Don't add the DDCUtils module as an editor dependency if we're building a Rocket distribution.

diff --git a/Engine/Source/Developer/DerivedDataCache/DerivedDataCache.Build.cs b/Engine/Source/Developer/DerivedDataCache/DerivedDataCache.Build.cs
--- a/Engine/Source/Developer/DerivedDataCache/DerivedDataCache.Build.cs
+++ b/Engine/Source/Developer/DerivedDataCache/DerivedDataCache.Build.cs
@@ -9,7 +9,7 @@ public class DerivedDataCache : ModuleRules
 	{
 		PrivateDependencyModuleNames.Add("Core");
 		// Internal (NoRedist) module
-		if (Directory.Exists(Path.Combine("Developer", "NoRedist", "DDCUtils")))
+		if (Directory.Exists(Path.Combine("Developer", "NoRedist", "DDCUtils")) && !UnrealBuildTool.UnrealBuildTool.BuildingRocket())
 		{
 			DynamicallyLoadedModuleNames.Add("DDCUtils");
 		}
