Moved DDCUtils from NoRedist to NotForLicensees

diff --git a/Engine/Source/Developer/DerivedDataCache/DerivedDataCache.Build.cs b/Engine/Source/Developer/DerivedDataCache/DerivedDataCache.Build.cs
--- a/Engine/Source/Developer/DerivedDataCache/DerivedDataCache.Build.cs
+++ b/Engine/Source/Developer/DerivedDataCache/DerivedDataCache.Build.cs
@@ -9,7 +9,7 @@ public class DerivedDataCache : ModuleRules
 	{
 		PrivateDependencyModuleNames.Add("Core");
 		// Internal (NoRedist) module
-		if (Directory.Exists(Path.Combine("Developer", "NoRedist", "DDCUtils")) && !UnrealBuildTool.UnrealBuildTool.BuildingRocket())
+		if (Directory.Exists(Path.Combine("Developer", "NotForLicensees", "DDCUtils")) && !UnrealBuildTool.UnrealBuildTool.BuildingRocket())
 		{
 			DynamicallyLoadedModuleNames.Add("DDCUtils");
 		}
