Update comment NoRedist -> NotForLicensees

diff --git a/Engine/Source/Developer/DerivedDataCache/DerivedDataCache.Build.cs b/Engine/Source/Developer/DerivedDataCache/DerivedDataCache.Build.cs
--- a/Engine/Source/Developer/DerivedDataCache/DerivedDataCache.Build.cs
+++ b/Engine/Source/Developer/DerivedDataCache/DerivedDataCache.Build.cs
@@ -8,7 +8,7 @@ public class DerivedDataCache : ModuleRules
 	public DerivedDataCache(TargetInfo Target)
 	{
 		PrivateDependencyModuleNames.Add("Core");
-		// Internal (NoRedist) module
+		// Internal (NotForLicensees) module
 		if (Directory.Exists(Path.Combine("Developer", "NotForLicensees", "DDCUtils")) && !UnrealBuildTool.UnrealBuildTool.BuildingRocket())
 		{
 			DynamicallyLoadedModuleNames.Add("DDCUtils");
