Removing message box when DDC filesystem is slow.

diff --git a/Engine/Source/Developer/DerivedDataCache/Private/FileSystemDerivedDataBackend.cpp b/Engine/Source/Developer/DerivedDataCache/Private/FileSystemDerivedDataBackend.cpp
--- a/Engine/Source/Developer/DerivedDataCache/Private/FileSystemDerivedDataBackend.cpp
+++ b/Engine/Source/Developer/DerivedDataCache/Private/FileSystemDerivedDataBackend.cpp
@@ -55,7 +55,7 @@ public:
 		// If we find a platform that has more stingent limits, this needs to be rethought.
 		checkAtCompileTime(MAX_BACKEND_KEY_LENGTH + MAX_CACHE_DIR_LEN + MAX_BACKEND_NUMBERED_SUBFOLDER_LENGTH + MAX_CACHE_EXTENTION_LEN < PLATFORM_MAX_FILEPATH_LENGTH, not_enough_room_left_for_cache_keys_in_max_path);
 		const double SlowInitDuration = 10.0;
-		bool bFilesystemIsSlow = false;
+		double AccessDuration = 0.0;
 
 		check(CachePath.Len());
 		FPaths::NormalizeFilename(CachePath);
@@ -84,27 +84,19 @@ public:
 			{
 				IFileManager::Get().Delete(*TempFilename, false, false, true);
 			}
-			double TestDuration = FPlatformTime::Seconds() - TestStart;
-			if (TestDuration > SlowInitDuration)
-			{
-				bFilesystemIsSlow = true;
-			}
+			AccessDuration = FPlatformTime::Seconds() - TestStart;
 		}
 		if (bFailed)
 		{
 			double StartTime = FPlatformTime::Seconds();
 			TArray<FString> FilesAndDirectories;
 			IFileManager::Get().FindFiles(FilesAndDirectories,*(CachePath / TEXT("*.*")), true, true);
-			double TestDuration = FPlatformTime::Seconds() - StartTime;
+			AccessDuration = FPlatformTime::Seconds() - StartTime;
 			if (FilesAndDirectories.Num() > 0)
 			{
 				bReadOnly = true;
 				bFailed = false;
 			}
-			if (TestDuration > SlowInitDuration)
-			{
-				bFilesystemIsSlow = true;
-			}
 		}
 		if (FString(FCommandLine::Get()).Contains(TEXT("DerivedDataCache")) )
 		{
@@ -123,14 +115,9 @@ public:
 			UE_LOG(LogDerivedDataCache, Display, TEXT("Files in %s will be touched."),*CachePath);
 		}
 
-		if (!bFailed && bFilesystemIsSlow)
+		if (!bFailed && AccessDuration > SlowInitDuration)
 		{
-			UE_LOG(LogDerivedDataCache, Warning, TEXT("%s access is very slow, consider disabling it."), *CachePath);
-			uint32 Result = FMessageDialog::Open(EAppMsgType::YesNo, FText::Format(NSLOCTEXT("Engine", "SlowDDC", "DDC filesystem backend is very slow:\n\n    {0}\n\nWould you like to disable it?"), FText::FromString(CachePath)));
-			if (Result == EAppReturnType::Yes)
-			{
-				bFailed = true;
-			}
+			UE_LOG(LogDerivedDataCache, Warning, TEXT("%s access is very slow (initialization took %.2lf seconds), consider disabling it."), *CachePath, AccessDuration);
 		}
 
 		if (!bReadOnly && !bFailed && bDeleteOldFiles && !FParse::Param(FCommandLine::Get(),TEXT("NODDCCLEANUP")) && FDDCCleanup::Get())
