Fix DDC not correctly handling project paths containing parentheses, due to them being stripped out on startup.

diff --git a/Engine/Source/Developer/DerivedDataCache/Private/DerivedDataBackends.cpp b/Engine/Source/Developer/DerivedDataCache/Private/DerivedDataBackends.cpp
--- a/Engine/Source/Developer/DerivedDataCache/Private/DerivedDataBackends.cpp
+++ b/Engine/Source/Developer/DerivedDataCache/Private/DerivedDataBackends.cpp
@@ -127,8 +127,8 @@ public:
 			// Trim whitespace at the beginning.
 			Entry = Entry.Trim();
 			// Remove brackets.
-			Entry = Entry.Replace( TEXT("("), TEXT("") );
-			Entry = Entry.Replace( TEXT(")"), TEXT("") );
+			Entry.RemoveFromStart(TEXT("("));
+			Entry.RemoveFromEnd(TEXT(")"));
 
 			FString	NodeType;
 			if( FParse::Value( *Entry, TEXT("Type="), NodeType ) )
