Slow DDC warning tweaks

diff --git a/Engine/Source/Developer/DerivedDataCache/Private/FileSystemDerivedDataBackend.cpp b/Engine/Source/Developer/DerivedDataCache/Private/FileSystemDerivedDataBackend.cpp
--- a/Engine/Source/Developer/DerivedDataCache/Private/FileSystemDerivedDataBackend.cpp
+++ b/Engine/Source/Developer/DerivedDataCache/Private/FileSystemDerivedDataBackend.cpp
@@ -17,27 +17,6 @@
 **/
 class FFileSystemDerivedDataBackend : public FDerivedDataBackendInterface
 {
-	/* Scoped timer that warns if the DDC access is taking too long. */
-	class FSlowAccessWarning
-	{
-		const FFileSystemDerivedDataBackend& Backend;
-		double StartTime;
-	public:
-		FSlowAccessWarning(const FFileSystemDerivedDataBackend& InBackend)
-			: Backend(InBackend)
-		{
-			StartTime = FPlatformTime::Seconds();
-		}
-		~FSlowAccessWarning()
-		{			
-			double Duration = FPlatformTime::Seconds() - StartTime;
-			const double SlowDuration = 10.0;
-			if (Duration >= SlowDuration)
-			{
-				UE_LOG(LogDerivedDataCache, Warning, TEXT("%s access is very slow, consider disabling it."), *Backend.CachePath);
-			}
-		}
-	};
 public:
 	/**
 	 * Constructor that should only be called once by the singleton, grabs the cache path from the ini
@@ -170,9 +149,15 @@ public:
 	{
 		check(!bFailed);
 		FString Filename = BuildFilename(CacheKey);
-		FSlowAccessWarning ScopedWarning(*this);
+
+		double StartTime = FPlatformTime::Seconds();
 		if (FFileHelper::LoadFileToArray(Data,*Filename,FILEREAD_Silent))
 		{
+			double ReadDuration = FPlatformTime::Seconds() - StartTime;
+			double ReadSpeed = ReadDuration > 5.0 ? (Data.Num() / ReadDuration) / (1024.0 * 1024.0) : 100.0;
+			// Slower than 0.5MB/s?
+			UE_CLOG(ReadSpeed < 0.5, LogDerivedDataCache, Warning, TEXT("%s access is very slow (%.2lfMB/s), consider disabling it."), *CachePath, ReadSpeed);
+
 			UE_LOG(LogDerivedDataCache, Verbose, TEXT("FileSystemDerivedDataBackend: Cache hit on %s"),*Filename);
 			return true;
 		}
@@ -201,7 +186,6 @@ public:
 				TempFilename = FPaths::GetPath(Filename) / TempFilename;
 				bool bResult;
 				{
-					FSlowAccessWarning ScopedWarning(*this);
 					bResult = FFileHelper::SaveArrayToFile(Data, *TempFilename);
 				}
 				if (bResult)
