Use the Rocket DDC backend when the engine is installed. Don't switch behaviours based on IsRocket().

diff --git a/Engine/Source/Developer/DerivedDataCache/Private/DerivedDataBackends.cpp b/Engine/Source/Developer/DerivedDataCache/Private/DerivedDataBackends.cpp
--- a/Engine/Source/Developer/DerivedDataCache/Private/DerivedDataBackends.cpp
+++ b/Engine/Source/Developer/DerivedDataCache/Private/DerivedDataBackends.cpp
@@ -72,7 +72,7 @@ public:
 		if( !RootCache )
 		{
 			// Use default graph
-			GraphName = FRocketSupport::IsRocket() ? TEXT("RocketDerivedDataBackendGraph") : TEXT("DerivedDataBackendGraph");
+			GraphName = FApp::IsEngineInstalled() ? TEXT("RocketDerivedDataBackendGraph") : TEXT("DerivedDataBackendGraph");
 			FString Entry;
 			if( !GConfig->GetString( *GraphName, TEXT("Root"), Entry, GEngineIni ) || !Entry.Len())
 			{
