Fixed a problem with the editor creating /Volumes/UE4DDC folder if the shared DDC folder is not mounted on Mac

diff --git a/Engine/Source/Developer/DerivedDataCache/Private/DerivedDataBackends.cpp b/Engine/Source/Developer/DerivedDataCache/Private/DerivedDataBackends.cpp
--- a/Engine/Source/Developer/DerivedDataCache/Private/DerivedDataBackends.cpp
+++ b/Engine/Source/Developer/DerivedDataCache/Private/DerivedDataBackends.cpp
@@ -523,7 +523,13 @@ public:
 				DeleteOldFiles( *Path );
 			}
 
-			FDerivedDataBackendInterface* InnerFileSystem = CreateFileSystemDerivedDataBackend( *Path, bReadOnly, bTouch, bPurgeTransient, bDeleteUnused, UnusedFileAge, MaxFoldersToClean );
+			FDerivedDataBackendInterface* InnerFileSystem = NULL;
+
+			// Don't create the file system if shared data cache directory is not mounted
+			if( FCString::Strcmp(NodeName, TEXT("Shared")) != 0 || IFileManager::Get().DirectoryExists(*Path) )
+			{
+				InnerFileSystem = CreateFileSystemDerivedDataBackend( *Path, bReadOnly, bTouch, bPurgeTransient, bDeleteUnused, UnusedFileAge, MaxFoldersToClean );
+			}
 
 			if( InnerFileSystem )
 			{
