UE4 - add CacheDataProbablyExists stats to DDC

diff --git a/Engine/Source/Developer/DerivedDataCache/Private/DerivedDataCache.cpp b/Engine/Source/Developer/DerivedDataCache/Private/DerivedDataCache.cpp
--- a/Engine/Source/Developer/DerivedDataCache/Private/DerivedDataCache.cpp
+++ b/Engine/Source/Developer/DerivedDataCache/Private/DerivedDataCache.cpp
@@ -13,10 +13,12 @@
 DEFINE_STAT(STAT_DDC_NumGets);
 DEFINE_STAT(STAT_DDC_NumPuts);
 DEFINE_STAT(STAT_DDC_NumBuilds);
+DEFINE_STAT(STAT_DDC_NumExist);
 DEFINE_STAT(STAT_DDC_SyncGetTime);
 DEFINE_STAT(STAT_DDC_ASyncWaitTime);
 DEFINE_STAT(STAT_DDC_PutTime);
 DEFINE_STAT(STAT_DDC_SyncBuildTime);
+DEFINE_STAT(STAT_DDC_ExistTime);
 
 /** 
  * Implementation of the derived data cache
@@ -288,7 +290,15 @@ public:
 
 	virtual bool CachedDataProbablyExists(const TCHAR* CacheKey) override
 	{
-		return FDerivedDataBackend::Get().GetRoot().CachedDataProbablyExists(CacheKey);
+		bool bResult;
+		INC_DWORD_STAT(STAT_DDC_NumExist);
+		STAT(double ThisTime = 0);
+		{
+			SCOPE_SECONDS_COUNTER(ThisTime);
+			bResult = FDerivedDataBackend::Get().GetRoot().CachedDataProbablyExists(CacheKey);
+		}
+		INC_FLOAT_STAT_BY(STAT_DDC_ExistTime, (float)ThisTime);
+		return bResult;
 	}
 
 	void NotifyBootComplete() override
diff --git a/Engine/Source/Developer/DerivedDataCache/Public/DerivedDataBackendInterface.h b/Engine/Source/Developer/DerivedDataCache/Public/DerivedDataBackendInterface.h
--- a/Engine/Source/Developer/DerivedDataCache/Public/DerivedDataBackendInterface.h
+++ b/Engine/Source/Developer/DerivedDataCache/Public/DerivedDataBackendInterface.h
@@ -7,10 +7,12 @@ DECLARE_LOG_CATEGORY_EXTERN(LogDerivedDataCache, Log, All);
 DECLARE_DWORD_ACCUMULATOR_STAT_EXTERN(TEXT("Num Gets"),STAT_DDC_NumGets,STATGROUP_DDC, );
 DECLARE_DWORD_ACCUMULATOR_STAT_EXTERN(TEXT("Num Puts"),STAT_DDC_NumPuts,STATGROUP_DDC, );
 DECLARE_DWORD_ACCUMULATOR_STAT_EXTERN(TEXT("Num Build"),STAT_DDC_NumBuilds,STATGROUP_DDC, );
-DECLARE_FLOAT_ACCUMULATOR_STAT_EXTERN(TEXT("DDC Sync Get Time"),STAT_DDC_SyncGetTime,STATGROUP_DDC, );
-DECLARE_FLOAT_ACCUMULATOR_STAT_EXTERN(TEXT("DDC ASync Wait Time"),STAT_DDC_ASyncWaitTime,STATGROUP_DDC, );
-DECLARE_FLOAT_ACCUMULATOR_STAT_EXTERN(TEXT("DDC Sync Put Time"),STAT_DDC_PutTime,STATGROUP_DDC, );
-DECLARE_FLOAT_ACCUMULATOR_STAT_EXTERN(TEXT("DDC Sync Build Time"),STAT_DDC_SyncBuildTime,STATGROUP_DDC, );
+DECLARE_DWORD_ACCUMULATOR_STAT_EXTERN(TEXT("Num Exists"),STAT_DDC_NumExist,STATGROUP_DDC, );
+DECLARE_FLOAT_ACCUMULATOR_STAT_EXTERN(TEXT("Sync Get Time"),STAT_DDC_SyncGetTime,STATGROUP_DDC, );
+DECLARE_FLOAT_ACCUMULATOR_STAT_EXTERN(TEXT("ASync Wait Time"),STAT_DDC_ASyncWaitTime,STATGROUP_DDC, );
+DECLARE_FLOAT_ACCUMULATOR_STAT_EXTERN(TEXT("Sync Put Time"),STAT_DDC_PutTime,STATGROUP_DDC, );
+DECLARE_FLOAT_ACCUMULATOR_STAT_EXTERN(TEXT("Sync Build Time"),STAT_DDC_SyncBuildTime,STATGROUP_DDC, );
+DECLARE_FLOAT_ACCUMULATOR_STAT_EXTERN(TEXT("Exists Time"),STAT_DDC_ExistTime,STATGROUP_DDC, );
 
 /** 
  * Interface for cache server backends. 
