Prevent missing BOOT.DDP warnings on build machines when starting up a project for the first time.

diff --git a/Engine/Source/Developer/DerivedDataCache/Private/DerivedDataBackends.cpp b/Engine/Source/Developer/DerivedDataCache/Private/DerivedDataBackends.cpp
--- a/Engine/Source/Developer/DerivedDataCache/Private/DerivedDataBackends.cpp
+++ b/Engine/Source/Developer/DerivedDataCache/Private/DerivedDataBackends.cpp
@@ -213,7 +213,7 @@ public:
 	 * @param NodeName Node name
 	 * @param Entry Node definition
 	 * @param bWriting true to create pak interface for writing
-	 * @return Pak file data backend interface instance or NULL if unsuccessfull
+	 * @return Pak file data backend interface instance or NULL if unsuccessful
 	 */
 	FDerivedDataBackendInterface* ParsePak( const TCHAR* NodeName, const TCHAR* Entry, const bool bWriting )
 	{
@@ -263,7 +263,7 @@ public:
 	 * @param IniFilename ini filename.
 	 * @param IniSection ini section containing graph definition
 	 * @param InParsedNodes map of nodes and their names which have already been parsed
-	 * @return Verify wrapper backend interface instance or NULL if unsuccessfull
+	 * @return Verify wrapper backend interface instance or NULL if unsuccessful
 	 */
 	FDerivedDataBackendInterface* ParseVerify( const TCHAR* NodeName, const TCHAR* Entry, const FString& IniFilename, const TCHAR* IniSection, TMap<FString, FDerivedDataBackendInterface*>& InParsedNodes )
 	{
@@ -587,7 +587,11 @@ public:
 					return Cache;
 				}
 
-				if( Cache->LoadCache( *Filename ) )
+				if (IFileManager::Get().FileSize(*Filename) < 0)
+				{
+					UE_LOG( LogDerivedDataCache, Display, TEXT("Starting with empty %s cache"), NodeName );
+				}
+				else if( Cache->LoadCache( *Filename ) )
 				{
 					UE_LOG( LogDerivedDataCache, Display, TEXT("Loaded %s cache: %s"), NodeName, *Filename );
 				}
