[INTEGRATE] Change 2423070 by [REDACTED_PERSON][REDACTED_HANDLE] on 2015/01/28 19:39:35

	Add a parameter to the pak file backend to merge DDC with an existing cache, and include DDC for MobileStarterContent and templates in Rocket build.

diff --git a/Engine/Config/BaseEngine.ini b/Engine/Config/BaseEngine.ini
--- a/Engine/Config/BaseEngine.ini
+++ b/Engine/Config/BaseEngine.ini
@@ -1889,7 +1889,7 @@ Local=(Type=FileSystem, ReadOnly=false, Clean=false, Flush=false, PurgeTransient
 Shared=(Type=FileSystem, ReadOnly=false, Clean=false, Flush=false, DeleteUnused=false, UnusedFileAge=23, FoldersToClean=-1, Path=\\epicgames.net\root\UE4DDC, EnvPathOverride=UE-SharedDataCachePath)
 AltShared=(Type=FileSystem, ReadOnly=true, Clean=false, Flush=false, DeleteUnused=false, UnusedFileAge=23, FoldersToClean=-1, Path=\\epicgames.net\root\UE4DDC2, EnvPathOverride=UE-SharedDataCachePath2)
 PakRead=(Type=ReadPak, Filename=../../../Engine/DerivedDataCache/DDC.ddp)
-PakWrite=(Type=WritePak, Filename=../../../Engine/DerivedDataCache/DDC.ddp)
+PakWrite=(Type=WritePak, Filename=../../../Engine/DerivedDataCache/DDC.ddp, Merge=true)
 
 [CreateProjectCache]
 MinimumDaysToKeepFile=7
diff --git a/Engine/Source/Developer/DerivedDataCache/Private/DerivedDataBackends.cpp b/Engine/Source/Developer/DerivedDataCache/Private/DerivedDataBackends.cpp
--- a/Engine/Source/Developer/DerivedDataCache/Private/DerivedDataBackends.cpp
+++ b/Engine/Source/Developer/DerivedDataCache/Private/DerivedDataBackends.cpp
@@ -235,6 +235,12 @@ public:
 				WritePakFilename = PakFilename + TEXT(".") + Temp.ToString();
 				WritePakCache = new FPakFileDerivedDataBackend( *WritePakFilename, true );
 				PakNode = WritePakCache;
+
+				bool bMerge = false;
+				if(FParse::Bool(Entry, TEXT("Merge="), bMerge) && bMerge)
+				{
+					WritePakCache->MergeCache(*PakFilename);
+				}
 			}
 			else
 			{
diff --git a/Engine/Source/Developer/DerivedDataCache/Private/PakFileDerivedDataBackend.h b/Engine/Source/Developer/DerivedDataCache/Private/PakFileDerivedDataBackend.h
--- a/Engine/Source/Developer/DerivedDataCache/Private/PakFileDerivedDataBackend.h
+++ b/Engine/Source/Developer/DerivedDataCache/Private/PakFileDerivedDataBackend.h
@@ -347,6 +347,35 @@ public:
 		}
 		return true;
 	}
+
+	/**
+	 * Merges another cache file into this one.
+	 * @return true on success
+	 */
+	bool MergeCache(const TCHAR* OtherFileName)
+	{
+		if(FPlatformFileManager::Get().GetPlatformFile().FileExists(OtherFileName))
+		{
+			FPakFileDerivedDataBackend OtherPak(OtherFileName, false);
+			if(!OtherPak.bClosed)
+			{
+				// Get all the existing keys
+				TArray<FString> KeyNames;
+				OtherPak.CacheItems.GenerateKeyArray(KeyNames);
+
+				// Copy them all to the new cache
+				TArray<uint8> Buffer;
+				for(const FString& KeyName: KeyNames)
+				{
+					Buffer.Reset();
+					OtherPak.GetCachedData(*KeyName, Buffer);
+					PutCachedData(*KeyName, Buffer, false);
+				}
+				return true;
+			}
+		}
+		return false;
+	}
 	
 	const FString& GetFilename() const
 	{
