[INTEGRATE] Change 2425182 by [REDACTED_PERSON][REDACTED_HANDLE] on 2015/01/30 00:27:58

	Fix for building engine DDC on Mac; suspect that lock is due to recursion in FDerivedDataBackend::Get(), due to recursion through WaitForQuiescence. Instead of appending existing cache on startup, append it on shutdown if a command line parameter is present instead.

diff --git a/Engine/Config/BaseEngine.ini b/Engine/Config/BaseEngine.ini
--- a/Engine/Config/BaseEngine.ini
+++ b/Engine/Config/BaseEngine.ini
@@ -1888,7 +1888,8 @@ Boot=(Type=Boot, Filename="%GAMEDIR%DerivedDataCache/Boot.ddc", MaxCacheSize=512
 Local=(Type=FileSystem, ReadOnly=false, Clean=false, Flush=false, PurgeTransient=true, DeleteUnused=true, UnusedFileAge=34, FoldersToClean=-1, Path=../../../Engine/DerivedDataCache, EnvPathOverride=UE-LocalDataCachePath)
 Shared=(Type=FileSystem, ReadOnly=false, Clean=false, Flush=false, DeleteUnused=false, UnusedFileAge=23, FoldersToClean=-1, Path=\\epicgames.net\root\UE4DDC, EnvPathOverride=UE-SharedDataCachePath)
 AltShared=(Type=FileSystem, ReadOnly=true, Clean=false, Flush=false, DeleteUnused=false, UnusedFileAge=23, FoldersToClean=-1, Path=\\epicgames.net\root\UE4DDC2, EnvPathOverride=UE-SharedDataCachePath2)
-PakWrite=(Type=WritePak, Filename=../../../Engine/DerivedDataCache/DDC.ddp, Merge=true)
+PakRead=(Type=ReadPak, Filename=../../../Engine/DerivedDataCache/DDC.ddp)
+PakWrite=(Type=WritePak, Filename=../../../Engine/DerivedDataCache/DDC.ddp)
 
 [CreateProjectCache]
 MinimumDaysToKeepFile=7
diff --git a/Engine/Source/Developer/DerivedDataCache/Private/DerivedDataBackends.cpp b/Engine/Source/Developer/DerivedDataCache/Private/DerivedDataBackends.cpp
--- a/Engine/Source/Developer/DerivedDataCache/Private/DerivedDataBackends.cpp
+++ b/Engine/Source/Developer/DerivedDataCache/Private/DerivedDataBackends.cpp
@@ -235,12 +235,6 @@ public:
 				WritePakFilename = PakFilename + TEXT(".") + Temp.ToString();
 				WritePakCache = new FPakFileDerivedDataBackend( *WritePakFilename, true );
 				PakNode = WritePakCache;
-
-				bool bMerge = false;
-				if(FParse::Bool(Entry, TEXT("Merge="), bMerge) && bMerge)
-				{
-					WritePakCache->MergeCache(*PakFilename);
-				}
 			}
 			else
 			{
@@ -680,6 +674,13 @@ public:
 		}
 		if (bShutdown)
 		{
+			if(FParse::Param(FCommandLine::Get(), TEXT("MergePaks")) && WritePakCache && WritePakCache->IsWritable())
+			{
+				for (int32 ReadPakIndex = 0; ReadPakIndex < ReadPakCache.Num(); ReadPakIndex++)
+				{
+					WritePakCache->MergeCache(ReadPakCache[ReadPakIndex]);
+				}
+			}
 			for (int32 ReadPakIndex = 0; ReadPakIndex < ReadPakCache.Num(); ReadPakIndex++)
 			{
 				ReadPakCache[ReadPakIndex]->Close();
diff --git a/Engine/Source/Developer/DerivedDataCache/Private/PakFileDerivedDataBackend.h b/Engine/Source/Developer/DerivedDataCache/Private/PakFileDerivedDataBackend.h
--- a/Engine/Source/Developer/DerivedDataCache/Private/PakFileDerivedDataBackend.h
+++ b/Engine/Source/Developer/DerivedDataCache/Private/PakFileDerivedDataBackend.h
@@ -352,29 +352,25 @@ public:
 	 * Merges another cache file into this one.
 	 * @return true on success
 	 */
-	bool MergeCache(const TCHAR* OtherFileName)
+	void MergeCache(FPakFileDerivedDataBackend* OtherPak)
 	{
-		if(FPlatformFileManager::Get().GetPlatformFile().FileExists(OtherFileName))
+		// Get all the existing keys
+		TArray<FString> KeyNames;
+		OtherPak->CacheItems.GenerateKeyArray(KeyNames);
+		UE_LOG(LogDerivedDataCache, Display, TEXT("Found %d entries."), KeyNames.Num());
+
+		// Copy them all to the new cache
+		TArray<uint8> Buffer;
+		for(int32 Idx = 0; Idx < KeyNames.Num(); Idx++)
 		{
-			FPakFileDerivedDataBackend OtherPak(OtherFileName, false);
-			if(!OtherPak.bClosed)
-			{
-				// Get all the existing keys
-				TArray<FString> KeyNames;
-				OtherPak.CacheItems.GenerateKeyArray(KeyNames);
+			Buffer.Reset();
 
-				// Copy them all to the new cache
-				TArray<uint8> Buffer;
-				for(const FString& KeyName: KeyNames)
-				{
-					Buffer.Reset();
-					OtherPak.GetCachedData(*KeyName, Buffer);
-					PutCachedData(*KeyName, Buffer, false);
-				}
-				return true;
-			}
+			UE_LOG(LogDerivedDataCache, Display, TEXT("[%d/%d] Reading %s..."), Idx + 1, KeyNames.Num(), *KeyNames[Idx]);
+			OtherPak->GetCachedData(*KeyNames[Idx], Buffer);
+
+			UE_LOG(LogDerivedDataCache, Display, TEXT("[%d/%d] Writing %d bytes..."), Idx + 1, KeyNames.Num(), Buffer.Num());
+			PutCachedData(*KeyNames[Idx], Buffer, false);
 		}
-		return false;
 	}
 	
 	const FString& GetFilename() const
