Check for DDCUtils module file instead of directory when building DDC module

diff --git a/Engine/Source/Developer/DerivedDataCache/DerivedDataCache.Build.cs b/Engine/Source/Developer/DerivedDataCache/DerivedDataCache.Build.cs
--- a/Engine/Source/Developer/DerivedDataCache/DerivedDataCache.Build.cs
+++ b/Engine/Source/Developer/DerivedDataCache/DerivedDataCache.Build.cs
@@ -9,9 +9,13 @@ public class DerivedDataCache : ModuleRules
 	{
 		PrivateDependencyModuleNames.Add("Core");
 		// Internal (NotForLicensees) module
-		if (Directory.Exists(Path.Combine("Developer", "NotForLicensees", "DDCUtils")) && !UnrealBuildTool.UnrealBuildTool.BuildingRocket())
+		if (!UnrealBuildTool.UnrealBuildTool.BuildingRocket())
 		{
-			DynamicallyLoadedModuleNames.Add("DDCUtils");
+			var DDCUtilsModule = Path.Combine("Developer", "NotForLicensees", "DDCUtils", "DDCUtils.Build.cs");
+			if (File.Exists(DDCUtilsModule))
+			{
+				DynamicallyLoadedModuleNames.Add("DDCUtils");
+			}
 		}
 	}
 }
