[INTEGRATE] Change 2444249 by [REDACTED_PERSON][REDACTED_HANDLE] on 2015/02/12 22:09:45

	Add a wrapper around the PAK-file DDC backend which can compress cache entries. Big reduction in the size of the DDC pak file, and seems to add negligable amount (measured average < 0.5s) to startup time. For UE-8823.

diff --git a/Engine/Config/BaseEngine.ini b/Engine/Config/BaseEngine.ini
--- a/Engine/Config/BaseEngine.ini
+++ b/Engine/Config/BaseEngine.ini
@@ -1853,7 +1853,7 @@ Hierarchy=(Type=Hierarchical, Inner=Boot, Inner=Pak, Inner=EnginePak, Inner=Loca
 Boot=(Type=Boot, Filename="%ENGINEUSERDIR%DerivedDataCache/Boot.ddc", MaxCacheSize=512)
 Local=(Type=FileSystem, ReadOnly=false, Clean=false, Flush=false, PurgeTransient=true, DeleteUnused=true, UnusedFileAge=34, FoldersToClean=-1, Path="%ENGINEVERSIONAGNOSTICUSERDIR%DerivedDataCache")
 Pak=(Type=ReadPak, Filename="%GAMEDIR%DerivedDataCache/DDC.ddp")
-EnginePak=(Type=ReadPak, Filename=../../../Engine/DerivedDataCache/DDC.ddp)
+EnginePak=(Type=ReadPak, Filename=../../../Engine/DerivedDataCache/Compressed.ddp, Compressed=true)
 
 [NoShared]
 MinimumDaysToKeepFile=7
@@ -1898,8 +1898,8 @@ Boot=(Type=Boot, Filename="%GAMEDIR%DerivedDataCache/Boot.ddc", MaxCacheSize=512
 Local=(Type=FileSystem, ReadOnly=false, Clean=false, Flush=false, PurgeTransient=true, DeleteUnused=true, UnusedFileAge=34, FoldersToClean=-1, Path=../../../Engine/DerivedDataCache, EnvPathOverride=UE-LocalDataCachePath)
 Shared=(Type=FileSystem, ReadOnly=false, Clean=false, Flush=false, DeleteUnused=false, UnusedFileAge=23, FoldersToClean=-1, Path=?EpicDDC, EnvPathOverride=UE-SharedDataCachePath)
 AltShared=(Type=FileSystem, ReadOnly=true, Clean=false, Flush=false, DeleteUnused=false, UnusedFileAge=23, FoldersToClean=-1, Path=?EpicDDC2, EnvPathOverride=UE-SharedDataCachePath2)
-PakRead=(Type=ReadPak, Filename=../../../Engine/DerivedDataCache/DDC.ddp)
-PakWrite=(Type=WritePak, Filename=../../../Engine/DerivedDataCache/DDC.ddp)
+PakRead=(Type=ReadPak, Filename=../../../Engine/DerivedDataCache/Compressed.ddp, Compressed=true)
+PakWrite=(Type=WritePak, Filename=../../../Engine/DerivedDataCache/Compressed.ddp, Compressed=true)
 
 [CreateProjectCache]
 MinimumDaysToKeepFile=7
diff --git a/Engine/Source/Developer/DerivedDataCache/Private/DerivedDataBackends.cpp b/Engine/Source/Developer/DerivedDataCache/Private/DerivedDataBackends.cpp
--- a/Engine/Source/Developer/DerivedDataCache/Private/DerivedDataBackends.cpp
+++ b/Engine/Source/Developer/DerivedDataCache/Private/DerivedDataBackends.cpp
@@ -220,6 +220,7 @@ public:
 		FDerivedDataBackendInterface* PakNode = NULL;
 		FString PakFilename;
 		FParse::Value( Entry, TEXT("Filename="), PakFilename );
+		bool bCompressed = GetParsedBool(Entry, TEXT("Compressed="));
 
 		if( !PakFilename.Len() )
 		{
@@ -233,7 +234,7 @@ public:
 				FGuid Temp = FGuid::NewGuid();
 				ReadPakFilename = PakFilename;
 				WritePakFilename = PakFilename + TEXT(".") + Temp.ToString();
-				WritePakCache = new FPakFileDerivedDataBackend( *WritePakFilename, true );
+				WritePakCache = bCompressed? new FCompressedPakFileDerivedDataBackend( *WritePakFilename, true ) : new FPakFileDerivedDataBackend( *WritePakFilename, true );
 				PakNode = WritePakCache;
 			}
 			else
@@ -241,7 +242,7 @@ public:
 				bool bReadPak = FPlatformFileManager::Get().GetPlatformFile().FileExists( *PakFilename );
 				if( bReadPak )
 				{
-					FPakFileDerivedDataBackend* ReadPak = new FPakFileDerivedDataBackend( *PakFilename, false );
+					FPakFileDerivedDataBackend* ReadPak = bCompressed? new FCompressedPakFileDerivedDataBackend( *PakFilename, false ) : new FPakFileDerivedDataBackend( *PakFilename, false );
 					ReadPakFilename = PakFilename;
 					PakNode = ReadPak;
 					ReadPakCache.Add(ReadPak);
diff --git a/Engine/Source/Developer/DerivedDataCache/Private/PakFileDerivedDataBackend.h b/Engine/Source/Developer/DerivedDataCache/Private/PakFileDerivedDataBackend.h
--- a/Engine/Source/Developer/DerivedDataCache/Private/PakFileDerivedDataBackend.h
+++ b/Engine/Source/Developer/DerivedDataCache/Private/PakFileDerivedDataBackend.h
@@ -2,6 +2,8 @@
 
 #pragma once
 
+#include "Compression.h"
+
 /** 
  * A simple thread safe, pak file based backend. 
 **/
@@ -363,13 +365,15 @@ public:
 		TArray<uint8> Buffer;
 		for(int32 Idx = 0; Idx < KeyNames.Num(); Idx++)
 		{
-			Buffer.Reset();
-
-			UE_LOG(LogDerivedDataCache, Display, TEXT("[%d/%d] Reading %s..."), Idx + 1, KeyNames.Num(), *KeyNames[Idx]);
-			OtherPak->GetCachedData(*KeyNames[Idx], Buffer);
-
-			UE_LOG(LogDerivedDataCache, Display, TEXT("[%d/%d] Writing %d bytes..."), Idx + 1, KeyNames.Num(), Buffer.Num());
-			PutCachedData(*KeyNames[Idx], Buffer, false);
+			if(!CachedDataProbablyExists(*KeyNames[Idx]))
+			{
+				Buffer.Reset();
+				if(OtherPak->GetCachedData(*KeyNames[Idx], Buffer))
+				{
+					UE_LOG(LogDerivedDataCache, Display, TEXT("[%d/%d] Copying %s (%d bytes)..."), Idx + 1, KeyNames.Num(), *KeyNames[Idx], Buffer.Num());
+					PutCachedData(*KeyNames[Idx], Buffer, false);
+				}
+			}
 		}
 	}
 	
@@ -437,3 +441,45 @@ private:
 	};
 };
 
+class FCompressedPakFileDerivedDataBackend : public FPakFileDerivedDataBackend
+{
+public:
+	FCompressedPakFileDerivedDataBackend(const TCHAR* InFilename, bool bInWriting) 
+		: FPakFileDerivedDataBackend(InFilename, bInWriting)
+	{
+	}
+
+	virtual void PutCachedData(const TCHAR* CacheKey, TArray<uint8>& InData, bool bPutEvenIfExists) override
+	{
+		int32 UncompressedSize = InData.Num();
+		int32 CompressedSize = FCompression::CompressMemoryBound(CompressionFlags, UncompressedSize);
+
+		TArray<uint8> CompressedData;
+		CompressedData.AddUninitialized(CompressedSize + sizeof(UncompressedSize));
+
+		FMemory::Memcpy(&CompressedData[0], &UncompressedSize, sizeof(UncompressedSize));
+		verify(FCompression::CompressMemory(CompressionFlags, CompressedData.GetData() + sizeof(UncompressedSize), CompressedSize, InData.GetData(), InData.Num()));
+		CompressedData.SetNum(CompressedSize + sizeof(UncompressedSize), false);
+
+		FPakFileDerivedDataBackend::PutCachedData(CacheKey, CompressedData, bPutEvenIfExists);
+	}
+
+	virtual bool GetCachedData(const TCHAR* CacheKey, TArray<uint8>& OutData) override
+	{
+		TArray<uint8> CompressedData;
+		if(!FPakFileDerivedDataBackend::GetCachedData(CacheKey, CompressedData))
+		{
+			return false;
+		}
+
+		int32 UncompressedSize;
+		FMemory::Memcpy(&UncompressedSize, &CompressedData[0], sizeof(UncompressedSize));
+		OutData.SetNum(UncompressedSize);
+		verify(FCompression::UncompressMemory(CompressionFlags, OutData.GetData(), UncompressedSize, CompressedData.GetData() + sizeof(UncompressedSize), CompressedData.Num() - sizeof(UncompressedSize)));
+
+		return true;
+	}
+
+private:
+	static const ECompressionFlags CompressionFlags = (ECompressionFlags)(COMPRESS_ZLIB | COMPRESS_BiasMemory);
+};
