Change DDCUtils to written to a NotForLicensees folder by default, and remove special case excluding it from Rocket builds (NotForLicensees folder filtering will pick this up automatically).

diff --git a/Engine/Build/JunkManifest.txt b/Engine/Build/JunkManifest.txt
--- a/Engine/Build/JunkManifest.txt
+++ b/Engine/Build/JunkManifest.txt
@@ -36,6 +36,7 @@ Engine\Binaries\Mac\UE4Editor.app\Contents\MacOS\*-CameraDebugger.*
 Engine\Binaries\Mac\UE4Editor.app\Contents\MacOS\*-CameraEditor.*
 Engine\Binaries\Mac\UE4Editor.app\Contents\MacOS\*-SessionLauncher.*
 Engine\Binaries\Mac\*-MessagingDebugger.*
+Engine\Binaries\Win64\*-DDCUtils.*
 Engine\Binaries\Win64\*-MessagingDebugger.*
 Engine\Binaries\Win64\*-TextureXboxOneFormat.*
 Engine\Binaries\Win64\*-TextureFormatATITC.*
diff --git a/Engine/Source/Developer/DerivedDataCache/DerivedDataCache.Build.cs b/Engine/Source/Developer/DerivedDataCache/DerivedDataCache.Build.cs
--- a/Engine/Source/Developer/DerivedDataCache/DerivedDataCache.Build.cs
+++ b/Engine/Source/Developer/DerivedDataCache/DerivedDataCache.Build.cs
@@ -8,14 +8,12 @@ public class DerivedDataCache : ModuleRules
 	public DerivedDataCache(TargetInfo Target)
 	{
 		PrivateDependencyModuleNames.Add("Core");
+
 		// Internal (NotForLicensees) module
-		if (!UnrealBuildTool.UnrealBuildTool.BuildingRocket())
+		var DDCUtilsModule = Path.Combine("Developer", "NotForLicensees", "DDCUtils", "DDCUtils.Build.cs");
+		if (File.Exists(DDCUtilsModule))
 		{
-			var DDCUtilsModule = Path.Combine("Developer", "NotForLicensees", "DDCUtils", "DDCUtils.Build.cs");
-			if (File.Exists(DDCUtilsModule))
-			{
-				DynamicallyLoadedModuleNames.Add("DDCUtils");
-			}
+			DynamicallyLoadedModuleNames.Add("DDCUtils");
 		}
 	}
 }
