UE4 - add stats to all DDC entry points

diff --git a/Engine/Source/Developer/DerivedDataCache/Private/DerivedDataCache.cpp b/Engine/Source/Developer/DerivedDataCache/Private/DerivedDataCache.cpp
--- a/Engine/Source/Developer/DerivedDataCache/Private/DerivedDataCache.cpp
+++ b/Engine/Source/Developer/DerivedDataCache/Private/DerivedDataCache.cpp
@@ -139,6 +139,7 @@ public:
 
 	virtual bool GetSynchronous(FDerivedDataPluginInterface* DataDeriver, TArray<uint8>& OutData) override
 	{
+		QUICK_SCOPE_CYCLE_COUNTER(STAT_DDC_GetSynchronous);
 		check(DataDeriver);
 		FString CacheKey = FDerivedDataCache::BuildCacheKey(DataDeriver);
 		UE_LOG(LogDerivedDataCache, Verbose, TEXT("GetSynchronous %s"), *CacheKey);
@@ -151,6 +152,7 @@ public:
 
 	virtual uint32 GetAsynchronous(FDerivedDataPluginInterface* DataDeriver) override
 	{
+		QUICK_SCOPE_CYCLE_COUNTER(STAT_DDC_GetAsynchronous);
 		FScopeLock ScopeLock(&SynchronizationObject);
 		uint32 Handle = NextHandle();
 		FString CacheKey = FDerivedDataCache::BuildCacheKey(DataDeriver);
@@ -175,6 +177,7 @@ public:
 
 	virtual bool PollAsynchronousCompletion(uint32 Handle) override
 	{
+		QUICK_SCOPE_CYCLE_COUNTER(STAT_DDC_PollAsynchronousCompletion);
 		FAsyncTask<FBuildAsyncWorker>* AsyncTask = NULL;
 		{
 			FScopeLock ScopeLock(&SynchronizationObject);
@@ -186,6 +189,7 @@ public:
 
 	virtual void WaitAsynchronousCompletion(uint32 Handle) override
 	{
+		QUICK_SCOPE_CYCLE_COUNTER(STAT_DDC_WaitAsynchronousCompletion);
 		STAT(double ThisTime = 0);
 		{
 			SCOPE_SECONDS_COUNTER(ThisTime);
@@ -202,6 +206,7 @@ public:
 
 	virtual bool GetAsynchronousResults(uint32 Handle, TArray<uint8>& OutData) override
 	{
+		QUICK_SCOPE_CYCLE_COUNTER(STAT_DDC_GetAsynchronousResults);
 		FAsyncTask<FBuildAsyncWorker>* AsyncTask = NULL;
 		{
 			FScopeLock ScopeLock(&SynchronizationObject);
@@ -231,6 +236,7 @@ public:
 
 	virtual bool GetSynchronous(const TCHAR* CacheKey, TArray<uint8>& OutData) override
 	{
+		QUICK_SCOPE_CYCLE_COUNTER(STAT_DDC_GetSynchronous_Data);
 		UE_LOG(LogDerivedDataCache, Verbose, TEXT("GetSynchronous %s"), CacheKey);
 		FAsyncTask<FBuildAsyncWorker> PendingTask((FDerivedDataPluginInterface*)NULL, CacheKey, true);
 		AddToAsyncCompletionCounter(1);
@@ -241,6 +247,7 @@ public:
 
 	virtual uint32 GetAsynchronous(const TCHAR* CacheKey, IDerivedDataRollup* Rollup = NULL) override
 	{
+		QUICK_SCOPE_CYCLE_COUNTER(STAT_DDC_GetAsynchronous_Handle);
 		check(!Rollup); // this needs to be handled by someone else, if rollups are disabled, then it should be NULL
 		FScopeLock ScopeLock(&SynchronizationObject);
 		UE_LOG(LogDerivedDataCache, Verbose, TEXT("GetAsynchronous %s"), CacheKey);
@@ -261,6 +268,7 @@ public:
 	**/
 	void GetAsynchronousForRollup(const TCHAR* CacheKey, uint32 Handle)
 	{
+		QUICK_SCOPE_CYCLE_COUNTER(STAT_DDC_GetAsynchronousForRollup);
 		FScopeLock ScopeLock(&SynchronizationObject);
 		UE_LOG(LogDerivedDataCache, Verbose, TEXT("GetAsynchronous(handle) %s"), CacheKey);
 		FAsyncTask<FBuildAsyncWorker>* AsyncTask = new FAsyncTask<FBuildAsyncWorker>((FDerivedDataPluginInterface*)NULL, CacheKey, false);
@@ -272,6 +280,7 @@ public:
 
 	virtual void Put(const TCHAR* CacheKey, TArray<uint8>& Data, bool bPutEvenIfExists = false) override
 	{
+		QUICK_SCOPE_CYCLE_COUNTER(STAT_DDC_Put);
 		STAT(double ThisTime = 0);
 		{
 			SCOPE_SECONDS_COUNTER(ThisTime);
@@ -288,6 +297,7 @@ public:
 
 	virtual bool CachedDataProbablyExists(const TCHAR* CacheKey) override
 	{
+		QUICK_SCOPE_CYCLE_COUNTER(STAT_DDC_CachedDataProbablyExists);
 		bool bResult;
 		INC_DWORD_STAT(STAT_DDC_NumExist);
 		STAT(double ThisTime = 0);
@@ -301,6 +311,7 @@ public:
 
 	void NotifyBootComplete() override
 	{
+		QUICK_SCOPE_CYCLE_COUNTER(STAT_DDC_NotifyBootComplete);
 		FDerivedDataBackend::Get().NotifyBootComplete();
 	}
 
@@ -311,6 +322,7 @@ public:
 
 	void WaitForQuiescence(bool bShutdown) override
 	{
+		QUICK_SCOPE_CYCLE_COUNTER(STAT_DDC_WaitForQuiescence);
 		FDerivedDataBackend::Get().WaitForQuiescence(bShutdown);
 	}
 
