Linux: larger limits for absolute DDC paths.

diff --git a/Engine/Source/Developer/DerivedDataCache/Private/FileSystemDerivedDataBackend.cpp b/Engine/Source/Developer/DerivedDataCache/Private/FileSystemDerivedDataBackend.cpp
--- a/Engine/Source/Developer/DerivedDataCache/Private/FileSystemDerivedDataBackend.cpp
+++ b/Engine/Source/Developer/DerivedDataCache/Private/FileSystemDerivedDataBackend.cpp
@@ -8,7 +8,12 @@
 
 #define MAX_BACKEND_KEY_LENGTH (120)
 #define MAX_BACKEND_NUMBERED_SUBFOLDER_LENGTH (9)
-#define MAX_CACHE_DIR_LEN (119)
+#if PLATFORM_LINUX	// PATH_MAX on Linux is 4096 (getconf PATH_MAX /, also see limits.h), so this value can be larger (note that it is still arbitrary).
+                    // This should not affect sharing the cache between platforms as the absolute paths will be different anyway.
+	#define MAX_CACHE_DIR_LEN (3119)
+#else
+	#define MAX_CACHE_DIR_LEN (119)
+#endif // PLATFORM_LINUX
 #define MAX_CACHE_EXTENTION_LEN (4)
 
 /** 
