Change warning to a regular log message when waiting for DDC to finish.

diff --git a/Engine/Source/Developer/DerivedDataCache/Private/DerivedDataBackends.cpp b/Engine/Source/Developer/DerivedDataCache/Private/DerivedDataBackends.cpp
--- a/Engine/Source/Developer/DerivedDataCache/Private/DerivedDataBackends.cpp
+++ b/Engine/Source/Developer/DerivedDataCache/Private/DerivedDataBackends.cpp
@@ -669,7 +669,7 @@ public:
 			FPlatformProcess::Sleep(1.0f);
 			if (FPlatformTime::Seconds() - LastPrint > 5.0)
 			{
-				UE_LOG(LogDerivedDataCache, Warning, TEXT("Waited %ds for derived data cache to finish..."), int32(FPlatformTime::Seconds() - StartTime));
+				UE_LOG(LogDerivedDataCache, Log, TEXT("Waited %ds for derived data cache to finish..."), int32(FPlatformTime::Seconds() - StartTime));
 				LastPrint = FPlatformTime::Seconds();
 			}
 		}
