Renamed cooking stats modules.
Added additional stats.

diff --git a/Engine/Source/Developer/DerivedDataCache/Private/FileSystemDerivedDataBackend.cpp b/Engine/Source/Developer/DerivedDataCache/Private/FileSystemDerivedDataBackend.cpp
--- a/Engine/Source/Developer/DerivedDataCache/Private/FileSystemDerivedDataBackend.cpp
+++ b/Engine/Source/Developer/DerivedDataCache/Private/FileSystemDerivedDataBackend.cpp
@@ -164,6 +164,7 @@ public:
 		static FName NAME_GetCachedData(TEXT("GetCachedData"));
 		FDDCScopeStatHelper Stat(CacheKey, NAME_GetCachedData);
 		static FName NAME_FileDDCPath(TEXT("FileDDCPath"));
+		static FName NAME_Retrieved(TEXT("Retrieved"));
 		Stat.AddTag(NAME_FileDDCPath, CachePath);
 		
 		check(!bFailed);
@@ -178,10 +179,12 @@ public:
 			UE_CLOG(ReadSpeed < 0.5, LogDerivedDataCache, Warning, TEXT("%s access is very slow (%.2lfMB/s), consider disabling it."), *CachePath, ReadSpeed);
 
 			UE_LOG(LogDerivedDataCache, Verbose, TEXT("FileSystemDerivedDataBackend: Cache hit on %s"),*Filename);
+			Stat.AddTag(NAME_Retrieved, TEXT("true"));
 			return true;
 		}
 		UE_LOG(LogDerivedDataCache, Verbose, TEXT("FileSystemDerivedDataBackend: Cache miss on %s"),*Filename);
 		Data.Empty();
+		Stat.AddTag(NAME_Retrieved, TEXT("false"));
 		return false;	
 	}
 	/**
diff --git a/Engine/Source/Developer/DerivedDataCache/Private/HierarchicalDerivedDataBackend.h b/Engine/Source/Developer/DerivedDataCache/Private/HierarchicalDerivedDataBackend.h
--- a/Engine/Source/Developer/DerivedDataCache/Private/HierarchicalDerivedDataBackend.h
+++ b/Engine/Source/Developer/DerivedDataCache/Private/HierarchicalDerivedDataBackend.h
@@ -85,7 +85,8 @@ public:
 	{
 		const static FName NAME_GetCachedData = FName(TEXT("GetCachedData"));
 		const static FName NAME_HierarchicalDDC = FName(TEXT("HierarchicalDDC"));
-		const static FName NAME_DataSize = FName(TEXT("DataSize"));
+		const static FName NAME_DataSize = FName(TEXT("DataSize")); 
+		static FName NAME_Retrieved(TEXT("Retrieved"));
 
 		FDDCScopeStatHelper Stat(CacheKey, NAME_GetCachedData);
 		Stat.AddTag(NAME_HierarchicalDDC, FString());
@@ -130,9 +131,11 @@ public:
 					}
 				}
 				Stat.AddTag(NAME_DataSize, FString::Printf(TEXT("%d bytes"), OutData.Num()));
+				Stat.AddTag(NAME_Retrieved, TEXT("true"));
 				return true;
 			}
 		}
+		Stat.AddTag(NAME_Retrieved, TEXT("false"));
 		return false;
 	}
 	/**
diff --git a/Engine/Source/Editor/CookingStats/Public/CookingStatsModule.h b/Engine/Source/Editor/CookingStats/Public/CookingStatsModule.h
--- a/Engine/Source/Editor/CookingStats/Public/CookingStatsModule.h
+++ b/Engine/Source/Editor/CookingStats/Public/CookingStatsModule.h
@@ -5,7 +5,7 @@
 
 #include "ModuleInterface.h"
 #include "ICookingStats.h"
-//#include "CookingStatsInterface.h"
+#include "CookingStatsInterface.h"
 
 namespace CookingStatsConstants
 {
@@ -16,8 +16,8 @@ namespace CookingStatsConstants
 * Asset registry module
 */
 class FCookingStatsModule
-	: public IModuleInterface
-	//: public ICookingStatsModule
+	//: public IModuleInterface
+	: public ICookingStatsModule
 {
 
 public:
diff --git a/Engine/Source/Runtime/Core/Public/Misc/CookingStatsInterface.h b/Engine/Source/Runtime/Core/Public/Misc/CookingStatsInterface.h
--- a/Engine/Source/Runtime/Core/Public/Misc/CookingStatsInterface.h
+++ b/Engine/Source/Runtime/Core/Public/Misc/CookingStatsInterface.h
@@ -8,21 +8,21 @@
 /**
 * CookingStats module interface
 */
-class ICookingStatsInterface : public IModuleInterface
+class ICookingStatsModule : public IModuleInterface
 {
 public:
 	/**
 	 * Tries to gets a pointer to the active AssetRegistryInterface implementation. 
 	 */
-	static inline ICookingStatsInterface * GetPtr()
+	static inline ICookingStatsModule * GetPtr()
 	{
 		static FName CookingStats("CookingStats");
-		auto Module = FModuleManager::GetModulePtr<ICookingStatsInterface>(CookingStats);
+		auto Module = FModuleManager::GetModulePtr<ICookingStatsModule>(CookingStats);
 		if (Module == nullptr)
 		{
-			Module = &FModuleManager::LoadModuleChecked<ICookingStatsInterface>(CookingStats);
+			Module = &FModuleManager::LoadModuleChecked<ICookingStatsModule>(CookingStats);
 		}
-		return reinterpret_cast<ICookingStatsInterface*>(Module);
+		return reinterpret_cast<ICookingStatsModule*>(Module);
 	}
 
 };
