ddc stats

diff --git a/Engine/Source/Developer/DerivedDataCache/Private/DerivedDataBackendAsyncPutWrapper.h b/Engine/Source/Developer/DerivedDataCache/Private/DerivedDataBackendAsyncPutWrapper.h
--- a/Engine/Source/Developer/DerivedDataCache/Private/DerivedDataBackendAsyncPutWrapper.h
+++ b/Engine/Source/Developer/DerivedDataCache/Private/DerivedDataBackendAsyncPutWrapper.h
@@ -45,16 +45,18 @@ public:
 	FThreadSet*							FilesInFlight;
 	/**If true, then do not attempt skip the put even if CachedDataProbablyExists returns true **/
 	bool								bPutEvenIfExists;
+	FCacheStatRecord*					Record;
 
 	/** Constructor
 	*/
-	FCachePutAsyncWorker(const TCHAR* InCacheKey, const TArray<uint8>* InData, FDerivedDataBackendInterface* InInnerBackend, bool InbPutEvenIfExists, FDerivedDataBackendInterface* InInflightCache, FThreadSet* InInFilesInFlight)
+	FCachePutAsyncWorker(const TCHAR* InCacheKey, const TArray<uint8>* InData, FDerivedDataBackendInterface* InInnerBackend, bool InbPutEvenIfExists, FDerivedDataBackendInterface* InInflightCache, FThreadSet* InInFilesInFlight, FCacheStatRecord* Stats)
 		: CacheKey(InCacheKey)
 		, Data(*InData)
 		, InnerBackend(InInnerBackend)
 		, InflightCache(InInflightCache)
 		, FilesInFlight(InInFilesInFlight)
 		, bPutEvenIfExists(InbPutEvenIfExists)
+		, Record(Stats)
 	{
 		check(InnerBackend);
 	}
@@ -63,13 +65,13 @@ public:
 	void DoWork()
 	{
 		bool bOk = true;
-		InnerBackend->PutCachedData(*CacheKey, Data, bPutEvenIfExists);
+		InnerBackend->PutCachedData(*CacheKey, Data, bPutEvenIfExists, Record);
 		if (InflightCache)
 		{
 			if (!InnerBackend->CachedDataProbablyExists(*CacheKey))
 			{
 				// retry
-				InnerBackend->PutCachedData(*CacheKey, Data, false);
+				InnerBackend->PutCachedData(*CacheKey, Data, false, Record);
 				if (!InnerBackend->CachedDataProbablyExists(*CacheKey))
 				{
 					UE_LOG(LogDerivedDataCache, Warning, TEXT("FDerivedDataBackendAsyncPutWrapper: Put failed, keeping in memory copy %s."),*CacheKey);
@@ -152,13 +154,13 @@ public:
 	 * @param	OutData		Buffer to receive the results, if any were found
 	 * @return				true if any data was found, and in this case OutData is non-empty
 	 */
-	virtual bool GetCachedData(const TCHAR* CacheKey, TArray<uint8>& OutData) override
+	virtual bool GetCachedData(const TCHAR* CacheKey, TArray<uint8>& OutData, FCacheStatRecord* Stats) override
 	{
-		if (InflightCache && InflightCache->GetCachedData(CacheKey, OutData))
+		if (InflightCache && InflightCache->GetCachedData(CacheKey, OutData, Stats))
 		{
 			return true;
 		}
-		bool bSuccess = InnerBackend->GetCachedData(CacheKey, OutData);
+		bool bSuccess = InnerBackend->GetCachedData(CacheKey, OutData, Stats);
 		return bSuccess;
 	}
 	/**
@@ -168,7 +170,7 @@ public:
 	 * @param	InData		Buffer containing the data to cache, can be destroyed after the call returns, immediately
 	 * @param	bPutEvenIfExists	If true, then do not attempt skip the put even if CachedDataProbablyExists returns true
 	 */
-	virtual void PutCachedData(const TCHAR* CacheKey, TArray<uint8>& InData, bool bPutEvenIfExists) override
+	virtual void PutCachedData(const TCHAR* CacheKey, TArray<uint8>& InData, bool bPutEvenIfExists, FCacheStatRecord* Stats) override
 	{
 		if (!InnerBackend->IsWritable())
 		{
@@ -188,7 +190,7 @@ public:
 		}
 		FilesInFlight.Add(CacheKey);
 		FDerivedDataBackend::Get().AddToAsyncCompletionCounter(1);
-		(new FAutoDeleteAsyncTask<FCachePutAsyncWorker>(CacheKey, &InData, InnerBackend, bPutEvenIfExists, InflightCache.GetOwnedPointer(), &FilesInFlight))->StartBackgroundTask();
+		(new FAutoDeleteAsyncTask<FCachePutAsyncWorker>(CacheKey, &InData, InnerBackend, bPutEvenIfExists, InflightCache.GetOwnedPointer(), &FilesInFlight, Stats))->StartBackgroundTask();
 	}
 
 	virtual void RemoveCachedData(const TCHAR* CacheKey, bool bTransient) override
diff --git a/Engine/Source/Developer/DerivedDataCache/Private/DerivedDataBackendCorruptionWrapper.h b/Engine/Source/Developer/DerivedDataCache/Private/DerivedDataBackendCorruptionWrapper.h
--- a/Engine/Source/Developer/DerivedDataCache/Private/DerivedDataBackendCorruptionWrapper.h
+++ b/Engine/Source/Developer/DerivedDataCache/Private/DerivedDataBackendCorruptionWrapper.h
@@ -99,9 +99,9 @@ public:
 	 * @param	OutData		Buffer to receive the results, if any were found
 	 * @return				true if any data was found, and in this case OutData is non-empty
 	 */
-	virtual bool GetCachedData(const TCHAR* CacheKey, TArray<uint8>& OutData) override
+	virtual bool GetCachedData(const TCHAR* CacheKey, TArray<uint8>& OutData, FCacheStatRecord* Stats) override
 	{
-		bool bOk = InnerBackend->GetCachedData(CacheKey, OutData);
+		bool bOk = InnerBackend->GetCachedData(CacheKey, OutData, Stats);
 		if (bOk)
 		{
 			if (OutData.Num() < sizeof(FDerivedDataTrailer))
@@ -144,7 +144,7 @@ public:
 	 * @param	InData		Buffer containing the data to cache, can be destroyed after the call returns, immediately
 	 * @param	bPutEvenIfExists	If true, then do not attempt skip the put even if CachedDataProbablyExists returns true
 	 */
-	virtual void PutCachedData(const TCHAR* CacheKey, TArray<uint8>& InData, bool bPutEvenIfExists) override
+	virtual void PutCachedData(const TCHAR* CacheKey, TArray<uint8>& InData, bool bPutEvenIfExists, FCacheStatRecord* Stats) override
 	{
 		if (!InnerBackend->IsWritable())
 		{
@@ -159,7 +159,7 @@ public:
 		FDerivedDataTrailer Trailer(Data);
 		Data.AddUninitialized(sizeof(FDerivedDataTrailer));
 		FMemory::Memcpy(&Data[Data.Num() - sizeof(FDerivedDataTrailer)], &Trailer, sizeof(FDerivedDataTrailer));
-		InnerBackend->PutCachedData(CacheKey, Data, bPutEvenIfExists);
+		InnerBackend->PutCachedData(CacheKey, Data, bPutEvenIfExists, Stats);
 	}
 	
 	virtual void RemoveCachedData(const TCHAR* CacheKey, bool bTransient) override
diff --git a/Engine/Source/Developer/DerivedDataCache/Private/DerivedDataBackendVerifyWrapper.h b/Engine/Source/Developer/DerivedDataCache/Private/DerivedDataBackendVerifyWrapper.h
--- a/Engine/Source/Developer/DerivedDataCache/Private/DerivedDataBackendVerifyWrapper.h
+++ b/Engine/Source/Developer/DerivedDataCache/Private/DerivedDataBackendVerifyWrapper.h
@@ -35,7 +35,7 @@ public:
 		}
 		return false;
 	}
-	virtual bool GetCachedData(const TCHAR* CacheKey, TArray<uint8>& OutData) override
+	virtual bool GetCachedData(const TCHAR* CacheKey, TArray<uint8>& OutData, FCacheStatRecord* Stats) override
 	{
 		bool bAlreadyTested = false;
 		{
@@ -47,11 +47,11 @@ public:
 		}
 		if (bAlreadyTested)
 		{
-			return InnerBackend->GetCachedData(CacheKey, OutData);
+			return InnerBackend->GetCachedData(CacheKey, OutData, Stats);
 		}
 		return false;
 	}
-	virtual void PutCachedData(const TCHAR* CacheKey, TArray<uint8>& InData, bool bPutEvenIfExists) override
+	virtual void PutCachedData(const TCHAR* CacheKey, TArray<uint8>& InData, bool bPutEvenIfExists, FCacheStatRecord* Stats) override
 	{
 		bool bAlreadyTested = false;
 		{
@@ -68,7 +68,7 @@ public:
 		if (!bAlreadyTested)
 		{
 			TArray<uint8> OutData;
-			bool bSuccess = InnerBackend->GetCachedData(CacheKey, OutData);
+			bool bSuccess = InnerBackend->GetCachedData(CacheKey, OutData, Stats);
 			if (bSuccess)
 			{
 				if (OutData != InData)
@@ -81,7 +81,7 @@ public:
 					if (bFixProblems)
 					{
 						UE_LOG(LogDerivedDataCache, Display, TEXT("Verify: Wrote newly generated data to the cache."), CacheKey);
-						InnerBackend->PutCachedData(CacheKey, InData, true);
+						InnerBackend->PutCachedData(CacheKey, InData, true, Stats);
 					}
 				}
 				else
@@ -92,7 +92,7 @@ public:
 			else
 			{
 				UE_LOG(LogDerivedDataCache, Warning, TEXT("Verify: Cached data didn't exist %s."), CacheKey);
-				InnerBackend->PutCachedData(CacheKey, InData, bPutEvenIfExists);
+				InnerBackend->PutCachedData(CacheKey, InData, bPutEvenIfExists, Stats);
 			}
 		}
 	}
diff --git a/Engine/Source/Developer/DerivedDataCache/Private/DerivedDataCache.cpp b/Engine/Source/Developer/DerivedDataCache/Private/DerivedDataCache.cpp
--- a/Engine/Source/Developer/DerivedDataCache/Private/DerivedDataCache.cpp
+++ b/Engine/Source/Developer/DerivedDataCache/Private/DerivedDataCache.cpp
@@ -39,11 +39,13 @@ class FDerivedDataCache : public FDerivedDataCacheInterface
 		 * @param	InDataDeriver	plugin to produce cache key and in the event of a miss, return the data.
 		 * @param	InCacheKey		Complete cache key for this data.
 		**/
-		FBuildAsyncWorker(FDerivedDataPluginInterface* InDataDeriver, const TCHAR* InCacheKey, bool bInSynchronousForStats)
+		FBuildAsyncWorker(FDerivedDataPluginInterface* InDataDeriver, const TCHAR* InCacheKey, bool bInSynchronousForStats, FCacheStatRecord* InGetRecord, FCacheStatRecord* InPutRecord)
 		: bSuccess(false)
 		, bSynchronousForStats(bInSynchronousForStats)
 		, DataDeriver(InDataDeriver)
 		, CacheKey(InCacheKey)
+		, GetRecord(InGetRecord)
+		, PutRecord(InPutRecord)
 		{
 		}
 		
@@ -56,7 +58,15 @@ class FDerivedDataCache : public FDerivedDataCacheInterface
 				STAT(double ThisTime = 0);
 				{
 					SCOPE_SECONDS_COUNTER(ThisTime);
-					bGetResult = FDerivedDataBackend::Get().GetRoot().GetCachedData(*CacheKey, Data);
+					GetRecord->CacheKey = CacheKey;
+					GetRecord->CacheGet = 1;
+					GetRecord->bSynchronous = bSynchronousForStats;
+					GetRecord->StartTime = FPlatformTime::Seconds();
+					GetRecord->GetDuration = 0.;
+					GetRecord->PutDuration = 0.;
+					bGetResult = FDerivedDataBackend::Get().GetRoot().GetCachedData(*CacheKey, Data, GetRecord);
+					GetRecord->EndTime = FPlatformTime::Seconds();
+					GetRecord->DataSize = Data.Num();
 				}
 				INC_FLOAT_STAT_BY(STAT_DDC_SyncGetTime, bSynchronousForStats ? (float)ThisTime : 0.0f);
 			}
@@ -87,7 +97,15 @@ class FDerivedDataCache : public FDerivedDataCacheInterface
 					STAT(double ThisTime = 0);
 					{
 						SCOPE_SECONDS_COUNTER(ThisTime);
-						FDerivedDataBackend::Get().GetRoot().PutCachedData(*CacheKey, Data, true);
+						PutRecord->CacheKey = CacheKey;
+						PutRecord->CacheGet = 0;
+						PutRecord->bSynchronous = bSynchronousForStats;
+						PutRecord->StartTime = FPlatformTime::Seconds();
+						PutRecord->GetDuration = 0.;
+						PutRecord->PutDuration = 0.;
+						FDerivedDataBackend::Get().GetRoot().PutCachedData(*CacheKey, Data, true, PutRecord);
+						PutRecord->EndTime = FPlatformTime::Seconds();
+						PutRecord->DataSize = Data.Num();
 					}
 					INC_FLOAT_STAT_BY(STAT_DDC_PutTime, bSynchronousForStats ? (float)ThisTime : 0.0f);
 				}
@@ -114,6 +132,8 @@ class FDerivedDataCache : public FDerivedDataCacheInterface
 		FString							CacheKey;
 		/** Data to return to caller, later **/
 		TArray<uint8>					Data;
+		FCacheStatRecord*				GetRecord;
+		FCacheStatRecord*				PutRecord;
 	};
 
 public:
@@ -143,7 +163,12 @@ public:
 		check(DataDeriver);
 		FString CacheKey = FDerivedDataCache::BuildCacheKey(DataDeriver);
 		UE_LOG(LogDerivedDataCache, Verbose, TEXT("GetSynchronous %s"), *CacheKey);
-		FAsyncTask<FBuildAsyncWorker> PendingTask(DataDeriver, *CacheKey, true);
+		FCacheStatRecord* GetRecord = new FCacheStatRecord();
+		FCacheStatRecord* PutRecord = new FCacheStatRecord();
+		PutRecord->CacheGet = -1;
+		CacheStats.Add(GetRecord);
+		CacheStats.Add(PutRecord);
+		FAsyncTask<FBuildAsyncWorker> PendingTask(DataDeriver, *CacheKey, true, GetRecord, PutRecord);
 		AddToAsyncCompletionCounter(1);
 		PendingTask.StartSynchronousTask();
 		OutData = PendingTask.GetTask().Data;
@@ -158,7 +183,12 @@ public:
 		FString CacheKey = FDerivedDataCache::BuildCacheKey(DataDeriver);
 		UE_LOG(LogDerivedDataCache, Verbose, TEXT("GetAsynchronous %s"), *CacheKey);
 		bool bSync = !DataDeriver->IsBuildThreadsafe();
-		FAsyncTask<FBuildAsyncWorker>* AsyncTask = new FAsyncTask<FBuildAsyncWorker>(DataDeriver, *CacheKey, bSync);
+		FCacheStatRecord* GetRecord = new FCacheStatRecord();
+		FCacheStatRecord* PutRecord = new FCacheStatRecord();
+		PutRecord->CacheGet = -1;
+		CacheStats.Add(GetRecord);
+		CacheStats.Add(PutRecord);
+		FAsyncTask<FBuildAsyncWorker>* AsyncTask = new FAsyncTask<FBuildAsyncWorker>(DataDeriver, *CacheKey, bSync, GetRecord, PutRecord);
 		check(!PendingTasks.Contains(Handle));
 		PendingTasks.Add(Handle,AsyncTask);
 		AddToAsyncCompletionCounter(1);
@@ -238,7 +268,12 @@ public:
 	{
 		QUICK_SCOPE_CYCLE_COUNTER(STAT_DDC_GetSynchronous_Data);
 		UE_LOG(LogDerivedDataCache, Verbose, TEXT("GetSynchronous %s"), CacheKey);
-		FAsyncTask<FBuildAsyncWorker> PendingTask((FDerivedDataPluginInterface*)NULL, CacheKey, true);
+		FCacheStatRecord* GetRecord = new FCacheStatRecord();
+		FCacheStatRecord* PutRecord = new FCacheStatRecord();
+		PutRecord->CacheGet = -1;
+		CacheStats.Add(GetRecord);
+		CacheStats.Add(PutRecord);
+		FAsyncTask<FBuildAsyncWorker> PendingTask((FDerivedDataPluginInterface*)NULL, CacheKey, true, GetRecord, PutRecord);
 		AddToAsyncCompletionCounter(1);
 		PendingTask.StartSynchronousTask();
 		OutData = PendingTask.GetTask().Data;
@@ -252,7 +287,12 @@ public:
 		FScopeLock ScopeLock(&SynchronizationObject);
 		UE_LOG(LogDerivedDataCache, Verbose, TEXT("GetAsynchronous %s"), CacheKey);
 		uint32 Handle = NextHandle();
-		FAsyncTask<FBuildAsyncWorker>* AsyncTask = new FAsyncTask<FBuildAsyncWorker>((FDerivedDataPluginInterface*)NULL, CacheKey, false);
+		FCacheStatRecord* GetRecord = new FCacheStatRecord();
+		FCacheStatRecord* PutRecord = new FCacheStatRecord();
+		PutRecord->CacheGet = -1;
+		CacheStats.Add(GetRecord);
+		CacheStats.Add(PutRecord);
+		FAsyncTask<FBuildAsyncWorker>* AsyncTask = new FAsyncTask<FBuildAsyncWorker>((FDerivedDataPluginInterface*)NULL, CacheKey, false, GetRecord, PutRecord);
 		check(!PendingTasks.Contains(Handle));
 		PendingTasks.Add(Handle, AsyncTask);
 		AddToAsyncCompletionCounter(1);
@@ -271,7 +311,12 @@ public:
 		QUICK_SCOPE_CYCLE_COUNTER(STAT_DDC_GetAsynchronousForRollup);
 		FScopeLock ScopeLock(&SynchronizationObject);
 		UE_LOG(LogDerivedDataCache, Verbose, TEXT("GetAsynchronous(handle) %s"), CacheKey);
-		FAsyncTask<FBuildAsyncWorker>* AsyncTask = new FAsyncTask<FBuildAsyncWorker>((FDerivedDataPluginInterface*)NULL, CacheKey, false);
+		FCacheStatRecord* GetRecord = new FCacheStatRecord();
+		FCacheStatRecord* PutRecord = new FCacheStatRecord();
+		PutRecord->CacheGet = -1;
+		CacheStats.Add(GetRecord);
+		CacheStats.Add(PutRecord);
+		FAsyncTask<FBuildAsyncWorker>* AsyncTask = new FAsyncTask<FBuildAsyncWorker>((FDerivedDataPluginInterface*)NULL, CacheKey, false, GetRecord, PutRecord);
 		check(!PendingTasks.Contains(Handle));
 		PendingTasks.Add(Handle,AsyncTask);
 		AddToAsyncCompletionCounter(1);
@@ -284,7 +329,17 @@ public:
 		STAT(double ThisTime = 0);
 		{
 			SCOPE_SECONDS_COUNTER(ThisTime);
+			FCacheStatRecord* Record = new FCacheStatRecord();
+			CacheStats.Add(Record);
+			Record->CacheKey = CacheKey;
+			Record->CacheGet = 0;
+			Record->bSynchronous = true;
+			Record->DataSize = Data.Num();
+			Record->StartTime = FPlatformTime::Seconds();
+			Record->PutDuration = 0;
+			Record->GetDuration = 0;
 			FDerivedDataBackend::Get().GetRoot().PutCachedData(CacheKey, Data, bPutEvenIfExists);
+			Record->EndTime = FPlatformTime::Seconds();
 		}
 		INC_FLOAT_STAT_BY(STAT_DDC_PutTime,(float)ThisTime);
 		INC_DWORD_STAT(STAT_DDC_NumPuts);
@@ -337,6 +392,41 @@ public:
 		// Used by derived classes to spit out leaked pending rollups
 	}
 
+	virtual void PrintStats()
+	{
+		if (CacheStats.Num() > 0)
+		{
+#if ALLOW_DEBUG_FILES
+			const FString OutputDir = FPaths::GameLogDir() / TEXT("DDCStats");
+			IFileManager::Get().MakeDirectory(*OutputDir, true);
+
+			// Create archive for log data.
+			const FString ChartName = OutputDir / FString(TEXT("CookRunDDC.csv"));
+			FArchive* OutputFile = IFileManager::Get().CreateDebugFileWriter(*ChartName);
+
+			if (OutputFile)
+			{
+				OutputFile->Logf(TEXT("Key,Size,Total Duration(ms),Get Duration(ms),Put Duration(ms),Get,Sync,Network"));
+				for (int32 Index = 0; Index < CacheStats.Num(); Index++)
+				{
+					FCacheStatRecord* Record = CacheStats[Index];
+					if (Record->CacheGet != -1)
+					{
+						OutputFile->Logf(TEXT("%s,%d,%lf,%lf,%lf,%d,%s,%s"), *(Record->CacheKey), Record->DataSize, (Record->EndTime - Record->StartTime) * 1000.0f, Record->GetDuration * 1000.0f, Record->PutDuration * 1000.0f, Record->CacheGet, Record->bSynchronous ? TEXT("true") : TEXT("false"), Record->bFromNetwork ? TEXT("true") : TEXT("false"));
+					}
+				}
+				delete OutputFile;
+			}
+#endif
+			for (int32 Index = 0; Index < CacheStats.Num(); Index++)
+			{
+				FCacheStatRecord* Record = CacheStats[Index];
+				delete Record;
+			}
+			CacheStats.Empty();
+		}
+	}
+
 protected:
 	uint32 NextHandle()
 	{
@@ -364,6 +454,9 @@ private:
 	FCriticalSection			SynchronizationObject;
 	/** Map of handle to pending task **/
 	TMap<uint32,FAsyncTask<FBuildAsyncWorker>*>	PendingTasks;
+
+	// TEMP: stats
+	TArray<FCacheStatRecord*> CacheStats;
 };
 
 //Forward reference
@@ -934,6 +1027,7 @@ public:
 
 		FDerivedDataCache& DDC = static_cast< FDerivedDataCache& >( GetDDC() );
 		DDC.PrintLeaks();
+		DDC.PrintStats();
 	}
 };
 
diff --git a/Engine/Source/Developer/DerivedDataCache/Private/DerivedDataLimitKeyLengthWrapper.h b/Engine/Source/Developer/DerivedDataCache/Private/DerivedDataLimitKeyLengthWrapper.h
--- a/Engine/Source/Developer/DerivedDataCache/Private/DerivedDataLimitKeyLengthWrapper.h
+++ b/Engine/Source/Developer/DerivedDataCache/Private/DerivedDataLimitKeyLengthWrapper.h
@@ -46,14 +46,14 @@ public:
 	 * @param	OutData		Buffer to receive the results, if any were found
 	 * @return				true if any data was found, and in this case OutData is non-empty
 	 */
-	virtual bool GetCachedData(const TCHAR* CacheKey, TArray<uint8>& OutData) override
+	virtual bool GetCachedData(const TCHAR* CacheKey, TArray<uint8>& OutData, FCacheStatRecord* Stats) override
 	{
 		FString NewKey;
 		bool bOk;
 		if (!ShortenKey(CacheKey, NewKey))
 		{
 			// no shortening needed
-			bOk = InnerBackend->GetCachedData(CacheKey, OutData);
+			bOk = InnerBackend->GetCachedData(CacheKey, OutData, Stats);
 			// look for old bug
 			if (FString(CacheKey).StartsWith(TEXT("TEXTURE2D_0002")))
 			{
@@ -71,7 +71,7 @@ public:
 		}
 		else
 		{
-			bOk = InnerBackend->GetCachedData(*NewKey, OutData);
+			bOk = InnerBackend->GetCachedData(*NewKey, OutData, Stats);
 			if (bOk)
 			{
 				int32 KeyLen = FCString::Strlen(CacheKey) + 1;
@@ -114,7 +114,7 @@ public:
 	 * @param	InData		Buffer containing the data to cache, can be destroyed after the call returns, immediately
 	 * @param	bPutEvenIfExists	If true, then do not attempt skip the put even if CachedDataProbablyExists returns true
 	 */
-	virtual void PutCachedData(const TCHAR* CacheKey, TArray<uint8>& InData, bool bPutEvenIfExists) override
+	virtual void PutCachedData(const TCHAR* CacheKey, TArray<uint8>& InData, bool bPutEvenIfExists, FCacheStatRecord* Stats) override
 	{
 		if (!InnerBackend->IsWritable())
 		{
@@ -124,7 +124,7 @@ public:
 		if (!ShortenKey(CacheKey, NewKey))
 		{
 			// no shortening needed
-			InnerBackend->PutCachedData(CacheKey, InData, bPutEvenIfExists);
+			InnerBackend->PutCachedData(CacheKey, InData, bPutEvenIfExists, Stats);
 			return;
 		}
 		TArray<uint8> Data(InData);
@@ -133,7 +133,7 @@ public:
 		Data.AddUninitialized(KeyLen);
 		FCStringAnsi::Strcpy((char*)&Data[Data.Num() - KeyLen], KeyLen, TCHAR_TO_ANSI(CacheKey));
 		check(Data.Last()==0);
-		InnerBackend->PutCachedData(*NewKey, Data, bPutEvenIfExists);
+		InnerBackend->PutCachedData(*NewKey, Data, bPutEvenIfExists, Stats);
 	}
 
 	virtual void RemoveCachedData(const TCHAR* CacheKey, bool bTransient) override
diff --git a/Engine/Source/Developer/DerivedDataCache/Private/FileSystemDerivedDataBackend.cpp b/Engine/Source/Developer/DerivedDataCache/Private/FileSystemDerivedDataBackend.cpp
--- a/Engine/Source/Developer/DerivedDataCache/Private/FileSystemDerivedDataBackend.cpp
+++ b/Engine/Source/Developer/DerivedDataCache/Private/FileSystemDerivedDataBackend.cpp
@@ -162,7 +162,7 @@ public:
 	 * @param	OutData		Buffer to receive the results, if any were found
 	 * @return				true if any data was found, and in this case OutData is non-empty
 	 */
-	virtual bool GetCachedData(const TCHAR* CacheKey,TArray<uint8>& Data) override
+	virtual bool GetCachedData(const TCHAR* CacheKey, TArray<uint8>& Data, FCacheStatRecord* Stats) override
 	{
 		static FName NAME_GetCachedData(TEXT("GetCachedData"));
 		FDDCScopeStatHelper Stat(CacheKey, NAME_GetCachedData);
@@ -172,7 +172,18 @@ public:
 		
 		check(!bFailed);
 		FString Filename = BuildFilename(CacheKey);
-
+		if (Stats)
+		{
+			Stats->CacheKey = Filename;
+			if (Filename.StartsWith(TEXT("//")) || Filename.StartsWith(TEXT("\\\\")))
+			{
+				Stats->bFromNetwork = true;
+			}
+			else
+			{
+				Stats->bFromNetwork = false;
+			}
+		}
 		double StartTime = FPlatformTime::Seconds();
 		if (FFileHelper::LoadFileToArray(Data,*Filename,FILEREAD_Silent))
 		{
@@ -183,12 +194,20 @@ public:
 
 			UE_LOG(LogDerivedDataCache, Verbose, TEXT("FileSystemDerivedDataBackend: Cache hit on %s"),*Filename);
 			Stat.AddTag(NAME_Retrieved, TEXT("true"));
+			if (Stats)
+			{
+				Stats->GetDuration += FPlatformTime::Seconds() - StartTime;
+			}
 			return true;
 		}
 		UE_LOG(LogDerivedDataCache, Verbose, TEXT("FileSystemDerivedDataBackend: Cache miss on %s"),*Filename);
 		Data.Empty();
 		Stat.AddTag(NAME_Retrieved, TEXT("false"));
-		return false;	
+		if (Stats)
+		{
+			Stats->GetDuration += FPlatformTime::Seconds() - StartTime;
+		}
+		return false;
 	}
 	/**
 	 * Asynchronous, fire-and-forget placement of a cache item
@@ -197,7 +216,7 @@ public:
 	 * @param	OutData		Buffer containing the data to cache, can be destroyed after the call returns, immediately
 	 * @param	bPutEvenIfExists	If true, then do not attempt skip the put even if CachedDataProbablyExists returns true
 	 */
-	virtual void PutCachedData(const TCHAR* CacheKey, TArray<uint8>& Data, bool bPutEvenIfExists) override
+	virtual void PutCachedData(const TCHAR* CacheKey, TArray<uint8>& Data, bool bPutEvenIfExists, FCacheStatRecord* Stats) override
 	{
 		check(!bFailed);
 		if (!bReadOnly)
@@ -206,6 +225,19 @@ public:
 			{
 				check(Data.Num());
 				FString Filename = BuildFilename(CacheKey);
+				double StartTime = FPlatformTime::Seconds();
+				if (Stats)
+				{
+					Stats->CacheKey = Filename;
+					if (Filename.StartsWith(TEXT("//")) || Filename.StartsWith(TEXT("\\\\")))
+					{
+						Stats->bFromNetwork = true;
+					}
+					else
+					{
+						Stats->bFromNetwork = false;
+					}
+				}
 				FString TempFilename(TEXT("temp.")); 
 				TempFilename += FGuid::NewGuid().ToString();
 				TempFilename = FPaths::GetPath(Filename) / TempFilename;
@@ -244,6 +276,10 @@ public:
 				{
 					UE_LOG(LogDerivedDataCache, Warning, TEXT("FFileSystemDerivedDataBackend: Could not write temp file %s!"),*TempFilename);
 				}
+				if (Stats)
+				{
+					Stats->PutDuration += FPlatformTime::Seconds() - StartTime;
+				}
 				// if everything worked, this is not necessary, but we will make every effort to avoid leaving junk in the cache
 				if (FPaths::FileExists(TempFilename))
 				{
diff --git a/Engine/Source/Developer/DerivedDataCache/Private/HierarchicalDerivedDataBackend.h b/Engine/Source/Developer/DerivedDataCache/Private/HierarchicalDerivedDataBackend.h
--- a/Engine/Source/Developer/DerivedDataCache/Private/HierarchicalDerivedDataBackend.h
+++ b/Engine/Source/Developer/DerivedDataCache/Private/HierarchicalDerivedDataBackend.h
@@ -86,7 +86,7 @@ public:
 	 * @param	OutData		Buffer to receive the results, if any were found
 	 * @return				true if any data was found, and in this case OutData is non-empty
 	 */
-	virtual bool GetCachedData(const TCHAR* CacheKey, TArray<uint8>& OutData) override
+	virtual bool GetCachedData(const TCHAR* CacheKey, TArray<uint8>& OutData, FCacheStatRecord* Stats) override
 	{
 		const static FName NAME_GetCachedData = FName(TEXT("GetCachedData"));
 		const static FName NAME_HierarchicalDDC = FName(TEXT("HierarchicalDDC"));
@@ -98,7 +98,7 @@ public:
 
 		for (int32 CacheIndex = 0; CacheIndex < InnerBackends.Num(); CacheIndex++)
 		{
-			if (InnerBackends[CacheIndex]->CachedDataProbablyExists(CacheKey) && InnerBackends[CacheIndex]->GetCachedData(CacheKey, OutData))
+			if (InnerBackends[CacheIndex]->CachedDataProbablyExists(CacheKey) && InnerBackends[CacheIndex]->GetCachedData(CacheKey, OutData, Stats))
 			{
 				if (bIsWritable)
 				{
@@ -111,11 +111,13 @@ public:
 								InnerBackends[PutCacheIndex]->CachedDataProbablyExists(CacheKey))
 							{
 								InnerBackends[PutCacheIndex]->RemoveCachedData(CacheKey, /*bTransient=*/ false); // it apparently failed, so lets delete what is there
-								AsyncPutInnerBackends[PutCacheIndex]->PutCachedData(CacheKey, OutData, true); // we force a put here because it must have failed
+								AsyncPutInnerBackends[PutCacheIndex]->PutCachedData(CacheKey, OutData, true, Stats); // we force a put here because it must have failed
+								Stats->CacheGet++;
 							}
 							else
 							{
-								AsyncPutInnerBackends[PutCacheIndex]->PutCachedData(CacheKey, OutData, false); 
+								AsyncPutInnerBackends[PutCacheIndex]->PutCachedData(CacheKey, OutData, false, Stats); 
+								Stats->CacheGet++;
 							}
 						}
 					}
@@ -130,13 +132,15 @@ public:
 							}
 							if (InnerBackends[PutCacheIndex]->IsWritable())
 							{
-								AsyncPutInnerBackends[PutCacheIndex]->PutCachedData(CacheKey, OutData, false); // we do not need to force a put here
+								AsyncPutInnerBackends[PutCacheIndex]->PutCachedData(CacheKey, OutData, false, Stats); // we do not need to force a put here
+								Stats->CacheGet++;
 							}
 						}
 					}
 				}
 				Stat.AddTag(NAME_DataSize, FString::Printf(TEXT("%d bytes"), OutData.Num()));
 				Stat.AddTag(NAME_Retrieved, TEXT("true"));
+				Stats->EndTime = FPlatformTime::Seconds();
 				return true;
 			}
 		}
@@ -150,7 +154,7 @@ public:
 	 * @param	InData		Buffer containing the data to cache, can be destroyed after the call returns, immediately
 	 * @param	bPutEvenIfExists	If true, then do not attempt skip the put even if CachedDataProbablyExists returns true
 	 */
-	virtual void PutCachedData(const TCHAR* CacheKey, TArray<uint8>& InData, bool bPutEvenIfExists) override
+	virtual void PutCachedData(const TCHAR* CacheKey, TArray<uint8>& InData, bool bPutEvenIfExists, FCacheStatRecord* Stats) override
 	{
 		if (!bIsWritable)
 		{
@@ -167,12 +171,12 @@ public:
 			{
 				if (!bSynchronousPutPeformed)
 				{
-					InnerBackends[PutCacheIndex]->PutCachedData(CacheKey, InData, bPutEvenIfExists);
+					InnerBackends[PutCacheIndex]->PutCachedData(CacheKey, InData, bPutEvenIfExists, Stats);
 					bSynchronousPutPeformed = true;
 				}
 				else
 				{
-					AsyncPutInnerBackends[PutCacheIndex]->PutCachedData(CacheKey, InData, bPutEvenIfExists);
+					AsyncPutInnerBackends[PutCacheIndex]->PutCachedData(CacheKey, InData, bPutEvenIfExists, Stats);
 				}
 			}
 		}
diff --git a/Engine/Source/Developer/DerivedDataCache/Private/MemoryDerivedDataBackend.h b/Engine/Source/Developer/DerivedDataCache/Private/MemoryDerivedDataBackend.h
--- a/Engine/Source/Developer/DerivedDataCache/Private/MemoryDerivedDataBackend.h
+++ b/Engine/Source/Developer/DerivedDataCache/Private/MemoryDerivedDataBackend.h
@@ -56,7 +56,7 @@ public:
 	 * @param	OutData		Buffer to receive the results, if any were found
 	 * @return				true if any data was found, and in this case OutData is non-empty
 	 */
-	virtual bool GetCachedData(const TCHAR* CacheKey, TArray<uint8>& OutData) override
+	virtual bool GetCachedData(const TCHAR* CacheKey, TArray<uint8>& OutData, FCacheStatRecord* Stats) override
 	{
 		FScopeLock ScopeLock(&SynchronizationObject);
 		if (!bDisabled)
@@ -80,7 +80,7 @@ public:
 	 * @param	InData		Buffer containing the data to cache, can be destroyed after the call returns, immediately
 	 * @param	bPutEvenIfExists	If true, then do not attempt skip the put even if CachedDataProbablyExists returns true
 	 */
-	virtual void PutCachedData(const TCHAR* CacheKey, TArray<uint8>& InData, bool bPutEvenIfExists) override
+	virtual void PutCachedData(const TCHAR* CacheKey, TArray<uint8>& InData, bool bPutEvenIfExists, FCacheStatRecord* Stats) override
 	{
 		FScopeLock ScopeLock(&SynchronizationObject);
 		
diff --git a/Engine/Source/Developer/DerivedDataCache/Private/PakFileDerivedDataBackend.h b/Engine/Source/Developer/DerivedDataCache/Private/PakFileDerivedDataBackend.h
--- a/Engine/Source/Developer/DerivedDataCache/Private/PakFileDerivedDataBackend.h
+++ b/Engine/Source/Developer/DerivedDataCache/Private/PakFileDerivedDataBackend.h
@@ -98,13 +98,15 @@ public:
 	 * @param	OutData		Buffer to receive the results, if any were found
 	 * @return				true if any data was found, and in this case OutData is non-empty
 	 */
-	virtual bool GetCachedData(const TCHAR* CacheKey, TArray<uint8>& OutData) override
+	virtual bool GetCachedData(const TCHAR* CacheKey, TArray<uint8>& OutData, FCacheStatRecord* Stats) override
 	{
 		if (bWriting || bClosed)
 		{
 			return false;
 		}
 		FScopeLock ScopeLock(&SynchronizationObject);
+		if (Stats)
+			Stats->bFromNetwork = false;
 		FCacheValue* Item = CacheItems.Find(FString(CacheKey));
 		if (Item)
 		{
@@ -148,7 +150,7 @@ public:
 	 * @param	InData		Buffer containing the data to cache, can be destroyed after the call returns, immediately
 	 * @param	bPutEvenIfExists	If true, then do not attempt skip the put even if CachedDataProbablyExists returns true
 	 */
-	virtual void PutCachedData(const TCHAR* CacheKey, TArray<uint8>& InData, bool bPutEvenIfExists) override
+	virtual void PutCachedData(const TCHAR* CacheKey, TArray<uint8>& InData, bool bPutEvenIfExists, FCacheStatRecord* Stats) override
 	{
 		if (!bWriting || bClosed)
 		{
@@ -156,6 +158,8 @@ public:
 		}
 		{
 			FScopeLock ScopeLock(&SynchronizationObject);
+			if (Stats)
+				Stats->bFromNetwork = false;
 			FString Key(CacheKey);
 			FCacheValue* Item = CacheItems.Find(FString(CacheKey));
 			if (!Item)
@@ -376,9 +380,9 @@ public:
 		for(const FString& CopyKeyName : CopyKeyNames)
 		{
 			Buffer.Reset();
-			if(OtherPak->FPakFileDerivedDataBackend::GetCachedData(*CopyKeyName, Buffer))
+			if(OtherPak->FPakFileDerivedDataBackend::GetCachedData(*CopyKeyName, Buffer, NULL))
 			{
-				FPakFileDerivedDataBackend::PutCachedData(*CopyKeyName, Buffer, false);
+				FPakFileDerivedDataBackend::PutCachedData(*CopyKeyName, Buffer, false, NULL);
 			}
 		}
 	}
@@ -408,8 +412,8 @@ public:
 		for (int KeyIndex = 0; KeyIndex < KeyNames.Num(); KeyIndex++)
 		{
 			Buffer.Reset();
-			InputPak.GetCachedData(*KeyNames[KeyIndex], Buffer);
-			OutputPak.PutCachedData(*KeyNames[KeyIndex], Buffer, false);
+			InputPak.GetCachedData(*KeyNames[KeyIndex], Buffer, NULL);
+			OutputPak.PutCachedData(*KeyNames[KeyIndex], Buffer, false, NULL);
 			KeySizes.Add(Buffer.Num());
 		}
 
@@ -466,7 +470,7 @@ public:
 	{
 	}
 
-	virtual void PutCachedData(const TCHAR* CacheKey, TArray<uint8>& InData, bool bPutEvenIfExists) override
+	virtual void PutCachedData(const TCHAR* CacheKey, TArray<uint8>& InData, bool bPutEvenIfExists, FCacheStatRecord* Stats) override
 	{
 		int32 UncompressedSize = InData.Num();
 		int32 CompressedSize = FCompression::CompressMemoryBound(CompressionFlags, UncompressedSize);
@@ -478,13 +482,13 @@ public:
 		verify(FCompression::CompressMemory(CompressionFlags, CompressedData.GetData() + sizeof(UncompressedSize), CompressedSize, InData.GetData(), InData.Num()));
 		CompressedData.SetNum(CompressedSize + sizeof(UncompressedSize), false);
 
-		FPakFileDerivedDataBackend::PutCachedData(CacheKey, CompressedData, bPutEvenIfExists);
+		FPakFileDerivedDataBackend::PutCachedData(CacheKey, CompressedData, bPutEvenIfExists, Stats);
 	}
 
-	virtual bool GetCachedData(const TCHAR* CacheKey, TArray<uint8>& OutData) override
+	virtual bool GetCachedData(const TCHAR* CacheKey, TArray<uint8>& OutData, FCacheStatRecord* Stats) override
 	{
 		TArray<uint8> CompressedData;
-		if(!FPakFileDerivedDataBackend::GetCachedData(CacheKey, CompressedData))
+		if(!FPakFileDerivedDataBackend::GetCachedData(CacheKey, CompressedData, Stats))
 		{
 			return false;
 		}
diff --git a/Engine/Source/Developer/DerivedDataCache/Public/DerivedDataBackendInterface.h b/Engine/Source/Developer/DerivedDataCache/Public/DerivedDataBackendInterface.h
--- a/Engine/Source/Developer/DerivedDataCache/Public/DerivedDataBackendInterface.h
+++ b/Engine/Source/Developer/DerivedDataCache/Public/DerivedDataBackendInterface.h
@@ -14,6 +14,20 @@ DECLARE_FLOAT_ACCUMULATOR_STAT_EXTERN(TEXT("Sync Put Time"),STAT_DDC_PutTime,STA
 DECLARE_FLOAT_ACCUMULATOR_STAT_EXTERN(TEXT("Sync Build Time"),STAT_DDC_SyncBuildTime,STATGROUP_DDC, );
 DECLARE_FLOAT_ACCUMULATOR_STAT_EXTERN(TEXT("Exists Time"),STAT_DDC_ExistTime,STATGROUP_DDC, );
 
+struct FCacheStatRecord
+{
+	FString CacheKey;
+	uint32 DataSize;
+	double StartTime;
+	double EndTime;
+	double GetDuration;
+	double PutDuration;
+	uint32 CacheGet;
+	bool bSynchronous;
+	bool bFromNetwork;
+};
+
+
 /** 
  * Interface for cache server backends. 
  * The entire API should be callable from any thread (except the singleton can be assumed to be called at least once before concurrent access).
@@ -51,7 +65,7 @@ public:
 	 * @param	OutData		Buffer to receive the results, if any were found
 	 * @return				true if any data was found, and in this case OutData is non-empty
 	 */
-	virtual bool GetCachedData(const TCHAR* CacheKey, TArray<uint8>& OutData)=0;
+	virtual bool GetCachedData(const TCHAR* CacheKey, TArray<uint8>& OutData, FCacheStatRecord* Stats = NULL)=0;
 	/**
 	 * Asynchronous, fire-and-forget placement of a cache item
 	 *
@@ -59,7 +73,7 @@ public:
 	 * @param	InData				Buffer containing the data to cache, can be destroyed after the call returns, immediately
 	 * @param	bPutEvenIfExists	If true, then do not attempt skip the put even if CachedDataProbablyExists returns true
 	 */
-	virtual void PutCachedData(const TCHAR* CacheKey, TArray<uint8>& InData, bool bPutEvenIfExists)=0;
+	virtual void PutCachedData(const TCHAR* CacheKey, TArray<uint8>& InData, bool bPutEvenIfExists, FCacheStatRecord* Stats = NULL) = 0;
 
 	/**
 	 * Remove data from cache (used in the event that corruption is detected at a higher level and possibly house keeping)
diff --git a/Engine/Source/Runtime/Core/Private/Misc/CoreMisc.cpp b/Engine/Source/Runtime/Core/Private/Misc/CoreMisc.cpp
--- a/Engine/Source/Runtime/Core/Private/Misc/CoreMisc.cpp
+++ b/Engine/Source/Runtime/Core/Private/Misc/CoreMisc.cpp
@@ -738,6 +738,15 @@ class FDerivedDataCacheInterface* GetDerivedDataCache()
 	return SingletonInterface;
 }
 
+void DerivedDataCachePrint()
+{
+	class IDerivedDataCacheModule* Module = FModuleManager::LoadModulePtr<IDerivedDataCacheModule>("DerivedDataCache");
+	if (Module)
+	{
+		Module->ShutdownModule();
+	}
+}
+
 class FDerivedDataCacheInterface& GetDerivedDataCacheRef()
 {
 	class FDerivedDataCacheInterface* SingletonInterface = GetDerivedDataCache();
diff --git a/Engine/Source/Runtime/Core/Public/Misc/CoreMisc.h b/Engine/Source/Runtime/Core/Public/Misc/CoreMisc.h
--- a/Engine/Source/Runtime/Core/Public/Misc/CoreMisc.h
+++ b/Engine/Source/Runtime/Core/Public/Misc/CoreMisc.h
@@ -241,6 +241,9 @@ CORE_API class FDerivedDataCacheInterface* GetDerivedDataCache();
 /** Return the DDC interface, fatal error if it is not available. **/
 CORE_API class FDerivedDataCacheInterface& GetDerivedDataCacheRef();
 
+/** Return the DDC interface, if it is available, otherwise return NULL **/
+CORE_API void DerivedDataCachePrint();
+
 /** Return the Target Platform Manager interface, if it is available, otherwise return NULL **/
 CORE_API class ITargetPlatformManagerModule* GetTargetPlatformManager();
 
diff --git a/Engine/Source/Runtime/Launch/Private/LaunchEngineLoop.cpp b/Engine/Source/Runtime/Launch/Private/LaunchEngineLoop.cpp
--- a/Engine/Source/Runtime/Launch/Private/LaunchEngineLoop.cpp
+++ b/Engine/Source/Runtime/Launch/Private/LaunchEngineLoop.cpp
@@ -2201,6 +2201,8 @@ void FEngineLoop::Exit()
 	}
 #endif
 
+	DerivedDataCachePrint();
+
 	// Unload all modules.  Note that this doesn't actually unload the module DLLs (that happens at
 	// process exit by the OS), but it does call ShutdownModule() on all loaded modules in the reverse
 	// order they were loaded in, so that systems can unregister and perform general clean up.
