potential fix for corrupted stat record array in the ddc

diff --git a/Engine/Source/Developer/DerivedDataCache/Private/DerivedDataCache.cpp b/Engine/Source/Developer/DerivedDataCache/Private/DerivedDataCache.cpp
--- a/Engine/Source/Developer/DerivedDataCache/Private/DerivedDataCache.cpp
+++ b/Engine/Source/Developer/DerivedDataCache/Private/DerivedDataCache.cpp
@@ -166,8 +166,11 @@ public:
 		FCacheStatRecord* GetRecord = new FCacheStatRecord();
 		FCacheStatRecord* PutRecord = new FCacheStatRecord();
 		PutRecord->CacheGet = -1;
-		CacheStats.Add(GetRecord);
-		CacheStats.Add(PutRecord);
+		{
+			FScopeLock ScopeLock(&SynchronizationObject);
+			CacheStats.Add(GetRecord);
+			CacheStats.Add(PutRecord);
+		}
 		FAsyncTask<FBuildAsyncWorker> PendingTask(DataDeriver, *CacheKey, true, GetRecord, PutRecord);
 		AddToAsyncCompletionCounter(1);
 		PendingTask.StartSynchronousTask();
@@ -271,8 +274,11 @@ public:
 		FCacheStatRecord* GetRecord = new FCacheStatRecord();
 		FCacheStatRecord* PutRecord = new FCacheStatRecord();
 		PutRecord->CacheGet = -1;
-		CacheStats.Add(GetRecord);
-		CacheStats.Add(PutRecord);
+		{
+			FScopeLock ScopeLock(&SynchronizationObject);
+			CacheStats.Add(GetRecord);
+			CacheStats.Add(PutRecord);
+		}
 		FAsyncTask<FBuildAsyncWorker> PendingTask((FDerivedDataPluginInterface*)NULL, CacheKey, true, GetRecord, PutRecord);
 		AddToAsyncCompletionCounter(1);
 		PendingTask.StartSynchronousTask();
@@ -330,7 +336,10 @@ public:
 		{
 			SCOPE_SECONDS_COUNTER(ThisTime);
 			FCacheStatRecord* Record = new FCacheStatRecord();
-			CacheStats.Add(Record);
+			{
+				FScopeLock ScopeLock(&SynchronizationObject);
+				CacheStats.Add(Record);
+			}
 			Record->CacheKey = CacheKey;
 			Record->CacheGet = 0;
 			Record->bSynchronous = true;
@@ -394,6 +403,7 @@ public:
 
 	virtual void PrintStats()
 	{
+		FScopeLock ScopeLock(&SynchronizationObject);
 		if (CacheStats.Num() > 0)
 		{
 #if ALLOW_DEBUG_FILES
