fix for location of the DDC stats log

diff --git a/Engine/Source/Developer/DerivedDataCache/Private/DerivedDataCache.cpp b/Engine/Source/Developer/DerivedDataCache/Private/DerivedDataCache.cpp
--- a/Engine/Source/Developer/DerivedDataCache/Private/DerivedDataCache.cpp
+++ b/Engine/Source/Developer/DerivedDataCache/Private/DerivedDataCache.cpp
@@ -407,11 +407,12 @@ public:
 		if (CacheStats.Num() > 0)
 		{
 #if ALLOW_DEBUG_FILES
-			const FString OutputDir = FPaths::GameLogDir() / TEXT("DDCStats");
+			const FString AbsLog = FPlatformOutputDevices::GetAbsoluteLogFilename();
+			const FString OutputDir = FPaths::GetPath(AbsLog);
 			IFileManager::Get().MakeDirectory(*OutputDir, true);
 
 			// Create archive for log data.
-			const FString ChartName = OutputDir / FString(TEXT("CookRunDDC.csv"));
+			const FString ChartName = OutputDir / FString(TEXT("CookRunDDC.txt"));
 			FArchive* OutputFile = IFileManager::Get().CreateDebugFileWriter(*ChartName);
 
 			if (OutputFile)
