Updating CIS Counter

diff --git a/Engine/Documentation/Source/Engine/Basics/DerivedDataCache/DerivedDataCache.INT.udn b/Engine/Documentation/Source/Engine/Basics/DerivedDataCache/DerivedDataCache.INT.udn
new file mode 100644
--- /dev/null
+++ b/Engine/Documentation/Source/Engine/Basics/DerivedDataCache/DerivedDataCache.INT.udn
@@ -0,0 +1,90 @@
+Availability:Public
+Title:Derived Data Cache
+Crumbs:%ROOT%, Engine
+Description: Overview of the Derived Data Cache and the storage of assets in formats used by UE4 and its target platforms. 
+
+The **Derived Data Cache** (DDC) stores versions of assets in the formats used by the engine and its 
+target platforms, as opposed to the source formats artists create that are imported into the editor 
+and stored in `.uasset` files. Content stored in the DDC is _disposable_ in that it can always be 
+regenerated at any time using the data stored in the `.uasset` file. Storing these _derived_ formats 
+externally makes it possible to easily add or change the formats used by the engine without needing to 
+modify the source asset file.
+
+## Using a Shared DDC
+
+Studios should use a shared DDC that all users in a particular location can access. This way, only one 
+person needs to build the derived asset format(s) and they are automatically available to all other 
+users. There will occasionally be stalls when assets need to be processed, but the results are remembered 
+and shared. So with a handful of developers or more, most people will not notice any impact.
+
+[REGION:warning]
+Never copy an entire DDC across the internet, back up your DDC, or restore a DDC from a backup. This will 
+not do any harm, but it is a waste of time as it generally takes longer to transfer the amount of data stored 
+in the DDC than it would to simply generate it from scratch locally.
+[/REGION]
+
+To set up a shared DDC, override the paths for the `[DerivedDataBackendGraph]` by declaring this 
+section in your game's `DefaultEngine.ini`. This section is originally declared in the `BaseEngine.ini` 
+where the paths are set to Epic's internal DDC share. Re-declare the section and change the paths to point 
+to a share on your network (e.g. `Path=\\mystudio.net\DDC`):
+
+	[DerivedDataBackendGraph]
+	MinimumDaysToKeepFile=7
+	Root=(Type=KeyLength, Length=120, Inner=AsyncPut)
+	AsyncPut=(Type=AsyncPut, Inner=Hierarchy)
+	Hierarchy=(Type=Hierarchical, Inner=Boot, Inner=Pak, Inner=EnginePak, Inner=Local, Inner=Shared)
+	Boot=(Type=Boot, Filename=\%GAMEDIR\%DerivedDataCache/Boot.ddc, MaxCacheSize=256)
+	Local=(Type=FileSystem, ReadOnly=false, Clean=false, Flush=false, PurgeTransient=true, DeleteUnused=true, UnusedFileAge=17, FoldersToClean=-1, Path=../../../Engine/DerivedDataCache)
+	Shared=(Type=FileSystem, ReadOnly=false, Clean=false, Flush=false, DeleteUnused=true, UnusedFileAge=19, FoldersToClean=-1, Path=\\mystudio.net\DDC, EnvPathOverride=UE-SharedDataCachePath)
+	AltShared=(Type=FileSystem, ReadOnly=true, Clean=false, Flush=false, DeleteUnused=true, UnusedFileAge=19, FoldersToClean=-1, Path=\\mystudio.net\DDC2, EnvPathOverride=UE-SharedDataCachePath2)
+	Pak=(Type=ReadPak, Filename=\%GAMEDIR\%DerivedDataCache/DDC.ddp)
+	EnginePak=(Type=ReadPak, Filename=../../../Engine/DerivedDataCache/DDC.ddp)
+
+
+Satellite studios working out of a single shared code base can set the `UE-SharedDataCachePath` environment 
+variable to a path that all users at each location can read and write to. This allows each location to 
+have its own shared DDC. 
+
+For example:
+
+	UE-SharedDataCachePath=\\mystudio.net\DDC
+
+## Using a Local DDC
+
+Lone developers or developers working remotely generally just use a local DDC. To tell the engine not to 
+try and use a shared DDC, use one of the methods below:
+
+* Pass `-ddc=noshared` on the command line.
+* Set the environment variable to a local drive hard drive: 
+    
+		UE-SharedDataCachePath=c:\DDC
+
+
+## Building Derived Data
+
+Generally, the user who imports the asset is the one who builds the derived data since they will most 
+likely be using and testing them out in the engine. However, there may be occasions when a new asset 
+needs to be processed. This happens automatically on an as-needed basis and should not result in much 
+of an impact when running fast hardware, though there may be occasional stalls.
+ 
+You can fill your derived data caches at any time by passing the arguments below to the executable:
+ 
+	UE4\Engine\Binaries\Win64\UE4Editor.exe MyGame -run=DerivedDataCache -fill 
+
+We do this on a nightly basis to ensure the DDC is always primed, but it is not necessary in general 
+as the automatic caching should _just work_.
+
+## Distributing with a DDC
+ 
+Cooking is generally the preferred method for packaging games since cooked builds do not need nor use a DDC. However, 
+DDCs can be packaged for distribution if the need arises.
+
+**To package a DDC:**
+ 
+1. Run the `UE4Editor.exe` from the `UE4/Engine/Binaries/Win64` directory passing it the arguments shown below:
+
+		UE4Editor.exe [GameName] -run=DerivedDataCache -fill -DDC=CreatePak
+ 
+1. This creates a `DDC.ddp` file in the `UE4\[GameName]\DerivedDataCache`
+
+1. The engine automatically detects and uses the `.ddp` file.
diff --git a/Engine/Source/Developer/DerivedDataCache/DerivedDataCache.Build.cs b/Engine/Source/Developer/DerivedDataCache/DerivedDataCache.Build.cs
new file mode 100644
--- /dev/null
+++ b/Engine/Source/Developer/DerivedDataCache/DerivedDataCache.Build.cs
@@ -0,0 +1,20 @@
+// Copyright 1998-2015 Epic Games, Inc. All Rights Reserved.
+
+using UnrealBuildTool;
+using System.IO;
+
+public class DerivedDataCache : ModuleRules
+{
+	public DerivedDataCache(TargetInfo Target)
+	{
+		PrivateDependencyModuleNames.Add("Core");
+        PrivateDependencyModuleNames.Add("CookingStats");
+
+		// Internal (NotForLicensees) module
+		var DDCUtilsModule = Path.Combine("Developer", "NotForLicensees", "DDCUtils", "DDCUtils.Build.cs");
+		if (File.Exists(DDCUtilsModule))
+		{
+			DynamicallyLoadedModuleNames.Add("DDCUtils");
+		}
+	}
+}
diff --git a/Engine/Source/Developer/DerivedDataCache/Private/DDCCleanup.cpp b/Engine/Source/Developer/DerivedDataCache/Private/DDCCleanup.cpp
new file mode 100644
--- /dev/null
+++ b/Engine/Source/Developer/DerivedDataCache/Private/DDCCleanup.cpp
@@ -0,0 +1,221 @@
+// Copyright 1998-2015 Epic Games, Inc. All Rights Reserved.
+
+#include "Core.h"
+
+#include "DerivedDataBackendInterface.h"
+#include "DDCCleanup.h"
+
+/** Struct containing a list of directories to cleanup. */
+struct FFilesystemInfo
+{
+	/** Filesystem Path to clean up. **/
+	const FString CachePath;
+	/** Minimum time a file has not been used for to delete it. */
+	const FTimespan UnusedFileTime;
+	/** The maximum number of folders to check. <= 0 means all. */
+	const int32 MaxNumFoldersToCheck;
+	/** The maximum number of files to check before pausing. <= 0 is no limit. */
+	const int32 MaxContinuousFileChecks;
+
+	/** The number of folders already checked. */
+	int32 FoldersChecked;
+	/** Filesystem directories left to clean up. */
+	TArray<int32> CacheDirectories;
+
+	FFilesystemInfo( FString& InCachePath, int32 InDaysToDelete, int32 InMaxNumFoldersToCheck, int32 InMaxContinuousFileChecks )
+		: CachePath( InCachePath )
+		, UnusedFileTime( InDaysToDelete, 0, 0, 0 )
+		, MaxNumFoldersToCheck( InMaxNumFoldersToCheck )
+		, MaxContinuousFileChecks( InMaxContinuousFileChecks )
+		, FoldersChecked( 0 )
+	{
+		CacheDirectories.Empty( 1000 );
+		for( int32 Index = 0; Index < 1000; Index++ )
+		{
+			CacheDirectories.Add( Index );
+		}
+			
+		// Initialize random stream using Cycles() 
+		const uint32 Cycles = FPlatformTime::Cycles();
+		FRandomStream RandomStream( *(int32*)&Cycles );
+		// Shuffle
+		for( int32 Index = CacheDirectories.Num() - 1; Index >= 0; Index-- )
+		{
+			const int32 RandomIndex = RandomStream.RandHelper(Index + 1);
+
+			Exchange( CacheDirectories[ Index ], CacheDirectories[ RandomIndex ] );
+		}
+	}
+};
+
+FDDCCleanup* FDDCCleanup::Runnable = NULL;
+
+FDDCCleanup::FDDCCleanup()
+	: Thread(NULL)
+	, StopTaskCounter(0)
+{
+	// Don't delete the runnable automatically. It's going to be manually deleted in FDDCCleanup::Shutdown.
+	Thread = FRunnableThread::Create(this, TEXT("FDDCCleanup"), 0, TPri_BelowNormal, FPlatformAffinity::GetPoolThreadMask());
+}
+
+FDDCCleanup::~FDDCCleanup()
+{
+	delete Thread;
+}
+
+void FDDCCleanup::Wait( const float InSeconds, const float InSleepTime )
+{
+	// Instead of waiting the given amount of seconds doing nothing
+	// check periodically if there's been any Stop requests.
+	for( float TimeToWait = InSeconds; TimeToWait > 0.0f && ShouldStop() == false; TimeToWait -= InSleepTime )
+	{
+		FPlatformProcess::Sleep( FMath::Min(InSleepTime, TimeToWait) );
+	}
+}
+
+bool FDDCCleanup::Init() 
+{
+	return true;
+}
+
+uint32 FDDCCleanup::Run()
+{
+	// Give up some time to the engine to start up and load everything
+	Wait( 120.0f, 0.5f );
+
+	int32 FilesystemToCleanup = 0;
+	// Check one directory every 5 seconds
+	do
+	{
+		// Pick one random filesystem.
+		TSharedPtr< FFilesystemInfo > FilesystemInfo;
+		{
+			FScopeLock ScopeLock( &DataLock );
+			if( CleanupList.Num() > 0 )
+			{
+				FilesystemToCleanup %= CleanupList.Num();
+				FilesystemInfo = CleanupList[ FilesystemToCleanup++ ];					
+			}
+		}
+
+		if( FilesystemInfo.IsValid() )
+		{
+			CleanupFilesystemDirectory( FilesystemInfo );
+		}
+		Wait( 5.0f );
+	}
+	while( ShouldStop() == false );
+
+	return 0;
+}
+
+bool FDDCCleanup::CleanupFilesystemDirectory( TSharedPtr< FFilesystemInfo > FilesystemInfo )
+{
+	bool bCleanedUp = false;
+
+	const double StartTime = FPlatformTime::Seconds();
+
+	// Pick one random directory.
+	TArray<FString> FileNames;
+	do
+	{
+		const int32 DirectoryIndex = FilesystemInfo->CacheDirectories.Pop();
+		const FString DirectoryPath( FilesystemInfo->CachePath / FString::Printf(TEXT("%1d/%1d/%1d/"),(DirectoryIndex/100)%10,(DirectoryIndex/10)%10,DirectoryIndex%10) );
+
+		IFileManager::Get().FindFilesRecursive( FileNames, *DirectoryPath, TEXT("*.*"), true, false );
+		if ( FilesystemInfo->CacheDirectories.Num() == 0 )
+		{
+			// Remove the filesystem and stop checking it
+			FScopeLock ScopeLock( &DataLock );
+			CleanupList.Remove( FilesystemInfo );
+			FilesystemInfo.Reset();
+		}
+		else if( ++FilesystemInfo->FoldersChecked >= FilesystemInfo->MaxNumFoldersToCheck && FilesystemInfo->MaxNumFoldersToCheck > 0 )
+		{
+			// Remove the filesystem but keep checking the current folder
+			FScopeLock ScopeLock( &DataLock );
+			CleanupList.Remove( FilesystemInfo );
+		}
+	}
+	while( FileNames.Num() == 0 && FilesystemInfo.IsValid() && ShouldStop() == false );
+
+	if( FilesystemInfo.IsValid() && ShouldStop() == false )
+	{
+		// Iterate over all files in the selected folder and check their last access time
+		int32 NumFilesChecked = 0;
+		for( int32 FileIndex = 0; FileIndex < FileNames.Num() && ShouldStop() == false; FileIndex++ )
+		{
+			const FDateTime LastModificationTime = IFileManager::Get().GetTimeStamp( *FileNames[ FileIndex ] );
+			const FDateTime LastAccessTime = IFileManager::Get().GetAccessTimeStamp( *FileNames[ FileIndex ] );
+			if( (LastAccessTime != FDateTime::MinValue()) || (LastModificationTime != FDateTime::MinValue())  )
+			{
+				const FTimespan TimeSinceLastAccess = FDateTime::UtcNow() - LastAccessTime;
+				const FTimespan TimeSinceLastModification = FDateTime::UtcNow() - LastModificationTime;
+				if( TimeSinceLastAccess >= FilesystemInfo->UnusedFileTime && TimeSinceLastModification >= FilesystemInfo->UnusedFileTime )
+				{
+					// Delete the file
+					bool Result = IFileManager::Get().Delete( *FileNames[ FileIndex ], false, true, true );
+				}
+			}
+
+			if( ++NumFilesChecked >= FilesystemInfo->MaxContinuousFileChecks && FilesystemInfo->MaxContinuousFileChecks > 0 && ShouldStop() == false )
+			{
+				NumFilesChecked = 0;
+				Wait( 1.0f );
+			}
+			else
+			{
+				// Give up a tiny amount of time so that we're not consuming too much cpu/hdd resources.
+				Wait( 0.05f );
+			}
+		}
+
+		bCleanedUp = true;			
+	}
+
+	UE_LOG(LogDerivedDataCache, VeryVerbose, TEXT("DDC Folder Cleanup (%s) took %.4lfs."), *FilesystemInfo->CachePath, FPlatformTime::Seconds() - StartTime);
+
+	return bCleanedUp;
+}
+
+void FDDCCleanup::Stop(void)
+{
+	StopTaskCounter.Increment();
+}
+
+void FDDCCleanup::Exit(void)
+{
+}
+
+void FDDCCleanup::EnsureCompletion()
+{
+	Stop();
+	Thread->WaitForCompletion();
+}
+
+void FDDCCleanup::AddFilesystem( FString& InCachePath, int32 InDaysToDelete, int32 InMaxNumFoldersToCheck, int32 InMaxContinuousFileChecks )
+{
+	FScopeLock Lock( &DataLock );
+	CleanupList.Add( MakeShareable( new FFilesystemInfo( InCachePath, InDaysToDelete, InMaxNumFoldersToCheck, InMaxContinuousFileChecks ) ) );
+}
+
+FDDCCleanup* FDDCCleanup::Get()
+{
+	if (!Runnable && FPlatformProcess::SupportsMultithreading())
+	{
+		Runnable = new FDDCCleanup();			
+	}
+	return Runnable;
+}
+
+void FDDCCleanup::Shutdown()
+{
+	if (Runnable)
+	{
+		Runnable->EnsureCompletion();
+		delete Runnable;
+		Runnable = NULL;
+	}
+}
+
+
diff --git a/Engine/Source/Developer/DerivedDataCache/Private/DDCCleanup.h b/Engine/Source/Developer/DerivedDataCache/Private/DDCCleanup.h
new file mode 100644
--- /dev/null
+++ b/Engine/Source/Developer/DerivedDataCache/Private/DDCCleanup.h
@@ -0,0 +1,76 @@
+// Copyright 1998-2015 Epic Games, Inc. All Rights Reserved.
+
+#pragma once
+
+/** 
+ * DDC Filesystem Cache cleanup thread.
+ */
+class FDDCCleanup : public FRunnable
+{
+	/** Singleton instance */
+	static FDDCCleanup* Runnable;
+	
+	/** Thread to run the cleanup FRunnable on */
+	FRunnableThread* Thread;
+	/** > 0 if we've been asked to abort work in progress at the next opportunity */
+	FThreadSafeCounter StopTaskCounter;
+	/** List of filesystems to clean up */
+	TArray< TSharedPtr< struct FFilesystemInfo > > CleanupList;
+	/** Synchronization object */
+	FCriticalSection DataLock;
+
+	/** Constructor */
+	FDDCCleanup();
+
+	/** Destructor */
+	virtual ~FDDCCleanup();
+
+	// Begin FRunnable interface.
+	virtual bool Init();
+	virtual uint32 Run();
+	virtual void Stop();
+	virtual void Exit();
+	// End FRunnable interface
+
+	/** Checks if there's been any Stop requests */
+	FORCEINLINE bool ShouldStop() const
+	{
+		return StopTaskCounter.GetValue() > 0;
+	}
+
+	/**
+	 * Waits for a given amount of time periodically checking if there's been any Stop requests.
+	 *
+	 * @param InSeconds time in seconds to wait.
+	 * @param InSleepTime interval at which to check for Stop requests.
+	 */
+	void Wait( const float InSeconds, const float InSleepTime = 0.1f );
+
+	/**
+	 * Performs directory cleanup.
+	 *
+	 * @param FilesystemInfo Filesystem to cleanup.
+	 * @return true if the filesystem contained any directories to clean up, false otherwise.
+	 */
+	bool CleanupFilesystemDirectory( TSharedPtr< FFilesystemInfo > FilesystemInfo );
+
+	/** Makes sure this thread has stopped properly */
+	void EnsureCompletion();
+
+public:
+
+	/**
+	 * Adds DDC filesystem to clean up.
+	 *
+	 * @param InCachePath Filesystem path.
+	 * @param InDaysToDelete Number of days since last access time to consider a file as unused.
+	 * @param InMaxContinuousFileChecks Number of files to check without pausing.
+	 */
+	void AddFilesystem( FString& InCachePath, int32 InDaysToDelete, int32 InMaxNumFoldersToCheck, int32 InMaxContinuousFileChecks );
+
+	/** Gets DDC Cleanup singleton instance */
+	static FDDCCleanup* Get();
+
+	/** Shuts down DDC Cleanup thread. */
+	static void Shutdown();
+};
diff --git a/Engine/Source/Developer/DerivedDataCache/Private/DerivedDataBackendAsyncPutWrapper.h b/Engine/Source/Developer/DerivedDataCache/Private/DerivedDataBackendAsyncPutWrapper.h
new file mode 100644
--- /dev/null
+++ b/Engine/Source/Developer/DerivedDataCache/Private/DerivedDataBackendAsyncPutWrapper.h
@@ -0,0 +1,223 @@
+// Copyright 1998-2015 Epic Games, Inc. All Rights Reserved.
+
+#pragma once
+
+/** 
+ * Thread safe set helper
+**/
+struct FThreadSet
+{
+	FCriticalSection	SynchronizationObject;
+	TSet<FString>		FilesInFlight;
+
+	void Add(const FString& Key)
+	{
+		FScopeLock ScopeLock(&SynchronizationObject);
+		check(Key.Len());
+		FilesInFlight.Add(Key);
+	}
+	void Remove(const FString& Key)
+	{
+		FScopeLock ScopeLock(&SynchronizationObject);
+		FilesInFlight.Remove(Key);
+	}
+	bool Exists(const FString& Key)
+	{
+		FScopeLock ScopeLock(&SynchronizationObject);
+		return FilesInFlight.Contains(Key);
+	}
+};
+/** 
+	* Async task to handle the fire and forget async put
+**/
+class FCachePutAsyncWorker
+{
+public:
+	/** Cache Key for the put to InnerBackend **/
+	FString								CacheKey;
+	/** Data for the put to InnerBackend **/
+	TArray<uint8>						Data;
+	/** Backend to use for storage, my responsibilities are about async puts **/
+	FDerivedDataBackendInterface*		InnerBackend;
+	/** Memory based cache to clear once the put is finished **/
+	FDerivedDataBackendInterface*		InflightCache;
+	/** We remember outstanding puts so that we don't do them redundantly **/
+	FThreadSet*							FilesInFlight;
+	/**If true, then do not attempt skip the put even if CachedDataProbablyExists returns true **/
+	bool								bPutEvenIfExists;
+	FCacheStatRecord*					Record;
+
+	/** Constructor
+	*/
+	FCachePutAsyncWorker(const TCHAR* InCacheKey, const TArray<uint8>* InData, FDerivedDataBackendInterface* InInnerBackend, bool InbPutEvenIfExists, FDerivedDataBackendInterface* InInflightCache, FThreadSet* InInFilesInFlight, FCacheStatRecord* Stats)
+		: CacheKey(InCacheKey)
+		, Data(*InData)
+		, InnerBackend(InInnerBackend)
+		, InflightCache(InInflightCache)
+		, FilesInFlight(InInFilesInFlight)
+		, bPutEvenIfExists(InbPutEvenIfExists)
+		, Record(Stats)
+	{
+		check(InnerBackend);
+	}
+		
+	/** Call the inner backend and when that completes, remove the memory cache */
+	void DoWork()
+	{
+		bool bOk = true;
+		InnerBackend->PutCachedData(*CacheKey, Data, bPutEvenIfExists, Record);
+		if (InflightCache)
+		{
+			if (!InnerBackend->CachedDataProbablyExists(*CacheKey))
+			{
+				// retry
+				InnerBackend->PutCachedData(*CacheKey, Data, false, Record);
+				if (!InnerBackend->CachedDataProbablyExists(*CacheKey))
+				{
+					UE_LOG(LogDerivedDataCache, Warning, TEXT("FDerivedDataBackendAsyncPutWrapper: Put failed, keeping in memory copy %s."),*CacheKey);
+					bOk = false;
+				}
+			}
+		}
+		if (bOk && InflightCache)
+		{
+			InflightCache->RemoveCachedData(*CacheKey, /*bTransient=*/ false); // we can remove this from the temp cache, since the real cache will hit now
+		}
+		FilesInFlight->Remove(CacheKey);
+		FDerivedDataBackend::Get().AddToAsyncCompletionCounter(-1);
+	}
+
+	FORCEINLINE TStatId GetStatId() const
+	{
+		RETURN_QUICK_DECLARE_CYCLE_STAT(FCachePutAsyncWorker, STATGROUP_ThreadPoolAsyncTasks);
+	}
+
+	/** Indicates to the thread pool that this task is abandonable */
+	bool CanAbandon()
+	{
+		return true;
+	}
+
+	/** Abandon routine, we need to remove the item from the in flight cache because something might be waiting for that */
+	void Abandon()
+	{
+		if (InflightCache)
+		{
+			InflightCache->RemoveCachedData(*CacheKey, /*bTransient=*/ false); // we can remove this from the temp cache, since the real cache will hit now
+		}
+		FilesInFlight->Remove(CacheKey);
+		FDerivedDataBackend::Get().AddToAsyncCompletionCounter(-1);
+	}
+};
+
+
+/** 
+ * A backend wrapper that coordinates async puts. This means that a Get will hit an in-memory cache while the async put is still in flight.
+**/
+class FDerivedDataBackendAsyncPutWrapper : public FDerivedDataBackendInterface
+{
+public:
+
+	/**
+	 * Constructor
+	 *
+	 * @param	InInnerBackend		Backend to use for storage, my responsibilities are about async puts
+	 * @param	bCacheInFlightPuts	if true, cache in-flight puts in a memory cache so that they hit immediately
+	 */
+	FDerivedDataBackendAsyncPutWrapper(FDerivedDataBackendInterface* InInnerBackend, bool bCacheInFlightPuts)
+		: InnerBackend(InInnerBackend)
+		, InflightCache(bCacheInFlightPuts ? (new FMemoryDerivedDataBackend) : NULL)
+	{
+		check(InnerBackend);
+	}
+
+	/** return true if this cache is writable **/
+	virtual bool IsWritable() override
+	{
+		return InnerBackend->IsWritable();
+	}
+
+	/**
+	 * Synchronous test for the existence of a cache item
+	 *
+	 * @param	CacheKey	Alphanumeric+underscore key of this cache item
+	 * @return				true if the data probably will be found, this can't be guaranteed because of concurrency in the backends, corruption, etc
+	 */
+	virtual bool CachedDataProbablyExists(const TCHAR* CacheKey) override
+	{
+		return (InflightCache && InflightCache->CachedDataProbablyExists(CacheKey)) || InnerBackend->CachedDataProbablyExists(CacheKey);
+	}
+	/**
+	 * Synchronous retrieve of a cache item
+	 *
+	 * @param	CacheKey	Alphanumeric+underscore key of this cache item
+	 * @param	OutData		Buffer to receive the results, if any were found
+	 * @return				true if any data was found, and in this case OutData is non-empty
+	 */
+	virtual bool GetCachedData(const TCHAR* CacheKey, TArray<uint8>& OutData, FCacheStatRecord* Stats) override
+	{
+		if (InflightCache && InflightCache->GetCachedData(CacheKey, OutData, Stats))
+		{
+			return true;
+		}
+		bool bSuccess = InnerBackend->GetCachedData(CacheKey, OutData, Stats);
+		return bSuccess;
+	}
+	/**
+	 * Asynchronous, fire-and-forget placement of a cache item
+	 *
+	 * @param	CacheKey	Alphanumeric+underscore key of this cache item
+	 * @param	InData		Buffer containing the data to cache, can be destroyed after the call returns, immediately
+	 * @param	bPutEvenIfExists	If true, then do not attempt skip the put even if CachedDataProbablyExists returns true
+	 */
+	virtual void PutCachedData(const TCHAR* CacheKey, TArray<uint8>& InData, bool bPutEvenIfExists, FCacheStatRecord* Stats) override
+	{
+		if (!InnerBackend->IsWritable())
+		{
+			return; // no point in continuing down the chain
+		}
+		if (FilesInFlight.Exists(CacheKey))
+		{
+			return; // if it is already on its way, we don't need to send it again
+		}
+		if (InflightCache)
+		{
+			if (InflightCache->CachedDataProbablyExists(CacheKey))
+			{
+				return; // if it is already on its way, we don't need to send it again
+			}
+			InflightCache->PutCachedData(CacheKey, InData, true); // temp copy stored in memory while the async task waits to complete
+		}
+		FilesInFlight.Add(CacheKey);
+		FDerivedDataBackend::Get().AddToAsyncCompletionCounter(1);
+		(new FAutoDeleteAsyncTask<FCachePutAsyncWorker>(CacheKey, &InData, InnerBackend, bPutEvenIfExists, InflightCache.GetOwnedPointer(), &FilesInFlight, Stats))->StartBackgroundTask();
+	}
+
+	virtual void RemoveCachedData(const TCHAR* CacheKey, bool bTransient) override
+	{
+		if (!InnerBackend->IsWritable())
+		{
+			return; // no point in continuing down the chain
+		}
+		while (FilesInFlight.Exists(CacheKey))
+		{
+			FPlatformProcess::Sleep(0.0f); // this is an exception condition (corruption), spin and wait for it to clear
+		}
+		if (InflightCache)
+		{
+			InflightCache->RemoveCachedData(CacheKey, bTransient);
+		}
+		InnerBackend->RemoveCachedData(CacheKey, bTransient);
+	}
+private:
+
+	/** Backend to use for storage, my responsibilities are about async puts **/
+	FDerivedDataBackendInterface*					InnerBackend;
+	/** Memory based cache to deal with gets that happen while an async put is still in flight **/
+	TScopedPointer<FDerivedDataBackendInterface>	InflightCache;
+	/** We remember outstanding puts so that we don't do them redundantly **/
+	FThreadSet										FilesInFlight;
+};
+
+
+
diff --git a/Engine/Source/Developer/DerivedDataCache/Private/DerivedDataBackendCorruptionWrapper.h b/Engine/Source/Developer/DerivedDataCache/Private/DerivedDataBackendCorruptionWrapper.h
new file mode 100644
--- /dev/null
+++ b/Engine/Source/Developer/DerivedDataCache/Private/DerivedDataBackendCorruptionWrapper.h
@@ -0,0 +1,178 @@
+// Copyright 1998-2015 Epic Games, Inc. All Rights Reserved.
+
+#pragma once
+
+/** 
+ * Helper class for placing a footer at the end of of a cache file.
+ * No effort is made to byte-swap this as we assume local format.
+**/
+struct FDerivedDataTrailer
+{
+	enum
+	{
+		/** Arbitrary number used to identify corruption **/
+		MagicConstant = 0x1e873d89
+	};
+
+	/** Arbitrary number used to identify corruption **/
+	uint32 Magic;
+	/** Version of the backend, for future use **/
+	uint32 Version;
+	/** CRC of the payload, used to detect corruption **/
+	uint32 CRCofPayload;
+	/** Size of the payload, used to detect corruption **/
+	uint32 SizeOfPayload;
+
+	/** Constructor that zeros everything. This is never valid because even the magic number is wrong **/
+	FDerivedDataTrailer() 
+		: Magic(0)
+		, Version(0)
+		, CRCofPayload(0)
+		, SizeOfPayload(0)
+	{
+	}
+
+	/**
+		* Constructor that sets the data properly for a given buffer of data
+		* @param	Data	Buffer of data to setup this trailer for
+		*/
+	FDerivedDataTrailer(const TArray<uint8>& Data)
+		: Magic(MagicConstant)
+		, Version(1)
+		, CRCofPayload(FCrc::MemCrc_DEPRECATED(Data.GetData(), Data.Num()))
+		, SizeOfPayload(Data.Num())
+	{
+	}
+	/**
+		* Comparison operator
+		* @param	Other	trailer to compare "this" to
+		* @return			true if all elements of the trailer match exactly
+		*/
+	bool operator==(const FDerivedDataTrailer& Other)
+	{
+		return Magic == Other.Magic &&
+			Version == Other.Version &&
+			CRCofPayload == Other.CRCofPayload &&
+			SizeOfPayload == Other.SizeOfPayload;
+	}
+};
+
+
+/** 
+ * A backend wrapper that adds a footer to the data to check the CRC, length, etc.
+**/
+class FDerivedDataBackendCorruptionWrapper : public FDerivedDataBackendInterface
+{
+public:
+
+	/**
+	 * Constructor
+	 *
+	 * @param	InInnerBackend	Backend to use for storage, my responsibilities are about corruption
+	 */
+	FDerivedDataBackendCorruptionWrapper(FDerivedDataBackendInterface* InInnerBackend)
+		: InnerBackend(InInnerBackend)
+	{
+		check(InnerBackend);
+	}
+
+	/** return true if this cache is writable **/
+	virtual bool IsWritable() override
+	{
+		return InnerBackend->IsWritable();
+	}
+
+	/**
+	 * Synchronous test for the existence of a cache item
+	 *
+	 * @param	CacheKey	Alphanumeric+underscore key of this cache item
+	 * @return				true if the data probably will be found, this can't be guaranteed because of concurrency in the backends, corruption, etc
+	 */
+	virtual bool CachedDataProbablyExists(const TCHAR* CacheKey) override
+	{
+		return InnerBackend->CachedDataProbablyExists(CacheKey);
+	}
+	/**
+	 * Synchronous retrieve of a cache item
+	 *
+	 * @param	CacheKey	Alphanumeric+underscore key of this cache item
+	 * @param	OutData		Buffer to receive the results, if any were found
+	 * @return				true if any data was found, and in this case OutData is non-empty
+	 */
+	virtual bool GetCachedData(const TCHAR* CacheKey, TArray<uint8>& OutData, FCacheStatRecord* Stats) override
+	{
+		bool bOk = InnerBackend->GetCachedData(CacheKey, OutData, Stats);
+		if (bOk)
+		{
+			if (OutData.Num() < sizeof(FDerivedDataTrailer))
+			{
+				UE_LOG(LogDerivedDataCache, Warning, TEXT("FDerivedDataBackendCorruptionWrapper: Corrupted file (short), ignoring and deleting %s."),CacheKey);
+				bOk	= false;
+			}
+			else
+			{
+				FDerivedDataTrailer Trailer;
+				FMemory::Memcpy(&Trailer,&OutData[OutData.Num() - sizeof(FDerivedDataTrailer)], sizeof(FDerivedDataTrailer));
+				OutData.RemoveAt(OutData.Num() - sizeof(FDerivedDataTrailer),sizeof(FDerivedDataTrailer));
+				FDerivedDataTrailer RecomputedTrailer(OutData);
+				if (Trailer == RecomputedTrailer)
+				{
+					UE_LOG(LogDerivedDataCache, Verbose, TEXT("FDerivedDataBackendCorruptionWrapper: cache hit, footer is ok %s"),CacheKey);
+				}
+				else
+				{
+					UE_LOG(LogDerivedDataCache, Warning, TEXT("FDerivedDataBackendCorruptionWrapper: Corrupted file, ignoring and deleting %s."),CacheKey);
+					bOk	= false;
+				}
+			}
+			if (!bOk)
+			{
+				// _we_ detected corruption, so _we_ will force a flush of the corrupted data
+				InnerBackend->RemoveCachedData(CacheKey, /*bTransient=*/ false);
+			}
+		}
+		if (!bOk)
+		{
+			OutData.Empty();
+		}
+		return bOk;
+	}
+	/**
+	 * Asynchronous, fire-and-forget placement of a cache item
+	 *
+	 * @param	CacheKey	Alphanumeric+underscore key of this cache item
+	 * @param	InData		Buffer containing the data to cache, can be destroyed after the call returns, immediately
+	 * @param	bPutEvenIfExists	If true, then do not attempt skip the put even if CachedDataProbablyExists returns true
+	 */
+	virtual void PutCachedData(const TCHAR* CacheKey, TArray<uint8>& InData, bool bPutEvenIfExists, FCacheStatRecord* Stats) override
+	{
+		if (!InnerBackend->IsWritable())
+		{
+			return; // no point in continuing down the chain
+		}
+
+		// Get rid of the double copy!
+		TArray<uint8> Data;
+		Data.Reset( InData.Num() + sizeof(FDerivedDataTrailer) );
+		Data.AddUninitialized(InData.Num());
+		FMemory::Memcpy(&Data[0], &InData[0], InData.Num());
+		FDerivedDataTrailer Trailer(Data);
+		Data.AddUninitialized(sizeof(FDerivedDataTrailer));
+		FMemory::Memcpy(&Data[Data.Num() - sizeof(FDerivedDataTrailer)], &Trailer, sizeof(FDerivedDataTrailer));
+		InnerBackend->PutCachedData(CacheKey, Data, bPutEvenIfExists, Stats);
+	}
+	
+	virtual void RemoveCachedData(const TCHAR* CacheKey, bool bTransient) override
+	{
+		if (!InnerBackend->IsWritable())
+		{
+			return; // no point in continuing down the chain
+		}
+		return InnerBackend->RemoveCachedData(CacheKey, bTransient);
+	}
+private:
+
+	/** Backend to use for storage, my responsibilities are about corruption **/
+	FDerivedDataBackendInterface* InnerBackend;
+};
+
diff --git a/Engine/Source/Developer/DerivedDataCache/Private/DerivedDataBackendVerifyWrapper.h b/Engine/Source/Developer/DerivedDataCache/Private/DerivedDataBackendVerifyWrapper.h
new file mode 100644
--- /dev/null
+++ b/Engine/Source/Developer/DerivedDataCache/Private/DerivedDataBackendVerifyWrapper.h
@@ -0,0 +1,112 @@
+// Copyright 1998-2015 Epic Games, Inc. All Rights Reserved.
+
+#pragma once
+/**
+ * FDerivedDataBackendVerifyWrapper, a wrapper for derived data that verifies the cache is bit-wise identical by failing all gets and verifying the puts
+**/
+class FDerivedDataBackendVerifyWrapper : public FDerivedDataBackendInterface
+{
+public:
+
+	/**
+	 * Constructor
+	 *
+	 * @param	InInnerBackend		Backend to use for storage, my responsibilities are about async puts
+	 * @param	InbFixProblems		if true, fix any problems found
+	 */
+	FDerivedDataBackendVerifyWrapper(FDerivedDataBackendInterface* InInnerBackend, bool InbFixProblems)
+		: bFixProblems(InbFixProblems)
+		, InnerBackend(InInnerBackend)
+	{
+		check(InnerBackend);
+	}
+
+	virtual bool IsWritable() override
+	{
+		return true;
+	}
+
+	virtual bool CachedDataProbablyExists(const TCHAR* CacheKey) override
+	{
+		FScopeLock ScopeLock(&SynchronizationObject);
+		if (AlreadyTested.Contains(FString(CacheKey)))
+		{
+			return true;
+		}
+		return false;
+	}
+	virtual bool GetCachedData(const TCHAR* CacheKey, TArray<uint8>& OutData, FCacheStatRecord* Stats) override
+	{
+		bool bAlreadyTested = false;
+		{
+			FScopeLock ScopeLock(&SynchronizationObject);
+			if (AlreadyTested.Contains(FString(CacheKey)))
+			{
+				bAlreadyTested = true;
+			}
+		}
+		if (bAlreadyTested)
+		{
+			return InnerBackend->GetCachedData(CacheKey, OutData, Stats);
+		}
+		return false;
+	}
+	virtual void PutCachedData(const TCHAR* CacheKey, TArray<uint8>& InData, bool bPutEvenIfExists, FCacheStatRecord* Stats) override
+	{
+		bool bAlreadyTested = false;
+		{
+			FScopeLock ScopeLock(&SynchronizationObject);
+			if (AlreadyTested.Contains(FString(CacheKey)))
+			{
+				bAlreadyTested = true;
+			}
+			else
+			{
+				AlreadyTested.Add(FString(CacheKey));
+			}
+		}
+		if (!bAlreadyTested)
+		{
+			TArray<uint8> OutData;
+			bool bSuccess = InnerBackend->GetCachedData(CacheKey, OutData, Stats);
+			if (bSuccess)
+			{
+				if (OutData != InData)
+				{
+					UE_LOG(LogDerivedDataCache, Error, TEXT("Verify: Cached data differs from newly generated data %s."), CacheKey);
+					FString Cache = FString(FPaths::GameSavedDir()) / TEXT("VerifyDDC") / CacheKey + TEXT(".fromcache");
+					FFileHelper::SaveArrayToFile(OutData, *Cache);
+					FString Verify = FString(FPaths::GameSavedDir()) / TEXT("VerifyDDC") / CacheKey + TEXT(".verify");;
+					FFileHelper::SaveArrayToFile(InData, *Verify);
+					if (bFixProblems)
+					{
+						UE_LOG(LogDerivedDataCache, Display, TEXT("Verify: Wrote newly generated data to the cache."), CacheKey);
+						InnerBackend->PutCachedData(CacheKey, InData, true, Stats);
+					}
+				}
+				else
+				{
+					UE_LOG(LogDerivedDataCache, Log, TEXT("Verify: Cached data exists and matched %s."), CacheKey);
+				}
+			}
+			else
+			{
+				UE_LOG(LogDerivedDataCache, Warning, TEXT("Verify: Cached data didn't exist %s."), CacheKey);
+				InnerBackend->PutCachedData(CacheKey, InData, bPutEvenIfExists, Stats);
+			}
+		}
+	}
+	virtual void RemoveCachedData(const TCHAR* CacheKey, bool bTransient) override
+	{
+		InnerBackend->RemoveCachedData(CacheKey, bTransient);
+	}
+private:
+	/** If problems are encountered, do we fix them?						*/
+	bool											bFixProblems;
+	/** Object used for synchronization via a scoped lock						*/
+	FCriticalSection								SynchronizationObject;
+	/** Set of cache keys we already tested **/
+	TSet<FString>									AlreadyTested;
+	/** Backend to service the actual requests */
+	FDerivedDataBackendInterface*					InnerBackend;
+};
diff --git a/Engine/Source/Developer/DerivedDataCache/Private/DerivedDataBackends.cpp b/Engine/Source/Developer/DerivedDataCache/Private/DerivedDataBackends.cpp
new file mode 100644
--- /dev/null
+++ b/Engine/Source/Developer/DerivedDataCache/Private/DerivedDataBackends.cpp
@@ -0,0 +1,868 @@
+// Copyright 1998-2015 Epic Games, Inc. All Rights Reserved.
+
+#include "Core.h"
+#include "SecureHash.h"
+#include "DerivedDataBackendInterface.h"
+#include "DerivedDataCacheInterface.h"
+#include "DDCStatsHelper.h"
+#include "MemoryDerivedDataBackend.h"
+#include "DerivedDataBackendAsyncPutWrapper.h"
+#include "PakFileDerivedDataBackend.h"
+#include "HierarchicalDerivedDataBackend.h"
+#include "DerivedDataLimitKeyLengthWrapper.h"
+#include "DerivedDataBackendCorruptionWrapper.h"
+#include "DerivedDataBackendVerifyWrapper.h"
+#include "DerivedDataUtilsInterface.h"
+#include "Misc/EngineBuildSettings.h"
+
+DEFINE_LOG_CATEGORY(LogDerivedDataCache);
+
+#define MAX_BACKEND_KEY_LENGTH (120)
+#define LOCTEXT_NAMESPACE "DerivedDataBackendGraph"
+
+FDerivedDataBackendInterface* CreateFileSystemDerivedDataBackend(const TCHAR* CacheDirectory, bool bForceReadOnly = false, bool bTouchFiles = false, bool bPurgeTransient = false, bool bDeleteOldFiles = false, int32 InDaysToDeleteUnusedFiles = 60, int32 InMaxNumFoldersToCheck = -1, int32 InMaxContinuousFileChecks = -1);
+
+/**
+  * This class is used to create a singleton that represents the derived data cache hierarchy and all of the wrappers necessary
+  * ideally this would be data driven and the backends would be plugins...
+**/
+class FDerivedDataBackendGraph : public FDerivedDataBackend
+{
+public:
+	/**
+	* constructor, builds the cache tree
+	**/
+	FDerivedDataBackendGraph()
+		: RootCache(NULL)
+		, BootCache(NULL)
+		, WritePakCache(NULL)
+		, AsyncPutWrapper(NULL)
+		, KeyLengthWrapper(NULL)
+		, HierarchicalWrapper(NULL)
+		, MountPakCommand(
+		TEXT( "DDC.MountPak" ),
+		*LOCTEXT("CommandText_DDCMountPak", "Mounts read-only pak file").ToString(),
+		FConsoleCommandWithArgsDelegate::CreateRaw( this, &FDerivedDataBackendGraph::MountPakCommandHandler ) )
+		, UnountPakCommand(
+		TEXT( "DDC.UnmountPak" ),
+		*LOCTEXT("CommandText_DDCUnmountPak", "Unmounts read-only pak file").ToString(),
+		FConsoleCommandWithArgsDelegate::CreateRaw( this, &FDerivedDataBackendGraph::UnmountPakCommandHandler ) )
+	{
+		check(IsInGameThread()); // we pretty much need this to be initialized from the main thread...it uses GConfig, etc
+		check(GConfig && GConfig->IsReadyForUse());
+		RootCache = NULL;
+		TMap<FString, FDerivedDataBackendInterface*> ParsedNodes;		
+
+		// Create the graph using ini settings
+		if( FParse::Value( FCommandLine::Get(), TEXT("-DDC="), GraphName ) )
+		{
+			if( GraphName.Len() > 0 )
+			{
+				RootCache = ParseNode( TEXT("Root"), GEngineIni, *GraphName, ParsedNodes );				
+			}
+
+			if( RootCache == NULL )
+			{
+				// Remove references to any backend instances that might have been created
+				ParsedNodes.Empty();
+				DestroyCreatedBackends();
+				UE_LOG( LogDerivedDataCache, Warning, TEXT("FDerivedDataBackendGraph:  Unable to create backend graph using the specified graph settings (%s). Reverting to default."), *GraphName );
+			}
+		}
+
+		if( !RootCache )
+		{
+			// Use default graph
+			GraphName = FApp::IsEngineInstalled() ? TEXT("InstalledDerivedDataBackendGraph") : TEXT("DerivedDataBackendGraph");
+			FString Entry;
+			if( !GConfig->GetString( *GraphName, TEXT("Root"), Entry, GEngineIni ) || !Entry.Len())
+			{
+				UE_LOG( LogDerivedDataCache, Fatal, TEXT("Unable to create backend graph using the default graph settings (%s) ini=%s."), *GraphName, *GEngineIni );
+			}
+			RootCache = ParseNode( TEXT("Root"), GEngineIni, *GraphName, ParsedNodes );
+		}
+		check(RootCache);
+
+		// Make sure AsyncPutWrapper and KeyLengthWrapper are created
+		if( !AsyncPutWrapper )
+		{
+			AsyncPutWrapper = new FDerivedDataBackendAsyncPutWrapper( RootCache, true );
+			check(AsyncPutWrapper);
+			CreatedBackends.Add( AsyncPutWrapper );
+			RootCache = AsyncPutWrapper;
+		}
+		if ( !KeyLengthWrapper )
+		{
+			KeyLengthWrapper = new FDerivedDataLimitKeyLengthWrapper( RootCache, MAX_BACKEND_KEY_LENGTH );
+			check(KeyLengthWrapper);
+			CreatedBackends.Add( KeyLengthWrapper );
+			RootCache = KeyLengthWrapper;
+		}
+	}
+
+	/**
+	 * Helper function to get the value of parsed bool as the return value
+	 **/
+	bool GetParsedBool( const TCHAR* Stream, const TCHAR* Match ) const
+	{
+		bool bValue = 0;
+		FParse::Bool( Stream, Match, bValue );
+		return bValue;
+	}
+
+	/**
+	 * Parses backend graph node from ini settings
+	 *
+	 * @param NodeName Name of the node to parse
+	 * @param IniFilename Ini filename
+	 * @param IniSection Section in the ini file containing the graph definition
+	 * @param InParsedNodes Map of parsed nodes and their names to be able to find already parsed nodes
+	 * @return Derived data backend interface instance created from ini settings
+	 */
+	FDerivedDataBackendInterface* ParseNode( const TCHAR* NodeName, const FString& IniFilename, const TCHAR* IniSection, TMap<FString, FDerivedDataBackendInterface*>& InParsedNodes  )
+	{
+		FDerivedDataBackendInterface* ParsedNode = NULL;
+		FString Entry;
+		if( GConfig->GetString( IniSection, NodeName, Entry, IniFilename ) )
+		{
+			// Trim whitespace at the beginning.
+			Entry = Entry.Trim();
+			// Remove brackets.
+			Entry.RemoveFromStart(TEXT("("));
+			Entry.RemoveFromEnd(TEXT(")"));
+
+			FString	NodeType;
+			if( FParse::Value( *Entry, TEXT("Type="), NodeType ) )
+			{
+				if( NodeType == TEXT("FileSystem") )
+				{
+					ParsedNode = ParseDataCache( NodeName, *Entry );
+				}
+				else if( NodeType == TEXT("Boot") )
+				{
+					if( BootCache == NULL )
+					{
+						BootCache = ParseBootCache( NodeName, *Entry, BootCacheFilename );
+						ParsedNode = BootCache;
+					}
+					else
+					{
+						UE_LOG( LogDerivedDataCache, Warning, TEXT("FDerivedDataBackendGraph:  Unable to create %s Boot cache because only one Boot cache node is supported."), *NodeName );
+					}
+				}
+				else if( NodeType == TEXT("Memory") )
+				{
+					ParsedNode = ParseMemoryCache( NodeName, *Entry );
+				}
+				else if( NodeType == TEXT("Hierarchical") )
+				{
+					ParsedNode = ParseHierarchicalCache( NodeName, *Entry, IniFilename, IniSection, InParsedNodes );
+				}
+				else if( NodeType == TEXT("KeyLength") )
+				{
+					if( KeyLengthWrapper == NULL )
+					{
+						KeyLengthWrapper = ParseKeyLength( NodeName, *Entry, IniFilename, IniSection, InParsedNodes );
+						ParsedNode = KeyLengthWrapper;
+					}
+					else
+					{
+						UE_LOG( LogDerivedDataCache, Warning, TEXT("FDerivedDataBackendGraph:  Unable to create %s KeyLengthWrapper because only one KeyLength node is supported."), *NodeName );
+					}
+				}
+				else if( NodeType == TEXT("AsyncPut") )
+				{
+					if( AsyncPutWrapper == NULL )
+					{
+						AsyncPutWrapper = ParseAsyncPut( NodeName, *Entry, IniFilename, IniSection, InParsedNodes );
+						ParsedNode = AsyncPutWrapper;
+					}
+					else
+					{
+						UE_LOG( LogDerivedDataCache, Warning, TEXT("FDerivedDataBackendGraph:  Unable to create %s AsyncPutWrapper because only one AsyncPutWrapper node is supported."), *NodeName );
+					}
+				}
+				else if( NodeType == TEXT("Verify") )
+				{
+					ParsedNode = ParseVerify( NodeName, *Entry, IniFilename, IniSection, InParsedNodes );
+				}
+				else if( NodeType == TEXT("ReadPak") )
+				{
+					ParsedNode = ParsePak( NodeName, *Entry, false );
+				}
+				else if( NodeType == TEXT("WritePak") )
+				{
+					ParsedNode = ParsePak( NodeName, *Entry, true );
+				}
+			}
+		}
+
+		if( ParsedNode )
+		{
+			// Store this node so that we don't require any order of adding backend nodes
+			InParsedNodes.Add( NodeName, ParsedNode );
+			// Keep references to all created nodes.
+			CreatedBackends.Add( ParsedNode );
+		}
+
+		return ParsedNode;
+	}
+
+	/**
+	 * Creates Read/write Pak file interface from ini settings
+	 *
+	 * @param NodeName Node name
+	 * @param Entry Node definition
+	 * @param bWriting true to create pak interface for writing
+	 * @return Pak file data backend interface instance or NULL if unsuccessful
+	 */
+	FDerivedDataBackendInterface* ParsePak( const TCHAR* NodeName, const TCHAR* Entry, const bool bWriting )
+	{
+		FDerivedDataBackendInterface* PakNode = NULL;
+		FString PakFilename;
+		FParse::Value( Entry, TEXT("Filename="), PakFilename );
+		bool bCompressed = GetParsedBool(Entry, TEXT("Compressed="));
+
+		if( !PakFilename.Len() )
+		{
+			UE_LOG( LogDerivedDataCache, Log, TEXT("FDerivedDataBackendGraph:  %s pak cache Filename not found in *engine.ini, will not use a pak cache."), NodeName );
+		}
+		else
+		{
+			// now add the pak read cache (if any) to the front of the cache hierarchy
+			if ( bWriting )
+			{
+				FGuid Temp = FGuid::NewGuid();
+				ReadPakFilename = PakFilename;
+				WritePakFilename = PakFilename + TEXT(".") + Temp.ToString();
+				WritePakCache = bCompressed? new FCompressedPakFileDerivedDataBackend( *WritePakFilename, true ) : new FPakFileDerivedDataBackend( *WritePakFilename, true );
+				PakNode = WritePakCache;
+			}
+			else
+			{
+				bool bReadPak = FPlatformFileManager::Get().GetPlatformFile().FileExists( *PakFilename );
+				if( bReadPak )
+				{
+					FPakFileDerivedDataBackend* ReadPak = bCompressed? new FCompressedPakFileDerivedDataBackend( *PakFilename, false ) : new FPakFileDerivedDataBackend( *PakFilename, false );
+					ReadPakFilename = PakFilename;
+					PakNode = ReadPak;
+					ReadPakCache.Add(ReadPak);
+				}
+				else
+				{
+					UE_LOG( LogDerivedDataCache, Log, TEXT("FDerivedDataBackendGraph:  %s pak cache file %s not found, will not use a pak cache."), NodeName, *PakFilename );
+				}
+			}
+		}
+		return PakNode;
+	}
+
+	/**
+	 * Creates Verify wrapper interface from ini settings.
+	 *
+	 * @param NodeName Node name.
+	 * @param Entry Node definition.
+	 * @param IniFilename ini filename.
+	 * @param IniSection ini section containing graph definition
+	 * @param InParsedNodes map of nodes and their names which have already been parsed
+	 * @return Verify wrapper backend interface instance or NULL if unsuccessful
+	 */
+	FDerivedDataBackendInterface* ParseVerify( const TCHAR* NodeName, const TCHAR* Entry, const FString& IniFilename, const TCHAR* IniSection, TMap<FString, FDerivedDataBackendInterface*>& InParsedNodes )
+	{
+		FDerivedDataBackendInterface* InnerNode = NULL;
+		FString InnerName;
+		if( FParse::Value( Entry, TEXT("Inner="), InnerName ) )
+		{
+			FDerivedDataBackendInterface** ParsedInnerNode = InParsedNodes.Find( InnerName );
+			if( ParsedInnerNode )
+			{
+				UE_LOG( LogDerivedDataCache, Warning, TEXT("Inner node %s for Verify node %s already exists. Nodes can only be used once."), *InnerName, NodeName );
+				return NULL;
+			}
+			else
+			{
+				InnerNode = ParseNode( *InnerName, IniFilename, IniSection, InParsedNodes );
+			}
+		}
+
+		FDerivedDataBackendInterface* VerifyNode = NULL;
+		if( InnerNode )
+		{
+			IFileManager::Get().DeleteDirectory(*(FPaths::GameSavedDir() / TEXT("VerifyDDC/")), false, true);
+
+			const bool bFix = GetParsedBool( Entry, TEXT("Fix=") );
+			VerifyNode = new FDerivedDataBackendVerifyWrapper( InnerNode, bFix );
+		}
+		else
+		{
+			UE_LOG( LogDerivedDataCache, Warning, TEXT("Unable to find inner node %s for Verify node %s. Verify node will not be created."), *InnerName, NodeName );
+		}
+
+		return VerifyNode;
+	}
+
+	/**
+	 * Creates AsyncPut wrapper interface from ini settings.
+	 *
+	 * @param NodeName Node name.
+	 * @param Entry Node definition.
+	 * @param IniFilename ini filename.
+	 * @param IniSection ini section containing graph definition
+	 * @param InParsedNodes map of nodes and their names which have already been parsed
+	 * @return AsyncPut wrapper backend interface instance or NULL if unsuccessfull
+	 */
+	FDerivedDataBackendInterface* ParseAsyncPut( const TCHAR* NodeName, const TCHAR* Entry, const FString& IniFilename, const TCHAR* IniSection, TMap<FString, FDerivedDataBackendInterface*>& InParsedNodes )
+	{
+		FDerivedDataBackendInterface* InnerNode = NULL;
+		FString InnerName;
+		if( FParse::Value( Entry, TEXT("Inner="), InnerName ) )
+		{
+			FDerivedDataBackendInterface** ParsedInnerNode = InParsedNodes.Find( InnerName );
+			if( ParsedInnerNode )
+			{
+				UE_LOG( LogDerivedDataCache, Warning, TEXT("Inner node %s for AsyncPut node %s already exists. Nodes can only be used once."), *InnerName, NodeName );
+				return NULL;
+			}
+			else
+			{
+				InnerNode = ParseNode( *InnerName, IniFilename, IniSection, InParsedNodes );
+			}
+		}
+
+		FDerivedDataBackendInterface* AsyncNode = NULL;
+		if( InnerNode )
+		{
+			AsyncNode = new FDerivedDataBackendAsyncPutWrapper( InnerNode, true );
+		}
+		else
+		{
+			UE_LOG( LogDerivedDataCache, Warning, TEXT("Unable to find inner node %s for AsyncPut node %s. AsyncPut node will not be created."), *InnerName, NodeName );
+		}
+
+		return AsyncNode;
+	}
+
+	/**
+	 * Creates KeyLength wrapper interface from ini settings.
+	 *
+	 * @param NodeName Node name.
+	 * @param Entry Node definition.
+	 * @param IniFilename ini filename.
+	 * @param IniSection ini section containing graph definition
+	 * @param InParsedNodes map of nodes and their names which have already been parsed
+	 * @return KeyLength wrapper backend interface instance or NULL if unsuccessfull
+	 */
+	FDerivedDataBackendInterface* ParseKeyLength( const TCHAR* NodeName, const TCHAR* Entry, const FString& IniFilename, const TCHAR* IniSection, TMap<FString, FDerivedDataBackendInterface*>& InParsedNodes )
+	{
+		FDerivedDataBackendInterface* InnerNode = NULL;
+		FString InnerName;
+		if( FParse::Value( Entry, TEXT("Inner="), InnerName ) )
+		{
+			FDerivedDataBackendInterface** ParsedInnerNode = InParsedNodes.Find( InnerName );
+			if( ParsedInnerNode )
+			{
+				UE_LOG( LogDerivedDataCache, Warning, TEXT("Inner node %s for KeyLength node %s already exists. Nodes can only be used once."), *InnerName, NodeName );
+				return NULL;
+			}
+			else
+			{
+				InnerNode = ParseNode( *InnerName, IniFilename, IniSection, InParsedNodes );
+			}
+		}
+
+		FDerivedDataBackendInterface* KeyLengthNode = NULL;
+		if( InnerNode )
+		{
+			int32 KeyLength = MAX_BACKEND_KEY_LENGTH;
+			FParse::Value( Entry, TEXT("Length="), KeyLength );
+			KeyLength = FMath::Clamp( KeyLength, 0, MAX_BACKEND_KEY_LENGTH );
+
+			KeyLengthNode = new FDerivedDataLimitKeyLengthWrapper( InnerNode, KeyLength );
+		}
+		else
+		{
+			UE_LOG( LogDerivedDataCache, Warning, TEXT("Unable to find inner node %s for KeyLength node %s. KeyLength node will not be created."), *InnerName, NodeName );
+		}
+
+		return KeyLengthNode;
+	}
+
+	/**
+	 * Creates Hierarchical interface from ini settings.
+	 *
+	 * @param NodeName Node name.
+	 * @param Entry Node definition.
+	 * @param IniFilename ini filename.
+	 * @param IniSection ini section containing graph definition
+	 * @param InParsedNodes map of nodes and their names which have already been parsed
+	 * @return Hierarchical backend interface instance or NULL if unsuccessfull
+	 */
+	FDerivedDataBackendInterface* ParseHierarchicalCache( const TCHAR* NodeName, const TCHAR* Entry, const FString& IniFilename, const TCHAR* IniSection, TMap<FString, FDerivedDataBackendInterface*>& InParsedNodes )
+	{
+		const TCHAR* InnerMatch = TEXT("Inner=");
+		const int32 InnerMatchLength = FCString::Strlen( InnerMatch );
+
+		TArray< FDerivedDataBackendInterface* > InnerNodes;
+		FString InnerName;
+		while ( FParse::Value( Entry, InnerMatch, InnerName ) )
+		{
+			// Check if the child has already been parsed
+			FDerivedDataBackendInterface** ParsedInnerNode = InParsedNodes.Find( InnerName );
+			if( ParsedInnerNode )
+			{
+				UE_LOG( LogDerivedDataCache, Warning, TEXT("Inner node %s for hierarchical node %s already exists. Nodes can only be used once."), *InnerName, NodeName );
+			}
+			else
+			{
+				FDerivedDataBackendInterface* InnerNode = ParseNode( *InnerName, IniFilename, IniSection, InParsedNodes );
+				if( InnerNode )
+				{
+					InnerNodes.Add( InnerNode );
+				}
+				else
+				{
+					UE_LOG( LogDerivedDataCache, Log, TEXT("Unable to find inner node %s for hierarchical cache %s."), *InnerName, NodeName );
+				}
+			}
+
+			// Move the Entry pointer forward so that we can find more children
+			Entry = FCString::Strfind( Entry, InnerMatch );
+			Entry += InnerMatchLength;
+		}
+
+		FDerivedDataBackendInterface* Hierarchy = NULL;
+		if( InnerNodes.Num() > 1 )
+		{
+			FHierarchicalDerivedDataBackend* HierarchyBackend = new FHierarchicalDerivedDataBackend( InnerNodes );
+			Hierarchy = HierarchyBackend;
+			if (HierarchicalWrapper == NULL)
+			{
+				HierarchicalWrapper = HierarchyBackend;
+			}
+		}
+		else if ( InnerNodes.Num() == 1 )
+		{
+			Hierarchy = InnerNodes[ 0 ];
+		}
+		else
+		{
+			UE_LOG( LogDerivedDataCache, Warning, TEXT("Hierarchical cache %s has no inner backends and will not be created."), NodeName );
+		}
+
+		return Hierarchy;
+	}
+
+	/**
+	 * Creates Filesystem data cache interface from ini settings.
+	 *
+	 * @param NodeName Node name.
+	 * @param Entry Node definition.
+	 * @return Filesystem data cache backend interface instance or NULL if unsuccessfull
+	 */
+	FDerivedDataBackendInterface* ParseDataCache( const TCHAR* NodeName, const TCHAR* Entry )
+	{
+		FDerivedDataBackendInterface* DataCache = NULL;
+
+		// Parse Path by default, it may be overwriten by EnvPathOverride
+		FString Path;
+		FParse::Value( Entry, TEXT("Path="), Path );
+
+		// Check the EnvPathOverride environment variable to allow persistent overriding of data cache path, eg for offsite workers.
+		FString EnvPathOverride;
+		if( FParse::Value( Entry, TEXT("EnvPathOverride="), EnvPathOverride ) )
+		{
+			TCHAR FilesystemCachePathEnv[256];
+			FPlatformMisc::GetEnvironmentVariable( *EnvPathOverride, FilesystemCachePathEnv, ARRAY_COUNT(FilesystemCachePathEnv) );
+			if( FilesystemCachePathEnv[0] )
+			{
+				Path = FilesystemCachePathEnv;
+				UE_LOG( LogDerivedDataCache, Log, TEXT("Found environment variable %s=%s"), *EnvPathOverride, *Path );
+			}
+		}
+		// Check if the Path is a real path or a special keyword
+		if (FEngineBuildSettings::IsInternalBuild())
+		{
+			auto DDCUtils = FModuleManager::LoadModulePtr< IDDCUtilsModuleInterface >("DDCUtils");
+			if (DDCUtils)
+			{
+				FString PathFromName = DDCUtils->GetSharedCachePath(Path);
+				if (!PathFromName.IsEmpty())
+				{
+					Path = PathFromName;
+				}
+			}
+		}
+		else if (Path.StartsWith(TEXT("?")))
+		{
+			Path = TEXT("");
+		}
+
+		if( !Path.Len() )
+		{
+			UE_LOG( LogDerivedDataCache, Log, TEXT("%s data cache path not found in *engine.ini, will not use an %s cache."), NodeName, NodeName );
+		}
+		else
+		{
+			const bool bReadOnly = GetParsedBool( Entry, TEXT("ReadOnly=") );
+			const bool bClean = GetParsedBool( Entry, TEXT("Clean=") );
+			const bool bFlush = GetParsedBool( Entry, TEXT("Flush=") );
+			const bool bTouch = GetParsedBool( Entry, TEXT("Touch=") );
+			const bool bPurgeTransient = GetParsedBool( Entry, TEXT("PurgeTransient=") );
+
+			bool bDeleteUnused = true; // On by default
+			FParse::Bool( Entry, TEXT("DeleteUnused="), bDeleteUnused );
+			int32 UnusedFileAge = 17;
+			FParse::Value( Entry, TEXT("UnusedFileAge="), UnusedFileAge );
+			int32 MaxFoldersToClean = -1;
+			FParse::Value( Entry, TEXT("FoldersToClean="), MaxFoldersToClean );
+			int32 MaxFileChecksPerSec = -1;
+			FParse::Value( Entry, TEXT("MaxFileChecksPerSec="), MaxFileChecksPerSec );
+
+			if( bFlush )
+			{
+				IFileManager::Get().DeleteDirectory( *(Path / TEXT("")), false, true );
+			}
+			else if( bClean )
+			{
+				DeleteOldFiles( *Path );
+			}
+
+			FDerivedDataBackendInterface* InnerFileSystem = NULL;
+
+			// Don't create the file system if shared data cache directory is not mounted
+			if( FCString::Strcmp(NodeName, TEXT("Shared")) != 0 || IFileManager::Get().DirectoryExists(*Path) )
+			{
+				InnerFileSystem = CreateFileSystemDerivedDataBackend( *Path, bReadOnly, bTouch, bPurgeTransient, bDeleteUnused, UnusedFileAge, MaxFoldersToClean, MaxFileChecksPerSec);
+			}
+
+			if( InnerFileSystem )
+			{
+				DataCache = new FDerivedDataBackendCorruptionWrapper( InnerFileSystem );
+				UE_LOG( LogDerivedDataCache, Log, TEXT("Using %s data cache path %s: %s"), NodeName, *Path, bReadOnly ? TEXT("ReadOnly") : TEXT("Writable") );
+				Directories.AddUnique(Path);
+			}
+			else
+			{
+				UE_LOG( LogDerivedDataCache, Warning, TEXT("%s data cache path was not usable, will not use it."), NodeName );
+			}
+		}
+
+		return DataCache;
+	}
+
+	/**
+	 * Creates Boot data cache interface from ini settings.
+	 *
+	 * @param NodeName Node name.
+	 * @param Entry Node definition.
+	 * @param OutFilename filename specified for the cache
+	 * @return Boot data cache backend interface instance or NULL if unsuccessfull
+	 */
+	FMemoryDerivedDataBackend* ParseBootCache( const TCHAR* NodeName, const TCHAR* Entry, FString& OutFilename )
+	{
+		FMemoryDerivedDataBackend* Cache = NULL;
+		FString Filename;
+		int64 MaxCacheSize = -1; // in MB
+		const int64 MaxSupportedCacheSize = 2048; // 2GB
+
+		FParse::Value( Entry, TEXT("MaxCacheSize="), MaxCacheSize);
+		FParse::Value( Entry, TEXT("Filename="), Filename );
+		if ( !Filename.Len() )
+		{
+			UE_LOG( LogDerivedDataCache, Warning, TEXT("FDerivedDataBackendGraph:  %s filename not found in *engine.ini, will not use %s cache."), NodeName, NodeName );
+		}
+		else
+		{
+			// make sure MaxCacheSize does not exceed 2GB
+			MaxCacheSize = FMath::Min(MaxCacheSize, MaxSupportedCacheSize);
+
+			UE_LOG( LogDerivedDataCache, Display, TEXT("Max Cache Size: %d MB"), MaxCacheSize);
+			Cache = new FMemoryDerivedDataBackend(MaxCacheSize * 1024 * 1024);
+
+			if( Cache && Filename.Len() )
+			{
+				OutFilename = Filename;
+
+				if (MaxCacheSize > 0 && IFileManager::Get().FileSize(*Filename) >= (MaxCacheSize * 1024 * 1024))
+				{
+					UE_LOG( LogDerivedDataCache, Warning, TEXT("FDerivedDataBackendGraph:  %s filename exceeds max size."), NodeName );
+					return Cache;
+				}
+
+				if (IFileManager::Get().FileSize(*Filename) < 0)
+				{
+					UE_LOG( LogDerivedDataCache, Display, TEXT("Starting with empty %s cache"), NodeName );
+				}
+				else if( Cache->LoadCache( *Filename ) )
+				{
+					UE_LOG( LogDerivedDataCache, Display, TEXT("Loaded %s cache: %s"), NodeName, *Filename );
+				}
+				else
+				{
+					UE_LOG( LogDerivedDataCache, Warning, TEXT("Could not load %s cache: %s"), NodeName, *Filename );
+				}
+			}
+		}
+		return Cache;
+	}
+
+	/**
+	 * Creates Memory data cache interface from ini settings.
+	 *
+	 * @param NodeName Node name.
+	 * @param Entry Node definition.
+	 * @return Memory data cache backend interface instance or NULL if unsuccessfull
+	 */
+	FMemoryDerivedDataBackend* ParseMemoryCache( const TCHAR* NodeName, const TCHAR* Entry )
+	{
+		FMemoryDerivedDataBackend* Cache = NULL;
+		FString Filename;
+
+		FParse::Value( Entry, TEXT("Filename="), Filename );
+		Cache = new FMemoryDerivedDataBackend;
+		if( Cache && Filename.Len() )
+		{
+			if( Cache->LoadCache( *Filename ) )
+			{
+				UE_LOG( LogDerivedDataCache, Display, TEXT("Loaded %s cache: %s"), NodeName, *Filename );
+			}
+			else
+			{
+				UE_LOG( LogDerivedDataCache, Warning, TEXT("Could not load %s cache: %s"), NodeName, *Filename );
+			}
+		}
+
+		return Cache;
+	}
+
+	~FDerivedDataBackendGraph()
+	{
+		RootCache = NULL;
+		DestroyCreatedBackends();
+	}
+
+	FDerivedDataBackendInterface& GetRoot() override
+	{
+		check(RootCache);
+		return *RootCache;
+	}
+
+	virtual void NotifyBootComplete() override
+	{
+		check(RootCache);
+		if (BootCache)
+		{
+			if( !FParse::Param( FCommandLine::Get(), TEXT("DDCNOSAVEBOOT") ) && !FParse::Param( FCommandLine::Get(), TEXT("Multiprocess") ) )
+			{
+				BootCache->SaveCache(*BootCacheFilename);
+			}
+			BootCache->Disable();
+		}
+	}
+
+	virtual void WaitForQuiescence(bool bShutdown) override
+	{
+		double StartTime = FPlatformTime::Seconds();
+		double LastPrint = StartTime;
+		while (AsyncCompletionCounter.GetValue())
+		{
+			check(AsyncCompletionCounter.GetValue() >= 0);
+			FPlatformProcess::Sleep(1.0f);
+			if (FPlatformTime::Seconds() - LastPrint > 5.0)
+			{
+				UE_LOG(LogDerivedDataCache, Log, TEXT("Waited %ds for derived data cache to finish..."), int32(FPlatformTime::Seconds() - StartTime));
+				LastPrint = FPlatformTime::Seconds();
+			}
+		}
+		if (bShutdown)
+		{
+			FString MergePaks;
+			if(WritePakCache && WritePakCache->IsWritable() && FParse::Value(FCommandLine::Get(), TEXT("MergePaks="), MergePaks))
+			{
+				TArray<FString> MergePakList;
+				MergePaks.FString::ParseIntoArray(MergePakList, TEXT("+"));
+
+				for(const FString& MergePakName : MergePakList)
+				{
+					FPakFileDerivedDataBackend ReadPak(*FPaths::Combine(*FPaths::GetPath(WritePakFilename), *MergePakName), false);
+					WritePakCache->MergeCache(&ReadPak);
+				}
+			}
+			for (int32 ReadPakIndex = 0; ReadPakIndex < ReadPakCache.Num(); ReadPakIndex++)
+			{
+				ReadPakCache[ReadPakIndex]->Close();
+			}
+			if (WritePakCache && WritePakCache->IsWritable())
+			{
+				WritePakCache->Close();
+				if (!FPlatformFileManager::Get().GetPlatformFile().FileExists(*WritePakFilename))
+				{
+					UE_LOG(LogDerivedDataCache, Error, TEXT("Pak file %s was not produced?"), *WritePakFilename);
+				}
+				else
+				{
+					if (FPlatformFileManager::Get().GetPlatformFile().FileExists(*ReadPakFilename))
+					{
+						FPlatformFileManager::Get().GetPlatformFile().SetReadOnly(*ReadPakFilename, false);
+						if (!FPlatformFileManager::Get().GetPlatformFile().DeleteFile(*ReadPakFilename))
+						{
+							UE_LOG(LogDerivedDataCache, Error, TEXT("Could not delete the pak file %s to overwrite it with a new one."), *ReadPakFilename);
+						}
+					}
+					if (!FPakFileDerivedDataBackend::SortAndCopy(WritePakFilename, ReadPakFilename))
+					{
+						UE_LOG(LogDerivedDataCache, Error, TEXT("Couldn't sort pak file (%s)"), *WritePakFilename);
+					}
+					else if (!IFileManager::Get().Delete(*WritePakFilename))
+					{
+						UE_LOG(LogDerivedDataCache, Error, TEXT("Couldn't delete pak file (%s)"), *WritePakFilename);
+					}
+					else
+					{
+						UE_LOG(LogDerivedDataCache, Display, TEXT("Sucessfully wrote %s."), *ReadPakFilename);
+					}
+				}
+			}
+		}
+	}
+
+	virtual void AddToAsyncCompletionCounter(int32 Addend) override
+	{
+		AsyncCompletionCounter.Add(Addend);
+		check(AsyncCompletionCounter.GetValue() >= 0);
+	}
+
+	virtual void GetDirectories(TArray<FString>& OutResults) override
+	{
+		OutResults = Directories;
+	}
+
+	static FORCEINLINE FDerivedDataBackendGraph& Get()
+	{
+		static FDerivedDataBackendGraph SingletonInstance;
+		return SingletonInstance;
+	}
+
+	virtual FDerivedDataBackendInterface* MountPakFile(const TCHAR* PakFilename) override
+	{
+		// Assumptions: there's at least one read-only pak backend in the hierarchy
+		// and its parent is a hierarchical backend.
+		FPakFileDerivedDataBackend* ReadPak = NULL;
+		if (HierarchicalWrapper && FPlatformFileManager::Get().GetPlatformFile().FileExists(PakFilename))
+		{
+			ReadPak = new FPakFileDerivedDataBackend(PakFilename, false);
+
+			HierarchicalWrapper->AddInnerBackend(ReadPak);
+			CreatedBackends.Add(ReadPak);
+			ReadPakCache.Add(ReadPak);
+		}
+		else
+		{
+			UE_LOG(LogDerivedDataCache, Warning, TEXT("Failed to add %s read-only pak DDC backend. Make sure it exists and there's at least one hierarchical backend in the cache tree."), PakFilename);
+		}
+
+		return ReadPak;
+	}
+
+	virtual bool UnmountPakFile(const TCHAR* PakFilename) override
+	{
+		for (int PakIndex = 0; PakIndex < ReadPakCache.Num(); ++PakIndex)
+		{
+			FPakFileDerivedDataBackend* ReadPak = ReadPakCache[PakIndex];
+			if (ReadPak->GetFilename() == PakFilename)
+			{
+				check(HierarchicalWrapper);
+
+				// Wait until all async requests are complete.
+				WaitForQuiescence(false);
+
+				HierarchicalWrapper->RemoveInnerBackend(ReadPak);
+				ReadPakCache.RemoveAt(PakIndex);
+				CreatedBackends.Remove(ReadPak);
+				ReadPak->Close();
+				delete ReadPak;
+				return true;
+			}
+		}
+		return false;
+	}
+
+private:
+
+	/** Delete the old files in a directory **/
+	void DeleteOldFiles(const TCHAR* Directory)
+	{
+		float MinimumDaysToKeepFile = 7;
+		GConfig->GetFloat( *GraphName, TEXT("MinimumDaysToKeepFile"), MinimumDaysToKeepFile, GEngineIni );
+		check(MinimumDaysToKeepFile > 0.0f); // sanity
+		//@todo 
+	}
+
+	/** Delete all created backends in the reversed order they were created. */
+	void DestroyCreatedBackends()
+	{
+		for (int32 BackendIndex = CreatedBackends.Num() - 1; BackendIndex >= 0; --BackendIndex)
+		{
+			delete CreatedBackends[BackendIndex];
+		}
+		CreatedBackends.Empty();
+	}
+
+	/** MountPak console command handler. */
+	void UnmountPakCommandHandler(const TArray<FString>& Args)
+	{
+		if (Args.Num() < 1)
+		{
+			UE_LOG(LogDerivedDataCache, Log, TEXT("Usage: DDC.MountPak PakFilename"));
+			return;
+		}
+		UnmountPakFile(*Args[0]);
+	}
+
+	/** UnmountPak console command handler. */
+	void MountPakCommandHandler(const TArray<FString>& Args)
+	{
+		if (Args.Num() < 1)
+		{
+			UE_LOG(LogDerivedDataCache, Log, TEXT("Usage: DDC.UnmountPak PakFilename"));
+			return;
+		}
+		MountPakFile(*Args[0]);
+	}
+
+	FThreadSafeCounter								AsyncCompletionCounter;
+	FString											GraphName;
+	FString											BootCacheFilename;
+	FString											ReadPakFilename;
+	FString											WritePakFilename;
+
+	/** Root of the graph */
+	FDerivedDataBackendInterface*					RootCache;
+
+	/** References to all created backed interfaces */
+	TArray< FDerivedDataBackendInterface* > CreatedBackends;
+
+	/** Instances of backend interfaces which exist in only one copy */
+	FMemoryDerivedDataBackend*		BootCache;
+	FPakFileDerivedDataBackend*		WritePakCache;
+	FDerivedDataBackendInterface*	AsyncPutWrapper;
+	FDerivedDataBackendInterface*	KeyLengthWrapper;
+	FHierarchicalDerivedDataBackend* HierarchicalWrapper;
+	/** Support for multiple read only pak files. */
+	TArray<FPakFileDerivedDataBackend*>		ReadPakCache;
+
+	/** List of directories used by the DDC */
+	TArray<FString> Directories;
+
+	/** MountPak console command */
+	FAutoConsoleCommand MountPakCommand;
+	/** UnmountPak console command */
+	FAutoConsoleCommand UnountPakCommand;
+};
+
+FDerivedDataBackend& FDerivedDataBackend::Get()
+{
+	return FDerivedDataBackendGraph::Get();
+}
+
+#undef LOCTEXT_NAMESPACE
\ No newline at end of file
diff --git a/Engine/Source/Developer/DerivedDataCache/Private/DerivedDataCache.cpp b/Engine/Source/Developer/DerivedDataCache/Private/DerivedDataCache.cpp
new file mode 100644
--- /dev/null
+++ b/Engine/Source/Developer/DerivedDataCache/Private/DerivedDataCache.cpp
@@ -0,0 +1,1046 @@
+// Copyright 1998-2015 Epic Games, Inc. All Rights Reserved.
+
+
+#include "Core.h"
+#include "ModuleManager.h"
+
+#include "DerivedDataCacheInterface.h"
+#include "DerivedDataBackendInterface.h"
+#include "DerivedDataPluginInterface.h"
+#include "DDCCleanup.h"
+
+
+DEFINE_STAT(STAT_DDC_NumGets);
+DEFINE_STAT(STAT_DDC_NumPuts);
+DEFINE_STAT(STAT_DDC_NumBuilds);
+DEFINE_STAT(STAT_DDC_NumExist);
+DEFINE_STAT(STAT_DDC_SyncGetTime);
+DEFINE_STAT(STAT_DDC_ASyncWaitTime);
+DEFINE_STAT(STAT_DDC_PutTime);
+DEFINE_STAT(STAT_DDC_SyncBuildTime);
+DEFINE_STAT(STAT_DDC_ExistTime);
+
+/** 
+ * Implementation of the derived data cache
+ * This API is fully threadsafe
+**/
+class FDerivedDataCache : public FDerivedDataCacheInterface
+{
+
+	/** 
+	 * Async worker that checks the cache backend and if that fails, calls the deriver to build the data and then puts the results to the cache
+	**/
+	friend class FBuildAsyncWorker;
+	class FBuildAsyncWorker : public FNonAbandonableTask
+	{
+	public:
+		/** 
+		 * Constructor for async task 
+		 * @param	InDataDeriver	plugin to produce cache key and in the event of a miss, return the data.
+		 * @param	InCacheKey		Complete cache key for this data.
+		**/
+		FBuildAsyncWorker(FDerivedDataPluginInterface* InDataDeriver, const TCHAR* InCacheKey, bool bInSynchronousForStats, FCacheStatRecord* InGetRecord, FCacheStatRecord* InPutRecord)
+		: bSuccess(false)
+		, bSynchronousForStats(bInSynchronousForStats)
+		, DataDeriver(InDataDeriver)
+		, CacheKey(InCacheKey)
+		, GetRecord(InGetRecord)
+		, PutRecord(InPutRecord)
+		{
+		}
+		
+		/** Async worker that checks the cache backend and if that fails, calls the deriver to build the data and then puts the results to the cache **/
+		void DoWork()
+		{
+			bool bGetResult;
+			{
+				INC_DWORD_STAT(STAT_DDC_NumGets);
+				STAT(double ThisTime = 0);
+				{
+					SCOPE_SECONDS_COUNTER(ThisTime);
+					GetRecord->CacheKey = CacheKey;
+					GetRecord->CacheGet = 1;
+					GetRecord->bSynchronous = bSynchronousForStats;
+					GetRecord->StartTime = FPlatformTime::Seconds();
+					GetRecord->GetDuration = 0.;
+					GetRecord->PutDuration = 0.;
+					bGetResult = FDerivedDataBackend::Get().GetRoot().GetCachedData(*CacheKey, Data, GetRecord);
+					GetRecord->EndTime = FPlatformTime::Seconds();
+					GetRecord->DataSize = Data.Num();
+				}
+				INC_FLOAT_STAT_BY(STAT_DDC_SyncGetTime, bSynchronousForStats ? (float)ThisTime : 0.0f);
+			}
+			if (bGetResult)
+			{
+				check(Data.Num());
+				bSuccess = true;
+				delete DataDeriver;
+				DataDeriver = NULL;
+			}
+			else if (DataDeriver)
+			{
+				{
+					INC_DWORD_STAT(STAT_DDC_NumBuilds);
+					STAT(double ThisTime = 0);
+					{
+						SCOPE_SECONDS_COUNTER(ThisTime);
+						bSuccess = DataDeriver->Build(Data);
+					}
+					INC_FLOAT_STAT_BY(STAT_DDC_SyncBuildTime, bSynchronousForStats ? (float)ThisTime : 0.0f);
+				}
+				delete DataDeriver;
+				DataDeriver = NULL;
+				if (bSuccess)
+				{
+					check(Data.Num());
+					INC_DWORD_STAT(STAT_DDC_NumPuts);
+					STAT(double ThisTime = 0);
+					{
+						SCOPE_SECONDS_COUNTER(ThisTime);
+						PutRecord->CacheKey = CacheKey;
+						PutRecord->CacheGet = 0;
+						PutRecord->bSynchronous = bSynchronousForStats;
+						PutRecord->StartTime = FPlatformTime::Seconds();
+						PutRecord->GetDuration = 0.;
+						PutRecord->PutDuration = 0.;
+						FDerivedDataBackend::Get().GetRoot().PutCachedData(*CacheKey, Data, true, PutRecord);
+						PutRecord->EndTime = FPlatformTime::Seconds();
+						PutRecord->DataSize = Data.Num();
+					}
+					INC_FLOAT_STAT_BY(STAT_DDC_PutTime, bSynchronousForStats ? (float)ThisTime : 0.0f);
+				}
+			}
+			if (!bSuccess)
+			{
+				Data.Empty();
+			}
+			FDerivedDataBackend::Get().AddToAsyncCompletionCounter(-1);
+		}
+
+		FORCEINLINE TStatId GetStatId() const
+		{
+			RETURN_QUICK_DECLARE_CYCLE_STAT(FBuildAsyncWorker, STATGROUP_ThreadPoolAsyncTasks);
+		}
+
+		/** true in the case of a cache hit, otherwise the result of the deriver build call **/
+		bool							bSuccess;
+		/** true if we should record the timing **/
+		bool							bSynchronousForStats;
+		/** Data dervier we are operating on **/
+		FDerivedDataPluginInterface*	DataDeriver;
+		/** Cache key associated with this build **/
+		FString							CacheKey;
+		/** Data to return to caller, later **/
+		TArray<uint8>					Data;
+		FCacheStatRecord*				GetRecord;
+		FCacheStatRecord*				PutRecord;
+	};
+
+public:
+
+	/** Constructor, called once to cereate a singleton **/
+	FDerivedDataCache()
+		: CurrentHandle(19248) // we will skip some potential handles to catch errors
+	{
+		FDerivedDataBackend::Get(); // we need to make sure this starts before we all us to start
+	}
+
+	/** Destructor, flushes all sync tasks **/
+	~FDerivedDataCache()
+	{
+		FScopeLock ScopeLock(&SynchronizationObject);
+		for (TMap<uint32,FAsyncTask<FBuildAsyncWorker>*>::TIterator It(PendingTasks); It; ++It)
+		{
+			It.Value()->EnsureCompletion();
+			delete It.Value();
+		}
+		PendingTasks.Empty();
+	}
+
+	virtual bool GetSynchronous(FDerivedDataPluginInterface* DataDeriver, TArray<uint8>& OutData) override
+	{
+		QUICK_SCOPE_CYCLE_COUNTER(STAT_DDC_GetSynchronous);
+		check(DataDeriver);
+		FString CacheKey = FDerivedDataCache::BuildCacheKey(DataDeriver);
+		UE_LOG(LogDerivedDataCache, Verbose, TEXT("GetSynchronous %s"), *CacheKey);
+		FCacheStatRecord* GetRecord = new FCacheStatRecord();
+		FCacheStatRecord* PutRecord = new FCacheStatRecord();
+		PutRecord->CacheGet = -1;
+		{
+			FScopeLock ScopeLock(&SynchronizationObject);
+			CacheStats.Add(GetRecord);
+			CacheStats.Add(PutRecord);
+		}
+		FAsyncTask<FBuildAsyncWorker> PendingTask(DataDeriver, *CacheKey, true, GetRecord, PutRecord);
+		AddToAsyncCompletionCounter(1);
+		PendingTask.StartSynchronousTask();
+		OutData = PendingTask.GetTask().Data;
+		return PendingTask.GetTask().bSuccess;
+	}
+
+	virtual uint32 GetAsynchronous(FDerivedDataPluginInterface* DataDeriver) override
+	{
+		QUICK_SCOPE_CYCLE_COUNTER(STAT_DDC_GetAsynchronous);
+		FScopeLock ScopeLock(&SynchronizationObject);
+		uint32 Handle = NextHandle();
+		FString CacheKey = FDerivedDataCache::BuildCacheKey(DataDeriver);
+		UE_LOG(LogDerivedDataCache, Verbose, TEXT("GetAsynchronous %s"), *CacheKey);
+		bool bSync = !DataDeriver->IsBuildThreadsafe();
+		FCacheStatRecord* GetRecord = new FCacheStatRecord();
+		FCacheStatRecord* PutRecord = new FCacheStatRecord();
+		PutRecord->CacheGet = -1;
+		CacheStats.Add(GetRecord);
+		CacheStats.Add(PutRecord);
+		FAsyncTask<FBuildAsyncWorker>* AsyncTask = new FAsyncTask<FBuildAsyncWorker>(DataDeriver, *CacheKey, bSync, GetRecord, PutRecord);
+		check(!PendingTasks.Contains(Handle));
+		PendingTasks.Add(Handle,AsyncTask);
+		AddToAsyncCompletionCounter(1);
+		if (!bSync)
+		{
+			AsyncTask->StartBackgroundTask();
+		}
+		else
+		{
+			AsyncTask->StartSynchronousTask();
+		}
+		// Must return a valid handle
+		check(Handle != 0);
+		return Handle;
+	}
+
+	virtual bool PollAsynchronousCompletion(uint32 Handle) override
+	{
+		QUICK_SCOPE_CYCLE_COUNTER(STAT_DDC_PollAsynchronousCompletion);
+		FAsyncTask<FBuildAsyncWorker>* AsyncTask = NULL;
+		{
+			FScopeLock ScopeLock(&SynchronizationObject);
+			AsyncTask = PendingTasks.FindRef(Handle);
+		}
+		check(AsyncTask);
+		return AsyncTask->IsDone();
+	}
+
+	virtual void WaitAsynchronousCompletion(uint32 Handle) override
+	{
+		QUICK_SCOPE_CYCLE_COUNTER(STAT_DDC_WaitAsynchronousCompletion);
+		STAT(double ThisTime = 0);
+		{
+			SCOPE_SECONDS_COUNTER(ThisTime);
+			FAsyncTask<FBuildAsyncWorker>* AsyncTask = NULL;
+			{
+				FScopeLock ScopeLock(&SynchronizationObject);
+				AsyncTask = PendingTasks.FindRef(Handle);
+			}
+			check(AsyncTask);
+			AsyncTask->EnsureCompletion();
+		}
+		INC_FLOAT_STAT_BY(STAT_DDC_ASyncWaitTime,(float)ThisTime);
+	}
+
+	virtual bool GetAsynchronousResults(uint32 Handle, TArray<uint8>& OutData) override
+	{
+		QUICK_SCOPE_CYCLE_COUNTER(STAT_DDC_GetAsynchronousResults);
+		FAsyncTask<FBuildAsyncWorker>* AsyncTask = NULL;
+		{
+			FScopeLock ScopeLock(&SynchronizationObject);
+			PendingTasks.RemoveAndCopyValue(Handle,AsyncTask);
+		}
+		check(AsyncTask);
+		if (!AsyncTask->GetTask().bSuccess)
+		{
+			delete AsyncTask;
+			return false;
+		}
+		OutData = AsyncTask->GetTask().Data;
+		delete AsyncTask;
+		check(OutData.Num());
+		return true;
+	}
+
+	virtual IDerivedDataRollup* StartRollup() override
+	{
+		return NULL;
+	}
+
+	virtual void EndRollup(IDerivedDataRollup*& Rollup) override
+	{
+		check(!Rollup); // this needs to be handled by someone else, if rollups are disabled, then it should be NULL
+	}
+
+	virtual bool GetSynchronous(const TCHAR* CacheKey, TArray<uint8>& OutData) override
+	{
+		QUICK_SCOPE_CYCLE_COUNTER(STAT_DDC_GetSynchronous_Data);
+		UE_LOG(LogDerivedDataCache, Verbose, TEXT("GetSynchronous %s"), CacheKey);
+		FCacheStatRecord* GetRecord = new FCacheStatRecord();
+		FCacheStatRecord* PutRecord = new FCacheStatRecord();
+		PutRecord->CacheGet = -1;
+		{
+			FScopeLock ScopeLock(&SynchronizationObject);
+			CacheStats.Add(GetRecord);
+			CacheStats.Add(PutRecord);
+		}
+		FAsyncTask<FBuildAsyncWorker> PendingTask((FDerivedDataPluginInterface*)NULL, CacheKey, true, GetRecord, PutRecord);
+		AddToAsyncCompletionCounter(1);
+		PendingTask.StartSynchronousTask();
+		OutData = PendingTask.GetTask().Data;
+		return PendingTask.GetTask().bSuccess;
+	}
+
+	virtual uint32 GetAsynchronous(const TCHAR* CacheKey, IDerivedDataRollup* Rollup = NULL) override
+	{
+		QUICK_SCOPE_CYCLE_COUNTER(STAT_DDC_GetAsynchronous_Handle);
+		check(!Rollup); // this needs to be handled by someone else, if rollups are disabled, then it should be NULL
+		FScopeLock ScopeLock(&SynchronizationObject);
+		UE_LOG(LogDerivedDataCache, Verbose, TEXT("GetAsynchronous %s"), CacheKey);
+		uint32 Handle = NextHandle();
+		FCacheStatRecord* GetRecord = new FCacheStatRecord();
+		FCacheStatRecord* PutRecord = new FCacheStatRecord();
+		PutRecord->CacheGet = -1;
+		CacheStats.Add(GetRecord);
+		CacheStats.Add(PutRecord);
+		FAsyncTask<FBuildAsyncWorker>* AsyncTask = new FAsyncTask<FBuildAsyncWorker>((FDerivedDataPluginInterface*)NULL, CacheKey, false, GetRecord, PutRecord);
+		check(!PendingTasks.Contains(Handle));
+		PendingTasks.Add(Handle, AsyncTask);
+		AddToAsyncCompletionCounter(1);
+		AsyncTask->StartBackgroundTask();
+		return Handle;
+	}
+
+	/** 
+	 * Starts the async process of checking the cache and if the item is present, retrieving the cached results (version for internal use by rollups)
+	 * @param	CacheKey	Key to identify the data
+	 * @param	Rollup		Rollup pointer, if this is part of a rollup
+	 * @param	Handle		The handle that can be used for PollAsynchronousCompletion, WaitAsynchronousCompletion and GetAsynchronousResults
+	**/
+	void GetAsynchronousForRollup(const TCHAR* CacheKey, uint32 Handle)
+	{
+		QUICK_SCOPE_CYCLE_COUNTER(STAT_DDC_GetAsynchronousForRollup);
+		FScopeLock ScopeLock(&SynchronizationObject);
+		UE_LOG(LogDerivedDataCache, Verbose, TEXT("GetAsynchronous(handle) %s"), CacheKey);
+		FCacheStatRecord* GetRecord = new FCacheStatRecord();
+		FCacheStatRecord* PutRecord = new FCacheStatRecord();
+		PutRecord->CacheGet = -1;
+		CacheStats.Add(GetRecord);
+		CacheStats.Add(PutRecord);
+		FAsyncTask<FBuildAsyncWorker>* AsyncTask = new FAsyncTask<FBuildAsyncWorker>((FDerivedDataPluginInterface*)NULL, CacheKey, false, GetRecord, PutRecord);
+		check(!PendingTasks.Contains(Handle));
+		PendingTasks.Add(Handle,AsyncTask);
+		AddToAsyncCompletionCounter(1);
+		AsyncTask->StartBackgroundTask();
+	}
+
+	virtual void Put(const TCHAR* CacheKey, TArray<uint8>& Data, bool bPutEvenIfExists = false) override
+	{
+		QUICK_SCOPE_CYCLE_COUNTER(STAT_DDC_Put);
+		STAT(double ThisTime = 0);
+		{
+			SCOPE_SECONDS_COUNTER(ThisTime);
+			FCacheStatRecord* Record = new FCacheStatRecord();
+			{
+				FScopeLock ScopeLock(&SynchronizationObject);
+				CacheStats.Add(Record);
+			}
+			Record->CacheKey = CacheKey;
+			Record->CacheGet = 0;
+			Record->bSynchronous = true;
+			Record->DataSize = Data.Num();
+			Record->StartTime = FPlatformTime::Seconds();
+			Record->PutDuration = 0;
+			Record->GetDuration = 0;
+			FDerivedDataBackend::Get().GetRoot().PutCachedData(CacheKey, Data, bPutEvenIfExists);
+			Record->EndTime = FPlatformTime::Seconds();
+		}
+		INC_FLOAT_STAT_BY(STAT_DDC_PutTime,(float)ThisTime);
+		INC_DWORD_STAT(STAT_DDC_NumPuts);
+	}
+
+	virtual void MarkTransient(const TCHAR* CacheKey) override
+	{
+		FDerivedDataBackend::Get().GetRoot().RemoveCachedData(CacheKey, /*bTransient=*/ true);
+	}
+
+	virtual bool CachedDataProbablyExists(const TCHAR* CacheKey) override
+	{
+		QUICK_SCOPE_CYCLE_COUNTER(STAT_DDC_CachedDataProbablyExists);
+		bool bResult;
+		INC_DWORD_STAT(STAT_DDC_NumExist);
+		STAT(double ThisTime = 0);
+		{
+			SCOPE_SECONDS_COUNTER(ThisTime);
+			bResult = FDerivedDataBackend::Get().GetRoot().CachedDataProbablyExists(CacheKey);
+		}
+		INC_FLOAT_STAT_BY(STAT_DDC_ExistTime, (float)ThisTime);
+		return bResult;
+	}
+
+	void NotifyBootComplete() override
+	{
+		QUICK_SCOPE_CYCLE_COUNTER(STAT_DDC_NotifyBootComplete);
+		FDerivedDataBackend::Get().NotifyBootComplete();
+	}
+
+	void AddToAsyncCompletionCounter(int32 Addend) override
+	{
+		FDerivedDataBackend::Get().AddToAsyncCompletionCounter(Addend);
+	}
+
+	void WaitForQuiescence(bool bShutdown) override
+	{
+		QUICK_SCOPE_CYCLE_COUNTER(STAT_DDC_WaitForQuiescence);
+		FDerivedDataBackend::Get().WaitForQuiescence(bShutdown);
+	}
+
+	void GetDirectories(TArray<FString>& OutResults) override
+	{
+		FDerivedDataBackend::Get().GetDirectories(OutResults);
+	}
+
+	/** Called at ShutdownModule() time to print out status before we're cleaned up */
+	virtual void PrintLeaks()
+	{
+		// Used by derived classes to spit out leaked pending rollups
+	}
+
+	virtual void PrintStats()
+	{
+		FScopeLock ScopeLock(&SynchronizationObject);
+		if (CacheStats.Num() > 0)
+		{
+#if ALLOW_DEBUG_FILES
+			const FString AbsLog = FPlatformOutputDevices::GetAbsoluteLogFilename();
+			const FString OutputDir = FPaths::GetPath(AbsLog);
+			IFileManager::Get().MakeDirectory(*OutputDir, true);
+
+			// Create archive for log data.
+			const FString ChartName = OutputDir / FString(TEXT("CookRunDDC.txt"));
+			FArchive* OutputFile = IFileManager::Get().CreateDebugFileWriter(*ChartName);
+
+			if (OutputFile)
+			{
+				OutputFile->Logf(TEXT("Key,Sync,Network,Size,Total Duration(ms),Get Duration(ms),Put Duration(ms),Get"));
+				for (int32 Index = 0; Index < CacheStats.Num(); Index++)
+				{
+					FCacheStatRecord* Record = CacheStats[Index];
+					if (Record->CacheGet != -1)
+					{
+						OutputFile->Logf(TEXT("%s,%s,%s,%d,%lf,%lf,%lf,%d"), *(Record->CacheKey), Record->bSynchronous ? TEXT("true") : TEXT("false"), Record->bFromNetwork ? TEXT("true") : TEXT("false"), Record->DataSize, (Record->EndTime - Record->StartTime) * 1000.0f, Record->GetDuration * 1000.0f, Record->PutDuration * 1000.0f, Record->CacheGet);
+					}
+				}
+				delete OutputFile;
+			}
+#endif
+			for (int32 Index = 0; Index < CacheStats.Num(); Index++)
+			{
+				FCacheStatRecord* Record = CacheStats[Index];
+				delete Record;
+			}
+			CacheStats.Empty();
+		}
+	}
+
+protected:
+	uint32 NextHandle()
+	{
+		return (uint32)CurrentHandle.Increment();
+	}
+
+
+private:
+
+	/** 
+	 * Internal function to build a cache key out of the plugin name, versions and plugin specific info
+	 * @param	DataDeriver	plugin to produce the elements of the cache key.
+	 * @return				Assembled cache key
+	**/
+	static FString BuildCacheKey(FDerivedDataPluginInterface* DataDeriver)
+	{
+		FString Result = FDerivedDataCacheInterface::BuildCacheKey(DataDeriver->GetPluginName(), DataDeriver->GetVersionString(), *DataDeriver->GetPluginSpecificCacheKeySuffix());
+		return Result;
+	}
+
+
+	/** Counter used to produce unique handles **/
+	FThreadSafeCounter			CurrentHandle;
+	/** Object used for synchronization via a scoped lock **/
+	FCriticalSection			SynchronizationObject;
+	/** Map of handle to pending task **/
+	TMap<uint32,FAsyncTask<FBuildAsyncWorker>*>	PendingTasks;
+
+	// TEMP: stats
+	TArray<FCacheStatRecord*> CacheStats;
+};
+
+//Forward reference
+FDerivedDataCache& InternalSingleton();
+
+namespace EPhase
+{
+	enum Type
+	{
+		Adding,
+		AsyncRollupGet,
+		AsyncRollupGetSucceed,
+		AsyncRollupGetFailed_GettingItemsAndWaitingForPuts,
+		Done,
+	};
+}
+
+/** 
+ * Opaque class for rollup handling
+**/
+class FDerivedDataRollup : public IDerivedDataRollup
+{
+	/** Magic numbers to verify integrity and check endianness **/
+	enum
+	{
+		MAGIC=0x9E1B83C1,
+		MAGIC_SWAPPED=0xC1831B9E
+	};
+
+	/** Helper structure for an element of a rollup **/
+	struct FRollupItem
+	{
+		FRollupItem(const FString& InCacheKey, uint32 InAsyncHandle)
+			: CacheKey(InCacheKey)
+			, AsyncHandle(InAsyncHandle)
+			, FinishedFromThePerspectiveOfTheCaller(false)
+		{
+		}
+		/** Cache key for this item **/
+		FString CacheKey;
+		/** Async handle for this item, used both to return to original caller, and for calls to the actual DDC **/
+		uint32 AsyncHandle;
+		/** Payload of this item, used for both from the get of the rollup and a put to the rollup **/
+		TArray<uint8> Payload;
+		/** If true, then the caller has either asked for the results. This means we don't need to keep them any more. **/
+		bool FinishedFromThePerspectiveOfTheCaller;
+	};
+
+	/** Items in this rollup **/
+	TArray<FRollupItem> Items;
+	/** Redundant copy of the keys in this rollup **/
+	TSet<FString> CacheKeys;
+	/** Redundant copy of the async handles in this rollup **/
+	TSet<uint32> AsyncHandles;
+	/** Cache key for the rollup itself **/
+	FString RollupCacheKey;
+	/** Async handle for the rollup **/
+	uint32 RollupAsyncHandle;
+	/** Tracks the phase this rollup is in. Mostly used for check's **/
+	EPhase::Type CurrentPhase;
+	/** If true, the rollup was corrupted, so we need to force a put when we get to the put. **/
+	bool ForcePutForCorruption;
+
+	/** Called when the rollup is ready. This is indirectly caused by the original caller waiting for an item to be ready. **/
+	void GetRollupResults()
+	{
+		check(Items.Num());
+		check(CurrentPhase == EPhase::AsyncRollupGet);
+		TArray<uint8> Payload;
+		bool bFailed = true;
+		if (InternalSingleton().FDerivedDataCache::GetAsynchronousResults(RollupAsyncHandle, Payload))
+		{
+			ForcePutForCorruption = true;
+			if (Payload.Num() > sizeof(uint32) * 2)
+			{
+				FMemoryReader Ar(Payload);
+				uint32 Magic;
+				Ar << Magic;
+				if (Magic == MAGIC_SWAPPED)
+				{
+					Ar.SetByteSwapping(!Ar.ForceByteSwapping());
+					Magic = MAGIC;
+				}
+				if (Magic == MAGIC)
+				{
+					int32 Count;
+					Ar << Count;
+					if (Count == Items.Num())
+					{
+						bFailed = false;
+						for (int32 Index = 0; Index < Items.Num(); Index++)
+						{
+							FString Key;
+							Ar << Key;
+							if (Key != Items[Index].CacheKey)
+							{
+								bFailed = true;
+								break;
+							}
+							Ar << Items[Index].Payload;
+							if (!Items[Index].Payload.Num())
+							{
+								bFailed = true;
+								break;
+							}
+						}
+					}
+				}
+			}
+		}
+		if (!bFailed)
+		{
+			ForcePutForCorruption = false;
+			CurrentPhase = EPhase::AsyncRollupGetSucceed;
+		}
+		else
+		{
+			CurrentPhase = EPhase::AsyncRollupGetFailed_GettingItemsAndWaitingForPuts;
+			for (int32 Index = 0; Index < Items.Num(); Index++)
+			{
+				Items[Index].Payload.Empty(); // we might have had partial success on a corrupted rollup; we won't accept those
+				InternalSingleton().FDerivedDataCache::GetAsynchronousForRollup(*Items[Index].CacheKey, Items[Index].AsyncHandle);
+			}
+		}
+	}
+
+	/** Tests to see if the rollup is complete and ready to be put; if it is, it packages it and puts it. **/
+	bool CheckForPut()
+	{
+		check(Items.Num());
+		check(CurrentPhase == EPhase::AsyncRollupGetFailed_GettingItemsAndWaitingForPuts);
+		for (int32 Index = 0; Index < Items.Num(); Index++)
+		{
+			if (!Items[Index].Payload.Num())
+			{
+				return false; // not done yet because we don't have all of the data
+			}
+			if (!Items[Index].FinishedFromThePerspectiveOfTheCaller)
+			{
+				return false; // not done yet because the caller still hasn't retrieved their results
+			}
+		}
+		uint32 Magic = MAGIC;
+		int32 Count = Items.Num();
+		TArray<uint8> Buffer;
+		FMemoryWriter Ar(Buffer);
+		Ar << Magic;
+		Ar << Count;
+		for (int32 Index = 0; Index < Items.Num(); Index++)
+		{
+			Ar << Items[Index].CacheKey;
+			Ar << Items[Index].Payload;
+		}
+		InternalSingleton().FDerivedDataCache::Put(*RollupCacheKey, Buffer, ForcePutForCorruption);
+		CurrentPhase = EPhase::Done;
+		return true;
+	}
+
+	/** 
+	 * Checks to see if there is any reason for this rollup to stay alive. 
+	 * @return true when done
+	**/
+	bool CheckForDone()
+	{
+		check(Items.Num());
+		check(CurrentPhase == EPhase::AsyncRollupGetSucceed);
+		for (int32 Index = 0; Index < Items.Num(); Index++)
+		{
+			if (!Items[Index].FinishedFromThePerspectiveOfTheCaller)
+			{
+				return false; // not done yet because the caller still hasn't retrieved their results
+			}
+			if (!Items[Index].Payload.Num())
+			{
+				check(0); // what is this?
+			}
+		}
+		CurrentPhase = EPhase::Done;
+		return true;
+	}
+
+	/** 
+	 * Finds an item by async handle, not legal to call if this rollup does not contain this handle
+	 * @param Handle		Async handle to find the corresponding item of
+	 * @return				Item with this async handle
+	**/
+	FRollupItem& FindItem(uint32 Handle)
+	{
+		for (int32 Index = 0; Index < Items.Num(); Index++)
+		{
+			if (Items[Index].AsyncHandle == Handle)
+			{
+				return Items[Index];
+			}
+		}
+		check(0);
+		return (FRollupItem&)*(FRollupItem*)NULL;
+	}
+
+	/** 
+	 * Finds an item by cache keu, not legal to call if this rollup does not contain this cache key
+	 * @param InCacheKey	cache key to find the corresponding item of
+	 * @return				Item with this cache key
+	**/
+	FRollupItem& FindItem(const FString& InCacheKey)
+	{
+		for (int32 Index = 0; Index < Items.Num(); Index++)
+		{
+			if (Items[Index].CacheKey == InCacheKey)
+			{
+				return Items[Index];
+			}
+		}
+		check(0);
+		return (FRollupItem&)*(FRollupItem*)NULL;
+	}
+
+public:
+	FDerivedDataRollup()
+		: RollupCacheKey(TEXT("ROLLUP_"))
+		, RollupAsyncHandle(0)
+		, CurrentPhase(EPhase::Adding)
+		, ForcePutForCorruption(false)
+	{
+	}
+
+	/** Return the cache key, used for error spew **/
+	FString GetName()
+	{
+		return RollupCacheKey;
+	}
+
+	/** Return true if this rollup can be deleted because it has completed its life cycle. **/
+	bool IsDone()
+	{
+		return CurrentPhase == EPhase::Done;
+	}
+
+	/** Return true if this rollup contains an item with the given cache key. **/
+	bool Contains(const FString& InCacheKey)
+	{
+		return CacheKeys.Contains(InCacheKey);
+	}
+
+	/** Return true if this rollup contains an item with the given async handle. **/
+	bool Contains(uint32 InAsyncHandle)
+	{
+		return AsyncHandles.Contains(InAsyncHandle);
+	}
+
+	/** Add a new item to this rollup with the given cache key and async handle. **/
+	void Add(const FString& InCacheKey, uint32 InAsyncHandle)
+	{
+		check(CurrentPhase == EPhase::Adding);
+		RollupCacheKey += InCacheKey;
+		CacheKeys.Add(InCacheKey);
+		AsyncHandles.Add(InAsyncHandle);
+		new (Items) FRollupItem(InCacheKey, InAsyncHandle);
+	}
+
+	/** Signifies the end of the adding phase and starts an async get of the rollup. **/
+	void Close()
+	{
+		check(CurrentPhase == EPhase::Adding);
+		if (Items.Num())
+		{
+			RollupAsyncHandle = InternalSingleton().FDerivedDataCache::GetAsynchronous(*RollupCacheKey);
+			CurrentPhase = EPhase::AsyncRollupGet;
+		}
+		else
+		{
+			CurrentPhase = EPhase::Done;
+		}
+	}
+
+	/** 
+	 * Handle PollAsynchronousCompletion from the calling code. 
+	 * @param Handle	Async completion handle for the item
+	 * @return			true, if the calling code can request results yet
+	 **/
+	bool PollAsynchronousCompletion(uint32 Handle)
+	{
+		check(Contains(Handle));
+		if (CurrentPhase == EPhase::AsyncRollupGet)
+		{
+			// in this phase we see if the rollup is done
+			if (!InternalSingleton().FDerivedDataCache::PollAsynchronousCompletion(RollupAsyncHandle))
+			{
+				return false;
+			}
+			GetRollupResults();
+			// fall through to handle the other cases
+		}
+		if (CurrentPhase == EPhase::AsyncRollupGetSucceed)
+		{
+			// Rollup succeeded, so the calling code can get the results
+			return true;
+		}
+		if (CurrentPhase == EPhase::AsyncRollupGetFailed_GettingItemsAndWaitingForPuts)
+		{
+			// Rollup failed, so call the actual PAC for the individual item
+			return InternalSingleton().FDerivedDataCache::PollAsynchronousCompletion(Handle);
+		}
+		check(0); // bad phase
+		return false;
+	}
+
+	/** 
+	 * Handle WaitAsynchronousCompletion from the calling code. 
+	 * @param Handle	Async completion handle for the item
+	 **/
+	void WaitAsynchronousCompletion(uint32 Handle)
+	{
+		check(Contains(Handle));
+		if (CurrentPhase == EPhase::AsyncRollupGet)
+		{
+			// in this phase we wait for the rollup to complete, then deal with the results
+			InternalSingleton().FDerivedDataCache::WaitAsynchronousCompletion(RollupAsyncHandle);
+			GetRollupResults();
+			// fall through to handle the other cases
+		}
+		if (CurrentPhase == EPhase::AsyncRollupGetSucceed)
+		{
+			// Rollup succeeded, so the calling code can get the results
+			return;
+		}
+		if (CurrentPhase == EPhase::AsyncRollupGetFailed_GettingItemsAndWaitingForPuts)
+		{
+			// Rollup failed, so call the actual WAC for the individual item
+			InternalSingleton().FDerivedDataCache::WaitAsynchronousCompletion(Handle);
+			return;
+		}
+		check(0); // bad phase
+	}
+
+	/** 
+	 * Handle GetAsynchronousResults from the calling code. If this is the last piece of data, the rollup will be put.
+	 * @param Handle	Async completion handle for the item
+	 * @param OutData	returned payload
+	 * @return			true if the payload contains data and everything is peachy
+	 **/
+	bool GetAsynchronousResults(uint32 Handle, TArray<uint8>& OutData)
+	{
+		check(Contains(Handle));
+		FRollupItem& Item = FindItem(Handle);
+		Item.FinishedFromThePerspectiveOfTheCaller = true;
+		OutData.Empty();
+		if (CurrentPhase == EPhase::AsyncRollupGetSucceed)
+		{
+			OutData = Item.Payload;
+			CheckForDone();
+			return !!OutData.Num();
+		}
+		if (CurrentPhase == EPhase::AsyncRollupGetFailed_GettingItemsAndWaitingForPuts)
+		{
+			if (InternalSingleton().FDerivedDataCache::GetAsynchronousResults(Handle, OutData))
+			{
+				Item.Payload = OutData;
+				CheckForPut();
+			}
+			return !!OutData.Num();
+		}
+		check(0); // bad phase
+		return false;
+	}
+	/** 
+	 * Handle Put from the calling code. If this is the last piece of data, the rollup will be put.
+	 * @param CacheKey	Cache key for the item
+	 * @param Data		returned payload
+	 **/
+	void Put(const TCHAR* CacheKey, TArray<uint8>& Data)
+	{
+		if (CurrentPhase != EPhase::AsyncRollupGetFailed_GettingItemsAndWaitingForPuts)
+		{
+			return;
+		}
+		check(Contains(CacheKey));
+		FRollupItem& Item = FindItem(CacheKey);
+		check(Data.Num());
+		Item.Payload = Data;
+		CheckForPut();
+	}
+};
+
+
+
+/** 
+ * Implementation of the derived data cache, this layer implements rollups
+**/
+class FDerivedDataCacheWithRollups : public FDerivedDataCache
+{
+	typedef FDerivedDataCache Super;
+public:
+
+	~FDerivedDataCacheWithRollups()
+	{
+	}
+
+	virtual void PrintLeaks() override
+	{
+		FScopeLock ScopeLock(&SynchronizationObject);
+		UE_LOG(LogDerivedDataCache, Log, TEXT("Shutdown"));
+		for (TSet<FDerivedDataRollup*>::TIterator Iter(PendingRollups); Iter; ++Iter)
+		{
+			FString Name((*Iter)->GetName());
+			if (Name.Len() > 1024)
+			{
+				Name = Name.Left(1024);
+				Name += TEXT("...");
+			}
+			UE_LOG(LogDerivedDataCache, Warning, TEXT("Leaked Rollup! %s"), *Name);
+		}
+		//check(!PendingRollups.Num()); // leaked rollups
+	}
+
+	virtual bool PollAsynchronousCompletion(uint32 Handle) override
+	{
+		FScopeLock ScopeLock(&SynchronizationObject);
+		for (TSet<FDerivedDataRollup*>::TIterator Iter(PendingRollups); Iter; ++Iter)
+		{
+			if ((*Iter)->Contains(Handle))
+			{
+				return (*Iter)->PollAsynchronousCompletion(Handle);
+			}
+		}
+		return Super::PollAsynchronousCompletion(Handle);
+	}
+
+	virtual void WaitAsynchronousCompletion(uint32 Handle) override
+	{
+		STAT(double ThisTime = 0);
+		{
+			SCOPE_SECONDS_COUNTER(ThisTime);
+			FScopeLock ScopeLock(&SynchronizationObject);
+			for (TSet<FDerivedDataRollup*>::TIterator Iter(PendingRollups); Iter; ++Iter)
+			{
+				if ((*Iter)->Contains(Handle))
+				{
+					(*Iter)->WaitAsynchronousCompletion(Handle);
+					return;
+				}
+			}
+		}
+		INC_FLOAT_STAT_BY(STAT_DDC_ASyncWaitTime,(float)ThisTime);
+		Super::WaitAsynchronousCompletion(Handle);
+	}
+
+	virtual bool GetAsynchronousResults(uint32 Handle, TArray<uint8>& OutData) override
+	{
+		FScopeLock ScopeLock(&SynchronizationObject);
+		for (TSet<FDerivedDataRollup*>::TIterator Iter(PendingRollups); Iter; ++Iter)
+		{
+			if ((*Iter)->Contains(Handle))
+			{
+				bool Result = (*Iter)->GetAsynchronousResults(Handle, OutData);
+				if ((*Iter)->IsDone())
+				{
+					delete *Iter;
+					Iter.RemoveCurrent();
+				}
+				return Result;
+			}
+		}
+		return Super::GetAsynchronousResults(Handle, OutData);
+	}
+
+	virtual IDerivedDataRollup* StartRollup() override
+	{
+		FScopeLock ScopeLock(&SynchronizationObject);
+		FDerivedDataRollup* Result = new FDerivedDataRollup();
+		PendingRollups.Add(Result);
+		return Result;
+	}
+
+	virtual void EndRollup(IDerivedDataRollup*& InRollup) override
+	{
+		FDerivedDataRollup* Rollup = static_cast<FDerivedDataRollup*>(InRollup);
+		if (Rollup)
+		{
+			FScopeLock ScopeLock(&SynchronizationObject);
+			Rollup->Close();
+			if (Rollup->IsDone())
+			{
+				PendingRollups.Remove(Rollup);
+				delete Rollup;
+			}
+			InRollup = NULL; // set the pointer to NULL so it cannot be reused
+		}
+	}
+
+	virtual uint32 GetAsynchronous(const TCHAR* CacheKey, IDerivedDataRollup* Rollup = NULL) override
+	{
+		if (Rollup)
+		{
+			FScopeLock ScopeLock(&SynchronizationObject);
+			UE_LOG(LogDerivedDataCache, Verbose, TEXT("GetAsynchronous (Rollup) %s"), CacheKey);
+			uint32 CurrentHandle = NextHandle();
+			((FDerivedDataRollup*)Rollup)->Add(CacheKey, CurrentHandle);
+			return CurrentHandle;
+		}
+		return Super::GetAsynchronous(CacheKey, NULL);
+	}
+
+	virtual void Put(const TCHAR* CacheKey, TArray<uint8>& Data, bool bPutEvenIfExists = false) override
+	{
+		FScopeLock ScopeLock(&SynchronizationObject);
+		for (TSet<FDerivedDataRollup*>::TIterator Iter(PendingRollups); Iter; ++Iter)
+		{
+			if ((*Iter)->Contains(CacheKey))
+			{
+				(*Iter)->Put(CacheKey, Data);
+				if ((*Iter)->IsDone())
+				{
+					delete *Iter;
+					Iter.RemoveCurrent();
+				}
+			}
+		}
+		return Super::Put(CacheKey, Data, bPutEvenIfExists);
+	}
+
+private:
+
+	/** Object used for synchronization via a scoped lock **/
+	FCriticalSection			SynchronizationObject;
+	/** Set of rollups **/
+	TSet<FDerivedDataRollup*>	PendingRollups;
+};
+
+/**
+ * Singleton used both internally, and through the module.
+ * We look at the commandline to check if we should disable rollups or not
+ */
+FDerivedDataCache& InternalSingleton()
+{
+	static FDerivedDataCache* Singleton = NULL;
+	if (!Singleton)
+	{
+		if (FParse::Param(FCommandLine::Get(), TEXT("DDCNoRollups")))
+		{
+			UE_LOG(LogDerivedDataCache, Warning, TEXT("Rollups are disabled."));
+			static FDerivedDataCache SingletonInstance;
+			Singleton = &SingletonInstance;
+		}
+		else
+		{
+			static FDerivedDataCacheWithRollups SingletonInstance;
+			Singleton = &SingletonInstance;
+		}
+	}
+	return *Singleton;
+}
+
+/**
+ * Module for the DDC
+ */
+class FDerivedDataCacheModule : public IDerivedDataCacheModule
+{
+public:
+	virtual FDerivedDataCacheInterface& GetDDC() override
+	{
+		return InternalSingleton();
+	}
+
+	virtual void ShutdownModule() override
+	{
+		FDDCCleanup::Shutdown();
+
+		FDerivedDataCache& DDC = static_cast< FDerivedDataCache& >( GetDDC() );
+		DDC.PrintLeaks();
+		DDC.PrintStats();
+	}
+};
+
+IMPLEMENT_MODULE( FDerivedDataCacheModule, DerivedDataCache);
+
diff --git a/Engine/Source/Developer/DerivedDataCache/Private/DerivedDataLimitKeyLengthWrapper.h b/Engine/Source/Developer/DerivedDataCache/Private/DerivedDataLimitKeyLengthWrapper.h
new file mode 100644
--- /dev/null
+++ b/Engine/Source/Developer/DerivedDataCache/Private/DerivedDataLimitKeyLengthWrapper.h
@@ -0,0 +1,187 @@
+// Copyright 1998-2015 Epic Games, Inc. All Rights Reserved.
+
+#pragma once
+
+/** 
+ * A backend wrapper that limits the key size and uses hashing...in this case it wraps the payload and the payload contains the full key to verify the integrity of the hash
+**/
+class FDerivedDataLimitKeyLengthWrapper : public FDerivedDataBackendInterface
+{
+public:
+
+	/**
+	 * Constructor
+	 *
+	 * @param	InInnerBackend	Backend to use for storage, my responsibilities are about key length
+	 */
+	FDerivedDataLimitKeyLengthWrapper(FDerivedDataBackendInterface* InInnerBackend, int32 InMaxKeyLength)
+		: InnerBackend(InInnerBackend)
+		, MaxKeyLength(InMaxKeyLength)
+	{
+		check(InnerBackend);
+	}
+
+	/** return true if this cache is writable **/
+	virtual bool IsWritable() override
+	{
+		return InnerBackend->IsWritable();
+	}
+
+	/**
+	 * Synchronous test for the existence of a cache item
+	 *
+	 * @param	CacheKey	Alphanumeric+underscore key of this cache item
+	 * @return				true if the data probably will be found, this can't be guaranteed because of concurrency in the backends, corruption, etc
+	 */
+	virtual bool CachedDataProbablyExists(const TCHAR* CacheKey) override
+	{
+		FString NewKey;
+		ShortenKey(CacheKey, NewKey);
+		return InnerBackend->CachedDataProbablyExists(*NewKey);
+	}
+	/**
+	 * Synchronous retrieve of a cache item
+	 *
+	 * @param	CacheKey	Alphanumeric+underscore key of this cache item
+	 * @param	OutData		Buffer to receive the results, if any were found
+	 * @return				true if any data was found, and in this case OutData is non-empty
+	 */
+	virtual bool GetCachedData(const TCHAR* CacheKey, TArray<uint8>& OutData, FCacheStatRecord* Stats) override
+	{
+		FString NewKey;
+		bool bOk;
+		if (!ShortenKey(CacheKey, NewKey))
+		{
+			// no shortening needed
+			bOk = InnerBackend->GetCachedData(CacheKey, OutData, Stats);
+			// look for old bug
+			if (FString(CacheKey).StartsWith(TEXT("TEXTURE2D_0002")))
+			{
+				int32 KeyLen = FCString::Strlen(CacheKey) + 1;
+				if (OutData.Num() > KeyLen && OutData.Last() == 0)
+				{
+					int32 Compare = FCStringAnsi::Strcmp(TCHAR_TO_ANSI(CacheKey), (char*)&OutData[OutData.Num() - KeyLen]);
+					if (Compare == 0)
+					{
+						UE_LOG(LogDerivedDataCache, Warning, TEXT("FDerivedDataLimitKeyLengthWrapper: Fixed old bug %s."), CacheKey);
+						OutData.RemoveAt(OutData.Num() - KeyLen, KeyLen);
+					}
+				}
+			}
+		}
+		else
+		{
+			bOk = InnerBackend->GetCachedData(*NewKey, OutData, Stats);
+			if (bOk)
+			{
+				int32 KeyLen = FCString::Strlen(CacheKey) + 1;
+				if (OutData.Num() < KeyLen)
+				{
+					UE_LOG(LogDerivedDataCache, Warning, TEXT("FDerivedDataLimitKeyLengthWrapper: Short file or Hash Collision, ignoring and deleting %s."), CacheKey);
+					bOk	= false;
+				}
+				else
+				{
+					int32 Compare = FCStringAnsi::Strcmp(TCHAR_TO_ANSI(CacheKey), (char*)&OutData[OutData.Num() - KeyLen]);
+					OutData.RemoveAt(OutData.Num() - KeyLen, KeyLen);
+					if (Compare == 0)
+					{
+						UE_LOG(LogDerivedDataCache, Verbose, TEXT("FDerivedDataLimitKeyLengthWrapper: cache hit, key match is ok %s"), CacheKey);
+					}
+					else
+					{
+						UE_LOG(LogDerivedDataCache, Warning, TEXT("FDerivedDataLimitKeyLengthWrapper: HASH COLLISION, ignoring and deleting %s."), CacheKey);
+						bOk	= false;
+					}
+				}
+				if (!bOk)
+				{
+					// _we_ detected corruption, so _we_ will force a flush of the corrupted data
+					InnerBackend->RemoveCachedData(*NewKey, /*bTransient=*/ false);
+				}
+			}
+		}
+		if (!bOk)
+		{
+			OutData.Empty();
+		}
+		return bOk;
+	}
+	/**
+	 * Asynchronous, fire-and-forget placement of a cache item
+	 *
+	 * @param	CacheKey	Alphanumeric+underscore key of this cache item
+	 * @param	InData		Buffer containing the data to cache, can be destroyed after the call returns, immediately
+	 * @param	bPutEvenIfExists	If true, then do not attempt skip the put even if CachedDataProbablyExists returns true
+	 */
+	virtual void PutCachedData(const TCHAR* CacheKey, TArray<uint8>& InData, bool bPutEvenIfExists, FCacheStatRecord* Stats) override
+	{
+		if (!InnerBackend->IsWritable())
+		{
+			return; // no point in continuing down the chain
+		}
+		FString NewKey;
+		if (!ShortenKey(CacheKey, NewKey))
+		{
+			// no shortening needed
+			InnerBackend->PutCachedData(CacheKey, InData, bPutEvenIfExists, Stats);
+			return;
+		}
+		TArray<uint8> Data(InData);
+		check(Data.Num());
+		int32 KeyLen = FCString::Strlen(CacheKey) + 1;
+		Data.AddUninitialized(KeyLen);
+		FCStringAnsi::Strcpy((char*)&Data[Data.Num() - KeyLen], KeyLen, TCHAR_TO_ANSI(CacheKey));
+		check(Data.Last()==0);
+		InnerBackend->PutCachedData(*NewKey, Data, bPutEvenIfExists, Stats);
+	}
+
+	virtual void RemoveCachedData(const TCHAR* CacheKey, bool bTransient) override
+	{
+		if (!InnerBackend->IsWritable())
+		{
+			return; // no point in continuing down the chain
+		}
+		FString NewKey;
+		ShortenKey(CacheKey, NewKey);
+		return InnerBackend->RemoveCachedData(*NewKey, bTransient);
+	}
+private:
+
+	/** Shorten the cache key and return true if shortening was required **/
+	bool ShortenKey(const TCHAR* CacheKey, FString& Result)
+	{
+		Result = FString(CacheKey);
+		if (Result.Len() <= MaxKeyLength)
+		{
+			return false;
+		}
+
+		FSHA1 HashState;
+		int32 Length = Result.Len();
+		HashState.Update((const uint8*)&Length, sizeof(int32));
+
+		auto ResultSrc = StringCast<UCS2CHAR>(*Result);
+		uint32 CRCofPayload(FCrc::MemCrc32(ResultSrc.Get(), Length * sizeof(UCS2CHAR)));
+
+		HashState.Update((const uint8*)&CRCofPayload, sizeof(uint32));
+		HashState.Update((const uint8*)ResultSrc.Get(), Length * sizeof(UCS2CHAR));
+
+		HashState.Final();
+		uint8 Hash[FSHA1::DigestSize];
+		HashState.GetHash(Hash);
+		FString HashString = BytesToHex(Hash, FSHA1::DigestSize);
+
+		int32 HashStringSize = HashString.Len();
+		int32 OriginalPart = MaxKeyLength - HashStringSize - 2;
+		Result = Result.Left(OriginalPart) + TEXT("__") + HashString;
+		check(Result.Len() == MaxKeyLength && Result.Len() > 0);
+		return true;
+	}
+
+	/** Backend to use for storage, my responsibilities are about key length **/
+	FDerivedDataBackendInterface* InnerBackend;
+
+	/** Maximum size of the backend key length **/
+	int32 MaxKeyLength;
+};
diff --git a/Engine/Source/Developer/DerivedDataCache/Private/FileSystemDerivedDataBackend.cpp b/Engine/Source/Developer/DerivedDataCache/Private/FileSystemDerivedDataBackend.cpp
new file mode 100644
--- /dev/null
+++ b/Engine/Source/Developer/DerivedDataCache/Private/FileSystemDerivedDataBackend.cpp
@@ -0,0 +1,356 @@
+// Copyright 1998-2015 Epic Games, Inc. All Rights Reserved.
+
+
+#include "Core.h"
+
+#include "DerivedDataBackendInterface.h"
+#include "DDCCleanup.h"
+
+#include "DDCStatsHelper.h"
+
+#define MAX_BACKEND_KEY_LENGTH (120)
+#define MAX_BACKEND_NUMBERED_SUBFOLDER_LENGTH (9)
+#if PLATFORM_LINUX	// PATH_MAX on Linux is 4096 (getconf PATH_MAX /, also see limits.h), so this value can be larger (note that it is still arbitrary).
+                    // This should not affect sharing the cache between platforms as the absolute paths will be different anyway.
+	#define MAX_CACHE_DIR_LEN (3119)
+#else
+	#define MAX_CACHE_DIR_LEN (119)
+#endif // PLATFORM_LINUX
+#define MAX_CACHE_EXTENTION_LEN (4)
+
+
+
+
+/** 
+ * Cache server that uses the OS filesystem
+ * The entire API should be callable from any thread (except the singleton can be assumed to be called at least once before concurrent access).
+**/
+class FFileSystemDerivedDataBackend : public FDerivedDataBackendInterface
+{
+public:
+	/**
+	 * Constructor that should only be called once by the singleton, grabs the cache path from the ini
+	 * @param InCacheDirectory	directory to store the cache in
+	 * @param bForceReadOnly	if true, do not attempt to write to this cache
+	*/
+	FFileSystemDerivedDataBackend(const TCHAR* InCacheDirectory, bool bForceReadOnly, bool bTouchFiles, bool bPurgeTransientData, bool bDeleteOldFiles, int32 InDaysToDeleteUnusedFiles, int32 InMaxNumFoldersToCheck, int32 InMaxContinuousFileChecks)
+		: CachePath(InCacheDirectory)
+		, bReadOnly(bForceReadOnly)
+		, bFailed(true)
+		, bTouch(bTouchFiles)
+		, bPurgeTransient(bPurgeTransientData)
+		, DaysToDeleteUnusedFiles(InDaysToDeleteUnusedFiles)
+	{
+		// If we find a platform that has more stingent limits, this needs to be rethought.
+		static_assert(MAX_BACKEND_KEY_LENGTH + MAX_CACHE_DIR_LEN + MAX_BACKEND_NUMBERED_SUBFOLDER_LENGTH + MAX_CACHE_EXTENTION_LEN < PLATFORM_MAX_FILEPATH_LENGTH,
+			"Not enough room left for cache keys in max path.");
+		const double SlowInitDuration = 10.0;
+		double AccessDuration = 0.0;
+
+		check(CachePath.Len());
+		FPaths::NormalizeFilename(CachePath);
+		const FString AbsoluteCachePath = IFileManager::Get().ConvertToAbsolutePathForExternalAppForRead(*CachePath);
+		if (AbsoluteCachePath.Len() > MAX_CACHE_DIR_LEN)
+		{
+			const FText ErrorMessage = FText::Format( NSLOCTEXT("DerivedDataCache", "PathTooLong", "Cache path {0} is longer than {1} characters...please adjust [DerivedDataBackendGraph] paths to be shorter (this leaves more room for cache keys)."), FText::FromString( AbsoluteCachePath ), FText::AsNumber( MAX_CACHE_DIR_LEN ) );
+			FMessageDialog::Open(EAppMsgType::Ok, ErrorMessage);
+			UE_LOG(LogDerivedDataCache, Fatal, TEXT("%s"), *ErrorMessage.ToString());
+		}
+		if (!bReadOnly)
+		{
+			double TestStart = FPlatformTime::Seconds();
+			FString TempFilename = CachePath / FGuid::NewGuid().ToString() + ".tmp";
+			FFileHelper::SaveStringToFile(FString("TEST"),*TempFilename);
+			int32 TestFileSize = IFileManager::Get().FileSize(*TempFilename);
+			if (TestFileSize < 4)
+			{
+				UE_LOG(LogDerivedDataCache, Warning, TEXT("Fail to write to %s, derived data cache to this directory will be read only."),*CachePath);
+			}
+			else
+			{
+				bFailed = false;
+			}			
+			if (TestFileSize >= 0)
+			{
+				IFileManager::Get().Delete(*TempFilename, false, false, true);
+			}
+			AccessDuration = FPlatformTime::Seconds() - TestStart;
+		}
+		if (bFailed)
+		{
+			double StartTime = FPlatformTime::Seconds();
+			TArray<FString> FilesAndDirectories;
+			IFileManager::Get().FindFiles(FilesAndDirectories,*(CachePath / TEXT("*.*")), true, true);
+			AccessDuration = FPlatformTime::Seconds() - StartTime;
+			if (FilesAndDirectories.Num() > 0)
+			{
+				bReadOnly = true;
+				bFailed = false;
+			}
+		}
+		if (FString(FCommandLine::Get()).Contains(TEXT("DerivedDataCache")) )
+		{
+			bTouch = true; // we always touch files when running the DDC commandlet
+		}
+		// The command line (-ddctouch) enables touch on all filesystem backends if specified. 
+		bTouch = bTouch || FParse::Param(FCommandLine::Get(), TEXT("DDCTOUCH"));
+
+		if (bReadOnly)
+		{
+			bTouch = false; // we won't touch read only paths
+		}
+
+		if (bTouch)
+		{
+			UE_LOG(LogDerivedDataCache, Display, TEXT("Files in %s will be touched."),*CachePath);
+		}
+
+		if (!bFailed && AccessDuration > SlowInitDuration)
+		{
+			UE_LOG(LogDerivedDataCache, Warning, TEXT("%s access is very slow (initialization took %.2lf seconds), consider disabling it."), *CachePath, AccessDuration);
+		}
+
+		if (!bReadOnly && !bFailed && bDeleteOldFiles && !FParse::Param(FCommandLine::Get(),TEXT("NODDCCLEANUP")) && FDDCCleanup::Get())
+		{			
+			FDDCCleanup::Get()->AddFilesystem( CachePath, InDaysToDeleteUnusedFiles, InMaxNumFoldersToCheck, InMaxContinuousFileChecks );
+		}
+	}
+
+	/** return true if the cache is usable **/
+	bool IsUsable()
+	{
+		return !bFailed;
+	}
+
+	/** return true if this cache is writable **/
+	virtual bool IsWritable() override
+	{
+		return !bReadOnly;
+	}
+	/**
+	 * Synchronous test for the existence of a cache item
+	 *
+	 * @param	CacheKey	Alphanumeric+underscore key of this cache item
+	 * @return				true if the data probably will be found, this can't be guaranteed because of concurrency in the backends, corruption, etc
+	 */
+	virtual bool CachedDataProbablyExists(const TCHAR* CacheKey) override
+	{
+		static FName NAME_CachedDataProbablyExists(TEXT("CachedDataProbablyExists"));
+		FDDCScopeStatHelper Stat(CacheKey, NAME_CachedDataProbablyExists);
+		static FName NAME_FileDDCPath(TEXT("FileDDCPath"));
+		Stat.AddTag(NAME_FileDDCPath, CachePath);
+
+		check(!bFailed);
+		FString Filename = BuildFilename(CacheKey);
+		FDateTime TimeStamp = IFileManager::Get().GetTimeStamp(*Filename);
+		bool bExists = TimeStamp > FDateTime::MinValue();
+		if (bExists)
+		{
+			// Update file timestamp to prevent it from being deleted by DDC Cleanup.
+			if (bTouch || 
+				 (!bReadOnly && (FDateTime::UtcNow() - TimeStamp).GetDays() > (DaysToDeleteUnusedFiles / 4)))
+			{
+				IFileManager::Get().SetTimeStamp(*Filename, FDateTime::UtcNow());
+			}
+		}
+		return bExists;
+	}
+	/**
+	 * Synchronous retrieve of a cache item
+	 *
+	 * @param	CacheKey	Alphanumeric+underscore key of this cache item
+	 * @param	OutData		Buffer to receive the results, if any were found
+	 * @return				true if any data was found, and in this case OutData is non-empty
+	 */
+	virtual bool GetCachedData(const TCHAR* CacheKey, TArray<uint8>& Data, FCacheStatRecord* Stats) override
+	{
+		static FName NAME_GetCachedData(TEXT("GetCachedData"));
+		FDDCScopeStatHelper Stat(CacheKey, NAME_GetCachedData);
+		static FName NAME_FileDDCPath(TEXT("FileDDCPath"));
+		static FName NAME_Retrieved(TEXT("Retrieved"));
+		Stat.AddTag(NAME_FileDDCPath, CachePath);
+		
+		check(!bFailed);
+		FString Filename = BuildFilename(CacheKey);
+		if (Stats)
+		{
+			Stats->CacheKey = Filename;
+			if (Filename.StartsWith(TEXT("//")) || Filename.StartsWith(TEXT("\\\\")))
+			{
+				Stats->bFromNetwork = true;
+			}
+			else
+			{
+				Stats->bFromNetwork = false;
+			}
+		}
+		double StartTime = FPlatformTime::Seconds();
+		if (FFileHelper::LoadFileToArray(Data,*Filename,FILEREAD_Silent))
+		{
+			double ReadDuration = FPlatformTime::Seconds() - StartTime;
+			double ReadSpeed = ReadDuration > 5.0 ? (Data.Num() / ReadDuration) / (1024.0 * 1024.0) : 100.0;
+			// Slower than 0.5MB/s?
+			UE_CLOG(ReadSpeed < 0.5, LogDerivedDataCache, Warning, TEXT("%s access is very slow (%.2lfMB/s), consider disabling it."), *CachePath, ReadSpeed);
+
+			UE_LOG(LogDerivedDataCache, Verbose, TEXT("FileSystemDerivedDataBackend: Cache hit on %s"),*Filename);
+			Stat.AddTag(NAME_Retrieved, TEXT("true"));
+			if (Stats)
+			{
+				Stats->GetDuration += FPlatformTime::Seconds() - StartTime;
+			}
+			return true;
+		}
+		UE_LOG(LogDerivedDataCache, Verbose, TEXT("FileSystemDerivedDataBackend: Cache miss on %s"),*Filename);
+		Data.Empty();
+		Stat.AddTag(NAME_Retrieved, TEXT("false"));
+		if (Stats)
+		{
+			Stats->GetDuration += FPlatformTime::Seconds() - StartTime;
+		}
+		return false;
+	}
+	/**
+	 * Asynchronous, fire-and-forget placement of a cache item
+	 *
+	 * @param	CacheKey	Alphanumeric+underscore key of this cache item
+	 * @param	OutData		Buffer containing the data to cache, can be destroyed after the call returns, immediately
+	 * @param	bPutEvenIfExists	If true, then do not attempt skip the put even if CachedDataProbablyExists returns true
+	 */
+	virtual void PutCachedData(const TCHAR* CacheKey, TArray<uint8>& Data, bool bPutEvenIfExists, FCacheStatRecord* Stats) override
+	{
+		static FName NAME_PutCachedData(TEXT("PutCachedData"));
+		FDDCScopeStatHelper Stat(CacheKey, NAME_PutCachedData);
+		static FName NAME_FileDDCPath(TEXT("FileDDCPath"));
+		Stat.AddTag(NAME_FileDDCPath, CachePath);
+
+
+		check(!bFailed);
+		if (!bReadOnly)
+		{
+			if (bPutEvenIfExists || !CachedDataProbablyExists(CacheKey))
+			{
+				check(Data.Num());
+				FString Filename = BuildFilename(CacheKey);
+				double StartTime = FPlatformTime::Seconds();
+				if (Stats)
+				{
+					Stats->CacheKey = Filename;
+					if (Filename.StartsWith(TEXT("//")) || Filename.StartsWith(TEXT("\\\\")))
+					{
+						Stats->bFromNetwork = true;
+					}
+					else
+					{
+						Stats->bFromNetwork = false;
+					}
+				}
+				FString TempFilename(TEXT("temp.")); 
+				TempFilename += FGuid::NewGuid().ToString();
+				TempFilename = FPaths::GetPath(Filename) / TempFilename;
+				bool bResult;
+				{
+					bResult = FFileHelper::SaveArrayToFile(Data, *TempFilename);
+				}
+				if (bResult)
+				{
+					if (IFileManager::Get().FileSize(*TempFilename) == Data.Num())
+					{
+						bool DoMove = !CachedDataProbablyExists(CacheKey);
+						if (bPutEvenIfExists && !DoMove)
+						{
+							DoMove = true;
+							RemoveCachedData(CacheKey, /*bTransient=*/ false);
+						}
+						if (DoMove) 
+						{
+							if (!IFileManager::Get().Move(*Filename, *TempFilename, true, true, false, true))
+							{
+								UE_LOG(LogDerivedDataCache, Log, TEXT("FFileSystemDerivedDataBackend: Move collision, attempt at redundant update, OK %s."),*Filename);
+							}
+							else
+							{
+								UE_LOG(LogDerivedDataCache, Verbose, TEXT("FFileSystemDerivedDataBackend: Successful cache put to %s"),*Filename);
+							}
+						}
+					}
+					else
+					{
+						UE_LOG(LogDerivedDataCache, Warning, TEXT("FFileSystemDerivedDataBackend: Temp file is short %s!"),*TempFilename);
+					}
+				}
+				else
+				{
+					UE_LOG(LogDerivedDataCache, Warning, TEXT("FFileSystemDerivedDataBackend: Could not write temp file %s!"),*TempFilename);
+				}
+				if (Stats)
+				{
+					Stats->PutDuration += FPlatformTime::Seconds() - StartTime;
+				}
+				// if everything worked, this is not necessary, but we will make every effort to avoid leaving junk in the cache
+				if (FPaths::FileExists(TempFilename))
+				{
+					IFileManager::Get().Delete(*TempFilename, false, false, true);
+				}
+			}
+		}
+	}
+
+	void RemoveCachedData(const TCHAR* CacheKey, bool bTransient) override
+	{
+		check(!bFailed);
+		if (!bReadOnly && (!bTransient || bPurgeTransient))
+		{
+			FString Filename = BuildFilename(CacheKey);
+			if (bTransient)
+			{
+				UE_LOG(LogDerivedDataCache,Verbose,TEXT("Deleting transient cached data. Key=%s Filename=%s"),CacheKey,*Filename);
+			}
+			IFileManager::Get().Delete(*Filename, false, false, true);
+		}
+	}
+
+private:
+
+	/**
+	 * Threadsafe method to compute the filename from the cachekey, currently just adds a path and an extension.
+	 *
+	 * @param	CacheKey	Alphanumeric+underscore key of this cache item
+	 * @return				filename built from the cache key
+	 */
+	FString BuildFilename(const TCHAR* CacheKey)
+	{
+		FString Key = FString(CacheKey).ToUpper();
+		for (int32 i = 0; i < Key.Len(); i++)
+		{
+			check(FChar::IsAlnum(Key[i]) || FChar::IsUnderscore(Key[i]) || Key[i] == L'$');
+		}
+		uint32 Hash = FCrc::StrCrc_DEPRECATED(*Key);
+		// this creates a tree of 1000 directories
+		FString HashPath = FString::Printf(TEXT("%1d/%1d/%1d/"),(Hash/100)%10,(Hash/10)%10,Hash%10);
+		return CachePath / HashPath / Key + TEXT(".udd");
+	}
+
+	/** Base path we are storing the cache files in. **/
+	FString	CachePath;
+	/** If true, do not attempt to write to this cache **/
+	bool		bReadOnly;
+	/** If true, we failed to write to this directory and it did not contain anything so we should not be used **/
+	bool		bFailed;
+	/** If true, CachedDataProbablyExists will update the file timestamps. */
+	bool		bTouch;
+	/** If true, allow transient data to be removed from the cache. */
+	bool		bPurgeTransient;
+	/** Age of file when it should be deleted from DDC cache. */
+	int32		DaysToDeleteUnusedFiles;
+};
+
+FDerivedDataBackendInterface* CreateFileSystemDerivedDataBackend(const TCHAR* CacheDirectory, bool bForceReadOnly /*= false*/, bool bTouchFiles /*= false*/, bool bPurgeTransient /*= false*/, bool bDeleteOldFiles /*= false*/, int32 InDaysToDeleteUnusedFiles /*= 60*/, int32 InMaxNumFoldersToCheck /*= -1*/, int32 InMaxContinuousFileChecks /*= -1*/)
+{
+	FFileSystemDerivedDataBackend* FileDDB = new FFileSystemDerivedDataBackend(CacheDirectory, bForceReadOnly, bTouchFiles, bPurgeTransient, bDeleteOldFiles, InDaysToDeleteUnusedFiles, InMaxNumFoldersToCheck, InMaxContinuousFileChecks);
+	if (!FileDDB->IsUsable())
+	{
+		delete FileDDB;
+		FileDDB = NULL;
+	}
+	return FileDDB;
+}
diff --git a/Engine/Source/Developer/DerivedDataCache/Private/HierarchicalDerivedDataBackend.h b/Engine/Source/Developer/DerivedDataCache/Private/HierarchicalDerivedDataBackend.h
new file mode 100644
--- /dev/null
+++ b/Engine/Source/Developer/DerivedDataCache/Private/HierarchicalDerivedDataBackend.h
@@ -0,0 +1,213 @@
+// Copyright 1998-2015 Epic Games, Inc. All Rights Reserved.
+
+#pragma once
+
+#include "DDCStatsHelper.h"
+
+
+/** 
+ * A backend wrapper that implements a cache hierarchy of backends. 
+**/
+class FHierarchicalDerivedDataBackend : public FDerivedDataBackendInterface
+{
+public:
+
+	/**
+	 * Constructor
+	 * @param	InInnerBackends Backends to call into for actual storage of the cache, first item is the "fastest cache"
+	 */
+	FHierarchicalDerivedDataBackend(const TArray<FDerivedDataBackendInterface*>& InInnerBackends)
+		: InnerBackends(InInnerBackends)
+		, bIsWritable(false)
+	{
+		check(InnerBackends.Num() > 1); // if it is just one, then you don't need this wrapper
+		for (int32 CacheIndex = 0; CacheIndex < InnerBackends.Num(); CacheIndex++)
+		{
+			if (InnerBackends[CacheIndex]->IsWritable())
+			{
+				bIsWritable = true;
+			}
+		}
+		if (bIsWritable)
+		{
+			for (int32 CacheIndex = 0; CacheIndex < InnerBackends.Num(); CacheIndex++)
+			{
+				// async puts to allow us to fill all levels without holding up the engine
+				new (AsyncPutInnerBackends) TScopedPointer<FDerivedDataBackendInterface>(new FDerivedDataBackendAsyncPutWrapper(InnerBackends[CacheIndex], false));
+			}
+		}
+	}
+
+	/** Adds inner backend. */
+	void AddInnerBackend(FDerivedDataBackendInterface* InInner) 
+	{
+		InnerBackends.Add(InInner);
+	}
+
+	/** Removes inner backend. */
+	bool RemoveInnerBackend(FDerivedDataBackendInterface* InInner) 
+	{
+		const int32 NumRemoved = InnerBackends.Remove(InInner);
+		return NumRemoved != 0;
+	}
+
+	/** return true if this cache is writable **/
+	virtual bool IsWritable() override
+	{
+		return bIsWritable;
+	}
+
+	/**
+	 * Synchronous test for the existence of a cache item
+	 *
+	 * @param	CacheKey	Alphanumeric+underscore key of this cache item
+	 * @return				true if the data probably will be found, this can't be guaranteed because of concurrency in the backends, corruption, etc
+	 */
+	virtual bool CachedDataProbablyExists(const TCHAR* CacheKey) override
+	{
+		static FName NAME_CachedDataProbablyExists(TEXT("CachedDataProbablyExists"));
+		FDDCScopeStatHelper Stat(CacheKey, NAME_CachedDataProbablyExists);
+		const static FName NAME_HierarchicalDDC = FName(TEXT("HierarchicalDDC"));
+		Stat.AddTag(NAME_HierarchicalDDC, FString());
+
+		for (int32 CacheIndex = 0; CacheIndex < InnerBackends.Num(); CacheIndex++)
+		{
+			if (InnerBackends[CacheIndex]->CachedDataProbablyExists(CacheKey))
+			{
+				return true;
+			}
+		}
+		return false;
+	}
+	/**
+	 * Synchronous retrieve of a cache item
+	 *
+	 * @param	CacheKey	Alphanumeric+underscore key of this cache item
+	 * @param	OutData		Buffer to receive the results, if any were found
+	 * @return				true if any data was found, and in this case OutData is non-empty
+	 */
+	virtual bool GetCachedData(const TCHAR* CacheKey, TArray<uint8>& OutData, FCacheStatRecord* Stats) override
+	{
+		const static FName NAME_GetCachedData = FName(TEXT("GetCachedData"));
+		const static FName NAME_HierarchicalDDC = FName(TEXT("HierarchicalDDC"));
+		const static FName NAME_DataSize = FName(TEXT("DataSize")); 
+		static FName NAME_Retrieved(TEXT("Retrieved"));
+
+		FDDCScopeStatHelper Stat(CacheKey, NAME_GetCachedData);
+		Stat.AddTag(NAME_HierarchicalDDC, FString());
+
+		for (int32 CacheIndex = 0; CacheIndex < InnerBackends.Num(); CacheIndex++)
+		{
+			if (InnerBackends[CacheIndex]->CachedDataProbablyExists(CacheKey) && InnerBackends[CacheIndex]->GetCachedData(CacheKey, OutData, Stats))
+			{
+				if (bIsWritable)
+				{
+					// fill in the higher level caches
+					for (int32 PutCacheIndex = CacheIndex - 1; PutCacheIndex >= 0; PutCacheIndex--)
+					{
+						if (InnerBackends[PutCacheIndex]->IsWritable())
+						{
+							if (InnerBackends[PutCacheIndex]->BackfillLowerCacheLevels() &&
+								InnerBackends[PutCacheIndex]->CachedDataProbablyExists(CacheKey))
+							{
+								InnerBackends[PutCacheIndex]->RemoveCachedData(CacheKey, /*bTransient=*/ false); // it apparently failed, so lets delete what is there
+								AsyncPutInnerBackends[PutCacheIndex]->PutCachedData(CacheKey, OutData, true, Stats); // we force a put here because it must have failed
+								Stats->CacheGet++;
+							}
+							else
+							{
+								AsyncPutInnerBackends[PutCacheIndex]->PutCachedData(CacheKey, OutData, false, Stats); 
+								Stats->CacheGet++;
+							}
+						}
+					}
+					if (InnerBackends[CacheIndex]->BackfillLowerCacheLevels())
+					{
+						// fill in the lower level caches
+						for (int32 PutCacheIndex = CacheIndex + 1; PutCacheIndex < AsyncPutInnerBackends.Num(); PutCacheIndex++)
+						{
+							if (!InnerBackends[PutCacheIndex]->IsWritable() && !InnerBackends[PutCacheIndex]->BackfillLowerCacheLevels() && InnerBackends[PutCacheIndex]->CachedDataProbablyExists(CacheKey))
+							{
+								break; //do not write things that are already in the read only pak file
+							}
+							if (InnerBackends[PutCacheIndex]->IsWritable())
+							{
+								AsyncPutInnerBackends[PutCacheIndex]->PutCachedData(CacheKey, OutData, false, Stats); // we do not need to force a put here
+								Stats->CacheGet++;
+							}
+						}
+					}
+				}
+				Stat.AddTag(NAME_DataSize, FString::Printf(TEXT("%d bytes"), OutData.Num()));
+				Stat.AddTag(NAME_Retrieved, TEXT("true"));
+				Stats->EndTime = FPlatformTime::Seconds();
+				return true;
+			}
+		}
+		Stat.AddTag(NAME_Retrieved, TEXT("false"));
+		return false;
+	}
+	/**
+	 * Asynchronous, fire-and-forget placement of a cache item
+	 *
+	 * @param	CacheKey	Alphanumeric+underscore key of this cache item
+	 * @param	InData		Buffer containing the data to cache, can be destroyed after the call returns, immediately
+	 * @param	bPutEvenIfExists	If true, then do not attempt skip the put even if CachedDataProbablyExists returns true
+	 */
+	virtual void PutCachedData(const TCHAR* CacheKey, TArray<uint8>& InData, bool bPutEvenIfExists, FCacheStatRecord* Stats) override
+	{
+		const static FName NAME_PutCachedData = FName(TEXT("PutCachedData"));
+		const static FName NAME_HierarchicalDDC = FName(TEXT("HierarchicalDDC"));
+		const static FName NAME_DataSize = FName(TEXT("DataSize"));
+
+		FDDCScopeStatHelper Stat(CacheKey, NAME_PutCachedData);
+		Stat.AddTag(NAME_HierarchicalDDC, FString());
+		Stat.AddTag(NAME_DataSize, FString::Printf(TEXT("%d bytes"), InData.Num()));
+
+
+		if (!bIsWritable)
+		{
+			return; // no point in continuing down the chain
+		}
+		bool bSynchronousPutPeformed = false;  // we must do at least one synchronous put to a writable cache before we return
+		for (int32 PutCacheIndex = 0; PutCacheIndex < InnerBackends.Num(); PutCacheIndex++)
+		{
+			if (!InnerBackends[PutCacheIndex]->IsWritable() && !InnerBackends[PutCacheIndex]->BackfillLowerCacheLevels() && InnerBackends[PutCacheIndex]->CachedDataProbablyExists(CacheKey))
+			{
+				break; //do not write things that are already in the read only pak file
+			}
+			if (InnerBackends[PutCacheIndex]->IsWritable())
+			{
+				if (!bSynchronousPutPeformed)
+				{
+					InnerBackends[PutCacheIndex]->PutCachedData(CacheKey, InData, bPutEvenIfExists, Stats);
+					bSynchronousPutPeformed = true;
+				}
+				else
+				{
+					AsyncPutInnerBackends[PutCacheIndex]->PutCachedData(CacheKey, InData, bPutEvenIfExists, Stats);
+				}
+			}
+		}
+	}
+
+	virtual void RemoveCachedData(const TCHAR* CacheKey, bool bTransient) override
+	{
+		if (!bIsWritable)
+		{
+			return; // no point in continuing down the chain
+		}
+		for (int32 PutCacheIndex = 0; PutCacheIndex < InnerBackends.Num(); PutCacheIndex++)
+		{
+			InnerBackends[PutCacheIndex]->RemoveCachedData(CacheKey, bTransient);
+		}
+	}
+private:
+
+	/** Array of backends forming the hierarchical cache...the first element is the fastest cache. **/
+	TArray<FDerivedDataBackendInterface*> InnerBackends;
+	/** Each of the backends wrapped with an async put **/
+	TArray<TScopedPointer<FDerivedDataBackendInterface> > AsyncPutInnerBackends;
+	/** As an optimization, we check our writable status at contruction **/
+	bool bIsWritable;
+};
diff --git a/Engine/Source/Developer/DerivedDataCache/Private/MemoryDerivedDataBackend.h b/Engine/Source/Developer/DerivedDataCache/Private/MemoryDerivedDataBackend.h
new file mode 100644
--- /dev/null
+++ b/Engine/Source/Developer/DerivedDataCache/Private/MemoryDerivedDataBackend.h
@@ -0,0 +1,357 @@
+// Copyright 1998-2015 Epic Games, Inc. All Rights Reserved.
+
+#pragma once
+
+/** 
+ * A simple thread safe, memory based backend. This is used for Async puts and the boot cache.
+**/
+class FMemoryDerivedDataBackend : public FDerivedDataBackendInterface
+{
+public:
+	FMemoryDerivedDataBackend(int64 InMaxCacheSize = -1)
+		: MaxCacheSize(InMaxCacheSize)
+		, bDisabled( false )
+		, CurrentCacheSize( SerializationSpecificDataSize )
+		, bMaxSizeExceeded(false)
+	{
+	}
+	~FMemoryDerivedDataBackend()
+	{
+		Disable();
+	}
+
+	/** return true if this cache is writable **/
+	virtual bool IsWritable() override
+	{
+		FScopeLock ScopeLock(&SynchronizationObject);
+		return !bDisabled;
+	}
+
+	/**
+	 * Synchronous test for the existence of a cache item
+	 *
+	 * @param	CacheKey	Alphanumeric+underscore key of this cache item
+	 * @return				true if the data probably will be found, this can't be guaranteed because of concurrency in the backends, corruption, etc
+	 */
+	virtual bool CachedDataProbablyExists(const TCHAR* CacheKey) override
+	{
+		FScopeLock ScopeLock(&SynchronizationObject);
+		if (bDisabled)
+		{
+			return false;
+		}
+		// to avoid constant error reporting in async put due to restricted cache size, 
+		// we report true if the max size has been exceeded
+		if (bMaxSizeExceeded)
+		{
+			return true;
+		}
+
+		return CacheItems.Contains(FString(CacheKey));
+	}
+	/**
+	 * Synchronous retrieve of a cache item
+	 *
+	 * @param	CacheKey	Alphanumeric+underscore key of this cache item
+	 * @param	OutData		Buffer to receive the results, if any were found
+	 * @return				true if any data was found, and in this case OutData is non-empty
+	 */
+	virtual bool GetCachedData(const TCHAR* CacheKey, TArray<uint8>& OutData, FCacheStatRecord* Stats) override
+	{
+		FScopeLock ScopeLock(&SynchronizationObject);
+		if (!bDisabled)
+		{
+			FCacheValue* Item = CacheItems.FindRef(FString(CacheKey));
+			if (Item)
+			{
+				OutData = Item->Data;
+				Item->Age = 0;
+				check(OutData.Num());
+				return true;
+			}
+		}
+		OutData.Empty();
+		return false;
+	}
+	/**
+	 * Asynchronous, fire-and-forget placement of a cache item
+	 *
+	 * @param	CacheKey	Alphanumeric+underscore key of this cache item
+	 * @param	InData		Buffer containing the data to cache, can be destroyed after the call returns, immediately
+	 * @param	bPutEvenIfExists	If true, then do not attempt skip the put even if CachedDataProbablyExists returns true
+	 */
+	virtual void PutCachedData(const TCHAR* CacheKey, TArray<uint8>& InData, bool bPutEvenIfExists, FCacheStatRecord* Stats) override
+	{
+		FScopeLock ScopeLock(&SynchronizationObject);
+		
+		if (bDisabled || bMaxSizeExceeded)
+		{
+			return;
+		}
+		
+		FString Key(CacheKey);
+		FCacheValue* Item = CacheItems.FindRef(FString(CacheKey));
+		if (Item)
+		{
+			//check(Item->Data == InData); // any second attempt to push data should be identical data
+		}
+		else
+		{
+			FCacheValue* Val = new FCacheValue(InData);
+			int32 CacheValueSize = CalcCacheValueSize(Key, *Val);
+
+			// check if we haven't exceeded the MaxCacheSize
+			if (MaxCacheSize > 0 && (CurrentCacheSize + CacheValueSize) > MaxCacheSize)
+			{
+				delete Val;
+
+				UE_LOG(LogDerivedDataCache, Warning, TEXT("Failed to cache data. Maximum cache size reached. CurrentSize %d kb / MaxSize: %d kb"), CurrentCacheSize / 1024, MaxCacheSize / 1024);
+				bMaxSizeExceeded = true;
+			}
+			else
+			{
+				CacheItems.Add(Key, Val);
+				CalcCacheValueSize(Key, *Val);
+
+				CurrentCacheSize += CacheValueSize;
+			}
+		}
+	}
+
+	virtual void RemoveCachedData(const TCHAR* CacheKey, bool bTransient) override
+	{
+		FScopeLock ScopeLock(&SynchronizationObject);
+		if (bDisabled || bTransient)
+		{
+			return;
+		}
+		FString Key(CacheKey);
+		FCacheValue* Item = NULL;
+		if (CacheItems.RemoveAndCopyValue(Key, Item))
+		{
+			CurrentCacheSize -= CalcCacheValueSize(Key, *Item);
+			bMaxSizeExceeded = false;
+
+			check(Item);
+			delete Item;
+		}
+		else
+		{
+			check(!Item);
+		}
+	}
+
+	/**
+	 * Save the cache to disk
+	 * @param	Filename	Filename to save
+	 * @return	true if file was saved sucessfully
+	 */
+	bool SaveCache(const TCHAR* Filename)
+	{
+		double StartTime = FPlatformTime::Seconds();
+		TAutoPtr<FArchive> SaverArchive(IFileManager::Get().CreateFileWriter(Filename, FILEWRITE_EvenIfReadOnly));
+		if (!SaverArchive.IsValid())
+		{
+			UE_LOG(LogDerivedDataCache, Error, TEXT("Could not save memory cache %s."), Filename);
+			return false;
+		}
+
+		FArchive& Saver = *SaverArchive;
+		uint32 Magic = MemCache_Magic64;
+		Saver << Magic;
+		const int64 DataStartOffset = Saver.Tell();
+		{
+			FScopeLock ScopeLock(&SynchronizationObject);
+			check(!bDisabled);
+			for (TMap<FString, FCacheValue*>::TIterator It(CacheItems); It; ++It )
+			{
+				Saver << It.Key();
+				Saver << It.Value()->Age;
+				Saver << It.Value()->Data;
+			}
+		}
+		const int64 DataSize = Saver.Tell(); // Everything except the footer
+		int64 Size = DataSize;
+		uint32 Crc = MemCache_Magic64; // Crc takes more time than I want to spend  FCrc::MemCrc_DEPRECATED(&Buffer[0], Size);
+		Saver << Size;
+		Saver << Crc;
+
+		check(SerializationSpecificDataSize + DataSize <= MaxCacheSize || MaxCacheSize <= 0);
+
+		UE_LOG(LogDerivedDataCache, Log, TEXT("Saved boot cache %4.2fs %lldMB %s."), float(FPlatformTime::Seconds() - StartTime), DataSize / (1024 * 1024), Filename);
+		return true;
+	}
+	/**
+	 * Load the cache from disk
+	 * @param	Filename	Filename to load
+	 * @return	true if file was loaded sucessfully
+	 */
+	bool LoadCache(const TCHAR* Filename)
+	{
+		double StartTime = FPlatformTime::Seconds();
+		const int64 FileSize = IFileManager::Get().FileSize(Filename);
+		if (FileSize < 0)
+		{
+			UE_LOG(LogDerivedDataCache, Warning, TEXT("Could not find memory cache %s."), Filename);
+			return false;
+		}
+		// We test 3 * uint32 which is the old format (< SerializationSpecificDataSize). We'll test
+		// agains SerializationSpecificDataSize later when we read the magic number from the cache.
+		if (FileSize < sizeof(uint32) * 3)
+		{
+			UE_LOG(LogDerivedDataCache, Error, TEXT("Memory cache was corrputed (short) %s."), Filename);
+			return false;
+		}
+		if (FileSize > MaxCacheSize*2 && MaxCacheSize > 0)
+		{
+			UE_LOG(LogDerivedDataCache, Error, TEXT("Refusing to load DDC cache %s. Size exceeds doubled MaxCacheSize."), Filename);
+			return false;
+		}
+
+		TAutoPtr<FArchive> LoderArchive(IFileManager::Get().CreateFileReader(Filename));
+		if (!LoderArchive.IsValid())
+		{
+			UE_LOG(LogDerivedDataCache, Warning, TEXT("Could not read memory cache %s."), Filename);
+			return false;
+		}
+
+		FArchive& Loader = *LoderArchive;
+		uint32 Magic = 0;
+		Loader << Magic;
+		if (Magic != MemCache_Magic && Magic != MemCache_Magic64)
+		{
+			UE_LOG(LogDerivedDataCache, Error, TEXT("Memory cache was corrputed (magic) %s."), Filename);
+			return false;
+		}
+		// Check the file size again, this time against the correct minimum size.
+		if (Magic == MemCache_Magic64 && FileSize < SerializationSpecificDataSize)
+		{
+			UE_LOG(LogDerivedDataCache, Error, TEXT("Memory cache was corrputed (short) %s."), Filename);
+			return false;
+		}
+		// Calculate expected DataSize based on the magic number (footer size difference)
+		const int64 DataSize = FileSize - (Magic == MemCache_Magic64 ? (SerializationSpecificDataSize - sizeof(uint32)) : (sizeof(uint32) * 2));		
+		Loader.Seek(DataSize);
+		int64 Size = 0;
+		uint32 Crc = 0;
+		if (Magic == MemCache_Magic64)
+		{
+			Loader << Size;
+		}
+		else
+		{
+			uint32 Size32 = 0;
+			Loader << Size32;
+			Size = (int64)Size32;
+		}
+		Loader << Crc;
+		if (Size != DataSize)
+		{
+			UE_LOG(LogDerivedDataCache, Error, TEXT("Memory cache was corrputed (size) %s."), Filename);
+			return false;
+		}
+		if ((Crc != MemCache_Magic && Crc != MemCache_Magic64) || Crc != Magic)
+		{
+			UE_LOG(LogDerivedDataCache, Warning, TEXT("Memory cache was corrputed (crc) %s."), Filename);
+			return false;
+		}
+		// Seek to data start offset (skip magic number)
+		Loader.Seek(sizeof(uint32));
+		{
+			TArray<uint8> Working;
+			FScopeLock ScopeLock(&SynchronizationObject);
+			check(!bDisabled);
+			while (Loader.Tell() < DataSize)
+			{
+				FString Key;
+				int32 Age;
+				Loader << Key;
+				Loader << Age;
+				Age++;
+				Loader << Working;
+				if (Age < MaxAge)
+				{
+					CacheItems.Add(Key, new FCacheValue(Working, Age));
+				}
+				Working.Reset();
+			}
+			// these are just a double check on ending correctly
+			if (Magic == MemCache_Magic64)
+			{
+				Loader << Size;
+			}
+			else
+			{
+				uint32 Size32 = 0;
+				Loader << Size32;
+				Size = (int64)Size32;
+			}
+			Loader << Crc;
+		}
+		
+		CurrentCacheSize = FileSize;
+		UE_LOG(LogDerivedDataCache, Log, TEXT("Loaded boot cache %4.2fs %lldMB %s."), float(FPlatformTime::Seconds() - StartTime), DataSize / (1024 * 1024), Filename);
+		return true;
+	}
+	/**
+	 * Disable cache and ignore all subsequent requests
+	 */
+	void Disable()
+	{
+		FScopeLock ScopeLock(&SynchronizationObject);
+		bDisabled = true;
+		for (TMap<FString,FCacheValue*>::TIterator It(CacheItems); It; ++It )
+		{
+			delete It.Value();
+		}
+		CacheItems.Empty();
+
+		CurrentCacheSize = SerializationSpecificDataSize;
+	}
+
+private:
+
+	struct FCacheValue
+	{
+		int32 Age;
+		TArray<uint8> Data;
+		FCacheValue(const TArray<uint8>& InData, int32 InAge = 0)
+			: Age(InAge)
+			, Data(InData)
+		{
+		}
+	};
+
+	FORCEINLINE int32 CalcCacheValueSize(const FString& Key, const FCacheValue& Val)
+	{
+		return (Key.Len() + 1) * sizeof(TCHAR) + sizeof(Val.Age) + Val.Data.Num();
+	}
+
+	/** Set of files that are being written to disk asynchronously. */
+	TMap<FString, FCacheValue*> CacheItems;
+	/** Maximum size the cached items can grow up to ( in bytes ) */
+	int64 MaxCacheSize;
+	/** When set to true, this cache is disabled...ignore all requests. */
+	bool bDisabled;
+	/** Object used for synchronization via a scoped lock						*/
+	FCriticalSection	SynchronizationObject;
+	/** Current estimated cache size in bytes */
+	int64 CurrentCacheSize;
+	/** Indicates that the cache max size has been exceeded. This is used to avoid
+		warning spam after the size has reached the limit. */
+	bool bMaxSizeExceeded;
+	enum 
+	{
+		/** Magic number to use in header */
+		MemCache_Magic = 0x0cac0ddc,
+		/** Magic number to use in header (new, > 2GB size compatible) */
+		MemCache_Magic64 = 0x0cac1ddc,
+		/** Oldest cache items to keep */
+		MaxAge = 3,
+		/** Size of data that is stored in the cachefile apart from the cache entries (64 bit size). */
+		SerializationSpecificDataSize = sizeof(uint32)	// Magic
+									  + sizeof(int64)	// Size
+									  + sizeof(uint32), // CRC
+	};
+};
+
diff --git a/Engine/Source/Developer/DerivedDataCache/Private/PakFileDerivedDataBackend.h b/Engine/Source/Developer/DerivedDataCache/Private/PakFileDerivedDataBackend.h
new file mode 100644
--- /dev/null
+++ b/Engine/Source/Developer/DerivedDataCache/Private/PakFileDerivedDataBackend.h
@@ -0,0 +1,506 @@
+// Copyright 1998-2015 Epic Games, Inc. All Rights Reserved.
+
+#pragma once
+
+#include "Compression.h"
+
+/** 
+ * A simple thread safe, pak file based backend. 
+**/
+class FPakFileDerivedDataBackend : public FDerivedDataBackendInterface
+{
+public:
+	FPakFileDerivedDataBackend(const TCHAR* InFilename, bool bInWriting)
+		: bWriting(bInWriting)
+		, bClosed(false)
+		, Filename(InFilename)
+	{
+		if (bWriting)
+		{
+			FileHandle = IFileManager::Get().CreateFileWriter(InFilename, FILEWRITE_NoReplaceExisting);
+			if (!FileHandle)
+			{
+				UE_LOG(LogDerivedDataCache, Fatal, TEXT("Pak cache could not be opened for writing %s."), InFilename);
+				bClosed = true;
+			}
+			else
+			{
+				UE_LOG(LogDerivedDataCache, Display, TEXT("Pak cache opened for writing %s."), InFilename);
+			}
+		}
+		else
+		{
+			FileHandle = IFileManager::Get().CreateFileReader(InFilename);
+			if (!FileHandle)
+			{
+				UE_LOG(LogDerivedDataCache, Warning, TEXT("Pak cache could not be opened for reading %s."), InFilename);
+			}
+			else if (!LoadCache(InFilename))
+			{
+				FileHandle.Reset();
+				CacheItems.Empty();
+				bClosed = true;
+			}
+			else
+			{
+				UE_LOG(LogDerivedDataCache, Display, TEXT("Pak cache opened for reading %s."), InFilename);
+			}
+		}
+	}
+
+	~FPakFileDerivedDataBackend()
+	{
+		Close();
+	}
+
+	void Close()
+	{
+		FDerivedDataBackend::Get().WaitForQuiescence();
+		FScopeLock ScopeLock(&SynchronizationObject);
+		if (!bClosed)
+		{
+			if (bWriting)
+			{
+				SaveCache();
+			}
+			FileHandle.Reset();
+			CacheItems.Empty();
+			bClosed = true;
+		}
+	}
+
+	/** return true if this cache is writable **/
+	virtual bool IsWritable() override
+	{
+		return bWriting && !bClosed;
+	}
+
+	virtual bool BackfillLowerCacheLevels() override
+	{
+		return false;
+	}
+
+	/**
+	 * Synchronous test for the existence of a cache item
+	 *
+	 * @param	CacheKey	Alphanumeric+underscore key of this cache item
+	 * @return				true if the data probably will be found, this can't be guaranteed because of concurrency in the backends, corruption, etc
+	 */
+	virtual bool CachedDataProbablyExists(const TCHAR* CacheKey) override
+	{
+		FScopeLock ScopeLock(&SynchronizationObject);
+		return CacheItems.Contains(FString(CacheKey));
+	}
+	/**
+	 * Synchronous retrieve of a cache item
+	 *
+	 * @param	CacheKey	Alphanumeric+underscore key of this cache item
+	 * @param	OutData		Buffer to receive the results, if any were found
+	 * @return				true if any data was found, and in this case OutData is non-empty
+	 */
+	virtual bool GetCachedData(const TCHAR* CacheKey, TArray<uint8>& OutData, FCacheStatRecord* Stats) override
+	{
+		if (bWriting || bClosed)
+		{
+			return false;
+		}
+		FScopeLock ScopeLock(&SynchronizationObject);
+		if (Stats)
+			Stats->bFromNetwork = false;
+		FCacheValue* Item = CacheItems.Find(FString(CacheKey));
+		if (Item)
+		{
+			check(FileHandle);
+			FileHandle->Seek(Item->Offset);
+			if (FileHandle->Tell() != Item->Offset)
+			{
+				UE_LOG(LogDerivedDataCache, Warning, TEXT("Pak file, bad seek."));
+			}
+			else
+			{
+				check(Item->Size);
+				check(!OutData.Num());
+				check(FileHandle->IsLoading());
+				OutData.AddUninitialized(Item->Size);
+				FileHandle->Serialize(OutData.GetData(), int64(Item->Size));
+				uint32 TestCrc = FCrc::MemCrc_DEPRECATED(OutData.GetData(), Item->Size);
+				if (TestCrc != Item->Crc)
+				{
+					UE_LOG(LogDerivedDataCache, Warning, TEXT("Pak file, bad crc."));
+				}
+				else
+				{
+					UE_LOG(LogDerivedDataCache, Verbose, TEXT("FPakFileDerivedDataBackend: Cache hit on %s"), CacheKey);
+					check(OutData.Num());
+					return true;
+				}
+			}
+		}
+		else
+		{
+			UE_LOG(LogDerivedDataCache, Verbose, TEXT("FPakFileDerivedDataBackend: Miss on %s"), CacheKey);
+		}
+		OutData.Empty();
+		return false;
+	}
+	/**
+	 * Asynchronous, fire-and-forget placement of a cache item
+	 *
+	 * @param	CacheKey	Alphanumeric+underscore key of this cache item
+	 * @param	InData		Buffer containing the data to cache, can be destroyed after the call returns, immediately
+	 * @param	bPutEvenIfExists	If true, then do not attempt skip the put even if CachedDataProbablyExists returns true
+	 */
+	virtual void PutCachedData(const TCHAR* CacheKey, TArray<uint8>& InData, bool bPutEvenIfExists, FCacheStatRecord* Stats) override
+	{
+		if (!bWriting || bClosed)
+		{
+			return;
+		}
+		{
+			FScopeLock ScopeLock(&SynchronizationObject);
+			if (Stats)
+				Stats->bFromNetwork = false;
+			FString Key(CacheKey);
+			FCacheValue* Item = CacheItems.Find(FString(CacheKey));
+			if (!Item)
+			{
+				check(InData.Num());
+				check(Key.Len());
+				uint32 Crc = FCrc::MemCrc_DEPRECATED(InData.GetData(), InData.Num());
+				check(FileHandle);
+				check(FileHandle->IsSaving());
+				int64 Offset = FileHandle->Tell();
+				if (Offset < 0)
+				{
+					CacheItems.Empty();
+					FileHandle.Reset();
+					bClosed = true;
+				}
+				else
+				{
+					FileHandle->Serialize(InData.GetData(), int64(InData.Num()));
+					UE_LOG(LogDerivedDataCache, Verbose, TEXT("FPakFileDerivedDataBackend: Put %s"), CacheKey);
+					CacheItems.Add(Key,FCacheValue(Offset, InData.Num(), Crc));
+				}
+			}
+		}
+		if (bClosed)
+		{
+			UE_LOG(LogDerivedDataCache, Fatal, TEXT("Could not write pak file...out of disk space?"));
+		}
+	}
+
+	virtual void RemoveCachedData(const TCHAR* CacheKey, bool bTransient) override
+	{
+		if (bClosed || bTransient)
+		{
+			return;
+		}
+		// strangish. We can delete from a pak, but it only deletes the index 
+		// if this is a read cache, it will read it next time
+		// if this is a write cache, we wated space
+		FScopeLock ScopeLock(&SynchronizationObject);
+		FString Key(CacheKey);
+		CacheItems.Remove(Key);
+	}
+
+	/**
+	 * Save the cache to disk
+	 * @return	true if file was saved sucessfully
+	 */
+	bool SaveCache()
+	{
+		FScopeLock ScopeLock(&SynchronizationObject);
+		check(FileHandle);
+		check(FileHandle->IsSaving());
+		int64 IndexOffset = FileHandle->Tell();
+		check(IndexOffset >= 0);
+		uint32 NumItems =  uint32(CacheItems.Num());
+		check(IndexOffset > 0 || !NumItems);
+		TArray<uint8> IndexBuffer;
+		{
+			FMemoryWriter Saver(IndexBuffer);
+			uint32 NumProcessed = 0;
+			for (TMap<FString, FCacheValue>::TIterator It(CacheItems); It; ++It )
+			{
+				check(It.Value().Offset >= 0 && It.Value().Offset < IndexOffset);
+				check(It.Value().Size);
+				check(It.Key().Len());
+				Saver << It.Key();
+				Saver << It.Value().Offset;
+				Saver << It.Value().Size;
+				Saver << It.Value().Crc;
+				NumProcessed++;
+			}
+			check(NumProcessed == NumItems);
+		}
+		uint32 IndexCrc = FCrc::MemCrc_DEPRECATED(IndexBuffer.GetData(), IndexBuffer.Num());
+		uint32 SizeIndex = uint32(IndexBuffer.Num());
+
+		uint32 Magic = PakCache_Magic;
+		TArray<uint8> Buffer;
+		FMemoryWriter Saver(Buffer);
+		Saver << Magic;
+		Saver << IndexCrc;
+		Saver << NumItems;
+		Saver << SizeIndex;
+		Saver.Serialize(IndexBuffer.GetData(), IndexBuffer.Num());
+		Saver << Magic;
+		Saver << IndexOffset;
+		FileHandle->Serialize(Buffer.GetData(), Buffer.Num());
+		CacheItems.Empty();
+		FileHandle.Reset();
+		bClosed = true;
+		return true;
+	}
+	/**
+	 * Load the cache to disk	 * @param	Filename	Filename to load
+	 * @return	true if file was loaded sucessfully
+	 */
+	bool LoadCache(const TCHAR* InFilename)
+	{
+		check(FileHandle);
+		check(FileHandle->IsLoading());
+		int64 FileSize = FileHandle->TotalSize();
+		check(FileSize >= 0);
+		if (FileSize < sizeof(int64) + sizeof(uint32) * 5)
+		{
+			UE_LOG(LogDerivedDataCache, Error, TEXT("Pak cache was corrputed (short) %s."), InFilename);
+			return false;
+		}
+		int64 IndexOffset = -1;
+		int64 Trailer = -1;
+		{
+			TArray<uint8> Buffer;
+			const int64 SeekPos = FileSize - int64(sizeof(int64) + sizeof(uint32));
+			FileHandle->Seek(SeekPos);
+			Trailer = FileHandle->Tell();
+			if (Trailer != SeekPos)
+			{
+				UE_LOG(LogDerivedDataCache, Error, TEXT("Pak cache was corrputed (bad seek) %s."), InFilename);
+				return false;
+			}
+			check(Trailer >= 0 && Trailer < FileSize);
+			Buffer.AddUninitialized(sizeof(int64) + sizeof(uint32));
+			FileHandle->Serialize(Buffer.GetData(), int64(sizeof(int64)+sizeof(uint32)));
+			FMemoryReader Loader(Buffer);
+			uint32 Magic = 0;
+			Loader << Magic;
+			Loader << IndexOffset;
+			if (Magic != PakCache_Magic || IndexOffset < 0 || IndexOffset + int64(sizeof(uint32) * 4) > Trailer)
+			{
+				UE_LOG(LogDerivedDataCache, Error, TEXT("Pak cache was corrputed (bad footer) %s."), InFilename);
+				return false;
+			}
+		}
+		uint32 IndexCrc = 0;
+		uint32 NumIndex = 0;
+		uint32 SizeIndex = 0;
+		{
+			TArray<uint8> Buffer;
+			FileHandle->Seek(IndexOffset);
+			if (FileHandle->Tell() != IndexOffset)
+			{
+				UE_LOG(LogDerivedDataCache, Error, TEXT("Pak cache was corrputed (bad seek index) %s."), InFilename);
+				return false;
+			}
+			Buffer.AddUninitialized(sizeof(uint32) * 4);
+			FileHandle->Serialize(Buffer.GetData(), int64(sizeof(uint32)* 4));
+			FMemoryReader Loader(Buffer);
+			uint32 Magic = 0;
+			Loader << Magic;
+			Loader << IndexCrc;
+			Loader << NumIndex;
+			Loader << SizeIndex;
+			if (Magic != PakCache_Magic || (SizeIndex != 0 && NumIndex == 0) || (SizeIndex == 0 && NumIndex != 0)) 
+			{
+				UE_LOG(LogDerivedDataCache, Error, TEXT("Pak cache was corrputed (bad index header) %s."), InFilename);
+				return false;
+			}
+			if (IndexOffset + sizeof(uint32) * 4 + SizeIndex != Trailer) 
+			{
+				UE_LOG(LogDerivedDataCache, Error, TEXT("Pak cache was corrputed (bad index size) %s."), InFilename);
+				return false;
+			}
+		}
+		{
+			TArray<uint8> Buffer;
+			Buffer.AddUninitialized(SizeIndex);
+			FileHandle->Serialize(Buffer.GetData(), int64(SizeIndex));
+			FMemoryReader Loader(Buffer);
+			while (Loader.Tell() < (int32)SizeIndex)
+			{
+				FString Key;
+				int64 Offset;
+				int64 Size;
+				uint32 Crc;
+				Loader << Key;
+				Loader << Offset;
+				Loader << Size;
+				Loader << Crc;
+				if (!Key.Len() || Offset < 0 || Offset >= IndexOffset || !Size)
+				{
+					UE_LOG(LogDerivedDataCache, Error, TEXT("Pak cache was corrputed (bad index entry) %s."), InFilename);
+					return false;
+				}
+				CacheItems.Add(Key, FCacheValue(Offset, Size, Crc));
+			}
+			if (CacheItems.Num() != NumIndex)
+			{
+				UE_LOG(LogDerivedDataCache, Error, TEXT("Pak cache was corrputed (bad index count) %s."), InFilename);
+				return false;
+			}
+		}
+		return true;
+	}
+
+	/**
+	 * Merges another cache file into this one.
+	 * @return true on success
+	 */
+	void MergeCache(FPakFileDerivedDataBackend* OtherPak)
+	{
+		// Get all the existing keys
+		TArray<FString> KeyNames;
+		OtherPak->CacheItems.GenerateKeyArray(KeyNames);
+
+		// Find all the keys to copy
+		TArray<FString> CopyKeyNames;
+		for(const FString& KeyName : KeyNames)
+		{
+			if(!CachedDataProbablyExists(*KeyName))
+			{
+				CopyKeyNames.Add(KeyName);
+			}
+		}
+		UE_LOG(LogDerivedDataCache, Display, TEXT("Merging %d entries (%d skipped)."), CopyKeyNames.Num(), KeyNames.Num() - CopyKeyNames.Num());
+
+		// Copy them all to the new cache. Don't use the overloaded get/put methods (which may compress/decompress); copy the raw data directly.
+		TArray<uint8> Buffer;
+		for(const FString& CopyKeyName : CopyKeyNames)
+		{
+			Buffer.Reset();
+			if(OtherPak->FPakFileDerivedDataBackend::GetCachedData(*CopyKeyName, Buffer, NULL))
+			{
+				FPakFileDerivedDataBackend::PutCachedData(*CopyKeyName, Buffer, false, NULL);
+			}
+		}
+	}
+	
+	const FString& GetFilename() const
+	{
+		return Filename;
+	}
+
+	static bool SortAndCopy(const FString &InputFilename, const FString &OutputFilename)
+	{
+		// Open the input and output files
+		FPakFileDerivedDataBackend InputPak(*InputFilename, false);
+		if (InputPak.bClosed) return false;
+
+		FPakFileDerivedDataBackend OutputPak(*OutputFilename, true);
+		if (OutputPak.bClosed) return false;
+
+		// Sort the key names
+		TArray<FString> KeyNames;
+		InputPak.CacheItems.GenerateKeyArray(KeyNames);
+		KeyNames.Sort();
+
+		// Copy all the DDC to the new cache
+		TArray<uint8> Buffer;
+		TArray<uint32> KeySizes;
+		for (int KeyIndex = 0; KeyIndex < KeyNames.Num(); KeyIndex++)
+		{
+			Buffer.Reset();
+			InputPak.GetCachedData(*KeyNames[KeyIndex], Buffer, NULL);
+			OutputPak.PutCachedData(*KeyNames[KeyIndex], Buffer, false, NULL);
+			KeySizes.Add(Buffer.Num());
+		}
+
+		// Write out a TOC listing for debugging
+		FStringOutputDevice Output;
+		Output.Logf(TEXT("Asset,Size") LINE_TERMINATOR);
+		for(int KeyIndex = 0; KeyIndex < KeyNames.Num(); KeyIndex++)
+		{
+			Output.Logf(TEXT("%s,%d") LINE_TERMINATOR, *KeyNames[KeyIndex], KeySizes[KeyIndex]);
+		}
+		FFileHelper::SaveStringToFile(Output, *FPaths::Combine(*FPaths::GetPath(OutputFilename), *(FPaths::GetBaseFilename(OutputFilename) + TEXT(".csv"))));
+		return true;
+	}
+
+private:
+
+	struct FCacheValue
+	{
+		int64 Offset;
+		int64 Size;
+		uint32 Crc;
+		FCacheValue(int64 InOffset, uint32 InSize, uint32 InCrc)
+			: Offset(InOffset)
+			, Size(InSize)
+			, Crc(InCrc)
+		{
+		}
+	};
+
+	/** When set to true, we are a pak writer (we don't do reads). */
+	bool bWriting;
+	/** When set to true, we are a pak writer and we saved, so we shouldn't be used anymore. Also, a read cache that failed to open. */
+	bool bClosed;
+	/** Object used for synchronization via a scoped lock						*/
+	FCriticalSection	SynchronizationObject;
+	/** Set of files that are being written to disk asynchronously. */
+	TMap<FString, FCacheValue> CacheItems;
+	/** File handle of pak. */
+	TAutoPtr<FArchive> FileHandle;
+	/** File name of pak. */
+	FString Filename;
+	enum 
+	{
+		/** Magic number to use in header */
+		PakCache_Magic = 0x0c7c0ddc,
+	};
+};
+
+class FCompressedPakFileDerivedDataBackend : public FPakFileDerivedDataBackend
+{
+public:
+	FCompressedPakFileDerivedDataBackend(const TCHAR* InFilename, bool bInWriting) 
+		: FPakFileDerivedDataBackend(InFilename, bInWriting)
+	{
+	}
+
+	virtual void PutCachedData(const TCHAR* CacheKey, TArray<uint8>& InData, bool bPutEvenIfExists, FCacheStatRecord* Stats) override
+	{
+		int32 UncompressedSize = InData.Num();
+		int32 CompressedSize = FCompression::CompressMemoryBound(CompressionFlags, UncompressedSize);
+
+		TArray<uint8> CompressedData;
+		CompressedData.AddUninitialized(CompressedSize + sizeof(UncompressedSize));
+
+		FMemory::Memcpy(&CompressedData[0], &UncompressedSize, sizeof(UncompressedSize));
+		verify(FCompression::CompressMemory(CompressionFlags, CompressedData.GetData() + sizeof(UncompressedSize), CompressedSize, InData.GetData(), InData.Num()));
+		CompressedData.SetNum(CompressedSize + sizeof(UncompressedSize), false);
+
+		FPakFileDerivedDataBackend::PutCachedData(CacheKey, CompressedData, bPutEvenIfExists, Stats);
+	}
+
+	virtual bool GetCachedData(const TCHAR* CacheKey, TArray<uint8>& OutData, FCacheStatRecord* Stats) override
+	{
+		TArray<uint8> CompressedData;
+		if(!FPakFileDerivedDataBackend::GetCachedData(CacheKey, CompressedData, Stats))
+		{
+			return false;
+		}
+
+		int32 UncompressedSize;
+		FMemory::Memcpy(&UncompressedSize, &CompressedData[0], sizeof(UncompressedSize));
+		OutData.SetNum(UncompressedSize);
+		verify(FCompression::UncompressMemory(CompressionFlags, OutData.GetData(), UncompressedSize, CompressedData.GetData() + sizeof(UncompressedSize), CompressedData.Num() - sizeof(UncompressedSize)));
+
+		return true;
+	}
+
+private:
+	static const ECompressionFlags CompressionFlags = (ECompressionFlags)(COMPRESS_ZLIB | COMPRESS_BiasMemory);
+};
diff --git a/Engine/Source/Developer/DerivedDataCache/Public/DerivedDataBackendInterface.h b/Engine/Source/Developer/DerivedDataCache/Public/DerivedDataBackendInterface.h
new file mode 100644
--- /dev/null
+++ b/Engine/Source/Developer/DerivedDataCache/Public/DerivedDataBackendInterface.h
@@ -0,0 +1,128 @@
+// Copyright 1998-2015 Epic Games, Inc. All Rights Reserved.
+
+
+#pragma once
+DECLARE_LOG_CATEGORY_EXTERN(LogDerivedDataCache, Log, All);
+
+DECLARE_DWORD_ACCUMULATOR_STAT_EXTERN(TEXT("Num Gets"),STAT_DDC_NumGets,STATGROUP_DDC, );
+DECLARE_DWORD_ACCUMULATOR_STAT_EXTERN(TEXT("Num Puts"),STAT_DDC_NumPuts,STATGROUP_DDC, );
+DECLARE_DWORD_ACCUMULATOR_STAT_EXTERN(TEXT("Num Build"),STAT_DDC_NumBuilds,STATGROUP_DDC, );
+DECLARE_DWORD_ACCUMULATOR_STAT_EXTERN(TEXT("Num Exists"),STAT_DDC_NumExist,STATGROUP_DDC, );
+DECLARE_FLOAT_ACCUMULATOR_STAT_EXTERN(TEXT("Sync Get Time"),STAT_DDC_SyncGetTime,STATGROUP_DDC, );
+DECLARE_FLOAT_ACCUMULATOR_STAT_EXTERN(TEXT("ASync Wait Time"),STAT_DDC_ASyncWaitTime,STATGROUP_DDC, );
+DECLARE_FLOAT_ACCUMULATOR_STAT_EXTERN(TEXT("Sync Put Time"),STAT_DDC_PutTime,STATGROUP_DDC, );
+DECLARE_FLOAT_ACCUMULATOR_STAT_EXTERN(TEXT("Sync Build Time"),STAT_DDC_SyncBuildTime,STATGROUP_DDC, );
+DECLARE_FLOAT_ACCUMULATOR_STAT_EXTERN(TEXT("Exists Time"),STAT_DDC_ExistTime,STATGROUP_DDC, );
+
+struct FCacheStatRecord
+{
+	FString CacheKey;
+	uint32 DataSize;
+	double StartTime;
+	double EndTime;
+	double GetDuration;
+	double PutDuration;
+	uint32 CacheGet;
+	bool bSynchronous;
+	bool bFromNetwork;
+};
+
+
+/** 
+ * Interface for cache server backends. 
+ * The entire API should be callable from any thread (except the singleton can be assumed to be called at least once before concurrent access).
+**/
+class FDerivedDataBackendInterface
+{
+public:
+	virtual ~FDerivedDataBackendInterface()
+	{
+	}
+
+	/** return true if this cache is writable **/
+	virtual bool IsWritable()=0;
+
+	/** 
+	 * return true if hits on this cache should propagate to lower cache level. Typically false for a PAK file. 
+	 * Caution! This generally isn't propagated, so the thing that returns false must be a direct child of the heirarchical cache.
+	 **/
+	virtual bool BackfillLowerCacheLevels()
+	{
+		return true;
+	}
+
+	/**
+	 * Synchronous test for the existence of a cache item
+	 *
+	 * @param	CacheKey	Alphanumeric+underscore key of this cache item
+	 * @return				true if the data probably will be found, this can't be guaranteed because of concurrency in the backends, corruption, etc
+	 */
+	virtual bool CachedDataProbablyExists(const TCHAR* CacheKey)=0;
+	/**
+	 * Synchronous retrieve of a cache item
+	 *
+	 * @param	CacheKey	Alphanumeric+underscore key of this cache item
+	 * @param	OutData		Buffer to receive the results, if any were found
+	 * @return				true if any data was found, and in this case OutData is non-empty
+	 */
+	virtual bool GetCachedData(const TCHAR* CacheKey, TArray<uint8>& OutData, FCacheStatRecord* Stats = NULL)=0;
+	/**
+	 * Asynchronous, fire-and-forget placement of a cache item
+	 *
+	 * @param	CacheKey			Alphanumeric+underscore key of this cache item
+	 * @param	InData				Buffer containing the data to cache, can be destroyed after the call returns, immediately
+	 * @param	bPutEvenIfExists	If true, then do not attempt skip the put even if CachedDataProbablyExists returns true
+	 */
+	virtual void PutCachedData(const TCHAR* CacheKey, TArray<uint8>& InData, bool bPutEvenIfExists, FCacheStatRecord* Stats = NULL) = 0;
+
+	/**
+	 * Remove data from cache (used in the event that corruption is detected at a higher level and possibly house keeping)
+	 *
+	 * @param	CacheKey	Alphanumeric+underscore key of this cache item
+	 * @param	bTransient	true if the data is transient and it is up to the backend to decide when and if to remove cached data.
+	 */
+	virtual void RemoveCachedData(const TCHAR* CacheKey, bool bTransient)=0;
+};
+
+class FDerivedDataBackend
+{
+public:
+	/**
+	 * Singleton to retrieve the GLOBAL backend
+	 *
+	 * @return Reference to the global cache backend
+	 */
+	static FDerivedDataBackend& Get();
+
+	/**
+	 * Singleton to retrieve the root cache
+	 * @return Reference to the global cache root
+	 */
+	virtual FDerivedDataBackendInterface& GetRoot() = 0;
+
+	//--------------------
+	// System Interface, copied from FDerivedDataCacheInterface
+	//--------------------
+
+	virtual void NotifyBootComplete() = 0;
+	virtual void AddToAsyncCompletionCounter(int32 Addend) = 0;
+	virtual void WaitForQuiescence(bool bShutdown = false) = 0;
+	virtual void GetDirectories(TArray<FString>& OutResults) = 0;
+
+	/**
+	 * Mounts a read-only pak file.
+	 *
+	 * @param PakFilename Pak filename
+	 */
+	virtual FDerivedDataBackendInterface* MountPakFile(const TCHAR* PakFilename) = 0;
+
+	/**
+	 * Unmounts a read-only pak file.
+	 *
+	 * @param PakFilename Pak filename
+	 */
+	virtual bool UnmountPakFile(const TCHAR* PakFilename) = 0;
+
+};
+
+
diff --git a/Engine/Source/Developer/DerivedDataCache/Public/DerivedDataCacheInterface.h b/Engine/Source/Developer/DerivedDataCache/Public/DerivedDataCacheInterface.h
new file mode 100644
--- /dev/null
+++ b/Engine/Source/Developer/DerivedDataCache/Public/DerivedDataCacheInterface.h
@@ -0,0 +1,214 @@
+// Copyright 1998-2015 Epic Games, Inc. All Rights Reserved.
+
+
+#pragma once
+#include "ModuleManager.h"
+
+/** 
+ * Opaque class for rollup handling
+**/
+class IDerivedDataRollup
+{
+public:
+	virtual ~IDerivedDataRollup()
+	{
+	}
+};
+
+/** 
+ * Interface for the derived data cache
+ * This API is fully threadsafe (with the possible exception of the system interface: NotfiyBootComplete, etc).
+**/
+class FDerivedDataCacheInterface
+{
+public:
+	virtual ~FDerivedDataCacheInterface()
+	{
+	}
+
+	//--------------------
+	// High Level Interface
+	//--------------------
+
+	/** 
+	 * Synchronously checks the cache and if the item is present, it returns the cached results, otherwise tells the deriver to build the data and then updates the cache
+	 * @param	DataDeriver	plugin to produce cache key and in the event of a miss, return the data.
+	 * @return	true if the data was retrieved from the cache or the deriver built the data sucessfully. false can only occur if the plugin returns false.
+	**/
+	virtual bool GetSynchronous(class FDerivedDataPluginInterface* DataDeriver, TArray<uint8>& OutData) = 0;
+
+	/** 
+	 * Starts the async process of checking the cache and if the item is present, retrieving the cached results, otherwise telling the deriver to build the data and then updating the cache
+	 * If the plugin does not support threading, all of the above will be completed before the call returns.
+	 * @param	DataDeriver	plugin to produce cache key and in the event of a miss, return the data.
+	 * @return	a handle that can be used for PollAsynchronousCompletion, WaitAsynchronousCompletion and GetAsynchronousResults
+	**/
+	virtual uint32 GetAsynchronous(class FDerivedDataPluginInterface* DataDeriver) = 0;
+	/** 
+	 * Polls a previous GetAsynchronous get for completion.
+	 * @param	Handle	Handle returned from GetAsynchronous.
+	 * @return			true if the build / retrieve is complete and the results can be obtained.
+	**/
+	virtual bool PollAsynchronousCompletion(uint32 Handle) = 0;
+	/** 
+	 * Blocks the current thread until an previous GetAsynchronous request is ready
+	 * @param	Handle	Handle returned from GetAsynchronous.
+	**/
+	virtual void WaitAsynchronousCompletion(uint32 Handle) = 0;
+	/** 
+	 * Retrieves the results from an async lookup / build. MUST only be called after the results are known to be ready by one of the aforementioned calls.
+	 * @param	Handle	Handle returned from GetAsynchronous.
+	 * @param	OutData	Array to receive the output results.
+	 * @return			true if the data was retrieved from the cache or the deriver built the data successfully. false can only occur if the plugin returns false.
+	**/
+	virtual bool GetAsynchronousResults(uint32 Handle, TArray<uint8>& OutData) = 0;
+
+	//--------------------------
+	// Low Level Static Helpers
+	//--------------------------
+
+	/** 
+	 * Static function to make sure a cache key contains only legal characters by using an escape
+	 * @param CacheKey							Cache key to sanitize
+	 * @return									Sanitized cache key
+	**/
+	static FString SanitizeCacheKey(const TCHAR* CacheKey)
+	{
+		FString Output;
+		FString Input(CacheKey);
+		int32 StartValid = 0;
+		int32 NumValid = 0;
+
+		for (int32 i = 0; i < Input.Len(); i++)
+		{
+			if (FChar::IsAlnum(Input[i]) || FChar::IsUnderscore(Input[i]))
+			{
+				NumValid++;
+			}
+			else
+			{
+				// Copy the valid range so far
+				Output += Input.Mid(StartValid, NumValid);
+
+				// Reset valid ranges
+				StartValid = i + 1;
+				NumValid = 0;
+
+				// Replace the invalid character with a special string
+				Output += FString::Printf(TEXT("$%x"), uint32(Input[i]));
+			}
+		}
+
+		// Just return the input if the entire string was valid
+		if (StartValid == 0 && NumValid == Input.Len())
+		{
+			return Input;
+		}
+		else if (NumValid > 0)
+		{
+			// Copy the remaining valid part
+			Output += Input.Mid(StartValid, NumValid);
+		}
+
+		return Output;
+	}
+
+	/** 
+	 * Static function to build a cache key out of the plugin name, versions and plugin specific info
+	 * @param PluginName						Name of the derived data type
+	 * @param VersionString						Version string of the derived data
+	 * @param PluginSpecificCacheKeySuffix		GUIDS and / or Hashes, etc to uniquely identify the specific cache entry
+	 * @return									Assembled cache key
+	**/
+	static FString BuildCacheKey(const TCHAR* PluginName, const TCHAR* VersionString, const TCHAR* PluginSpecificCacheKeySuffix)
+	{
+		return SanitizeCacheKey(*FString::Printf(TEXT("%s_%s_%s"), PluginName, VersionString, PluginSpecificCacheKeySuffix));
+	}
+
+	//--------------------
+	// Low Level Interface
+	//--------------------
+
+	/** 
+	 * Starts a rollup. Use this for GetAsynchronous calls, then call close on it
+	 * @return	a pointer to a rollup
+	**/
+	virtual IDerivedDataRollup* StartRollup() = 0;
+
+	/** 
+	 * Ends a rollup. Use this for GetAsynchronous calls, then call close on it
+	 * @param Rollup	reference to a pointer to the rollup to close...also sets the pointer to NULL
+	**/
+	virtual void EndRollup(IDerivedDataRollup*& Rollup) = 0;
+
+	/** 
+	 * Synchronously checks the cache and if the item is present, it returns the cached results, otherwise it returns false
+	 * @param	CacheKey	Key to identify the data
+	 * @return	true if the data was retrieved from the cache
+	**/
+	virtual bool GetSynchronous(const TCHAR* CacheKey, TArray<uint8>& OutData) = 0; 
+
+	/** 
+	 * Starts the async process of checking the cache and if the item is present, retrieving the cached results
+	 * @param	CacheKey	Key to identify the data
+	 * @param	Rollup		Rollup pointer, if this is part of a rollup
+	 * @return	a handle that can be used for PollAsynchronousCompletion, WaitAsynchronousCompletion and GetAsynchronousResults
+	**/
+	virtual uint32 GetAsynchronous(const TCHAR* CacheKey, IDerivedDataRollup* Rollup = NULL) = 0;
+
+	/** 
+	 * Puts data into the cache. This is fire-and-forget and typically asynchronous.
+	 * @param	CacheKey	Key to identify the data
+	 * @param	Data		Data to put in the cache under this key
+	**/
+	virtual void Put(const TCHAR* CacheKey, TArray<uint8>& Data, bool bPutEvenIfExists = false) = 0;
+
+	/**
+	 * Hint that the data associated with the key is transient and may be optionally purged from the cache.
+	 * @param	CacheKey	Key that is associated with transient data.
+	 */
+	virtual void MarkTransient(const TCHAR* CacheKey) = 0;
+
+	/**
+	 * Returns true if the data associated with the key is likely to exist in the cache.
+	 * Even if this function returns true, a get for this key may still fail!
+	 * @param	CacheKey	Key to see if data probably exists.
+	 */
+	virtual bool CachedDataProbablyExists(const TCHAR* CacheKey) = 0;
+
+	//--------------------
+	// System Interface
+	//--------------------
+
+	/**
+	 * Notify the system that the boot process is complete and so we can write the boot cache and get rid of it
+	 */
+	virtual void NotifyBootComplete() = 0;
+
+	/**
+	 * Adds or subtracts a number from the thread safe counter which tracks outstand async requests. This is used to ensure everything is complete prior to shutdown.
+	 */
+	virtual void AddToAsyncCompletionCounter(int32 Addend) = 0;
+
+	/**
+	 * Wait for all outstanding async DDC operations to complete.
+	 */
+	virtual void WaitForQuiescence(bool bShutdown = false) = 0;
+
+	/**
+	 * Retrieve the directories used by the DDC
+	 */
+	virtual void GetDirectories(TArray<FString>& OutResults) = 0;
+
+};
+
+/**
+ * Module for the DDC
+ */
+class IDerivedDataCacheModule : public IModuleInterface
+{
+public:
+	/** Return the DDC interface **/
+	virtual FDerivedDataCacheInterface& GetDDC() = 0;
+};
+
diff --git a/Engine/Source/Developer/DerivedDataCache/Public/DerivedDataPluginInterface.h b/Engine/Source/Developer/DerivedDataCache/Public/DerivedDataPluginInterface.h
new file mode 100644
--- /dev/null
+++ b/Engine/Source/Developer/DerivedDataCache/Public/DerivedDataPluginInterface.h
@@ -0,0 +1,45 @@
+// Copyright 1998-2015 Epic Games, Inc. All Rights Reserved.
+
+
+#pragma once
+/** 
+ * Interface for data deriving backends
+ * This API will not be called concurrently, except that Build might be called on different instances if IsBuildThreadsafe.
+**/
+class FDerivedDataPluginInterface
+{
+public:
+	virtual ~FDerivedDataPluginInterface()
+	{
+	}
+	/** 
+	* Get the plugin name, this is used as the first part of the cache key 
+	* @return	Name of the plugin
+	**/
+ 	virtual const TCHAR* GetPluginName() const = 0;
+	/** 
+	* Get the version of the plugin, this is used as part of the cache key. This is supposed to
+	* be a guid string ( ex. "69C8C8A6-A9F8-4EFC-875C-CFBB72E66486" )
+	* @return	Version string of the plugin
+	**/
+ 	virtual const TCHAR* GetVersionString() const = 0;
+	
+	/** 
+	* Returns the largest and plugin specific part of the cache key. This must be a alphanumeric+underscore
+	* @return	Version number of the plugin, for licensees.
+	**/
+ 	virtual FString GetPluginSpecificCacheKeySuffix() const = 0;
+	/** 
+	* Indicates that this plugin is threadsafe. Note, the system itself will not call it concurrently if this false, however, then you are responsible for not calling the system itself concurrently.
+	* @return	true if this plugin is threadsafe
+	**/
+	virtual bool IsBuildThreadsafe() const = 0;
+	/** 
+	* Does the work of deriving the data. 
+	* @param	OutData	Array of bytes to fill in with the result data
+	* @return	true if successful, in the event of failure the cache is not updated and failure is propagated to the original caller.
+	**/
+ 	virtual bool Build(TArray<uint8>& OutData) = 0;
+};
+
+
diff --git a/Engine/Source/Developer/DerivedDataCache/Public/DerivedDataUtilsInterface.h b/Engine/Source/Developer/DerivedDataCache/Public/DerivedDataUtilsInterface.h
new file mode 100644
--- /dev/null
+++ b/Engine/Source/Developer/DerivedDataCache/Public/DerivedDataUtilsInterface.h
@@ -0,0 +1,14 @@
+// Copyright 1998-2015 Epic Games, Inc. All Rights Reserved.
+#pragma once
+
+/** 
+ * Interface for DDC utility modules.
+ **/
+class IDDCUtilsModuleInterface : public IModuleInterface
+{
+public:
+	/**
+	 * Returns the path to the shared cache location given its name.
+	 **/
+	virtual FString GetSharedCachePath(const FString& CacheName) = 0;
+};
\ No newline at end of file
