Some fixes for the DDC stat collection

diff --git a/Engine/Source/Developer/DerivedDataCache/Private/DerivedDataCache.cpp b/Engine/Source/Developer/DerivedDataCache/Private/DerivedDataCache.cpp
--- a/Engine/Source/Developer/DerivedDataCache/Private/DerivedDataCache.cpp
+++ b/Engine/Source/Developer/DerivedDataCache/Private/DerivedDataCache.cpp
@@ -60,11 +60,11 @@ class FDerivedDataCache : public FDerivedDataCacheInterface
 					SCOPE_SECONDS_COUNTER(ThisTime);
 					GetRecord->CacheKey = CacheKey;
 					GetRecord->CacheGet = 1;
+					GetRecord->PutDuration = 0.;
 					GetRecord->bSynchronous = bSynchronousForStats;
 					GetRecord->StartTime = FPlatformTime::Seconds();
-					GetRecord->GetDuration = 0.;
-					GetRecord->PutDuration = 0.;
 					bGetResult = FDerivedDataBackend::Get().GetRoot().GetCachedData(*CacheKey, Data, GetRecord);
+					GetRecord->GetDuration = FPlatformTime::Seconds() - GetRecord->StartTime;
 					GetRecord->EndTime = FPlatformTime::Seconds();
 					GetRecord->DataSize = Data.Num();
 				}
@@ -99,11 +99,11 @@ class FDerivedDataCache : public FDerivedDataCacheInterface
 						SCOPE_SECONDS_COUNTER(ThisTime);
 						PutRecord->CacheKey = CacheKey;
 						PutRecord->CacheGet = 0;
+						PutRecord->GetDuration = 0.f;
 						PutRecord->bSynchronous = bSynchronousForStats;
 						PutRecord->StartTime = FPlatformTime::Seconds();
-						PutRecord->GetDuration = 0.;
-						PutRecord->PutDuration = 0.;
 						FDerivedDataBackend::Get().GetRoot().PutCachedData(*CacheKey, Data, true, PutRecord);
+						PutRecord->PutDuration = FPlatformTime::Seconds() - PutRecord->StartTime;
 						PutRecord->EndTime = FPlatformTime::Seconds();
 						PutRecord->DataSize = Data.Num();
 					}
@@ -345,10 +345,9 @@ public:
 			Record->bSynchronous = true;
 			Record->DataSize = Data.Num();
 			Record->StartTime = FPlatformTime::Seconds();
-			Record->PutDuration = 0;
-			Record->GetDuration = 0;
 			FDerivedDataBackend::Get().GetRoot().PutCachedData(CacheKey, Data, bPutEvenIfExists);
 			Record->EndTime = FPlatformTime::Seconds();
+			Record->PutDuration = Record->EndTime - Record->StartTime;
 		}
 		INC_FLOAT_STAT_BY(STAT_DDC_PutTime,(float)ThisTime);
 		INC_DWORD_STAT(STAT_DDC_NumPuts);
@@ -417,13 +416,13 @@ public:
 
 			if (OutputFile)
 			{
-				OutputFile->Logf(TEXT("Key,Sync,Network,Size,Total Duration(ms),Get Duration(ms),Put Duration(ms),Get"));
+				OutputFile->Logf(TEXT("Key,Sync,FromNetwork,ToNetwork,Size,Get Duration(ms),Put Duration(ms),Get"));
 				for (int32 Index = 0; Index < CacheStats.Num(); Index++)
 				{
 					FCacheStatRecord* Record = CacheStats[Index];
 					if (Record->CacheGet != -1)
 					{
-						OutputFile->Logf(TEXT("%s,%s,%s,%d,%lf,%lf,%lf,%d"), *(Record->CacheKey), Record->bSynchronous ? TEXT("true") : TEXT("false"), Record->bFromNetwork ? TEXT("true") : TEXT("false"), Record->DataSize, (Record->EndTime - Record->StartTime) * 1000.0f, Record->GetDuration * 1000.0f, Record->PutDuration * 1000.0f, Record->CacheGet);
+						OutputFile->Logf(TEXT("%s,%s,%s,%s,%d,%lf,%lf,%d"), *(Record->CacheKey), Record->bSynchronous ? TEXT("true") : TEXT("false"), Record->bFromNetwork ? TEXT("true") : TEXT("false"), Record->bToNetwork ? TEXT("true") : TEXT("false"), Record->DataSize, Record->GetDuration * 1000.0f, Record->PutDuration * 1000.0f, Record->CacheGet);
 					}
 				}
 				delete OutputFile;
diff --git a/Engine/Source/Developer/DerivedDataCache/Private/FileSystemDerivedDataBackend.cpp b/Engine/Source/Developer/DerivedDataCache/Private/FileSystemDerivedDataBackend.cpp
--- a/Engine/Source/Developer/DerivedDataCache/Private/FileSystemDerivedDataBackend.cpp
+++ b/Engine/Source/Developer/DerivedDataCache/Private/FileSystemDerivedDataBackend.cpp
@@ -237,11 +237,11 @@ public:
 					Stats->CacheKey = Filename;
 					if (Filename.StartsWith(TEXT("//")) || Filename.StartsWith(TEXT("\\\\")))
 					{
-						Stats->bFromNetwork = true;
+						Stats->bToNetwork = true;
 					}
 					else
 					{
-						Stats->bFromNetwork = false;
+						Stats->bToNetwork = false;
 					}
 				}
 				FString TempFilename(TEXT("temp.")); 
diff --git a/Engine/Source/Developer/DerivedDataCache/Private/HierarchicalDerivedDataBackend.h b/Engine/Source/Developer/DerivedDataCache/Private/HierarchicalDerivedDataBackend.h
--- a/Engine/Source/Developer/DerivedDataCache/Private/HierarchicalDerivedDataBackend.h
+++ b/Engine/Source/Developer/DerivedDataCache/Private/HierarchicalDerivedDataBackend.h
@@ -96,10 +96,14 @@ public:
 		FDDCScopeStatHelper Stat(CacheKey, NAME_GetCachedData);
 		Stat.AddTag(NAME_HierarchicalDDC, FString());
 
+		Stats->StartTime = FPlatformTime::Seconds();
 		for (int32 CacheIndex = 0; CacheIndex < InnerBackends.Num(); CacheIndex++)
 		{
 			if (InnerBackends[CacheIndex]->CachedDataProbablyExists(CacheKey) && InnerBackends[CacheIndex]->GetCachedData(CacheKey, OutData, Stats))
 			{
+				Stats->GetDuration = FPlatformTime::Seconds() - Stats->StartTime;
+				Stats->DataSize = OutData.Num();
+				double StartTime = FPlatformTime::Seconds();
 				if (bIsWritable)
 				{
 					// fill in the higher level caches
@@ -141,9 +145,13 @@ public:
 				Stat.AddTag(NAME_DataSize, FString::Printf(TEXT("%d bytes"), OutData.Num()));
 				Stat.AddTag(NAME_Retrieved, TEXT("true"));
 				Stats->EndTime = FPlatformTime::Seconds();
+				Stats->PutDuration = Stats->EndTime - StartTime;
 				return true;
 			}
 		}
+		Stats->EndTime = FPlatformTime::Seconds();
+		Stats->GetDuration = FPlatformTime::Seconds() - Stats->StartTime;
+		Stats->PutDuration = 0.;
 		Stat.AddTag(NAME_Retrieved, TEXT("false"));
 		return false;
 	}
diff --git a/Engine/Source/Developer/DerivedDataCache/Private/PakFileDerivedDataBackend.h b/Engine/Source/Developer/DerivedDataCache/Private/PakFileDerivedDataBackend.h
--- a/Engine/Source/Developer/DerivedDataCache/Private/PakFileDerivedDataBackend.h
+++ b/Engine/Source/Developer/DerivedDataCache/Private/PakFileDerivedDataBackend.h
@@ -159,7 +159,7 @@ public:
 		{
 			FScopeLock ScopeLock(&SynchronizationObject);
 			if (Stats)
-				Stats->bFromNetwork = false;
+				Stats->bToNetwork = false;
 			FString Key(CacheKey);
 			FCacheValue* Item = CacheItems.Find(FString(CacheKey));
 			if (!Item)
diff --git a/Engine/Source/Developer/DerivedDataCache/Public/DerivedDataBackendInterface.h b/Engine/Source/Developer/DerivedDataCache/Public/DerivedDataBackendInterface.h
--- a/Engine/Source/Developer/DerivedDataCache/Public/DerivedDataBackendInterface.h
+++ b/Engine/Source/Developer/DerivedDataCache/Public/DerivedDataBackendInterface.h
@@ -25,6 +25,7 @@ struct FCacheStatRecord
 	uint32 CacheGet;
 	bool bSynchronous;
 	bool bFromNetwork;
+	bool bToNetwork;
 };
 
 
