Fixed crash on exit when running -game.

diff --git a/Engine/Source/Developer/DerivedDataCache/Private/DerivedDataCache.cpp b/Engine/Source/Developer/DerivedDataCache/Private/DerivedDataCache.cpp
--- a/Engine/Source/Developer/DerivedDataCache/Private/DerivedDataCache.cpp
+++ b/Engine/Source/Developer/DerivedDataCache/Private/DerivedDataCache.cpp
@@ -1031,6 +1031,14 @@ public:
 		return InternalSingleton();
 	}
 
+	virtual void StartupModule() override
+	{
+#if WITH_EDITOR
+		// Make sure the CookingStats module gets loaded on the correct thread (used by DDCStats on a background thread)
+		FModuleManager::Get().LoadModule(TEXT("CookingStats"));
+#endif // WITH_EDITOR
+	}
+
 	virtual void ShutdownModule() override
 	{
 		FDDCCleanup::Shutdown();
