Disable DDC speed warnings on build machines. These messages can be useful when working from home or a different studio, but on builders they're usually caused by momentary saturation of the network and cause unactionable notification emails to be sent to developers.

diff --git a/Engine/Source/Developer/DerivedDataCache/Private/FileSystemDerivedDataBackend.cpp b/Engine/Source/Developer/DerivedDataCache/Private/FileSystemDerivedDataBackend.cpp
--- a/Engine/Source/Developer/DerivedDataCache/Private/FileSystemDerivedDataBackend.cpp
+++ b/Engine/Source/Developer/DerivedDataCache/Private/FileSystemDerivedDataBackend.cpp
@@ -103,7 +103,7 @@ public:
 			UE_LOG(LogDerivedDataCache, Display, TEXT("Files in %s will be touched."),*CachePath);
 		}
 
-		if (!bFailed && AccessDuration > SlowInitDuration)
+		if (!bFailed && AccessDuration > SlowInitDuration && !GIsBuildMachine)
 		{
 			UE_LOG(LogDerivedDataCache, Warning, TEXT("%s access is very slow (initialization took %.2lf seconds), consider disabling it."), *CachePath, AccessDuration);
 		}
@@ -176,10 +176,13 @@ public:
 		double StartTime = FPlatformTime::Seconds();
 		if (FFileHelper::LoadFileToArray(Data,*Filename,FILEREAD_Silent))
 		{
-			double ReadDuration = FPlatformTime::Seconds() - StartTime;
-			double ReadSpeed = ReadDuration > 5.0 ? (Data.Num() / ReadDuration) / (1024.0 * 1024.0) : 100.0;
-			// Slower than 0.5MB/s?
-			UE_CLOG(ReadSpeed < 0.5, LogDerivedDataCache, Warning, TEXT("%s is very slow (%.2lfMB/s) when accessing %s, consider disabling it."), *CachePath, ReadSpeed, *Filename);
+			if(!GIsBuildMachine)
+			{
+				double ReadDuration = FPlatformTime::Seconds() - StartTime;
+				double ReadSpeed = ReadDuration > 5.0 ? (Data.Num() / ReadDuration) / (1024.0 * 1024.0) : 100.0;
+				// Slower than 0.5MB/s?
+				UE_CLOG(ReadSpeed < 0.5, LogDerivedDataCache, Warning, TEXT("%s is very slow (%.2lfMB/s) when accessing %s, consider disabling it."), *CachePath, ReadSpeed, *Filename);
+			}
 
 			UE_LOG(LogDerivedDataCache, Verbose, TEXT("FileSystemDerivedDataBackend: Cache hit on %s"),*Filename);
 #if ENABLE_DDC_STATS
