Fix error parsing documentation - UnrealDocTool barfs if there's a source file with two % characters on the same line, because it parses it as a variable reference.

diff --git a/Engine/Source/Developer/DerivedDataCache/Public/DerivedDataCacheUsageStats.h b/Engine/Source/Developer/DerivedDataCache/Public/DerivedDataCacheUsageStats.h
--- a/Engine/Source/Developer/DerivedDataCache/Public/DerivedDataCacheUsageStats.h
+++ b/Engine/Source/Developer/DerivedDataCache/Public/DerivedDataCacheUsageStats.h
@@ -48,7 +48,7 @@
  *       <override GatherUsageStats>
  *       {
  *           // Add this node's UsageStats to the usage map. Your Key name should be UNIQUE to the entire graph (so use the file name, or pointer to this if you have to).
- *           UsageStatsMap.Add(FString::Printf(TEXT("%s: %s"), *GraphPath, <Some unique name for this node instance>), UsageStats);
+ *           UsageStatsMap.Add(FString::Printf(TEXT("%s: <Some unique name for this node instance>"), *GraphPath), UsageStats);
  *       }
 *   }
  */
