Removed DDC cache size warning (per Gil should not be a warning)

diff --git a/Engine/Source/Developer/DerivedDataCache/Private/MemoryDerivedDataBackend.h b/Engine/Source/Developer/DerivedDataCache/Private/MemoryDerivedDataBackend.h
--- a/Engine/Source/Developer/DerivedDataCache/Private/MemoryDerivedDataBackend.h
+++ b/Engine/Source/Developer/DerivedDataCache/Private/MemoryDerivedDataBackend.h
@@ -114,8 +114,7 @@ public:
 			if (MaxCacheSize > 0 && (CurrentCacheSize + CacheValueSize) > MaxCacheSize)
 			{
 				delete Val;
-
-				UE_LOG(LogDerivedDataCache, Warning, TEXT("Failed to cache data. Maximum cache size reached. CurrentSize %d kb / MaxSize: %d kb"), CurrentCacheSize / 1024, MaxCacheSize / 1024);
+				UE_LOG(LogDerivedDataCache, Display, TEXT("Failed to cache data. Maximum cache size reached. CurrentSize %d kb / MaxSize: %d kb"), CurrentCacheSize / 1024, MaxCacheSize / 1024);
 				bMaxSizeExceeded = true;
 			}
 			else
