Prevent double-free of DDC nodes due to duplicated entries in the CreatedBackends array. Can happen if the hierarchy node decides to optimize itself out of the graph.

diff --git a/Engine/Source/Developer/DerivedDataCache/Private/DerivedDataBackends.cpp b/Engine/Source/Developer/DerivedDataCache/Private/DerivedDataBackends.cpp
--- a/Engine/Source/Developer/DerivedDataCache/Private/DerivedDataBackends.cpp
+++ b/Engine/Source/Developer/DerivedDataCache/Private/DerivedDataBackends.cpp
@@ -202,7 +202,7 @@ public:
 			// Store this node so that we don't require any order of adding backend nodes
 			InParsedNodes.Add( NodeName, ParsedNode );
 			// Keep references to all created nodes.
-			CreatedBackends.Add( ParsedNode );
+			CreatedBackends.AddUnique( ParsedNode );
 		}
 
 		return ParsedNode;
