Copying //UE4/Dev-Core to //UE4/Dev-Main (Source: //UE4/Dev-Core @ 3208226)

==========================
MAJOR FEATURES + CHANGES
==========================

Change 3173153 on 2016/10/25 by [REDACTED_PERSON]

	Pak signing changes
	 - Integrated into EDL loader
	 - Changed to not encrypt each CRC in the sig file, rather just store a single encryped signature of the entire sig file. Removes need to decrypt thousands of signatures at startup.

Change 3173531 on 2016/10/25 by [REDACTED_PERSON]

	Removing unused j query packages.

Change 3174743 on 2016/10/26 by [REDACTED_PERSON]

	UE4 - fixed COTF with EDL

Change 3177896 on 2016/10/28 by [REDACTED_PERSON]

	TSharedPtr and TSharedRef aliasing constructors.
	Removal of static_asserts for TSharedPtr<UObject>.

Change 3180343 on 2016/10/31 by [REDACTED_PERSON]

	Reimplementation of changes from CL#s 3050329 and 3105715 that were lost in merges 3094597 and 3105741.

Change 3181382 on 2016/11/01 by [REDACTED_PERSON]

	Visual Studio debugger visualizers for delegates.

Change 3182738 on 2016/11/02 by [REDACTED_PERSON]

	Re-enable signed archive reader so non-pakpreacher based reads still get signature checked

Change 3183420 on 2016/11/02 by [REDACTED_PERSON]

	Fix to TIsZeroConstructType for TScriptDelegate.

Change 3184872 on 2016/11/03 by [REDACTED_PERSON]

	Fixing memory stomps in SSL certificate initialization (found with mallocstomp)

Change 3184873 on 2016/11/03 by [REDACTED_PERSON]

	Adding thread safety checks to async loading code

Change 3185535 on 2016/11/03 by [REDACTED_PERSON]

	Fix it so calling CreateDefaultSubobject with bTransient = true sets the object transient flag. This fixes EDL Crashes involving components.

Change 3186636 on 2016/11/04 by [REDACTED_PERSON]

	AES encryption integrated into EDL system
	Pak signing and AES encryption now configurable by ini files rather than magical text files

Change 3186637 on 2016/11/04 by [REDACTED_PERSON]

	Configured pak signing and encryption in ShooterGame for reference

Change 3186639 on 2016/11/04 by [REDACTED_PERSON]

	Encryption changes for Orion
	* Move pak signing keys into new INI format
	* Add AES key and enable INI file encryption

Change 3186661 on 2016/11/04 by [REDACTED_PERSON]

	Change unrealpak command line params to accept AES key as a separete parameter

Change 3186670 on 2016/11/04 by [REDACTED_PERSON]

	Adding a null check before using a package pointer in Linker code

Change 3186775 on 2016/11/04 by [REDACTED_PERSON]

	Fix UBT defines that come in as quoted strings, losing the quotes when passed to the compiler
	 - PS4 and Mac fixes. Other platforms might need fixing too!

Change 3186823 on 2016/11/04 by [REDACTED_PERSON]

	Fixed an incorrect size check in the EDL pak signing code

Change 3186925 on 2016/11/04 by [REDACTED_PERSON]

	Allow UnrealPak to read encryption settings from project ini files

Change 3189885 on 2016/11/08 by [REDACTED_PERSON]

	Static analysis warning fix

Change 3190015 on 2016/11/08 by [REDACTED_PERSON]

	Thread safety fix for UBlueprintGeneratedClass::PostLoadDefaultObject while UBlueprintGeneratedClass::SerializeDefaultObject runs on the async loading thread

Change 3190253 on 2016/11/08 by [REDACTED_PERSON]

	Improved MDD performance for on the CR server.
	[UE-37566] - Improve MDD performance on CR server

	Blocked MDD init'ing the crash handling code as it isn't desirable on the server.
	Removed redundant call to SetSymbolPathsFromModules() from CrashDebugHelper.

Change 3192993 on 2016/11/10 by [REDACTED_PERSON]

	Thread Heartbeat will no longer report the same hang multiple times.

Change 3193111 on 2016/11/10 by [REDACTED_PERSON]

	Minor change in the condition that detects the same hangs - allow the same callstacks from different threads

Change 3193168 on 2016/11/10 by [REDACTED_PERSON]

	TSparseArray now reserves space in reverse so that new elements get added to the front of the allocation rather than the back, which is better for memory traversal and meets expectations more closely.

Change 3193171 on 2016/11/10 by [REDACTED_PERSON]

	Easier debugging of FPendingRegistrantInfo map.

Change 3193188 on 2016/11/10 by [REDACTED_PERSON]

	TAutoPointer deprecated.

Change 3193796 on 2016/11/10 by [REDACTED_PERSON]

	Fix pak creation failure when no pak signing keys are supplied

Change 3194524 on 2016/11/11 by [REDACTED_PERSON]

	Another static analysis warning fix

Change 3195119 on 2016/11/11 by [REDACTED_PERSON]

	TAutoPtr deprecated.
	Fixes to use of TAutoPtr with incompatible memory deallocations (TAutoPtr with FMemory::Malloc and new[]).
	Some large headers moved into .cpp files.

Change 3196582 on 2016/11/14 by [REDACTED_PERSON]

	UE4 - Changed a check to a warning related to detaching linekrs twice. Seen in nativized BP version of platformer game.

Change 3196878 on 2016/11/14 by [REDACTED_PERSON]

	TScopedPointer deprecated.

Change 3198061 on 2016/11/15 by [REDACTED_PERSON]

	Class array is no longer regenerated when saving UClasses.

Change 3198065 on 2016/11/15 by [REDACTED_PERSON]

	Making AssembleReferenceTokenStream thread safe for blueprints loaded on the async loading thread.

Change 3198199 on 2016/11/15 by [REDACTED_PERSON]

	Pak platform file will now only be used if pak files exist regardless of command line paraks like -pak, -singedpak and -signed.

Change 3199954 on 2016/11/16 by [REDACTED_PERSON]

	Removing USING_SIGNED_CONTENT

Change 3200221 on 2016/11/16 by [REDACTED_PERSON]

	CrashReportProcess code cleanup - removing unused using directives

Change 3200232 on 2016/11/16 by [REDACTED_PERSON]

	Multiple CrashReportProcess updates and improvements (CRP v1.2.6)
	UE-36248 - CRP scalability: All bulk storage or shared data to S3 or suitable network drives

	InvalidCrashReports now saved to S3 instead of local folder
	Removed option tosync MinidumpDiagnostics from Perforce
	Moved MinidumpDiagnostics from old Perforce synched location to its own folder in E:\Services (makes more sense with manual publishing)
	Added improved logging to Slack with option to monitor MDD performance
	Added hourly log folders to MDD logs
	Added support for types of crashes we don't want to symbolicate (using it to skip callstack gen for hang detected ensures)

Change 3200382 on 2016/11/16 by [REDACTED_PERSON]

	Async Loading code will now detach the linker when resetting async package loader to avoid situations when loading the same asset multiple times results in the following load request finding the old linker after the package has been loading but the async package hasn't been deleted yet (async package for the old request in limbo state but linker exists).

Change 3200562 on 2016/11/16 by [REDACTED_PERSON]

	UE4 - Fixed rare issue with reloading nativized blueprints with the EDL and a minor simplication.

Change 3201093 on 2016/11/16 by [REDACTED_PERSON]

	#UE 38654 Fix EDL cooking to correctly search components created directly by UBlueprints, as well as the CDO components it already covered. Also explicitly mark subobject templates as editor only.
	Fix issue where the AssetImportData associated with Blueprint-owned Curves was ending up in the cooked subobject template list. Stopped it from creating those objects, and mark the class editor only.

Change 3201736 on 2016/11/17 by [REDACTED_PERSON]

	Strtoi64 platform and TCString functions.

Change 3201938 on 2016/11/17 by [REDACTED_PERSON]

	Dummy integrate of the Square render version workaround (CL 3201913) with _accept target_ to prevent it being integrated to dev-core in future.

	Commandline:
	p4 integrate //Tasks/UE4/Dev-LoadTimes/Engine/Source/Runtime/CoreUObject/Private/UObject/LinkerLoad.cpp@3201913,3201913 //UE4/Dev-Core/Engine/Source/Runtime/CoreUObject/Private/UObject/LinkerLoad.cpp

Change 3203757 on 2016/11/18 by [REDACTED_PERSON]

	Removing debug code from async loading code.

Change 3203927 on 2016/11/18 by [REDACTED_PERSON]

	Fixing comments in the async loading code.

Change 3204851 on 2016/11/18 by [REDACTED_PERSON]

	Metafunction for testing if a particular operator<< overload exists, e.g. THasInserterOperator<FArchive&, FMyType&>::Value.

Change 3204854 on 2016/11/18 by [REDACTED_PERSON]

	UEnumProperty.

Change 3205027 on 2016/11/18 by [REDACTED_PERSON]

	Add useful functions to FAssetPtr and TAssetSubclassOf that already existed on TAssetPtr
	Add Get() to TSubclassOf so it matches our other wrappers
	Fix TSubclassOf and TAssetSubclassOf to use the more efficient template method of checking class compatibility
	Comment and template cleanups for AssetPtr, StringAssetReference, LazyPtr, and SubclassOf

Change 3206334 on 2016/11/21 by [REDACTED_PERSON]

	#UE-38773: Fix it so non-component template subobjects of CDOs are not included as creation dependencies for BP classes, also clean up GetPreloadDependencies as it was adding redundant and null entries
	#UE-38799: Fix it so WidgetTrees don't get picked up as subobjects, and add ensure at cook time to find null outers that would crash at runtime. Make sure the instanced widget trees are transient.
	Cook finishes but game is still crashing in some cases, so I might adjust this after other testing

Change 3206353 on 2016/11/21 by [REDACTED_PERSON]

	Fix EnumProperty to handle EDL preload dependencies properly

Change 3206625 on 2016/11/21 by [REDACTED_PERSON]

	Fix enum property crash at runtime by copying what array property does and making sure inner property is not transient

Change 3206937 on 2016/11/21 by [REDACTED_PERSON]

	Disable other nested enum migrations as they are unlikely to work. Array property tags need to be refactored to be safer
	Correctly save enum tag for enum properties, it was being set but not serialized

Change 3207002 on 2016/11/21 by [REDACTED_PERSON]

	Fix it so per-widget copy of widget tree and all widgets inside are properly transient, they were being cooked before but never accessed.
	Fix case where non ClientOnly public objects nested instead ClientOnly objects would cook but fail to load, and add ensure to catch these cases in the future.
	If the full outer chain isn't available, it can't be loaded anyway, and this finds issues at cook time instead of load time.
	We should generally outlaw non-transient objects with transient outers, it does not do what people expect.

Change 3207032 on 2016/11/21 by [REDACTED_PERSON]

diff --git a/Engine/Plugins/NetcodeUnitTest/NetcodeUnitTest/Source/NetcodeUnitTest/Private/NUTUtilReflection.cpp b/Engine/Plugins/NetcodeUnitTest/NetcodeUnitTest/Source/NetcodeUnitTest/Private/NUTUtilReflection.cpp
--- a/Engine/Plugins/NetcodeUnitTest/NetcodeUnitTest/Source/NetcodeUnitTest/Private/NUTUtilReflection.cpp
+++ b/Engine/Plugins/NetcodeUnitTest/NetcodeUnitTest/Source/NetcodeUnitTest/Private/NUTUtilReflection.cpp
@@ -172,6 +172,7 @@ FVMReflection& FVMReflection::operator ->*(const ANSICHAR* InPropertyName)
 		static const TArray<UClass*> SupportedTypes = TArrayBuilder<UClass*>()
 														.Add(UClass::StaticClass())
 														.Add(UByteProperty::StaticClass())
+														.Add(UEnumProperty::StaticClass())
 														.Add(UUInt16Property::StaticClass())
 														.Add(UUInt32Property::StaticClass())
 														.Add(UUInt64Property::StaticClass())
@@ -387,6 +388,47 @@ FVMReflection& FVMReflection::operator [](const ANSICHAR* InFieldType)
 			FIELD_TYPE_CHECK(TEXT("uint64"), UUInt64Property)
 			FIELD_TYPE_CHECK(TEXT("FString"), UStrProperty)
 			FIELD_TYPE_CHECK(TEXT("FText"), UTextProperty)
+			else if (ActualFieldType->IsA(UEnumProperty::StaticClass()))
+			{
+				const UProperty* UnderlyingActualFieldType = static_cast<const UEnumProperty*>(ActualFieldType)->GetUnderlyingProperty();
+				if (ExpectedFieldType == "uint8" && UnderlyingActualFieldType->IsA<UByteProperty>())
+				{
+					bTypeValid = true;
+				}
+				else if (ExpectedFieldType == "int16" && UnderlyingActualFieldType->IsA<UInt16Property>())
+				{
+					bTypeValid = true;
+				}
+				else if (ExpectedFieldType == "int64" && UnderlyingActualFieldType->IsA<UInt64Property>())
+				{
+					bTypeValid = true;
+				}
+				else if (ExpectedFieldType == "int8" && UnderlyingActualFieldType->IsA<UInt8Property>())
+				{
+					bTypeValid = true;
+				}
+				else if (ExpectedFieldType == "int32" && UnderlyingActualFieldType->IsA<UIntProperty>())
+				{
+					bTypeValid = true;
+				}
+				else if (ExpectedFieldType == "uint16" && UnderlyingActualFieldType->IsA<UUInt16Property>())
+				{
+					bTypeValid = true;
+				}
+				else if (ExpectedFieldType == "uint32" && UnderlyingActualFieldType->IsA<UUInt32Property>())
+				{
+					bTypeValid = true;
+				}
+				else if (ExpectedFieldType == "uint64" && UnderlyingActualFieldType->IsA<UUInt64Property>())
+				{
+					bTypeValid = true;
+				}
+				else
+				{
+					SetError(FString::Printf(TEXT("Tried to verify %s as being of type '%s', but it has underlying type '%s' instead."),
+								CheckType, *ExpectedFieldType, *UnderlyingActualFieldType->GetClass()->GetName()));
+				}
+			}
 			// UObject and subclasses
 			else if (ExpectedFieldType.Len() > 2 &&
 						(ExpectedFieldType.Left(1) == TEXT("U") || ExpectedFieldType.Left(1) == TEXT("A")) &&
@@ -546,6 +588,10 @@ InType* FVMReflection::GetWritableCast(const TCHAR* InTypeStr, bool bDoingUpCast
 		{
 			ReturnVal = (InType*)FieldAddress;
 		}
+		else if (FieldInstance->IsA<UEnumProperty>() && static_cast<const UEnumProperty*>(FieldInstance)->GetUnderlyingProperty()->IsA(InTypeClass::StaticClass()))
+		{
+			ReturnVal = (InType*)FieldAddress;
+		}
 		else if (!bDoingUpCast)
 		{
 			FString Error = FString::Printf(TEXT("Tried to cast type '%s' to type '%s'."), *FieldInstance->GetClass()->GetName(),
@@ -616,6 +662,59 @@ InType FVMReflection::GetNumericTypeCast(const TCHAR* InTypeStr, const TArray<UC
 			NUMERIC_UPCAST(int16, UInt16Property)
 			NUMERIC_UPCAST(int32, UIntProperty)
 			NUMERIC_UPCAST(float, UFloatProperty)
+			else if (FieldInstance->IsA(UEnumProperty::StaticClass()))
+			{
+				const UEnumProperty* EnumFieldInstance = static_cast<const UEnumProperty*>(FieldInstance);
+				const UNumericProperty* UnderlyingProp = EnumFieldInstance->GetUnderlyingProperty();
+				const UClass* UnderlyingPropClass = UnderlyingProp->GetClass();
+				if (SupportedUpCasts.Contains(UnderlyingPropClass))
+				{
+					if (UnderlyingPropClass == UByteProperty::StaticClass())
+					{
+						ReturnVal = (InType)(GetTypeCast<uint8>(InTypeStr));
+					}
+					else if (UnderlyingPropClass == UUInt16Property::StaticClass())
+					{
+						ReturnVal = (InType)(GetTypeCast<uint16>(InTypeStr));
+					}
+					else if (UnderlyingPropClass == UUInt32Property::StaticClass())
+					{
+						ReturnVal = (InType)(GetTypeCast<uint32>(InTypeStr));
+					}
+					else if (UnderlyingPropClass == UUInt64Property::StaticClass())
+					{
+						ReturnVal = (InType)(GetTypeCast<uint64>(InTypeStr));
+					}
+					else if (UnderlyingPropClass == UInt8Property::StaticClass())
+					{
+						ReturnVal = (InType)(GetTypeCast<int8>(InTypeStr));
+					}
+					else if (UnderlyingPropClass == UInt16Property::StaticClass())
+					{
+						ReturnVal = (InType)(GetTypeCast<int16>(InTypeStr));
+					}
+					else if (UnderlyingPropClass == UIntProperty::StaticClass())
+					{
+						ReturnVal = (InType)(GetTypeCast<int32>(InTypeStr));
+					}
+					else if (UnderlyingPropClass == UInt64Property::StaticClass())
+					{
+						ReturnVal = (InType)(GetTypeCast<int64>(InTypeStr));
+					}
+					else
+					{
+						FString Error = FString::Printf(TEXT("Enum property with underlying type '%s' does not support upcasting to type '%s'."),
+											*UnderlyingPropClass->GetName(), *InTypeClass::StaticClass()->GetName());
+						SetCastError(Error);
+					}
+				}
+				else
+				{
+					FString Error = FString::Printf(TEXT("Type '%s' does not support upcasting to type '%s'."),
+										*FieldInstance->GetClass()->GetName(), *InTypeClass::StaticClass()->GetName());
+					SetCastError(Error);
+				}
+			}
 			else
 			{
 				FString Error = FString::Printf(TEXT("No upcast possible from type '%s' to type '%s'."),
diff --git a/Engine/Plugins/Online/OnlineFramework/Source/Lobby/Private/LobbyBeaconClient.cpp b/Engine/Plugins/Online/OnlineFramework/Source/Lobby/Private/LobbyBeaconClient.cpp
--- a/Engine/Plugins/Online/OnlineFramework/Source/Lobby/Private/LobbyBeaconClient.cpp
+++ b/Engine/Plugins/Online/OnlineFramework/Source/Lobby/Private/LobbyBeaconClient.cpp
@@ -152,7 +152,7 @@ void ALobbyBeaconClient::LoginLocalPlayers()
 	UWorld* World = GetWorld();
 	for (FConstPlayerControllerIterator Iterator = World->GetPlayerControllerIterator(); Iterator; ++Iterator)
 	{
-		APlayerController* PC = (*Iterator);
+		APlayerController* PC = Iterator->Get();
 		if (PC && PC->PlayerState && PC->PlayerState->UniqueId.IsValid())
 		{
 			ULocalPlayer* LP = Cast<ULocalPlayer>(PC->Player);
diff --git a/Engine/Plugins/Online/OnlineSubsystem/Source/Private/OnlineKeyValuePair.cpp b/Engine/Plugins/Online/OnlineSubsystem/Source/Private/OnlineKeyValuePair.cpp
--- a/Engine/Plugins/Online/OnlineSubsystem/Source/Private/OnlineKeyValuePair.cpp
+++ b/Engine/Plugins/Online/OnlineSubsystem/Source/Private/OnlineKeyValuePair.cpp
@@ -820,21 +820,83 @@ bool FVariantDataConverter::VariantDataToUProperty(const FVariantData* Variant,
 
 bool FVariantDataConverter::ConvertScalarVariantToUProperty(const FVariantData* Variant, UProperty* Property, void* OutValue, int64 CheckFlags, int64 SkipFlags)
 {
-	UNumericProperty* NumericProperty = Cast<UNumericProperty>(Property);
-	if (NumericProperty)
+	if (UEnumProperty* EnumProperty = Cast<UEnumProperty>(Property))
+	{
+		const UEnum* Enum = EnumProperty->GetEnum();
+
+		if (Variant->GetType() == EOnlineKeyValuePairDataType::String)
+		{
+			FString StrValue;
+			Variant->GetValue(StrValue);
+
+			int64 IntValue = Enum->GetValueByName(FName(*StrValue));
+			if (IntValue != INDEX_NONE)
+			{
+				EnumProperty->GetUnderlyingProperty()->SetIntPropertyValue(OutValue, IntValue);
+				return true;
+			}
+		}
+		else
+		{
+			int64 Value = 0;
+			if (Variant->GetType() == EOnlineKeyValuePairDataType::Double)
+			{
+				double DoubleValue = 0.0;
+				Variant->GetValue(DoubleValue);
+				Value = (int64)DoubleValue;
+			}
+			else if (Variant->GetType() == EOnlineKeyValuePairDataType::Float)
+			{
+				float FloatValue = 0.0f;
+				Variant->GetValue(FloatValue);
+				Value = (int64)FloatValue;
+			}
+			else if (Variant->GetType() == EOnlineKeyValuePairDataType::Int32)
+			{
+				int32 IntValue = 0;
+				Variant->GetValue(IntValue);
+				Value = (int64)IntValue;
+			}
+			else if (Variant->GetType() == EOnlineKeyValuePairDataType::UInt32)
+			{
+				uint32 IntValue = 0;
+				Variant->GetValue(IntValue);
+				Value = (int64)IntValue;
+			}
+			else if (Variant->GetType() == EOnlineKeyValuePairDataType::Int64)
+			{
+				int64 Int64Value;
+				Variant->GetValue(Int64Value);
+				Value = (int64)Int64Value;
+			}
+			else if (Variant->GetType() == EOnlineKeyValuePairDataType::UInt64)
+			{
+				uint64 UInt64Value;
+				Variant->GetValue(UInt64Value);
+				Value = (int64)UInt64Value;
+			}
+
+			// AsNumber will log an error for completely inappropriate types (then give us a default)
+			EnumProperty->GetUnderlyingProperty()->SetIntPropertyValue(OutValue, Value);
+		}
+
+		UE_LOG(LogOnline, Error, TEXT("ConvertScalarVariantToUProperty - Unable to import enum %s from %s value for property %s"), *Enum->CppType, Variant->GetTypeString(), *Property->GetNameCPP());
+		return false;
+	}
+	else if (UNumericProperty* NumericProperty = Cast<UNumericProperty>(Property))
 	{
 		if (NumericProperty->IsEnum() && Variant->GetType() == EOnlineKeyValuePairDataType::String)
 		{
 			// see if we were passed a string for the enum
-			const UEnum* EnumProperty = NumericProperty->GetIntPropertyEnum();
-			check(EnumProperty); // should be assured by IsEnum()
+			const UEnum* Enum = NumericProperty->GetIntPropertyEnum();
+			check(Enum); // should be assured by IsEnum()
 			FString StrValue;
 			Variant->GetValue(StrValue);
 
-			int32 IntValue = EnumProperty->GetValueByName(FName(*StrValue));
+			int32 IntValue = Enum->GetValueByName(FName(*StrValue));
 			if (IntValue == INDEX_NONE)
 			{
-				UE_LOG(LogOnline, Error, TEXT("ConvertScalarVariantToUProperty - Unable import enum %s from string value %s for property %s"), *EnumProperty->CppType, *StrValue, *Property->GetNameCPP());
+				UE_LOG(LogOnline, Error, TEXT("ConvertScalarVariantToUProperty - Unable import enum %s from string value %s for property %s"), *Enum->CppType, *StrValue, *Property->GetNameCPP());
 				return false;
 			}
 			NumericProperty->SetIntPropertyValue(OutValue, (int64)IntValue);
@@ -1073,7 +1135,14 @@ bool FVariantDataConverter::ConvertScalarUPropertyToVariant(UProperty* Property,
 {
 	OutVariantData.Empty();
 
-	if (UNumericProperty *NumericProperty = Cast<UNumericProperty>(Property))
+	if (UEnumProperty* EnumProperty = Cast<UEnumProperty>(Property))
+	{
+		// export enums as strings
+		UEnum* EnumDef = EnumProperty->GetEnum();
+		FString StringValue = EnumDef->GetEnumName(EnumProperty->GetUnderlyingProperty()->GetSignedIntPropertyValue(Value));
+		OutVariantData.SetValue(StringValue);
+	}
+	else if (UNumericProperty *NumericProperty = Cast<UNumericProperty>(Property))
 	{
 		// see if it's an enum
 		UEnum* EnumDef = NumericProperty->GetIntPropertyEnum();
diff --git a/Engine/Plugins/Runtime/GameplayAbilities/Source/GameplayAbilitiesEditor/Private/GameplayAbilitiesEditorModule.cpp b/Engine/Plugins/Runtime/GameplayAbilities/Source/GameplayAbilitiesEditor/Private/GameplayAbilitiesEditorModule.cpp
--- a/Engine/Plugins/Runtime/GameplayAbilities/Source/GameplayAbilitiesEditor/Private/GameplayAbilitiesEditorModule.cpp
+++ b/Engine/Plugins/Runtime/GameplayAbilities/Source/GameplayAbilitiesEditor/Private/GameplayAbilitiesEditorModule.cpp
@@ -368,7 +368,7 @@ void FGameplayAbilitiesEditorModule::ApplyGameplayModEvaluationChannelAliasesToE
 			const UEnum* AttributeBasedFloatCalculationTypeEnum = FindObject<UEnum>(ANY_PACKAGE, TEXT("EAttributeBasedFloatCalculationType"));
 			if (ensure(AttributeBasedFloatCalculationTypeEnum))
 			{
-				const int32 ChannelBasedCalcIdx = AttributeBasedFloatCalculationTypeEnum->GetIndexByValue(static_cast<int32>(EAttributeBasedFloatCalculationType::AttributeMagnitudeEvaluatedUpToChannel));
+				const int32 ChannelBasedCalcIdx = AttributeBasedFloatCalculationTypeEnum->GetIndexByValue(static_cast<int64>(EAttributeBasedFloatCalculationType::AttributeMagnitudeEvaluatedUpToChannel));
 				AttributeBasedFloatCalculationTypeEnum->SetMetaData(HiddenMeta, TEXT(""), ChannelBasedCalcIdx);
 			}
 		}
diff --git a/Engine/Plugins/ScriptPlugin/Source/ScriptEditorPlugin/Private/ScriptBlueprintCompiler.cpp b/Engine/Plugins/ScriptPlugin/Source/ScriptEditorPlugin/Private/ScriptBlueprintCompiler.cpp
--- a/Engine/Plugins/ScriptPlugin/Source/ScriptEditorPlugin/Private/ScriptBlueprintCompiler.cpp
+++ b/Engine/Plugins/ScriptPlugin/Source/ScriptEditorPlugin/Private/ScriptBlueprintCompiler.cpp
@@ -192,9 +192,9 @@ void FScriptBlueprintCompiler::FinishCompilingClass(UClass* Class)
 void FScriptBlueprintCompiler::Compile()
 {
 	ScriptBlueprint()->UpdateSourceCodeIfChanged();
-	ScriptContext = FScriptContextBase::CreateContext(ScriptBlueprint()->SourceCode, NULL, NULL);
+	ScriptContext.Reset(FScriptContextBase::CreateContext(ScriptBlueprint()->SourceCode, NULL, NULL));
 	bool Result = true;
-	if (ScriptContext.IsValid())
+	if (ScriptContext)
 	{
 		ScriptDefinedFields.Empty();
 		ScriptContext->GetScriptDefinedFields(ScriptDefinedFields);
diff --git a/Engine/Plugins/ScriptPlugin/Source/ScriptEditorPlugin/Private/ScriptBlueprintCompiler.h b/Engine/Plugins/ScriptPlugin/Source/ScriptEditorPlugin/Private/ScriptBlueprintCompiler.h
--- a/Engine/Plugins/ScriptPlugin/Source/ScriptEditorPlugin/Private/ScriptBlueprintCompiler.h
+++ b/Engine/Plugins/ScriptPlugin/Source/ScriptEditorPlugin/Private/ScriptBlueprintCompiler.h
@@ -1,6 +1,7 @@
 // Copyright 1998-2016 Epic Games, Inc. All Rights Reserved.
 #pragma once
 #include "KismetCompiler.h"
+#include "UniquePtr.h"
 
 /**
 * Script Blueprint Compiler
@@ -49,7 +50,7 @@ protected:
 	/** New script class */
 	UScriptBlueprintGeneratedClass* NewScriptBlueprintClass;
 	/** Script context */
-	TAutoPtr<FScriptContextBase> ScriptContext;
+	TUniquePtr<FScriptContextBase> ScriptContext;
 	/** Script-defined properties and functions */
 	TArray<FScriptField> ScriptDefinedFields;
 	/** Script context property generated for the compiled class */
diff --git a/Engine/Plugins/ScriptPlugin/Source/ScriptEditorPlugin/Private/ScriptFactory.cpp b/Engine/Plugins/ScriptPlugin/Source/ScriptEditorPlugin/Private/ScriptFactory.cpp
--- a/Engine/Plugins/ScriptPlugin/Source/ScriptEditorPlugin/Private/ScriptFactory.cpp
+++ b/Engine/Plugins/ScriptPlugin/Source/ScriptEditorPlugin/Private/ScriptFactory.cpp
@@ -1,6 +1,7 @@
 // Copyright 1998-2016 Epic Games, Inc. All Rights Reserved. 
 #include "ScriptEditorPluginPrivatePCH.h"
 #include "ScriptBlueprintGeneratedClass.h"
+#include "UniquePtr.h"
 
 UScriptFactory::UScriptFactory(const FObjectInitializer& ObjectInitializer)
 	: Super( ObjectInitializer )
@@ -54,8 +55,8 @@ UObject* UScriptFactory::FactoryCreateText(UClass* InClass, UObject* InParent, F
 	GEditor->SelectNone(true, true, false);
 
 	UScriptBlueprint* NewBlueprint = nullptr;
-	TAutoPtr<FScriptContextBase> ScriptContext(FScriptContextBase::CreateContext(Buffer, nullptr, nullptr));
-	if (ScriptContext.IsValid())
+	TUniquePtr<FScriptContextBase> ScriptContext(FScriptContextBase::CreateContext(Buffer, nullptr, nullptr));
+	if (ScriptContext)
 	{
 	    NewBlueprint = Cast<UScriptBlueprint>( FindObject<UBlueprint>( InParent, *InName.ToString() ) );
 	    if( NewBlueprint != nullptr )
diff --git a/Engine/Plugins/ScriptPlugin/Source/ScriptGeneratorPlugin/Private/ScriptGeneratorPlugin.cpp b/Engine/Plugins/ScriptPlugin/Source/ScriptGeneratorPlugin/Private/ScriptGeneratorPlugin.cpp
--- a/Engine/Plugins/ScriptPlugin/Source/ScriptGeneratorPlugin/Private/ScriptGeneratorPlugin.cpp
+++ b/Engine/Plugins/ScriptPlugin/Source/ScriptGeneratorPlugin/Private/ScriptGeneratorPlugin.cpp
@@ -6,13 +6,14 @@
 #include "LuaScriptCodeGenerator.h"
 #include "IProjectManager.h"
 #include "Runtime/Core/Public/Features/IModularFeatures.h"
+#include "UniquePtr.h"
 
 DEFINE_LOG_CATEGORY(LogScriptGenerator);
 
 class FScriptGeneratorPlugin : public IScriptGeneratorPlugin
 {
 	/** Specialized script code generator */
-	TAutoPtr<FScriptCodeGeneratorBase> CodeGenerator;
+	TUniquePtr<FScriptCodeGeneratorBase> CodeGenerator;
 
 	/** IModuleInterface implementation */
 	virtual void StartupModule() override;
@@ -58,10 +59,10 @@ void FScriptGeneratorPlugin::Initialize(const FString& RootLocalPath, const FStr
 {
 #if WITH_LUA
 	UE_LOG(LogScriptGenerator, Log, TEXT("Using Lua Script Generator."));
-	CodeGenerator = new FLuaScriptCodeGenerator(RootLocalPath, RootBuildPath, OutputDirectory, IncludeBase);
+	CodeGenerator = MakeUnique<FLuaScriptCodeGenerator>(RootLocalPath, RootBuildPath, OutputDirectory, IncludeBase);
 #else
 	UE_LOG(LogScriptGenerator, Log, TEXT("Using Generic Script Generator."));
-	CodeGenerator = new FGenericScriptCodeGenerator(RootLocalPath, RootBuildPath, OutputDirectory, IncludeBase);
+	CodeGenerator = MakeUnique<FGenericScriptCodeGenerator>(RootLocalPath, RootBuildPath, OutputDirectory, IncludeBase);
 #endif
 	UE_LOG(LogScriptGenerator, Log, TEXT("Output directory: %s"), *OutputDirectory);
 }
diff --git a/Engine/Source/Developer/BlueprintCompilerCppBackend/Private/BlueprintCompilerCppBackend.cpp b/Engine/Source/Developer/BlueprintCompilerCppBackend/Private/BlueprintCompilerCppBackend.cpp
--- a/Engine/Source/Developer/BlueprintCompilerCppBackend/Private/BlueprintCompilerCppBackend.cpp
+++ b/Engine/Source/Developer/BlueprintCompilerCppBackend/Private/BlueprintCompilerCppBackend.cpp
@@ -6,6 +6,7 @@
 #include "EdGraphSchema_K2.h"
 #include "Kismet/KismetNodeHelperLibrary.h"
 #include "Kismet/KismetArrayLibrary.h"
+#include "UniquePtr.h"
 
 // Generates single "if" scope. Its condition checks context of given term.
 struct FSafeContextScopedEmmitter
@@ -718,14 +719,14 @@ FString FBlueprintCompilerCppBackend::EmitCallStatmentInner(FEmitterLocalContext
 
 	FString Result;
 	FString CloseCast;
-	TAutoPtr<FSetterExpressionBuilder> SetterExpression;
+	TUniquePtr<FSetterExpressionBuilder> SetterExpression;
 	if (!bInline)
 	{
 		// Handle the return value of the function being called
 		UProperty* FuncToCallReturnProperty = Statement.FunctionToCall->GetReturnProperty();
 		if (FuncToCallReturnProperty && ensure(Statement.LHS))
 		{
-			SetterExpression = new FSetterExpressionBuilder(*this, EmitterContext, Statement.LHS);
+			SetterExpression = MakeUnique<FSetterExpressionBuilder>(*this, EmitterContext, Statement.LHS);
 			Result += SetterExpression->BuildStart();
 
 			FString BeginCast;
@@ -836,7 +837,7 @@ FString FBlueprintCompilerCppBackend::EmitCallStatmentInner(FEmitterLocalContext
 	}
 
 	Result += CloseCast;
-	if (SetterExpression.IsValid())
+	if (SetterExpression)
 	{
 		Result += SetterExpression->BuildEnd(false);
 	}
diff --git a/Engine/Source/Developer/BlueprintCompilerCppBackend/Private/BlueprintCompilerCppBackendBase.cpp b/Engine/Source/Developer/BlueprintCompilerCppBackend/Private/BlueprintCompilerCppBackendBase.cpp
--- a/Engine/Source/Developer/BlueprintCompilerCppBackend/Private/BlueprintCompilerCppBackendBase.cpp
+++ b/Engine/Source/Developer/BlueprintCompilerCppBackend/Private/BlueprintCompilerCppBackendBase.cpp
@@ -782,7 +782,7 @@ FString FBlueprintCompilerCppBackendBase::GenerateCodeFromEnum(UUserDefinedEnum*
 
 	auto EnumItemName = [&](int32 InIndex)
 	{
-		const int32 ElemValue = SourceEnum->GetValueByIndex(InIndex);
+		const int64 ElemValue = SourceEnum->GetValueByIndex(InIndex);
 		if (ElemValue == SourceEnum->GetMaxEnumValue())
 		{
 			return FString::Printf(TEXT("%s_MAX"), *EnumCppName);
@@ -793,12 +793,12 @@ FString FBlueprintCompilerCppBackendBase::GenerateCodeFromEnum(UUserDefinedEnum*
 	for (int32 Index = 0; Index < SourceEnum->NumEnums(); ++Index)
 	{
 		const FString ElemCppName = EnumItemName(Index);
-		const int32 ElemValue = SourceEnum->GetValueByIndex(Index);
+		const int64 ElemValue = SourceEnum->GetValueByIndex(Index);
 
 		const FString& DisplayNameMD = SourceEnum->GetMetaData(TEXT("DisplayName"), ElemValue);// TODO: value or index?
 		const FString MetaDisplayName = DisplayNameMD.IsEmpty() ? FString() : FString::Printf(TEXT("DisplayName = \"%s\","), *DisplayNameMD.ReplaceCharWithEscapedChar());
 		const FString MetaOverrideName = FString::Printf(TEXT("OverrideName = \"%s\""), *SourceEnum->GetNameByIndex(Index).ToString());
-		Header.AddLine(FString::Printf(TEXT("%s = %d UMETA(%s%s),"), *ElemCppName, ElemValue, *MetaDisplayName, *MetaOverrideName));
+		Header.AddLine(FString::Printf(TEXT("%s = %lld UMETA(%s%s),"), *ElemCppName, ElemValue, *MetaDisplayName, *MetaOverrideName));
 	}
 
 	Header.DecreaseIndent();
diff --git a/Engine/Source/Developer/BlueprintCompilerCppBackend/Private/BlueprintCompilerCppBackendGatherDependencies.cpp b/Engine/Source/Developer/BlueprintCompilerCppBackend/Private/BlueprintCompilerCppBackendGatherDependencies.cpp
--- a/Engine/Source/Developer/BlueprintCompilerCppBackend/Private/BlueprintCompilerCppBackendGatherDependencies.cpp
+++ b/Engine/Source/Developer/BlueprintCompilerCppBackend/Private/BlueprintCompilerCppBackendGatherDependencies.cpp
@@ -453,6 +453,11 @@ void FGatherConvertedClassDependencies::DependenciesForHeader()
 				// HeaderReferenceFinder.FindReferences(Obj); cannot find this enum..
 				IncludeInHeader.Add(ByteProperty->Enum);
 			}
+			else if (const UEnumProperty* EnumProperty = Cast<const UEnumProperty>(Property))
+			{ 
+				// HeaderReferenceFinder.FindReferences(Obj); cannot find this enum..
+				IncludeInHeader.Add(EnumProperty->GetEnum());
+			}
 			else
 			{
 				HeaderReferenceFinder.FindReferences(Obj);
diff --git a/Engine/Source/Developer/CrashDebugHelper/Private/CrashDebugPDBCache.cpp b/Engine/Source/Developer/CrashDebugHelper/Private/CrashDebugPDBCache.cpp
--- a/Engine/Source/Developer/CrashDebugHelper/Private/CrashDebugPDBCache.cpp
+++ b/Engine/Source/Developer/CrashDebugHelper/Private/CrashDebugPDBCache.cpp
@@ -2,6 +2,7 @@
 
 #include "CrashDebugHelperPrivatePCH.h"
 #include "CrashDebugPDBCache.h"
+#include "UniquePtr.h"
 
 /*-----------------------------------------------------------------------------
 	PDB Cache implementation
@@ -248,7 +249,7 @@ FPDBCacheEntryRef FPDBCache::CreateAndAddPDBCacheEntry( const FString& OriginalL
 	FPDBCacheEntryRef NewCacheEntry = MakeShareable( new FPDBCacheEntry( CachedFiles, CleanedLabelName, FDateTime::Now(), SizeGB ) );
 
 	// Verify there is an entry timestamp file, write the size of a PDB cache to avoid time consuming FindFilesRecursive during initialization.
-	TAutoPtr<FArchive> FileWriter( IFileManager::Get().CreateFileWriter( *EntryTimeStampFilename ) );
+	TUniquePtr<FArchive> FileWriter( IFileManager::Get().CreateFileWriter( *EntryTimeStampFilename ) );
 	UE_CLOG( !FileWriter, LogCrashDebugHelper, Fatal, TEXT( "Couldn't save the timestamp for a file: %s" ), *EntryTimeStampFilename );
 	*FileWriter << *NewCacheEntry;
 
@@ -288,7 +289,7 @@ FPDBCacheEntryRef FPDBCache::CreateAndAddPDBCacheEntryMixed( const FString& Prod
 	FPDBCacheEntryRef NewCacheEntry = MakeShareable( new FPDBCacheEntry( CachedFiles, ProductVersion, FDateTime::Now(), SizeGB ) );
 
 	// Verify there is an entry timestamp file, write the size of a PDB cache to avoid time consuming FindFilesRecursive during initialization.
-	TAutoPtr<FArchive> FileWriter( IFileManager::Get().CreateFileWriter( *EntryTimeStampFilename ) );
+	TUniquePtr<FArchive> FileWriter( IFileManager::Get().CreateFileWriter( *EntryTimeStampFilename ) );
 	UE_CLOG( !FileWriter, LogCrashDebugHelper, Fatal, TEXT( "Couldn't save the timestamp for a file: %s" ), *EntryTimeStampFilename );
 	*FileWriter << *NewCacheEntry;
 
@@ -313,7 +314,7 @@ FPDBCacheEntryPtr FPDBCache::ReadPDBCacheEntry( const FString& Directory )
 	if( LastAccessTime != FDateTime::MinValue() )
 	{
 		// Read the metadata
-		TAutoPtr<FArchive> FileReader( IFileManager::Get().CreateFileReader( *EntryTimeStampFilename ) );
+		TUniquePtr<FArchive> FileReader( IFileManager::Get().CreateFileReader( *EntryTimeStampFilename ) );
 		NewEntry = MakeShareable( new FPDBCacheEntry( LastAccessTime ) );
 		*FileReader << *NewEntry;
 	}
@@ -336,7 +337,7 @@ FPDBCacheEntryPtr FPDBCache::ReadPDBCacheEntry( const FString& Directory )
 		NewEntry = MakeShareable( new FPDBCacheEntry( PDBFiles, Directory, LastAccessTimeNoMeta, SizeGB ) );
 
 		// Remove the old file and save the metadata.
-		TAutoPtr<FArchive> FileWriter( IFileManager::Get().CreateFileWriter( *EntryTimeStampFilename ) );
+		TUniquePtr<FArchive> FileWriter( IFileManager::Get().CreateFileWriter( *EntryTimeStampFilename ) );
 		*FileWriter << *NewEntry;
 
 		IFileManager::Get().Delete( *EntryTimeStampFilenameNoMeta );
diff --git a/Engine/Source/Developer/DerivedDataCache/Private/DerivedDataBackendAsyncPutWrapper.cpp b/Engine/Source/Developer/DerivedDataCache/Private/DerivedDataBackendAsyncPutWrapper.cpp
new file mode 100644
--- /dev/null
+++ b/Engine/Source/Developer/DerivedDataCache/Private/DerivedDataBackendAsyncPutWrapper.cpp
@@ -0,0 +1,195 @@
+// Copyright 1998-2016 Epic Games, Inc. All Rights Reserved.
+
+#include "Core.h"
+#include "DerivedDataBackendAsyncPutWrapper.h"
+#include "MemoryDerivedDataBackend.h"
+
+/** 
+ * Async task to handle the fire and forget async put
+ */
+class FCachePutAsyncWorker
+{
+public:
+	/** Cache Key for the put to InnerBackend **/
+	FString								CacheKey;
+	/** Data for the put to InnerBackend **/
+	TArray<uint8>						Data;
+	/** Backend to use for storage, my responsibilities are about async puts **/
+	FDerivedDataBackendInterface*		InnerBackend;
+	/** Memory based cache to clear once the put is finished **/
+	FDerivedDataBackendInterface*		InflightCache;
+	/** We remember outstanding puts so that we don't do them redundantly **/
+	FThreadSet*							FilesInFlight;
+	/**If true, then do not attempt skip the put even if CachedDataProbablyExists returns true **/
+	bool								bPutEvenIfExists;
+	/** Usage stats to track thread times. */
+	FDerivedDataCacheUsageStats&        UsageStats;
+
+	/** Constructor
+	*/
+	FCachePutAsyncWorker(const TCHAR* InCacheKey, const TArray<uint8>* InData, FDerivedDataBackendInterface* InInnerBackend, bool InbPutEvenIfExists, FDerivedDataBackendInterface* InInflightCache, FThreadSet* InInFilesInFlight, FDerivedDataCacheUsageStats& InUsageStats)
+		: CacheKey(InCacheKey)
+		, Data(*InData)
+		, InnerBackend(InInnerBackend)
+		, InflightCache(InInflightCache)
+		, FilesInFlight(InInFilesInFlight)
+		, bPutEvenIfExists(InbPutEvenIfExists)
+		, UsageStats(InUsageStats)
+	{
+		check(InnerBackend);
+	}
+		
+	/** Call the inner backend and when that completes, remove the memory cache */
+	void DoWork()
+	{
+		COOK_STAT(auto Timer = UsageStats.TimePut());
+		bool bOk = true;
+		const bool bAlreadyExists = InnerBackend->CachedDataProbablyExists(*CacheKey);
+		if (!bAlreadyExists || bPutEvenIfExists)
+		{
+			InnerBackend->PutCachedData(*CacheKey, Data, bPutEvenIfExists);
+			COOK_STAT(Timer.AddHit(Data.Num()));
+		}
+		// if it already existed, don't bother checking if we need to retry. We don't.
+		if (InflightCache && !bAlreadyExists && !InnerBackend->CachedDataProbablyExists(*CacheKey))
+		{
+			// retry
+			InnerBackend->PutCachedData(*CacheKey, Data, false);
+			if (!InnerBackend->CachedDataProbablyExists(*CacheKey))
+			{
+				UE_LOG(LogDerivedDataCache, Warning, TEXT("FDerivedDataBackendAsyncPutWrapper: Put failed, keeping in memory copy %s."),*CacheKey);
+				bOk = false;
+			}
+		}
+		if (bOk && InflightCache)
+		{
+			InflightCache->RemoveCachedData(*CacheKey, /*bTransient=*/ false); // we can remove this from the temp cache, since the real cache will hit now
+		}
+		FilesInFlight->Remove(CacheKey);
+		FDerivedDataBackend::Get().AddToAsyncCompletionCounter(-1);
+	}
+
+	FORCEINLINE TStatId GetStatId() const
+	{
+		RETURN_QUICK_DECLARE_CYCLE_STAT(FCachePutAsyncWorker, STATGROUP_ThreadPoolAsyncTasks);
+	}
+
+	/** Indicates to the thread pool that this task is abandonable */
+	bool CanAbandon()
+	{
+		return true;
+	}
+
+	/** Abandon routine, we need to remove the item from the in flight cache because something might be waiting for that */
+	void Abandon()
+	{
+		if (InflightCache)
+		{
+			InflightCache->RemoveCachedData(*CacheKey, /*bTransient=*/ false); // we can remove this from the temp cache, since the real cache will hit now
+		}
+		FilesInFlight->Remove(CacheKey);
+		FDerivedDataBackend::Get().AddToAsyncCompletionCounter(-1);
+	}
+};
+
+FDerivedDataBackendAsyncPutWrapper::FDerivedDataBackendAsyncPutWrapper(FDerivedDataBackendInterface* InInnerBackend, bool bCacheInFlightPuts)
+	: InnerBackend(InInnerBackend)
+	, InflightCache(bCacheInFlightPuts ? (new FMemoryDerivedDataBackend) : NULL)
+{
+	check(InnerBackend);
+}
+
+/** return true if this cache is writable **/
+bool FDerivedDataBackendAsyncPutWrapper::IsWritable()
+{
+	return InnerBackend->IsWritable();
+}
+
+/**
+ * Synchronous test for the existence of a cache item
+ *
+ * @param	CacheKey	Alphanumeric+underscore key of this cache item
+ * @return				true if the data probably will be found, this can't be guaranteed because of concurrency in the backends, corruption, etc
+ */
+bool FDerivedDataBackendAsyncPutWrapper::CachedDataProbablyExists(const TCHAR* CacheKey)
+{
+	COOK_STAT(auto Timer = UsageStats.TimeProbablyExists());
+	bool Result = (InflightCache && InflightCache->CachedDataProbablyExists(CacheKey)) || InnerBackend->CachedDataProbablyExists(CacheKey);
+	COOK_STAT(if (Result) {	Timer.AddHit(0); });
+	return Result;
+}
+
+bool FDerivedDataBackendAsyncPutWrapper::GetCachedData(const TCHAR* CacheKey, TArray<uint8>& OutData)
+{
+	COOK_STAT(auto Timer = UsageStats.TimeGet());
+	if (InflightCache && InflightCache->GetCachedData(CacheKey, OutData))
+	{
+		COOK_STAT(Timer.AddHit(OutData.Num()));
+		return true;
+	}
+	bool bSuccess = InnerBackend->GetCachedData(CacheKey, OutData);
+	if (bSuccess)
+	{
+		COOK_STAT(Timer.AddHit(OutData.Num()));
+	}
+	return bSuccess;
+}
+
+void FDerivedDataBackendAsyncPutWrapper::PutCachedData(const TCHAR* CacheKey, TArray<uint8>& InData, bool bPutEvenIfExists)
+{
+	COOK_STAT(auto Timer = PutSyncUsageStats.TimePut());
+	if (!InnerBackend->IsWritable())
+	{
+		return; // no point in continuing down the chain
+	}
+	const bool bAdded = FilesInFlight.AddIfNotExists(CacheKey);
+	if (!bAdded)
+	{
+		return; // if it is already on its way, we don't need to send it again
+	}
+	if (InflightCache)
+	{
+		if (InflightCache->CachedDataProbablyExists(CacheKey))
+		{
+			return; // if it is already on its way, we don't need to send it again
+		}
+		InflightCache->PutCachedData(CacheKey, InData, true); // temp copy stored in memory while the async task waits to complete
+		COOK_STAT(Timer.AddHit(InData.Num()));
+	}
+	FDerivedDataBackend::Get().AddToAsyncCompletionCounter(1);
+	(new FAutoDeleteAsyncTask<FCachePutAsyncWorker>(CacheKey, &InData, InnerBackend, bPutEvenIfExists, InflightCache.Get(), &FilesInFlight, UsageStats))->StartBackgroundTask();
+}
+
+void FDerivedDataBackendAsyncPutWrapper::RemoveCachedData(const TCHAR* CacheKey, bool bTransient)
+{
+	if (!InnerBackend->IsWritable())
+	{
+		return; // no point in continuing down the chain
+	}
+	while (FilesInFlight.Exists(CacheKey))
+	{
+		FPlatformProcess::Sleep(0.0f); // this is an exception condition (corruption), spin and wait for it to clear
+	}
+	if (InflightCache)
+	{
+		InflightCache->RemoveCachedData(CacheKey, bTransient);
+	}
+	InnerBackend->RemoveCachedData(CacheKey, bTransient);
+}
+
+void FDerivedDataBackendAsyncPutWrapper::GatherUsageStats(TMap<FString, FDerivedDataCacheUsageStats>& UsageStatsMap, FString&& GraphPath)
+{
+	COOK_STAT(
+		{
+		UsageStatsMap.Add(GraphPath + TEXT(": AsyncPut"), UsageStats);
+		UsageStatsMap.Add(GraphPath + TEXT(": AsyncPutSync"), PutSyncUsageStats);
+		if (InnerBackend)
+		{
+			InnerBackend->GatherUsageStats(UsageStatsMap, GraphPath + TEXT(". 0"));
+		}
+		if (InflightCache)
+		{
+			InflightCache->GatherUsageStats(UsageStatsMap, GraphPath + TEXT(". 1"));
+		}
+	});
+}
diff --git a/Engine/Source/Developer/DerivedDataCache/Private/DerivedDataBackendAsyncPutWrapper.h b/Engine/Source/Developer/DerivedDataCache/Private/DerivedDataBackendAsyncPutWrapper.h
--- a/Engine/Source/Developer/DerivedDataCache/Private/DerivedDataBackendAsyncPutWrapper.h
+++ b/Engine/Source/Developer/DerivedDataCache/Private/DerivedDataBackendAsyncPutWrapper.h
@@ -1,6 +1,7 @@
 // Copyright 1998-2016 Epic Games, Inc. All Rights Reserved.
 
 #pragma once
+#include "DerivedDataBackendInterface.h"
 #include "DerivedDataCacheUsageStats.h"
 
 /** 
@@ -39,94 +40,6 @@ struct FThreadSet
 		return false;
 	}
 };
-/** 
-	* Async task to handle the fire and forget async put
-**/
-class FCachePutAsyncWorker
-{
-public:
-	/** Cache Key for the put to InnerBackend **/
-	FString								CacheKey;
-	/** Data for the put to InnerBackend **/
-	TArray<uint8>						Data;
-	/** Backend to use for storage, my responsibilities are about async puts **/
-	FDerivedDataBackendInterface*		InnerBackend;
-	/** Memory based cache to clear once the put is finished **/
-	FDerivedDataBackendInterface*		InflightCache;
-	/** We remember outstanding puts so that we don't do them redundantly **/
-	FThreadSet*							FilesInFlight;
-	/**If true, then do not attempt skip the put even if CachedDataProbablyExists returns true **/
-	bool								bPutEvenIfExists;
-	/** Usage stats to track thread times. */
-	FDerivedDataCacheUsageStats&        UsageStats;
-
-	/** Constructor
-	*/
-	FCachePutAsyncWorker(const TCHAR* InCacheKey, const TArray<uint8>* InData, FDerivedDataBackendInterface* InInnerBackend, bool InbPutEvenIfExists, FDerivedDataBackendInterface* InInflightCache, FThreadSet* InInFilesInFlight, FDerivedDataCacheUsageStats& InUsageStats)
-		: CacheKey(InCacheKey)
-		, Data(*InData)
-		, InnerBackend(InInnerBackend)
-		, InflightCache(InInflightCache)
-		, FilesInFlight(InInFilesInFlight)
-		, bPutEvenIfExists(InbPutEvenIfExists)
-		, UsageStats(InUsageStats)
-	{
-		check(InnerBackend);
-	}
-		
-	/** Call the inner backend and when that completes, remove the memory cache */
-	void DoWork()
-	{
-		COOK_STAT(auto Timer = UsageStats.TimePut());
-		bool bOk = true;
-		const bool bAlreadyExists = InnerBackend->CachedDataProbablyExists(*CacheKey);
-		if (!bAlreadyExists || bPutEvenIfExists)
-		{
-			InnerBackend->PutCachedData(*CacheKey, Data, bPutEvenIfExists);
-			COOK_STAT(Timer.AddHit(Data.Num()));
-		}
-		// if it already existed, don't bother checking if we need to retry. We don't.
-		if (InflightCache && !bAlreadyExists && !InnerBackend->CachedDataProbablyExists(*CacheKey))
-		{
-			// retry
-			InnerBackend->PutCachedData(*CacheKey, Data, false);
-			if (!InnerBackend->CachedDataProbablyExists(*CacheKey))
-			{
-				UE_LOG(LogDerivedDataCache, Warning, TEXT("FDerivedDataBackendAsyncPutWrapper: Put failed, keeping in memory copy %s."),*CacheKey);
-				bOk = false;
-			}
-		}
-		if (bOk && InflightCache)
-		{
-			InflightCache->RemoveCachedData(*CacheKey, /*bTransient=*/ false); // we can remove this from the temp cache, since the real cache will hit now
-		}
-		FilesInFlight->Remove(CacheKey);
-		FDerivedDataBackend::Get().AddToAsyncCompletionCounter(-1);
-	}
-
-	FORCEINLINE TStatId GetStatId() const
-	{
-		RETURN_QUICK_DECLARE_CYCLE_STAT(FCachePutAsyncWorker, STATGROUP_ThreadPoolAsyncTasks);
-	}
-
-	/** Indicates to the thread pool that this task is abandonable */
-	bool CanAbandon()
-	{
-		return true;
-	}
-
-	/** Abandon routine, we need to remove the item from the in flight cache because something might be waiting for that */
-	void Abandon()
-	{
-		if (InflightCache)
-		{
-			InflightCache->RemoveCachedData(*CacheKey, /*bTransient=*/ false); // we can remove this from the temp cache, since the real cache will hit now
-		}
-		FilesInFlight->Remove(CacheKey);
-		FDerivedDataBackend::Get().AddToAsyncCompletionCounter(-1);
-	}
-};
-
 
 /** 
  * A backend wrapper that coordinates async puts. This means that a Get will hit an in-memory cache while the async put is still in flight.
@@ -141,18 +54,10 @@ public:
 	 * @param	InInnerBackend		Backend to use for storage, my responsibilities are about async puts
 	 * @param	bCacheInFlightPuts	if true, cache in-flight puts in a memory cache so that they hit immediately
 	 */
-	FDerivedDataBackendAsyncPutWrapper(FDerivedDataBackendInterface* InInnerBackend, bool bCacheInFlightPuts)
-		: InnerBackend(InInnerBackend)
-		, InflightCache(bCacheInFlightPuts ? (new FMemoryDerivedDataBackend) : NULL)
-	{
-		check(InnerBackend);
-	}
+	FDerivedDataBackendAsyncPutWrapper(FDerivedDataBackendInterface* InInnerBackend, bool bCacheInFlightPuts);
 
 	/** return true if this cache is writable **/
-	virtual bool IsWritable() override
-	{
-		return InnerBackend->IsWritable();
-	}
+	virtual bool IsWritable() override;
 
 	/**
 	 * Synchronous test for the existence of a cache item
@@ -160,13 +65,8 @@ public:
 	 * @param	CacheKey	Alphanumeric+underscore key of this cache item
 	 * @return				true if the data probably will be found, this can't be guaranteed because of concurrency in the backends, corruption, etc
 	 */
-	virtual bool CachedDataProbablyExists(const TCHAR* CacheKey) override
-	{
-		COOK_STAT(auto Timer = UsageStats.TimeProbablyExists());
-		bool Result = (InflightCache && InflightCache->CachedDataProbablyExists(CacheKey)) || InnerBackend->CachedDataProbablyExists(CacheKey);
-		COOK_STAT(if (Result) {	Timer.AddHit(0); });
-		return Result;
-	}
+	virtual bool CachedDataProbablyExists(const TCHAR* CacheKey) override;
+
 	/**
 	 * Synchronous retrieve of a cache item
 	 *
@@ -174,21 +74,8 @@ public:
 	 * @param	OutData		Buffer to receive the results, if any were found
 	 * @return				true if any data was found, and in this case OutData is non-empty
 	 */
-	virtual bool GetCachedData(const TCHAR* CacheKey, TArray<uint8>& OutData) override
-	{
-		COOK_STAT(auto Timer = UsageStats.TimeGet());
-		if (InflightCache && InflightCache->GetCachedData(CacheKey, OutData))
-		{
-			COOK_STAT(Timer.AddHit(OutData.Num()));
-			return true;
-		}
-		bool bSuccess = InnerBackend->GetCachedData(CacheKey, OutData);
-		if (bSuccess)
-		{
-			COOK_STAT(Timer.AddHit(OutData.Num()));
-		}
-		return bSuccess;
-	}
+	virtual bool GetCachedData(const TCHAR* CacheKey, TArray<uint8>& OutData) override;
+
 	/**
 	 * Asynchronous, fire-and-forget placement of a cache item
 	 *
@@ -196,64 +83,11 @@ public:
 	 * @param	InData		Buffer containing the data to cache, can be destroyed after the call returns, immediately
 	 * @param	bPutEvenIfExists	If true, then do not attempt skip the put even if CachedDataProbablyExists returns true
 	 */
-	virtual void PutCachedData(const TCHAR* CacheKey, TArray<uint8>& InData, bool bPutEvenIfExists) override
-	{
-		COOK_STAT(auto Timer = PutSyncUsageStats.TimePut());
-		if (!InnerBackend->IsWritable())
-		{
-			return; // no point in continuing down the chain
-		}
-		const bool bAdded = FilesInFlight.AddIfNotExists(CacheKey);
-		if (!bAdded)
-		{
-			return; // if it is already on its way, we don't need to send it again
-		}
-		if (InflightCache)
-		{
-			if (InflightCache->CachedDataProbablyExists(CacheKey))
-			{
-				return; // if it is already on its way, we don't need to send it again
-			}
-			InflightCache->PutCachedData(CacheKey, InData, true); // temp copy stored in memory while the async task waits to complete
-			COOK_STAT(Timer.AddHit(InData.Num()));
-		}
-		FDerivedDataBackend::Get().AddToAsyncCompletionCounter(1);
-		(new FAutoDeleteAsyncTask<FCachePutAsyncWorker>(CacheKey, &InData, InnerBackend, bPutEvenIfExists, InflightCache.GetOwnedPointer(), &FilesInFlight, UsageStats))->StartBackgroundTask();
-	}
+	virtual void PutCachedData(const TCHAR* CacheKey, TArray<uint8>& InData, bool bPutEvenIfExists) override;
 
-	virtual void RemoveCachedData(const TCHAR* CacheKey, bool bTransient) override
-	{
-		if (!InnerBackend->IsWritable())
-		{
-			return; // no point in continuing down the chain
-		}
-		while (FilesInFlight.Exists(CacheKey))
-		{
-			FPlatformProcess::Sleep(0.0f); // this is an exception condition (corruption), spin and wait for it to clear
-		}
-		if (InflightCache)
-		{
-			InflightCache->RemoveCachedData(CacheKey, bTransient);
-		}
-		InnerBackend->RemoveCachedData(CacheKey, bTransient);
-	}
+	virtual void RemoveCachedData(const TCHAR* CacheKey, bool bTransient) override;
 
-	virtual void GatherUsageStats(TMap<FString, FDerivedDataCacheUsageStats>& UsageStatsMap, FString&& GraphPath) override
-	{
-		COOK_STAT(
-			{
-			UsageStatsMap.Add(GraphPath + TEXT(": AsyncPut"), UsageStats);
-			UsageStatsMap.Add(GraphPath + TEXT(": AsyncPutSync"), PutSyncUsageStats);
-			if (InnerBackend)
-			{
-				InnerBackend->GatherUsageStats(UsageStatsMap, GraphPath + TEXT(". 0"));
-			}
-			if (InflightCache)
-			{
-				InflightCache->GatherUsageStats(UsageStatsMap, GraphPath + TEXT(". 1"));
-			}
-		});
-	}
+	virtual void GatherUsageStats(TMap<FString, FDerivedDataCacheUsageStats>& UsageStatsMap, FString&& GraphPath) override;
 
 private:
 	FDerivedDataCacheUsageStats UsageStats;
@@ -262,7 +96,7 @@ private:
 	/** Backend to use for storage, my responsibilities are about async puts **/
 	FDerivedDataBackendInterface*					InnerBackend;
 	/** Memory based cache to deal with gets that happen while an async put is still in flight **/
-	TScopedPointer<FDerivedDataBackendInterface>	InflightCache;
+	TUniquePtr<FDerivedDataBackendInterface>		InflightCache;
 	/** We remember outstanding puts so that we don't do them redundantly **/
 	FThreadSet										FilesInFlight;
 };
diff --git a/Engine/Source/Developer/DerivedDataCache/Private/HierarchicalDerivedDataBackend.h b/Engine/Source/Developer/DerivedDataCache/Private/HierarchicalDerivedDataBackend.h
--- a/Engine/Source/Developer/DerivedDataCache/Private/HierarchicalDerivedDataBackend.h
+++ b/Engine/Source/Developer/DerivedDataCache/Private/HierarchicalDerivedDataBackend.h
@@ -3,6 +3,7 @@
 #pragma once
 
 #include "DerivedDataCacheUsageStats.h"
+#include "UniquePtr.h"
 
 
 /** 
@@ -33,7 +34,7 @@ public:
 			for (int32 CacheIndex = 0; CacheIndex < InnerBackends.Num(); CacheIndex++)
 			{
 				// async puts to allow us to fill all levels without holding up the engine
-				new (AsyncPutInnerBackends) TScopedPointer<FDerivedDataBackendInterface>(new FDerivedDataBackendAsyncPutWrapper(InnerBackends[CacheIndex], false));
+				AsyncPutInnerBackends.Emplace(new FDerivedDataBackendAsyncPutWrapper(InnerBackends[CacheIndex], false));
 			}
 		}
 	}
@@ -211,7 +212,7 @@ private:
 	/** Array of backends forming the hierarchical cache...the first element is the fastest cache. **/
 	TArray<FDerivedDataBackendInterface*> InnerBackends;
 	/** Each of the backends wrapped with an async put **/
-	TArray<TScopedPointer<FDerivedDataBackendInterface> > AsyncPutInnerBackends;
+	TArray<TUniquePtr<FDerivedDataBackendInterface> > AsyncPutInnerBackends;
 	/** As an optimization, we check our writable status at contruction **/
 	bool bIsWritable;
 };
diff --git a/Engine/Source/Developer/DerivedDataCache/Private/MemoryDerivedDataBackend.cpp b/Engine/Source/Developer/DerivedDataCache/Private/MemoryDerivedDataBackend.cpp
new file mode 100644
--- /dev/null
+++ b/Engine/Source/Developer/DerivedDataCache/Private/MemoryDerivedDataBackend.cpp
@@ -0,0 +1,292 @@
+// Copyright 1998-2016 Epic Games, Inc. All Rights Reserved.
+
+#include "Core.h"
+#include "MemoryDerivedDataBackend.h"
+#include "UniquePtr.h"
+
+FMemoryDerivedDataBackend::FMemoryDerivedDataBackend(int64 InMaxCacheSize)
+	: MaxCacheSize(InMaxCacheSize)
+	, bDisabled( false )
+	, CurrentCacheSize( SerializationSpecificDataSize )
+	, bMaxSizeExceeded(false)
+{
+}
+
+FMemoryDerivedDataBackend::~FMemoryDerivedDataBackend()
+{
+	Disable();
+}
+
+bool FMemoryDerivedDataBackend::IsWritable()
+{
+	FScopeLock ScopeLock(&SynchronizationObject);
+	return !bDisabled;
+}
+
+bool FMemoryDerivedDataBackend::CachedDataProbablyExists(const TCHAR* CacheKey)
+{
+	COOK_STAT(auto Timer = UsageStats.TimeProbablyExists());
+	FScopeLock ScopeLock(&SynchronizationObject);
+	if (bDisabled)
+	{
+		return false;
+	}
+	// to avoid constant error reporting in async put due to restricted cache size, 
+	// we report true if the max size has been exceeded
+	if (bMaxSizeExceeded)
+	{
+		return true;
+	}
+
+	bool Result = CacheItems.Contains(FString(CacheKey));
+	if (Result)
+	{
+		COOK_STAT(Timer.AddHit(0));
+	}
+	return Result;
+}
+
+bool FMemoryDerivedDataBackend::GetCachedData(const TCHAR* CacheKey, TArray<uint8>& OutData)
+{
+	COOK_STAT(auto Timer = UsageStats.TimeGet());
+	FScopeLock ScopeLock(&SynchronizationObject);
+	if (!bDisabled)
+	{
+		FCacheValue* Item = CacheItems.FindRef(FString(CacheKey));
+		if (Item)
+		{
+			OutData = Item->Data;
+			Item->Age = 0;
+			check(OutData.Num());
+			COOK_STAT(Timer.AddHit(OutData.Num()));
+			return true;
+		}
+	}
+	OutData.Empty();
+	return false;
+}
+
+void FMemoryDerivedDataBackend::PutCachedData(const TCHAR* CacheKey, TArray<uint8>& InData, bool bPutEvenIfExists)
+{
+	COOK_STAT(auto Timer = UsageStats.TimePut());
+	FScopeLock ScopeLock(&SynchronizationObject);
+	
+	if (bDisabled || bMaxSizeExceeded)
+	{
+		return;
+	}
+	
+	FString Key(CacheKey);
+	FCacheValue* Item = CacheItems.FindRef(FString(CacheKey));
+	if (Item)
+	{
+		//check(Item->Data == InData); // any second attempt to push data should be identical data
+	}
+	else
+	{
+		FCacheValue* Val = new FCacheValue(InData);
+		int32 CacheValueSize = CalcCacheValueSize(Key, *Val);
+
+		// check if we haven't exceeded the MaxCacheSize
+		if (MaxCacheSize > 0 && (CurrentCacheSize + CacheValueSize) > MaxCacheSize)
+		{
+			delete Val;
+			UE_LOG(LogDerivedDataCache, Display, TEXT("Failed to cache data. Maximum cache size reached. CurrentSize %d kb / MaxSize: %d kb"), CurrentCacheSize / 1024, MaxCacheSize / 1024);
+			bMaxSizeExceeded = true;
+		}
+		else
+		{
+			COOK_STAT(Timer.AddHit(InData.Num()));
+			CacheItems.Add(Key, Val);
+			CalcCacheValueSize(Key, *Val);
+
+			CurrentCacheSize += CacheValueSize;
+		}
+	}
+}
+
+void FMemoryDerivedDataBackend::RemoveCachedData(const TCHAR* CacheKey, bool bTransient)
+{
+	FScopeLock ScopeLock(&SynchronizationObject);
+	if (bDisabled || bTransient)
+	{
+		return;
+	}
+	FString Key(CacheKey);
+	FCacheValue* Item = NULL;
+	if (CacheItems.RemoveAndCopyValue(Key, Item))
+	{
+		CurrentCacheSize -= CalcCacheValueSize(Key, *Item);
+		bMaxSizeExceeded = false;
+
+		check(Item);
+		delete Item;
+	}
+	else
+	{
+		check(!Item);
+	}
+}
+
+bool FMemoryDerivedDataBackend::SaveCache(const TCHAR* Filename)
+{
+	double StartTime = FPlatformTime::Seconds();
+	TUniquePtr<FArchive> SaverArchive(IFileManager::Get().CreateFileWriter(Filename, FILEWRITE_EvenIfReadOnly));
+	if (!SaverArchive)
+	{
+		UE_LOG(LogDerivedDataCache, Error, TEXT("Could not save memory cache %s."), Filename);
+		return false;
+	}
+
+	FArchive& Saver = *SaverArchive;
+	uint32 Magic = MemCache_Magic64;
+	Saver << Magic;
+	const int64 DataStartOffset = Saver.Tell();
+	{
+		FScopeLock ScopeLock(&SynchronizationObject);
+		check(!bDisabled);
+		for (TMap<FString, FCacheValue*>::TIterator It(CacheItems); It; ++It )
+		{
+			Saver << It.Key();
+			Saver << It.Value()->Age;
+			Saver << It.Value()->Data;
+		}
+	}
+	const int64 DataSize = Saver.Tell(); // Everything except the footer
+	int64 Size = DataSize;
+	uint32 Crc = MemCache_Magic64; // Crc takes more time than I want to spend  FCrc::MemCrc_DEPRECATED(&Buffer[0], Size);
+	Saver << Size;
+	Saver << Crc;
+
+	check(SerializationSpecificDataSize + DataSize <= MaxCacheSize || MaxCacheSize <= 0);
+
+	UE_LOG(LogDerivedDataCache, Log, TEXT("Saved boot cache %4.2fs %lldMB %s."), float(FPlatformTime::Seconds() - StartTime), DataSize / (1024 * 1024), Filename);
+	return true;
+}
+
+bool FMemoryDerivedDataBackend::LoadCache(const TCHAR* Filename)
+{
+	double StartTime = FPlatformTime::Seconds();
+	const int64 FileSize = IFileManager::Get().FileSize(Filename);
+	if (FileSize < 0)
+	{
+		UE_LOG(LogDerivedDataCache, Warning, TEXT("Could not find memory cache %s."), Filename);
+		return false;
+	}
+	// We test 3 * uint32 which is the old format (< SerializationSpecificDataSize). We'll test
+	// against SerializationSpecificDataSize later when we read the magic number from the cache.
+	if (FileSize < sizeof(uint32) * 3)
+	{
+		UE_LOG(LogDerivedDataCache, Error, TEXT("Memory cache was corrputed (short) %s."), Filename);
+		return false;
+	}
+	if (FileSize > MaxCacheSize*2 && MaxCacheSize > 0)
+	{
+		UE_LOG(LogDerivedDataCache, Error, TEXT("Refusing to load DDC cache %s. Size exceeds doubled MaxCacheSize."), Filename);
+		return false;
+	}
+
+	TUniquePtr<FArchive> LoaderArchive(IFileManager::Get().CreateFileReader(Filename));
+	if (!LoaderArchive)
+	{
+		UE_LOG(LogDerivedDataCache, Warning, TEXT("Could not read memory cache %s."), Filename);
+		return false;
+	}
+
+	FArchive& Loader = *LoaderArchive;
+	uint32 Magic = 0;
+	Loader << Magic;
+	if (Magic != MemCache_Magic && Magic != MemCache_Magic64)
+	{
+		UE_LOG(LogDerivedDataCache, Error, TEXT("Memory cache was corrputed (magic) %s."), Filename);
+		return false;
+	}
+	// Check the file size again, this time against the correct minimum size.
+	if (Magic == MemCache_Magic64 && FileSize < SerializationSpecificDataSize)
+	{
+		UE_LOG(LogDerivedDataCache, Error, TEXT("Memory cache was corrputed (short) %s."), Filename);
+		return false;
+	}
+	// Calculate expected DataSize based on the magic number (footer size difference)
+	const int64 DataSize = FileSize - (Magic == MemCache_Magic64 ? (SerializationSpecificDataSize - sizeof(uint32)) : (sizeof(uint32) * 2));		
+	Loader.Seek(DataSize);
+	int64 Size = 0;
+	uint32 Crc = 0;
+	if (Magic == MemCache_Magic64)
+	{
+		Loader << Size;
+	}
+	else
+	{
+		uint32 Size32 = 0;
+		Loader << Size32;
+		Size = (int64)Size32;
+	}
+	Loader << Crc;
+	if (Size != DataSize)
+	{
+		UE_LOG(LogDerivedDataCache, Error, TEXT("Memory cache was corrputed (size) %s."), Filename);
+		return false;
+	}
+	if ((Crc != MemCache_Magic && Crc != MemCache_Magic64) || Crc != Magic)
+	{
+		UE_LOG(LogDerivedDataCache, Warning, TEXT("Memory cache was corrputed (crc) %s."), Filename);
+		return false;
+	}
+	// Seek to data start offset (skip magic number)
+	Loader.Seek(sizeof(uint32));
+	{
+		TArray<uint8> Working;
+		FScopeLock ScopeLock(&SynchronizationObject);
+		check(!bDisabled);
+		while (Loader.Tell() < DataSize)
+		{
+			FString Key;
+			int32 Age;
+			Loader << Key;
+			Loader << Age;
+			Age++;
+			Loader << Working;
+			if (Age < MaxAge)
+			{
+				CacheItems.Add(Key, new FCacheValue(Working, Age));
+			}
+			Working.Reset();
+		}
+		// these are just a double check on ending correctly
+		if (Magic == MemCache_Magic64)
+		{
+			Loader << Size;
+		}
+		else
+		{
+			uint32 Size32 = 0;
+			Loader << Size32;
+			Size = (int64)Size32;
+		}
+		Loader << Crc;
+	}
+		
+	CurrentCacheSize = FileSize;
+	CacheFilename = Filename;
+	UE_LOG(LogDerivedDataCache, Log, TEXT("Loaded boot cache %4.2fs %lldMB %s."), float(FPlatformTime::Seconds() - StartTime), DataSize / (1024 * 1024), Filename);
+	return true;
+}
+
+void FMemoryDerivedDataBackend::Disable()
+{
+	FScopeLock ScopeLock(&SynchronizationObject);
+	bDisabled = true;
+	for (TMap<FString,FCacheValue*>::TIterator It(CacheItems); It; ++It )
+	{
+		delete It.Value();
+	}
+	CacheItems.Empty();
+
+	CurrentCacheSize = SerializationSpecificDataSize;
+}
+
+void FMemoryDerivedDataBackend::GatherUsageStats(TMap<FString, FDerivedDataCacheUsageStats>& UsageStatsMap, FString&& GraphPath)
+{
+	COOK_STAT(UsageStatsMap.Add(FString::Printf(TEXT("%s: %s.%s"), *GraphPath, TEXT("MemoryBackend"), *CacheFilename), UsageStats));
+}
diff --git a/Engine/Source/Developer/DerivedDataCache/Private/MemoryDerivedDataBackend.h b/Engine/Source/Developer/DerivedDataCache/Private/MemoryDerivedDataBackend.h
--- a/Engine/Source/Developer/DerivedDataCache/Private/MemoryDerivedDataBackend.h
+++ b/Engine/Source/Developer/DerivedDataCache/Private/MemoryDerivedDataBackend.h
@@ -1,6 +1,7 @@
 // Copyright 1998-2016 Epic Games, Inc. All Rights Reserved.
 
 #pragma once
+#include "DerivedDataBackendInterface.h"
 #include "DerivedDataCacheUsageStats.h"
 
 /** 
@@ -9,24 +10,11 @@
 class FMemoryDerivedDataBackend : public FDerivedDataBackendInterface
 {
 public:
-	FMemoryDerivedDataBackend(int64 InMaxCacheSize = -1)
-		: MaxCacheSize(InMaxCacheSize)
-		, bDisabled( false )
-		, CurrentCacheSize( SerializationSpecificDataSize )
-		, bMaxSizeExceeded(false)
-	{
-	}
-	~FMemoryDerivedDataBackend()
-	{
-		Disable();
-	}
+	explicit FMemoryDerivedDataBackend(int64 InMaxCacheSize = -1);
+	~FMemoryDerivedDataBackend();
 
 	/** return true if this cache is writable **/
-	virtual bool IsWritable() override
-	{
-		FScopeLock ScopeLock(&SynchronizationObject);
-		return !bDisabled;
-	}
+	virtual bool IsWritable() override;
 
 	/**
 	 * Synchronous test for the existence of a cache item
@@ -34,28 +22,8 @@ public:
 	 * @param	CacheKey	Alphanumeric+underscore key of this cache item
 	 * @return				true if the data probably will be found, this can't be guaranteed because of concurrency in the backends, corruption, etc
 	 */
-	virtual bool CachedDataProbablyExists(const TCHAR* CacheKey) override
-	{
-		COOK_STAT(auto Timer = UsageStats.TimeProbablyExists());
-		FScopeLock ScopeLock(&SynchronizationObject);
-		if (bDisabled)
-		{
-			return false;
-		}
-		// to avoid constant error reporting in async put due to restricted cache size, 
-		// we report true if the max size has been exceeded
-		if (bMaxSizeExceeded)
-		{
-			return true;
-		}
+	virtual bool CachedDataProbablyExists(const TCHAR* CacheKey) override;
 
-		bool Result = CacheItems.Contains(FString(CacheKey));
-		if (Result)
-		{
-			COOK_STAT(Timer.AddHit(0));
-		}
-		return Result;
-	}
 	/**
 	 * Synchronous retrieve of a cache item
 	 *
@@ -63,25 +31,8 @@ public:
 	 * @param	OutData		Buffer to receive the results, if any were found
 	 * @return				true if any data was found, and in this case OutData is non-empty
 	 */
-	virtual bool GetCachedData(const TCHAR* CacheKey, TArray<uint8>& OutData) override
-	{
-		COOK_STAT(auto Timer = UsageStats.TimeGet());
-		FScopeLock ScopeLock(&SynchronizationObject);
-		if (!bDisabled)
-		{
-			FCacheValue* Item = CacheItems.FindRef(FString(CacheKey));
-			if (Item)
-			{
-				OutData = Item->Data;
-				Item->Age = 0;
-				check(OutData.Num());
-				COOK_STAT(Timer.AddHit(OutData.Num()));
-				return true;
-			}
-		}
-		OutData.Empty();
-		return false;
-	}
+	virtual bool GetCachedData(const TCHAR* CacheKey, TArray<uint8>& OutData) override;
+
 	/**
 	 * Asynchronous, fire-and-forget placement of a cache item
 	 *
@@ -89,241 +40,30 @@ public:
 	 * @param	InData		Buffer containing the data to cache, can be destroyed after the call returns, immediately
 	 * @param	bPutEvenIfExists	If true, then do not attempt skip the put even if CachedDataProbablyExists returns true
 	 */
-	virtual void PutCachedData(const TCHAR* CacheKey, TArray<uint8>& InData, bool bPutEvenIfExists) override
-	{
-		COOK_STAT(auto Timer = UsageStats.TimePut());
-		FScopeLock ScopeLock(&SynchronizationObject);
-		
-		if (bDisabled || bMaxSizeExceeded)
-		{
-			return;
-		}
-		
-		FString Key(CacheKey);
-		FCacheValue* Item = CacheItems.FindRef(FString(CacheKey));
-		if (Item)
-		{
-			//check(Item->Data == InData); // any second attempt to push data should be identical data
-		}
-		else
-		{
-			FCacheValue* Val = new FCacheValue(InData);
-			int32 CacheValueSize = CalcCacheValueSize(Key, *Val);
+	virtual void PutCachedData(const TCHAR* CacheKey, TArray<uint8>& InData, bool bPutEvenIfExists) override;
 
-			// check if we haven't exceeded the MaxCacheSize
-			if (MaxCacheSize > 0 && (CurrentCacheSize + CacheValueSize) > MaxCacheSize)
-			{
-				delete Val;
-				UE_LOG(LogDerivedDataCache, Display, TEXT("Failed to cache data. Maximum cache size reached. CurrentSize %d kb / MaxSize: %d kb"), CurrentCacheSize / 1024, MaxCacheSize / 1024);
-				bMaxSizeExceeded = true;
-			}
-			else
-			{
-				COOK_STAT(Timer.AddHit(InData.Num()));
-				CacheItems.Add(Key, Val);
-				CalcCacheValueSize(Key, *Val);
-
-				CurrentCacheSize += CacheValueSize;
-			}
-		}
-	}
-
-	virtual void RemoveCachedData(const TCHAR* CacheKey, bool bTransient) override
-	{
-		FScopeLock ScopeLock(&SynchronizationObject);
-		if (bDisabled || bTransient)
-		{
-			return;
-		}
-		FString Key(CacheKey);
-		FCacheValue* Item = NULL;
-		if (CacheItems.RemoveAndCopyValue(Key, Item))
-		{
-			CurrentCacheSize -= CalcCacheValueSize(Key, *Item);
-			bMaxSizeExceeded = false;
-
-			check(Item);
-			delete Item;
-		}
-		else
-		{
-			check(!Item);
-		}
-	}
+	virtual void RemoveCachedData(const TCHAR* CacheKey, bool bTransient) override;
 
 	/**
 	 * Save the cache to disk
 	 * @param	Filename	Filename to save
 	 * @return	true if file was saved sucessfully
 	 */
-	bool SaveCache(const TCHAR* Filename)
-	{
-		double StartTime = FPlatformTime::Seconds();
-		TAutoPtr<FArchive> SaverArchive(IFileManager::Get().CreateFileWriter(Filename, FILEWRITE_EvenIfReadOnly));
-		if (!SaverArchive.IsValid())
-		{
-			UE_LOG(LogDerivedDataCache, Error, TEXT("Could not save memory cache %s."), Filename);
-			return false;
-		}
-
-		FArchive& Saver = *SaverArchive;
-		uint32 Magic = MemCache_Magic64;
-		Saver << Magic;
-		const int64 DataStartOffset = Saver.Tell();
-		{
-			FScopeLock ScopeLock(&SynchronizationObject);
-			check(!bDisabled);
-			for (TMap<FString, FCacheValue*>::TIterator It(CacheItems); It; ++It )
-			{
-				Saver << It.Key();
-				Saver << It.Value()->Age;
-				Saver << It.Value()->Data;
-			}
-		}
-		const int64 DataSize = Saver.Tell(); // Everything except the footer
-		int64 Size = DataSize;
-		uint32 Crc = MemCache_Magic64; // Crc takes more time than I want to spend  FCrc::MemCrc_DEPRECATED(&Buffer[0], Size);
-		Saver << Size;
-		Saver << Crc;
+	bool SaveCache(const TCHAR* Filename);
 
-		check(SerializationSpecificDataSize + DataSize <= MaxCacheSize || MaxCacheSize <= 0);
-
-		UE_LOG(LogDerivedDataCache, Log, TEXT("Saved boot cache %4.2fs %lldMB %s."), float(FPlatformTime::Seconds() - StartTime), DataSize / (1024 * 1024), Filename);
-		return true;
-	}
 	/**
 	 * Load the cache from disk
 	 * @param	Filename	Filename to load
 	 * @return	true if file was loaded sucessfully
 	 */
-	bool LoadCache(const TCHAR* Filename)
-	{
-		double StartTime = FPlatformTime::Seconds();
-		const int64 FileSize = IFileManager::Get().FileSize(Filename);
-		if (FileSize < 0)
-		{
-			UE_LOG(LogDerivedDataCache, Warning, TEXT("Could not find memory cache %s."), Filename);
-			return false;
-		}
-		// We test 3 * uint32 which is the old format (< SerializationSpecificDataSize). We'll test
-		// agains SerializationSpecificDataSize later when we read the magic number from the cache.
-		if (FileSize < sizeof(uint32) * 3)
-		{
-			UE_LOG(LogDerivedDataCache, Error, TEXT("Memory cache was corrputed (short) %s."), Filename);
-			return false;
-		}
-		if (FileSize > MaxCacheSize*2 && MaxCacheSize > 0)
-		{
-			UE_LOG(LogDerivedDataCache, Error, TEXT("Refusing to load DDC cache %s. Size exceeds doubled MaxCacheSize."), Filename);
-			return false;
-		}
-
-		TAutoPtr<FArchive> LoderArchive(IFileManager::Get().CreateFileReader(Filename));
-		if (!LoderArchive.IsValid())
-		{
-			UE_LOG(LogDerivedDataCache, Warning, TEXT("Could not read memory cache %s."), Filename);
-			return false;
-		}
+	bool LoadCache(const TCHAR* Filename);
 
-		FArchive& Loader = *LoderArchive;
-		uint32 Magic = 0;
-		Loader << Magic;
-		if (Magic != MemCache_Magic && Magic != MemCache_Magic64)
-		{
-			UE_LOG(LogDerivedDataCache, Error, TEXT("Memory cache was corrputed (magic) %s."), Filename);
-			return false;
-		}
-		// Check the file size again, this time against the correct minimum size.
-		if (Magic == MemCache_Magic64 && FileSize < SerializationSpecificDataSize)
-		{
-			UE_LOG(LogDerivedDataCache, Error, TEXT("Memory cache was corrputed (short) %s."), Filename);
-			return false;
-		}
-		// Calculate expected DataSize based on the magic number (footer size difference)
-		const int64 DataSize = FileSize - (Magic == MemCache_Magic64 ? (SerializationSpecificDataSize - sizeof(uint32)) : (sizeof(uint32) * 2));		
-		Loader.Seek(DataSize);
-		int64 Size = 0;
-		uint32 Crc = 0;
-		if (Magic == MemCache_Magic64)
-		{
-			Loader << Size;
-		}
-		else
-		{
-			uint32 Size32 = 0;
-			Loader << Size32;
-			Size = (int64)Size32;
-		}
-		Loader << Crc;
-		if (Size != DataSize)
-		{
-			UE_LOG(LogDerivedDataCache, Error, TEXT("Memory cache was corrputed (size) %s."), Filename);
-			return false;
-		}
-		if ((Crc != MemCache_Magic && Crc != MemCache_Magic64) || Crc != Magic)
-		{
-			UE_LOG(LogDerivedDataCache, Warning, TEXT("Memory cache was corrputed (crc) %s."), Filename);
-			return false;
-		}
-		// Seek to data start offset (skip magic number)
-		Loader.Seek(sizeof(uint32));
-		{
-			TArray<uint8> Working;
-			FScopeLock ScopeLock(&SynchronizationObject);
-			check(!bDisabled);
-			while (Loader.Tell() < DataSize)
-			{
-				FString Key;
-				int32 Age;
-				Loader << Key;
-				Loader << Age;
-				Age++;
-				Loader << Working;
-				if (Age < MaxAge)
-				{
-					CacheItems.Add(Key, new FCacheValue(Working, Age));
-				}
-				Working.Reset();
-			}
-			// these are just a double check on ending correctly
-			if (Magic == MemCache_Magic64)
-			{
-				Loader << Size;
-			}
-			else
-			{
-				uint32 Size32 = 0;
-				Loader << Size32;
-				Size = (int64)Size32;
-			}
-			Loader << Crc;
-		}
-		
-		CurrentCacheSize = FileSize;
-		CacheFilename = Filename;
-		UE_LOG(LogDerivedDataCache, Log, TEXT("Loaded boot cache %4.2fs %lldMB %s."), float(FPlatformTime::Seconds() - StartTime), DataSize / (1024 * 1024), Filename);
-		return true;
-	}
 	/**
 	 * Disable cache and ignore all subsequent requests
 	 */
-	void Disable()
-	{
-		FScopeLock ScopeLock(&SynchronizationObject);
-		bDisabled = true;
-		for (TMap<FString,FCacheValue*>::TIterator It(CacheItems); It; ++It )
-		{
-			delete It.Value();
-		}
-		CacheItems.Empty();
+	void Disable();
 
-		CurrentCacheSize = SerializationSpecificDataSize;
-	}
-
-	virtual void GatherUsageStats(TMap<FString, FDerivedDataCacheUsageStats>& UsageStatsMap, FString&& GraphPath) override
-	{
-		COOK_STAT(UsageStatsMap.Add(FString::Printf(TEXT("%s: %s.%s"), *GraphPath, TEXT("MemoryBackend"), *CacheFilename), UsageStats));
-	}
+	virtual void GatherUsageStats(TMap<FString, FDerivedDataCacheUsageStats>& UsageStatsMap, FString&& GraphPath) override;
 
 private:
 	/** Name of the cache file loaded (if any). */
diff --git a/Engine/Source/Developer/DerivedDataCache/Private/PakFileDerivedDataBackend.cpp b/Engine/Source/Developer/DerivedDataCache/Private/PakFileDerivedDataBackend.cpp
new file mode 100644
--- /dev/null
+++ b/Engine/Source/Developer/DerivedDataCache/Private/PakFileDerivedDataBackend.cpp
@@ -0,0 +1,438 @@
+// Copyright 1998-2016 Epic Games, Inc. All Rights Reserved.
+
+#pragma once
+
+#include "Core.h"
+#include "Compression.h"
+#include "DerivedDataCacheUsageStats.h"
+#include "PakFileDerivedDataBackend.h"
+
+FPakFileDerivedDataBackend::FPakFileDerivedDataBackend(const TCHAR* InFilename, bool bInWriting)
+	: bWriting(bInWriting)
+	, bClosed(false)
+	, Filename(InFilename)
+{
+	if (bWriting)
+	{
+		FileHandle.Reset(IFileManager::Get().CreateFileWriter(InFilename, FILEWRITE_NoReplaceExisting));
+		if (!FileHandle)
+		{
+			UE_LOG(LogDerivedDataCache, Fatal, TEXT("Pak cache could not be opened for writing %s."), InFilename);
+			bClosed = true;
+		}
+		else
+		{
+			UE_LOG(LogDerivedDataCache, Display, TEXT("Pak cache opened for writing %s."), InFilename);
+		}
+	}
+	else
+	{
+		FileHandle.Reset(IFileManager::Get().CreateFileReader(InFilename));
+		if (!FileHandle)
+		{
+			UE_LOG(LogDerivedDataCache, Warning, TEXT("Pak cache could not be opened for reading %s."), InFilename);
+		}
+		else if (!LoadCache(InFilename))
+		{
+			FileHandle.Reset();
+			CacheItems.Empty();
+			bClosed = true;
+		}
+		else
+		{
+			UE_LOG(LogDerivedDataCache, Display, TEXT("Pak cache opened for reading %s."), InFilename);
+		}
+	}
+}
+
+FPakFileDerivedDataBackend::~FPakFileDerivedDataBackend()
+{
+	Close();
+}
+
+void FPakFileDerivedDataBackend::Close()
+{
+	FDerivedDataBackend::Get().WaitForQuiescence();
+	FScopeLock ScopeLock(&SynchronizationObject);
+	if (!bClosed)
+	{
+		if (bWriting)
+		{
+			SaveCache();
+		}
+		FileHandle.Reset();
+		CacheItems.Empty();
+		bClosed = true;
+	}
+}
+
+bool FPakFileDerivedDataBackend::IsWritable()
+{
+	return bWriting && !bClosed;
+}
+
+bool FPakFileDerivedDataBackend::BackfillLowerCacheLevels()
+{
+	return false;
+}
+
+bool FPakFileDerivedDataBackend::CachedDataProbablyExists(const TCHAR* CacheKey)
+{
+	COOK_STAT(auto Timer = UsageStats.TimeProbablyExists());
+	FScopeLock ScopeLock(&SynchronizationObject);
+	bool Result = CacheItems.Contains(FString(CacheKey));
+	if (Result)
+	{
+		COOK_STAT(Timer.AddHit(0));
+	}
+	return Result;
+}
+
+bool FPakFileDerivedDataBackend::GetCachedData(const TCHAR* CacheKey, TArray<uint8>& OutData)
+{
+	COOK_STAT(auto Timer = UsageStats.TimeGet());
+	if (bWriting || bClosed)
+	{
+		return false;
+	}
+	FScopeLock ScopeLock(&SynchronizationObject);
+	FCacheValue* Item = CacheItems.Find(FString(CacheKey));
+	if (Item)
+	{
+		check(FileHandle);
+		FileHandle->Seek(Item->Offset);
+		if (FileHandle->Tell() != Item->Offset)
+		{
+			UE_LOG(LogDerivedDataCache, Warning, TEXT("Pak file, bad seek."));
+		}
+		else
+		{
+			check(Item->Size);
+			check(!OutData.Num());
+			check(FileHandle->IsLoading());
+			OutData.AddUninitialized(Item->Size);
+			FileHandle->Serialize(OutData.GetData(), int64(Item->Size));
+			uint32 TestCrc = FCrc::MemCrc_DEPRECATED(OutData.GetData(), Item->Size);
+			if (TestCrc != Item->Crc)
+			{
+				UE_LOG(LogDerivedDataCache, Warning, TEXT("Pak file, bad crc."));
+			}
+			else
+			{
+				UE_LOG(LogDerivedDataCache, Verbose, TEXT("FPakFileDerivedDataBackend: Cache hit on %s"), CacheKey);
+				check(OutData.Num());
+				COOK_STAT(Timer.AddHit(OutData.Num()));
+				return true;
+			}
+		}
+	}
+	else
+	{
+		UE_LOG(LogDerivedDataCache, Verbose, TEXT("FPakFileDerivedDataBackend: Miss on %s"), CacheKey);
+	}
+	OutData.Empty();
+	return false;
+}
+
+void FPakFileDerivedDataBackend::PutCachedData(const TCHAR* CacheKey, TArray<uint8>& InData, bool bPutEvenIfExists)
+{
+	COOK_STAT(auto Timer = UsageStats.TimePut());
+	if (!bWriting || bClosed)
+	{
+		return;
+	}
+	{
+		FScopeLock ScopeLock(&SynchronizationObject);
+		FString Key(CacheKey);
+		FCacheValue* Item = CacheItems.Find(FString(CacheKey));
+		if (!Item)
+		{
+			check(InData.Num());
+			check(Key.Len());
+			uint32 Crc = FCrc::MemCrc_DEPRECATED(InData.GetData(), InData.Num());
+			check(FileHandle);
+			check(FileHandle->IsSaving());
+			int64 Offset = FileHandle->Tell();
+			if (Offset < 0)
+			{
+				CacheItems.Empty();
+				FileHandle.Reset();
+				bClosed = true;
+			}
+			else
+			{
+				COOK_STAT(Timer.AddHit(InData.Num()));
+				FileHandle->Serialize(InData.GetData(), int64(InData.Num()));
+				UE_LOG(LogDerivedDataCache, Verbose, TEXT("FPakFileDerivedDataBackend: Put %s"), CacheKey);
+				CacheItems.Add(Key,FCacheValue(Offset, InData.Num(), Crc));
+			}
+		}
+	}
+	if (bClosed)
+	{
+		UE_LOG(LogDerivedDataCache, Fatal, TEXT("Could not write pak file...out of disk space?"));
+	}
+}
+
+void FPakFileDerivedDataBackend::RemoveCachedData(const TCHAR* CacheKey, bool bTransient)
+{
+	if (bClosed || bTransient)
+	{
+		return;
+	}
+	// strangish. We can delete from a pak, but it only deletes the index 
+	// if this is a read cache, it will read it next time
+	// if this is a write cache, we wated space
+	FScopeLock ScopeLock(&SynchronizationObject);
+	FString Key(CacheKey);
+	CacheItems.Remove(Key);
+}
+
+bool FPakFileDerivedDataBackend::SaveCache()
+{
+	FScopeLock ScopeLock(&SynchronizationObject);
+	check(FileHandle);
+	check(FileHandle->IsSaving());
+	int64 IndexOffset = FileHandle->Tell();
+	check(IndexOffset >= 0);
+	uint32 NumItems =  uint32(CacheItems.Num());
+	check(IndexOffset > 0 || !NumItems);
+	TArray<uint8> IndexBuffer;
+	{
+		FMemoryWriter Saver(IndexBuffer);
+		uint32 NumProcessed = 0;
+		for (TMap<FString, FCacheValue>::TIterator It(CacheItems); It; ++It )
+		{
+			check(It.Value().Offset >= 0 && It.Value().Offset < IndexOffset);
+			check(It.Value().Size);
+			check(It.Key().Len());
+			Saver << It.Key();
+			Saver << It.Value().Offset;
+			Saver << It.Value().Size;
+			Saver << It.Value().Crc;
+			NumProcessed++;
+		}
+		check(NumProcessed == NumItems);
+	}
+	uint32 IndexCrc = FCrc::MemCrc_DEPRECATED(IndexBuffer.GetData(), IndexBuffer.Num());
+	uint32 SizeIndex = uint32(IndexBuffer.Num());
+
+	uint32 Magic = PakCache_Magic;
+	TArray<uint8> Buffer;
+	FMemoryWriter Saver(Buffer);
+	Saver << Magic;
+	Saver << IndexCrc;
+	Saver << NumItems;
+	Saver << SizeIndex;
+	Saver.Serialize(IndexBuffer.GetData(), IndexBuffer.Num());
+	Saver << Magic;
+	Saver << IndexOffset;
+	FileHandle->Serialize(Buffer.GetData(), Buffer.Num());
+	CacheItems.Empty();
+	FileHandle.Reset();
+	bClosed = true;
+	return true;
+}
+
+bool FPakFileDerivedDataBackend::LoadCache(const TCHAR* InFilename)
+{
+	check(FileHandle);
+	check(FileHandle->IsLoading());
+	int64 FileSize = FileHandle->TotalSize();
+	check(FileSize >= 0);
+	if (FileSize < sizeof(int64) + sizeof(uint32) * 5)
+	{
+		UE_LOG(LogDerivedDataCache, Error, TEXT("Pak cache was corrputed (short) %s."), InFilename);
+		return false;
+	}
+	int64 IndexOffset = -1;
+	int64 Trailer = -1;
+	{
+		TArray<uint8> Buffer;
+		const int64 SeekPos = FileSize - int64(sizeof(int64) + sizeof(uint32));
+		FileHandle->Seek(SeekPos);
+		Trailer = FileHandle->Tell();
+		if (Trailer != SeekPos)
+		{
+			UE_LOG(LogDerivedDataCache, Error, TEXT("Pak cache was corrputed (bad seek) %s."), InFilename);
+			return false;
+		}
+		check(Trailer >= 0 && Trailer < FileSize);
+		Buffer.AddUninitialized(sizeof(int64) + sizeof(uint32));
+		FileHandle->Serialize(Buffer.GetData(), int64(sizeof(int64)+sizeof(uint32)));
+		FMemoryReader Loader(Buffer);
+		uint32 Magic = 0;
+		Loader << Magic;
+		Loader << IndexOffset;
+		if (Magic != PakCache_Magic || IndexOffset < 0 || IndexOffset + int64(sizeof(uint32) * 4) > Trailer)
+		{
+			UE_LOG(LogDerivedDataCache, Error, TEXT("Pak cache was corrputed (bad footer) %s."), InFilename);
+			return false;
+		}
+	}
+	uint32 IndexCrc = 0;
+	uint32 NumIndex = 0;
+	uint32 SizeIndex = 0;
+	{
+		TArray<uint8> Buffer;
+		FileHandle->Seek(IndexOffset);
+		if (FileHandle->Tell() != IndexOffset)
+		{
+			UE_LOG(LogDerivedDataCache, Error, TEXT("Pak cache was corrputed (bad seek index) %s."), InFilename);
+			return false;
+		}
+		Buffer.AddUninitialized(sizeof(uint32) * 4);
+		FileHandle->Serialize(Buffer.GetData(), int64(sizeof(uint32)* 4));
+		FMemoryReader Loader(Buffer);
+		uint32 Magic = 0;
+		Loader << Magic;
+		Loader << IndexCrc;
+		Loader << NumIndex;
+		Loader << SizeIndex;
+		if (Magic != PakCache_Magic || (SizeIndex != 0 && NumIndex == 0) || (SizeIndex == 0 && NumIndex != 0)) 
+		{
+			UE_LOG(LogDerivedDataCache, Error, TEXT("Pak cache was corrputed (bad index header) %s."), InFilename);
+			return false;
+		}
+		if (IndexOffset + sizeof(uint32) * 4 + SizeIndex != Trailer) 
+		{
+			UE_LOG(LogDerivedDataCache, Error, TEXT("Pak cache was corrputed (bad index size) %s."), InFilename);
+			return false;
+		}
+	}
+	{
+		TArray<uint8> Buffer;
+		Buffer.AddUninitialized(SizeIndex);
+		FileHandle->Serialize(Buffer.GetData(), int64(SizeIndex));
+		FMemoryReader Loader(Buffer);
+		while (Loader.Tell() < (int32)SizeIndex)
+		{
+			FString Key;
+			int64 Offset;
+			int64 Size;
+			uint32 Crc;
+			Loader << Key;
+			Loader << Offset;
+			Loader << Size;
+			Loader << Crc;
+			if (!Key.Len() || Offset < 0 || Offset >= IndexOffset || !Size)
+			{
+				UE_LOG(LogDerivedDataCache, Error, TEXT("Pak cache was corrputed (bad index entry) %s."), InFilename);
+				return false;
+			}
+			CacheItems.Add(Key, FCacheValue(Offset, Size, Crc));
+		}
+		if (CacheItems.Num() != NumIndex)
+		{
+			UE_LOG(LogDerivedDataCache, Error, TEXT("Pak cache was corrputed (bad index count) %s."), InFilename);
+			return false;
+		}
+	}
+	return true;
+}
+
+void FPakFileDerivedDataBackend::MergeCache(FPakFileDerivedDataBackend* OtherPak)
+{
+	// Get all the existing keys
+	TArray<FString> KeyNames;
+	OtherPak->CacheItems.GenerateKeyArray(KeyNames);
+
+	// Find all the keys to copy
+	TArray<FString> CopyKeyNames;
+	for(const FString& KeyName : KeyNames)
+	{
+		if(!CachedDataProbablyExists(*KeyName))
+		{
+			CopyKeyNames.Add(KeyName);
+		}
+	}
+	UE_LOG(LogDerivedDataCache, Display, TEXT("Merging %d entries (%d skipped)."), CopyKeyNames.Num(), KeyNames.Num() - CopyKeyNames.Num());
+
+	// Copy them all to the new cache. Don't use the overloaded get/put methods (which may compress/decompress); copy the raw data directly.
+	TArray<uint8> Buffer;
+	for(const FString& CopyKeyName : CopyKeyNames)
+	{
+		Buffer.Reset();
+		if(OtherPak->FPakFileDerivedDataBackend::GetCachedData(*CopyKeyName, Buffer))
+		{
+			FPakFileDerivedDataBackend::PutCachedData(*CopyKeyName, Buffer, false);
+		}
+	}
+}
+
+bool FPakFileDerivedDataBackend::SortAndCopy(const FString &InputFilename, const FString &OutputFilename)
+{
+	// Open the input and output files
+	FPakFileDerivedDataBackend InputPak(*InputFilename, false);
+	if (InputPak.bClosed) return false;
+
+	FPakFileDerivedDataBackend OutputPak(*OutputFilename, true);
+	if (OutputPak.bClosed) return false;
+
+	// Sort the key names
+	TArray<FString> KeyNames;
+	InputPak.CacheItems.GenerateKeyArray(KeyNames);
+	KeyNames.Sort();
+
+	// Copy all the DDC to the new cache
+	TArray<uint8> Buffer;
+	TArray<uint32> KeySizes;
+	for (int KeyIndex = 0; KeyIndex < KeyNames.Num(); KeyIndex++)
+	{
+		Buffer.Reset();
+		InputPak.GetCachedData(*KeyNames[KeyIndex], Buffer);
+		OutputPak.PutCachedData(*KeyNames[KeyIndex], Buffer, false);
+		KeySizes.Add(Buffer.Num());
+	}
+
+	// Write out a TOC listing for debugging
+	FStringOutputDevice Output;
+	Output.Logf(TEXT("Asset,Size") LINE_TERMINATOR);
+	for(int KeyIndex = 0; KeyIndex < KeyNames.Num(); KeyIndex++)
+	{
+		Output.Logf(TEXT("%s,%d") LINE_TERMINATOR, *KeyNames[KeyIndex], KeySizes[KeyIndex]);
+	}
+	FFileHelper::SaveStringToFile(Output, *FPaths::Combine(*FPaths::GetPath(OutputFilename), *(FPaths::GetBaseFilename(OutputFilename) + TEXT(".csv"))));
+	return true;
+}
+
+void FPakFileDerivedDataBackend::GatherUsageStats(TMap<FString, FDerivedDataCacheUsageStats>& UsageStatsMap, FString&& GraphPath)
+{
+	COOK_STAT(UsageStatsMap.Add(FString::Printf(TEXT("%s: %s.%s"), *GraphPath, TEXT("PakFile"), *Filename), UsageStats));
+}
+
+FCompressedPakFileDerivedDataBackend::FCompressedPakFileDerivedDataBackend(const TCHAR* InFilename, bool bInWriting)
+	: FPakFileDerivedDataBackend(InFilename, bInWriting)
+{
+}
+
+void FCompressedPakFileDerivedDataBackend::PutCachedData(const TCHAR* CacheKey, TArray<uint8>& InData, bool bPutEvenIfExists)
+{
+	int32 UncompressedSize = InData.Num();
+	int32 CompressedSize = FCompression::CompressMemoryBound(CompressionFlags, UncompressedSize);
+
+	TArray<uint8> CompressedData;
+	CompressedData.AddUninitialized(CompressedSize + sizeof(UncompressedSize));
+
+	FMemory::Memcpy(&CompressedData[0], &UncompressedSize, sizeof(UncompressedSize));
+	verify(FCompression::CompressMemory(CompressionFlags, CompressedData.GetData() + sizeof(UncompressedSize), CompressedSize, InData.GetData(), InData.Num()));
+	CompressedData.SetNum(CompressedSize + sizeof(UncompressedSize), false);
+
+	FPakFileDerivedDataBackend::PutCachedData(CacheKey, CompressedData, bPutEvenIfExists);
+}
+
+bool FCompressedPakFileDerivedDataBackend::GetCachedData(const TCHAR* CacheKey, TArray<uint8>& OutData)
+{
+	TArray<uint8> CompressedData;
+	if(!FPakFileDerivedDataBackend::GetCachedData(CacheKey, CompressedData))
+	{
+		return false;
+	}
+
+	int32 UncompressedSize;
+	FMemory::Memcpy(&UncompressedSize, &CompressedData[0], sizeof(UncompressedSize));
+	OutData.SetNum(UncompressedSize);
+	verify(FCompression::UncompressMemory(CompressionFlags, OutData.GetData(), UncompressedSize, CompressedData.GetData() + sizeof(UncompressedSize), CompressedData.Num() - sizeof(UncompressedSize)));
+
+	return true;
+}
diff --git a/Engine/Source/Developer/DerivedDataCache/Private/PakFileDerivedDataBackend.h b/Engine/Source/Developer/DerivedDataCache/Private/PakFileDerivedDataBackend.h
--- a/Engine/Source/Developer/DerivedDataCache/Private/PakFileDerivedDataBackend.h
+++ b/Engine/Source/Developer/DerivedDataCache/Private/PakFileDerivedDataBackend.h
@@ -4,6 +4,8 @@
 
 #include "Compression.h"
 #include "DerivedDataCacheUsageStats.h"
+#include "DerivedDataBackendInterface.h"
+#include "UniquePtr.h"
 
 /** 
  * A simple thread safe, pak file based backend. 
@@ -11,75 +13,15 @@
 class FPakFileDerivedDataBackend : public FDerivedDataBackendInterface
 {
 public:
-	FPakFileDerivedDataBackend(const TCHAR* InFilename, bool bInWriting)
-		: bWriting(bInWriting)
-		, bClosed(false)
-		, Filename(InFilename)
-	{
-		if (bWriting)
-		{
-			FileHandle = IFileManager::Get().CreateFileWriter(InFilename, FILEWRITE_NoReplaceExisting);
-			if (!FileHandle)
-			{
-				UE_LOG(LogDerivedDataCache, Fatal, TEXT("Pak cache could not be opened for writing %s."), InFilename);
-				bClosed = true;
-			}
-			else
-			{
-				UE_LOG(LogDerivedDataCache, Display, TEXT("Pak cache opened for writing %s."), InFilename);
-			}
-		}
-		else
-		{
-			FileHandle = IFileManager::Get().CreateFileReader(InFilename);
-			if (!FileHandle)
-			{
-				UE_LOG(LogDerivedDataCache, Warning, TEXT("Pak cache could not be opened for reading %s."), InFilename);
-			}
-			else if (!LoadCache(InFilename))
-			{
-				FileHandle.Reset();
-				CacheItems.Empty();
-				bClosed = true;
-			}
-			else
-			{
-				UE_LOG(LogDerivedDataCache, Display, TEXT("Pak cache opened for reading %s."), InFilename);
-			}
-		}
-	}
+	FPakFileDerivedDataBackend(const TCHAR* InFilename, bool bInWriting);
+	~FPakFileDerivedDataBackend();
 
-	~FPakFileDerivedDataBackend()
-	{
-		Close();
-	}
-
-	void Close()
-	{
-		FDerivedDataBackend::Get().WaitForQuiescence();
-		FScopeLock ScopeLock(&SynchronizationObject);
-		if (!bClosed)
-		{
-			if (bWriting)
-			{
-				SaveCache();
-			}
-			FileHandle.Reset();
-			CacheItems.Empty();
-			bClosed = true;
-		}
-	}
+	void Close();
 
 	/** return true if this cache is writable **/
-	virtual bool IsWritable() override
-	{
-		return bWriting && !bClosed;
-	}
+	virtual bool IsWritable() override;
 
-	virtual bool BackfillLowerCacheLevels() override
-	{
-		return false;
-	}
+	virtual bool BackfillLowerCacheLevels() override;
 
 	/**
 	 * Synchronous test for the existence of a cache item
@@ -87,17 +29,8 @@ public:
 	 * @param	CacheKey	Alphanumeric+underscore key of this cache item
 	 * @return				true if the data probably will be found, this can't be guaranteed because of concurrency in the backends, corruption, etc
 	 */
-	virtual bool CachedDataProbablyExists(const TCHAR* CacheKey) override
-	{
-		COOK_STAT(auto Timer = UsageStats.TimeProbablyExists());
-		FScopeLock ScopeLock(&SynchronizationObject);
-		bool Result = CacheItems.Contains(FString(CacheKey));
-		if (Result)
-		{
-			COOK_STAT(Timer.AddHit(0));
-		}
-		return Result;
-	}
+	virtual bool CachedDataProbablyExists(const TCHAR* CacheKey) override;
+
 	/**
 	 * Synchronous retrieve of a cache item
 	 *
@@ -105,51 +38,8 @@ public:
 	 * @param	OutData		Buffer to receive the results, if any were found
 	 * @return				true if any data was found, and in this case OutData is non-empty
 	 */
-	virtual bool GetCachedData(const TCHAR* CacheKey, TArray<uint8>& OutData) override
-	{
-		COOK_STAT(auto Timer = UsageStats.TimeGet());
-		if (bWriting || bClosed)
-		{
-			return false;
-		}
-		FScopeLock ScopeLock(&SynchronizationObject);
-		FCacheValue* Item = CacheItems.Find(FString(CacheKey));
-		if (Item)
-		{
-			check(FileHandle);
-			FileHandle->Seek(Item->Offset);
-			if (FileHandle->Tell() != Item->Offset)
-			{
-				UE_LOG(LogDerivedDataCache, Warning, TEXT("Pak file, bad seek."));
-			}
-			else
-			{
-				check(Item->Size);
-				check(!OutData.Num());
-				check(FileHandle->IsLoading());
-				OutData.AddUninitialized(Item->Size);
-				FileHandle->Serialize(OutData.GetData(), int64(Item->Size));
-				uint32 TestCrc = FCrc::MemCrc_DEPRECATED(OutData.GetData(), Item->Size);
-				if (TestCrc != Item->Crc)
-				{
-					UE_LOG(LogDerivedDataCache, Warning, TEXT("Pak file, bad crc."));
-				}
-				else
-				{
-					UE_LOG(LogDerivedDataCache, Verbose, TEXT("FPakFileDerivedDataBackend: Cache hit on %s"), CacheKey);
-					check(OutData.Num());
-					COOK_STAT(Timer.AddHit(OutData.Num()));
-					return true;
-				}
-			}
-		}
-		else
-		{
-			UE_LOG(LogDerivedDataCache, Verbose, TEXT("FPakFileDerivedDataBackend: Miss on %s"), CacheKey);
-		}
-		OutData.Empty();
-		return false;
-	}
+	virtual bool GetCachedData(const TCHAR* CacheKey, TArray<uint8>& OutData) override;
+
 	/**
 	 * Asynchronous, fire-and-forget placement of a cache item
 	 *
@@ -157,288 +47,35 @@ public:
 	 * @param	InData		Buffer containing the data to cache, can be destroyed after the call returns, immediately
 	 * @param	bPutEvenIfExists	If true, then do not attempt skip the put even if CachedDataProbablyExists returns true
 	 */
-	virtual void PutCachedData(const TCHAR* CacheKey, TArray<uint8>& InData, bool bPutEvenIfExists) override
-	{
-		COOK_STAT(auto Timer = UsageStats.TimePut());
-		if (!bWriting || bClosed)
-		{
-			return;
-		}
-		{
-			FScopeLock ScopeLock(&SynchronizationObject);
-			FString Key(CacheKey);
-			FCacheValue* Item = CacheItems.Find(FString(CacheKey));
-			if (!Item)
-			{
-				check(InData.Num());
-				check(Key.Len());
-				uint32 Crc = FCrc::MemCrc_DEPRECATED(InData.GetData(), InData.Num());
-				check(FileHandle);
-				check(FileHandle->IsSaving());
-				int64 Offset = FileHandle->Tell();
-				if (Offset < 0)
-				{
-					CacheItems.Empty();
-					FileHandle.Reset();
-					bClosed = true;
-				}
-				else
-				{
-					COOK_STAT(Timer.AddHit(InData.Num()));
-					FileHandle->Serialize(InData.GetData(), int64(InData.Num()));
-					UE_LOG(LogDerivedDataCache, Verbose, TEXT("FPakFileDerivedDataBackend: Put %s"), CacheKey);
-					CacheItems.Add(Key,FCacheValue(Offset, InData.Num(), Crc));
-				}
-			}
-		}
-		if (bClosed)
-		{
-			UE_LOG(LogDerivedDataCache, Fatal, TEXT("Could not write pak file...out of disk space?"));
-		}
-	}
-
-	virtual void RemoveCachedData(const TCHAR* CacheKey, bool bTransient) override
-	{
-		if (bClosed || bTransient)
-		{
-			return;
-		}
-		// strangish. We can delete from a pak, but it only deletes the index 
-		// if this is a read cache, it will read it next time
-		// if this is a write cache, we wated space
-		FScopeLock ScopeLock(&SynchronizationObject);
-		FString Key(CacheKey);
-		CacheItems.Remove(Key);
-	}
+	virtual void PutCachedData(const TCHAR* CacheKey, TArray<uint8>& InData, bool bPutEvenIfExists) override;
+	virtual void RemoveCachedData(const TCHAR* CacheKey, bool bTransient) override;
 
 	/**
 	 * Save the cache to disk
 	 * @return	true if file was saved sucessfully
 	 */
-	bool SaveCache()
-	{
-		FScopeLock ScopeLock(&SynchronizationObject);
-		check(FileHandle);
-		check(FileHandle->IsSaving());
-		int64 IndexOffset = FileHandle->Tell();
-		check(IndexOffset >= 0);
-		uint32 NumItems =  uint32(CacheItems.Num());
-		check(IndexOffset > 0 || !NumItems);
-		TArray<uint8> IndexBuffer;
-		{
-			FMemoryWriter Saver(IndexBuffer);
-			uint32 NumProcessed = 0;
-			for (TMap<FString, FCacheValue>::TIterator It(CacheItems); It; ++It )
-			{
-				check(It.Value().Offset >= 0 && It.Value().Offset < IndexOffset);
-				check(It.Value().Size);
-				check(It.Key().Len());
-				Saver << It.Key();
-				Saver << It.Value().Offset;
-				Saver << It.Value().Size;
-				Saver << It.Value().Crc;
-				NumProcessed++;
-			}
-			check(NumProcessed == NumItems);
-		}
-		uint32 IndexCrc = FCrc::MemCrc_DEPRECATED(IndexBuffer.GetData(), IndexBuffer.Num());
-		uint32 SizeIndex = uint32(IndexBuffer.Num());
+	bool SaveCache();
 
-		uint32 Magic = PakCache_Magic;
-		TArray<uint8> Buffer;
-		FMemoryWriter Saver(Buffer);
-		Saver << Magic;
-		Saver << IndexCrc;
-		Saver << NumItems;
-		Saver << SizeIndex;
-		Saver.Serialize(IndexBuffer.GetData(), IndexBuffer.Num());
-		Saver << Magic;
-		Saver << IndexOffset;
-		FileHandle->Serialize(Buffer.GetData(), Buffer.Num());
-		CacheItems.Empty();
-		FileHandle.Reset();
-		bClosed = true;
-		return true;
-	}
 	/**
 	 * Load the cache to disk	 * @param	Filename	Filename to load
 	 * @return	true if file was loaded sucessfully
 	 */
-	bool LoadCache(const TCHAR* InFilename)
-	{
-		check(FileHandle);
-		check(FileHandle->IsLoading());
-		int64 FileSize = FileHandle->TotalSize();
-		check(FileSize >= 0);
-		if (FileSize < sizeof(int64) + sizeof(uint32) * 5)
-		{
-			UE_LOG(LogDerivedDataCache, Error, TEXT("Pak cache was corrputed (short) %s."), InFilename);
-			return false;
-		}
-		int64 IndexOffset = -1;
-		int64 Trailer = -1;
-		{
-			TArray<uint8> Buffer;
-			const int64 SeekPos = FileSize - int64(sizeof(int64) + sizeof(uint32));
-			FileHandle->Seek(SeekPos);
-			Trailer = FileHandle->Tell();
-			if (Trailer != SeekPos)
-			{
-				UE_LOG(LogDerivedDataCache, Error, TEXT("Pak cache was corrputed (bad seek) %s."), InFilename);
-				return false;
-			}
-			check(Trailer >= 0 && Trailer < FileSize);
-			Buffer.AddUninitialized(sizeof(int64) + sizeof(uint32));
-			FileHandle->Serialize(Buffer.GetData(), int64(sizeof(int64)+sizeof(uint32)));
-			FMemoryReader Loader(Buffer);
-			uint32 Magic = 0;
-			Loader << Magic;
-			Loader << IndexOffset;
-			if (Magic != PakCache_Magic || IndexOffset < 0 || IndexOffset + int64(sizeof(uint32) * 4) > Trailer)
-			{
-				UE_LOG(LogDerivedDataCache, Error, TEXT("Pak cache was corrputed (bad footer) %s."), InFilename);
-				return false;
-			}
-		}
-		uint32 IndexCrc = 0;
-		uint32 NumIndex = 0;
-		uint32 SizeIndex = 0;
-		{
-			TArray<uint8> Buffer;
-			FileHandle->Seek(IndexOffset);
-			if (FileHandle->Tell() != IndexOffset)
-			{
-				UE_LOG(LogDerivedDataCache, Error, TEXT("Pak cache was corrputed (bad seek index) %s."), InFilename);
-				return false;
-			}
-			Buffer.AddUninitialized(sizeof(uint32) * 4);
-			FileHandle->Serialize(Buffer.GetData(), int64(sizeof(uint32)* 4));
-			FMemoryReader Loader(Buffer);
-			uint32 Magic = 0;
-			Loader << Magic;
-			Loader << IndexCrc;
-			Loader << NumIndex;
-			Loader << SizeIndex;
-			if (Magic != PakCache_Magic || (SizeIndex != 0 && NumIndex == 0) || (SizeIndex == 0 && NumIndex != 0)) 
-			{
-				UE_LOG(LogDerivedDataCache, Error, TEXT("Pak cache was corrputed (bad index header) %s."), InFilename);
-				return false;
-			}
-			if (IndexOffset + sizeof(uint32) * 4 + SizeIndex != Trailer) 
-			{
-				UE_LOG(LogDerivedDataCache, Error, TEXT("Pak cache was corrputed (bad index size) %s."), InFilename);
-				return false;
-			}
-		}
-		{
-			TArray<uint8> Buffer;
-			Buffer.AddUninitialized(SizeIndex);
-			FileHandle->Serialize(Buffer.GetData(), int64(SizeIndex));
-			FMemoryReader Loader(Buffer);
-			while (Loader.Tell() < (int32)SizeIndex)
-			{
-				FString Key;
-				int64 Offset;
-				int64 Size;
-				uint32 Crc;
-				Loader << Key;
-				Loader << Offset;
-				Loader << Size;
-				Loader << Crc;
-				if (!Key.Len() || Offset < 0 || Offset >= IndexOffset || !Size)
-				{
-					UE_LOG(LogDerivedDataCache, Error, TEXT("Pak cache was corrputed (bad index entry) %s."), InFilename);
-					return false;
-				}
-				CacheItems.Add(Key, FCacheValue(Offset, Size, Crc));
-			}
-			if (CacheItems.Num() != NumIndex)
-			{
-				UE_LOG(LogDerivedDataCache, Error, TEXT("Pak cache was corrputed (bad index count) %s."), InFilename);
-				return false;
-			}
-		}
-		return true;
-	}
+	bool LoadCache(const TCHAR* InFilename);
 
 	/**
 	 * Merges another cache file into this one.
 	 * @return true on success
 	 */
-	void MergeCache(FPakFileDerivedDataBackend* OtherPak)
-	{
-		// Get all the existing keys
-		TArray<FString> KeyNames;
-		OtherPak->CacheItems.GenerateKeyArray(KeyNames);
-
-		// Find all the keys to copy
-		TArray<FString> CopyKeyNames;
-		for(const FString& KeyName : KeyNames)
-		{
-			if(!CachedDataProbablyExists(*KeyName))
-			{
-				CopyKeyNames.Add(KeyName);
-			}
-		}
-		UE_LOG(LogDerivedDataCache, Display, TEXT("Merging %d entries (%d skipped)."), CopyKeyNames.Num(), KeyNames.Num() - CopyKeyNames.Num());
-
-		// Copy them all to the new cache. Don't use the overloaded get/put methods (which may compress/decompress); copy the raw data directly.
-		TArray<uint8> Buffer;
-		for(const FString& CopyKeyName : CopyKeyNames)
-		{
-			Buffer.Reset();
-			if(OtherPak->FPakFileDerivedDataBackend::GetCachedData(*CopyKeyName, Buffer))
-			{
-				FPakFileDerivedDataBackend::PutCachedData(*CopyKeyName, Buffer, false);
-			}
-		}
-	}
+	void MergeCache(FPakFileDerivedDataBackend* OtherPak);
 	
 	const FString& GetFilename() const
 	{
 		return Filename;
 	}
 
-	static bool SortAndCopy(const FString &InputFilename, const FString &OutputFilename)
-	{
-		// Open the input and output files
-		FPakFileDerivedDataBackend InputPak(*InputFilename, false);
-		if (InputPak.bClosed) return false;
-
-		FPakFileDerivedDataBackend OutputPak(*OutputFilename, true);
-		if (OutputPak.bClosed) return false;
-
-		// Sort the key names
-		TArray<FString> KeyNames;
-		InputPak.CacheItems.GenerateKeyArray(KeyNames);
-		KeyNames.Sort();
+	static bool SortAndCopy(const FString &InputFilename, const FString &OutputFilename);
 
-		// Copy all the DDC to the new cache
-		TArray<uint8> Buffer;
-		TArray<uint32> KeySizes;
-		for (int KeyIndex = 0; KeyIndex < KeyNames.Num(); KeyIndex++)
-		{
-			Buffer.Reset();
-			InputPak.GetCachedData(*KeyNames[KeyIndex], Buffer);
-			OutputPak.PutCachedData(*KeyNames[KeyIndex], Buffer, false);
-			KeySizes.Add(Buffer.Num());
-		}
-
-		// Write out a TOC listing for debugging
-		FStringOutputDevice Output;
-		Output.Logf(TEXT("Asset,Size") LINE_TERMINATOR);
-		for(int KeyIndex = 0; KeyIndex < KeyNames.Num(); KeyIndex++)
-		{
-			Output.Logf(TEXT("%s,%d") LINE_TERMINATOR, *KeyNames[KeyIndex], KeySizes[KeyIndex]);
-		}
-		FFileHelper::SaveStringToFile(Output, *FPaths::Combine(*FPaths::GetPath(OutputFilename), *(FPaths::GetBaseFilename(OutputFilename) + TEXT(".csv"))));
-		return true;
-	}
-
-	virtual void GatherUsageStats(TMap<FString, FDerivedDataCacheUsageStats>& UsageStatsMap, FString&& GraphPath) override
-	{
-		COOK_STAT(UsageStatsMap.Add(FString::Printf(TEXT("%s: %s.%s"), *GraphPath, TEXT("PakFile"), *Filename), UsageStats));
-	}
+	virtual void GatherUsageStats(TMap<FString, FDerivedDataCacheUsageStats>& UsageStatsMap, FString&& GraphPath) override;
 
 private:
 	FDerivedDataCacheUsageStats UsageStats;
@@ -465,7 +102,7 @@ private:
 	/** Set of files that are being written to disk asynchronously. */
 	TMap<FString, FCacheValue> CacheItems;
 	/** File handle of pak. */
-	TAutoPtr<FArchive> FileHandle;
+	TUniquePtr<FArchive> FileHandle;
 	/** File name of pak. */
 	FString Filename;
 	enum 
@@ -478,41 +115,10 @@ private:
 class FCompressedPakFileDerivedDataBackend : public FPakFileDerivedDataBackend
 {
 public:
-	FCompressedPakFileDerivedDataBackend(const TCHAR* InFilename, bool bInWriting) 
-		: FPakFileDerivedDataBackend(InFilename, bInWriting)
-	{
-	}
+	FCompressedPakFileDerivedDataBackend(const TCHAR* InFilename, bool bInWriting);
 
-	virtual void PutCachedData(const TCHAR* CacheKey, TArray<uint8>& InData, bool bPutEvenIfExists) override
-	{
-		int32 UncompressedSize = InData.Num();
-		int32 CompressedSize = FCompression::CompressMemoryBound(CompressionFlags, UncompressedSize);
-
-		TArray<uint8> CompressedData;
-		CompressedData.AddUninitialized(CompressedSize + sizeof(UncompressedSize));
-
-		FMemory::Memcpy(&CompressedData[0], &UncompressedSize, sizeof(UncompressedSize));
-		verify(FCompression::CompressMemory(CompressionFlags, CompressedData.GetData() + sizeof(UncompressedSize), CompressedSize, InData.GetData(), InData.Num()));
-		CompressedData.SetNum(CompressedSize + sizeof(UncompressedSize), false);
-
-		FPakFileDerivedDataBackend::PutCachedData(CacheKey, CompressedData, bPutEvenIfExists);
-	}
-
-	virtual bool GetCachedData(const TCHAR* CacheKey, TArray<uint8>& OutData) override
-	{
-		TArray<uint8> CompressedData;
-		if(!FPakFileDerivedDataBackend::GetCachedData(CacheKey, CompressedData))
-		{
-			return false;
-		}
-
-		int32 UncompressedSize;
-		FMemory::Memcpy(&UncompressedSize, &CompressedData[0], sizeof(UncompressedSize));
-		OutData.SetNum(UncompressedSize);
-		verify(FCompression::UncompressMemory(CompressionFlags, OutData.GetData(), UncompressedSize, CompressedData.GetData() + sizeof(UncompressedSize), CompressedData.Num() - sizeof(UncompressedSize)));
-
-		return true;
-	}
+	virtual void PutCachedData(const TCHAR* CacheKey, TArray<uint8>& InData, bool bPutEvenIfExists) override;
+	virtual bool GetCachedData(const TCHAR* CacheKey, TArray<uint8>& OutData) override;
 
 private:
 	static const ECompressionFlags CompressionFlags = (ECompressionFlags)(COMPRESS_ZLIB | COMPRESS_BiasMemory);
diff --git a/Engine/Source/Developer/DesktopPlatform/Private/Windows/DesktopPlatformWindows.cpp b/Engine/Source/Developer/DesktopPlatform/Private/Windows/DesktopPlatformWindows.cpp
--- a/Engine/Source/Developer/DesktopPlatform/Private/Windows/DesktopPlatformWindows.cpp
+++ b/Engine/Source/Developer/DesktopPlatform/Private/Windows/DesktopPlatformWindows.cpp
@@ -595,13 +595,13 @@ void FDesktopPlatformWindows::GetRequiredRegistrySettings(TIndirectArray<FRegist
 
 	// HKLM\SOFTWARE\Classes\.uproject
 	FRegistryRootedKey *RootExtensionKey = new FRegistryRootedKey(HKEY_LOCAL_MACHINE, TEXT("SOFTWARE\\Classes\\.uproject"));
-	RootExtensionKey->Key = new FRegistryKey();
+	RootExtensionKey->Key = MakeUnique<FRegistryKey>();
 	RootExtensionKey->Key->SetValue(TEXT(""), TEXT("Unreal.ProjectFile"));
 	RootedKeys.Add(RootExtensionKey);
 
 	// HKLM\SOFTWARE\Classes\Unreal.ProjectFile
 	FRegistryRootedKey *RootFileTypeKey = new FRegistryRootedKey(HKEY_LOCAL_MACHINE, TEXT("SOFTWARE\\Classes\\Unreal.ProjectFile"));
-	RootFileTypeKey->Key = new FRegistryKey();
+	RootFileTypeKey->Key = MakeUnique<FRegistryKey>();
 	RootFileTypeKey->Key->SetValue(TEXT(""), TEXT("Unreal Engine Project File"));
 	RootFileTypeKey->Key->FindOrAddKey(L"DefaultIcon")->SetValue(TEXT(""), QuotedExecutableFileName);
 	RootedKeys.Add(RootFileTypeKey);
diff --git a/Engine/Source/Developer/DesktopPlatform/Private/Windows/WindowsRegistry.h b/Engine/Source/Developer/DesktopPlatform/Private/Windows/WindowsRegistry.h
--- a/Engine/Source/Developer/DesktopPlatform/Private/Windows/WindowsRegistry.h
+++ b/Engine/Source/Developer/DesktopPlatform/Private/Windows/WindowsRegistry.h
@@ -2,6 +2,8 @@
 
 #pragma once
 
+#include "UniquePtr.h"
+
 struct FRegistryValue
 {
 	uint32 Type;
@@ -45,7 +47,7 @@ struct FRegistryRootedKey
 {
 	HKEY hRootKey;
 	FString Path;
-	TScopedPointer<FRegistryKey> Key;
+	TUniquePtr<FRegistryKey> Key;
 
 	FRegistryRootedKey(HKEY hInKeyRoot, const FString &InPath);
 
diff --git a/Engine/Source/Developer/GameplayDebugger/Private/GameplayDebuggerLocalController.cpp b/Engine/Source/Developer/GameplayDebugger/Private/GameplayDebuggerLocalController.cpp
--- a/Engine/Source/Developer/GameplayDebugger/Private/GameplayDebuggerLocalController.cpp
+++ b/Engine/Source/Developer/GameplayDebugger/Private/GameplayDebuggerLocalController.cpp
@@ -609,7 +609,7 @@ void UGameplayDebuggerLocalController::OnSelectActorTick()
 		const FVector ViewDir = CameraRotation.Vector();
 		for (FConstPawnIterator It = OwnerPC->GetWorld()->GetPawnIterator(); It; ++It)
 		{
-			APawn* TestPawn = *It;
+			APawn* TestPawn = It->Get();
 			if (TestPawn && !TestPawn->bHidden && TestPawn->GetActorEnableCollision() && TestPawn != OwnerPC->GetPawn())
 			{
 				FVector DirToPawn = (TestPawn->GetActorLocation() - CameraLocation);
diff --git a/Engine/Source/Developer/GameplayDebugger/Private/GameplayDebuggerPlayerManager.cpp b/Engine/Source/Developer/GameplayDebugger/Private/GameplayDebuggerPlayerManager.cpp
--- a/Engine/Source/Developer/GameplayDebugger/Private/GameplayDebuggerPlayerManager.cpp
+++ b/Engine/Source/Developer/GameplayDebugger/Private/GameplayDebuggerPlayerManager.cpp
@@ -103,7 +103,7 @@ void AGameplayDebuggerPlayerManager::UpdateAuthReplicators()
 
 	for (FConstPlayerControllerIterator It = World->GetPlayerControllerIterator(); It; It++)
 	{
-		APlayerController* TestPC = *It;
+		APlayerController* TestPC = It->Get();
 		if (TestPC && !TestPC->IsA<ADebugCameraController>())
 		{
 			const bool bNeedsReplicator = (GetReplicator(*TestPC) == nullptr);
diff --git a/Engine/Source/Developer/HierarchicalLODUtilities/Private/HierarchicalLODUtilities.cpp b/Engine/Source/Developer/HierarchicalLODUtilities/Private/HierarchicalLODUtilities.cpp
--- a/Engine/Source/Developer/HierarchicalLODUtilities/Private/HierarchicalLODUtilities.cpp
+++ b/Engine/Source/Developer/HierarchicalLODUtilities/Private/HierarchicalLODUtilities.cpp
@@ -712,7 +712,7 @@ int32 FHierarchicalLODUtilities::ComputeStaticMeshLODLevel(const TArray<FStaticM
 int32 FHierarchicalLODUtilities::GetLODLevelForScreenSize(const UStaticMeshComponent* StaticMeshComponent, const float ScreenSize)
 {
 	check(StaticMeshComponent != nullptr);
-	const FStaticMeshRenderData* RenderData = StaticMeshComponent->GetStaticMesh()->RenderData.GetOwnedPointer();
+	const FStaticMeshRenderData* RenderData = StaticMeshComponent->GetStaticMesh()->RenderData.Get();
 	checkf(RenderData != nullptr, TEXT("StaticMesh in StaticMeshComponent %s contains invalid render data"), *StaticMeshComponent->GetName());
 	checkf(StaticMeshComponent->GetStaticMesh()->SourceModels.Num() > 0, TEXT("StaticMesh in StaticMeshComponent %s contains no SourceModels"), *StaticMeshComponent->GetName());
 
