Fix a DDC issue where an early abort due to lack of memory cache space causes the disk not to be searched

This manifested itself in a game when textures ended up rebuilding every single run of the game but suddenly stopped building and coming from the DDC as they should have after the first run.

When a texture was being requested from the DDC, the in-memory DDC cache space ran out, and FMemoryDerivedDataBackend::CachedDataProbablyExists() returned true. FDerivedDataBackendAsyncPutWrapper::PutCachedData() assumes the data is already on its way, so it doesn't send it again and exits the function. Unfortunately, the data is not really on the disk, and FCachePutAsyncWorker never gets a chance to put it there.

Because of changing memory requirements from run to run, this game was eventually able to write all of the texture data to disk, but it took dozens of runs to do so, as it generally would only write a single mip from a mipchain in any given run. When all of the mips were finally written, the texture would be fully retrieved from the DDC, and no build would be necessary.

With this fix, no early abort is had, and all textures write themselves fully to the disk.

diff --git a/Engine/Source/Developer/DerivedDataCache/Private/MemoryDerivedDataBackend.cpp b/Engine/Source/Developer/DerivedDataCache/Private/MemoryDerivedDataBackend.cpp
--- a/Engine/Source/Developer/DerivedDataCache/Private/MemoryDerivedDataBackend.cpp
+++ b/Engine/Source/Developer/DerivedDataCache/Private/MemoryDerivedDataBackend.cpp
@@ -30,12 +30,6 @@ bool FMemoryDerivedDataBackend::CachedDataProbablyExists(const TCHAR* CacheKey)
 	{
 		return false;
 	}
-	// to avoid constant error reporting in async put due to restricted cache size, 
-	// we report true if the max size has been exceeded
-	if (bMaxSizeExceeded)
-	{
-		return true;
-	}
 
 	bool Result = CacheItems.Contains(FString(CacheKey));
 	if (Result)
