Exposed hardcoded DDCCleanup settings to ini

#rnx

diff --git a/Engine/Config/BaseEngine.ini b/Engine/Config/BaseEngine.ini
--- a/Engine/Config/BaseEngine.ini
+++ b/Engine/Config/BaseEngine.ini
@@ -1622,6 +1622,11 @@ UnusedFileAge=34
 MaxFileChecksPerSec=-1 ;no limit
 Path=%GAMEDIR%DerivedDataCache/VT
 
+[DDCCleanup]
+TimeToWaitAfterInit=120
+TimeBetweenDeleteingDirectories=5
+TimeBetweenDeletingFiles=2
+
 [/Script/Engine.LocalPlayer]
 AspectRatioAxisConstraint=AspectRatio_MaintainXFOV
 
diff --git a/Engine/Source/Developer/DerivedDataCache/Private/DDCCleanup.cpp b/Engine/Source/Developer/DerivedDataCache/Private/DDCCleanup.cpp
--- a/Engine/Source/Developer/DerivedDataCache/Private/DDCCleanup.cpp
+++ b/Engine/Source/Developer/DerivedDataCache/Private/DDCCleanup.cpp
@@ -58,7 +58,15 @@ FDDCCleanup::FDDCCleanup()
 	: Thread(NULL)
 	, StopTaskCounter(0)
 	, bDontWaitBetweenDeletes(false)
+	, TimeToWaitAfterInit(120.0f)
+	, TimeBetweenDeleteingDirectories(5.0f)
+	, TimeBetweenDeletingFiles(2.0f)
 {
+	check(GConfig);
+	GConfig->GetFloat(TEXT("DDCCleanup"), TEXT("TimeToWaitAfterInit"), TimeToWaitAfterInit, GEngineIni);
+	GConfig->GetFloat(TEXT("DDCCleanup"), TEXT("TimeBetweenDeleteingDirectories"), TimeBetweenDeleteingDirectories, GEngineIni);
+	GConfig->GetFloat(TEXT("DDCCleanup"), TEXT("TimeBetweenDeletingFiles"), TimeBetweenDeletingFiles, GEngineIni);
+
 	// Don't delete the runnable automatically. It's going to be manually deleted in FDDCCleanup::Shutdown.
 	Thread = FRunnableThread::Create(this, TEXT("FDDCCleanup"), 0, TPri_BelowNormal, FPlatformAffinity::GetPoolThreadMask());
 }
@@ -79,14 +87,14 @@ void FDDCCleanup::Wait( const float InSeconds, const float InSleepTime )
 }
 
 bool FDDCCleanup::Init() 
-{
+{	
 	return true;
 }
 
 uint32 FDDCCleanup::Run()
 {
 	// Give up some time to the engine to start up and load everything
-	Wait( 120.0f, 0.5f );
+	Wait( TimeToWaitAfterInit, 0.5f );
 
 	int32 FilesystemToCleanup = 0;
 	// Check one directory every 5 seconds
@@ -107,7 +115,7 @@ uint32 FDDCCleanup::Run()
 		{
 			CleanupFilesystemDirectory( FilesystemInfo );
 		}
-		Wait( 5.0f );
+		Wait( TimeBetweenDeleteingDirectories );
 	}
 	while(ShouldStop() == false && CleanupList.Num() > 0);
 
@@ -174,7 +182,7 @@ bool FDDCCleanup::CleanupFilesystemDirectory( TSharedPtr< FFilesystemInfo > File
 			if( !bDontWaitBetweenDeletes && ++NumFilesChecked >= FilesystemInfo->MaxContinuousFileChecks && FilesystemInfo->MaxContinuousFileChecks > 0 && ShouldStop() == false )
 			{
 				NumFilesChecked = 0;
-				Wait( 1.0f );
+				Wait( TimeBetweenDeletingFiles );
 			}
 			else
 			{
diff --git a/Engine/Source/Developer/DerivedDataCache/Public/DDCCleanup.h b/Engine/Source/Developer/DerivedDataCache/Public/DDCCleanup.h
--- a/Engine/Source/Developer/DerivedDataCache/Public/DDCCleanup.h
+++ b/Engine/Source/Developer/DerivedDataCache/Public/DDCCleanup.h
@@ -27,6 +27,13 @@ class DERIVEDDATACACHE_API FDDCCleanup : public FRunnable
 	/** If true, work without giving up time to other threads */
 	bool bDontWaitBetweenDeletes;
 
+	/** Time before any deleting starts */
+	float TimeToWaitAfterInit;
+	/** Time to wait before starting deleting the next DDC directory files */
+	float TimeBetweenDeleteingDirectories;
+	/** Time to wait before deleting the next DDC file */
+	float TimeBetweenDeletingFiles;
+
 	/** Constructor */
 	FDDCCleanup();
 
