CIS guess fix

diff --git a/Engine/Source/Developer/DerivedDataCache/Private/DDCCleanup.cpp b/Engine/Source/Developer/DerivedDataCache/Private/DDCCleanup.cpp
--- a/Engine/Source/Developer/DerivedDataCache/Private/DDCCleanup.cpp
+++ b/Engine/Source/Developer/DerivedDataCache/Private/DDCCleanup.cpp
@@ -6,7 +6,7 @@
 #include "HAL/RunnableThread.h"
 #include "Misc/ScopeLock.h"
 #include "Math/RandomStream.h"
-
+#include "Misc/ConfigCacheIni.h"
 #include "DerivedDataBackendInterface.h"
 
 /** Struct containing a list of directories to cleanup. */
