Added DDCCleanup commandlet that iterates over shared DDC directories and removes old (unused) cache files

diff --git a/Engine/Source/Developer/DerivedDataCache/Private/DDCCleanup.cpp b/Engine/Source/Developer/DerivedDataCache/Private/DDCCleanup.cpp
--- a/Engine/Source/Developer/DerivedDataCache/Private/DDCCleanup.cpp
+++ b/Engine/Source/Developer/DerivedDataCache/Private/DDCCleanup.cpp
@@ -57,6 +57,7 @@ FDDCCleanup* FDDCCleanup::Runnable = NULL;
 FDDCCleanup::FDDCCleanup()
 	: Thread(NULL)
 	, StopTaskCounter(0)
+	, bDontWaitBetweenDeletes(false)
 {
 	// Don't delete the runnable automatically. It's going to be manually deleted in FDDCCleanup::Shutdown.
 	Thread = FRunnableThread::Create(this, TEXT("FDDCCleanup"), 0, TPri_BelowNormal, FPlatformAffinity::GetPoolThreadMask());
@@ -71,9 +72,9 @@ void FDDCCleanup::Wait( const float InSeconds, const float InSleepTime )
 {
 	// Instead of waiting the given amount of seconds doing nothing
 	// check periodically if there's been any Stop requests.
-	for( float TimeToWait = InSeconds; TimeToWait > 0.0f && ShouldStop() == false; TimeToWait -= InSleepTime )
+	for( float TimeToWait = InSeconds; TimeToWait > 0.0f && ShouldStop() == false && !bDontWaitBetweenDeletes; TimeToWait -= InSleepTime )
 	{
-		FPlatformProcess::Sleep( FMath::Min(InSleepTime, TimeToWait) );
+		FPlatformProcess::SleepNoStats( FMath::Min(InSleepTime, TimeToWait) );
 	}
 }
 
@@ -108,7 +109,7 @@ uint32 FDDCCleanup::Run()
 		}
 		Wait( 5.0f );
 	}
-	while( ShouldStop() == false );
+	while(ShouldStop() == false && CleanupList.Num() > 0);
 
 	return 0;
 }
@@ -134,7 +135,7 @@ bool FDDCCleanup::CleanupFilesystemDirectory( TSharedPtr< FFilesystemInfo > File
 			CleanupList.Remove( FilesystemInfo );
 			FilesystemInfo.Reset();
 		}
-		else if( ++FilesystemInfo->FoldersChecked >= FilesystemInfo->MaxNumFoldersToCheck && FilesystemInfo->MaxNumFoldersToCheck > 0 )
+		else if( !bDontWaitBetweenDeletes && ++FilesystemInfo->FoldersChecked >= FilesystemInfo->MaxNumFoldersToCheck && FilesystemInfo->MaxNumFoldersToCheck > 0 )
 		{
 			// Remove the filesystem but keep checking the current folder
 			FScopeLock ScopeLock( &DataLock );
@@ -158,11 +159,19 @@ bool FDDCCleanup::CleanupFilesystemDirectory( TSharedPtr< FFilesystemInfo > File
 				if( TimeSinceLastAccess >= FilesystemInfo->UnusedFileTime && TimeSinceLastModification >= FilesystemInfo->UnusedFileTime )
 				{
 					// Delete the file
-					bool Result = IFileManager::Get().Delete( *FileNames[ FileIndex ], false, true, true );
+					bool bResult = IFileManager::Get().Delete( *FileNames[ FileIndex ], false, true, true );
+					if (bResult)
+					{
+						UE_LOG(LogDerivedDataCache, VeryVerbose, TEXT("Deleted %s"), *FileNames[FileIndex]);
+					}
+					else
+					{
+						UE_LOG(LogDerivedDataCache, VeryVerbose, TEXT("Failed to delete %s"), *FileNames[FileIndex]);
+					}
 				}
 			}
 
-			if( ++NumFilesChecked >= FilesystemInfo->MaxContinuousFileChecks && FilesystemInfo->MaxContinuousFileChecks > 0 && ShouldStop() == false )
+			if( !bDontWaitBetweenDeletes && ++NumFilesChecked >= FilesystemInfo->MaxContinuousFileChecks && FilesystemInfo->MaxContinuousFileChecks > 0 && ShouldStop() == false )
 			{
 				NumFilesChecked = 0;
 				Wait( 1.0f );
@@ -177,7 +186,7 @@ bool FDDCCleanup::CleanupFilesystemDirectory( TSharedPtr< FFilesystemInfo > File
 		bCleanedUp = true;			
 	}
 
-	UE_LOG(LogDerivedDataCache, VeryVerbose, TEXT("DDC Folder Cleanup (%s) took %.4lfs."), *FilesystemInfo->CachePath, FPlatformTime::Seconds() - StartTime);
+	UE_CLOG(FilesystemInfo.IsValid(), LogDerivedDataCache, VeryVerbose, TEXT("DDC Folder Cleanup (%s) took %.4lfs."), *FilesystemInfo->CachePath, FPlatformTime::Seconds() - StartTime);
 
 	return bCleanedUp;
 }
@@ -218,7 +227,7 @@ void FDDCCleanup::Shutdown()
 	{
 		Runnable->EnsureCompletion();
 		delete Runnable;
-		Runnable = NULL;
+		Runnable = nullptr;
 	}
 }
 
diff --git a/Engine/Source/Developer/DerivedDataCache/Private/DDCCleanup.h b/Engine/Source/Developer/DerivedDataCache/Private/DDCCleanup.h
deleted file mode 100644
diff --git a/Engine/Source/Developer/DerivedDataCache/Public/DDCCleanup.h b/Engine/Source/Developer/DerivedDataCache/Public/DDCCleanup.h
new file mode 100644
--- /dev/null
+++ b/Engine/Source/Developer/DerivedDataCache/Public/DDCCleanup.h
@@ -0,0 +1,102 @@
+// Copyright 1998-2019 Epic Games, Inc. All Rights Reserved.
+
+#pragma once
+
+#include "CoreMinimal.h"
+#include "HAL/ThreadSafeCounter.h"
+#include "HAL/Runnable.h"
+
+struct FFilesystemInfo;
+
+/** 
+ * DDC Filesystem Cache cleanup thread.
+ */
+class DERIVEDDATACACHE_API FDDCCleanup : public FRunnable
+{
+	/** Singleton instance */
+	static FDDCCleanup* Runnable;
+	
+	/** Thread to run the cleanup FRunnable on */
+	FRunnableThread* Thread;
+	/** > 0 if we've been asked to abort work in progress at the next opportunity */
+	FThreadSafeCounter StopTaskCounter;
+	/** List of filesystems to clean up */
+	TArray< TSharedPtr< struct FFilesystemInfo > > CleanupList;
+	/** Synchronization object */
+	FCriticalSection DataLock;
+	/** If true, work without giving up time to other threads */
+	bool bDontWaitBetweenDeletes;
+
+	/** Constructor */
+	FDDCCleanup();
+
+	/** Destructor */
+	virtual ~FDDCCleanup();
+
+	//~ Begin FRunnable Interface.
+	virtual bool Init();
+	virtual uint32 Run();
+	virtual void Stop();
+	virtual void Exit();
+	//~ End FRunnable Interface
+
+	/** Checks if there's been any Stop requests */
+	FORCEINLINE bool ShouldStop() const
+	{
+		return StopTaskCounter.GetValue() > 0;
+	}
+
+	/**
+	 * Waits for a given amount of time periodically checking if there's been any Stop requests.
+	 *
+	 * @param InSeconds time in seconds to wait.
+	 * @param InSleepTime interval at which to check for Stop requests.
+	 */
+	void Wait( const float InSeconds, const float InSleepTime = 0.1f );
+
+	/**
+	 * Performs directory cleanup.
+	 *
+	 * @param FilesystemInfo Filesystem to cleanup.
+	 * @return true if the filesystem contained any directories to clean up, false otherwise.
+	 */
+	bool CleanupFilesystemDirectory( TSharedPtr< FFilesystemInfo > FilesystemInfo );
+
+	/** Makes sure this thread has stopped properly */
+	void EnsureCompletion();
+
+public:
+
+	/**
+	 * Adds DDC filesystem to clean up.
+	 *
+	 * @param InCachePath Filesystem path.
+	 * @param InDaysToDelete Number of days since last access time to consider a file as unused.
+	 * @param InMaxContinuousFileChecks Number of files to check without pausing.
+	 */
+	void AddFilesystem( FString& InCachePath, int32 InDaysToDelete, int32 InMaxNumFoldersToCheck, int32 InMaxContinuousFileChecks );
+
+	/** Gets DDC Cleanup singleton instance, crates one if it doesn't exist already */
+	static FDDCCleanup* Get();
+
+	/** Gets DDC Cleanup singleton instance */
+	static FDDCCleanup* GetNoInit()
+	{
+		return Runnable;
+	}
+
+	/** Shuts down DDC Cleanup thread. */
+	static void Shutdown();
+
+	/** Sets whether cleanup should give up time to other threads between deletes */
+	void WaitBetweenDeletes(bool bWait)
+	{
+		bDontWaitBetweenDeletes = !bWait;
+	}
+
+	/** Checks if the cleanup thread is done deleting files */
+	bool IsFinished() const
+	{
+		return !CleanupList.Num();
+	}
+};
diff --git a/Engine/Source/Editor/UnrealEd/Classes/Commandlets/DDCCleanupCommandlet.h b/Engine/Source/Editor/UnrealEd/Classes/Commandlets/DDCCleanupCommandlet.h
new file mode 100644
--- /dev/null
+++ b/Engine/Source/Editor/UnrealEd/Classes/Commandlets/DDCCleanupCommandlet.h
@@ -0,0 +1,19 @@
+// Copyright 1998-2019 Epic Games, Inc. All Rights Reserved.
+
+#pragma once
+
+#include "CoreMinimal.h"
+#include "UObject/ObjectMacros.h"
+#include "Commandlets/Commandlet.h"
+#include "DDCCleanupCommandlet.generated.h"
+
+UCLASS()
+class UDDCCleanupCommandlet : public UCommandlet
+{
+	GENERATED_BODY()
+	//~ Begin UCommandlet Interface
+	virtual int32 Main(const FString& Params) override;
+	//~ End UCommandlet Interface
+};
+
+
diff --git a/Engine/Source/Editor/UnrealEd/Private/Commandlets/DDCCleanupCommandlet.cpp b/Engine/Source/Editor/UnrealEd/Private/Commandlets/DDCCleanupCommandlet.cpp
new file mode 100644
--- /dev/null
+++ b/Engine/Source/Editor/UnrealEd/Private/Commandlets/DDCCleanupCommandlet.cpp
@@ -0,0 +1,21 @@
+// Copyright 1998-2019 Epic Games, Inc. All Rights Reserved.
+
+#include "Commandlets/DDCCleanupCommandlet.h"
+#include "DDCCleanup.h"
+#include "HAL/PlatformProcess.h"
+
+int32 UDDCCleanupCommandlet::Main(const FString& Params)
+{
+	if (FDDCCleanup::Get())
+	{
+		FDDCCleanup::GetNoInit()->WaitBetweenDeletes(false);
+		do
+		{			
+			// Cleanup works from its own thread so just wait until it's done and flush logs once in a while
+			FPlatformProcess::SleepNoStats(1.0f);
+			GLog->Flush();
+		} while (FDDCCleanup::GetNoInit() && !FDDCCleanup::GetNoInit()->IsFinished());
+	}
+
+	return 0;
+}
diff --git a/Engine/Source/Editor/UnrealEd/UnrealEd.Build.cs b/Engine/Source/Editor/UnrealEd/UnrealEd.Build.cs
--- a/Engine/Source/Editor/UnrealEd/UnrealEd.Build.cs
+++ b/Engine/Source/Editor/UnrealEd/UnrealEd.Build.cs
@@ -167,7 +167,8 @@ public class UnrealEd : ModuleRules
 				"PIEPreviewDeviceProfileSelector",
 				"PakFileUtilities",
 				"TimeManagement",
-			}
+                "DerivedDataCache",
+            }
 		);
 
 		DynamicallyLoadedModuleNames.AddRange(
