#UE4 Add a little more info to a warning about failing to write a file to a filesystem DDC backend

diff --git a/Engine/Source/Developer/DerivedDataCache/Private/FileSystemDerivedDataBackend.cpp b/Engine/Source/Developer/DerivedDataCache/Private/FileSystemDerivedDataBackend.cpp
--- a/Engine/Source/Developer/DerivedDataCache/Private/FileSystemDerivedDataBackend.cpp
+++ b/Engine/Source/Developer/DerivedDataCache/Private/FileSystemDerivedDataBackend.cpp
@@ -274,7 +274,9 @@ public:
 				}
 				else
 				{
-					UE_LOG(LogDerivedDataCache, Warning, TEXT("FFileSystemDerivedDataBackend: Could not write temp file %s!"),*TempFilename);
+					TCHAR ErrorBuffer[1024];
+					FPlatformMisc::GetSystemErrorMessage(ErrorBuffer, 1024, 0);
+					UE_LOG(LogDerivedDataCache, Warning, TEXT("FFileSystemDerivedDataBackend: Could not write temp file %s! Error: %s"), *TempFilename, ErrorBuffer);
 				}
 				// if everything worked, this is not necessary, but we will make every effort to avoid leaving junk in the cache
 				if (FPaths::FileExists(TempFilename))
