#UE4 Also including the error code in DDC write failure log

diff --git a/Engine/Source/Developer/DerivedDataCache/Private/FileSystemDerivedDataBackend.cpp b/Engine/Source/Developer/DerivedDataCache/Private/FileSystemDerivedDataBackend.cpp
--- a/Engine/Source/Developer/DerivedDataCache/Private/FileSystemDerivedDataBackend.cpp
+++ b/Engine/Source/Developer/DerivedDataCache/Private/FileSystemDerivedDataBackend.cpp
@@ -274,9 +274,10 @@ public:
 				}
 				else
 				{
+					uint32 ErrorCode = FPlatformMisc::GetLastError();
 					TCHAR ErrorBuffer[1024];
-					FPlatformMisc::GetSystemErrorMessage(ErrorBuffer, 1024, 0);
-					UE_LOG(LogDerivedDataCache, Warning, TEXT("FFileSystemDerivedDataBackend: Could not write temp file %s! Error: %s"), *TempFilename, ErrorBuffer);
+					FPlatformMisc::GetSystemErrorMessage(ErrorBuffer, 1024, ErrorCode);
+					UE_LOG(LogDerivedDataCache, Warning, TEXT("FFileSystemDerivedDataBackend: Could not write temp file %s! Error: %d (%s)"), *TempFilename, ErrorCode, ErrorBuffer);
 				}
 				// if everything worked, this is not necessary, but we will make every effort to avoid leaving junk in the cache
 				if (FPaths::FileExists(TempFilename))
