#UE4 Added some log info to another warning about writing to DDC folders.

diff --git a/Engine/Source/Developer/DerivedDataCache/Private/FileSystemDerivedDataBackend.cpp b/Engine/Source/Developer/DerivedDataCache/Private/FileSystemDerivedDataBackend.cpp
--- a/Engine/Source/Developer/DerivedDataCache/Private/FileSystemDerivedDataBackend.cpp
+++ b/Engine/Source/Developer/DerivedDataCache/Private/FileSystemDerivedDataBackend.cpp
@@ -76,7 +76,10 @@ public:
 			int32 TestFileSize = IFileManager::Get().FileSize(*TempFilename);
 			if (TestFileSize < 4)
 			{
-				UE_LOG(LogDerivedDataCache, Warning, TEXT("Fail to write to %s, derived data cache to this directory will be read only."),*CachePath);
+				uint32 ErrorCode = FPlatformMisc::GetLastError();
+				TCHAR ErrorBuffer[1024];
+				FPlatformMisc::GetSystemErrorMessage(ErrorBuffer, 1024, ErrorCode);
+				UE_LOG(LogDerivedDataCache, Warning, TEXT("Fail to write to %s, derived data cache to this directory will be read only. Error: %u (%s)"), *CachePath, ErrorCode, ErrorBuffer);
 			}
 			else
 			{
@@ -277,7 +280,7 @@ public:
 					uint32 ErrorCode = FPlatformMisc::GetLastError();
 					TCHAR ErrorBuffer[1024];
 					FPlatformMisc::GetSystemErrorMessage(ErrorBuffer, 1024, ErrorCode);
-					UE_LOG(LogDerivedDataCache, Warning, TEXT("FFileSystemDerivedDataBackend: Could not write temp file %s! Error: %d (%s)"), *TempFilename, ErrorCode, ErrorBuffer);
+					UE_LOG(LogDerivedDataCache, Warning, TEXT("FFileSystemDerivedDataBackend: Could not write temp file %s! Error: %u (%s)"), *TempFilename, ErrorCode, ErrorBuffer);
 				}
 				// if everything worked, this is not necessary, but we will make every effort to avoid leaving junk in the cache
 				if (FPaths::FileExists(TempFilename))
