#UE4 Adding even more info to a warning about failing to write to the DDC

diff --git a/Engine/Source/Developer/DerivedDataCache/Private/FileSystemDerivedDataBackend.cpp b/Engine/Source/Developer/DerivedDataCache/Private/FileSystemDerivedDataBackend.cpp
--- a/Engine/Source/Developer/DerivedDataCache/Private/FileSystemDerivedDataBackend.cpp
+++ b/Engine/Source/Developer/DerivedDataCache/Private/FileSystemDerivedDataBackend.cpp
@@ -73,13 +73,16 @@ public:
 			double TestStart = FPlatformTime::Seconds();
 			FString TempFilename = CachePath / FGuid::NewGuid().ToString() + ".tmp";
 			FFileHelper::SaveStringToFile(FString("TEST"),*TempFilename);
+			uint32 WriteErrorCode = FPlatformMisc::GetLastError();
 			int32 TestFileSize = IFileManager::Get().FileSize(*TempFilename);
 			if (TestFileSize < 4)
 			{
-				uint32 ErrorCode = FPlatformMisc::GetLastError();
-				TCHAR ErrorBuffer[1024];
-				FPlatformMisc::GetSystemErrorMessage(ErrorBuffer, 1024, ErrorCode);
-				UE_LOG(LogDerivedDataCache, Warning, TEXT("Fail to write to %s, derived data cache to this directory will be read only. Error: %u (%s)"), *CachePath, ErrorCode, ErrorBuffer);
+				uint32 ReadErrorCode = FPlatformMisc::GetLastError();
+				TCHAR WriteErrorBuffer[1024];
+				TCHAR ReadErrorBuffer[1024];
+				FPlatformMisc::GetSystemErrorMessage(WriteErrorBuffer, 1024, WriteErrorCode);
+				FPlatformMisc::GetSystemErrorMessage(ReadErrorBuffer, 1024, ReadErrorCode);
+				UE_LOG(LogDerivedDataCache, Warning, TEXT("Fail to write to %s, derived data cache to this directory will be read only. WriteError: %u (%s) ReadError: %u (%s)"), *CachePath, WriteErrorCode, WriteErrorBuffer, ReadErrorCode, ReadErrorBuffer);
 			}
 			else
 			{
