Move derived data cache async tasks to GLargeThreadPool when in editor
Niagara is spawning a lot of async tasks to compile scripts, which saturates GThreadPool, then
* Texture streaming is blocked and no mipmaps can be loaded until compilation finishes
* (Critical) D3D12's async deferred deletion is blocked and we run out of vmem and crash
This change moves Niagara to GLargeThreadPool where editor slow tasks such as lightmap encoding & texture DDC generation are

diff --git a/Engine/Source/Developer/DerivedDataCache/Private/DerivedDataCache.cpp b/Engine/Source/Developer/DerivedDataCache/Private/DerivedDataCache.cpp
--- a/Engine/Source/Developer/DerivedDataCache/Private/DerivedDataCache.cpp
+++ b/Engine/Source/Developer/DerivedDataCache/Private/DerivedDataCache.cpp
@@ -306,7 +306,11 @@ public:
 		AddToAsyncCompletionCounter(1);
 		if (!bSync)
 		{
+#if WITH_EDITOR
+			AsyncTask->StartBackgroundTask(GLargeThreadPool);
+#else
 			AsyncTask->StartBackgroundTask();
+#endif
 		}
 		else
 		{
