Revert CL 6596039 to fix Niagara stack overflow

diff --git a/Engine/Source/Developer/DerivedDataCache/Private/DerivedDataCache.cpp b/Engine/Source/Developer/DerivedDataCache/Private/DerivedDataCache.cpp
--- a/Engine/Source/Developer/DerivedDataCache/Private/DerivedDataCache.cpp
+++ b/Engine/Source/Developer/DerivedDataCache/Private/DerivedDataCache.cpp
@@ -319,11 +319,7 @@ public:
 		AddToAsyncCompletionCounter(1);
 		if (!bSync)
 		{
-#if WITH_EDITOR
-			AsyncTask->StartBackgroundTask(GLargeThreadPool);
-#else
 			AsyncTask->StartBackgroundTask();
-#endif
 		}
 		else
 		{
