Fixed a hang on editor shutdown due to long file system scan by DDC

diff --git a/Engine/Source/Developer/DerivedDataCache/Private/DDCCleanup.cpp b/Engine/Source/Developer/DerivedDataCache/Private/DDCCleanup.cpp
--- a/Engine/Source/Developer/DerivedDataCache/Private/DDCCleanup.cpp
+++ b/Engine/Source/Developer/DerivedDataCache/Private/DDCCleanup.cpp
@@ -122,6 +122,11 @@ uint32 FDDCCleanup::Run()
 	return 0;
 }
 
+bool FDDCCleanup::ShouldStop() const
+{
+	return StopTaskCounter.GetValue() > 0 || GIsRequestingExit;
+}
+
 bool FDDCCleanup::CleanupFilesystemDirectory( TSharedPtr< FFilesystemInfo > FilesystemInfo )
 {
 	bool bCleanedUp = false;
@@ -135,7 +140,15 @@ bool FDDCCleanup::CleanupFilesystemDirectory( TSharedPtr< FFilesystemInfo > File
 		const int32 DirectoryIndex = FilesystemInfo->CacheDirectories.Pop();
 		const FString DirectoryPath( FilesystemInfo->CachePath / FString::Printf(TEXT("%1d/%1d/%1d/"),(DirectoryIndex/100)%10,(DirectoryIndex/10)%10,DirectoryIndex%10) );
 
-		IFileManager::Get().FindFilesRecursive( FileNames, *DirectoryPath, TEXT("*.*"), true, false );
+		IFileManager::Get().IterateDirectoryRecursively(*DirectoryPath, [this, &FileNames](const TCHAR* InFilenameOrDirectory, const bool InIsDirectory) -> bool
+		{
+			if (!InIsDirectory)
+			{
+				FileNames.Emplace(FString(InFilenameOrDirectory));
+			}
+			return !ShouldStop();
+		});
+
 		if ( FilesystemInfo->CacheDirectories.Num() == 0 )
 		{
 			// Remove the filesystem and stop checking it
diff --git a/Engine/Source/Developer/DerivedDataCache/Public/DDCCleanup.h b/Engine/Source/Developer/DerivedDataCache/Public/DDCCleanup.h
--- a/Engine/Source/Developer/DerivedDataCache/Public/DDCCleanup.h
+++ b/Engine/Source/Developer/DerivedDataCache/Public/DDCCleanup.h
@@ -48,10 +48,7 @@ class DERIVEDDATACACHE_API FDDCCleanup : public FRunnable
 	//~ End FRunnable Interface
 
 	/** Checks if there's been any Stop requests */
-	FORCEINLINE bool ShouldStop() const
-	{
-		return StopTaskCounter.GetValue() > 0;
-	}
+	bool ShouldStop() const;
 
 	/**
 	 * Waits for a given amount of time periodically checking if there's been any Stop requests.
