* CookOnTheFlyServer trace instrumentation
* Selective DerivedDataCache instrumentation

(peafour-cherrypick of 6543978 by [REDACTED_PERSON])

diff --git a/Engine/Source/Developer/DerivedDataCache/Private/DerivedDataCache.cpp b/Engine/Source/Developer/DerivedDataCache/Private/DerivedDataCache.cpp
--- a/Engine/Source/Developer/DerivedDataCache/Private/DerivedDataCache.cpp
+++ b/Engine/Source/Developer/DerivedDataCache/Private/DerivedDataCache.cpp
@@ -29,6 +29,10 @@ DEFINE_STAT(STAT_DDC_PutTime);
 DEFINE_STAT(STAT_DDC_SyncBuildTime);
 DEFINE_STAT(STAT_DDC_ExistTime);
 
+//#define DDC_SCOPE_CYCLE_COUNTER(x) QUICK_SCOPE_CYCLE_COUNTER(STAT_ ## x)
+#define DDC_SCOPE_CYCLE_COUNTER(x) TRACE_CPUPROFILER_EVENT_SCOPE_TEXT(TEXT(#x));
+
+
 #if ENABLE_COOK_STATS
 #include "DerivedDataCacheUsageStats.h"
 namespace DerivedDataCacheCookStats
@@ -148,9 +152,13 @@ class FDerivedDataCache : public FDerivedDataCacheInterface
 		/** Async worker that checks the cache backend and if that fails, calls the deriver to build the data and then puts the results to the cache **/
 		void DoWork()
 		{
+			TRACE_CPUPROFILER_EVENT_SCOPE_TEXT(TEXT("DDC_DoWork"));
+
 			const int32 NumBeforeDDC = Data.Num();
 			bool bGetResult;
 			{
+				TRACE_CPUPROFILER_EVENT_SCOPE_TEXT(TEXT("DDC_Get"));
+
 				INC_DWORD_STAT(STAT_DDC_NumGets);
 				STAT(double ThisTime = 0);
 				{
@@ -203,6 +211,8 @@ class FDerivedDataCache : public FDerivedDataCacheInterface
 			else if (DataDeriver)
 			{
 				{
+					TRACE_CPUPROFILER_EVENT_SCOPE_TEXT(TEXT("DDC_Build"));
+
 					INC_DWORD_STAT(STAT_DDC_NumBuilds);
 					STAT(double ThisTime = 0);
 					{
@@ -217,6 +227,9 @@ class FDerivedDataCache : public FDerivedDataCacheInterface
 				if (bSuccess)
 				{
 					check(Data.Num());
+
+					TRACE_CPUPROFILER_EVENT_SCOPE_TEXT(TEXT("DDC_Put"));
+
 					INC_DWORD_STAT(STAT_DDC_NumPuts);
 					STAT(double ThisTime = 0);
 					{
@@ -277,7 +290,7 @@ public:
 
 	virtual bool GetSynchronous(FDerivedDataPluginInterface* DataDeriver, TArray<uint8>& OutData, bool* bDataWasBuilt = nullptr) override
 	{
-		QUICK_SCOPE_CYCLE_COUNTER(STAT_DDC_GetSynchronous);
+		DDC_SCOPE_CYCLE_COUNTER(DDC_GetSynchronous);
 		check(DataDeriver);
 		FString CacheKey = FDerivedDataCache::BuildCacheKey(DataDeriver);
 		UE_LOG(LogDerivedDataCache, Verbose, TEXT("GetSynchronous %s"), *CacheKey);
@@ -294,7 +307,7 @@ public:
 
 	virtual uint32 GetAsynchronous(FDerivedDataPluginInterface* DataDeriver) override
 	{
-		QUICK_SCOPE_CYCLE_COUNTER(STAT_DDC_GetAsynchronous);
+		DDC_SCOPE_CYCLE_COUNTER(DDC_GetAsynchronous);
 		FScopeLock ScopeLock(&SynchronizationObject);
 		uint32 Handle = NextHandle();
 		FString CacheKey = FDerivedDataCache::BuildCacheKey(DataDeriver);
@@ -319,7 +332,7 @@ public:
 
 	virtual bool PollAsynchronousCompletion(uint32 Handle) override
 	{
-		QUICK_SCOPE_CYCLE_COUNTER(STAT_DDC_PollAsynchronousCompletion);
+		DDC_SCOPE_CYCLE_COUNTER(DDC_PollAsynchronousCompletion);
 		FAsyncTask<FBuildAsyncWorker>* AsyncTask = NULL;
 		{
 			FScopeLock ScopeLock(&SynchronizationObject);
@@ -331,7 +344,7 @@ public:
 
 	virtual void WaitAsynchronousCompletion(uint32 Handle) override
 	{
-		QUICK_SCOPE_CYCLE_COUNTER(STAT_DDC_WaitAsynchronousCompletion);
+		DDC_SCOPE_CYCLE_COUNTER(DDC_WaitAsynchronousCompletion);
 		STAT(double ThisTime = 0);
 		{
 			SCOPE_SECONDS_COUNTER(ThisTime);
@@ -348,7 +361,7 @@ public:
 
 	virtual bool GetAsynchronousResults(uint32 Handle, TArray<uint8>& OutData, bool* bDataWasBuilt = nullptr) override
 	{
-		QUICK_SCOPE_CYCLE_COUNTER(STAT_DDC_GetAsynchronousResults);
+		DDC_SCOPE_CYCLE_COUNTER(DDC_GetAsynchronousResults);
 		FAsyncTask<FBuildAsyncWorker>* AsyncTask = NULL;
 		{
 			FScopeLock ScopeLock(&SynchronizationObject);
@@ -382,7 +395,7 @@ public:
 
 	virtual bool GetSynchronous(const TCHAR* CacheKey, TArray<uint8>& OutData) override
 	{
-		QUICK_SCOPE_CYCLE_COUNTER(STAT_DDC_GetSynchronous_Data);
+		DDC_SCOPE_CYCLE_COUNTER(DDC_GetSynchronous_Data);
 		UE_LOG(LogDerivedDataCache, Verbose, TEXT("GetSynchronous %s"), CacheKey);
 		FAsyncTask<FBuildAsyncWorker> PendingTask((FDerivedDataPluginInterface*)NULL, CacheKey, true);
 		AddToAsyncCompletionCounter(1);
@@ -393,7 +406,7 @@ public:
 
 	virtual uint32 GetAsynchronous(const TCHAR* CacheKey, IDerivedDataRollup* Rollup = NULL) override
 	{
-		QUICK_SCOPE_CYCLE_COUNTER(STAT_DDC_GetAsynchronous_Handle);
+		DDC_SCOPE_CYCLE_COUNTER(DDC_GetAsynchronous_Handle);
 		check(!Rollup); // this needs to be handled by someone else, if rollups are disabled, then it should be NULL
 		FScopeLock ScopeLock(&SynchronizationObject);
 		UE_LOG(LogDerivedDataCache, Verbose, TEXT("GetAsynchronous %s"), CacheKey);
@@ -414,7 +427,7 @@ public:
 	**/
 	void GetAsynchronousForRollup(const TCHAR* CacheKey, uint32 Handle)
 	{
-		QUICK_SCOPE_CYCLE_COUNTER(STAT_DDC_GetAsynchronousForRollup);
+		DDC_SCOPE_CYCLE_COUNTER(DDC_GetAsynchronousForRollup);
 		FScopeLock ScopeLock(&SynchronizationObject);
 		UE_LOG(LogDerivedDataCache, Verbose, TEXT("GetAsynchronous(handle) %s"), CacheKey);
 		FAsyncTask<FBuildAsyncWorker>* AsyncTask = new FAsyncTask<FBuildAsyncWorker>((FDerivedDataPluginInterface*)NULL, CacheKey, false);
@@ -426,7 +439,7 @@ public:
 
 	virtual void Put(const TCHAR* CacheKey, TArray<uint8>& Data, bool bPutEvenIfExists = false) override
 	{
-		QUICK_SCOPE_CYCLE_COUNTER(STAT_DDC_Put);
+		DDC_SCOPE_CYCLE_COUNTER(DDC_Put);
 		STAT(double ThisTime = 0);
 		{
 			SCOPE_SECONDS_COUNTER(ThisTime);
@@ -443,7 +456,7 @@ public:
 
 	virtual bool CachedDataProbablyExists(const TCHAR* CacheKey) override
 	{
-		QUICK_SCOPE_CYCLE_COUNTER(STAT_DDC_CachedDataProbablyExists);
+		DDC_SCOPE_CYCLE_COUNTER(DDC_CachedDataProbablyExists);
 		bool bResult;
 		INC_DWORD_STAT(STAT_DDC_NumExist);
 		STAT(double ThisTime = 0);
@@ -457,7 +470,7 @@ public:
 
 	void NotifyBootComplete() override
 	{
-		QUICK_SCOPE_CYCLE_COUNTER(STAT_DDC_NotifyBootComplete);
+		DDC_SCOPE_CYCLE_COUNTER(DDC_NotifyBootComplete);
 		FDerivedDataBackend::Get().NotifyBootComplete();
 	}
 
@@ -468,7 +481,7 @@ public:
 
 	void WaitForQuiescence(bool bShutdown) override
 	{
-		QUICK_SCOPE_CYCLE_COUNTER(STAT_DDC_WaitForQuiescence);
+		DDC_SCOPE_CYCLE_COUNTER(DDC_WaitForQuiescence);
 		FDerivedDataBackend::Get().WaitForQuiescence(bShutdown);
 	}
 
diff --git a/Engine/Source/Editor/UnrealEd/Private/Commandlets/CookCommandlet.cpp b/Engine/Source/Editor/UnrealEd/Private/Commandlets/CookCommandlet.cpp
--- a/Engine/Source/Editor/UnrealEd/Private/Commandlets/CookCommandlet.cpp
+++ b/Engine/Source/Editor/UnrealEd/Private/Commandlets/CookCommandlet.cpp
@@ -597,6 +597,8 @@ int32 UCookCommandlet::Main(const FString& CmdLineParams)
 
 bool UCookCommandlet::CookByTheBook( const TArray<ITargetPlatform*>& Platforms, TArray<FString>& FilesInPath )
 {
+	TRACE_CPUPROFILER_EVENT_SCOPE_TEXT(TEXT("CookByTheBook"));
+
 	COOK_STAT(FScopedDurationTimer CookByTheBookTimer(DetailedCookStats::CookByTheBookTimeSec));
 	UCookOnTheFlyServer *CookOnTheFlyServer = NewObject<UCookOnTheFlyServer>();
 
diff --git a/Engine/Source/Editor/UnrealEd/Private/CookOnTheFlyServer.cpp b/Engine/Source/Editor/UnrealEd/Private/CookOnTheFlyServer.cpp
--- a/Engine/Source/Editor/UnrealEd/Private/CookOnTheFlyServer.cpp
+++ b/Engine/Source/Editor/UnrealEd/Private/CookOnTheFlyServer.cpp
@@ -307,18 +307,11 @@ void ClearHierarchyTimers()
 
 #define CREATE_TIMER(name, incrementScope) FScopeTimer ScopeTimer##name(__COUNTER__, #name, incrementScope); 
 
-#define SCOPE_TIMER(name)				CREATE_TIMER(name, true); ScopeTimer##name.Start();
+#define SCOPE_TIMER(name)				TRACE_CPUPROFILER_EVENT_SCOPE_TEXT(TEXT(#name)); CREATE_TIMER(name, true); ScopeTimer##name.Start();
 
-#define ACCUMULATE_TIMER(name)			CREATE_TIMER(name, false);
-#define ACCUMULATE_TIMER_START(name)	ScopeTimer##name.Start();
-#define ACCUMULATE_TIMER_STOP(name)		ScopeTimer##name.Stop();
 #else
 #define SCOPE_TIMER(name)
 
-#define ACCUMULATE_TIMER(name) 
-#define ACCUMULATE_TIMER_START(name) 
-#define ACCUMULATE_TIMER_STOP(name) 
-
 void OutputHierarchyTimers() {}
 void ClearHierarchyTimers() {}
 #endif
@@ -1455,6 +1448,8 @@ UCookOnTheFlyServer::~UCookOnTheFlyServer()
 // this tick only happens in the editor cook commandlet directly calls tick on the side
 void UCookOnTheFlyServer::Tick(float DeltaTime)
 {
+	TRACE_CPUPROFILER_EVENT_SCOPE_TEXT(TEXT("UCookOnTheFlyServer::Tick"));
+
 	check(IsCookingInEditor());
 
 	if (IsCookByTheBookMode() && !IsCookByTheBookRunning() && !GIsSlowTask)
@@ -4047,6 +4042,8 @@ void UCookOnTheFlyServer::SaveCookedPackage(UPackage* Package, uint32 SaveFlags,
 
 void UCookOnTheFlyServer::Initialize( ECookMode::Type DesiredCookMode, ECookInitializationFlags InCookFlags, const FString &InOutputDirectoryOverride )
 {
+	TRACE_CPUPROFILER_EVENT_SCOPE_TEXT(TEXT("UCookOnTheFlyServer::Initialize"));
+
 	OutputDirectoryOverride = InOutputDirectoryOverride;
 	CurrentCookMode = DesiredCookMode;
 	CookFlags = InCookFlags;
@@ -5037,6 +5034,8 @@ FName UCookOnTheFlyServer::ConvertCookedPathToUncookedPath(
 
 void UCookOnTheFlyServer::GetAllCookedFiles(TMap<FName, FName>& UncookedPathToCookedPath, const FString& SandboxRootDir)
 {
+	TRACE_CPUPROFILER_EVENT_SCOPE_TEXT(TEXT("UCookOnTheFlyServer::GetAllCookedFiles"));
+
 	TArray<FString> CookedFiles;
 	{
 		IPlatformFile& PlatformFile = FPlatformFileManager::Get().GetPlatformFile();
@@ -5064,6 +5063,8 @@ void UCookOnTheFlyServer::GetAllCookedFiles(TMap<FName, FName>& UncookedPathToCo
 
 void UCookOnTheFlyServer::PopulateCookedPackagesFromDisk(const TArray<ITargetPlatform*>& Platforms)
 {
+	TRACE_CPUPROFILER_EVENT_SCOPE_TEXT(TEXT("UCookOnTheFlyServer::PopulateCookedPackagesFromDisk"));
+
 	// See what files are out of date in the sandbox folder
 	for (int32 Index = 0; Index < Platforms.Num(); Index++)
 	{
@@ -5326,6 +5327,8 @@ const FString ExtractPackageNameFromObjectPath( const FString ObjectPath )
 
 void UCookOnTheFlyServer::CleanSandbox(const bool bIterative)
 {
+	TRACE_CPUPROFILER_EVENT_SCOPE_TEXT(TEXT("UCookOnTheFlyServer::CleanSandbox"));
+
 	const TArray<ITargetPlatform*>& Platforms = GetCookingTargetPlatforms();
 
 	// before we can delete any cooked files we need to make sure that we have finished writing them
@@ -5399,6 +5402,8 @@ void UCookOnTheFlyServer::CleanSandbox(const bool bIterative)
 
 void UCookOnTheFlyServer::GenerateAssetRegistry()
 {
+	TRACE_CPUPROFILER_EVENT_SCOPE_TEXT(TEXT("UCookOnTheFlyServer::GenerateAssetRegistry"));
+
 	// Cache asset registry for later
 	FAssetRegistryModule& AssetRegistryModule = FModuleManager::LoadModuleChecked<FAssetRegistryModule>(TEXT("AssetRegistry"));
 	AssetRegistry = &AssetRegistryModule.Get();
@@ -5536,6 +5541,8 @@ void UCookOnTheFlyServer::AddFileToCook( TArray<FName>& InOutFilesToCook, const
 
 void UCookOnTheFlyServer::CollectFilesToCook(TArray<FName>& FilesInPath, const TArray<FString>& CookMaps, const TArray<FString>& InCookDirectories, const TArray<FString> &IniMapSections, ECookByTheBookOptions FilesToCookFlags)
 {
+	TRACE_CPUPROFILER_EVENT_SCOPE_TEXT(TEXT("UCookOnTheFlyServer::GenerateAssetRegistry"));
+
 #if OUTPUT_TIMING
 	SCOPE_TIMER(CollectFilesToCook);
 #endif
diff --git a/Engine/Source/Runtime/Engine/Private/StaticMesh.cpp b/Engine/Source/Runtime/Engine/Private/StaticMesh.cpp
--- a/Engine/Source/Runtime/Engine/Private/StaticMesh.cpp
+++ b/Engine/Source/Runtime/Engine/Private/StaticMesh.cpp
@@ -2221,6 +2221,8 @@ void FStaticMeshRenderData::Cache(UStaticMesh* Owner, const FStaticMeshLODSettin
 
 
 	{
+		TRACE_CPUPROFILER_EVENT_SCOPE_TEXT(TEXT("StaticMesh_Cache"));
+
 		COOK_STAT(auto Timer = StaticMeshCookStats::UsageStats.TimeSyncWork());
 		int32 T0 = FPlatformTime::Cycles();
 		int32 NumLODs = Owner->SourceModels.Num();
