Fixed nopch compile error.

diff --git a/Engine/Source/Developer/DerivedDataCache/Private/DDCCleanup.cpp b/Engine/Source/Developer/DerivedDataCache/Private/DDCCleanup.cpp
--- a/Engine/Source/Developer/DerivedDataCache/Private/DDCCleanup.cpp
+++ b/Engine/Source/Developer/DerivedDataCache/Private/DDCCleanup.cpp
@@ -4,6 +4,7 @@
 #include "HAL/PlatformProcess.h"
 #include "HAL/FileManager.h"
 #include "HAL/RunnableThread.h"
+#include "Misc/ConfigCacheIni.h"
 #include "Misc/ScopeLock.h"
 #include "Math/RandomStream.h"
 
