Fix some minor issues with the output of "DDC Summary Stats" after a cook.

Add checks for divide by zero so that we output 0, rather than INF.
Increased the padding of the first output column so that the indentation of each output line is consistent.

diff --git a/Engine/Source/Developer/DerivedDataCache/Private/DerivedDataCache.cpp b/Engine/Source/Developer/DerivedDataCache/Private/DerivedDataCache.cpp
--- a/Engine/Source/Developer/DerivedDataCache/Private/DerivedDataCache.cpp
+++ b/Engine/Source/Developer/DerivedDataCache/Private/DerivedDataCache.cpp
@@ -37,6 +37,12 @@ DEFINE_STAT(STAT_DDC_ExistTime);
 #include "DerivedDataCacheUsageStats.h"
 namespace DerivedDataCacheCookStats
 {
+	// Use to prevent potential divide by zero issues
+	inline double SafeDivide(const int64 Numerator, const int64 Denominator)
+	{
+		return Denominator != 0 ? (double)Numerator / (double)Denominator : 0.0;
+	}
+
 	FCookStatsManager::FAutoRegisterCallback RegisterCookStats([](FCookStatsManager::AddStatFuncRef AddStat)
 	{
 		TMap<FString, FDerivedDataCacheUsageStats> DDCStats;
@@ -102,15 +108,15 @@ namespace DerivedDataCacheCookStats
 				AddStat(TEXT("DDC.Summary"), FCookStatsManager::CreateKeyValueArray(
 					TEXT("TotalGetHits"), TotalGetHits,
 					TEXT("TotalGets"), TotalGets,
-					TEXT("TotalGetHitPct"), (double)TotalGetHits / TotalGets,
-					TEXT("LocalGetHitPct"), (double)LocalHits / TotalGets,
-					TEXT("SharedGetHitPct"), (double)SharedHits / TotalGets,
-					TEXT("OtherGetHitPct"), double(TotalGetHits - LocalHits - SharedHits) / TotalGets,
-					TEXT("GetMissPct"), (double)TotalGetMisses / TotalGets,
+					TEXT("TotalGetHitPct"), SafeDivide(TotalGetHits, TotalGets),
+					TEXT("LocalGetHitPct"), SafeDivide(LocalHits, TotalGets),
+					TEXT("SharedGetHitPct"), SafeDivide(SharedHits, TotalGets),
+					TEXT("OtherGetHitPct"), SafeDivide((TotalGetHits - LocalHits - SharedHits), TotalGets),
+					TEXT("GetMissPct"), SafeDivide(TotalGetMisses, TotalGets),
 					TEXT("TotalPutHits"), TotalPutHits,
 					TEXT("TotalPuts"), TotalPuts,
-					TEXT("TotalPutHitPct"), (double)TotalPutHits / TotalPuts,
-					TEXT("PutMissPct"), (double)TotalPutMisses / TotalPuts
+					TEXT("TotalPutHitPct"), SafeDivide(TotalPutHits, TotalPuts),
+					TEXT("PutMissPct"), SafeDivide(TotalPutMisses, TotalPuts)
 					));
 			}
 		}
diff --git a/Engine/Source/Editor/UnrealEd/Private/Commandlets/CookCommandlet.cpp b/Engine/Source/Editor/UnrealEd/Private/Commandlets/CookCommandlet.cpp
--- a/Engine/Source/Editor/UnrealEd/Private/Commandlets/CookCommandlet.cpp
+++ b/Engine/Source/Editor/UnrealEd/Private/Commandlets/CookCommandlet.cpp
@@ -286,7 +286,7 @@ namespace DetailedCookStats
 			UE_LOG(LogCookCommandlet, Display, TEXT("================="));
 			for (const auto& Attr : DDCSummaryStats)
 			{
-				UE_LOG(LogCookCommandlet, Display, TEXT("%-14s=%10s"), *Attr.Key, *Attr.Value);
+				UE_LOG(LogCookCommandlet, Display, TEXT("%-16s=%10s"), *Attr.Key, *Attr.Value);
 			}
 		}
 		if (DDCResourceUsageStats.Num() > 0)
