Stopping DDC cleanup thread as soon as the engine begins shutting down

diff --git a/Engine/Source/Developer/DerivedDataCache/Private/DDCCleanup.cpp b/Engine/Source/Developer/DerivedDataCache/Private/DDCCleanup.cpp
--- a/Engine/Source/Developer/DerivedDataCache/Private/DDCCleanup.cpp
+++ b/Engine/Source/Developer/DerivedDataCache/Private/DDCCleanup.cpp
@@ -123,6 +123,11 @@ uint32 FDDCCleanup::Run()
 	return 0;
 }
 
+bool FDDCCleanup::ShouldStop() const
+{
+	return StopTaskCounter.GetValue() > 0 || GIsRequestingExit;
+}
+
 bool FDDCCleanup::CleanupFilesystemDirectory( TSharedPtr< FFilesystemInfo > FilesystemInfo )
 {
 	bool bCleanedUp = false;
diff --git a/Engine/Source/Developer/DerivedDataCache/Public/DDCCleanup.h b/Engine/Source/Developer/DerivedDataCache/Public/DDCCleanup.h
--- a/Engine/Source/Developer/DerivedDataCache/Public/DDCCleanup.h
+++ b/Engine/Source/Developer/DerivedDataCache/Public/DDCCleanup.h
@@ -48,10 +48,7 @@ class DERIVEDDATACACHE_API FDDCCleanup : public FRunnable
 	//~ End FRunnable Interface
 
 	/** Checks if there's been any Stop requests */
-	FORCEINLINE bool ShouldStop() const
-	{
-		return StopTaskCounter.GetValue() > 0;
-	}
+	bool ShouldStop() const;
 
 	/**
 	 * Waits for a given amount of time periodically checking if there's been any Stop requests.
