Work around a false positive from static analysis

Static analysis is reporting that certain local variables inside of lambdas are shadowing themselves. Steve.Robb found https://developercommunity.visualstudio.com/content/problem/576913/c6244-regression-in-new-lambda-processorpermissive.html where the same problem was reported last year. Microsoft closed the issue without resolving it.

#rnx

diff --git a/Engine/Source/Developer/DerivedDataCache/Private/DerivedDataCache.cpp b/Engine/Source/Developer/DerivedDataCache/Private/DerivedDataCache.cpp
--- a/Engine/Source/Developer/DerivedDataCache/Private/DerivedDataCache.cpp
+++ b/Engine/Source/Developer/DerivedDataCache/Private/DerivedDataCache.cpp
@@ -43,7 +43,9 @@ namespace DerivedDataCacheCookStats
 		return Denominator != 0 ? (double)Numerator / (double)Denominator : 0.0;
 	}
 
-	FCookStatsManager::FAutoRegisterCallback RegisterCookStats([](FCookStatsManager::AddStatFuncRef AddStat)
+	// AddCookStats cannot be a lambda because of false positives in static analysis.
+	// See https://developercommunity.visualstudio.com/content/problem/576913/c6244-regression-in-new-lambda-processorpermissive.html
+	static void AddCookStats(FCookStatsManager::AddStatFuncRef AddStat)
 	{
 		TMap<FString, FDerivedDataCacheUsageStats> DDCStats;
 		GetDerivedDataCacheRef().GatherUsageStats(DDCStats);
@@ -120,7 +122,9 @@ namespace DerivedDataCacheCookStats
 					));
 			}
 		}
-	});
+	}
+
+	FCookStatsManager::FAutoRegisterCallback RegisterCookStats(AddCookStats);
 }
 #endif
 
