Fix extra DDC puts when DDC logging is enabled on the command line

Enabling verbose DDC logging with -ini:Engine:[Core.Log]:LogDerivedDataCache=Verbose was causing extra idempotent puts because of checks for the DerivedDataCache commandlet that checked for "DerivedDataCache" anywhere in the command line. Check for "Run=DerivedDataCache" instead to avoid the problem.

diff --git a/Engine/Source/Developer/DerivedDataCache/Private/FileSystemDerivedDataBackend.cpp b/Engine/Source/Developer/DerivedDataCache/Private/FileSystemDerivedDataBackend.cpp
--- a/Engine/Source/Developer/DerivedDataCache/Private/FileSystemDerivedDataBackend.cpp
+++ b/Engine/Source/Developer/DerivedDataCache/Private/FileSystemDerivedDataBackend.cpp
@@ -106,7 +106,7 @@ public:
 				bFailed = false;
 			}
 		}
-		if (FString(FCommandLine::Get()).Contains(TEXT("DerivedDataCache")) )
+		if (FString(FCommandLine::Get()).Contains(TEXT("Run=DerivedDataCache")) )
 		{
 			bTouch = true; // we always touch files when running the DDC commandlet
 		}
diff --git a/Engine/Source/Runtime/Engine/Private/AudioDerivedData.cpp b/Engine/Source/Runtime/Engine/Private/AudioDerivedData.cpp
--- a/Engine/Source/Runtime/Engine/Private/AudioDerivedData.cpp
+++ b/Engine/Source/Runtime/Engine/Private/AudioDerivedData.cpp
@@ -596,7 +596,7 @@ void FStreamedAudioPlatformData::Cache(USoundWave& InSoundWave, const FPlatformA
 
 	uint32 Flags = InFlags;
 
-	static bool bForDDC = FString(FCommandLine::Get()).Contains(TEXT("DerivedDataCache"));
+	static bool bForDDC = FString(FCommandLine::Get()).Contains(TEXT("Run=DerivedDataCache"));
 	if (bForDDC)
 	{
 		Flags |= EStreamedAudioCacheFlags::ForDDCBuild;
diff --git a/Engine/Source/Runtime/Engine/Private/TextureDerivedData.cpp b/Engine/Source/Runtime/Engine/Private/TextureDerivedData.cpp
--- a/Engine/Source/Runtime/Engine/Private/TextureDerivedData.cpp
+++ b/Engine/Source/Runtime/Engine/Private/TextureDerivedData.cpp
@@ -808,7 +808,7 @@ void FTexturePlatformData::Cache(
 
 	uint32 Flags = InFlags;
 
-	static bool bForDDC = FString(FCommandLine::Get()).Contains(TEXT("DerivedDataCache"));
+	static bool bForDDC = FString(FCommandLine::Get()).Contains(TEXT("Run=DerivedDataCache"));
 	if (bForDDC)
 	{
 		Flags |= ETextureCacheFlags::ForDDCBuild;
