Added some comments on the strategy used to shorten the DDC keys

diff --git a/Engine/Source/Developer/DerivedDataCache/Private/DerivedDataLimitKeyLengthWrapper.h b/Engine/Source/Developer/DerivedDataCache/Private/DerivedDataLimitKeyLengthWrapper.h
--- a/Engine/Source/Developer/DerivedDataCache/Private/DerivedDataLimitKeyLengthWrapper.h
+++ b/Engine/Source/Developer/DerivedDataCache/Private/DerivedDataLimitKeyLengthWrapper.h
@@ -196,15 +196,21 @@ private:
 		HashState.Update((const uint8*)&Length, sizeof(int32));
 
 		auto ResultSrc = StringCast<UCS2CHAR>(*Result);
-		uint32 CRCofPayload(FCrc::MemCrc32(ResultSrc.Get(), Length * sizeof(UCS2CHAR)));
 
+		// This is pretty redundant. Incorporating the CRC of the name into the hash
+		// which also ends up computing SHA1 of the name is not really going to make 
+		// any meaningful difference to the strength of the key so it's just a waste
+		// of CPU. However it's difficult to get rid of without invalidating the DDC
+		// contents so here we are.
+		const uint32 CRCofPayload(FCrc::MemCrc32(ResultSrc.Get(), Length * sizeof(UCS2CHAR)));
 		HashState.Update((const uint8*)&CRCofPayload, sizeof(uint32));
+
 		HashState.Update((const uint8*)ResultSrc.Get(), Length * sizeof(UCS2CHAR));
 
 		HashState.Final();
 		uint8 Hash[FSHA1::DigestSize];
 		HashState.GetHash(Hash);
-		FString HashString = BytesToHex(Hash, FSHA1::DigestSize);
+		const FString HashString = BytesToHex(Hash, FSHA1::DigestSize);
 
 		int32 HashStringSize = HashString.Len();
 		int32 OriginalPart = MaxKeyLength - HashStringSize - 2;
