Removed rollup support from DDC interface

diff --git a/Engine/Source/Developer/DerivedDataCache/Private/DerivedDataCache.cpp b/Engine/Source/Developer/DerivedDataCache/Private/DerivedDataCache.cpp
--- a/Engine/Source/Developer/DerivedDataCache/Private/DerivedDataCache.cpp
+++ b/Engine/Source/Developer/DerivedDataCache/Private/DerivedDataCache.cpp
@@ -315,10 +315,10 @@ public:
 	{
 		DDC_SCOPE_CYCLE_COUNTER(DDC_GetAsynchronous);
 		FScopeLock ScopeLock(&SynchronizationObject);
-		uint32 Handle = NextHandle();
+		const uint32 Handle = NextHandle();
 		FString CacheKey = FDerivedDataCache::BuildCacheKey(DataDeriver);
 		UE_LOG(LogDerivedDataCache, Verbose, TEXT("GetAsynchronous %s"), *CacheKey);
-		bool bSync = !DataDeriver->IsBuildThreadsafe();
+		const bool bSync = !DataDeriver->IsBuildThreadsafe();
 		FAsyncTask<FBuildAsyncWorker>* AsyncTask = new FAsyncTask<FBuildAsyncWorker>(DataDeriver, *CacheKey, bSync);
 		check(!PendingTasks.Contains(Handle));
 		PendingTasks.Add(Handle,AsyncTask);
@@ -389,16 +389,6 @@ public:
 		return true;
 	}
 
-	virtual IDerivedDataRollup* StartRollup() override
-	{
-		return NULL;
-	}
-
-	virtual void EndRollup(IDerivedDataRollup*& Rollup) override
-	{
-		check(!Rollup); // this needs to be handled by someone else, if rollups are disabled, then it should be NULL
-	}
-
 	virtual bool GetSynchronous(const TCHAR* CacheKey, TArray<uint8>& OutData) override
 	{
 		DDC_SCOPE_CYCLE_COUNTER(DDC_GetSynchronous_Data);
@@ -410,13 +400,12 @@ public:
 		return PendingTask.GetTask().bSuccess;
 	}
 
-	virtual uint32 GetAsynchronous(const TCHAR* CacheKey, IDerivedDataRollup* Rollup = NULL) override
+	virtual uint32 GetAsynchronous(const TCHAR* CacheKey) override
 	{
 		DDC_SCOPE_CYCLE_COUNTER(DDC_GetAsynchronous_Handle);
-		check(!Rollup); // this needs to be handled by someone else, if rollups are disabled, then it should be NULL
 		FScopeLock ScopeLock(&SynchronizationObject);
 		UE_LOG(LogDerivedDataCache, Verbose, TEXT("GetAsynchronous %s"), CacheKey);
-		uint32 Handle = NextHandle();
+		const uint32 Handle = NextHandle();
 		FAsyncTask<FBuildAsyncWorker>* AsyncTask = new FAsyncTask<FBuildAsyncWorker>((FDerivedDataPluginInterface*)NULL, CacheKey, false);
 		check(!PendingTasks.Contains(Handle));
 		PendingTasks.Add(Handle, AsyncTask);
@@ -425,24 +414,6 @@ public:
 		return Handle;
 	}
 
-	/** 
-	 * Starts the async process of checking the cache and if the item is present, retrieving the cached results (version for internal use by rollups)
-	 * @param	CacheKey	Key to identify the data
-	 * @param	Rollup		Rollup pointer, if this is part of a rollup
-	 * @param	Handle		The handle that can be used for PollAsynchronousCompletion, WaitAsynchronousCompletion and GetAsynchronousResults
-	**/
-	void GetAsynchronousForRollup(const TCHAR* CacheKey, uint32 Handle)
-	{
-		DDC_SCOPE_CYCLE_COUNTER(DDC_GetAsynchronousForRollup);
-		FScopeLock ScopeLock(&SynchronizationObject);
-		UE_LOG(LogDerivedDataCache, Verbose, TEXT("GetAsynchronous(handle) %s"), CacheKey);
-		FAsyncTask<FBuildAsyncWorker>* AsyncTask = new FAsyncTask<FBuildAsyncWorker>((FDerivedDataPluginInterface*)NULL, CacheKey, false);
-		check(!PendingTasks.Contains(Handle));
-		PendingTasks.Add(Handle,AsyncTask);
-		AddToAsyncCompletionCounter(1);
-		AsyncTask->StartBackgroundTask();
-	}
-
 	virtual void Put(const TCHAR* CacheKey, TArray<uint8>& Data, bool bPutEvenIfExists = false) override
 	{
 		DDC_SCOPE_CYCLE_COUNTER(DDC_Put);
@@ -502,12 +473,6 @@ public:
 		FDerivedDataBackend::Get().GetDirectories(OutResults);
 	}
 
-	/** Called at ShutdownModule() time to print out status before we're cleaned up */
-	virtual void PrintLeaks()
-	{
-		// Used by derived classes to spit out leaked pending rollups
-	}
-
 	virtual void GatherUsageStats(TMap<FString, FDerivedDataCacheUsageStats>& UsageStatsMap) override
 	{
 		FDerivedDataBackend::Get().GatherUsageStats(UsageStatsMap);
@@ -551,597 +516,27 @@ private:
 	FOnDDCNotification DDCNotificationEvent;
 };
 
-//Forward reference
-FDerivedDataCache& InternalSingleton();
-
-namespace EPhase
-{
-	enum Type
-	{
-		Adding,
-		AsyncRollupGet,
-		AsyncRollupGetSucceed,
-		AsyncRollupGetFailed_GettingItemsAndWaitingForPuts,
-		Done,
-	};
-}
-
-/** 
- * Opaque class for rollup handling
-**/
-class FDerivedDataRollup : public IDerivedDataRollup
-{
-	/** Magic numbers to verify integrity and check endianness **/
-	enum
-	{
-		MAGIC=0x9E1B83C1,
-		MAGIC_SWAPPED=0xC1831B9E
-	};
-
-	/** Helper structure for an element of a rollup **/
-	struct FRollupItem
-	{
-		FRollupItem(const FString& InCacheKey, uint32 InAsyncHandle)
-			: CacheKey(InCacheKey)
-			, AsyncHandle(InAsyncHandle)
-			, FinishedFromThePerspectiveOfTheCaller(false)
-		{
-		}
-		/** Cache key for this item **/
-		FString CacheKey;
-		/** Async handle for this item, used both to return to original caller, and for calls to the actual DDC **/
-		uint32 AsyncHandle;
-		/** Payload of this item, used for both from the get of the rollup and a put to the rollup **/
-		TArray<uint8> Payload;
-		/** If true, then the caller has either asked for the results. This means we don't need to keep them any more. **/
-		bool FinishedFromThePerspectiveOfTheCaller;
-	};
-
-	/** Items in this rollup **/
-	TArray<FRollupItem> Items;
-	/** Redundant copy of the keys in this rollup **/
-	TSet<FString> CacheKeys;
-	/** Redundant copy of the async handles in this rollup **/
-	TSet<uint32> AsyncHandles;
-	/** Cache key for the rollup itself **/
-	FString RollupCacheKey;
-	/** Async handle for the rollup **/
-	uint32 RollupAsyncHandle;
-	/** Tracks the phase this rollup is in. Mostly used for check's **/
-	EPhase::Type CurrentPhase;
-	/** If true, the rollup was corrupted, so we need to force a put when we get to the put. **/
-	bool ForcePutForCorruption;
-
-	/** Called when the rollup is ready. This is indirectly caused by the original caller waiting for an item to be ready. **/
-	void GetRollupResults()
-	{
-		check(Items.Num());
-		check(CurrentPhase == EPhase::AsyncRollupGet);
-		TArray<uint8> Payload;
-		bool bFailed = true;
-		if (InternalSingleton().FDerivedDataCache::GetAsynchronousResults(RollupAsyncHandle, Payload))
-		{
-			ForcePutForCorruption = true;
-			if (Payload.Num() > sizeof(uint32) * 2)
-			{
-				FMemoryReader Ar(Payload);
-				uint32 Magic;
-				Ar << Magic;
-				if (Magic == MAGIC_SWAPPED)
-				{
-					Ar.SetByteSwapping(!Ar.ForceByteSwapping());
-					Magic = MAGIC;
-				}
-				if (Magic == MAGIC)
-				{
-					int32 Count;
-					Ar << Count;
-					if (Count == Items.Num())
-					{
-						bFailed = false;
-						for (int32 Index = 0; Index < Items.Num(); Index++)
-						{
-							FString Key;
-							Ar << Key;
-							if (Key != Items[Index].CacheKey)
-							{
-								bFailed = true;
-								break;
-							}
-							Ar << Items[Index].Payload;
-							if (!Items[Index].Payload.Num())
-							{
-								bFailed = true;
-								break;
-							}
-						}
-					}
-				}
-			}
-		}
-		if (!bFailed)
-		{
-			ForcePutForCorruption = false;
-			CurrentPhase = EPhase::AsyncRollupGetSucceed;
-		}
-		else
-		{
-			CurrentPhase = EPhase::AsyncRollupGetFailed_GettingItemsAndWaitingForPuts;
-			for (int32 Index = 0; Index < Items.Num(); Index++)
-			{
-				Items[Index].Payload.Empty(); // we might have had partial success on a corrupted rollup; we won't accept those
-				InternalSingleton().FDerivedDataCache::GetAsynchronousForRollup(*Items[Index].CacheKey, Items[Index].AsyncHandle);
-			}
-		}
-	}
-
-	/** Tests to see if the rollup is complete and ready to be put; if it is, it packages it and puts it. **/
-	bool CheckForPut()
-	{
-		check(Items.Num());
-		check(CurrentPhase == EPhase::AsyncRollupGetFailed_GettingItemsAndWaitingForPuts);
-		for (int32 Index = 0; Index < Items.Num(); Index++)
-		{
-			if (!Items[Index].Payload.Num())
-			{
-				return false; // not done yet because we don't have all of the data
-			}
-			if (!Items[Index].FinishedFromThePerspectiveOfTheCaller)
-			{
-				return false; // not done yet because the caller still hasn't retrieved their results
-			}
-		}
-		uint32 Magic = MAGIC;
-		int32 Count = Items.Num();
-		TArray<uint8> Buffer;
-		FMemoryWriter Ar(Buffer);
-		Ar << Magic;
-		Ar << Count;
-		for (int32 Index = 0; Index < Items.Num(); Index++)
-		{
-			Ar << Items[Index].CacheKey;
-			Ar << Items[Index].Payload;
-		}
-		InternalSingleton().FDerivedDataCache::Put(*RollupCacheKey, Buffer, ForcePutForCorruption);
-		CurrentPhase = EPhase::Done;
-		return true;
-	}
-
-	/** 
-	 * Checks to see if there is any reason for this rollup to stay alive. 
-	 * @return true when done
-	**/
-	bool CheckForDone()
-	{
-		check(Items.Num());
-		check(CurrentPhase == EPhase::AsyncRollupGetSucceed);
-		for (int32 Index = 0; Index < Items.Num(); Index++)
-		{
-			if (!Items[Index].FinishedFromThePerspectiveOfTheCaller)
-			{
-				return false; // not done yet because the caller still hasn't retrieved their results
-			}
-			if (!Items[Index].Payload.Num())
-			{
-				check(0); // what is this?
-			}
-		}
-		CurrentPhase = EPhase::Done;
-		return true;
-	}
-
-	/** 
-	 * Finds an item by async handle, not legal to call if this rollup does not contain this handle
-	 * @param Handle		Async handle to find the corresponding item of
-	 * @return				Item with this async handle
-	**/
-	FRollupItem& FindItem(uint32 Handle)
-	{
-		for (int32 Index = 0; Index < Items.Num(); Index++)
-		{
-			if (Items[Index].AsyncHandle == Handle)
-			{
-				return Items[Index];
-			}
-		}
-		UE_LOG(LogDerivedDataCache, Fatal, TEXT("Illegal async handle passed to FRollupItem::FindItem()."));
-		return Items[0];// shouldn't get here anyway, won't make it worse
-	}
-
-	/** 
-	 * Finds an item by cache key, not legal to call if this rollup does not contain this cache key
-	 * @param InCacheKey	cache key to find the corresponding item of
-	 * @return				Item with this cache key
-	**/
-	FRollupItem& FindItem(const FString& InCacheKey)
-	{
-		for (int32 Index = 0; Index < Items.Num(); Index++)
-		{
-			if (Items[Index].CacheKey == InCacheKey)
-			{
-				return Items[Index];
-			}
-		}
-		UE_LOG(LogDerivedDataCache, Fatal, TEXT("Illegal cache key passed to FRollupItem::FindItem()."));
-		return Items[0];// shouldn't get here anyway, won't make it worse
-	}
-
-public:
-	FDerivedDataRollup()
-		: RollupCacheKey(TEXT("ROLLUP_"))
-		, RollupAsyncHandle(0)
-		, CurrentPhase(EPhase::Adding)
-		, ForcePutForCorruption(false)
-	{
-	}
-
-	/** Return the cache key, used for error spew **/
-	FString GetName()
-	{
-		return RollupCacheKey;
-	}
-
-	/** Return true if this rollup can be deleted because it has completed its life cycle. **/
-	bool IsDone()
-	{
-		return CurrentPhase == EPhase::Done;
-	}
-
-	/** Return true if this rollup contains an item with the given cache key. **/
-	bool Contains(const FString& InCacheKey)
-	{
-		return CacheKeys.Contains(InCacheKey);
-	}
-
-	/** Return true if this rollup contains an item with the given async handle. **/
-	bool Contains(uint32 InAsyncHandle)
-	{
-		return AsyncHandles.Contains(InAsyncHandle);
-	}
-
-	/** Add a new item to this rollup with the given cache key and async handle. **/
-	void Add(const FString& InCacheKey, uint32 InAsyncHandle)
-	{
-		check(CurrentPhase == EPhase::Adding);
-		RollupCacheKey += InCacheKey;
-		CacheKeys.Add(InCacheKey);
-		AsyncHandles.Add(InAsyncHandle);
-		new (Items) FRollupItem(InCacheKey, InAsyncHandle);
-	}
-
-	/** Signifies the end of the adding phase and starts an async get of the rollup. **/
-	void Close()
-	{
-		check(CurrentPhase == EPhase::Adding);
-		if (Items.Num())
-		{
-			RollupAsyncHandle = InternalSingleton().FDerivedDataCache::GetAsynchronous(*RollupCacheKey);
-			CurrentPhase = EPhase::AsyncRollupGet;
-		}
-		else
-		{
-			CurrentPhase = EPhase::Done;
-		}
-	}
-
-	/** 
-	 * Handle PollAsynchronousCompletion from the calling code. 
-	 * @param Handle	Async completion handle for the item
-	 * @return			true, if the calling code can request results yet
-	 **/
-	bool PollAsynchronousCompletion(uint32 Handle)
-	{
-		check(Contains(Handle));
-		if (CurrentPhase == EPhase::AsyncRollupGet)
-		{
-			// in this phase we see if the rollup is done
-			if (!InternalSingleton().FDerivedDataCache::PollAsynchronousCompletion(RollupAsyncHandle))
-			{
-				return false;
-			}
-			GetRollupResults();
-			// fall through to handle the other cases
-		}
-		if (CurrentPhase == EPhase::AsyncRollupGetSucceed)
-		{
-			// Rollup succeeded, so the calling code can get the results
-			return true;
-		}
-		if (CurrentPhase == EPhase::AsyncRollupGetFailed_GettingItemsAndWaitingForPuts)
-		{
-			// Rollup failed, so call the actual PAC for the individual item
-			return InternalSingleton().FDerivedDataCache::PollAsynchronousCompletion(Handle);
-		}
-		check(0); // bad phase
-		return false;
-	}
-
-	/** 
-	 * Handle WaitAsynchronousCompletion from the calling code. 
-	 * @param Handle	Async completion handle for the item
-	 **/
-	void WaitAsynchronousCompletion(uint32 Handle)
-	{
-		check(Contains(Handle));
-		if (CurrentPhase == EPhase::AsyncRollupGet)
-		{
-			// in this phase we wait for the rollup to complete, then deal with the results
-			InternalSingleton().FDerivedDataCache::WaitAsynchronousCompletion(RollupAsyncHandle);
-			GetRollupResults();
-			// fall through to handle the other cases
-		}
-		if (CurrentPhase == EPhase::AsyncRollupGetSucceed)
-		{
-			// Rollup succeeded, so the calling code can get the results
-			return;
-		}
-		if (CurrentPhase == EPhase::AsyncRollupGetFailed_GettingItemsAndWaitingForPuts)
-		{
-			// Rollup failed, so call the actual WAC for the individual item
-			InternalSingleton().FDerivedDataCache::WaitAsynchronousCompletion(Handle);
-			return;
-		}
-		check(0); // bad phase
-	}
-
-	/** 
-	 * Handle GetAsynchronousResults from the calling code. If this is the last piece of data, the rollup will be put.
-	 * @param Handle	Async completion handle for the item
-	 * @param OutData	returned payload
-	 * @return			true if the payload contains data and everything is peachy
-	 **/
-	bool GetAsynchronousResults(uint32 Handle, TArray<uint8>& OutData, bool* bDataWasBuilt)
-	{
-		check(Contains(Handle));
-		FRollupItem& Item = FindItem(Handle);
-		Item.FinishedFromThePerspectiveOfTheCaller = true;
-		OutData.Empty();
-		if (CurrentPhase == EPhase::AsyncRollupGetSucceed)
-		{
-			if (bDataWasBuilt)
-			{
-				*bDataWasBuilt = false;
-			}
-			OutData = Item.Payload;
-			CheckForDone();
-			return !!OutData.Num();
-		}
-		if (CurrentPhase == EPhase::AsyncRollupGetFailed_GettingItemsAndWaitingForPuts)
-		{
-			if (InternalSingleton().FDerivedDataCache::GetAsynchronousResults(Handle, OutData, bDataWasBuilt))
-			{
-				Item.Payload = OutData;
-				CheckForPut();
-			}
-			return !!OutData.Num();
-		}
-		check(0); // bad phase
-		return false;
-	}
-	/** 
-	 * Handle Put from the calling code. If this is the last piece of data, the rollup will be put.
-	 * @param CacheKey	Cache key for the item
-	 * @param Data		returned payload
-	 **/
-	void Put(const TCHAR* CacheKey, TArray<uint8>& Data)
-	{
-		if (CurrentPhase != EPhase::AsyncRollupGetFailed_GettingItemsAndWaitingForPuts)
-		{
-			return;
-		}
-		check(Contains(CacheKey));
-		FRollupItem& Item = FindItem(CacheKey);
-		check(Data.Num());
-		Item.Payload = Data;
-		CheckForPut();
-	}
-};
-
-
-
-/** 
- * Implementation of the derived data cache, this layer implements rollups
-**/
-class FDerivedDataCacheWithRollups : public FDerivedDataCache
-{
-	typedef FDerivedDataCache Super;
-public:
-
-	~FDerivedDataCacheWithRollups()
-	{
-	}
-
-	virtual void PrintLeaks() override
-	{
-		FScopeLock ScopeLock(&SynchronizationObject);
-		UE_LOG(LogDerivedDataCache, Log, TEXT("Shutdown"));
-		for (TSet<FDerivedDataRollup*>::TIterator Iter(PendingRollups); Iter; ++Iter)
-		{
-			FString Name((*Iter)->GetName());
-			if (Name.Len() > 1024)
-			{
-				Name = Name.Left(1024);
-				Name += TEXT("...");
-			}
-			UE_LOG(LogDerivedDataCache, Warning, TEXT("Leaked Rollup! %s"), *Name);
-		}
-		//check(!PendingRollups.Num()); // leaked rollups
-	}
-
-	virtual bool PollAsynchronousCompletion(uint32 Handle) override
-	{
-		FScopeLock ScopeLock(&SynchronizationObject);
-		for (TSet<FDerivedDataRollup*>::TIterator Iter(PendingRollups); Iter; ++Iter)
-		{
-			if ((*Iter)->Contains(Handle))
-			{
-				return (*Iter)->PollAsynchronousCompletion(Handle);
-			}
-		}
-		return Super::PollAsynchronousCompletion(Handle);
-	}
-
-	virtual void WaitAsynchronousCompletion(uint32 Handle) override
-	{
-		STAT(double ThisTime = 0);
-		{
-			SCOPE_SECONDS_COUNTER(ThisTime);
-			FScopeLock ScopeLock(&SynchronizationObject);
-			for (TSet<FDerivedDataRollup*>::TIterator Iter(PendingRollups); Iter; ++Iter)
-			{
-				if ((*Iter)->Contains(Handle))
-				{
-					(*Iter)->WaitAsynchronousCompletion(Handle);
-					return;
-				}
-			}
-		}
-		INC_FLOAT_STAT_BY(STAT_DDC_ASyncWaitTime,(float)ThisTime);
-		Super::WaitAsynchronousCompletion(Handle);
-	}
-
-	virtual bool GetAsynchronousResults(uint32 Handle, TArray<uint8>& OutData, bool* bDataWasBuilt = nullptr) override
-	{
-		FScopeLock ScopeLock(&SynchronizationObject);
-		for (TSet<FDerivedDataRollup*>::TIterator Iter(PendingRollups); Iter; ++Iter)
-		{
-			if ((*Iter)->Contains(Handle))
-			{
-				bool Result = (*Iter)->GetAsynchronousResults(Handle, OutData, bDataWasBuilt);
-				if ((*Iter)->IsDone())
-				{
-					delete *Iter;
-					Iter.RemoveCurrent();
-				}
-				return Result;
-			}
-		}
-		return Super::GetAsynchronousResults(Handle, OutData, bDataWasBuilt);
-	}
-
-	virtual IDerivedDataRollup* StartRollup() override
-	{
-		FScopeLock ScopeLock(&SynchronizationObject);
-		FDerivedDataRollup* Result = new FDerivedDataRollup();
-		PendingRollups.Add(Result);
-		return Result;
-	}
-
-	virtual void EndRollup(IDerivedDataRollup*& InRollup) override
-	{
-		FDerivedDataRollup* Rollup = static_cast<FDerivedDataRollup*>(InRollup);
-		if (Rollup)
-		{
-			FScopeLock ScopeLock(&SynchronizationObject);
-			Rollup->Close();
-			if (Rollup->IsDone())
-			{
-				PendingRollups.Remove(Rollup);
-				delete Rollup;
-			}
-			InRollup = NULL; // set the pointer to NULL so it cannot be reused
-		}
-	}
-
-	virtual uint32 GetAsynchronous(const TCHAR* CacheKey, IDerivedDataRollup* Rollup = NULL) override
-	{
-		if (Rollup)
-		{
-			FScopeLock ScopeLock(&SynchronizationObject);
-			UE_LOG(LogDerivedDataCache, Verbose, TEXT("GetAsynchronous (Rollup) %s"), CacheKey);
-			const uint32 RollupHandle = NextHandle();
-			((FDerivedDataRollup*)Rollup)->Add(CacheKey, RollupHandle);
-			return RollupHandle;
-		}
-		return Super::GetAsynchronous(CacheKey, nullptr);
-	}
-
-	virtual void Put(const TCHAR* CacheKey, TArray<uint8>& Data, bool bPutEvenIfExists = false) override
-	{
-		FScopeLock ScopeLock(&SynchronizationObject);
-		for (TSet<FDerivedDataRollup*>::TIterator Iter(PendingRollups); Iter; ++Iter)
-		{
-			if ((*Iter)->Contains(CacheKey))
-			{
-				(*Iter)->Put(CacheKey, Data);
-				if ((*Iter)->IsDone())
-				{
-					delete *Iter;
-					Iter.RemoveCurrent();
-				}
-			}
-		}
-		return Super::Put(CacheKey, Data, bPutEvenIfExists);
-	}
-
-private:
-
-	/** Object used for synchronization via a scoped lock **/
-	FCriticalSection			SynchronizationObject;
-	/** Set of rollups **/
-	TSet<FDerivedDataRollup*>	PendingRollups;
-};
-
-/**
- * Singleton used both internally, and through the module.
- * We look at the commandline to check if we should disable rollups or not
- */
-FDerivedDataCache& InternalSingleton()
-{
-	static FDerivedDataCache* Singleton = NULL;
-	if (!Singleton)
-	{
-		if (FParse::Param(FCommandLine::Get(), TEXT("DDCNoRollups")))
-		{
-			UE_LOG(LogDerivedDataCache, Warning, TEXT("Rollups are disabled."));
-			static FDerivedDataCache SingletonInstance;
-			Singleton = &SingletonInstance;
-		}
-		else
-		{
-			static FDerivedDataCacheWithRollups SingletonInstance;
-			Singleton = &SingletonInstance;
-		}
-	}
-	return *Singleton;
-}
-
 /**
  * Module for the DDC
  */
 class FDerivedDataCacheModule : public IDerivedDataCacheModule
 {
-	/** Cached reference to DDC singleton, helpful to control singleton's lifetime. */
-	FDerivedDataCache* DDC;
-
 public:
 	virtual FDerivedDataCacheInterface& GetDDC() override
 	{
-		return InternalSingleton();
+		static FDerivedDataCache SingletonInstance;
+		return SingletonInstance;
 	}
 
 	virtual void StartupModule() override
 	{
-
-		// make sure DDC gets created early, previously it might have happened in ShutdownModule() (for PrintLeaks()) when it was already too late
-		DDC = static_cast< FDerivedDataCache* >( &GetDDC() );
+		GetDDC();
 	}
 
 	virtual void ShutdownModule() override
 	{
 		FDDCCleanup::Shutdown();
-
-		if (DDC)
-		{
-			DDC->PrintLeaks();
-		}
-	}
-
-	FDerivedDataCacheModule()
-		:	DDC(nullptr)
-	{
 	}
 };
 
 IMPLEMENT_MODULE( FDerivedDataCacheModule, DerivedDataCache);
-
diff --git a/Engine/Source/Developer/DerivedDataCache/Public/DerivedDataCacheInterface.h b/Engine/Source/Developer/DerivedDataCache/Public/DerivedDataCacheInterface.h
--- a/Engine/Source/Developer/DerivedDataCache/Public/DerivedDataCacheInterface.h
+++ b/Engine/Source/Developer/DerivedDataCache/Public/DerivedDataCacheInterface.h
@@ -8,17 +8,6 @@
 
 class FDerivedDataCacheUsageStats;
 
-/** 
- * Opaque class for rollup handling
-**/
-class IDerivedDataRollup
-{
-public:
-	virtual ~IDerivedDataRollup()
-	{
-	}
-};
-
 /** 
  * Interface for the derived data cache
  * This API is fully threadsafe (with the possible exception of the system interface: NotfiyBootComplete, etc).
@@ -135,18 +124,6 @@ public:
 	// Low Level Interface
 	//--------------------
 
-	/** 
-	 * Starts a rollup. Use this for GetAsynchronous calls, then call close on it
-	 * @return	a pointer to a rollup
-	**/
-	virtual IDerivedDataRollup* StartRollup() = 0;
-
-	/** 
-	 * Ends a rollup. Use this for GetAsynchronous calls, then call close on it
-	 * @param Rollup	reference to a pointer to the rollup to close...also sets the pointer to NULL
-	**/
-	virtual void EndRollup(IDerivedDataRollup*& Rollup) = 0;
-
 	/** 
 	 * Synchronously checks the cache and if the item is present, it returns the cached results, otherwise it returns false
 	 * @param	CacheKey	Key to identify the data
@@ -157,10 +134,9 @@ public:
 	/** 
 	 * Starts the async process of checking the cache and if the item is present, retrieving the cached results
 	 * @param	CacheKey	Key to identify the data
-	 * @param	Rollup		Rollup pointer, if this is part of a rollup
 	 * @return	a handle that can be used for PollAsynchronousCompletion, WaitAsynchronousCompletion and GetAsynchronousResults
 	**/
-	virtual uint32 GetAsynchronous(const TCHAR* CacheKey, IDerivedDataRollup* Rollup = NULL) = 0;
+	virtual uint32 GetAsynchronous(const TCHAR* CacheKey) = 0;
 
 	/** 
 	 * Puts data into the cache. This is fire-and-forget and typically asynchronous.
