Changing the DDC interface to use TArrayView instead of TArray for placing data into the DDC.
Updating the ArrayView with new comparison operators to have greater replacement for what TArray provides.

diff --git a/Engine/Source/Developer/DerivedDataCache/Private/DerivedDataBackendAsyncPutWrapper.cpp b/Engine/Source/Developer/DerivedDataCache/Private/DerivedDataBackendAsyncPutWrapper.cpp
--- a/Engine/Source/Developer/DerivedDataCache/Private/DerivedDataBackendAsyncPutWrapper.cpp
+++ b/Engine/Source/Developer/DerivedDataCache/Private/DerivedDataBackendAsyncPutWrapper.cpp
@@ -26,9 +26,9 @@ public:
 
 	/** Constructor
 	*/
-	FCachePutAsyncWorker(const TCHAR* InCacheKey, const TArray<uint8>* InData, FDerivedDataBackendInterface* InInnerBackend, bool InbPutEvenIfExists, FDerivedDataBackendInterface* InInflightCache, FThreadSet* InInFilesInFlight, FDerivedDataCacheUsageStats& InUsageStats)
+	FCachePutAsyncWorker(const TCHAR* InCacheKey, TArrayView<const uint8> InData, FDerivedDataBackendInterface* InInnerBackend, bool InbPutEvenIfExists, FDerivedDataBackendInterface* InInflightCache, FThreadSet* InInFilesInFlight, FDerivedDataCacheUsageStats& InUsageStats)
 		: CacheKey(InCacheKey)
-		, Data(*InData)
+		, Data(InData.GetData(), InData.Num())
 		, InnerBackend(InInnerBackend)
 		, InflightCache(InInflightCache)
 		, FilesInFlight(InInFilesInFlight)
@@ -135,7 +135,7 @@ bool FDerivedDataBackendAsyncPutWrapper::GetCachedData(const TCHAR* CacheKey, TA
 	return bSuccess;
 }
 
-void FDerivedDataBackendAsyncPutWrapper::PutCachedData(const TCHAR* CacheKey, TArray<uint8>& InData, bool bPutEvenIfExists)
+void FDerivedDataBackendAsyncPutWrapper::PutCachedData(const TCHAR* CacheKey, TArrayView<const uint8> InData, bool bPutEvenIfExists)
 {
 	COOK_STAT(auto Timer = PutSyncUsageStats.TimePut());
 	if (!InnerBackend->IsWritable())
@@ -157,7 +157,7 @@ void FDerivedDataBackendAsyncPutWrapper::PutCachedData(const TCHAR* CacheKey, TA
 		COOK_STAT(Timer.AddHit(InData.Num()));
 	}
 	FDerivedDataBackend::Get().AddToAsyncCompletionCounter(1);
-	(new FAutoDeleteAsyncTask<FCachePutAsyncWorker>(CacheKey, &InData, InnerBackend, bPutEvenIfExists, InflightCache.Get(), &FilesInFlight, UsageStats))->StartBackgroundTask();
+	(new FAutoDeleteAsyncTask<FCachePutAsyncWorker>(CacheKey, InData, InnerBackend, bPutEvenIfExists, InflightCache.Get(), &FilesInFlight, UsageStats))->StartBackgroundTask();
 }
 
 void FDerivedDataBackendAsyncPutWrapper::RemoveCachedData(const TCHAR* CacheKey, bool bTransient)
diff --git a/Engine/Source/Developer/DerivedDataCache/Private/DerivedDataBackendAsyncPutWrapper.h b/Engine/Source/Developer/DerivedDataCache/Private/DerivedDataBackendAsyncPutWrapper.h
--- a/Engine/Source/Developer/DerivedDataCache/Private/DerivedDataBackendAsyncPutWrapper.h
+++ b/Engine/Source/Developer/DerivedDataCache/Private/DerivedDataBackendAsyncPutWrapper.h
@@ -90,7 +90,7 @@ public:
 	 * @param	InData		Buffer containing the data to cache, can be destroyed after the call returns, immediately
 	 * @param	bPutEvenIfExists	If true, then do not attempt skip the put even if CachedDataProbablyExists returns true
 	 */
-	virtual void PutCachedData(const TCHAR* CacheKey, TArray<uint8>& InData, bool bPutEvenIfExists) override;
+	virtual void PutCachedData(const TCHAR* CacheKey, TArrayView<const uint8> InData, bool bPutEvenIfExists) override;
 
 	virtual void RemoveCachedData(const TCHAR* CacheKey, bool bTransient) override;
 
diff --git a/Engine/Source/Developer/DerivedDataCache/Private/DerivedDataBackendCorruptionWrapper.h b/Engine/Source/Developer/DerivedDataCache/Private/DerivedDataBackendCorruptionWrapper.h
--- a/Engine/Source/Developer/DerivedDataCache/Private/DerivedDataBackendCorruptionWrapper.h
+++ b/Engine/Source/Developer/DerivedDataCache/Private/DerivedDataBackendCorruptionWrapper.h
@@ -160,7 +160,7 @@ public:
 	 * @param	InData		Buffer containing the data to cache, can be destroyed after the call returns, immediately
 	 * @param	bPutEvenIfExists	If true, then do not attempt skip the put even if CachedDataProbablyExists returns true
 	 */
-	virtual void PutCachedData(const TCHAR* CacheKey, TArray<uint8>& InData, bool bPutEvenIfExists) override
+	virtual void PutCachedData(const TCHAR* CacheKey, TArrayView<const uint8> InData, bool bPutEvenIfExists) override
 	{
 		COOK_STAT(auto Timer = UsageStats.TimePut());
 		if (!InnerBackend->IsWritable())
diff --git a/Engine/Source/Developer/DerivedDataCache/Private/DerivedDataBackendVerifyWrapper.h b/Engine/Source/Developer/DerivedDataCache/Private/DerivedDataBackendVerifyWrapper.h
--- a/Engine/Source/Developer/DerivedDataCache/Private/DerivedDataBackendVerifyWrapper.h
+++ b/Engine/Source/Developer/DerivedDataCache/Private/DerivedDataBackendVerifyWrapper.h
@@ -48,6 +48,7 @@ public:
 		}
 		return false;
 	}
+
 	virtual bool GetCachedData(const TCHAR* CacheKey, TArray<uint8>& OutData) override
 	{
 		COOK_STAT(auto Timer = UsageStats.TimeGet());
@@ -70,7 +71,8 @@ public:
 		}
 		return false;
 	}
-	virtual void PutCachedData(const TCHAR* CacheKey, TArray<uint8>& InData, bool bPutEvenIfExists) override
+
+	virtual void PutCachedData(const TCHAR* CacheKey, TArrayView<const uint8> InData, bool bPutEvenIfExists) override
 	{
 		COOK_STAT(auto Timer = UsageStats.TimePut());
 		bool bAlreadyTested = false;
diff --git a/Engine/Source/Developer/DerivedDataCache/Private/DerivedDataCache.cpp b/Engine/Source/Developer/DerivedDataCache/Private/DerivedDataCache.cpp
--- a/Engine/Source/Developer/DerivedDataCache/Private/DerivedDataCache.cpp
+++ b/Engine/Source/Developer/DerivedDataCache/Private/DerivedDataCache.cpp
@@ -418,7 +418,7 @@ public:
 		return Handle;
 	}
 
-	virtual void Put(const TCHAR* CacheKey, TArray<uint8>& Data, const FStringView& DataContext, bool bPutEvenIfExists = false) override
+	virtual void Put(const TCHAR* CacheKey, TArrayView<const uint8> Data, const FStringView& DataContext, bool bPutEvenIfExists = false) override
 	{
 		DDC_SCOPE_CYCLE_COUNTER(DDC_Put);
 		UE_LOG(LogDerivedDataCache, Verbose, TEXT("Put %s from '%.*s'"), CacheKey, DataContext.Len(), DataContext.GetData());
diff --git a/Engine/Source/Developer/DerivedDataCache/Private/DerivedDataLimitKeyLengthWrapper.h b/Engine/Source/Developer/DerivedDataCache/Private/DerivedDataLimitKeyLengthWrapper.h
--- a/Engine/Source/Developer/DerivedDataCache/Private/DerivedDataLimitKeyLengthWrapper.h
+++ b/Engine/Source/Developer/DerivedDataCache/Private/DerivedDataLimitKeyLengthWrapper.h
@@ -133,7 +133,7 @@ public:
 	 * @param	InData		Buffer containing the data to cache, can be destroyed after the call returns, immediately
 	 * @param	bPutEvenIfExists	If true, then do not attempt skip the put even if CachedDataProbablyExists returns true
 	 */
-	virtual void PutCachedData(const TCHAR* CacheKey, TArray<uint8>& InData, bool bPutEvenIfExists) override
+	virtual void PutCachedData(const TCHAR* CacheKey, TArrayView<const uint8> InData, bool bPutEvenIfExists) override
 	{
 		COOK_STAT(auto Timer = UsageStats.TimePut());
 		if (!InnerBackend->IsWritable())
@@ -147,7 +147,7 @@ public:
 			InnerBackend->PutCachedData(CacheKey, InData, bPutEvenIfExists);
 			return;
 		}
-		TArray<uint8> Data(InData);
+		TArray<uint8> Data(InData.GetData(), InData.Num());
 		check(Data.Num());
 		int32 KeyLen = FCString::Strlen(CacheKey) + 1;
 		Data.AddUninitialized(KeyLen);
diff --git a/Engine/Source/Developer/DerivedDataCache/Private/FileSystemDerivedDataBackend.cpp b/Engine/Source/Developer/DerivedDataCache/Private/FileSystemDerivedDataBackend.cpp
--- a/Engine/Source/Developer/DerivedDataCache/Private/FileSystemDerivedDataBackend.cpp
+++ b/Engine/Source/Developer/DerivedDataCache/Private/FileSystemDerivedDataBackend.cpp
@@ -233,7 +233,7 @@ public:
 	 * @param	OutData		Buffer containing the data to cache, can be destroyed after the call returns, immediately
 	 * @param	bPutEvenIfExists	If true, then do not attempt skip the put even if CachedDataProbablyExists returns true
 	 */
-	virtual void PutCachedData(const TCHAR* CacheKey, TArray<uint8>& Data, bool bPutEvenIfExists) override
+	virtual void PutCachedData(const TCHAR* CacheKey, TArrayView<const uint8> Data, bool bPutEvenIfExists) override
 	{
 		COOK_STAT(auto Timer = UsageStats.TimePut());
 		check(!bFailed);
diff --git a/Engine/Source/Developer/DerivedDataCache/Private/HierarchicalDerivedDataBackend.h b/Engine/Source/Developer/DerivedDataCache/Private/HierarchicalDerivedDataBackend.h
--- a/Engine/Source/Developer/DerivedDataCache/Private/HierarchicalDerivedDataBackend.h
+++ b/Engine/Source/Developer/DerivedDataCache/Private/HierarchicalDerivedDataBackend.h
@@ -153,7 +153,7 @@ public:
 	 * @param	InData		Buffer containing the data to cache, can be destroyed after the call returns, immediately
 	 * @param	bPutEvenIfExists	If true, then do not attempt skip the put even if CachedDataProbablyExists returns true
 	 */
-	virtual void PutCachedData(const TCHAR* CacheKey, TArray<uint8>& InData, bool bPutEvenIfExists) override
+	virtual void PutCachedData(const TCHAR* CacheKey, TArrayView<const uint8> InData, bool bPutEvenIfExists) override
 	{
 		COOK_STAT(auto Timer = UsageStats.TimePut());
 		if (!bIsWritable)
diff --git a/Engine/Source/Developer/DerivedDataCache/Private/MemoryDerivedDataBackend.cpp b/Engine/Source/Developer/DerivedDataCache/Private/MemoryDerivedDataBackend.cpp
--- a/Engine/Source/Developer/DerivedDataCache/Private/MemoryDerivedDataBackend.cpp
+++ b/Engine/Source/Developer/DerivedDataCache/Private/MemoryDerivedDataBackend.cpp
@@ -59,7 +59,7 @@ bool FMemoryDerivedDataBackend::GetCachedData(const TCHAR* CacheKey, TArray<uint
 	return false;
 }
 
-void FMemoryDerivedDataBackend::PutCachedData(const TCHAR* CacheKey, TArray<uint8>& InData, bool bPutEvenIfExists)
+void FMemoryDerivedDataBackend::PutCachedData(const TCHAR* CacheKey, TArrayView<const uint8> InData, bool bPutEvenIfExists)
 {
 	COOK_STAT(auto Timer = UsageStats.TimePut());
 	FScopeLock ScopeLock(&SynchronizationObject);
diff --git a/Engine/Source/Developer/DerivedDataCache/Private/MemoryDerivedDataBackend.h b/Engine/Source/Developer/DerivedDataCache/Private/MemoryDerivedDataBackend.h
--- a/Engine/Source/Developer/DerivedDataCache/Private/MemoryDerivedDataBackend.h
+++ b/Engine/Source/Developer/DerivedDataCache/Private/MemoryDerivedDataBackend.h
@@ -47,21 +47,21 @@ public:
 	 * @param	InData		Buffer containing the data to cache, can be destroyed after the call returns, immediately
 	 * @param	bPutEvenIfExists	If true, then do not attempt skip the put even if CachedDataProbablyExists returns true
 	 */
-	virtual void PutCachedData(const TCHAR* CacheKey, TArray<uint8>& InData, bool bPutEvenIfExists) override;
+	virtual void PutCachedData(const TCHAR* CacheKey, TArrayView<const uint8> InData, bool bPutEvenIfExists) override;
 
 	virtual void RemoveCachedData(const TCHAR* CacheKey, bool bTransient) override;
 
 	/**
 	 * Save the cache to disk
 	 * @param	Filename	Filename to save
-	 * @return	true if file was saved sucessfully
+	 * @return	true if file was saved successfully
 	 */
 	bool SaveCache(const TCHAR* Filename);
 
 	/**
 	 * Load the cache from disk
 	 * @param	Filename	Filename to load
-	 * @return	true if file was loaded sucessfully
+	 * @return	true if file was loaded successfully
 	 */
 	bool LoadCache(const TCHAR* Filename);
 
@@ -81,9 +81,9 @@ private:
 	{
 		int32 Age;
 		TArray<uint8> Data;
-		FCacheValue(const TArray<uint8>& InData, int32 InAge = 0)
+		FCacheValue(TArrayView<const uint8> InData, int32 InAge = 0)
 			: Age(InAge)
-			, Data(InData)
+			, Data(InData.GetData(), InData.Num())
 		{
 		}
 	};
diff --git a/Engine/Source/Developer/DerivedDataCache/Private/PakFileDerivedDataBackend.cpp b/Engine/Source/Developer/DerivedDataCache/Private/PakFileDerivedDataBackend.cpp
--- a/Engine/Source/Developer/DerivedDataCache/Private/PakFileDerivedDataBackend.cpp
+++ b/Engine/Source/Developer/DerivedDataCache/Private/PakFileDerivedDataBackend.cpp
@@ -131,7 +131,7 @@ bool FPakFileDerivedDataBackend::GetCachedData(const TCHAR* CacheKey, TArray<uin
 	return false;
 }
 
-void FPakFileDerivedDataBackend::PutCachedData(const TCHAR* CacheKey, TArray<uint8>& InData, bool bPutEvenIfExists)
+void FPakFileDerivedDataBackend::PutCachedData(const TCHAR* CacheKey, TArrayView<const uint8> InData, bool bPutEvenIfExists)
 {
 	COOK_STAT(auto Timer = UsageStats.TimePut());
 	if (!bWriting || bClosed)
@@ -159,7 +159,8 @@ void FPakFileDerivedDataBackend::PutCachedData(const TCHAR* CacheKey, TArray<uin
 			else
 			{
 				COOK_STAT(Timer.AddHit(InData.Num()));
-				FileHandle->Serialize(InData.GetData(), int64(InData.Num()));
+				// NOTE: Gross that FArchive doesn't have a const version just for saving...
+				FileHandle->Serialize(const_cast<uint8*>(InData.GetData()), int64(InData.Num()));
 				UE_LOG(LogDerivedDataCache, Verbose, TEXT("FPakFileDerivedDataBackend: Put %s"), CacheKey);
 				CacheItems.Add(Key,FCacheValue(Offset, InData.Num(), Crc));
 			}
@@ -403,7 +404,7 @@ FCompressedPakFileDerivedDataBackend::FCompressedPakFileDerivedDataBackend(const
 {
 }
 
-void FCompressedPakFileDerivedDataBackend::PutCachedData(const TCHAR* CacheKey, TArray<uint8>& InData, bool bPutEvenIfExists)
+void FCompressedPakFileDerivedDataBackend::PutCachedData(const TCHAR* CacheKey, TArrayView<const uint8> InData, bool bPutEvenIfExists)
 {
 	int32 UncompressedSize = InData.Num();
 	int32 CompressedSize = FCompression::CompressMemoryBound(CompressionFormat, UncompressedSize, CompressionFlags);
diff --git a/Engine/Source/Developer/DerivedDataCache/Private/PakFileDerivedDataBackend.h b/Engine/Source/Developer/DerivedDataCache/Private/PakFileDerivedDataBackend.h
--- a/Engine/Source/Developer/DerivedDataCache/Private/PakFileDerivedDataBackend.h
+++ b/Engine/Source/Developer/DerivedDataCache/Private/PakFileDerivedDataBackend.h
@@ -56,7 +56,7 @@ public:
 	 * @param	InData		Buffer containing the data to cache, can be destroyed after the call returns, immediately
 	 * @param	bPutEvenIfExists	If true, then do not attempt skip the put even if CachedDataProbablyExists returns true
 	 */
-	virtual void PutCachedData(const TCHAR* CacheKey, TArray<uint8>& InData, bool bPutEvenIfExists) override;
+	virtual void PutCachedData(const TCHAR* CacheKey, TArrayView<const uint8> InData, bool bPutEvenIfExists) override;
 	virtual void RemoveCachedData(const TCHAR* CacheKey, bool bTransient) override;
 
 	/**
@@ -126,7 +126,7 @@ class FCompressedPakFileDerivedDataBackend : public FPakFileDerivedDataBackend
 public:
 	FCompressedPakFileDerivedDataBackend(const TCHAR* InFilename, bool bInWriting);
 
-	virtual void PutCachedData(const TCHAR* CacheKey, TArray<uint8>& InData, bool bPutEvenIfExists) override;
+	virtual void PutCachedData(const TCHAR* CacheKey, TArrayView<const uint8> InData, bool bPutEvenIfExists) override;
 	virtual bool GetCachedData(const TCHAR* CacheKey, TArray<uint8>& OutData) override;
 
 private:
diff --git a/Engine/Source/Developer/DerivedDataCache/Public/DerivedDataBackendInterface.h b/Engine/Source/Developer/DerivedDataCache/Public/DerivedDataBackendInterface.h
--- a/Engine/Source/Developer/DerivedDataCache/Public/DerivedDataBackendInterface.h
+++ b/Engine/Source/Developer/DerivedDataCache/Public/DerivedDataBackendInterface.h
@@ -65,7 +65,7 @@ public:
 	 * @param	InData				Buffer containing the data to cache, can be destroyed after the call returns, immediately
 	 * @param	bPutEvenIfExists	If true, then do not attempt skip the put even if CachedDataProbablyExists returns true
 	 */
-	virtual void PutCachedData(const TCHAR* CacheKey, TArray<uint8>& InData, bool bPutEvenIfExists) = 0;
+	virtual void PutCachedData(const TCHAR* CacheKey, TArrayView<const uint8> InData, bool bPutEvenIfExists) = 0;
 
 	/**
 	 * Remove data from cache (used in the event that corruption is detected at a higher level and possibly house keeping)
diff --git a/Engine/Source/Developer/DerivedDataCache/Public/DerivedDataCacheInterface.h b/Engine/Source/Developer/DerivedDataCache/Public/DerivedDataCacheInterface.h
--- a/Engine/Source/Developer/DerivedDataCache/Public/DerivedDataCacheInterface.h
+++ b/Engine/Source/Developer/DerivedDataCache/Public/DerivedDataCacheInterface.h
@@ -180,7 +180,7 @@ public:
 	 * @param	Data		Data to put in the cache under this key
 	 * @param	DataContext	A string used to describe the data being generated. Typically the path to the object that it is generated from is sufficient.
 	**/
-	virtual void Put(const TCHAR* CacheKey, TArray<uint8>& Data, const FStringView& DataContext, bool bPutEvenIfExists = false) = 0;
+	virtual void Put(const TCHAR* CacheKey, TArrayView<const uint8> Data, const FStringView& DataContext, bool bPutEvenIfExists = false) = 0;
 
 	/**
 	 * Hint that the data associated with the key is transient and may be optionally purged from the cache.
diff --git a/Engine/Source/Runtime/Core/Public/Containers/ArrayView.h b/Engine/Source/Runtime/Core/Public/Containers/ArrayView.h
--- a/Engine/Source/Runtime/Core/Public/Containers/ArrayView.h
+++ b/Engine/Source/Runtime/Core/Public/Containers/ArrayView.h
@@ -576,6 +576,61 @@ TArrayView<const T> MakeArrayView(std::initializer_list<T> List)
 	return TArrayView<const T>(List);
 }
 
+//////////////////////////////////////////////////////////////////////////
+
+/**
+ * Equality operator.
+ *
+ * @param Lhs Another ranged type to compare.
+ * @returns True if this array view's contents and the other ranged type match. False otherwise.
+ */
+template <
+	typename RangeType,
+	typename ElementType,
+	typename = decltype(ImplicitConv<const ElementType*>(GetData(DeclVal<RangeType&>())))
+>
+bool operator==(RangeType&& Lhs, TArrayView<ElementType> Rhs)
+{
+	auto Count = GetNum(Lhs);
+	return Count == Rhs.Num() && CompareItems(GetData(Lhs), Rhs.GetData(), Count);
+}
+
+template <
+	typename RangeType,
+	typename ElementType,
+	typename = decltype(ImplicitConv<const ElementType*>(GetData(DeclVal<RangeType&>())))
+>
+bool operator==(TArrayView<ElementType> Lhs, RangeType&& Rhs)
+{
+	return (Rhs == Lhs);
+}
+
+/**
+ * Inequality operator.
+ *
+ * @param Lhs Another ranged type to compare.
+ * @returns False if this array view's contents and the other ranged type match. True otherwise.
+ */
+template <
+	typename RangeType,
+	typename ElementType,
+	typename = decltype(ImplicitConv<const ElementType*>(GetData(DeclVal<RangeType&>())))
+>
+bool operator!=(RangeType&& Lhs, TArrayView<ElementType> Rhs)
+{
+	return !(Lhs == Rhs);
+}
+
+template <
+	typename RangeType,
+	typename ElementType,
+	typename = decltype(ImplicitConv<const ElementType*>(GetData(DeclVal<RangeType&>())))
+>
+bool operator!=(TArrayView<ElementType> Lhs, RangeType&& Rhs)
+{
+	return !(Rhs == Lhs);
+}
+
 template<typename InElementType, typename InAllocatorType>
 template<typename OtherElementType>
 FORCEINLINE TArray<InElementType, InAllocatorType>::TArray(const TArrayView<OtherElementType>& Other)
