Add experimental DDC backend which fetches objects from an S3 bucket. Bucket is populated by a job running periodically on build machines, and generates a JSON manifest used to determine files to download.

#rnx

diff --git a/Engine/Source/Developer/DerivedDataCache/DerivedDataCache.Build.cs b/Engine/Source/Developer/DerivedDataCache/DerivedDataCache.Build.cs
--- a/Engine/Source/Developer/DerivedDataCache/DerivedDataCache.Build.cs
+++ b/Engine/Source/Developer/DerivedDataCache/DerivedDataCache.Build.cs
@@ -9,6 +9,12 @@ public class DerivedDataCache : ModuleRules
 	{
 		PrivateDependencyModuleNames.Add("Core");
 
+		// Dependencies for S3 backend
+		PrivateDependencyModuleNames.AddRange(new string[] { "SSL", "Json" });
+		PrivateIncludePathModuleNames.Add("DesktopPlatform");
+		AddEngineThirdPartyPrivateStaticDependencies(Target, "libcurl");
+		AddEngineThirdPartyPrivateStaticDependencies(Target, "OpenSSL");
+
 		// Internal (NotForLicensees) module
 		var DDCUtilsModule = Path.Combine("Developer", "NotForLicensees", "DDCUtils", "DDCUtils.Build.cs");
 		if (File.Exists(DDCUtilsModule))
diff --git a/Engine/Source/Developer/DerivedDataCache/Private/DerivedDataBackends.cpp b/Engine/Source/Developer/DerivedDataCache/Private/DerivedDataBackends.cpp
--- a/Engine/Source/Developer/DerivedDataCache/Private/DerivedDataBackends.cpp
+++ b/Engine/Source/Developer/DerivedDataCache/Private/DerivedDataBackends.cpp
@@ -14,6 +14,7 @@
 #include "MemoryDerivedDataBackend.h"
 #include "DerivedDataBackendAsyncPutWrapper.h"
 #include "PakFileDerivedDataBackend.h"
+#include "S3DerivedDataBackend.h"
 #include "HierarchicalDerivedDataBackend.h"
 #include "DerivedDataLimitKeyLengthWrapper.h"
 #include "DerivedDataBackendCorruptionWrapper.h"
@@ -203,6 +204,10 @@ public:
 				{
 					ParsedNode = ParsePak( NodeName, *Entry, true );
 				}
+				else if (NodeType == TEXT("S3"))
+				{
+					ParsedNode = ParseS3Cache(NodeName, *Entry);
+				}
 			}
 		}
 
@@ -598,6 +603,63 @@ public:
 		return DataCache;
 	}
 
+
+	/**
+	 * Creates an S3 data cache interface.
+	 */
+	FDerivedDataBackendInterface* ParseS3Cache(const TCHAR* NodeName, const TCHAR* Entry)
+	{
+		FString ManifestPath;
+		if (!FParse::Value(Entry, TEXT("Manifest="), ManifestPath))
+		{
+			UE_LOG(LogDerivedDataCache, Error, TEXT("Node %s does not specify 'Manifest'."), NodeName);
+			return nullptr;
+		}
+
+		FString BaseUrl;
+		if (!FParse::Value(Entry, TEXT("BaseUrl="), BaseUrl))
+		{
+			UE_LOG(LogDerivedDataCache, Error, TEXT("Node %s does not specify 'BaseUrl'."), NodeName);
+			return nullptr;
+		}
+
+		FString CanaryObjectKey;
+		FParse::Value(Entry, TEXT("Canary="), CanaryObjectKey);
+
+		FString Region;
+		if (!FParse::Value(Entry, TEXT("Region="), Region))
+		{
+			UE_LOG(LogDerivedDataCache, Error, TEXT("Node %s does not specify 'Region'."), NodeName);
+			return nullptr;
+		}
+
+		FString AccessKey;
+		if (!FParse::Value(Entry, TEXT("AccessKey="), AccessKey))
+		{
+			UE_LOG(LogDerivedDataCache, Error, TEXT("Node %s does not specify 'AccessKey'."), NodeName);
+			return nullptr;
+		}
+
+		FString SecretKey;
+		if (!FParse::Value(Entry, TEXT("SecretKey="), SecretKey))
+		{
+			UE_LOG(LogDerivedDataCache, Error, TEXT("Node %s does not specify 'SecretKey'."), NodeName);
+			return nullptr;
+		}
+
+		FS3DerivedDataBackend* Backend = new FS3DerivedDataBackend(*ManifestPath, *BaseUrl, *Region, *CanaryObjectKey, *AccessKey, *SecretKey);
+		if (!Backend->IsUsable())
+		{
+			UE_LOG(LogDerivedDataCache, Warning, TEXT("%s could not contact the service (%s), will not use it."), NodeName, *BaseUrl);
+			delete Backend;
+			return nullptr;
+		}
+
+		// Insert the backend corruption wrapper. Since the filesystem already uses this, and we're recycling the data with the trailer intact, we need to use it for the S3 cache too.
+		UE_LOG(LogDerivedDataCache, Display, TEXT("Using %s S3 backend at %s"), *Region, *BaseUrl);
+		return new FDerivedDataBackendCorruptionWrapper(Backend);
+	}
+
 	/**
 	 * Creates Boot data cache interface from ini settings.
 	 *
diff --git a/Engine/Source/Developer/DerivedDataCache/Private/FileSystemDerivedDataBackend.cpp b/Engine/Source/Developer/DerivedDataCache/Private/FileSystemDerivedDataBackend.cpp
--- a/Engine/Source/Developer/DerivedDataCache/Private/FileSystemDerivedDataBackend.cpp
+++ b/Engine/Source/Developer/DerivedDataCache/Private/FileSystemDerivedDataBackend.cpp
@@ -29,8 +29,18 @@
 #endif // PLATFORM_LINUX
 #define MAX_CACHE_EXTENTION_LEN (4)
 
-
-
+FString BuildPathForCacheKey(const TCHAR* CacheKey)
+{
+	FString Key = FString(CacheKey).ToUpper();
+	for (int32 i = 0; i < Key.Len(); i++)
+	{
+		check(FChar::IsAlnum(Key[i]) || FChar::IsUnderscore(Key[i]) || Key[i] == L'$');
+	}
+	uint32 Hash = FCrc::StrCrc_DEPRECATED(*Key);
+	// this creates a tree of 1000 directories
+	FString HashPath = FString::Printf(TEXT("%1d/%1d/%1d/"), (Hash / 100) % 10, (Hash / 10) % 10, Hash % 10);
+	return HashPath / Key + TEXT(".udd");
+}
 
 /** 
  * Cache server that uses the OS filesystem
@@ -367,15 +377,7 @@ private:
 	 */
 	FString BuildFilename(const TCHAR* CacheKey)
 	{
-		FString Key = FString(CacheKey).ToUpper();
-		for (int32 i = 0; i < Key.Len(); i++)
-		{
-			check(FChar::IsAlnum(Key[i]) || FChar::IsUnderscore(Key[i]) || Key[i] == L'$');
-		}
-		uint32 Hash = FCrc::StrCrc_DEPRECATED(*Key);
-		// this creates a tree of 1000 directories
-		FString HashPath = FString::Printf(TEXT("%1d/%1d/%1d/"),(Hash/100)%10,(Hash/10)%10,Hash%10);
-		return CachePath / HashPath / Key + TEXT(".udd");
+		return CachePath / BuildPathForCacheKey(CacheKey);
 	}
 
 	/** Base path we are storing the cache files in. **/
diff --git a/Engine/Source/Developer/DerivedDataCache/Private/S3DerivedDataBackend.cpp b/Engine/Source/Developer/DerivedDataCache/Private/S3DerivedDataBackend.cpp
new file mode 100644
--- /dev/null
+++ b/Engine/Source/Developer/DerivedDataCache/Private/S3DerivedDataBackend.cpp
@@ -0,0 +1,995 @@
+// Copyright Epic Games, Inc. All Rights Reserved.
+
+#include "S3DerivedDataBackend.h"
+
+#if PLATFORM_WINDOWS || PLATFORM_HOLOLENS
+	#include "Windows/WindowsHWrapper.h"
+	#include "Windows/AllowWindowsPlatformTypes.h"
+#endif
+#include "curl/curl.h"
+#if PLATFORM_WINDOWS || PLATFORM_HOLOLENS
+	#include "Windows/HideWindowsPlatformTypes.h"
+#endif
+#if WITH_SSL
+	#include "Ssl.h"
+	#include <openssl/ssl.h>
+	#include <openssl/sha.h>
+	#include <openssl/hmac.h>
+#endif
+#include "Misc/Base64.h"
+#include "Misc/Paths.h"
+#include "Misc/SecureHash.h"
+#include "Misc/FileHelper.h"
+#include "Misc/FeedbackContext.h"
+#include "Misc/OutputDeviceRedirector.h"
+#include "Async/ParallelFor.h"
+#include "Serialization/MemoryReader.h"
+#include "Dom/JsonObject.h"
+#include "Serialization/JsonReader.h"
+#include "Serialization/JsonSerializer.h"
+#include "ProfilingDebugging/CpuProfilerTrace.h"
+#include "ProfilingDebugging/CountersTrace.h"
+#include "HAL/PlatformFile.h"
+#include "HAL/PlatformFilemanager.h"
+#include "HAL/Runnable.h"
+#include "DesktopPlatformModule.h"
+
+#include "picosha2.h"
+
+#define S3DDC_BACKEND_WAIT_INTERVAL 0.01f
+#define S3DDC_HTTP_REQUEST_TIMEOUT_SECONDS 30L
+#define S3DDC_HTTP_REQUEST_TIMOUT_ENABLED 1
+#define S3DDC_REQUEST_POOL_SIZE 16
+#define S3DDC_MAX_FAILED_LOGIN_ATTEMPTS 16
+#define S3DDC_MAX_ATTEMPTS 4
+#define S3DDC_MAX_BUFFER_RESERVE 104857600u
+
+TRACE_DECLARE_INT_COUNTER(S3DDC_Exist, TEXT("S3DDC Exist"));
+TRACE_DECLARE_INT_COUNTER(S3DDC_ExistHit, TEXT("S3DDC Exist Hit"));
+TRACE_DECLARE_INT_COUNTER(S3DDC_Get, TEXT("S3DDC Get"));
+TRACE_DECLARE_INT_COUNTER(S3DDC_GetHit, TEXT("S3DDC Get Hit"));
+TRACE_DECLARE_INT_COUNTER(S3DDC_BytesRecieved, TEXT("S3DDC Bytes Recieved"));
+
+FString BuildPathForCacheKey(const TCHAR* CacheKey);
+
+class FStringAnsi
+{
+public:
+	FStringAnsi()
+	{
+		Inner.Add(0);
+	}
+
+	FStringAnsi(const ANSICHAR* Text)
+	{
+		Inner.Append(Text, FCStringAnsi::Strlen(Text) + 1);
+	}
+
+	void Append(ANSICHAR Character)
+	{
+		Inner[Inner.Num() - 1] = Character;
+		Inner.Add(0);
+	}
+
+	void Append(const FStringAnsi& Other)
+	{
+		Inner.RemoveAt(Inner.Num() - 1);
+		Inner.Append(Other.Inner);
+	}
+
+	void Append(const ANSICHAR* Text)
+	{
+		Inner.RemoveAt(Inner.Num() - 1);
+		Inner.Append(Text, FCStringAnsi::Strlen(Text) + 1);
+	}
+
+	void Append(const ANSICHAR* Start, const ANSICHAR* End)
+	{
+		Inner.RemoveAt(Inner.Num() - 1);
+		Inner.Append(Start, End - Start);
+		Inner.Add(0);
+	}
+
+	static FStringAnsi Printf(const ANSICHAR* Format, ...)
+	{
+		ANSICHAR Buffer[1024];
+		GET_VARARGS_ANSI(Buffer, UE_ARRAY_COUNT(Buffer), UE_ARRAY_COUNT(Buffer) - 1, Format, Format);
+		return Buffer;
+	}
+
+	FString ToWideString() const
+	{
+		return ANSI_TO_TCHAR(Inner.GetData());
+	}
+
+	const ANSICHAR* operator*() const
+	{
+		return Inner.GetData();
+	}
+
+	int32 Len() const
+	{
+		return Inner.Num() - 1;
+	}
+
+private:
+	TArray<ANSICHAR> Inner;
+};
+
+struct FSHA256
+{
+	uint8 Digest[32];
+
+	FStringAnsi ToString() const
+	{
+		ANSICHAR Buffer[65];
+		for (int Idx = 0; Idx < 32; Idx++)
+		{
+			FCStringAnsi::Sprintf(Buffer + (Idx * 2), "%02x", Digest[Idx]);
+		}
+		return Buffer;
+	}
+};
+
+FSHA256 Sha256(const uint8* Input, size_t InputLen)
+{
+#if WITH_SSL
+	FSHA256 Output;
+	SHA256(Input, InputLen, Output.Digest);
+	return Output;
+#else
+	UE_LOG(LogDerivedDataCache, Fatal, TEXT("SSL functionality is not available. Cannot use S3 backend."));
+#endif
+}
+
+FSHA256 HmacSha256(const uint8* Input, size_t InputLen, const uint8* Key, size_t KeyLen)
+{
+#if WITH_SSL
+	FSHA256 Output;
+	unsigned int OutputLen = 0;
+	HMAC(EVP_sha256(), Key, KeyLen, (const unsigned char*)Input, InputLen, Output.Digest, &OutputLen);
+	return Output;
+#else
+	UE_LOG(LogDerivedDataCache, Fatal, TEXT("SSL functionality is not available. Cannot use S3 backend."));
+#endif
+}
+
+FSHA256 HmacSha256(const FStringAnsi& Input, const uint8* Key, size_t KeyLen)
+{
+	return HmacSha256((const uint8*)*Input, (size_t)Input.Len(), Key, KeyLen);
+}
+
+FSHA256 HmacSha256(const char* Input, const uint8* Key, size_t KeyLen)
+{
+	return HmacSha256((const uint8*)Input, (size_t)FCStringAnsi::Strlen(Input), Key, KeyLen);
+}
+
+bool IsSuccessfulHttpResponse(long ResponseCode)
+{
+	return (ResponseCode >= 200 && ResponseCode <= 299);
+}
+
+struct IRequestCallback
+{
+	virtual ~IRequestCallback() { }
+	virtual bool Update(int NumBytes, int TotalBytes) = 0;
+};
+
+/**
+ * Minimal HTTP request type wrapping CURL without the need for managers. This request
+ * is written to allow reuse of request objects, in order to allow connections to be reused.
+ *
+ * CURL has a global library initialization (curl_global_init). We rely on this happening in 
+ * the Online/HTTP library which is a dependency on this module.
+ */
+class FS3DerivedDataBackend::FRequest
+{
+public:
+	FRequest(const ANSICHAR* InRegion, const ANSICHAR* InAccessKey, const ANSICHAR* InSecretKey)
+		: Region(InRegion)
+		, AccessKey(InAccessKey)
+		, SecretKey(InSecretKey)
+	{
+		Curl = curl_easy_init();
+	}
+
+	~FRequest()
+	{
+		curl_easy_cleanup(Curl);
+	}
+
+	/**
+	 * Performs the request, blocking until finished.
+	 * @param Url HTTP URL to fetch
+	 * @param Callback Object used to convey state to/from the operation
+	 * @param Buffer Optional buffer to directly receive the result of the request. 
+	 * If unset the response body will be stored in the request.
+	 */
+	long PerformBlocking(const ANSICHAR* Url, IRequestCallback* Callback, TArray<uint8>& OutResponseBody, FOutputDevice* Log)
+	{
+		TRACE_CPUPROFILER_EVENT_SCOPE(S3DDC_CurlPerform);
+
+		// Find the host from the URL
+		const ANSICHAR* ProtocolEnd = FCStringAnsi::Strchr(Url, ':');
+		check(ProtocolEnd != nullptr && *(ProtocolEnd + 1) == '/' && *(ProtocolEnd + 2) == '/');
+
+		const ANSICHAR* UrlHost = ProtocolEnd + 3;
+		const ANSICHAR* UrlHostEnd = FCStringAnsi::Strchr(UrlHost, '/');
+		check(UrlHostEnd != nullptr);
+
+		FStringAnsi Host;
+		Host.Append(UrlHost, UrlHostEnd);
+
+		// Get the header strings
+		FDateTime Timestamp = FDateTime::UtcNow();// FDateTime(2015, 9, 15, 12, 45, 0);
+		FStringAnsi TimeString = FStringAnsi::Printf("%04d%02d%02dT%02d%02d%02dZ", Timestamp.GetYear(), Timestamp.GetMonth(), Timestamp.GetDay(), Timestamp.GetHour(), Timestamp.GetMinute(), Timestamp.GetSecond());
+
+		// Payload string
+		FStringAnsi EmptyPayloadSha256 = Sha256(nullptr, 0).ToString();
+
+		// Create the headers
+		curl_slist* CurlHeaders = nullptr;
+		CurlHeaders = curl_slist_append(CurlHeaders, *FStringAnsi::Printf("Host: %s", *Host));
+		CurlHeaders = curl_slist_append(CurlHeaders, *FStringAnsi::Printf("x-amz-content-sha256: %s", *EmptyPayloadSha256));
+		CurlHeaders = curl_slist_append(CurlHeaders, *FStringAnsi::Printf("x-amz-date: %s", *TimeString));
+		CurlHeaders = curl_slist_append(CurlHeaders, *GetAuthorizationHeader("GET", UrlHostEnd, "", CurlHeaders, *TimeString, *EmptyPayloadSha256));
+
+		// Create the callback data
+		FStringAnsi Domain;
+		Domain.Append(Url, UrlHostEnd);
+
+		FCallbackData CallbackData(Domain, OutResponseBody);
+
+		// Setup the request
+		curl_easy_reset(Curl);
+		curl_easy_setopt(Curl, CURLOPT_FOLLOWLOCATION, 1L);
+		curl_easy_setopt(Curl, CURLOPT_NOSIGNAL, 1L);
+		curl_easy_setopt(Curl, CURLOPT_HTTPGET, 1L);
+		curl_easy_setopt(Curl, CURLOPT_URL, Url);
+		curl_easy_setopt(Curl, CURLOPT_ACCEPT_ENCODING, "gzip");
+#if S3DDC_HTTP_REQUEST_TIMOUT_ENABLED
+		curl_easy_setopt(Curl, CURLOPT_CONNECTTIMEOUT, S3DDC_HTTP_REQUEST_TIMEOUT_SECONDS);
+#endif
+
+		// Headers
+		curl_easy_setopt(Curl, CURLOPT_HTTPHEADER, CurlHeaders);
+		curl_easy_setopt(Curl, CURLOPT_HEADERDATA, CurlHeaders);
+
+		// Progress
+		curl_easy_setopt(Curl, CURLOPT_NOPROGRESS, 0);
+		curl_easy_setopt(Curl, CURLOPT_XFERINFODATA, Callback);
+		curl_easy_setopt(Curl, CURLOPT_XFERINFOFUNCTION, &FRequest::StaticStatusFn);
+
+		// Response
+		curl_easy_setopt(Curl, CURLOPT_HEADERDATA, &CallbackData);
+		curl_easy_setopt(Curl, CURLOPT_HEADERFUNCTION, &FRequest::StaticWriteHeaderFn);
+		curl_easy_setopt(Curl, CURLOPT_WRITEDATA, &CallbackData);
+		curl_easy_setopt(Curl, CURLOPT_WRITEFUNCTION, StaticWriteBodyFn);
+
+#if WITH_SSL
+		// SSL options
+		curl_easy_setopt(Curl, CURLOPT_USE_SSL, CURLUSESSL_ALL);
+		curl_easy_setopt(Curl, CURLOPT_SSL_VERIFYPEER, 1);
+		curl_easy_setopt(Curl, CURLOPT_SSL_VERIFYHOST, 1);
+		curl_easy_setopt(Curl, CURLOPT_SSLCERTTYPE, "PEM");
+
+		// SSL certification verification
+		curl_easy_setopt(Curl, CURLOPT_CAINFO, nullptr);
+		curl_easy_setopt(Curl, CURLOPT_SSL_CTX_FUNCTION, *sslctx_function);
+		curl_easy_setopt(Curl, CURLOPT_SSL_CTX_DATA, &CallbackData);
+#endif
+
+		// Send the request
+		CURLcode CurlResult = curl_easy_perform(Curl);
+
+		// Free the headers object
+		curl_slist_free_all(CurlHeaders);
+		curl_easy_setopt(Curl, CURLOPT_HEADERDATA, nullptr);
+
+		// Get the response code
+		long ResponseCode = 0;
+		if (CurlResult == CURLE_OK)
+		{
+			CurlResult = curl_easy_getinfo(Curl, CURLINFO_RESPONSE_CODE, &ResponseCode);
+		}
+		if (CurlResult != CURLE_OK)
+		{
+			Log->Logf(ELogVerbosity::Error, TEXT("Error while connecting to %s: %s"), ANSI_TO_TCHAR(Url), ANSI_TO_TCHAR(curl_easy_strerror(CurlResult)));
+			return 500;
+		}
+
+		// Print any diagnostic output
+		if (!(ResponseCode >= 200 && ResponseCode <= 299))
+		{
+			Log->Logf(ELogVerbosity::Error, TEXT("Download failed for %s (response %d):\n%s\n%s"), ANSI_TO_TCHAR(Url), ResponseCode, *CallbackData.ResponseHeader.ToWideString(), ANSI_TO_TCHAR((const ANSICHAR*)OutResponseBody.GetData()));
+		}
+		return ResponseCode;
+	}
+
+private:
+	struct FCallbackData
+	{
+		const FStringAnsi& Domain;
+		FStringAnsi ResponseHeader;
+		TArray<uint8>& ResponseBody;
+
+		FCallbackData(const FStringAnsi& InDomain, TArray<uint8>& InResponseBody)
+			: Domain(InDomain)
+			, ResponseBody(InResponseBody)
+		{
+		}
+	};
+
+	CURL* Curl;
+	FStringAnsi Region;
+	FStringAnsi AccessKey;
+	FStringAnsi SecretKey;
+
+	FStringAnsi GetAuthorizationHeader(const ANSICHAR* Verb, const ANSICHAR* RelativeUrl, const ANSICHAR* QueryString, const curl_slist* Headers, const ANSICHAR* Timestamp, const ANSICHAR* Digest)
+	{
+		// Create the canonical list of headers
+		FStringAnsi CanonicalHeaders;
+		for (const curl_slist* Header = Headers; Header != nullptr; Header = Header->next)
+		{
+			const ANSICHAR* Colon = FCStringAnsi::Strchr(Header->data, ':');
+			if (Colon != nullptr)
+			{
+				for (const ANSICHAR* Char = Header->data; Char != Colon; Char++)
+				{
+					CanonicalHeaders.Append(tolower(*Char));
+				}
+				CanonicalHeaders.Append(':');
+
+				const ANSICHAR* Value = Colon + 1;
+				while (*Value == ' ')
+				{
+					Value++;
+				}
+				for (; *Value != 0; Value++)
+				{
+					CanonicalHeaders.Append(*Value);
+				}
+				CanonicalHeaders.Append('\n');
+			}
+		}
+
+		// Create the list of signed headers
+		FStringAnsi SignedHeaders;
+		for (const curl_slist* Header = Headers; Header != nullptr; Header = Header->next)
+		{
+			const ANSICHAR* Colon = FCStringAnsi::Strchr(Header->data, ':');
+			if (Colon != nullptr)
+			{
+				if (SignedHeaders.Len() > 0)
+				{
+					SignedHeaders.Append(';');
+				}
+				for (const ANSICHAR* Char = Header->data; Char != Colon; Char++)
+				{
+					SignedHeaders.Append(tolower(*Char));
+				}
+			}
+		}
+
+		// Build the canonical request string
+		FStringAnsi CanonicalRequest;
+		CanonicalRequest.Append(Verb);
+		CanonicalRequest.Append('\n');
+		CanonicalRequest.Append(RelativeUrl);
+		CanonicalRequest.Append('\n');
+		CanonicalRequest.Append(QueryString);
+		CanonicalRequest.Append('\n');
+		CanonicalRequest.Append(CanonicalHeaders);
+		CanonicalRequest.Append('\n');
+		CanonicalRequest.Append(SignedHeaders);
+		CanonicalRequest.Append('\n');
+		CanonicalRequest.Append(Digest);
+
+		// Get the date
+		FStringAnsi DateString;
+		for (int32 Idx = 0; Timestamp[Idx] != 0 && Timestamp[Idx] != 'T'; Idx++)
+		{
+			DateString.Append(Timestamp[Idx]);
+		}
+
+		// Generate the signature key
+		FStringAnsi Key = FStringAnsi::Printf("AWS4%s", *SecretKey);
+
+		FSHA256 DateHash = HmacSha256(DateString, (const uint8*)*Key, Key.Len());
+		FSHA256 RegionHash = HmacSha256(Region, DateHash.Digest, sizeof(DateHash.Digest));
+		FSHA256 ServiceHash = HmacSha256("s3", RegionHash.Digest, sizeof(RegionHash.Digest));
+		FSHA256 SigningKeyHash = HmacSha256("aws4_request", ServiceHash.Digest, sizeof(ServiceHash.Digest));
+
+		// Calculate the signature
+		FStringAnsi DateRequest = FStringAnsi::Printf("%s/%s/s3/aws4_request", *DateString, *Region);
+		FStringAnsi CanonicalRequestSha256 = Sha256((const uint8*)*CanonicalRequest, CanonicalRequest.Len()).ToString();
+		FStringAnsi StringToSign = FStringAnsi::Printf("AWS4-HMAC-SHA256\n%s\n%s\n%s", Timestamp, *DateRequest, *CanonicalRequestSha256);
+		FStringAnsi Signature = HmacSha256(*StringToSign, SigningKeyHash.Digest, sizeof(SigningKeyHash.Digest)).ToString();
+
+		// Format the final header
+		return FStringAnsi::Printf("Authorization: AWS4-HMAC-SHA256 Credential=%s/%s, SignedHeaders=%s, Signature=%s", *AccessKey, *DateRequest, *SignedHeaders, *Signature);
+	}
+
+	/**
+	 * Returns the response buffer as a string. Note that is the request is performed
+	 * with an external buffer as target buffer this string will be empty.
+	 */
+	static FString GetResponseAsString(const TArray<uint8>& Buffer)
+	{
+		FUTF8ToTCHAR TCHARData(reinterpret_cast<const ANSICHAR*>(Buffer.GetData()), Buffer.Num());
+		return FString(TCHARData.Length(), TCHARData.Get());
+	}
+
+	static int StaticStatusFn(void* Ptr, curl_off_t TotalDownloadSize, curl_off_t CurrentDownloadSize, curl_off_t TotalUploadSize, curl_off_t CurrentUploadSize)
+	{
+		IRequestCallback* Callback = (IRequestCallback*)Ptr;
+		if (Callback != nullptr)
+		{
+			return Callback->Update(CurrentDownloadSize, TotalDownloadSize)? 0 : 1;
+		}
+		return 0;
+	}
+
+	static size_t StaticWriteHeaderFn(void* Ptr, size_t SizeInBlocks, size_t BlockSizeInBytes, void* UserData)
+	{
+		const size_t WriteSize = SizeInBlocks * BlockSizeInBytes;
+		if (WriteSize > 0)
+		{
+			FCallbackData* CallbackData = static_cast<FCallbackData*>(UserData);
+			CallbackData->ResponseHeader.Append((const ANSICHAR*)Ptr, (const ANSICHAR*)Ptr + WriteSize);
+			return WriteSize;
+		}
+		return 0;
+	}
+
+	static size_t StaticWriteBodyFn(void* Ptr, size_t SizeInBlocks, size_t BlockSizeInBytes, void* UserData)
+	{
+		const size_t WriteSize = SizeInBlocks * BlockSizeInBytes;
+		if (WriteSize > 0)
+		{
+			FCallbackData* CallbackData = static_cast<FCallbackData*>(UserData);
+
+			// If this is the first part of the body being received, try to reserve 
+			// memory if content length is defined in the header.
+			if (CallbackData->ResponseBody.Num() == 0)
+			{
+				static const ANSICHAR Prefix[] = "Content-Length: ";
+				static const size_t PrefixLen = UE_ARRAY_COUNT(Prefix) - 1;
+
+				for(const ANSICHAR* Header = *CallbackData->ResponseHeader;;Header++)
+				{
+					// Check this header
+					if (FCStringAnsi::Strnicmp(Header, Prefix, PrefixLen) == 0)
+					{
+						size_t ContentLength = (size_t)atol(Header + PrefixLen);
+						if (ContentLength > 0u && ContentLength < S3DDC_MAX_BUFFER_RESERVE)
+						{
+							CallbackData->ResponseBody.Reserve(ContentLength);
+						}
+						break;
+					}
+
+					// Move to the next string
+					Header = FCStringAnsi::Strchr(Header, '\n');
+					if (Header == nullptr)
+					{
+						break;
+					}
+				}
+			}
+
+			// Write to the target buffer
+			CallbackData->ResponseBody.Append((const uint8*)Ptr, WriteSize);
+			return WriteSize;
+		}
+		return 0;
+	}
+
+#if WITH_SSL
+	static int SslCertVerify(int PreverifyOk, X509_STORE_CTX* Context)
+	{
+		if (PreverifyOk == 1)
+		{
+			SSL* Handle = static_cast<SSL*>(X509_STORE_CTX_get_ex_data(Context, SSL_get_ex_data_X509_STORE_CTX_idx()));
+			check(Handle);
+
+			SSL_CTX* SslContext = SSL_get_SSL_CTX(Handle);
+			check(SslContext);
+
+			FCallbackData* CallbackData = reinterpret_cast<FCallbackData*>(SSL_CTX_get_app_data(SslContext));
+			check(CallbackData);
+
+			if (!FSslModule::Get().GetCertificateManager().VerifySslCertificates(Context, *CallbackData->Domain))
+			{
+				PreverifyOk = 0;
+			}
+		}
+
+		return PreverifyOk;
+	}
+
+	static CURLcode sslctx_function(CURL* curl, void* sslctx, void* parm)
+	{
+		SSL_CTX* Context = static_cast<SSL_CTX*>(sslctx);
+		const ISslCertificateManager& CertificateManager = FSslModule::Get().GetCertificateManager();
+
+		CertificateManager.AddCertificatesToSslContext(Context);
+		SSL_CTX_set_verify(Context, SSL_CTX_get_verify_mode(Context), SslCertVerify);
+		SSL_CTX_set_app_data(Context, parm);
+
+		/* all set to go */
+		return CURLE_OK;
+	}
+
+#endif //#if WITH_SSL
+};
+
+//----------------------------------------------------------------------------------------------------------
+// FS3DerivedDataBackend::FRequestPool
+//----------------------------------------------------------------------------------------------------------
+
+class FS3DerivedDataBackend::FRequestPool
+{
+public:
+	FRequestPool(const ANSICHAR* Region, const ANSICHAR* AccessKey, const ANSICHAR* SecretKey)
+	{
+		Pool.SetNum(S3DDC_REQUEST_POOL_SIZE);
+		for (uint8 i = 0; i < Pool.Num(); ++i)
+		{
+			Pool[i].Usage = 0u;
+			Pool[i].Request = new FRequest(Region, AccessKey, SecretKey);
+		}
+	}
+
+	~FRequestPool()
+	{
+		for (uint8 i = 0; i < Pool.Num(); ++i)
+		{
+			// No requests should be in use by now.
+			check(Pool[i].Usage.Load(EMemoryOrder::Relaxed) == 0u);
+			delete Pool[i].Request;
+		}
+	}
+
+	long Download(const TCHAR* Url, IRequestCallback* Callback, TArray<uint8>& OutData, FOutputDevice* Log)
+	{
+		FRequest* Request = WaitForFreeRequest();
+		long ResponseCode = Request->PerformBlocking(TCHAR_TO_ANSI(Url), Callback, OutData, Log);
+		ReleaseRequestToPool(Request);
+		return ResponseCode;
+	}
+
+private:
+	struct FEntry
+	{
+		TAtomic<uint8> Usage;
+		FRequest* Request;
+	};
+
+	TArray<FEntry> Pool;
+
+	FRequest* WaitForFreeRequest()
+	{
+		TRACE_CPUPROFILER_EVENT_SCOPE(S3DDC_WaitForConnPool);
+		while (true)
+		{
+			for (uint8 i = 0; i < Pool.Num(); ++i)
+			{
+				if (!Pool[i].Usage.Load(EMemoryOrder::Relaxed))
+				{
+					uint8 Expected = 0u;
+					if (Pool[i].Usage.CompareExchange(Expected, 1u))
+					{
+						return Pool[i].Request;
+					}
+				}
+			}
+			FPlatformProcess::Sleep(S3DDC_BACKEND_WAIT_INTERVAL);
+		}
+	}
+
+	void ReleaseRequestToPool(FRequest* Request)
+	{
+		for (uint8 i = 0; i < Pool.Num(); ++i)
+		{
+			if (Pool[i].Request == Request)
+			{
+				uint8 Expected = 1u;
+				Pool[i].Usage.CompareExchange(Expected, 0u);
+				return;
+			}
+		}
+		check(false);
+	}
+};
+
+//----------------------------------------------------------------------------------------------------------
+// FS3DerivedDataBackend::FBundleEntry
+//----------------------------------------------------------------------------------------------------------
+
+struct FS3DerivedDataBackend::FBundleEntry
+{
+	int64 Offset;
+	int32 Length;
+};
+
+//----------------------------------------------------------------------------------------------------------
+// FS3DerivedDataBackend::FBundle
+//----------------------------------------------------------------------------------------------------------
+
+struct FS3DerivedDataBackend::FBundle
+{
+	FString Name;
+	FString ObjectKey;
+	FString LocalFile;
+	int32 CompressedLength;
+	int32 UncompressedLength;
+	TMap<FSHAHash, FBundleEntry> Entries;
+};
+
+//----------------------------------------------------------------------------------------------------------
+// FS3DerivedDataBackend::FBundleDownloadInfo
+//----------------------------------------------------------------------------------------------------------
+
+struct FS3DerivedDataBackend::FBundleDownload : IRequestCallback
+{
+	FCriticalSection& CriticalSection;
+	FBundle& Bundle;
+	FString BundleUrl;
+	FRequestPool& RequestPool;
+	FFeedbackContext* Context;
+	FGraphEventRef Event;
+	int DownloadedBytes;
+
+	FBundleDownload(FCriticalSection& InCriticalSection, FBundle& InBundle, FString InBundleUrl, FRequestPool& InRequestPool, FFeedbackContext* InContext)
+		: CriticalSection(InCriticalSection)
+		, Bundle(InBundle)
+		, BundleUrl(InBundleUrl)
+		, RequestPool(InRequestPool)
+		, Context(InContext)
+		, DownloadedBytes(0)
+	{
+	}
+
+	void Execute()
+	{
+		if (Context->ReceivedUserCancel())
+		{
+			return;
+		}
+
+		Context->Logf(TEXT("Downloading %s (%d bytes)"), *BundleUrl, Bundle.CompressedLength);
+
+		TArray<uint8> CompressedData;
+		CompressedData.Reserve(Bundle.CompressedLength);
+
+		long ResponseCode = RequestPool.Download(*BundleUrl, this, CompressedData, Context);
+		if(!IsSuccessfulHttpResponse(ResponseCode))
+		{
+			Context->Logf(ELogVerbosity::Warning, TEXT("Unable to download bundle %s (%d)"), *BundleUrl, ResponseCode);
+			return;
+		}
+
+		Context->Logf(TEXT("Decompressing %s (%d bytes)"), *BundleUrl, Bundle.UncompressedLength);
+
+		TArray<uint8> UncompressedData;
+		UncompressedData.SetNum(Bundle.UncompressedLength);
+
+		if (!FCompression::UncompressMemory(NAME_Gzip, UncompressedData.GetData(), Bundle.UncompressedLength, CompressedData.GetData(), CompressedData.Num()))
+		{
+			Context->Logf(ELogVerbosity::Warning, TEXT("Unable to decompress bundle %s"), *BundleUrl);
+			return;
+		}
+
+		FString TempFile = Bundle.LocalFile + TEXT(".incoming");
+		if (!FFileHelper::SaveArrayToFile(UncompressedData, *TempFile))
+		{
+			Context->Logf(ELogVerbosity::Warning, TEXT("Unable to save bundle to %s"), *TempFile);
+			return;
+		}
+
+		IFileManager::Get().Move(*Bundle.LocalFile, *TempFile);
+		Context->Logf(TEXT("Finished downloading %s to %s"), *BundleUrl, *Bundle.LocalFile);
+	}
+
+	virtual bool Update(int NumBytes, int MaxBytes) override final
+	{
+		FScopeLock Lock(&CriticalSection);
+		DownloadedBytes = NumBytes;
+		return !Context->ReceivedUserCancel();
+	}
+};
+
+//----------------------------------------------------------------------------------------------------------
+// FS3DerivedDataBackend
+//----------------------------------------------------------------------------------------------------------
+
+FS3DerivedDataBackend::FS3DerivedDataBackend(const TCHAR* InManifestPath, const TCHAR* InBaseUrl, const TCHAR* InRegion, const TCHAR* InCanaryObjectKey, const TCHAR* InAccessKey, const TCHAR* InSecretKey)
+	: ManifestPath(InManifestPath)
+	, BaseUrl(InBaseUrl)
+	, Region(InRegion)
+	, CanaryObjectKey(InCanaryObjectKey)
+	, CacheDir(FPaths::ProjectSavedDir() / TEXT("S3DDC"))
+	, RequestPool(new FRequestPool(TCHAR_TO_ANSI(InRegion), TCHAR_TO_ANSI(InAccessKey), TCHAR_TO_ANSI(InSecretKey)))
+	, bEnabled(false)
+{
+	// Test whether we can reach the canary URL
+	bool bCanaryValid = true;
+	if (CanaryObjectKey.Len() > 0)
+	{
+		TArray<uint8> Data;
+		bCanaryValid = IsSuccessfulHttpResponse(RequestPool->Download(*(BaseUrl / CanaryObjectKey), nullptr, Data, GLog));
+	}
+
+	// Try to read the bundles
+	if (bCanaryValid)
+	{
+		FFeedbackContext* Context = FDesktopPlatformModule::Get()->GetNativeFeedbackContext();
+		Context->BeginSlowTask(NSLOCTEXT("S3DerivedDataBackend", "DownloadingDDCBundles", "Downloading DDC bundles..."), true, true);
+
+		if (ReadManifest())
+		{
+			// Get the path for each bundle that needs downloading
+			for (FBundle& Bundle : Bundles)
+			{
+				Bundle.LocalFile = CacheDir / Bundle.Name;
+			}
+
+			// Remove any bundles that are no longer required
+			RemoveUnusedBundles();
+
+			// Create a critical section used for updating download state
+			FCriticalSection CriticalSection;
+
+			// Create all the download tasks
+			TArray<TSharedPtr<FBundleDownload>> Downloads;
+			for (FBundle& Bundle : Bundles)
+			{
+				if (!FPaths::FileExists(Bundle.LocalFile))
+				{
+					TSharedPtr<FBundleDownload> Download(new FBundleDownload(CriticalSection, Bundle, BaseUrl + Bundle.ObjectKey, *RequestPool.Get(), Context));
+					Download->Event = FFunctionGraphTask::CreateAndDispatchWhenReady([Download]() { Download->Execute(); }, TStatId());
+					Downloads.Add(MoveTemp(Download));
+				}
+			}
+
+			// Loop until the downloads have all finished
+			for(bool bComplete = false; !bComplete; )
+			{
+				FPlatformProcess::Sleep(0.1f);
+
+				int NumBytes = 0;
+				int MaxBytes = 0;
+				bComplete = true;
+
+				CriticalSection.Lock();
+				for (TSharedPtr<FBundleDownload>& Download : Downloads)
+				{
+					NumBytes += Download->DownloadedBytes;
+					MaxBytes += Download->Bundle.CompressedLength;
+					bComplete &= Download->Event->IsComplete();
+				}
+				CriticalSection.Unlock();
+
+				Context->UpdateProgress(NumBytes, MaxBytes);
+			}
+
+			// Mount all the bundles
+			ParallelFor(Bundles.Num(), [this](int32 Index) { ReadBundle(Bundles[Index]); });
+			bEnabled = true;
+		}
+
+		Context->EndSlowTask();
+	}
+}
+
+FS3DerivedDataBackend::~FS3DerivedDataBackend()
+{
+}
+
+bool FS3DerivedDataBackend::IsUsable() const
+{
+	return bEnabled;
+}
+
+bool FS3DerivedDataBackend::CachedDataProbablyExists(const TCHAR* CacheKey)
+{
+	const FBundle* Bundle;
+	const FBundleEntry* BundleEntry;
+	return FindBundleEntry(CacheKey, Bundle, BundleEntry);
+}
+
+bool FS3DerivedDataBackend::GetCachedData(const TCHAR* CacheKey, TArray<uint8>& OutData)
+{
+	TRACE_CPUPROFILER_EVENT_SCOPE(S3DDC_Get);
+	TRACE_COUNTER_ADD(S3DDC_Get, int64(1));
+	COOK_STAT(auto Timer = UsageStats.TimeGet());
+
+	const FBundle* Bundle;
+	const FBundleEntry* BundleEntry;
+	if (FindBundleEntry(CacheKey, Bundle, BundleEntry))
+	{
+		TUniquePtr<FArchive> Reader(IFileManager::Get().CreateFileReader(*Bundle->LocalFile));
+		if (!Reader->IsError())
+		{
+			OutData.SetNum(BundleEntry->Length);
+			Reader->Seek(BundleEntry->Offset);
+			Reader->Serialize(OutData.GetData(), BundleEntry->Length);
+			return true;
+		}
+	}
+
+	return false;
+}
+
+void FS3DerivedDataBackend::PutCachedData(const TCHAR* CacheKey, TArrayView<const uint8> InData, bool bPutEvenIfExists)
+{
+	// Not implemented
+}
+
+void FS3DerivedDataBackend::RemoveCachedData(const TCHAR* CacheKey, bool bTransient)
+{
+	// Not implemented
+}
+
+void FS3DerivedDataBackend::GatherUsageStats(TMap<FString, FDerivedDataCacheUsageStats>& UsageStatsMap, FString&& GraphPath)
+{
+	COOK_STAT(UsageStatsMap.Add(FString::Printf(TEXT("%s: %s @ %s"), *GraphPath, TEXT("S3"), *BaseUrl), UsageStats));
+}
+
+bool FS3DerivedDataBackend::ReadManifest()
+{
+	// Read the manifest from disk
+	FString ManifestText;
+	if (!FFileHelper::LoadFileToString(ManifestText, *ManifestPath))
+	{
+		UE_LOG(LogDerivedDataCache, Warning, TEXT("Unable to read manifest from %s"), *ManifestPath);
+		return false;
+	}
+
+	// Deserialize a JSON object from the string
+	TSharedPtr<FJsonObject> ManifestObject;
+	if (!FJsonSerializer::Deserialize(TJsonReaderFactory<>::Create(ManifestText), ManifestObject) || !ManifestObject.IsValid())
+	{
+		UE_LOG(LogDerivedDataCache, Warning, TEXT("Unable to parse manifest from %s"), *ManifestPath);
+		return false;
+	}
+
+	// Parse out the list of bundles
+	const TArray<TSharedPtr<FJsonValue>>* BundlesArray;
+	if (!ManifestObject->TryGetArrayField(TEXT("Bundles"), BundlesArray))
+	{
+		UE_LOG(LogDerivedDataCache, Warning, TEXT("Manifest from %s is missing bundles array"), *ManifestPath);
+		return false;
+	}
+
+	// Parse out each bundle
+	for (const TSharedPtr<FJsonValue>& BundleValue : *BundlesArray)
+	{
+		const FJsonObject& BundleObject = *BundleValue->AsObject();
+
+		FBundle Bundle;
+		if (!BundleObject.TryGetStringField(TEXT("Name"), Bundle.Name))
+		{
+			UE_LOG(LogDerivedDataCache, Warning, TEXT("Manifest from %s is missing a bundle name"), *ManifestPath);
+			return false;
+		}
+		if (!BundleObject.TryGetStringField(TEXT("ObjectKey"), Bundle.ObjectKey))
+		{
+			UE_LOG(LogDerivedDataCache, Warning, TEXT("Manifest from %s is missing an bundle object key"), *ManifestPath);
+			return false;
+		}
+		if (!BundleObject.TryGetNumberField(TEXT("CompressedLength"), Bundle.CompressedLength))
+		{
+			UE_LOG(LogDerivedDataCache, Warning, TEXT("Manifest from %s is missing the compressed length"), *ManifestPath);
+			return false;
+		}
+		if (!BundleObject.TryGetNumberField(TEXT("UncompressedLength"), Bundle.UncompressedLength))
+		{
+			UE_LOG(LogDerivedDataCache, Warning, TEXT("Manifest from %s is missing the uncompressed length"), *ManifestPath);
+			return false;
+		}
+
+		Bundles.Add(MoveTemp(Bundle));
+	}
+
+	return true;
+}
+
+void FS3DerivedDataBackend::RemoveUnusedBundles()
+{
+	IFileManager& FileManager = IFileManager::Get();
+
+	// Find all the files we want to keep
+	TSet<FString> KeepFiles;
+	for (const FBundle& Bundle : Bundles)
+	{
+		KeepFiles.Add(Bundle.Name);
+	}
+
+	// Find all the files on disk
+	TArray<FString> Files;
+	FileManager.FindFiles(Files, *CacheDir);
+
+	// Remove anything left over
+	for (const FString& File : Files)
+	{
+		if (!KeepFiles.Contains(File))
+		{
+			FileManager.Delete(*(CacheDir / File));
+		}
+	}
+}
+
+void FS3DerivedDataBackend::ReadBundle(FBundle& Bundle)
+{
+	IFileManager& FileManager = IFileManager::Get();
+
+	// Open the file for reading. If this fails, assume it's because the download was aborted.
+	TUniquePtr<FArchive> Reader(FileManager.CreateFileReader(*Bundle.LocalFile));
+	if (!Reader.IsValid() || Reader->IsError())
+	{
+		UE_LOG(LogDerivedDataCache, Warning, TEXT("Unable to open bundle %s for reading. Ignoring."), *Bundle.LocalFile);
+		return;
+	}
+	
+	struct FFileHeader
+	{
+		uint32 Signature;
+		int32 NumRecords;
+	};
+
+	FFileHeader Header;
+	Reader->Serialize(&Header, sizeof(Header));
+
+	const uint32 BundleSignature = (uint32)'D' | ((uint32)'D' << 8) | ((uint32)'B' << 16);
+	const uint32 BundleSignatureV1 = BundleSignature | (1U << 24);
+	if (Header.Signature != BundleSignatureV1)
+	{
+		UE_LOG(LogDerivedDataCache, Warning, TEXT("Unable to read bundle with signature %08x"), Header.Signature);
+		return;
+	}
+
+	struct FFileRecord
+	{
+		FSHAHash Hash;
+		uint32 Length;
+	};
+
+	TArray<FFileRecord> Records;
+	Records.SetNum(Header.NumRecords);
+	Reader->Serialize(Records.GetData(), Header.NumRecords * sizeof(FFileRecord));
+
+	Bundle.Entries.Reserve(Records.Num());
+
+	int64 Offset = Reader->Tell();
+	for (const FFileRecord& Record : Records)
+	{
+		FBundleEntry& Entry = Bundle.Entries.Add(Record.Hash);
+		Entry.Offset = Offset;
+		Entry.Length = Record.Length;
+		Offset += Record.Length;
+		check(Offset <= Bundle.UncompressedLength);
+	}
+}
+
+bool FS3DerivedDataBackend::FindBundleEntry(const TCHAR* CacheKey, const FBundle*& OutBundle, const FBundleEntry*& OutBundleEntry) const
+{
+	FSHAHash Hash;
+
+	auto AnsiString = StringCast<ANSICHAR>(*BuildPathForCacheKey(CacheKey).ToUpper());
+	FSHA1::HashBuffer(AnsiString.Get(), AnsiString.Length(), Hash.Hash);
+
+	for (const FBundle& Bundle : Bundles)
+	{
+		const FBundleEntry* Entry = Bundle.Entries.Find(Hash);
+		if (Entry != nullptr)
+		{
+			OutBundle = &Bundle;
+			OutBundleEntry = Entry;
+			return true;
+		}
+	}
+
+	return false;
+}
diff --git a/Engine/Source/Developer/DerivedDataCache/Private/S3DerivedDataBackend.h b/Engine/Source/Developer/DerivedDataCache/Private/S3DerivedDataBackend.h
new file mode 100644
--- /dev/null
+++ b/Engine/Source/Developer/DerivedDataCache/Private/S3DerivedDataBackend.h
@@ -0,0 +1,63 @@
+// Copyright Epic Games, Inc. All Rights Reserved.
+
+#pragma once
+
+#include "CoreMinimal.h"
+#include "DerivedDataBackendInterface.h"
+#include "DerivedDataCacheUsageStats.h"
+
+/**
+ * Backend for a read-only AWS S3 based caching service.
+ **/
+class FS3DerivedDataBackend : public FDerivedDataBackendInterface
+{
+public:
+	/**
+	 * Creates the backend, checks health status and attempts to acquire an access token.
+	 *
+	 * @param  InManifestPath       Path to the JSON manifest in the workspace containing a list of files to download
+	 * @param  InBaseUrl            Base URL for the bucket, with trailing slash (eg. https://foo.s3.us-east-1.amazonaws.com/)
+	 * @param  InRegion	            Name of the AWS region (eg. us-east-1)
+	 * @param  InCanaryObjectKey    Key for a canary object used to test whether this backend is usable
+	 * @param  InAccessKey          The AWS access key
+	 * @param  InSecretKey          The AWS secret key
+	 */
+	FS3DerivedDataBackend(const TCHAR* InManifestPath, const TCHAR* InBaseUrl, const TCHAR* InRegion, const TCHAR* InCanaryObjectKey, const TCHAR* InAccessKey, const TCHAR* InSecretKey);
+	~FS3DerivedDataBackend();
+
+	/**
+	 * Checks is backend is usable (reachable and accessible).
+	 * @return true if usable
+	 */
+	bool IsUsable() const;
+	bool IsWritable() override { return false; }
+	bool CachedDataProbablyExists(const TCHAR* CacheKey) override;
+	bool GetCachedData(const TCHAR* CacheKey, TArray<uint8>& OutData) override;
+	void PutCachedData(const TCHAR* CacheKey, TArrayView<const uint8> InData, bool bPutEvenIfExists) override;
+	void RemoveCachedData(const TCHAR* CacheKey, bool bTransient) override;
+	void GatherUsageStats(TMap<FString, FDerivedDataCacheUsageStats>& UsageStatsMap, FString&& GraphPath) override;
+	
+private:
+	struct FBundle;
+	struct FBundleEntry;
+	struct FBundleDownload;
+
+	class FRequest;
+	class FRequestPool;
+
+	FString ManifestPath;
+	FString BaseUrl;
+	FString Region;
+	FString CanaryObjectKey;
+	FString CacheDir;
+	TArray<FBundle> Bundles;
+	TUniquePtr<FRequestPool> RequestPool;
+	FDerivedDataCacheUsageStats UsageStats;
+	bool bEnabled;
+
+	bool ReadManifest();
+	void RemoveUnusedBundles();
+	void ReadBundle(FBundle& Bundle);
+	bool FindBundleEntry(const TCHAR* CacheKey, const FBundle*& OutBundle, const FBundleEntry*& OutBundleEntry) const;
+};
+
