Fix missing header.

diff --git a/Engine/Source/Developer/DerivedDataCache/Private/S3DerivedDataBackend.cpp b/Engine/Source/Developer/DerivedDataCache/Private/S3DerivedDataBackend.cpp
--- a/Engine/Source/Developer/DerivedDataCache/Private/S3DerivedDataBackend.cpp
+++ b/Engine/Source/Developer/DerivedDataCache/Private/S3DerivedDataBackend.cpp
@@ -34,8 +34,6 @@
 #include "HAL/Runnable.h"
 #include "DesktopPlatformModule.h"
 
-#include "picosha2.h"
-
 #define S3DDC_BACKEND_WAIT_INTERVAL 0.01f
 #define S3DDC_HTTP_REQUEST_TIMEOUT_SECONDS 30L
 #define S3DDC_HTTP_REQUEST_TIMOUT_ENABLED 1
