Change S3 DDC backend to use a separate index for the currently valid manifests, so that it can be secured in source control.

diff --git a/Engine/Source/Developer/DerivedDataCache/Private/S3DerivedDataBackend.cpp b/Engine/Source/Developer/DerivedDataCache/Private/S3DerivedDataBackend.cpp
--- a/Engine/Source/Developer/DerivedDataCache/Private/S3DerivedDataBackend.cpp
+++ b/Engine/Source/Developer/DerivedDataCache/Private/S3DerivedDataBackend.cpp
@@ -702,8 +702,8 @@ struct FS3DerivedDataBackend::FBundleDownload : IRequestCallback
 // FS3DerivedDataBackend
 //----------------------------------------------------------------------------------------------------------
 
-FS3DerivedDataBackend::FS3DerivedDataBackend(const TCHAR* InManifestPath, const TCHAR* InBaseUrl, const TCHAR* InRegion, const TCHAR* InCanaryObjectKey, const TCHAR* InAccessKey, const TCHAR* InSecretKey)
-	: ManifestPath(InManifestPath)
+FS3DerivedDataBackend::FS3DerivedDataBackend(const TCHAR* InRootManifestPath, const TCHAR* InBaseUrl, const TCHAR* InRegion, const TCHAR* InCanaryObjectKey, const TCHAR* InAccessKey, const TCHAR* InSecretKey)
+	: RootManifestPath(InRootManifestPath)
 	, BaseUrl(InBaseUrl)
 	, Region(InRegion)
 	, CanaryObjectKey(InCanaryObjectKey)
@@ -725,7 +725,7 @@ FS3DerivedDataBackend::FS3DerivedDataBackend(const TCHAR* InManifestPath, const
 		FFeedbackContext* Context = FDesktopPlatformModule::Get()->GetNativeFeedbackContext();
 		Context->BeginSlowTask(NSLOCTEXT("S3DerivedDataBackend", "DownloadingDDCBundles", "Downloading DDC bundles..."), true, true);
 
-		if (ReadManifest())
+		if (DownloadManifest(Context))
 		{
 			// Get the path for each bundle that needs downloading
 			for (FBundle& Bundle : Bundles)
@@ -835,29 +835,67 @@ void FS3DerivedDataBackend::GatherUsageStats(TMap<FString, FDerivedDataCacheUsag
 	COOK_STAT(UsageStatsMap.Add(FString::Printf(TEXT("%s: %s @ %s"), *GraphPath, TEXT("S3"), *BaseUrl), UsageStats));
 }
 
-bool FS3DerivedDataBackend::ReadManifest()
+bool FS3DerivedDataBackend::DownloadManifest(FFeedbackContext* Context)
 {
-	// Read the manifest from disk
-	FString ManifestText;
-	if (!FFileHelper::LoadFileToString(ManifestText, *ManifestPath))
+	// Read the root manifest from disk
+	FString RootManifestText;
+	if (!FFileHelper::LoadFileToString(RootManifestText, *RootManifestPath))
 	{
-		UE_LOG(LogDerivedDataCache, Warning, TEXT("Unable to read manifest from %s"), *ManifestPath);
+		Context->Logf(ELogVerbosity::Warning, TEXT("Unable to read manifest from %s"), *RootManifestPath);
 		return false;
 	}
 
 	// Deserialize a JSON object from the string
-	TSharedPtr<FJsonObject> ManifestObject;
-	if (!FJsonSerializer::Deserialize(TJsonReaderFactory<>::Create(ManifestText), ManifestObject) || !ManifestObject.IsValid())
+	TSharedPtr<FJsonObject> RootManifestObject;
+	if (!FJsonSerializer::Deserialize(TJsonReaderFactory<>::Create(RootManifestText), RootManifestObject) || !RootManifestObject.IsValid())
 	{
-		UE_LOG(LogDerivedDataCache, Warning, TEXT("Unable to parse manifest from %s"), *ManifestPath);
+		Context->Logf(ELogVerbosity::Warning, TEXT("Unable to parse manifest from %s"), *RootManifestPath);
+		return false;
+	}
+
+	// Parse out the list of manifests
+	const TArray<TSharedPtr<FJsonValue>>* RootEntriesArray;
+	if (!RootManifestObject->TryGetArrayField(TEXT("Entries"), RootEntriesArray))
+	{
+		Context->Logf(ELogVerbosity::Warning, TEXT("Root manifest from %s is missing entries array"), *RootManifestPath);
+		return false;
+	}
+	if (RootEntriesArray->Num() == 0)
+	{
+		Context->Logf(ELogVerbosity::Warning, TEXT("Root manifest from %s has empty entries array"), *RootManifestPath);
+		return false;
+	}
+
+	// Get the object key for the last entry
+	const TSharedPtr<FJsonObject>& LastRootManifestEntry = (*RootEntriesArray)[RootEntriesArray->Num() - 1]->AsObject();
+	FString BundleManifestKey = LastRootManifestEntry->GetStringField("Key");
+
+	// Download the bundle manifest
+	TArray<uint8> BundleManifestData;
+	long ResponseCode = RequestPool->Download(*(BaseUrl + BundleManifestKey), nullptr, BundleManifestData, Context);
+	if (!IsSuccessfulHttpResponse(ResponseCode))
+	{
+		Context->Logf(ELogVerbosity::Warning, TEXT("Unable to download bundle manifest from %s (%d)"), *BundleManifestKey, (int)ResponseCode);
+		return false;
+	}
+
+	// Convert it to text
+	BundleManifestData.Add(0);
+	FString BundleManifestText = ANSI_TO_TCHAR((const ANSICHAR*)BundleManifestData.GetData());
+
+	// Deserialize a JSON object from the string
+	TSharedPtr<FJsonObject> BundleManifestObject;
+	if (!FJsonSerializer::Deserialize(TJsonReaderFactory<>::Create(BundleManifestText), BundleManifestObject) || !BundleManifestObject.IsValid())
+	{
+		Context->Logf(ELogVerbosity::Warning, TEXT("Unable to parse manifest from %s"), *BundleManifestKey);
 		return false;
 	}
 
 	// Parse out the list of bundles
 	const TArray<TSharedPtr<FJsonValue>>* BundlesArray;
-	if (!ManifestObject->TryGetArrayField(TEXT("Bundles"), BundlesArray))
+	if (!BundleManifestObject->TryGetArrayField(TEXT("Entries"), BundlesArray))
 	{
-		UE_LOG(LogDerivedDataCache, Warning, TEXT("Manifest from %s is missing bundles array"), *ManifestPath);
+		Context->Logf(ELogVerbosity::Warning, TEXT("Manifest from %s is missing bundles array"), *BundleManifestKey);
 		return false;
 	}
 
@@ -869,22 +907,22 @@ bool FS3DerivedDataBackend::ReadManifest()
 		FBundle Bundle;
 		if (!BundleObject.TryGetStringField(TEXT("Name"), Bundle.Name))
 		{
-			UE_LOG(LogDerivedDataCache, Warning, TEXT("Manifest from %s is missing a bundle name"), *ManifestPath);
+			Context->Logf(ELogVerbosity::Warning, TEXT("Manifest from %s is missing a bundle name"), *BundleManifestKey);
 			return false;
 		}
 		if (!BundleObject.TryGetStringField(TEXT("ObjectKey"), Bundle.ObjectKey))
 		{
-			UE_LOG(LogDerivedDataCache, Warning, TEXT("Manifest from %s is missing an bundle object key"), *ManifestPath);
+			Context->Logf(ELogVerbosity::Warning, TEXT("Manifest from %s is missing an bundle object key"), *BundleManifestKey);
 			return false;
 		}
 		if (!BundleObject.TryGetNumberField(TEXT("CompressedLength"), Bundle.CompressedLength))
 		{
-			UE_LOG(LogDerivedDataCache, Warning, TEXT("Manifest from %s is missing the compressed length"), *ManifestPath);
+			Context->Logf(ELogVerbosity::Warning, TEXT("Manifest from %s is missing the compressed length"), *BundleManifestKey);
 			return false;
 		}
 		if (!BundleObject.TryGetNumberField(TEXT("UncompressedLength"), Bundle.UncompressedLength))
 		{
-			UE_LOG(LogDerivedDataCache, Warning, TEXT("Manifest from %s is missing the uncompressed length"), *ManifestPath);
+			Context->Logf(ELogVerbosity::Warning, TEXT("Manifest from %s is missing the uncompressed length"), *BundleManifestKey);
 			return false;
 		}
 
diff --git a/Engine/Source/Developer/DerivedDataCache/Private/S3DerivedDataBackend.h b/Engine/Source/Developer/DerivedDataCache/Private/S3DerivedDataBackend.h
--- a/Engine/Source/Developer/DerivedDataCache/Private/S3DerivedDataBackend.h
+++ b/Engine/Source/Developer/DerivedDataCache/Private/S3DerivedDataBackend.h
@@ -15,14 +15,14 @@ public:
 	/**
 	 * Creates the backend, checks health status and attempts to acquire an access token.
 	 *
-	 * @param  InManifestPath       Path to the JSON manifest in the workspace containing a list of files to download
+	 * @param  InRootManifestPath   Local path to the JSON manifest in the workspace containing a list of files to download
 	 * @param  InBaseUrl            Base URL for the bucket, with trailing slash (eg. https://foo.s3.us-east-1.amazonaws.com/)
 	 * @param  InRegion	            Name of the AWS region (eg. us-east-1)
 	 * @param  InCanaryObjectKey    Key for a canary object used to test whether this backend is usable
 	 * @param  InAccessKey          The AWS access key
 	 * @param  InSecretKey          The AWS secret key
 	 */
-	FS3DerivedDataBackend(const TCHAR* InManifestPath, const TCHAR* InBaseUrl, const TCHAR* InRegion, const TCHAR* InCanaryObjectKey, const TCHAR* InAccessKey, const TCHAR* InSecretKey);
+	FS3DerivedDataBackend(const TCHAR* InRootManifestPath, const TCHAR* InBaseUrl, const TCHAR* InRegion, const TCHAR* InCanaryObjectKey, const TCHAR* InAccessKey, const TCHAR* InSecretKey);
 	~FS3DerivedDataBackend();
 
 	/**
@@ -45,7 +45,7 @@ private:
 	class FRequest;
 	class FRequestPool;
 
-	FString ManifestPath;
+	FString RootManifestPath;
 	FString BaseUrl;
 	FString Region;
 	FString CanaryObjectKey;
@@ -55,7 +55,7 @@ private:
 	FDerivedDataCacheUsageStats UsageStats;
 	bool bEnabled;
 
-	bool ReadManifest();
+	bool DownloadManifest(FFeedbackContext* Context);
 	void RemoveUnusedBundles();
 	void ReadBundle(FBundle& Bundle);
 	bool FindBundleEntry(const TCHAR* CacheKey, const FBundle*& OutBundle, const FBundleEntry*& OutBundleEntry) const;
