Disable S3 DDC backend on Mac due to missing libcurl implementation causing link errors.

diff --git a/Engine/Source/Developer/DerivedDataCache/Private/DerivedDataBackends.cpp b/Engine/Source/Developer/DerivedDataCache/Private/DerivedDataBackends.cpp
--- a/Engine/Source/Developer/DerivedDataCache/Private/DerivedDataBackends.cpp
+++ b/Engine/Source/Developer/DerivedDataCache/Private/DerivedDataBackends.cpp
@@ -603,12 +603,12 @@ public:
 		return DataCache;
 	}
 
-
 	/**
 	 * Creates an S3 data cache interface.
 	 */
 	FDerivedDataBackendInterface* ParseS3Cache(const TCHAR* NodeName, const TCHAR* Entry)
 	{
+#if WITH_S3_DDC_BACKEND
 		FString ManifestPath;
 		if (!FParse::Value(Entry, TEXT("Manifest="), ManifestPath))
 		{
@@ -658,6 +658,10 @@ public:
 		// Insert the backend corruption wrapper. Since the filesystem already uses this, and we're recycling the data with the trailer intact, we need to use it for the S3 cache too.
 		UE_LOG(LogDerivedDataCache, Display, TEXT("Using %s S3 backend at %s"), *Region, *BaseUrl);
 		return new FDerivedDataBackendCorruptionWrapper(Backend);
+#else
+		UE_LOG(LogDerivedDataCache, Warning, TEXT("S3 backend is not supported in the current build configuration."));
+		return nullptr;
+#endif
 	}
 
 	/**
diff --git a/Engine/Source/Developer/DerivedDataCache/Private/S3DerivedDataBackend.cpp b/Engine/Source/Developer/DerivedDataCache/Private/S3DerivedDataBackend.cpp
--- a/Engine/Source/Developer/DerivedDataCache/Private/S3DerivedDataBackend.cpp
+++ b/Engine/Source/Developer/DerivedDataCache/Private/S3DerivedDataBackend.cpp
@@ -2,6 +2,8 @@
 
 #include "S3DerivedDataBackend.h"
 
+#if WITH_S3_DDC_BACKEND
+
 #if PLATFORM_WINDOWS || PLATFORM_HOLOLENS
 	#include "Windows/WindowsHWrapper.h"
 	#include "Windows/AllowWindowsPlatformTypes.h"
@@ -10,12 +12,10 @@
 #if PLATFORM_WINDOWS || PLATFORM_HOLOLENS
 	#include "Windows/HideWindowsPlatformTypes.h"
 #endif
-#if WITH_SSL
-	#include "Ssl.h"
-	#include <openssl/ssl.h>
-	#include <openssl/sha.h>
-	#include <openssl/hmac.h>
-#endif
+#include "Ssl.h"
+#include <openssl/ssl.h>
+#include <openssl/sha.h>
+#include <openssl/hmac.h>
 #include "Misc/Base64.h"
 #include "Misc/Paths.h"
 #include "Misc/SecureHash.h"
@@ -131,25 +131,17 @@ struct FSHA256
 
 FSHA256 Sha256(const uint8* Input, size_t InputLen)
 {
-#if WITH_SSL
 	FSHA256 Output;
 	SHA256(Input, InputLen, Output.Digest);
 	return Output;
-#else
-	UE_LOG(LogDerivedDataCache, Fatal, TEXT("SSL functionality is not available. Cannot use S3 backend."));
-#endif
 }
 
 FSHA256 HmacSha256(const uint8* Input, size_t InputLen, const uint8* Key, size_t KeyLen)
 {
-#if WITH_SSL
 	FSHA256 Output;
 	unsigned int OutputLen = 0;
 	HMAC(EVP_sha256(), Key, KeyLen, (const unsigned char*)Input, InputLen, Output.Digest, &OutputLen);
 	return Output;
-#else
-	UE_LOG(LogDerivedDataCache, Fatal, TEXT("SSL functionality is not available. Cannot use S3 backend."));
-#endif
 }
 
 FSHA256 HmacSha256(const FStringAnsi& Input, const uint8* Key, size_t KeyLen)
@@ -264,7 +256,6 @@ public:
 		curl_easy_setopt(Curl, CURLOPT_WRITEDATA, &CallbackData);
 		curl_easy_setopt(Curl, CURLOPT_WRITEFUNCTION, StaticWriteBodyFn);
 
-#if WITH_SSL
 		// SSL options
 		curl_easy_setopt(Curl, CURLOPT_USE_SSL, CURLUSESSL_ALL);
 		curl_easy_setopt(Curl, CURLOPT_SSL_VERIFYPEER, 1);
@@ -275,7 +266,6 @@ public:
 		curl_easy_setopt(Curl, CURLOPT_CAINFO, nullptr);
 		curl_easy_setopt(Curl, CURLOPT_SSL_CTX_FUNCTION, *sslctx_function);
 		curl_easy_setopt(Curl, CURLOPT_SSL_CTX_DATA, &CallbackData);
-#endif
 
 		// Send the request
 		CURLcode CurlResult = curl_easy_perform(Curl);
@@ -483,7 +473,6 @@ private:
 		return 0;
 	}
 
-#if WITH_SSL
 	static int SslCertVerify(int PreverifyOk, X509_STORE_CTX* Context)
 	{
 		if (PreverifyOk == 1)
@@ -518,8 +507,6 @@ private:
 		/* all set to go */
 		return CURLE_OK;
 	}
-
-#endif //#if WITH_SSL
 };
 
 //----------------------------------------------------------------------------------------------------------
@@ -1029,3 +1016,5 @@ bool FS3DerivedDataBackend::FindBundleEntry(const TCHAR* CacheKey, const FBundle
 
 	return false;
 }
+
+#endif
diff --git a/Engine/Source/Developer/DerivedDataCache/Private/S3DerivedDataBackend.h b/Engine/Source/Developer/DerivedDataCache/Private/S3DerivedDataBackend.h
--- a/Engine/Source/Developer/DerivedDataCache/Private/S3DerivedDataBackend.h
+++ b/Engine/Source/Developer/DerivedDataCache/Private/S3DerivedDataBackend.h
@@ -6,6 +6,15 @@
 #include "DerivedDataBackendInterface.h"
 #include "DerivedDataCacheUsageStats.h"
 
+// Macro for whether to enable the S3 backend. libcurl is not currently available on Mac.
+#if PLATFORM_WINDOWS
+	#define WITH_S3_DDC_BACKEND 1
+#else
+	#define WITH_S3_DDC_BACKEND 0
+#endif
+
+#if WITH_S3_DDC_BACKEND
+
 /**
  * Backend for a read-only AWS S3 based caching service.
  **/
@@ -61,3 +70,4 @@ private:
 	bool FindBundleEntry(const TCHAR* CacheKey, const FBundle*& OutBundle, const FBundleEntry*& OutBundleEntry) const;
 };
 
+#endif
