Add verbose logging to S3 DDC backend.

diff --git a/Engine/Source/Developer/DerivedDataCache/Private/S3DerivedDataBackend.cpp b/Engine/Source/Developer/DerivedDataCache/Private/S3DerivedDataBackend.cpp
--- a/Engine/Source/Developer/DerivedDataCache/Private/S3DerivedDataBackend.cpp
+++ b/Engine/Source/Developer/DerivedDataCache/Private/S3DerivedDataBackend.cpp
@@ -781,7 +781,12 @@ bool FS3DerivedDataBackend::CachedDataProbablyExists(const TCHAR* CacheKey)
 {
 	const FBundle* Bundle;
 	const FBundleEntry* BundleEntry;
-	return FindBundleEntry(CacheKey, Bundle, BundleEntry);
+	if (!FindBundleEntry(CacheKey, Bundle, BundleEntry))
+	{
+		UE_LOG(LogDerivedDataCache, Verbose, TEXT("S3DerivedDataBackend: Cache miss on %s (probably)"), CacheKey);
+		return false;
+	}
+	return true;
 }
 
 bool FS3DerivedDataBackend::GetCachedData(const TCHAR* CacheKey, TArray<uint8>& OutData)
@@ -797,6 +802,7 @@ bool FS3DerivedDataBackend::GetCachedData(const TCHAR* CacheKey, TArray<uint8>&
 		TUniquePtr<FArchive> Reader(IFileManager::Get().CreateFileReader(*Bundle->LocalFile));
 		if (!Reader->IsError())
 		{
+			UE_LOG(LogDerivedDataCache, Verbose, TEXT("S3DerivedDataBackend: Cache hit on %s"), CacheKey);
 			OutData.SetNum(BundleEntry->Length);
 			Reader->Seek(BundleEntry->Offset);
 			Reader->Serialize(OutData.GetData(), BundleEntry->Length);
@@ -804,6 +810,7 @@ bool FS3DerivedDataBackend::GetCachedData(const TCHAR* CacheKey, TArray<uint8>&
 		}
 	}
 
+	UE_LOG(LogDerivedDataCache, Verbose, TEXT("S3DerivedDataBackend: Cache miss on %s"), CacheKey);
 	return false;
 }
 
