Add an optional parameter to the filesystem DDC backend that logs any accesses to a particular key, so we can extract data for uploading to S3.

diff --git a/Engine/Source/Developer/DerivedDataCache/Private/DerivedDataBackends.cpp b/Engine/Source/Developer/DerivedDataCache/Private/DerivedDataBackends.cpp
--- a/Engine/Source/Developer/DerivedDataCache/Private/DerivedDataBackends.cpp
+++ b/Engine/Source/Developer/DerivedDataCache/Private/DerivedDataBackends.cpp
@@ -29,7 +29,7 @@ DEFINE_LOG_CATEGORY(LogDerivedDataCache);
 #define MAX_BACKEND_KEY_LENGTH (120)
 #define LOCTEXT_NAMESPACE "DerivedDataBackendGraph"
 
-FDerivedDataBackendInterface* CreateFileSystemDerivedDataBackend(const TCHAR* CacheDirectory, bool bForceReadOnly = false, bool bTouchFiles = false, bool bPurgeTransient = false, bool bDeleteOldFiles = false, int32 InDaysToDeleteUnusedFiles = 60, int32 InMaxNumFoldersToCheck = -1, int32 InMaxContinuousFileChecks = -1);
+FDerivedDataBackendInterface* CreateFileSystemDerivedDataBackend(const TCHAR* CacheDirectory, bool bForceReadOnly = false, bool bTouchFiles = false, bool bPurgeTransient = false, bool bDeleteOldFiles = false, int32 InDaysToDeleteUnusedFiles = 60, int32 InMaxNumFoldersToCheck = -1, int32 InMaxContinuousFileChecks = -1, const TCHAR* InAccessLogFileName = nullptr);
 
 /**
   * This class is used to create a singleton that represents the derived data cache hierarchy and all of the wrappers necessary
@@ -567,6 +567,8 @@ public:
 			FParse::Value( Entry, TEXT("FoldersToClean="), MaxFoldersToClean );
 			int32 MaxFileChecksPerSec = -1;
 			FParse::Value( Entry, TEXT("MaxFileChecksPerSec="), MaxFileChecksPerSec );
+			FString WriteAccessLog;
+			FParse::Value( Entry, TEXT("WriteAccessLog="), WriteAccessLog );
 
 			if( bFlush )
 			{
@@ -583,7 +585,7 @@ public:
 			bool bShared = FCString::Stricmp(NodeName, TEXT("Shared")) == 0;
 			if( !bShared || IFileManager::Get().DirectoryExists(*Path) )
 			{
-				InnerFileSystem = CreateFileSystemDerivedDataBackend( *Path, bReadOnly, bTouch, bPurgeTransient, bDeleteUnused, UnusedFileAge, MaxFoldersToClean, MaxFileChecksPerSec);
+				InnerFileSystem = CreateFileSystemDerivedDataBackend( *Path, bReadOnly, bTouch, bPurgeTransient, bDeleteUnused, UnusedFileAge, MaxFoldersToClean, MaxFileChecksPerSec, *WriteAccessLog);
 			}
 
 			if( InnerFileSystem )
diff --git a/Engine/Source/Developer/DerivedDataCache/Private/FileSystemDerivedDataBackend.cpp b/Engine/Source/Developer/DerivedDataCache/Private/FileSystemDerivedDataBackend.cpp
--- a/Engine/Source/Developer/DerivedDataCache/Private/FileSystemDerivedDataBackend.cpp
+++ b/Engine/Source/Developer/DerivedDataCache/Private/FileSystemDerivedDataBackend.cpp
@@ -54,7 +54,7 @@ public:
 	 * @param InCacheDirectory	directory to store the cache in
 	 * @param bForceReadOnly	if true, do not attempt to write to this cache
 	*/
-	FFileSystemDerivedDataBackend(const TCHAR* InCacheDirectory, bool bForceReadOnly, bool bTouchFiles, bool bPurgeTransientData, bool bDeleteOldFiles, int32 InDaysToDeleteUnusedFiles, int32 InMaxNumFoldersToCheck, int32 InMaxContinuousFileChecks)
+	FFileSystemDerivedDataBackend(const TCHAR* InCacheDirectory, bool bForceReadOnly, bool bTouchFiles, bool bPurgeTransientData, bool bDeleteOldFiles, int32 InDaysToDeleteUnusedFiles, int32 InMaxNumFoldersToCheck, int32 InMaxContinuousFileChecks, const TCHAR* InAccessLogFileName)
 		: CachePath(InCacheDirectory)
 		, bReadOnly(bForceReadOnly)
 		, bFailed(true)
@@ -138,6 +138,11 @@ public:
 			UE_LOG(LogDerivedDataCache, Warning, TEXT("%s access is very slow (initialization took %.2f seconds), consider disabling it."), *CachePath, AccessDuration);
 		}
 
+		if (!bFailed && InAccessLogFileName != nullptr && *InAccessLogFileName != 0)
+		{
+			AccessLogWriter.Reset(new FAccessLogWriter(InAccessLogFileName));
+		}
+
 		if (!bReadOnly && !bFailed && bDeleteOldFiles && !FParse::Param(FCommandLine::Get(),TEXT("NODDCCLEANUP")) && FDDCCleanup::Get())
 		{			
 			FDDCCleanup::Get()->AddFilesystem( CachePath, InDaysToDeleteUnusedFiles, InMaxNumFoldersToCheck, InMaxContinuousFileChecks );
@@ -155,6 +160,36 @@ public:
 	{
 		return !bReadOnly;
 	}
+
+	class FAccessLogWriter
+	{
+	public:
+		FAccessLogWriter(const TCHAR* FileName)
+			: Archive(IFileManager::Get().CreateFileWriter(FileName, FILEWRITE_AllowRead))
+		{
+		}
+
+		void Append(const TCHAR* CacheKey)
+		{
+			FScopeLock Lock(&CriticalSection);
+
+			FString CacheKeyStr(CacheKey);
+			if (!CacheKeys.Contains(CacheKeyStr))
+			{
+				CacheKeys.Add(MoveTemp(CacheKeyStr));
+
+				auto FileName = StringCast<ANSICHAR>(*BuildPathForCacheKey(CacheKey));
+				Archive->Serialize(const_cast<ANSICHAR*>(FileName.Get()), FileName.Length());
+				Archive->Serialize(const_cast<ANSICHAR*>(LINE_TERMINATOR_ANSI), sizeof(LINE_TERMINATOR_ANSI) - 1);
+			}
+		}
+
+	private:
+		TUniquePtr<FArchive> Archive;
+		FCriticalSection CriticalSection;
+		TSet<FString> CacheKeys;
+	};
+
 	/**
 	 * Synchronous test for the existence of a cache item
 	 *
@@ -177,6 +212,11 @@ public:
 				IFileManager::Get().SetTimeStamp(*Filename, FDateTime::UtcNow());
 			}
 
+			if (AccessLogWriter.IsValid())
+			{
+				AccessLogWriter->Append(CacheKey);
+			}
+
 			COOK_STAT(Timer.AddHit(0));
 		}
 
@@ -216,6 +256,11 @@ public:
 				UE_CLOG(ReadSpeed < 0.5, LogDerivedDataCache, Warning, TEXT("%s is very slow (%.2fMB/s) when accessing %s, consider disabling it."), *CachePath, ReadSpeed, *Filename);
 			}
 
+			if (AccessLogWriter.IsValid())
+			{
+				AccessLogWriter->Append(CacheKey);
+			}
+
 			UE_LOG(LogDerivedDataCache, Verbose, TEXT("FileSystemDerivedDataBackend: Cache hit on %s"),*Filename);
 			COOK_STAT(Timer.AddHit(Data.Num()));
 			return true;
@@ -249,6 +294,11 @@ public:
 		check(!bFailed);
 		if (!bReadOnly)
 		{
+			if (AccessLogWriter.IsValid())
+			{
+				AccessLogWriter->Append(CacheKey);
+			}
+
 			if (bPutEvenIfExists || !CachedDataProbablyExists(CacheKey))
 			{
 				COOK_STAT(Timer.AddHit(Data.Num()));
@@ -403,11 +453,13 @@ private:
 	/** The total estimated build time accumulated from cache miss/put deltas */
 	double TotalEstimatedBuildTime;
 
+	/** Access log to write to */
+	TUniquePtr<FAccessLogWriter> AccessLogWriter;
 };
 
-FDerivedDataBackendInterface* CreateFileSystemDerivedDataBackend(const TCHAR* CacheDirectory, bool bForceReadOnly /*= false*/, bool bTouchFiles /*= false*/, bool bPurgeTransient /*= false*/, bool bDeleteOldFiles /*= false*/, int32 InDaysToDeleteUnusedFiles /*= 60*/, int32 InMaxNumFoldersToCheck /*= -1*/, int32 InMaxContinuousFileChecks /*= -1*/)
+FDerivedDataBackendInterface* CreateFileSystemDerivedDataBackend(const TCHAR* CacheDirectory, bool bForceReadOnly /*= false*/, bool bTouchFiles /*= false*/, bool bPurgeTransient /*= false*/, bool bDeleteOldFiles /*= false*/, int32 InDaysToDeleteUnusedFiles /*= 60*/, int32 InMaxNumFoldersToCheck /*= -1*/, int32 InMaxContinuousFileChecks /*= -1*/, const TCHAR* InAccessLogFileName /*= nullptr*/)
 {
-	FFileSystemDerivedDataBackend* FileDDB = new FFileSystemDerivedDataBackend(CacheDirectory, bForceReadOnly, bTouchFiles, bPurgeTransient, bDeleteOldFiles, InDaysToDeleteUnusedFiles, InMaxNumFoldersToCheck, InMaxContinuousFileChecks);
+	FFileSystemDerivedDataBackend* FileDDB = new FFileSystemDerivedDataBackend(CacheDirectory, bForceReadOnly, bTouchFiles, bPurgeTransient, bDeleteOldFiles, InDaysToDeleteUnusedFiles, InMaxNumFoldersToCheck, InMaxContinuousFileChecks, InAccessLogFileName);
 	if (!FileDDB->IsUsable())
 	{
 		delete FileDDB;
