Fix progress bar winding backwards when downloading DDC bundles from S3, due to integer overflow.

diff --git a/Engine/Source/Developer/DerivedDataCache/Private/S3DerivedDataBackend.cpp b/Engine/Source/Developer/DerivedDataCache/Private/S3DerivedDataBackend.cpp
--- a/Engine/Source/Developer/DerivedDataCache/Private/S3DerivedDataBackend.cpp
+++ b/Engine/Source/Developer/DerivedDataCache/Private/S3DerivedDataBackend.cpp
@@ -743,8 +743,8 @@ FS3DerivedDataBackend::FS3DerivedDataBackend(const TCHAR* InRootManifestPath, co
 			{
 				FPlatformProcess::Sleep(0.1f);
 
-				int NumBytes = 0;
-				int MaxBytes = 0;
+				int64 NumBytes = 0;
+				int64 MaxBytes = 0;
 				bComplete = true;
 
 				CriticalSection.Lock();
@@ -756,7 +756,13 @@ FS3DerivedDataBackend::FS3DerivedDataBackend(const TCHAR* InRootManifestPath, co
 				}
 				CriticalSection.Unlock();
 
-				Context->UpdateProgress(NumBytes, MaxBytes);
+				int NumMB = (int)((NumBytes + (1024 * 1024) - 1) / (1024 * 1024));
+				int MaxMB = (int)((MaxBytes + (1024 * 1024) - 1) / (1024 * 1024));
+				if (MaxBytes > 0)
+				{
+					FText StatusText = FText::Format(NSLOCTEXT("S3DerivedDataBackend", "DownloadingDDCBundlesPct", "Downloading DDC bundles... ({0}Mb/{1}Mb)"), NumMB, MaxMB);
+					Context->StatusUpdate((int)((NumBytes * 1000) / MaxBytes), 1000, StatusText);
+				}
 			}
 
 			// Mount all the bundles
