Changed FStringView parameters to virtual functions to pass by value instead of by reference

This is the preferred way to use string views. These interfaces were added in 4.25 so this is the time to fix them without breaking user code.

#rnx

diff --git a/Engine/Source/Developer/DerivedDataCache/Private/DerivedDataCache.cpp b/Engine/Source/Developer/DerivedDataCache/Private/DerivedDataCache.cpp
--- a/Engine/Source/Developer/DerivedDataCache/Private/DerivedDataCache.cpp
+++ b/Engine/Source/Developer/DerivedDataCache/Private/DerivedDataCache.cpp
@@ -393,7 +393,7 @@ public:
 		return true;
 	}
 
-	virtual bool GetSynchronous(const TCHAR* CacheKey, TArray<uint8>& OutData, const FStringView& DataContext) override
+	virtual bool GetSynchronous(const TCHAR* CacheKey, TArray<uint8>& OutData, FStringView DataContext) override
 	{
 		DDC_SCOPE_CYCLE_COUNTER(DDC_GetSynchronous_Data);
 		UE_LOG(LogDerivedDataCache, Verbose, TEXT("GetSynchronous %s from '%.*s'"), CacheKey, DataContext.Len(), DataContext.GetData());
@@ -404,7 +404,7 @@ public:
 		return PendingTask.GetTask().bSuccess;
 	}
 
-	virtual uint32 GetAsynchronous(const TCHAR* CacheKey, const FStringView& DataContext) override
+	virtual uint32 GetAsynchronous(const TCHAR* CacheKey, FStringView DataContext) override
 	{
 		DDC_SCOPE_CYCLE_COUNTER(DDC_GetAsynchronous_Handle);
 		FScopeLock ScopeLock(&SynchronizationObject);
@@ -418,7 +418,7 @@ public:
 		return Handle;
 	}
 
-	virtual void Put(const TCHAR* CacheKey, TArrayView<const uint8> Data, const FStringView& DataContext, bool bPutEvenIfExists = false) override
+	virtual void Put(const TCHAR* CacheKey, TArrayView<const uint8> Data, FStringView DataContext, bool bPutEvenIfExists = false) override
 	{
 		DDC_SCOPE_CYCLE_COUNTER(DDC_Put);
 		UE_LOG(LogDerivedDataCache, Verbose, TEXT("Put %s from '%.*s'"), CacheKey, DataContext.Len(), DataContext.GetData());
diff --git a/Engine/Source/Developer/DerivedDataCache/Public/DerivedDataCacheInterface.h b/Engine/Source/Developer/DerivedDataCache/Public/DerivedDataCacheInterface.h
--- a/Engine/Source/Developer/DerivedDataCache/Public/DerivedDataCacheInterface.h
+++ b/Engine/Source/Developer/DerivedDataCache/Public/DerivedDataCacheInterface.h
@@ -158,7 +158,7 @@ public:
 	 * @param	DebugContext	A string used to describe the data being generated. Typically the path to the object that it is generated from is sufficient.
 	 * @return	true if the data was retrieved from the cache.
 	**/
-	virtual bool GetSynchronous(const TCHAR* CacheKey, TArray<uint8>& OutData, const FStringView& DebugContext) = 0; 
+	virtual bool GetSynchronous(const TCHAR* CacheKey, TArray<uint8>& OutData, FStringView DebugContext) = 0; 
 
 	/** 
 	 * Starts the async process of checking the cache and if the item is present, retrieving the cached results.
@@ -169,7 +169,7 @@ public:
 	 * @param	DebugContext	A string used to describe the data being generated. Typically the path to the object that it is generated from is sufficient.
 	 * @return	A handle that can be used for PollAsynchronousCompletion, WaitAsynchronousCompletion, and GetAsynchronousResults.
 	**/
-	virtual uint32 GetAsynchronous(const TCHAR* CacheKey, const FStringView& DebugContext) = 0;
+	virtual uint32 GetAsynchronous(const TCHAR* CacheKey, FStringView DebugContext) = 0;
 
 	/** 
 	 * Puts data into the cache. This is fire-and-forget and typically asynchronous.
@@ -180,7 +180,7 @@ public:
 	 * @param	Data		Data to put in the cache under this key
 	 * @param	DataContext	A string used to describe the data being generated. Typically the path to the object that it is generated from is sufficient.
 	**/
-	virtual void Put(const TCHAR* CacheKey, TArrayView<const uint8> Data, const FStringView& DataContext, bool bPutEvenIfExists = false) = 0;
+	virtual void Put(const TCHAR* CacheKey, TArrayView<const uint8> Data, FStringView DataContext, bool bPutEvenIfExists = false) = 0;
 
 	/**
 	 * Hint that the data associated with the key is transient and may be optionally purged from the cache.
diff --git a/Engine/Source/Developer/TargetPlatform/Private/TargetPlatformManagerModule.cpp b/Engine/Source/Developer/TargetPlatform/Private/TargetPlatformManagerModule.cpp
--- a/Engine/Source/Developer/TargetPlatform/Private/TargetPlatformManagerModule.cpp
+++ b/Engine/Source/Developer/TargetPlatform/Private/TargetPlatformManagerModule.cpp
@@ -179,7 +179,7 @@ public:
 		return nullptr;
 	}
 
-	virtual ITargetPlatform* FindTargetPlatform(const FStringView& Name) override
+	virtual ITargetPlatform* FindTargetPlatform(FStringView Name) override
 	{
 		GetTargetPlatforms(); // Populates PlatformsByName
 
diff --git a/Engine/Source/Developer/TargetPlatform/Public/Interfaces/ITargetPlatformManagerModule.h b/Engine/Source/Developer/TargetPlatform/Public/Interfaces/ITargetPlatformManagerModule.h
--- a/Engine/Source/Developer/TargetPlatform/Public/Interfaces/ITargetPlatformManagerModule.h
+++ b/Engine/Source/Developer/TargetPlatform/Public/Interfaces/ITargetPlatformManagerModule.h
@@ -64,7 +64,7 @@ public:
 	 * @param Name The target platform's short or long name.
 	 * @return The target platform, or nullptr if not found.
 	 */
-	virtual ITargetPlatform* FindTargetPlatform( const FStringView& Name ) = 0;
+	virtual ITargetPlatform* FindTargetPlatform( FStringView Name ) = 0;
 
 	/**
 	 * Finds a target platform by looking for one that supports a given value for a generic type of support
