Improved DDC error when no backends are available

[RN] fix Core/Cooker

diff --git a/Engine/Source/Developer/DerivedDataCache/Private/DerivedDataBackends.cpp b/Engine/Source/Developer/DerivedDataCache/Private/DerivedDataBackends.cpp
--- a/Engine/Source/Developer/DerivedDataCache/Private/DerivedDataBackends.cpp
+++ b/Engine/Source/Developer/DerivedDataCache/Private/DerivedDataBackends.cpp
@@ -90,8 +90,12 @@ public:
 				UE_LOG( LogDerivedDataCache, Fatal, TEXT("Unable to create backend graph using the default graph settings (%s) ini=%s."), *GraphName, *GEngineIni );
 			}
 			RootCache = ParseNode( TEXT("Root"), GEngineIni, *GraphName, ParsedNodes );
+			if (!RootCache)
+			{
+				UE_LOG(LogDerivedDataCache, Fatal, TEXT("Unable to create backend graph using the default graph settings (%s) ini=%s. ")
+					TEXT("At least one backend in the graph must be available."), *GraphName, *GEngineIni);
+			}
 		}
-		check(RootCache);
 
 		// Make sure AsyncPutWrapper and KeyLengthWrapper are created
 		if( !AsyncPutWrapper )
