Optimizations for DDC access on remote / slow drives.

If a filesystem node is not available not prompt the user and optionally retry incase they need to mount a drive or start VPN

Fiilesystem nodes now perform a speed test using a selection of 'DDC sized' files to determine a classification (local, fast, ok, slow).

Add a new 'ConsiderSlowAt' property to the 'Filesystem' DDC node type. If latency to the node is >= this value then the node will be marked as slow which disables touch'ing and reduces file stats

Interface Changes

- Add the concept of a speed class to nodes
- Add GetName to nodes for better debugging / logging
- WouldCache query that allows caches to opt of of consideration early and avoid async tasks being created.
- Create a new 'FileBackedDerivedDataBackend' class that's the for the memory/boot backend and future classes
- TryToPrefetch interface functions for future use

Behavior Changes

- Moved parameter parsing into FileSysteDerivedDataBackend as things were getting out of hand
- FileSystemDerivedDataBackend now performs a speed test using 'DDC sized' files in separate directories and applies a classification
- Slow locations turn off touching of data on read
- Slow locations always return true for CachedDataProbablyExists. It's faster just to try to read and fail
- If the shared DDC is not available the user is prompted incase they need to mount it.

#tests lots of PIE runs with / without this option

diff --git a/Engine/Config/BaseEngine.ini b/Engine/Config/BaseEngine.ini
--- a/Engine/Config/BaseEngine.ini
+++ b/Engine/Config/BaseEngine.ini
@@ -1617,8 +1617,8 @@ AsyncPut=(Type=AsyncPut, Inner=Hierarchy)
 Hierarchy=(Type=Hierarchical, Inner=Boot, Inner=Pak, Inner=EnginePak, Inner=Local, Inner=Shared)
 Boot=(Type=Boot, Filename="%GAMEDIR%DerivedDataCache/Boot.ddc", MaxCacheSize=512)
 Local=(Type=FileSystem, ReadOnly=false, Clean=false, Flush=false, PurgeTransient=true, DeleteUnused=true, UnusedFileAge=34, FoldersToClean=-1, Path=%ENGINEDIR%DerivedDataCache, EnvPathOverride=UE-LocalDataCachePath, EditorOverrideSetting=LocalDerivedDataCache)
-Shared=(Type=FileSystem, ReadOnly=false, Clean=false, Flush=false, DeleteUnused=true, UnusedFileAge=10, FoldersToClean=10, MaxFileChecksPerSec=1, Path=?EpicDDC, EnvPathOverride=UE-SharedDataCachePath, EditorOverrideSetting=SharedDerivedDataCache, CommandLineOverride=SharedDataCachePath)
-AltShared=(Type=FileSystem, ReadOnly=true, Clean=false, Flush=false, DeleteUnused=true, UnusedFileAge=23, FoldersToClean=10, MaxFileChecksPerSec=1, Path=?EpicDDC2, EnvPathOverride=UE-SharedDataCachePath2)
+Shared=(Type=FileSystem, ReadOnly=false, Clean=false, Flush=false, DeleteUnused=true, UnusedFileAge=10, FoldersToClean=10, MaxFileChecksPerSec=1, ConsiderSlowAt=70, Path=?EpicDDC, EnvPathOverride=UE-SharedDataCachePath, EditorOverrideSetting=SharedDerivedDataCache, CommandLineOverride=SharedDataCachePath)
+AltShared=(Type=FileSystem, ReadOnly=true, Clean=false, Flush=false, DeleteUnused=true, UnusedFileAge=23, FoldersToClean=10, MaxFileChecksPerSec=1, ConsiderSlowAt=70, Path=?EpicDDC2, EnvPathOverride=UE-SharedDataCachePath2)
 Pak=(Type=ReadPak, Filename="%GAMEDIR%DerivedDataCache/DDC.ddp")
 EnginePak=(Type=ReadPak, Filename=%ENGINEDIR%DerivedDataCache/DDC.ddp)
 
diff --git a/Engine/Source/Developer/DerivedDataCache/Private/DerivedDataBackendAsyncPutWrapper.cpp b/Engine/Source/Developer/DerivedDataCache/Private/DerivedDataBackendAsyncPutWrapper.cpp
--- a/Engine/Source/Developer/DerivedDataCache/Private/DerivedDataBackendAsyncPutWrapper.cpp
+++ b/Engine/Source/Developer/DerivedDataCache/Private/DerivedDataBackendAsyncPutWrapper.cpp
@@ -94,7 +94,7 @@ public:
 
 FDerivedDataBackendAsyncPutWrapper::FDerivedDataBackendAsyncPutWrapper(FDerivedDataBackendInterface* InInnerBackend, bool bCacheInFlightPuts)
 	: InnerBackend(InInnerBackend)
-	, InflightCache(bCacheInFlightPuts ? (new FMemoryDerivedDataBackend) : NULL)
+	, InflightCache(bCacheInFlightPuts ? (new FMemoryDerivedDataBackend(TEXT("AsyncPutCache"))) : NULL)
 {
 	check(InnerBackend);
 }
@@ -105,6 +105,11 @@ bool FDerivedDataBackendAsyncPutWrapper::IsWritable()
 	return InnerBackend->IsWritable();
 }
 
+FDerivedDataBackendInterface::ESpeedClass FDerivedDataBackendAsyncPutWrapper::GetSpeedClass()
+{
+	return InnerBackend->GetSpeedClass();
+}
+
 /**
  * Synchronous test for the existence of a cache item
  *
@@ -119,6 +124,31 @@ bool FDerivedDataBackendAsyncPutWrapper::CachedDataProbablyExists(const TCHAR* C
 	return Result;
 }
 
+bool FDerivedDataBackendAsyncPutWrapper::TryToPrefetch(const TCHAR* CacheKey)
+{
+	COOK_STAT(auto Timer = UsageStats.TimePrefetch());
+
+	bool SkipCheck = (!InflightCache && InflightCache->CachedDataProbablyExists(CacheKey));
+	
+	bool Hit = false;
+
+	if (!SkipCheck)
+	{
+		Hit = InnerBackend->TryToPrefetch(CacheKey);
+	}
+
+	COOK_STAT(if (Hit) { Timer.AddHit(0); });
+	return Hit;
+}
+
+/*
+	Determine if we would cache this by asking all our inner layers
+*/
+bool FDerivedDataBackendAsyncPutWrapper::WouldCache(const TCHAR* CacheKey, TArrayView<const uint8> InData)
+{
+	return InnerBackend->WouldCache(CacheKey, InData);
+}
+
 bool FDerivedDataBackendAsyncPutWrapper::GetCachedData(const TCHAR* CacheKey, TArray<uint8>& OutData)
 {
 	COOK_STAT(auto Timer = UsageStats.TimeGet());
diff --git a/Engine/Source/Developer/DerivedDataCache/Private/DerivedDataBackendAsyncPutWrapper.h b/Engine/Source/Developer/DerivedDataCache/Private/DerivedDataBackendAsyncPutWrapper.h
--- a/Engine/Source/Developer/DerivedDataCache/Private/DerivedDataBackendAsyncPutWrapper.h
+++ b/Engine/Source/Developer/DerivedDataCache/Private/DerivedDataBackendAsyncPutWrapper.h
@@ -63,9 +63,18 @@ public:
 	 */
 	FDerivedDataBackendAsyncPutWrapper(FDerivedDataBackendInterface* InInnerBackend, bool bCacheInFlightPuts);
 
+	/** Return a name for this interface */
+	virtual FString GetName() const override
+	{
+		return FString::Printf(TEXT("AsyncPutWrapper (%s)"), *InnerBackend->GetName());
+	}
+
 	/** return true if this cache is writable **/
 	virtual bool IsWritable() override;
 
+	/** Returns a class of speed for this interface **/
+	virtual ESpeedClass GetSpeedClass() override;
+
 	/**
 	 * Synchronous test for the existence of a cache item
 	 *
@@ -74,6 +83,19 @@ public:
 	 */
 	virtual bool CachedDataProbablyExists(const TCHAR* CacheKey) override;
 
+	/**
+	 * Attempts to make sure the cached data will be available as optimally as possible. This is left up to the implementation to do
+	 * @param	CacheKey	Alphanumeric+underscore key of this cache item
+	 * @return				true if any steps were performed to optimize future retrieval
+	 */
+	virtual bool TryToPrefetch(const TCHAR* CacheKey) override;
+
+	/**
+	 * Allows the DDC backend to determine if it wants to cache the provided data. Reasons for returning false could be a slow connection,
+	 * a file size limit, etc.
+	 */
+	virtual bool WouldCache(const TCHAR* CacheKey, TArrayView<const uint8> InData) override;
+
 	/**
 	 * Synchronous retrieve of a cache item
 	 *
diff --git a/Engine/Source/Developer/DerivedDataCache/Private/DerivedDataBackendCorruptionWrapper.h b/Engine/Source/Developer/DerivedDataCache/Private/DerivedDataBackendCorruptionWrapper.h
--- a/Engine/Source/Developer/DerivedDataCache/Private/DerivedDataBackendCorruptionWrapper.h
+++ b/Engine/Source/Developer/DerivedDataCache/Private/DerivedDataBackendCorruptionWrapper.h
@@ -81,12 +81,24 @@ public:
 		check(InnerBackend);
 	}
 
+	/** Return a name for this interface */
+	virtual FString GetName() const override
+	{
+		return FString::Printf(TEXT("CorruptionWrapper (%s)"), *InnerBackend->GetName());
+	}
+
 	/** return true if this cache is writable **/
 	virtual bool IsWritable() override
 	{
 		return InnerBackend->IsWritable();
 	}
 
+	/** Returns a class of speed for this interface **/
+	virtual ESpeedClass GetSpeedClass() override
+	{
+		return InnerBackend->GetSpeedClass();
+	}
+
 	/**
 	 * Synchronous test for the existence of a cache item
 	 *
@@ -129,7 +141,7 @@ public:
 				FDerivedDataTrailer RecomputedTrailer(OutData);
 				if (Trailer == RecomputedTrailer)
 				{
-					UE_LOG(LogDerivedDataCache, Verbose, TEXT("FDerivedDataBackendCorruptionWrapper: cache hit, footer is ok %s"),CacheKey);
+					UE_LOG(LogDerivedDataCache, VeryVerbose, TEXT("FDerivedDataBackendCorruptionWrapper: cache hit, footer is ok %s"),CacheKey);
 				}
 				else
 				{
@@ -200,6 +212,16 @@ public:
 		});
 	}
 
+	virtual bool TryToPrefetch(const TCHAR* CacheKey) override
+	{
+		return InnerBackend->TryToPrefetch(CacheKey);
+	}
+
+	virtual bool WouldCache(const TCHAR* CacheKey, TArrayView<const uint8> InData) override
+	{
+		return InnerBackend->WouldCache(CacheKey, InData);
+	}
+
 private:
 	FDerivedDataCacheUsageStats UsageStats;
 
diff --git a/Engine/Source/Developer/DerivedDataCache/Private/DerivedDataBackendVerifyWrapper.h b/Engine/Source/Developer/DerivedDataCache/Private/DerivedDataBackendVerifyWrapper.h
--- a/Engine/Source/Developer/DerivedDataCache/Private/DerivedDataBackendVerifyWrapper.h
+++ b/Engine/Source/Developer/DerivedDataCache/Private/DerivedDataBackendVerifyWrapper.h
@@ -32,11 +32,22 @@ public:
 		check(InnerBackend);
 	}
 
+	/** Return a name for this interface */
+	virtual FString GetName() const override
+	{
+		return FString::Printf(TEXT("VerifyWrapper (%s)"), *InnerBackend->GetName());
+	}
+
 	virtual bool IsWritable() override
 	{
 		return true;
 	}
 
+	virtual ESpeedClass GetSpeedClass() override
+	{
+		return ESpeedClass::Local;
+	}
+
 	virtual bool CachedDataProbablyExists(const TCHAR* CacheKey) override
 	{
 		COOK_STAT(auto Timer = UsageStats.TimeProbablyExists());
@@ -136,6 +147,16 @@ public:
 		});
 	}
 
+	virtual bool TryToPrefetch(const TCHAR* CacheKey) override
+	{
+		return InnerBackend->TryToPrefetch(CacheKey);
+	}
+
+	virtual bool WouldCache(const TCHAR* CacheKey, TArrayView<const uint8> InData) override
+	{
+		return InnerBackend->WouldCache(CacheKey, InData);
+	}
+
 private:
 	FDerivedDataCacheUsageStats UsageStats;
 
diff --git a/Engine/Source/Developer/DerivedDataCache/Private/DerivedDataBackends.cpp b/Engine/Source/Developer/DerivedDataCache/Private/DerivedDataBackends.cpp
--- a/Engine/Source/Developer/DerivedDataCache/Private/DerivedDataBackends.cpp
+++ b/Engine/Source/Developer/DerivedDataCache/Private/DerivedDataBackends.cpp
@@ -24,12 +24,13 @@
 #include "Modules/ModuleManager.h"
 #include "Misc/ConfigCacheIni.h"
 
+
 DEFINE_LOG_CATEGORY(LogDerivedDataCache);
 
 #define MAX_BACKEND_KEY_LENGTH (120)
 #define LOCTEXT_NAMESPACE "DerivedDataBackendGraph"
 
-FDerivedDataBackendInterface* CreateFileSystemDerivedDataBackend(const TCHAR* CacheDirectory, bool bForceReadOnly = false, bool bTouchFiles = false, bool bPurgeTransient = false, bool bDeleteOldFiles = false, int32 InDaysToDeleteUnusedFiles = 60, int32 InMaxNumFoldersToCheck = -1, int32 InMaxContinuousFileChecks = -1, const TCHAR* InAccessLogFileName = nullptr);
+FDerivedDataBackendInterface* CreateFileSystemDerivedDataBackend(const TCHAR* CacheDirectory, const TCHAR* InParams, const int InMissRate, const TCHAR* InAccessLogFileName = nullptr);
 
 /**
   * This class is used to create a singleton that represents the derived data cache hierarchy and all of the wrappers necessary
@@ -161,7 +162,7 @@ public:
 					}
 					else
 					{
-						UE_LOG( LogDerivedDataCache, Warning, TEXT("FDerivedDataBackendGraph:  Unable to create %s Boot cache because only one Boot cache node is supported."), *NodeName );
+						UE_LOG( LogDerivedDataCache, Warning, TEXT("FDerivedDataBackendGraph:  Unable to create %s Boot cache because only one Boot or Sparse cache node is supported."), *NodeName );
 					}
 				}
 				else if( NodeType == TEXT("Memory") )
@@ -557,55 +558,57 @@ public:
 		}
 		else
 		{
-			const bool bReadOnly = GetParsedBool( Entry, TEXT("ReadOnly=") );
-			const bool bClean = GetParsedBool( Entry, TEXT("Clean=") );
-			const bool bFlush = GetParsedBool( Entry, TEXT("Flush=") );
-			const bool bTouch = GetParsedBool( Entry, TEXT("Touch=") );
-			const bool bPurgeTransient = GetParsedBool( Entry, TEXT("PurgeTransient=") );
-
-			bool bDeleteUnused = true; // On by default
-			FParse::Bool( Entry, TEXT("DeleteUnused="), bDeleteUnused );
-			int32 UnusedFileAge = 17;
-			FParse::Value( Entry, TEXT("UnusedFileAge="), UnusedFileAge );
-			int32 MaxFoldersToClean = -1;
-			FParse::Value( Entry, TEXT("FoldersToClean="), MaxFoldersToClean );
-			int32 MaxFileChecksPerSec = -1;
-			FParse::Value( Entry, TEXT("MaxFileChecksPerSec="), MaxFileChecksPerSec );
-			FString WriteAccessLog;
-			FParse::Value( Entry, TEXT("WriteAccessLog="), WriteAccessLog );
-
-			if( bFlush )
-			{
-				IFileManager::Get().DeleteDirectory( *(Path / TEXT("")), false, true );
-			}
-			else if( bClean )
-			{
-				DeleteOldFiles( *Path );
-			}
-
 			FDerivedDataBackendInterface* InnerFileSystem = NULL;
-			
-			// Don't create the file system if shared data cache directory is not mounted
-			bool bShared = FCString::Stricmp(NodeName, TEXT("Shared")) == 0;
-			if( !bShared || IFileManager::Get().DirectoryExists(*Path) )
-			{
-				InnerFileSystem = CreateFileSystemDerivedDataBackend( *Path, bReadOnly, bTouch, bPurgeTransient, bDeleteUnused, UnusedFileAge, MaxFoldersToClean, MaxFileChecksPerSec, *WriteAccessLog);
-			}
 
-			if( InnerFileSystem )
+			// Try to set up the shared drive, allow user to correct any issues that may exist.
+			bool RetryOnFailure = false;
+			do
 			{
-				bUsingSharedDDC = bUsingSharedDDC ? bUsingSharedDDC : bShared;
+				RetryOnFailure = false;
 
-				DataCache = new FDerivedDataBackendCorruptionWrapper( InnerFileSystem );
-				UE_LOG( LogDerivedDataCache, Log, TEXT("Using %s data cache path %s: %s"), NodeName, *Path, bReadOnly ? TEXT("ReadOnly") : TEXT("Writable") );
-				Directories.AddUnique(Path);
-			}
-			else
-			{
-				UE_LOG( LogDerivedDataCache, Warning, TEXT("%s data cache path (%s) was not usable, will not use it."), NodeName, *Path );
-			}
+				// Don't create the file system if shared data cache directory is not mounted
+				bool bShared = FCString::Stricmp(NodeName, TEXT("Shared")) == 0;
+				
+				FString WriteAccessLog;
+				FParse::Value( Entry, TEXT("WriteAccessLog="), WriteAccessLog );
+
+				// look for -ddclocalmissrate, -ddcsharedmissrate etc
+				FString MissRateCmd = FString::Printf(TEXT("DDC%sMissRate="), NodeName);
+				int MissRate = 0;
+				FParse::Value(FCommandLine::Get(), *MissRateCmd, MissRate);
+
+				if (!bShared || IFileManager::Get().DirectoryExists(*Path))
+				{
+					InnerFileSystem = CreateFileSystemDerivedDataBackend( *Path, Entry, MissRate, *WriteAccessLog);
+				}
+
+				if (InnerFileSystem)
+				{
+					bUsingSharedDDC = bUsingSharedDDC ? bUsingSharedDDC : bShared;
+	
+					DataCache = new FDerivedDataBackendCorruptionWrapper(InnerFileSystem);
+					UE_LOG(LogDerivedDataCache, Log, TEXT("Using %s data cache path %s: %s"), NodeName, *Path, !InnerFileSystem->IsWritable() ? TEXT("ReadOnly") : TEXT("Writable"));
+					Directories.AddUnique(Path);
+				}
+				else
+				{
+					FString Message = FString::Printf(TEXT("%s data cache path (%s) is unavailable so cache will be disabled."), NodeName, *Path);
+					
+					UE_LOG(LogDerivedDataCache, Warning, TEXT("%s"), *Message);
+
+					// Give the user a chance to retry incase they need to connect a network drive or something.
+					if (!FApp::IsUnattended())
+					{
+						Message += FString::Printf(TEXT("\n\nRetry connection to %s?"), *Path);
+						EAppReturnType::Type MessageReturn = FPlatformMisc::MessageBoxExt(EAppMsgType::YesNo, *Message, TEXT("Could not access DDC"));
+						RetryOnFailure = MessageReturn == EAppReturnType::Yes;
+					}
+				}
+			} while (RetryOnFailure);
 		}
 
+
+
 		return DataCache;
 	}
 
@@ -678,7 +681,7 @@ public:
 	 * @param OutFilename filename specified for the cache
 	 * @return Boot data cache backend interface instance or NULL if unsuccessful
 	 */
-	FMemoryDerivedDataBackend* ParseBootCache( const TCHAR* NodeName, const TCHAR* Entry, FString& OutFilename )
+	FFileBackedDerivedDataBackend* ParseBootCache( const TCHAR* NodeName, const TCHAR* Entry, FString& OutFilename )
 	{
 		FMemoryDerivedDataBackend* Cache = NULL;
 
@@ -700,7 +703,7 @@ public:
 			MaxCacheSize = FMath::Min(MaxCacheSize, MaxSupportedCacheSize);
 
 			UE_LOG( LogDerivedDataCache, Display, TEXT("Max Cache Size: %d MB"), MaxCacheSize);
-			Cache = new FMemoryDerivedDataBackend(MaxCacheSize * 1024 * 1024);
+			Cache = new FMemoryDerivedDataBackend(TEXT("Boot"), MaxCacheSize * 1024 * 1024);
 
 			if( Cache && Filename.Len() )
 			{
@@ -709,7 +712,6 @@ public:
 				if (MaxCacheSize > 0 && IFileManager::Get().FileSize(*Filename) >= (MaxCacheSize * 1024 * 1024))
 				{
 					UE_LOG( LogDerivedDataCache, Warning, TEXT("FDerivedDataBackendGraph:  %s filename exceeds max size."), NodeName );
-					return Cache;
 				}
 
 				if (IFileManager::Get().FileSize(*Filename) < 0)
@@ -743,7 +745,7 @@ public:
 		FString Filename;
 
 		FParse::Value( Entry, TEXT("Filename="), Filename );
-		Cache = new FMemoryDerivedDataBackend;
+		Cache = new FMemoryDerivedDataBackend(TEXT("Memory"));
 		if( Cache && Filename.Len() )
 		{
 			if( Cache->LoadCache( *Filename ) )
@@ -927,15 +929,6 @@ public:
 
 private:
 
-	/** Delete the old files in a directory **/
-	void DeleteOldFiles(const TCHAR* Directory)
-	{
-		float MinimumDaysToKeepFile = 7;
-		GConfig->GetFloat( *GraphName, TEXT("MinimumDaysToKeepFile"), MinimumDaysToKeepFile, GEngineIni );
-		check(MinimumDaysToKeepFile > 0.0f); // sanity
-		//@todo 
-	}
-
 	/** Delete all created backends in the reversed order they were created. */
 	void DestroyCreatedBackends()
 	{
@@ -981,7 +974,7 @@ private:
 	TArray< FDerivedDataBackendInterface* > CreatedBackends;
 
 	/** Instances of backend interfaces which exist in only one copy */
-	FMemoryDerivedDataBackend*		BootCache;
+	FFileBackedDerivedDataBackend*	BootCache;
 	FPakFileDerivedDataBackend*		WritePakCache;
 	FDerivedDataBackendInterface*	AsyncPutWrapper;
 	FDerivedDataBackendInterface*	KeyLengthWrapper;
diff --git a/Engine/Source/Developer/DerivedDataCache/Private/DerivedDataCache.cpp b/Engine/Source/Developer/DerivedDataCache/Private/DerivedDataCache.cpp
--- a/Engine/Source/Developer/DerivedDataCache/Private/DerivedDataCache.cpp
+++ b/Engine/Source/Developer/DerivedDataCache/Private/DerivedDataCache.cpp
@@ -303,7 +303,7 @@ public:
 		DDC_SCOPE_CYCLE_COUNTER(DDC_GetSynchronous);
 		check(DataDeriver);
 		FString CacheKey = FDerivedDataCache::BuildCacheKey(DataDeriver);
-		UE_LOG(LogDerivedDataCache, Verbose, TEXT("GetSynchronous %s from '%s'"), *CacheKey, *DataDeriver->GetDebugContextString());
+		UE_LOG(LogDerivedDataCache, VeryVerbose, TEXT("GetSynchronous %s from '%s'"), *CacheKey, *DataDeriver->GetDebugContextString());
 		FAsyncTask<FBuildAsyncWorker> PendingTask(DataDeriver, *CacheKey, true);
 		AddToAsyncCompletionCounter(1);
 		PendingTask.StartSynchronousTask();
@@ -321,7 +321,7 @@ public:
 		FScopeLock ScopeLock(&SynchronizationObject);
 		const uint32 Handle = NextHandle();
 		FString CacheKey = FDerivedDataCache::BuildCacheKey(DataDeriver);
-		UE_LOG(LogDerivedDataCache, Verbose, TEXT("GetAsynchronous %s from '%s'"), *CacheKey, *DataDeriver->GetDebugContextString());
+		UE_LOG(LogDerivedDataCache, VeryVerbose, TEXT("GetAsynchronous %s from '%s'"), *CacheKey, *DataDeriver->GetDebugContextString());
 		const bool bSync = !DataDeriver->IsBuildThreadsafe();
 		FAsyncTask<FBuildAsyncWorker>* AsyncTask = new FAsyncTask<FBuildAsyncWorker>(DataDeriver, *CacheKey, bSync);
 		check(!PendingTasks.Contains(Handle));
@@ -396,7 +396,7 @@ public:
 	virtual bool GetSynchronous(const TCHAR* CacheKey, TArray<uint8>& OutData, FStringView DataContext) override
 	{
 		DDC_SCOPE_CYCLE_COUNTER(DDC_GetSynchronous_Data);
-		UE_LOG(LogDerivedDataCache, Verbose, TEXT("GetSynchronous %s from '%.*s'"), CacheKey, DataContext.Len(), DataContext.GetData());
+		UE_LOG(LogDerivedDataCache, VeryVerbose, TEXT("GetSynchronous %s from '%.*s'"), CacheKey, DataContext.Len(), DataContext.GetData());
 		FAsyncTask<FBuildAsyncWorker> PendingTask((FDerivedDataPluginInterface*)NULL, CacheKey, true);
 		AddToAsyncCompletionCounter(1);
 		PendingTask.StartSynchronousTask();
@@ -408,7 +408,7 @@ public:
 	{
 		DDC_SCOPE_CYCLE_COUNTER(DDC_GetAsynchronous_Handle);
 		FScopeLock ScopeLock(&SynchronizationObject);
-		UE_LOG(LogDerivedDataCache, Verbose, TEXT("GetAsynchronous %s from '%.*s'"), CacheKey, DataContext.Len(), DataContext.GetData());
+		UE_LOG(LogDerivedDataCache, VeryVerbose, TEXT("GetAsynchronous %s from '%.*s'"), CacheKey, DataContext.Len(), DataContext.GetData());
 		const uint32 Handle = NextHandle();
 		FAsyncTask<FBuildAsyncWorker>* AsyncTask = new FAsyncTask<FBuildAsyncWorker>((FDerivedDataPluginInterface*)NULL, CacheKey, false);
 		check(!PendingTasks.Contains(Handle));
@@ -421,7 +421,7 @@ public:
 	virtual void Put(const TCHAR* CacheKey, TArrayView<const uint8> Data, FStringView DataContext, bool bPutEvenIfExists = false) override
 	{
 		DDC_SCOPE_CYCLE_COUNTER(DDC_Put);
-		UE_LOG(LogDerivedDataCache, Verbose, TEXT("Put %s from '%.*s'"), CacheKey, DataContext.Len(), DataContext.GetData());
+		UE_LOG(LogDerivedDataCache, VeryVerbose, TEXT("Put %s from '%.*s'"), CacheKey, DataContext.Len(), DataContext.GetData());
 		STAT(double ThisTime = 0);
 		{
 			SCOPE_SECONDS_COUNTER(ThisTime);
diff --git a/Engine/Source/Developer/DerivedDataCache/Private/DerivedDataLimitKeyLengthWrapper.h b/Engine/Source/Developer/DerivedDataCache/Private/DerivedDataLimitKeyLengthWrapper.h
--- a/Engine/Source/Developer/DerivedDataCache/Private/DerivedDataLimitKeyLengthWrapper.h
+++ b/Engine/Source/Developer/DerivedDataCache/Private/DerivedDataLimitKeyLengthWrapper.h
@@ -27,12 +27,24 @@ public:
 		check(InnerBackend);
 	}
 
+	/** Return a name for this interface */
+	virtual FString GetName() const override 
+	{ 
+		return FString::Printf(TEXT("LimitKeyLengthWrapper (%s)"), *InnerBackend->GetName());
+	}
+
 	/** return true if this cache is writable **/
 	virtual bool IsWritable() override
 	{
 		return InnerBackend->IsWritable();
 	}
 
+	/** Returns a class of speed for this interface **/
+	virtual ESpeedClass GetSpeedClass() override
+	{
+		return InnerBackend->GetSpeedClass();
+	}
+
 	/**
 	 * Synchronous test for the existence of a cache item
 	 *
@@ -51,6 +63,34 @@ public:
 		}
 		return Result;
 	}
+
+	/**
+	 * Attempts to make sure the cached data will be available as optimally as possible. This is left up to the implementation to do
+	 * @param	CacheKey	Alphanumeric+underscore key of this cache item
+	 * @return				true if any steps were performed to optimize future retrieval
+	 */
+	virtual bool TryToPrefetch(const TCHAR* CacheKey) override
+	{
+		COOK_STAT(auto Timer = UsageStats.TimePrefetch());
+		FString NewKey;
+		ShortenKey(CacheKey, NewKey);
+		bool Result = InnerBackend->TryToPrefetch(*NewKey);
+		if (Result)
+		{
+			COOK_STAT(Timer.AddHit(0));
+		}
+
+		return Result;
+	}
+
+	/*
+		Determines if we have any interest in caching this data
+	*/
+	virtual bool WouldCache(const TCHAR* CacheKey, TArrayView<const uint8> InData) override
+	{
+		return InnerBackend->WouldCache(CacheKey, InData);
+	}
+
 	/**
 	 * Synchronous retrieve of a cache item
 	 *
@@ -101,7 +141,7 @@ public:
 					OutData.RemoveAt(OutData.Num() - KeyLen, KeyLen);
 					if (Compare == 0)
 					{
-						UE_LOG(LogDerivedDataCache, Verbose, TEXT("FDerivedDataLimitKeyLengthWrapper: cache hit, key match is ok %s"), CacheKey);
+						UE_LOG(LogDerivedDataCache, VeryVerbose, TEXT("FDerivedDataLimitKeyLengthWrapper: cache hit, key match is ok %s"), CacheKey);
 					}
 					else
 					{
diff --git a/Engine/Source/Developer/DerivedDataCache/Private/FileBackedDerivedDataBackend.h b/Engine/Source/Developer/DerivedDataCache/Private/FileBackedDerivedDataBackend.h
new file mode 100644
--- /dev/null
+++ b/Engine/Source/Developer/DerivedDataCache/Private/FileBackedDerivedDataBackend.h
@@ -0,0 +1,38 @@
+// Copyright Epic Games, Inc. All Rights Reserved.
+
+#pragma once
+
+#include "CoreMinimal.h"
+#include "DerivedDataBackendInterface.h"
+
+
+/** 
+ * A file backed backend. This is used for local caching of data where more flexibility is needed than loose files.
+**/
+class FFileBackedDerivedDataBackend : public FDerivedDataBackendInterface
+{
+public:
+	
+	/**
+	 * Save the cache to disk
+	 * @param	Filename	Filename to save
+	 * @return	true if file was saved successfully
+	 */
+	virtual bool SaveCache(const TCHAR* Filename) = 0;
+
+	/**
+	 * Load the cache from disk
+	 * @param	Filename	Filename to load
+	 * @return	true if file was loaded successfully
+	 */
+	virtual bool LoadCache(const TCHAR* Filename) = 0;
+
+	/**
+	 * Disables use of this cache and allows the underlying implementation to free up memory
+	 */
+	virtual void Disable() = 0;
+
+private:
+
+};
+
diff --git a/Engine/Source/Developer/DerivedDataCache/Private/FileSystemDerivedDataBackend.cpp b/Engine/Source/Developer/DerivedDataCache/Private/FileSystemDerivedDataBackend.cpp
--- a/Engine/Source/Developer/DerivedDataCache/Private/FileSystemDerivedDataBackend.cpp
+++ b/Engine/Source/Developer/DerivedDataCache/Private/FileSystemDerivedDataBackend.cpp
@@ -11,6 +11,7 @@
 #include "Misc/Guid.h"
 #include "Misc/ScopeLock.h"
 #include "Async/TaskGraphInterfaces.h"
+#include "Async/Async.h"
 
 #include "DerivedDataCacheInterface.h"
 #include "DerivedDataBackendInterface.h"
@@ -42,6 +43,22 @@ FString BuildPathForCacheKey(const TCHAR* CacheKey)
 	return HashPath / Key + TEXT(".udd");
 }
 
+/**
+ * Helper function to get the value of parsed bool as the return value
+ **/
+bool GetParsedBool(const TCHAR* Stream, const TCHAR* Match) 
+{
+	bool bValue = 0;
+	FParse::Bool(Stream, Match, bValue);
+	return bValue;
+}
+
+/** Delete the old files in a directory **/
+void DeleteOldFiles(const TCHAR* Directory, int Age)
+{
+	// @todo(agrant) the original implementation of this did nothing. Do we need this?
+}
+
 /** 
  * Cache server that uses the OS filesystem
  * The entire API should be callable from any thread (except the singleton can be assumed to be called at least once before concurrent access).
@@ -54,111 +71,368 @@ public:
 	 * @param InCacheDirectory	directory to store the cache in
 	 * @param bForceReadOnly	if true, do not attempt to write to this cache
 	*/
-	FFileSystemDerivedDataBackend(const TCHAR* InCacheDirectory, bool bForceReadOnly, bool bTouchFiles, bool bPurgeTransientData, bool bDeleteOldFiles, int32 InDaysToDeleteUnusedFiles, int32 InMaxNumFoldersToCheck, int32 InMaxContinuousFileChecks, const TCHAR* InAccessLogFileName)
+	FFileSystemDerivedDataBackend(const TCHAR* InCacheDirectory, const TCHAR* InParams, const int InMissRate, const TCHAR* InAccessLogFileName)
 		: CachePath(InCacheDirectory)
-		, bReadOnly(bForceReadOnly)
-		, bFailed(true)
-		, bTouch(bTouchFiles)
-		, bPurgeTransient(bPurgeTransientData)
-		, DaysToDeleteUnusedFiles(InDaysToDeleteUnusedFiles)
+		, SpeedClass(ESpeedClass::Unknown)
+		, bReadOnly(false)
+		, bTouch(false)
+		, bPurgeTransient(false)
+		, DaysToDeleteUnusedFiles(15)
+		, bDisabled(false)
 		, TotalEstimatedBuildTime(0)
+		, DebugMissRate(0)
 	{
 		// If we find a platform that has more stingent limits, this needs to be rethought.
 		checkf(MAX_BACKEND_KEY_LENGTH + MAX_CACHE_DIR_LEN + MAX_BACKEND_NUMBERED_SUBFOLDER_LENGTH + MAX_CACHE_EXTENTION_LEN < FPlatformMisc::GetMaxPathLength(),
 			TEXT("Not enough room left for cache keys in max path."));
-		const double SlowInitDuration = 10.0;
-		double AccessDuration = 0.0;
 
 		check(CachePath.Len());
 		FPaths::NormalizeFilename(CachePath);
-		const FString AbsoluteCachePath = IFileManager::Get().ConvertToAbsolutePathForExternalAppForRead(*CachePath);
-		if (AbsoluteCachePath.Len() > MAX_CACHE_DIR_LEN)
+
+		// Params that override our instance defaults
+		bReadOnly = GetParsedBool(InParams, TEXT("ReadOnly="));
+		bTouch = GetParsedBool(InParams, TEXT("Touch="));
+		bPurgeTransient = GetParsedBool(InParams, TEXT("PurgeTransient="));
+		FParse::Value(InParams, TEXT("UnusedFileAge="), DaysToDeleteUnusedFiles);
+
+		// Params that are used when setting up our path
+		const bool bClean = GetParsedBool(InParams, TEXT("Clean="));
+		const bool bFlush = GetParsedBool(InParams, TEXT("Flush="));
+	
+
+		// These are used to determine if we kick off a worker to cleanup our cache
+		bool bDeleteUnused = true; // On by default
+		int32 MaxFoldersToClean = -1;
+		int32 MaxFileChecksPerSec = -1;
+
+		FParse::Bool(InParams, TEXT("DeleteUnused="), bDeleteUnused);
+
+		if (bDeleteUnused)
 		{
-			const FText ErrorMessage = FText::Format( NSLOCTEXT("DerivedDataCache", "PathTooLong", "Cache path {0} is longer than {1} characters...please adjust [DerivedDataBackendGraph] paths to be shorter (this leaves more room for cache keys)."), FText::FromString( AbsoluteCachePath ), FText::AsNumber( MAX_CACHE_DIR_LEN ) );
-			FMessageDialog::Open(EAppMsgType::Ok, ErrorMessage);
-			UE_LOG(LogDerivedDataCache, Fatal, TEXT("%s"), *ErrorMessage.ToString());
+			FParse::Value(InParams, TEXT("FoldersToClean="), MaxFoldersToClean);
+			FParse::Value(InParams, TEXT("MaxFileChecksPerSec="), MaxFileChecksPerSec);
 		}
-		if (!bReadOnly)
+
+		if (bFlush)
+		{
+			IFileManager::Get().DeleteDirectory(*(CachePath / TEXT("")), false, true);
+		}
+		else if (bClean)
+		{
+			DeleteOldFiles(InCacheDirectory, DaysToDeleteUnusedFiles);
+		}
+
+		// check latency and speed. Read values should always be valid
+		double ReadSpeedMBs = 0.0;
+		double WriteSpeedMBs = 0.0;
+		double SeekTimeMS = 0.0;
+
+		/* Speeds faster than this are considered local*/
+		const float ConsiderFastAtMS = 10;
+		/* Speeds faster than this are ok. Everything else is slow. This value can be overridden in the ini file */
+		float ConsiderSlowAtMS = 50;
+		FParse::Value(InParams, TEXT("ConsiderSlowAt="), ConsiderSlowAtMS);
+
+		// can skip the speed test so everything acts as fast local (e.g. 4.25 and earlier behavior).  Mostly used
+		// for verifying speed related optimizations can can be removed soon (3/17/20).
+		bool SkipSpeedTest = FParse::Param(FCommandLine::Get(), TEXT("ddcskipspeedtest"));
+		if (SkipSpeedTest)
 		{
-			double TestStart = FPlatformTime::Seconds();
-			FString TempFilename = CachePath / FGuid::NewGuid().ToString() + ".tmp";
-			FFileHelper::SaveStringToFile(FString("TEST"),*TempFilename);
-			uint32 WriteErrorCode = FPlatformMisc::GetLastError();
-			int32 TestFileSize = IFileManager::Get().FileSize(*TempFilename);
-			if (TestFileSize < 4)
+			ReadSpeedMBs = 999;
+			WriteSpeedMBs = 999;
+			SeekTimeMS = 0;
+		}
+
+		if (!SkipSpeedTest && !RunSpeedTest(ConsiderSlowAtMS * 2, SeekTimeMS, ReadSpeedMBs, WriteSpeedMBs))
+		{
+			bDisabled = true;
+			UE_LOG(LogDerivedDataCache, Error, TEXT("No read or write access to %s"), *CachePath);
+		}
+		else
+		{
+			bool bReadTestPassed = ReadSpeedMBs > 0.0;
+			bool bWriteTestPassed = WriteSpeedMBs > 0.0;
+
+			// if we failed writes mark this as read only
+			bReadOnly = bReadOnly || !bWriteTestPassed;
+
+			// classify and report on these times
+			if (SeekTimeMS < 1)
+			{
+				SpeedClass = ESpeedClass::Local;
+			}
+			else if (SeekTimeMS <= ConsiderFastAtMS)
+			{
+				SpeedClass = ESpeedClass::Fast;
+			}
+			else if (SeekTimeMS >= ConsiderSlowAtMS)
 			{
-				uint32 ReadErrorCode = FPlatformMisc::GetLastError();
-				TCHAR WriteErrorBuffer[1024];
-				TCHAR ReadErrorBuffer[1024];
-				FPlatformMisc::GetSystemErrorMessage(WriteErrorBuffer, 1024, WriteErrorCode);
-				FPlatformMisc::GetSystemErrorMessage(ReadErrorBuffer, 1024, ReadErrorCode);
-				UE_LOG(LogDerivedDataCache, Warning, TEXT("Fail to write to %s, derived data cache to this directory will be read only. WriteError: %u (%s) ReadError: %u (%s)"), *CachePath, WriteErrorCode, WriteErrorBuffer, ReadErrorCode, ReadErrorBuffer);
+				SpeedClass = ESpeedClass::Slow;
 			}
 			else
 			{
-				bFailed = false;
-			}			
-			if (TestFileSize >= 0)
+				SpeedClass = ESpeedClass::Ok;
+			}
+
+			UE_LOG(LogDerivedDataCache, Display, TEXT("Performance to %s: Latency=%.02fms. RandomReadSpeed=%.02fMBs, RandomWriteSpeed=%.02fMBs. Assigned SpeedClass '%s'"), 
+				*CachePath, SeekTimeMS, ReadSpeedMBs, WriteSpeedMBs, LexToString(SpeedClass));
+
+			if (SpeedClass <= FDerivedDataBackendInterface::ESpeedClass::Slow && !bReadOnly)
 			{
-				IFileManager::Get().Delete(*TempFilename, false, false, true);
+				UE_LOG(LogDerivedDataCache, Warning, TEXT("Access to %s appears to be slow. 'Touch' will be disabled and queries/writes will be limited."), *CachePath);
+				bTouch = false;
+				//bReadOnly = true;
 			}
-			AccessDuration = FPlatformTime::Seconds() - TestStart;
-		}
-		if (bFailed)
-		{
-			double StartTime = FPlatformTime::Seconds();
-			TArray<FString> FilesAndDirectories;
-			IFileManager::Get().FindFiles(FilesAndDirectories,*(CachePath / TEXT("*.*")), true, true);
-			AccessDuration = FPlatformTime::Seconds() - StartTime;
-			if (FilesAndDirectories.Num() > 0)
+
+			if (!bReadOnly)
 			{
-				bReadOnly = true;
-				bFailed = false;
+				if (FString(FCommandLine::Get()).Contains(TEXT("Run=DerivedDataCache")))
+				{
+					bTouch = true; // we always touch files when running the DDC commandlet
+				}
+
+				// The command line (-ddctouch) enables touch on all filesystem backends if specified. 
+				bTouch = bTouch || FParse::Param(FCommandLine::Get(), TEXT("DDCTOUCH"));
+
+				if (bTouch)
+				{
+					UE_LOG(LogDerivedDataCache, Display, TEXT("Files in %s will be touched."), *CachePath);
+				}
+
+				if (bDeleteUnused && !FParse::Param(FCommandLine::Get(), TEXT("NODDCCLEANUP")) && FDDCCleanup::Get())
+				{
+					FDDCCleanup::Get()->AddFilesystem(CachePath, DaysToDeleteUnusedFiles, MaxFoldersToClean, MaxFileChecksPerSec);
+				}
+			}
+			
+			if (IsUsable() && InAccessLogFileName != nullptr && *InAccessLogFileName != 0)
+			{
+				AccessLogWriter.Reset(new FAccessLogWriter(InAccessLogFileName));
 			}
 		}
-		if (FString(FCommandLine::Get()).Contains(TEXT("Run=DerivedDataCache")) )
+	}
+
+	bool RunSpeedTest(double InSkipTestsIfSeeksExceedMS, double& OutSeekTimeMS, double& OutReadSpeedMBs, double& OutWriteSpeedMBs) const
+	{
+		//  files of increasing size. Most DDC data falls within this range so we don't want to skew by reading 
+		// large amounts of data. Ultimately we care most about latency anyway.
+		const int FileSizes[] = { 4, 8, 16, 64, 128, 256 };
+		const int NumTestFolders = 2; //(0-9)
+		const int FileSizeCount = UE_ARRAY_COUNT(FileSizes);
+
+		bool bWriteTestPassed = true;
+		bool bReadTestPassed = true;
+		bool bTestDataExists = true;
+
+		double TotalSeekTime = 0;
+		double TotalReadTime = 0;
+		double TotalWriteTime = 0;
+		int TotalDataRead = 0;
+		int TotalDataWritten = 0;
+
+		const FString AbsoluteCachePath = IFileManager::Get().ConvertToAbsolutePathForExternalAppForRead(*CachePath);
+		if (AbsoluteCachePath.Len() > MAX_CACHE_DIR_LEN)
 		{
-			bTouch = true; // we always touch files when running the DDC commandlet
+			const FText ErrorMessage = FText::Format(NSLOCTEXT("DerivedDataCache", "PathTooLong", "Cache path {0} is longer than {1} characters...please adjust [DerivedDataBackendGraph] paths to be shorter (this leaves more room for cache keys)."), FText::FromString(AbsoluteCachePath), FText::AsNumber(MAX_CACHE_DIR_LEN));
+			FMessageDialog::Open(EAppMsgType::Ok, ErrorMessage);
+			UE_LOG(LogDerivedDataCache, Fatal, TEXT("%s"), *ErrorMessage.ToString());
 		}
-		// The command line (-ddctouch) enables touch on all filesystem backends if specified. 
-		bTouch = bTouch || FParse::Param(FCommandLine::Get(), TEXT("DDCTOUCH"));
 
-		if (bReadOnly)
+		TArray<FString> Paths;
+		TArray<FString> MissingFiles;
+
+		MissingFiles.Reserve(NumTestFolders * FileSizeCount);
+
+		const FString TestDataPath = FPaths::Combine(CachePath, TEXT("TestData"));
+
+		// create an upfront map of paths to data size in bytes
+		// create the paths we'll use. <path>/0/TestData.dat, <path>/1/TestData.dat etc. If those files don't exist we'll
+		// create them which will likely give an invalid result when measuring them now but not in the future...
+		TMap<FString, int> TestFileEntries;
+		for (int iSize = 0; iSize < FileSizeCount; iSize++)
 		{
-			bTouch = false; // we won't touch read only paths
+			// make sure we dont stat/read/write to consecuting files in folders
+			for (int iFolder = 0; iFolder < NumTestFolders; iFolder++)
+			{
+				int FileSizeKB = FileSizes[iSize];
+				FString Path = FPaths::Combine(CachePath, TEXT("TestData"), *FString::FromInt(iFolder), *FString::Printf(TEXT("TestData_%dkb.dat"), FileSizeKB));
+				TestFileEntries.Add(Path, FileSizeKB * 1024);
+			}
 		}
 
-		if (bTouch)
+		// measure latency by checking for the presence of all these files. We'll also track which don't exist..
+		const double StatStartTime = FPlatformTime::Seconds();
+		int TotalStatChecks = 0;
+		FDateTime CurrentTime = FDateTime::Now();
+ 		for (auto& KV : TestFileEntries)
 		{
-			UE_LOG(LogDerivedDataCache, Display, TEXT("Files in %s will be touched."),*CachePath);
+			FFileStatData StatData = IFileManager::Get().GetStatData(*KV.Key);
+
+			if (!StatData.bIsValid || StatData.FileSize != KV.Value)
+			{
+				MissingFiles.Add(KV.Key);
+			}
+			else if (!bReadOnly)
+			{
+				IFileManager::Get().SetTimeStamp(*KV.Key, CurrentTime);
+			}
 		}
 
-		if (!bFailed && AccessDuration > SlowInitDuration && !GIsBuildMachine)
+		// save total stat time
+		TotalSeekTime = (FPlatformTime::Seconds() - StatStartTime);
+
+		// calculate seek time here
+		OutSeekTimeMS = (TotalSeekTime / TestFileEntries.Num()) * 1000;
+
+		UE_LOG(LogDerivedDataCache, Verbose, TEXT("Stat tests to %s took %.02f seconds"), *CachePath, TotalSeekTime);
+
+		// if seek times are very slow do a single read/write test just to confirm access
+		if (OutSeekTimeMS >= InSkipTestsIfSeeksExceedMS)
 		{
-			UE_LOG(LogDerivedDataCache, Warning, TEXT("%s access is very slow (initialization took %.2f seconds), consider disabling it."), *CachePath, AccessDuration);
+			UE_LOG(LogDerivedDataCache, Warning, TEXT("Limiting read/write speed tests due to seek times of %.02f exceeding %.02fms. Values will be inaccurate."), OutSeekTimeMS, InSkipTestsIfSeeksExceedMS);
+
+			FString Path = TestFileEntries.begin()->Key;
+			int Size = TestFileEntries.begin()->Value;
+
+			TestFileEntries.Reset();
+			TestFileEntries.Add(Path, Size);		
 		}
+	
+		// create any files that were missing
+		if (!bReadOnly)
+		{
+			TArray<uint8> Data;
+			for (auto& File : MissingFiles)
+			{
+				const int DesiredSize = TestFileEntries[File];
+				Data.SetNumUninitialized(DesiredSize);
 
-		if (!bFailed && InAccessLogFileName != nullptr && *InAccessLogFileName != 0)
+				if (!FFileHelper::SaveArrayToFile(Data, *File))
+				{
+					uint32 ErrorCode = FPlatformMisc::GetLastError();
+					TCHAR ErrorBuffer[1024];
+					FPlatformMisc::GetSystemErrorMessage(ErrorBuffer, 1024, ErrorCode);
+					UE_LOG(LogDerivedDataCache, Warning, TEXT("Fail to write to %s, derived data cache to this directory will be read only. WriteError: %u (%s)"), *File, ErrorCode, ErrorBuffer);
+					bTestDataExists = false;
+					bWriteTestPassed = false;
+					break;
+				}
+			}
+		}
+
+		// now read all sizes from random folders
 		{
-			AccessLogWriter.Reset(new FAccessLogWriter(InAccessLogFileName));
+			const int ArraySize = UE_ARRAY_COUNT(FileSizes);
+			TArray<uint8> TempData;
+			TempData.Empty(FileSizes[ArraySize - 1] * 1024);
+
+			const double ReadStartTime = FPlatformTime::Seconds();
+
+			for (auto& KV : TestFileEntries)
+			{
+				const int FileSize = KV.Value;
+				const FString& FilePath = KV.Key;
+
+				if (!FFileHelper::LoadFileToArray(TempData, *FilePath))
+				{
+					uint32 ErrorCode = FPlatformMisc::GetLastError();
+					TCHAR ErrorBuffer[1024];
+					FPlatformMisc::GetSystemErrorMessage(ErrorBuffer, 1024, ErrorCode);
+					UE_LOG(LogDerivedDataCache, Warning, TEXT("Fail to read from %s, derived data cache will be disabled. ReadError: %u (%s)"), *FilePath, ErrorCode, ErrorBuffer);
+					bReadTestPassed = false;
+					break;
+				}
+
+				TotalDataRead += TempData.Num();
+			}
+
+			TotalReadTime = FPlatformTime::Seconds() - ReadStartTime;
+
+			UE_LOG(LogDerivedDataCache, Verbose, TEXT("Read tests %s on %s and took %.02f seconds"), bReadTestPassed ? TEXT("passed") : TEXT("failed"), *CachePath, TotalReadTime);
 		}
+		
+		// do write tests if or read tests passed and our seeks were below the cut-off 
+		if (bReadTestPassed && !bReadOnly)
+		{
+			// do write tests but use a machine specific folder
+			FString CustomPath = FPaths::Combine(CachePath, TEXT("TestData"), FPlatformMisc::GetLoginId());
+
+			const int ArraySize = UE_ARRAY_COUNT(FileSizes);
+			TArray<uint8> TempData;
+			TempData.Empty(FileSizes[ArraySize - 1] * 1024);
+
+			const double WriteStartTime = FPlatformTime::Seconds();
+
+			for (auto& KV : TestFileEntries)
+			{
+				const int FileSize = KV.Value;
+				FString FilePath = KV.Key;
+
+				TempData.SetNumUninitialized(FileSize);
+
+				FilePath = FilePath.Replace(*CachePath, *CustomPath);
+
+				if (!FFileHelper::SaveArrayToFile(TempData, *FilePath))
+				{
+					uint32 ErrorCode = FPlatformMisc::GetLastError();
+					TCHAR ErrorBuffer[1024];
+					FPlatformMisc::GetSystemErrorMessage(ErrorBuffer, 1024, ErrorCode);
+					UE_LOG(LogDerivedDataCache, Warning, TEXT("Fail to write to %s, derived data cache will be disabled. ReadError: %u (%s)"), *FilePath, ErrorCode, ErrorBuffer);
+					bWriteTestPassed = false;
+					break;
+				}
+
+				TotalDataWritten += TempData.Num();
+			}
 
-		if (!bReadOnly && !bFailed && bDeleteOldFiles && !FParse::Param(FCommandLine::Get(),TEXT("NODDCCLEANUP")) && FDDCCleanup::Get())
-		{			
-			FDDCCleanup::Get()->AddFilesystem( CachePath, InDaysToDeleteUnusedFiles, InMaxNumFoldersToCheck, InMaxContinuousFileChecks );
+			TotalWriteTime = FPlatformTime::Seconds() - WriteStartTime;
+
+			UE_LOG(LogDerivedDataCache, Verbose, TEXT("write tests %s on %s and took %.02f seconds"), bWriteTestPassed ? TEXT("passed") : TEXT("failed"), *CachePath, TotalReadTime)
+
+			// remove the custom path but do it async as this can be slow on
+			AsyncTask(ENamedThreads::AnyThread, [CustomPath]() {
+				IFileManager::Get().DeleteDirectory(*CustomPath, false, true);
+			});
+	
+			// check latency and speed. Read values should always be valid
+			const double ReadSpeedMBs = (bReadTestPassed ? (TotalDataRead / TotalReadTime) : 0) / (1024 * 1024);
+			const double WriteSpeedMBs = (bWriteTestPassed ? (TotalDataWritten / TotalWriteTime) : 0) / (1024 * 1024);
+			const double SeekTimeMS = (TotalSeekTime / TestFileEntries.Num()) * 1000;
 		}
+
+		const double TotalTestTime = FPlatformTime::Seconds() - StatStartTime;
+
+		UE_LOG(LogDerivedDataCache, Log, TEXT("Speed tests for %s took %.02f seconds"), *CachePath, TotalTestTime);
+
+		// check latency and speed. Read values should always be valid
+		OutReadSpeedMBs = (bReadTestPassed ? (TotalDataRead / TotalReadTime) : 0) / (1024 * 1024);
+		OutWriteSpeedMBs = (bWriteTestPassed ? (TotalDataWritten / TotalWriteTime) : 0) / (1024 * 1024);
+
+		return bWriteTestPassed || bReadTestPassed;
+	}
+
+	/** Return a name for this interface */
+	virtual FString GetName() const override 
+	{ 
+		return CachePath; 
 	}
 
 	/** return true if the cache is usable **/
 	bool IsUsable()
 	{
-		return !bFailed;
+		return !bDisabled;
 	}
 
 	/** return true if this cache is writable **/
 	virtual bool IsWritable() override
 	{
-		return !bReadOnly;
+		return !bReadOnly && !bDisabled;
+	}
+
+	/** Returns a class of speed for this interface **/
+	virtual ESpeedClass GetSpeedClass() override
+	{
+		return SpeedClass;
 	}
 
 	class FAccessLogWriter
@@ -189,7 +463,8 @@ public:
 		FCriticalSection CriticalSection;
 		TSet<FString> CacheKeys;
 	};
-
+	
+	
 	/**
 	 * Synchronous test for the existence of a cache item
 	 *
@@ -199,12 +474,23 @@ public:
 	virtual bool CachedDataProbablyExists(const TCHAR* CacheKey) override
 	{
 		COOK_STAT(auto Timer = UsageStats.TimeProbablyExists());
-		check(!bFailed);
+		check(IsUsable());
+
+		// if we're a slow device just say we have the data. It's faster to try and fail than it
+		// is to check and succeed.
+		if (GetSpeedClass() <= FDerivedDataBackendInterface::ESpeedClass::Slow)
+		{
+			return true;
+		}
+
 		FString Filename = BuildFilename(CacheKey);
-		FDateTime TimeStamp = IFileManager::Get().GetTimeStamp(*Filename);
-		bool bExists = TimeStamp > FDateTime::MinValue();
-		if (bExists)
+
+		FFileStatData FileStat = IFileManager::Get().GetStatData(*Filename);
+
+		if (FileStat.bIsValid)
 		{
+			FDateTime TimeStamp = FileStat.ModificationTime;
+
 			// Update file timestamp to prevent it from being deleted by DDC Cleanup.
 			if (bTouch || 
 				 (!bReadOnly && (FDateTime::UtcNow() - TimeStamp).GetDays() > (DaysToDeleteUnusedFiles / 4)))
@@ -221,7 +507,7 @@ public:
 		}
 
 		// If not using a shared cache, record a (probable) miss
-		if (!bExists && !GetDerivedDataCacheRef().GetUsingSharedDDC())
+		if (!FileStat.bIsValid && !GetDerivedDataCacheRef().GetUsingSharedDDC())
 		{
 			// store a cache miss
 			FScopeLock ScopeLock(&SynchronizationObject);
@@ -231,7 +517,7 @@ public:
 			}
 		}
 
-		return bExists;
+		return FileStat.bIsValid;
 	}
 	/**
 	 * Synchronous retrieve of a cache item
@@ -243,15 +529,30 @@ public:
 	virtual bool GetCachedData(const TCHAR* CacheKey, TArray<uint8>& Data) override
 	{
 		COOK_STAT(auto Timer = UsageStats.TimeGet());
-		check(!bFailed);
+		check(IsUsable());
 		FString Filename = BuildFilename(CacheKey);
 		double StartTime = FPlatformTime::Seconds();
-		if (FFileHelper::LoadFileToArray(Data,*Filename,FILEREAD_Silent))
+
+		if (DebugMissRate > 0)
 		{
-			if(!GIsBuildMachine)
+			bool Miss = FMath::RandHelper(100) < DebugMissRate;
+
+			if (Miss)
 			{
-				double ReadDuration = FPlatformTime::Seconds() - StartTime;
-				double ReadSpeed = ReadDuration > 5.0 ? (Data.Num() / ReadDuration) / (1024.0 * 1024.0) : 100.0;
+				FScopeLock Lock(&MissedKeysCS);
+				UE_LOG(LogDerivedDataCache, Log, TEXT("%s Pretending key didn't exist: %s"), *GetName(), CacheKey);
+				DebugMissedKeys.Add(FName(CacheKey));
+				return false;
+			}
+		}
+
+		if (FFileHelper::LoadFileToArray(Data,*Filename,FILEREAD_Silent))
+		{
+			double ReadDuration = FPlatformTime::Seconds() - StartTime;
+			double ReadSpeed = (Data.Num() / ReadDuration) / (1024.0 * 1024.0);
+
+			if(!GIsBuildMachine && ReadDuration > 5.0)
+			{				
 				// Slower than 0.5MB/s?
 				UE_CLOG(ReadSpeed < 0.5, LogDerivedDataCache, Warning, TEXT("%s is very slow (%.2fMB/s) when accessing %s, consider disabling it."), *CachePath, ReadSpeed, *Filename);
 			}
@@ -261,11 +562,11 @@ public:
 				AccessLogWriter->Append(CacheKey);
 			}
 
-			UE_LOG(LogDerivedDataCache, Verbose, TEXT("FileSystemDerivedDataBackend: Cache hit on %s"),*Filename);
+			UE_LOG(LogDerivedDataCache, Verbose, TEXT("%s: Cache hit on %s (%d bytes, %.02f secs, %.2fMB/s)"), *GetName(), CacheKey, Data.Num(), ReadDuration, ReadSpeed);
 			COOK_STAT(Timer.AddHit(Data.Num()));
 			return true;
 		}
-		UE_LOG(LogDerivedDataCache, Verbose, TEXT("FileSystemDerivedDataBackend: Cache miss on %s"),*Filename);
+		UE_LOG(LogDerivedDataCache, Verbose, TEXT("%s: Cache miss on %s"), *GetName(), CacheKey);
 		Data.Empty();
 
 		// If not using a shared cache, record a miss
@@ -281,6 +582,15 @@ public:
 
 		return false;
 	}
+
+	/**
+	 * Would we cache this? Say yes so long as we aren't read-only.
+	 */
+	bool WouldCache(const TCHAR* CacheKey, TArrayView<const uint8> InData) override
+	{
+		return IsWritable() && !CachedDataProbablyExists(CacheKey);
+	}
+
 	/**
 	 * Asynchronous, fire-and-forget placement of a cache item
 	 *
@@ -291,14 +601,21 @@ public:
 	virtual void PutCachedData(const TCHAR* CacheKey, TArrayView<const uint8> Data, bool bPutEvenIfExists) override
 	{
 		COOK_STAT(auto Timer = UsageStats.TimePut());
-		check(!bFailed);
-		if (!bReadOnly)
+		check(IsUsable());
+
+		if (IsWritable())
 		{
 			if (AccessLogWriter.IsValid())
 			{
 				AccessLogWriter->Append(CacheKey);
 			}
 
+			// don't put anything we pretended didn't exist
+			if (DebugMissRate > 0 && WasDebugMissed(CacheKey))
+			{
+				return;
+			}
+
 			if (bPutEvenIfExists || !CachedDataProbablyExists(CacheKey))
 			{
 				COOK_STAT(Timer.AddHit(Data.Num()));
@@ -399,8 +716,8 @@ public:
 
 	void RemoveCachedData(const TCHAR* CacheKey, bool bTransient) override
 	{
-		check(!bFailed);
-		if (!bReadOnly && (!bTransient || bPurgeTransient))
+		check(IsUsable());
+		if (IsWritable() && (!bTransient || bPurgeTransient))
 		{
 			FString Filename = BuildFilename(CacheKey);
 			if (bTransient)
@@ -416,6 +733,23 @@ public:
 		COOK_STAT(UsageStatsMap.Add(FString::Printf(TEXT("%s: %s.%s"), *GraphPath, TEXT("FileSystem"), *CachePath), UsageStats));
 	}
 
+	bool TryToPrefetch(const TCHAR* CacheKey) override
+	{
+		return false;
+	}
+
+	void SetDebugMissRate(int InMissRate)
+	{
+		UE_LOG(LogDerivedDataCache, Log, TEXT("Setting DebugMissRate to %d for cache %s"), InMissRate, *GetName());
+		DebugMissRate = InMissRate;
+	}
+
+	bool WasDebugMissed(const TCHAR* InKey)
+	{
+		FScopeLock Lock(&MissedKeysCS);
+		return DebugMissedKeys.Contains(FName(InKey));
+	}
+
 private:
 	FDerivedDataCacheUsageStats UsageStats;
 
@@ -432,16 +766,21 @@ private:
 
 	/** Base path we are storing the cache files in. **/
 	FString	CachePath;
+	/** Class of this cache */
+	ESpeedClass SpeedClass;
 	/** If true, do not attempt to write to this cache **/
 	bool		bReadOnly;
-	/** If true, we failed to write to this directory and it did not contain anything so we should not be used **/
-	bool		bFailed;
+	
 	/** If true, CachedDataProbablyExists will update the file timestamps. */
 	bool		bTouch;
 	/** If true, allow transient data to be removed from the cache. */
 	bool		bPurgeTransient;
 	/** Age of file when it should be deleted from DDC cache. */
 	int32		DaysToDeleteUnusedFiles;
+
+	/** If true, we failed to write to this directory and it did not contain anything so we should not be used **/
+	bool		bDisabled;
+
 	/** Object used for synchronization via a scoped lock						*/
 	FCriticalSection SynchronizationObject;
 
@@ -453,17 +792,32 @@ private:
 	/** The total estimated build time accumulated from cache miss/put deltas */
 	double TotalEstimatedBuildTime;
 
+	/** Miss rate that can be set for profiling */
+	int DebugMissRate;
+
 	/** Access log to write to */
 	TUniquePtr<FAccessLogWriter> AccessLogWriter;
+
+	/** Keys we ignored due to miss rate settings */
+	FCriticalSection MissedKeysCS;
+	TSet<FName> DebugMissedKeys;
+
 };
 
-FDerivedDataBackendInterface* CreateFileSystemDerivedDataBackend(const TCHAR* CacheDirectory, bool bForceReadOnly /*= false*/, bool bTouchFiles /*= false*/, bool bPurgeTransient /*= false*/, bool bDeleteOldFiles /*= false*/, int32 InDaysToDeleteUnusedFiles /*= 60*/, int32 InMaxNumFoldersToCheck /*= -1*/, int32 InMaxContinuousFileChecks /*= -1*/, const TCHAR* InAccessLogFileName /*= nullptr*/)
+FDerivedDataBackendInterface* CreateFileSystemDerivedDataBackend(const TCHAR* CacheDirectory, const TCHAR* InParams, const int InMissRate, const TCHAR* InAccessLogFileName /*= nullptr*/)
 {
-	FFileSystemDerivedDataBackend* FileDDB = new FFileSystemDerivedDataBackend(CacheDirectory, bForceReadOnly, bTouchFiles, bPurgeTransient, bDeleteOldFiles, InDaysToDeleteUnusedFiles, InMaxNumFoldersToCheck, InMaxContinuousFileChecks, InAccessLogFileName);
+	FFileSystemDerivedDataBackend* FileDDB = new FFileSystemDerivedDataBackend( CacheDirectory, InParams, InMissRate, InAccessLogFileName);
+
+	if (InMissRate > 0)
+	{
+		FileDDB->SetDebugMissRate(InMissRate);
+	}
+
 	if (!FileDDB->IsUsable())
 	{
 		delete FileDDB;
 		FileDDB = NULL;
 	}
+
 	return FileDDB;
 }
diff --git a/Engine/Source/Developer/DerivedDataCache/Private/HierarchicalDerivedDataBackend.h b/Engine/Source/Developer/DerivedDataCache/Private/HierarchicalDerivedDataBackend.h
--- a/Engine/Source/Developer/DerivedDataCache/Private/HierarchicalDerivedDataBackend.h
+++ b/Engine/Source/Developer/DerivedDataCache/Private/HierarchicalDerivedDataBackend.h
@@ -29,6 +29,18 @@ public:
 		UpdateAsyncInnerBackends();
 	}
 
+	/** Return a name for this interface */
+	virtual FString GetName() const override
+	{
+		return TEXT("HierarchicalDerivedDataBackend");
+	}
+
+	/** Are we a remote cache? */
+	virtual ESpeedClass GetSpeedClass() override
+	{
+		return ESpeedClass::Local;
+	}
+
 	void UpdateAsyncInnerBackends()
 	{
 		bIsWritable = false;
@@ -91,6 +103,68 @@ public:
 		}
 		return false;
 	}
+
+	/**
+	 * Attempts to make sure the cached data will be available as optimally as possible. This is left up to the implementation to do
+	 * @param	CacheKey	Alphanumeric+underscore key of this cache item
+	 * @return				true if any steps were performed to optimize future retrieval
+	 */
+	virtual bool TryToPrefetch(const TCHAR* CacheKey) override
+	{
+		COOK_STAT(auto Timer = UsageStats.TimePrefetch());
+
+		// Search all backends for this key. If it can be moved into a faster class then we'll do so.
+		bool WorthFetching = false;
+
+		FDerivedDataBackendInterface* LastMissedInterface = nullptr;
+
+		for (int32 CacheIndex = 0; CacheIndex < InnerBackends.Num(); CacheIndex++)
+		{
+			FDerivedDataBackendInterface* Interface = InnerBackends[CacheIndex];
+
+			if (!Interface->CachedDataProbablyExists(CacheKey) && Interface->IsWritable())
+			{
+				LastMissedInterface = Interface;
+			}
+			else
+			{
+				// if we have an interface that's writable and faster, lets get it
+				if (LastMissedInterface && LastMissedInterface->GetSpeedClass() > Interface->GetSpeedClass())
+				{
+					WorthFetching = true;
+				}
+			}
+		}
+		
+		// If it's remote then fetch it. We don't care about the data but we 
+		// Need to read a copy from the remote store anyway to fill the caches
+		if (WorthFetching)
+		{
+			TArray<uint8> DontCare;
+			GetCachedData(CacheKey, DontCare);
+			COOK_STAT(Timer.AddHit(0));
+		}			
+
+		// Return true if we did anything
+		return WorthFetching;
+	}
+
+	/*
+		Determine if we would cache this by asking all our inner layers
+	*/
+	virtual bool WouldCache(const TCHAR* CacheKey, TArrayView<const uint8> InData) override
+	{
+		for (int32 CacheIndex = 0; CacheIndex < InnerBackends.Num(); CacheIndex++)
+		{
+			if (InnerBackends[CacheIndex]->WouldCache(CacheKey, InData))
+			{
+				return true;
+			}
+		}
+
+		return false;
+	}
+
 	/**
 	 * Synchronous retrieve of a cache item
 	 *
@@ -103,39 +177,53 @@ public:
 		COOK_STAT(auto Timer = UsageStats.TimeGet());
 		for (int32 CacheIndex = 0; CacheIndex < InnerBackends.Num(); CacheIndex++)
 		{
-			if (InnerBackends[CacheIndex]->CachedDataProbablyExists(CacheKey) && InnerBackends[CacheIndex]->GetCachedData(CacheKey, OutData))
+			FDerivedDataBackendInterface* GetInterface = InnerBackends[CacheIndex];
+
+			if (GetInterface->CachedDataProbablyExists(CacheKey) && GetInterface->GetCachedData(CacheKey, OutData))
 			{
 				if (bIsWritable)
 				{
 					// fill in the higher level caches
 					for (int32 PutCacheIndex = CacheIndex - 1; PutCacheIndex >= 0; PutCacheIndex--)
 					{
-						if (InnerBackends[PutCacheIndex]->IsWritable())
+						FDerivedDataBackendInterface* PutBackend = InnerBackends[PutCacheIndex];
+
+						if (PutBackend->IsWritable() && PutBackend->WouldCache(CacheKey, OutData))
 						{
-							if (InnerBackends[PutCacheIndex]->BackfillLowerCacheLevels() &&
-								InnerBackends[PutCacheIndex]->CachedDataProbablyExists(CacheKey))
+							if (PutBackend->BackfillLowerCacheLevels() &&
+								PutBackend->CachedDataProbablyExists(CacheKey))
 							{
-								InnerBackends[PutCacheIndex]->RemoveCachedData(CacheKey, /*bTransient=*/ false); // it apparently failed, so lets delete what is there
+								PutBackend->RemoveCachedData(CacheKey, /*bTransient=*/ false); // it apparently failed, so lets delete what is there
 								AsyncPutInnerBackends[PutCacheIndex]->PutCachedData(CacheKey, OutData, true); // we force a put here because it must have failed
+								UE_LOG(LogDerivedDataCache, Verbose, TEXT("Forward-filling cache %s with: %s (%d bytes) (force=%d)"), *PutBackend->GetName(), CacheKey, OutData.Num(), true);
 							}
 							else
 							{
-								AsyncPutInnerBackends[PutCacheIndex]->PutCachedData(CacheKey, OutData, false); 
+								UE_LOG(LogDerivedDataCache, Verbose, TEXT("Forward-filling cache %s with: %s (%d bytes) (force=%d)"), *PutBackend->GetName(), CacheKey, OutData.Num(), false);
+								PutBackend->PutCachedData(CacheKey, OutData, false);
 							}
 						}
 					}
+
 					if (InnerBackends[CacheIndex]->BackfillLowerCacheLevels())
 					{
 						// fill in the lower level caches
 						for (int32 PutCacheIndex = CacheIndex + 1; PutCacheIndex < AsyncPutInnerBackends.Num(); PutCacheIndex++)
 						{
-							if (!InnerBackends[PutCacheIndex]->IsWritable() && !InnerBackends[PutCacheIndex]->BackfillLowerCacheLevels() && InnerBackends[PutCacheIndex]->CachedDataProbablyExists(CacheKey))
+							FDerivedDataBackendInterface* PutBackend = InnerBackends[PutCacheIndex];
+
+							if (!PutBackend->IsWritable())
 							{
-								break; //do not write things that are already in the read only pak file
+								if (!PutBackend->BackfillLowerCacheLevels() && PutBackend->CachedDataProbablyExists(CacheKey))
+								{
+									break; //do not write things that are already in the read only pak file
+								}
 							}
-							if (InnerBackends[PutCacheIndex]->IsWritable())
+							else if (PutBackend->GetSpeedClass() >= FDerivedDataBackendInterface::ESpeedClass::Ok && PutBackend->WouldCache(CacheKey, OutData))
 							{
+								
 								AsyncPutInnerBackends[PutCacheIndex]->PutCachedData(CacheKey, OutData, false); // we do not need to force a put here
+								UE_LOG(LogDerivedDataCache, Verbose, TEXT("Back-filling cache %s with: %s (%d bytes) (force=%d)"), *PutBackend->GetName(), CacheKey, OutData.Num(), false);
 							}
 						}
 					}
diff --git a/Engine/Source/Developer/DerivedDataCache/Private/MemoryDerivedDataBackend.cpp b/Engine/Source/Developer/DerivedDataCache/Private/MemoryDerivedDataBackend.cpp
--- a/Engine/Source/Developer/DerivedDataCache/Private/MemoryDerivedDataBackend.cpp
+++ b/Engine/Source/Developer/DerivedDataCache/Private/MemoryDerivedDataBackend.cpp
@@ -3,8 +3,9 @@
 #include "MemoryDerivedDataBackend.h"
 #include "Templates/UniquePtr.h"
 
-FMemoryDerivedDataBackend::FMemoryDerivedDataBackend(int64 InMaxCacheSize)
-	: MaxCacheSize(InMaxCacheSize)
+FMemoryDerivedDataBackend::FMemoryDerivedDataBackend(const TCHAR* InName, int64 InMaxCacheSize)
+	: Name(InName)
+	, MaxCacheSize(InMaxCacheSize)
 	, bDisabled( false )
 	, CurrentCacheSize( SerializationSpecificDataSize )
 	, bMaxSizeExceeded(false)
@@ -22,6 +23,11 @@ bool FMemoryDerivedDataBackend::IsWritable()
 	return !bDisabled;
 }
 
+FDerivedDataBackendInterface::ESpeedClass FMemoryDerivedDataBackend::GetSpeedClass()
+{
+	return ESpeedClass::Local;
+}
+
 bool FMemoryDerivedDataBackend::CachedDataProbablyExists(const TCHAR* CacheKey)
 {
 	COOK_STAT(auto Timer = UsageStats.TimeProbablyExists());
@@ -59,13 +65,30 @@ bool FMemoryDerivedDataBackend::GetCachedData(const TCHAR* CacheKey, TArray<uint
 	return false;
 }
 
+bool FMemoryDerivedDataBackend::TryToPrefetch(const TCHAR* CacheKey)
+{
+	return false;
+}
+
+bool FMemoryDerivedDataBackend::WouldCache(const TCHAR* CacheKey, TArrayView<const uint8> InData)
+{
+	if (bDisabled || bMaxSizeExceeded)
+	{
+		return false;
+	}
+
+	return true;
+}
+
 void FMemoryDerivedDataBackend::PutCachedData(const TCHAR* CacheKey, TArrayView<const uint8> InData, bool bPutEvenIfExists)
 {
 	COOK_STAT(auto Timer = UsageStats.TimePut());
 	FScopeLock ScopeLock(&SynchronizationObject);
 	
-	if (bDisabled || bMaxSizeExceeded)
+	// Should never hit this as higher level code should be checking..
+	if (!WouldCache(CacheKey, InData))
 	{
+		UE_LOG(LogDerivedDataCache, Warning, TEXT("WouldCache was not called prior to attempted Put!"));
 		return;
 	}
 	
diff --git a/Engine/Source/Developer/DerivedDataCache/Private/MemoryDerivedDataBackend.h b/Engine/Source/Developer/DerivedDataCache/Private/MemoryDerivedDataBackend.h
--- a/Engine/Source/Developer/DerivedDataCache/Private/MemoryDerivedDataBackend.h
+++ b/Engine/Source/Developer/DerivedDataCache/Private/MemoryDerivedDataBackend.h
@@ -4,7 +4,7 @@
 
 #include "CoreMinimal.h"
 #include "HAL/FileManager.h"
-#include "DerivedDataBackendInterface.h"
+#include "FileBackedDerivedDataBackend.h"
 #include "ProfilingDebugging/CookStats.h"
 #include "DerivedDataCacheUsageStats.h"
 #include "Misc/ScopeLock.h"
@@ -14,15 +14,21 @@ class Error;
 /** 
  * A simple thread safe, memory based backend. This is used for Async puts and the boot cache.
 **/
-class FMemoryDerivedDataBackend : public FDerivedDataBackendInterface
+class FMemoryDerivedDataBackend : public FFileBackedDerivedDataBackend
 {
 public:
-	explicit FMemoryDerivedDataBackend(int64 InMaxCacheSize = -1);
+	explicit FMemoryDerivedDataBackend(const TCHAR* InName, int64 InMaxCacheSize = -1);
 	~FMemoryDerivedDataBackend();
 
+	/** Return a name for this interface */
+	virtual FString GetName() const override { return Name; }
+
 	/** return true if this cache is writable **/
 	virtual bool IsWritable() override;
 
+	/** Returns a class of speed for this interface **/
+	virtual ESpeedClass GetSpeedClass() override;
+
 	/**
 	 * Synchronous test for the existence of a cache item
 	 *
@@ -68,10 +74,18 @@ public:
 	/**
 	 * Disable cache and ignore all subsequent requests
 	 */
-	void Disable();
+	void Disable() override;
 
 	virtual void GatherUsageStats(TMap<FString, FDerivedDataCacheUsageStats>& UsageStatsMap, FString&& GraphPath) override;
 
+	virtual bool TryToPrefetch(const TCHAR* CacheKey) override;
+
+	/**
+	 *  Determines if we would cache the provided data
+	 */
+	virtual bool WouldCache(const TCHAR* CacheKey, TArrayView<const uint8> InData) override;
+	
+
 private:
 	/** Name of the cache file loaded (if any). */
 	FString CacheFilename;
@@ -93,6 +107,9 @@ private:
 		return (Key.Len() + 1) * sizeof(TCHAR) + sizeof(Val.Age) + Val.Data.Num();
 	}
 
+	/** Name of this cache (used for debugging) */
+	FString Name;
+
 	/** Set of files that are being written to disk asynchronously. */
 	TMap<FString, FCacheValue*> CacheItems;
 	/** Maximum size the cached items can grow up to ( in bytes ) */
diff --git a/Engine/Source/Developer/DerivedDataCache/Private/PakFileDerivedDataBackend.cpp b/Engine/Source/Developer/DerivedDataCache/Private/PakFileDerivedDataBackend.cpp
--- a/Engine/Source/Developer/DerivedDataCache/Private/PakFileDerivedDataBackend.cpp
+++ b/Engine/Source/Developer/DerivedDataCache/Private/PakFileDerivedDataBackend.cpp
@@ -68,6 +68,12 @@ bool FPakFileDerivedDataBackend::IsWritable()
 	return bWriting && !bClosed;
 }
 
+/** Returns a class of speed for this interface **/
+FDerivedDataBackendInterface::ESpeedClass FPakFileDerivedDataBackend::GetSpeedClass() 
+{
+	return ESpeedClass::Local;
+}
+
 bool FPakFileDerivedDataBackend::BackfillLowerCacheLevels()
 {
 	return false;
diff --git a/Engine/Source/Developer/DerivedDataCache/Private/PakFileDerivedDataBackend.h b/Engine/Source/Developer/DerivedDataCache/Private/PakFileDerivedDataBackend.h
--- a/Engine/Source/Developer/DerivedDataCache/Private/PakFileDerivedDataBackend.h
+++ b/Engine/Source/Developer/DerivedDataCache/Private/PakFileDerivedDataBackend.h
@@ -27,9 +27,15 @@ public:
 
 	void Close();
 
+	/** Return a name for this interface */
+	virtual FString GetName() const override { return Filename; }
+
 	/** return true if this cache is writable **/
 	virtual bool IsWritable() override;
 
+	/** Returns a class of speed for this interface **/
+	virtual ESpeedClass GetSpeedClass() override;
+
 	virtual bool BackfillLowerCacheLevels() override;
 
 	/**
@@ -86,6 +92,10 @@ public:
 
 	virtual void GatherUsageStats(TMap<FString, FDerivedDataCacheUsageStats>& UsageStatsMap, FString&& GraphPath) override;
 
+	virtual bool TryToPrefetch(const TCHAR* CacheKey) override { return false; }
+
+	virtual bool WouldCache(const TCHAR* CacheKey, TArrayView<const uint8> InData) override { return true; }
+
 private:
 	FDerivedDataCacheUsageStats UsageStats;
 
@@ -129,6 +139,12 @@ public:
 	virtual void PutCachedData(const TCHAR* CacheKey, TArrayView<const uint8> InData, bool bPutEvenIfExists) override;
 	virtual bool GetCachedData(const TCHAR* CacheKey, TArray<uint8>& OutData) override;
 
+	/** Returns a class of speed for this interface **/
+	virtual ESpeedClass GetSpeedClass() override
+	{
+		return ESpeedClass::Fast;
+	}
+
 private:
 	static const EName CompressionFormat = NAME_Zlib;
 	static const ECompressionFlags CompressionFlags = COMPRESS_BiasMemory;
diff --git a/Engine/Source/Developer/DerivedDataCache/Private/S3DerivedDataBackend.cpp b/Engine/Source/Developer/DerivedDataCache/Private/S3DerivedDataBackend.cpp
--- a/Engine/Source/Developer/DerivedDataCache/Private/S3DerivedDataBackend.cpp
+++ b/Engine/Source/Developer/DerivedDataCache/Private/S3DerivedDataBackend.cpp
@@ -835,6 +835,26 @@ void FS3DerivedDataBackend::GatherUsageStats(TMap<FString, FDerivedDataCacheUsag
 	COOK_STAT(UsageStatsMap.Add(FString::Printf(TEXT("%s: %s @ %s"), *GraphPath, TEXT("S3"), *BaseUrl), UsageStats));
 }
 
+FString FS3DerivedDataBackend::GetName() const
+{
+	return BaseUrl;
+}
+
+FDerivedDataBackendInterface::ESpeedClass FS3DerivedDataBackend::GetSpeedClass()
+{
+	return ESpeedClass::Local;
+}
+
+bool FS3DerivedDataBackend::TryToPrefetch(const TCHAR* CacheKey)
+{
+	return false;
+}
+
+bool FS3DerivedDataBackend::WouldCache(const TCHAR* CacheKey, TArrayView<const uint8> InData)
+{
+	return false;
+}
+
 bool FS3DerivedDataBackend::DownloadManifest(FFeedbackContext* Context)
 {
 	// Read the root manifest from disk
diff --git a/Engine/Source/Developer/DerivedDataCache/Private/S3DerivedDataBackend.h b/Engine/Source/Developer/DerivedDataCache/Private/S3DerivedDataBackend.h
--- a/Engine/Source/Developer/DerivedDataCache/Private/S3DerivedDataBackend.h
+++ b/Engine/Source/Developer/DerivedDataCache/Private/S3DerivedDataBackend.h
@@ -45,7 +45,12 @@ public:
 	void PutCachedData(const TCHAR* CacheKey, TArrayView<const uint8> InData, bool bPutEvenIfExists) override;
 	void RemoveCachedData(const TCHAR* CacheKey, bool bTransient) override;
 	void GatherUsageStats(TMap<FString, FDerivedDataCacheUsageStats>& UsageStatsMap, FString&& GraphPath) override;
-	
+
+	FString GetName() const override;
+	ESpeedClass GetSpeedClass() override;
+	bool TryToPrefetch(const TCHAR* CacheKey) override;
+	bool WouldCache(const TCHAR* CacheKey, TArrayView<const uint8> InData) override;
+
 private:
 	struct FBundle;
 	struct FBundleEntry;
diff --git a/Engine/Source/Developer/DerivedDataCache/Public/DerivedDataBackendInterface.h b/Engine/Source/Developer/DerivedDataCache/Public/DerivedDataBackendInterface.h
--- a/Engine/Source/Developer/DerivedDataCache/Public/DerivedDataBackendInterface.h
+++ b/Engine/Source/Developer/DerivedDataCache/Public/DerivedDataBackendInterface.h
@@ -27,10 +27,26 @@ DECLARE_FLOAT_ACCUMULATOR_STAT_EXTERN(TEXT("Exists Time"),STAT_DDC_ExistTime,STA
 class FDerivedDataBackendInterface
 {
 public:
+
+	/*
+		Speed classes. Higher values are faster so > / < comparisons make sense.
+	*/
+	enum class ESpeedClass
+	{
+		Unknown,		/* Don't know yet*/
+		Slow,			/* Slow, likely a remote drive. Some benefit but handle with care */
+		Ok,				/* Ok but not great.  */
+		Fast,			/* Fast but seek times still have an impact */
+		Local			/* Little to no impact from seek times and extremly fast reads */
+	};
+
 	virtual ~FDerivedDataBackendInterface()
 	{
 	}
 
+	/** Return a name for this interface */
+	virtual FString GetName() const = 0;
+
 	/** return true if this cache is writable **/
 	virtual bool IsWritable()=0;
 
@@ -43,6 +59,11 @@ public:
 		return true;
 	}
 
+	/**
+	 * Returns a class of speed for this interface
+	 **/
+	virtual ESpeedClass GetSpeedClass() = 0;
+
 	/**
 	 * Synchronous test for the existence of a cache item
 	 *
@@ -82,6 +103,20 @@ public:
 	 *								This will create a path such as "0. 1. 0. 2", which can uniquely identify this node.
 	 */
 	virtual void GatherUsageStats(TMap<FString, FDerivedDataCacheUsageStats>& UsageStatsMap, FString&& GraphPath) = 0;
+
+	/**
+	 * Synchronous attempt to make sure the cached data will be available as optimally as possible. This is left up to the implementation 
+	 * to implement or just skip.
+	 * @param	CacheKey	Alphanumeric+underscore key of this cache item
+	 * @return				true if any steps were performed to optimize future retrieval
+	 */
+	virtual bool TryToPrefetch(const TCHAR* CacheKey) = 0;
+
+	/**
+	 * Allows the DDC backend to determine if it wants to cache the provided data. Reasons for returning false could be a slow connection,
+	 * a file size limit, etc.
+	 */
+	virtual bool WouldCache(const TCHAR* CacheKey, TArrayView<const uint8> InData) = 0;
 };
 
 class FDerivedDataBackend
@@ -127,3 +162,23 @@ public:
 	virtual void GatherUsageStats(TMap<FString, FDerivedDataCacheUsageStats>& UsageStats) = 0;
 
 };
+
+
+inline const TCHAR* LexToString(FDerivedDataBackendInterface::ESpeedClass SpeedClass)
+{
+	switch (SpeedClass)
+	{
+	case FDerivedDataBackendInterface::ESpeedClass::Unknown:
+		return TEXT("Unknown");
+	case FDerivedDataBackendInterface::ESpeedClass::Slow:
+		return TEXT("Slow");
+	case FDerivedDataBackendInterface::ESpeedClass::Ok:
+		return TEXT("Ok");
+	case FDerivedDataBackendInterface::ESpeedClass::Fast:
+		return TEXT("Fast");
+	case FDerivedDataBackendInterface::ESpeedClass::Local:
+		return TEXT("Local");
+	}
+
+	return TEXT("Unknow value! (Update LexToString!)");
+}
\ No newline at end of file
diff --git a/Engine/Source/Developer/DerivedDataCache/Public/DerivedDataCacheUsageStats.h b/Engine/Source/Developer/DerivedDataCache/Public/DerivedDataCacheUsageStats.h
--- a/Engine/Source/Developer/DerivedDataCache/Public/DerivedDataCacheUsageStats.h
+++ b/Engine/Source/Developer/DerivedDataCache/Public/DerivedDataCacheUsageStats.h
@@ -72,16 +72,24 @@ public:
 		return FCookStats::FScopedStatsCounter(PutStats);
 	}
 
+	/** Call this at the top of the PutCachedData override. auto Timer = TimePut(); Use AddHit on the returned type to track a cache hit. */
+	FCookStats::FScopedStatsCounter TimePrefetch()
+	{
+		return FCookStats::FScopedStatsCounter(PrefetchStats);
+	}
+
 	void LogStats(FCookStatsManager::AddStatFuncRef AddStat, const FString& StatName, const FString& NodeName) const
 	{
 		GetStats.LogStats(AddStat, StatName, NodeName, TEXT("Get"));
 		PutStats.LogStats(AddStat, StatName, NodeName, TEXT("Put"));
 		ExistsStats.LogStats(AddStat, StatName, NodeName, TEXT("Exists"));
+		PrefetchStats.LogStats(AddStat, StatName, NodeName, TEXT("Prefetch"));
 	}
 
 	// expose these publicly for low level access. These should really never be accessed directly except when finished accumulating them.
 	FCookStats::CallStats GetStats;
 	FCookStats::CallStats PutStats;
 	FCookStats::CallStats ExistsStats;
+	FCookStats::CallStats PrefetchStats;
 #endif
 };
