Enable S3DDC by default if available, with a check for canary file not reachable from internal IPs. Also add an editor setting (in the Global panel) for disabling it.

Also suppress some warning messages when a download is cancelled.

diff --git a/Engine/Source/Developer/DerivedDataCache/Private/S3DerivedDataBackend.cpp b/Engine/Source/Developer/DerivedDataCache/Private/S3DerivedDataBackend.cpp
--- a/Engine/Source/Developer/DerivedDataCache/Private/S3DerivedDataBackend.cpp
+++ b/Engine/Source/Developer/DerivedDataCache/Private/S3DerivedDataBackend.cpp
@@ -33,6 +33,7 @@
 #include "HAL/PlatformFilemanager.h"
 #include "HAL/Runnable.h"
 #include "DesktopPlatformModule.h"
+#include "Misc/ConfigCacheIni.h"
 
 #define S3DDC_BACKEND_WAIT_INTERVAL 0.01f
 #define S3DDC_HTTP_REQUEST_TIMEOUT_SECONDS 30L
@@ -282,7 +283,10 @@ public:
 		}
 		if (CurlResult != CURLE_OK)
 		{
-			Log->Logf(ELogVerbosity::Error, TEXT("Error while connecting to %s: %s"), ANSI_TO_TCHAR(Url), ANSI_TO_TCHAR(curl_easy_strerror(CurlResult)));
+			if (CurlResult != CURLE_ABORTED_BY_CALLBACK)
+			{
+				Log->Logf(ELogVerbosity::Error, TEXT("Error while connecting to %s: %d (%s)"), ANSI_TO_TCHAR(Url), CurlResult, ANSI_TO_TCHAR(curl_easy_strerror(CurlResult)));
+			}
 			return 500;
 		}
 
@@ -651,7 +655,10 @@ struct FS3DerivedDataBackend::FBundleDownload : IRequestCallback
 		long ResponseCode = RequestPool.Download(*BundleUrl, this, CompressedData, Context);
 		if(!IsSuccessfulHttpResponse(ResponseCode))
 		{
-			Context->Logf(ELogVerbosity::Warning, TEXT("Unable to download bundle %s (%d)"), *BundleUrl, ResponseCode);
+			if (!Context->ReceivedUserCancel())
+			{
+				Context->Logf(ELogVerbosity::Warning, TEXT("Unable to download bundle %s (%d)"), *BundleUrl, ResponseCode);
+			}
 			return;
 		}
 
@@ -706,6 +713,14 @@ FS3DerivedDataBackend::FS3DerivedDataBackend(const TCHAR* InRootManifestPath, co
 		bCanaryValid = IsSuccessfulHttpResponse(RequestPool->Download(*(BaseUrl / CanaryObjectKey), nullptr, Data, GLog));
 	}
 
+	// Allow the user to override it from the editor
+	bool bSetting;
+	if (GConfig->GetBool(TEXT("/Script/UnrealEd.EditorSettings"), TEXT("bEnableS3DDC"), bSetting, GEditorSettingsIni) && !bSetting)
+	{
+		UE_LOG(LogDerivedDataCache, Display, TEXT("S3DerivedDataBackend: Disabling due to config setting"));
+		bCanaryValid = false;
+	}
+
 	// Try to read the bundles
 	if (bCanaryValid)
 	{
diff --git a/Engine/Source/Editor/UnrealEd/Public/Settings/EditorSettings.h b/Engine/Source/Editor/UnrealEd/Public/Settings/EditorSettings.h
--- a/Engine/Source/Editor/UnrealEd/Public/Settings/EditorSettings.h
+++ b/Engine/Source/Editor/UnrealEd/Public/Settings/EditorSettings.h
@@ -31,6 +31,14 @@ class UNREALED_API UEditorSettings : public UObject
 	UPROPERTY(EditAnywhere, config, Category=DerivedData, meta = (ConfigRestartRequired = true))
 	FDirectoryPath SharedDerivedDataCache;
 
+	/** Whether to show the S3DDC settings. Needs to be enabled on a project-by-project basis. */
+	UPROPERTY(config)
+	bool bShowEnableS3DDC;
+
+	/** Whether to enable the S3 derived data cache backend */
+	UPROPERTY(EditAnywhere, config, Category=DerivedData, meta = (DisplayName = "Enable AWS S3 Cache", Tooltip = "Enables the experimental S3 derived data cache.", ConfigRestartRequired = true, EditCondition="bShowEnableS3DDC"))
+	bool bEnableS3DDC = true;
+
 	// =====================================================================
 	// The following options are NOT exposed in the preferences Editor
 	// (usually because there is a different way to set them interactively!)
