Suppressing errors about S3DDC if disabled.

diff --git a/Engine/Source/Developer/DerivedDataCache/Private/DerivedDataBackends.cpp b/Engine/Source/Developer/DerivedDataCache/Private/DerivedDataBackends.cpp
--- a/Engine/Source/Developer/DerivedDataCache/Private/DerivedDataBackends.cpp
+++ b/Engine/Source/Developer/DerivedDataCache/Private/DerivedDataBackends.cpp
@@ -656,16 +656,8 @@ public:
 			return nullptr;
 		}
 
-		FS3DerivedDataBackend* Backend = new FS3DerivedDataBackend(*ManifestPath, *BaseUrl, *Region, *CanaryObjectKey, *AccessKey, *SecretKey);
-		if (!Backend->IsUsable())
-		{
-			UE_LOG(LogDerivedDataCache, Warning, TEXT("%s could not contact the service (%s), will not use it."), NodeName, *BaseUrl);
-			delete Backend;
-			return nullptr;
-		}
-
 		// Insert the backend corruption wrapper. Since the filesystem already uses this, and we're recycling the data with the trailer intact, we need to use it for the S3 cache too.
-		UE_LOG(LogDerivedDataCache, Display, TEXT("Using %s S3 backend at %s"), *Region, *BaseUrl);
+		FS3DerivedDataBackend* Backend = new FS3DerivedDataBackend(*ManifestPath, *BaseUrl, *Region, *CanaryObjectKey, *AccessKey, *SecretKey);
 		return new FDerivedDataBackendCorruptionWrapper(Backend);
 #else
 		UE_LOG(LogDerivedDataCache, Warning, TEXT("S3 backend is not supported in the current build configuration."));
diff --git a/Engine/Source/Developer/DerivedDataCache/Private/S3DerivedDataBackend.cpp b/Engine/Source/Developer/DerivedDataCache/Private/S3DerivedDataBackend.cpp
--- a/Engine/Source/Developer/DerivedDataCache/Private/S3DerivedDataBackend.cpp
+++ b/Engine/Source/Developer/DerivedDataCache/Private/S3DerivedDataBackend.cpp
@@ -707,23 +707,36 @@ FS3DerivedDataBackend::FS3DerivedDataBackend(const TCHAR* InRootManifestPath, co
 {
 	// Test whether we can reach the canary URL
 	bool bCanaryValid = true;
-	if (CanaryObjectKey.Len() > 0)
+	if (GIsBuildMachine)
+	{
+		UE_LOG(LogDerivedDataCache, Log, TEXT("S3DerivedDataBackend: Disabling on build machine"));
+		bCanaryValid = false;
+	}
+	else if (CanaryObjectKey.Len() > 0)
 	{
 		TArray<uint8> Data;
-		bCanaryValid = IsSuccessfulHttpResponse(RequestPool->Download(*(BaseUrl / CanaryObjectKey), nullptr, Data, GLog));
+
+		FStringOutputDevice DummyOutputDevice;
+		if (!IsSuccessfulHttpResponse(RequestPool->Download(*(BaseUrl / CanaryObjectKey), nullptr, Data, &DummyOutputDevice)))
+		{
+			UE_LOG(LogDerivedDataCache, Log, TEXT("S3DerivedDataBackend: Unable to download canary file. Disabling."));
+			bCanaryValid = false;
+		}
 	}
 
 	// Allow the user to override it from the editor
 	bool bSetting;
 	if (GConfig->GetBool(TEXT("/Script/UnrealEd.EditorSettings"), TEXT("bEnableS3DDC"), bSetting, GEditorSettingsIni) && !bSetting)
 	{
-		UE_LOG(LogDerivedDataCache, Display, TEXT("S3DerivedDataBackend: Disabling due to config setting"));
+		UE_LOG(LogDerivedDataCache, Log, TEXT("S3DerivedDataBackend: Disabling due to config setting"));
 		bCanaryValid = false;
 	}
 
 	// Try to read the bundles
 	if (bCanaryValid)
 	{
+		UE_LOG(LogDerivedDataCache, Log, TEXT("Using %s S3 backend at %s"), *Region, *BaseUrl);
+
 		FFeedbackContext* Context = FDesktopPlatformModule::Get()->GetNativeFeedbackContext();
 		Context->BeginSlowTask(NSLOCTEXT("S3DerivedDataBackend", "DownloadingDDCBundles", "Downloading DDC bundles..."), true, true);
 
