Moved some DDC debug options into a common class that can be provided to a backend to simulate certain behavior.

E.g.

'-ddc-local-missrate=50' simulates a 50% miss rate of keys on the local backend,
'-ddc-shared-misstypes=AnimSeq' simulates misses for any AnimSeq key. Use + to combine keys. E.g. AnimSeq+Audio

Memory DDC, FileSystemDDC and S3DDC support these options as do the classes that wrap other backends (e.g. AsyncPut).

#tests lots of warm and cold DDC

diff --git a/Engine/Source/Developer/DerivedDataCache/Private/DerivedDataBackendAsyncPutWrapper.cpp b/Engine/Source/Developer/DerivedDataCache/Private/DerivedDataBackendAsyncPutWrapper.cpp
--- a/Engine/Source/Developer/DerivedDataCache/Private/DerivedDataBackendAsyncPutWrapper.cpp
+++ b/Engine/Source/Developer/DerivedDataCache/Private/DerivedDataBackendAsyncPutWrapper.cpp
@@ -149,6 +149,11 @@ bool FDerivedDataBackendAsyncPutWrapper::WouldCache(const TCHAR* CacheKey, TArra
 	return InnerBackend->WouldCache(CacheKey, InData);
 }
 
+bool FDerivedDataBackendAsyncPutWrapper::ApplyDebugOptions(FBackendDebugOptions& InOptions)
+{
+	return InnerBackend->ApplyDebugOptions(InOptions);
+}
+
 bool FDerivedDataBackendAsyncPutWrapper::GetCachedData(const TCHAR* CacheKey, TArray<uint8>& OutData)
 {
 	COOK_STAT(auto Timer = UsageStats.TimeGet());
diff --git a/Engine/Source/Developer/DerivedDataCache/Private/DerivedDataBackendAsyncPutWrapper.h b/Engine/Source/Developer/DerivedDataCache/Private/DerivedDataBackendAsyncPutWrapper.h
--- a/Engine/Source/Developer/DerivedDataCache/Private/DerivedDataBackendAsyncPutWrapper.h
+++ b/Engine/Source/Developer/DerivedDataCache/Private/DerivedDataBackendAsyncPutWrapper.h
@@ -118,6 +118,8 @@ public:
 
 	virtual void GatherUsageStats(TMap<FString, FDerivedDataCacheUsageStats>& UsageStatsMap, FString&& GraphPath) override;
 
+	virtual bool ApplyDebugOptions(FBackendDebugOptions& InOptions) override;
+
 private:
 	FDerivedDataCacheUsageStats UsageStats;
 	FDerivedDataCacheUsageStats PutSyncUsageStats;
diff --git a/Engine/Source/Developer/DerivedDataCache/Private/DerivedDataBackendCorruptionWrapper.h b/Engine/Source/Developer/DerivedDataCache/Private/DerivedDataBackendCorruptionWrapper.h
--- a/Engine/Source/Developer/DerivedDataCache/Private/DerivedDataBackendCorruptionWrapper.h
+++ b/Engine/Source/Developer/DerivedDataCache/Private/DerivedDataBackendCorruptionWrapper.h
@@ -222,6 +222,11 @@ public:
 		return InnerBackend->WouldCache(CacheKey, InData);
 	}
 
+	virtual bool ApplyDebugOptions(FBackendDebugOptions& InOptions) 
+	{ 
+		return InnerBackend->ApplyDebugOptions(InOptions);
+	}
+
 private:
 	FDerivedDataCacheUsageStats UsageStats;
 
diff --git a/Engine/Source/Developer/DerivedDataCache/Private/DerivedDataBackendVerifyWrapper.h b/Engine/Source/Developer/DerivedDataCache/Private/DerivedDataBackendVerifyWrapper.h
--- a/Engine/Source/Developer/DerivedDataCache/Private/DerivedDataBackendVerifyWrapper.h
+++ b/Engine/Source/Developer/DerivedDataCache/Private/DerivedDataBackendVerifyWrapper.h
@@ -157,6 +157,11 @@ public:
 		return InnerBackend->WouldCache(CacheKey, InData);
 	}
 
+	bool ApplyDebugOptions(FBackendDebugOptions& InOptions) override
+	{
+		return InnerBackend->ApplyDebugOptions(InOptions);
+	}
+
 private:
 	FDerivedDataCacheUsageStats UsageStats;
 
diff --git a/Engine/Source/Developer/DerivedDataCache/Private/DerivedDataBackends.cpp b/Engine/Source/Developer/DerivedDataCache/Private/DerivedDataBackends.cpp
--- a/Engine/Source/Developer/DerivedDataCache/Private/DerivedDataBackends.cpp
+++ b/Engine/Source/Developer/DerivedDataCache/Private/DerivedDataBackends.cpp
@@ -30,7 +30,7 @@ DEFINE_LOG_CATEGORY(LogDerivedDataCache);
 #define MAX_BACKEND_KEY_LENGTH (120)
 #define LOCTEXT_NAMESPACE "DerivedDataBackendGraph"
 
-FDerivedDataBackendInterface* CreateFileSystemDerivedDataBackend(const TCHAR* CacheDirectory, const TCHAR* InParams, const int InMissRate, const TCHAR* InAccessLogFileName = nullptr);
+FDerivedDataBackendInterface* CreateFileSystemDerivedDataBackend(const TCHAR* CacheDirectory, const TCHAR* InParams, const TCHAR* InAccessLogFileName = nullptr);
 
 /**
   * This class is used to create a singleton that represents the derived data cache hierarchy and all of the wrappers necessary
@@ -222,6 +222,17 @@ public:
 			InParsedNodes.Add( NodeName, ParsedNode );
 			// Keep references to all created nodes.
 			CreatedBackends.AddUnique( ParsedNode );
+
+			// parse any debug options for this backend. E.g. -ddc-<name>-missrate
+			FDerivedDataBackendInterface::FBackendDebugOptions DebugOptions;
+
+			if (FDerivedDataBackendInterface::FBackendDebugOptions::ParseFromTokens(DebugOptions, NodeName, FCommandLine::Get()))
+			{
+				if (!ParsedNode->ApplyDebugOptions(DebugOptions))
+				{
+					UE_LOG(LogDerivedDataCache, Warning, TEXT("Node %s is ignoring one or mode -ddc-<nodename>-opt debug options"), NodeName);
+				}
+			}
 		}
 
 		return ParsedNode;
@@ -570,16 +581,11 @@ public:
 				bool bShared = FCString::Stricmp(NodeName, TEXT("Shared")) == 0;
 				
 				FString WriteAccessLog;
-				FParse::Value( Entry, TEXT("WriteAccessLog="), WriteAccessLog );
-
-				// look for -ddclocalmissrate, -ddcsharedmissrate etc
-				FString MissRateCmd = FString::Printf(TEXT("DDC%sMissRate="), NodeName);
-				int MissRate = 0;
-				FParse::Value(FCommandLine::Get(), *MissRateCmd, MissRate);
+				FParse::Value( Entry, TEXT("WriteAccessLog="), WriteAccessLog );				
 
 				if (!bShared || IFileManager::Get().DirectoryExists(*Path))
 				{
-					InnerFileSystem = CreateFileSystemDerivedDataBackend( *Path, Entry, MissRate, *WriteAccessLog);
+					InnerFileSystem = CreateFileSystemDerivedDataBackend( *Path, Entry, *WriteAccessLog);
 				}
 
 				if (InnerFileSystem)
@@ -991,4 +997,120 @@ FDerivedDataBackend& FDerivedDataBackend::Get()
 	return FDerivedDataBackendGraph::Get();
 }
 
+
+/**
+ * Parse debug options for the provided node name. Returns true if any options were specified
+ */
+bool FDerivedDataBackendInterface::FBackendDebugOptions::ParseFromTokens(FDerivedDataBackendInterface::FBackendDebugOptions& OutOptions, const TCHAR* InNodeName, const TCHAR* InInputTokens)
+{
+	// check if the input stream has any ddc options for this node
+	FString PrefixKey = FString(TEXT("-ddc-")) + InNodeName;
+
+	if (FCString::Stristr(InInputTokens, *PrefixKey) == nullptr)
+	{
+		// check if it has any -ddc-all- args
+		PrefixKey = FString(TEXT("-ddc-all"));
+
+		if (FCString::Stristr(InInputTokens, *PrefixKey) == nullptr)
+		{
+			return false;
+		}
+	}
+
+	// turn -arg= into arg= for parsing
+	PrefixKey.RightChopInline(1);
+
+	/** types that can be set to ignored (-ddc-<name>-misstypes="foo+bar" etc) */
+	// look for -ddc-local-misstype=AnimSeq+Audio -ddc-shared-misstype=AnimSeq+Audio 
+	FString ArgName = FString::Printf(TEXT("%s-misstypes="), *PrefixKey);
+
+	FString TempArg;
+	FParse::Value(InInputTokens, *ArgName, TempArg);
+	TempArg.ParseIntoArray(OutOptions.SimulateMissTypes, TEXT("+"), true);
+
+	// look for -ddc-local-missrate=, -ddc-shared-missrate= etc
+	ArgName = FString::Printf(TEXT("%s-missrate="), *PrefixKey);
+	int MissRate = 0;
+	FParse::Value(InInputTokens, *ArgName, OutOptions.RandomMissRate);
+
+	// look for -ddc-local-speed=, -ddc-shared-speed= etc
+	ArgName = FString::Printf(TEXT("%s-speed="), *PrefixKey);
+	if (FParse::Value(InInputTokens, *ArgName, TempArg))
+	{
+		if (!TempArg.IsEmpty())
+		{
+			LexFromString(OutOptions.SpeedClass, *TempArg);
+		}
+	}
+
+	return true;
+}
+
+/* Convenience function for backends that check if the key should be missed or not */
+bool FDerivedDataBackendInterface::FBackendDebugOptions::ShouldSimulateMiss(const TCHAR* InCacheKey)
+{
+	bool bDoMiss = false;
+
+	if (SimulateMissTypes.Num() > 0)
+	{
+		FString TypeStr = FString(InCacheKey);
+		TypeStr = TypeStr.Left(TypeStr.Find(TEXT("_")));
+
+		if (SimulateMissTypes.Contains(TypeStr))
+		{
+			bDoMiss = true;
+		}
+	}
+
+	if (!bDoMiss && RandomMissRate > 0)
+	{
+		bDoMiss = FMath::RandHelper(100) < RandomMissRate;
+	}
+
+	return bDoMiss;
+}
+
+
+const TCHAR* LexToString(FDerivedDataBackendInterface::ESpeedClass SpeedClass)
+{
+	switch (SpeedClass)
+	{
+	case FDerivedDataBackendInterface::ESpeedClass::Unknown:
+		return TEXT("Unknown");
+	case FDerivedDataBackendInterface::ESpeedClass::Slow:
+		return TEXT("Slow");
+	case FDerivedDataBackendInterface::ESpeedClass::Ok:
+		return TEXT("Ok");
+	case FDerivedDataBackendInterface::ESpeedClass::Fast:
+		return TEXT("Fast");
+	case FDerivedDataBackendInterface::ESpeedClass::Local:
+		return TEXT("Local");
+	}
+
+	return TEXT("Unknow value! (Update LexToString!)");
+}
+
+
+void LexFromString(FDerivedDataBackendInterface::ESpeedClass& OutValue, const TCHAR* Buffer)
+{
+	OutValue = FDerivedDataBackendInterface::ESpeedClass::Unknown;
+
+	if (FCString::Stricmp(Buffer, TEXT("Slow")) == 0)
+	{
+		OutValue = FDerivedDataBackendInterface::ESpeedClass::Slow;
+	}
+	else if (FCString::Stricmp(Buffer, TEXT("Ok")) == 0)
+	{
+		OutValue = FDerivedDataBackendInterface::ESpeedClass::Ok;
+	}
+	if (FCString::Stricmp(Buffer, TEXT("Fast")) == 0)
+	{
+		OutValue = FDerivedDataBackendInterface::ESpeedClass::Fast;
+	}
+	if (FCString::Stricmp(Buffer, TEXT("Local")) == 0)
+	{
+		OutValue = FDerivedDataBackendInterface::ESpeedClass::Local;
+	}
+}
+
 #undef LOCTEXT_NAMESPACE
diff --git a/Engine/Source/Developer/DerivedDataCache/Private/DerivedDataLimitKeyLengthWrapper.h b/Engine/Source/Developer/DerivedDataCache/Private/DerivedDataLimitKeyLengthWrapper.h
--- a/Engine/Source/Developer/DerivedDataCache/Private/DerivedDataLimitKeyLengthWrapper.h
+++ b/Engine/Source/Developer/DerivedDataCache/Private/DerivedDataLimitKeyLengthWrapper.h
@@ -91,6 +91,11 @@ public:
 		return InnerBackend->WouldCache(CacheKey, InData);
 	}
 
+	virtual bool ApplyDebugOptions(FBackendDebugOptions& InOptions) override 
+	{ 
+		return InnerBackend->ApplyDebugOptions(InOptions);
+	}
+
 	/**
 	 * Synchronous retrieve of a cache item
 	 *
diff --git a/Engine/Source/Developer/DerivedDataCache/Private/FileSystemDerivedDataBackend.cpp b/Engine/Source/Developer/DerivedDataCache/Private/FileSystemDerivedDataBackend.cpp
--- a/Engine/Source/Developer/DerivedDataCache/Private/FileSystemDerivedDataBackend.cpp
+++ b/Engine/Source/Developer/DerivedDataCache/Private/FileSystemDerivedDataBackend.cpp
@@ -71,7 +71,7 @@ public:
 	 * @param InCacheDirectory	directory to store the cache in
 	 * @param bForceReadOnly	if true, do not attempt to write to this cache
 	*/
-	FFileSystemDerivedDataBackend(const TCHAR* InCacheDirectory, const TCHAR* InParams, const int InMissRate, const TCHAR* InAccessLogFileName)
+	FFileSystemDerivedDataBackend(const TCHAR* InCacheDirectory, const TCHAR* InParams, const TCHAR* InAccessLogFileName)
 		: CachePath(InCacheDirectory)
 		, SpeedClass(ESpeedClass::Unknown)
 		, bReadOnly(false)
@@ -80,7 +80,6 @@ public:
 		, DaysToDeleteUnusedFiles(15)
 		, bDisabled(false)
 		, TotalEstimatedBuildTime(0)
-		, DebugMissRate(0)
 	{
 		// If we find a platform that has more stingent limits, this needs to be rethought.
 		checkf(MAX_BACKEND_KEY_LENGTH + MAX_CACHE_DIR_LEN + MAX_BACKEND_NUMBERED_SUBFOLDER_LENGTH + MAX_CACHE_EXTENTION_LEN < FPlatformMisc::GetMaxPathLength(),
@@ -272,10 +271,6 @@ public:
 			{
 				MissingFiles.Add(KV.Key);
 			}
-			else if (!bReadOnly)
-			{
-				IFileManager::Get().SetTimeStamp(*KV.Key, CurrentTime);
-			}
 		}
 
 		// save total stat time
@@ -287,7 +282,7 @@ public:
 		UE_LOG(LogDerivedDataCache, Verbose, TEXT("Stat tests to %s took %.02f seconds"), *CachePath, TotalSeekTime);
 
 		// if seek times are very slow do a single read/write test just to confirm access
-		if (OutSeekTimeMS >= InSkipTestsIfSeeksExceedMS)
+		/*if (OutSeekTimeMS >= InSkipTestsIfSeeksExceedMS)
 		{
 			UE_LOG(LogDerivedDataCache, Warning, TEXT("Limiting read/write speed tests due to seek times of %.02f exceeding %.02fms. Values will be inaccurate."), OutSeekTimeMS, InSkipTestsIfSeeksExceedMS);
 
@@ -296,7 +291,7 @@ public:
 
 			TestFileEntries.Reset();
 			TestFileEntries.Add(Path, Size);		
-		}
+		}*/
 	
 		// create any files that were missing
 		if (!bReadOnly)
@@ -488,6 +483,11 @@ public:
 		}
 		*/
 
+		if (ShouldSimulateMiss(CacheKey))
+		{
+			return false;
+		}
+
 		FString Filename = BuildFilename(CacheKey);
 
 		FFileStatData FileStat = IFileManager::Get().GetStatData(*Filename);
@@ -523,7 +523,8 @@ public:
 		}
 
 		return FileStat.bIsValid;
-	}
+	}	
+
 	/**
 	 * Synchronous retrieve of a cache item
 	 *
@@ -538,17 +539,12 @@ public:
 		FString Filename = BuildFilename(CacheKey);
 		double StartTime = FPlatformTime::Seconds();
 
-		if (DebugMissRate > 0)
+		if (ShouldSimulateMiss(CacheKey))
 		{
-			bool Miss = FMath::RandHelper(100) < DebugMissRate;
-
-			if (Miss)
-			{
-				FScopeLock Lock(&MissedKeysCS);
-				UE_LOG(LogDerivedDataCache, Log, TEXT("%s Pretending key didn't exist: %s"), *GetName(), CacheKey);
-				DebugMissedKeys.Add(FName(CacheKey));
-				return false;
-			}
+			FScopeLock Lock(&MissedKeysCS);
+			UE_LOG(LogDerivedDataCache, Verbose, TEXT("Simulating miss in %s for %s"), *GetName(), CacheKey);
+			DebugMissedKeys.Add(FName(CacheKey));
+			return false;
 		}
 
 		if (FFileHelper::LoadFileToArray(Data,*Filename,FILEREAD_Silent))
@@ -616,7 +612,7 @@ public:
 			}
 
 			// don't put anything we pretended didn't exist
-			if (DebugMissRate > 0 && WasDebugMissed(CacheKey))
+			if ((CacheKey))
 			{
 				return;
 			}
@@ -743,16 +739,10 @@ public:
 		return false;
 	}
 
-	void SetDebugMissRate(int InMissRate)
+	bool ApplyDebugOptions(FBackendDebugOptions& InOptions) override
 	{
-		UE_LOG(LogDerivedDataCache, Log, TEXT("Setting DebugMissRate to %d for cache %s"), InMissRate, *GetName());
-		DebugMissRate = InMissRate;
-	}
-
-	bool WasDebugMissed(const TCHAR* InKey)
-	{
-		FScopeLock Lock(&MissedKeysCS);
-		return DebugMissedKeys.Contains(FName(InKey));
+		DebugOptions = InOptions;
+		return true;
 	}
 
 private:
@@ -797,27 +787,51 @@ private:
 	/** The total estimated build time accumulated from cache miss/put deltas */
 	double TotalEstimatedBuildTime;
 
-	/** Miss rate that can be set for profiling */
-	int DebugMissRate;
-
 	/** Access log to write to */
 	TUniquePtr<FAccessLogWriter> AccessLogWriter;
 
+	/** Debug Options */
+	FBackendDebugOptions DebugOptions;
+
 	/** Keys we ignored due to miss rate settings */
 	FCriticalSection MissedKeysCS;
 	TSet<FName> DebugMissedKeys;
 
-};
-
-FDerivedDataBackendInterface* CreateFileSystemDerivedDataBackend(const TCHAR* CacheDirectory, const TCHAR* InParams, const int InMissRate, const TCHAR* InAccessLogFileName /*= nullptr*/)
-{
-	FFileSystemDerivedDataBackend* FileDDB = new FFileSystemDerivedDataBackend( CacheDirectory, InParams, InMissRate, InAccessLogFileName);
+	bool DidSimulateMiss(const TCHAR* InKey)
+	{
+		if (DebugOptions.RandomMissRate == 0 || DebugOptions.SimulateMissTypes.Num() == 0)
+		{
+			return false;
+		}
+		FScopeLock Lock(&MissedKeysCS);
+		return DebugMissedKeys.Contains(FName(InKey));
+	}
 
-	if (InMissRate > 0)
+	bool ShouldSimulateMiss(const TCHAR* InKey)
 	{
-		FileDDB->SetDebugMissRate(InMissRate);
+		// once missed, always missed
+		if (DidSimulateMiss(InKey))
+		{
+			return true;
+		}
+
+		if (DebugOptions.ShouldSimulateMiss(InKey))
+		{
+			FScopeLock Lock(&MissedKeysCS);
+			UE_LOG(LogDerivedDataCache, Verbose, TEXT("Simulating miss in %s for %s"), *GetName(), InKey);
+			DebugMissedKeys.Add(FName(InKey));
+			return true;
+		}
+
+		return true;
 	}
 
+};
+
+FDerivedDataBackendInterface* CreateFileSystemDerivedDataBackend(const TCHAR* CacheDirectory, const TCHAR* InParams, const TCHAR* InAccessLogFileName /*= nullptr*/)
+{
+	FFileSystemDerivedDataBackend* FileDDB = new FFileSystemDerivedDataBackend( CacheDirectory, InParams, InAccessLogFileName);
+
 	if (!FileDDB->IsUsable())
 	{
 		delete FileDDB;
diff --git a/Engine/Source/Developer/DerivedDataCache/Private/HierarchicalDerivedDataBackend.h b/Engine/Source/Developer/DerivedDataCache/Private/HierarchicalDerivedDataBackend.h
--- a/Engine/Source/Developer/DerivedDataCache/Private/HierarchicalDerivedDataBackend.h
+++ b/Engine/Source/Developer/DerivedDataCache/Private/HierarchicalDerivedDataBackend.h
@@ -165,6 +165,20 @@ public:
 		return false;
 	}
 
+	bool ApplyDebugOptions(FBackendDebugOptions& InOptions) override
+	{
+		bool bSuccess = true;
+		for (int32 CacheIndex = 0; CacheIndex < InnerBackends.Num(); CacheIndex++)
+		{
+			if (!InnerBackends[CacheIndex]->ApplyDebugOptions(InOptions))
+			{
+				bSuccess = false;
+			}
+		}
+
+		return bSuccess;
+	}
+
 	/**
 	 * Synchronous retrieve of a cache item
 	 *
diff --git a/Engine/Source/Developer/DerivedDataCache/Private/MemoryDerivedDataBackend.cpp b/Engine/Source/Developer/DerivedDataCache/Private/MemoryDerivedDataBackend.cpp
--- a/Engine/Source/Developer/DerivedDataCache/Private/MemoryDerivedDataBackend.cpp
+++ b/Engine/Source/Developer/DerivedDataCache/Private/MemoryDerivedDataBackend.cpp
@@ -31,12 +31,18 @@ FDerivedDataBackendInterface::ESpeedClass FMemoryDerivedDataBackend::GetSpeedCla
 bool FMemoryDerivedDataBackend::CachedDataProbablyExists(const TCHAR* CacheKey)
 {
 	COOK_STAT(auto Timer = UsageStats.TimeProbablyExists());
-	FScopeLock ScopeLock(&SynchronizationObject);
+
+	if (ShouldSimulateMiss(CacheKey))
+	{
+		return false;
+	}
+
 	if (bDisabled)
 	{
 		return false;
 	}
 
+	FScopeLock ScopeLock(&SynchronizationObject);
 	bool Result = CacheItems.Contains(FString(CacheKey));
 	if (Result)
 	{
@@ -48,9 +54,16 @@ bool FMemoryDerivedDataBackend::CachedDataProbablyExists(const TCHAR* CacheKey)
 bool FMemoryDerivedDataBackend::GetCachedData(const TCHAR* CacheKey, TArray<uint8>& OutData)
 {
 	COOK_STAT(auto Timer = UsageStats.TimeGet());
-	FScopeLock ScopeLock(&SynchronizationObject);
+
+	if (ShouldSimulateMiss(CacheKey))
+	{
+		return false;
+	}
+	
 	if (!bDisabled)
 	{
+		FScopeLock ScopeLock(&SynchronizationObject);
+
 		FCacheValue* Item = CacheItems.FindRef(FString(CacheKey));
 		if (Item)
 		{
@@ -84,6 +97,11 @@ void FMemoryDerivedDataBackend::PutCachedData(const TCHAR* CacheKey, TArrayView<
 {
 	COOK_STAT(auto Timer = UsageStats.TimePut());
 	FScopeLock ScopeLock(&SynchronizationObject);
+
+	if (DidSimulateMiss(CacheKey))
+	{
+		return;
+	}
 	
 	// Should never hit this as higher level code should be checking..
 	if (!WouldCache(CacheKey, InData))
@@ -306,3 +324,38 @@ void FMemoryDerivedDataBackend::GatherUsageStats(TMap<FString, FDerivedDataCache
 {
 	COOK_STAT(UsageStatsMap.Add(FString::Printf(TEXT("%s: %s.%s"), *GraphPath, TEXT("MemoryBackend"), *CacheFilename), UsageStats));
 }
+
+bool FMemoryDerivedDataBackend::ApplyDebugOptions(FBackendDebugOptions& InOptions)
+{
+	DebugOptions = InOptions;
+	return true;
+}
+
+bool FMemoryDerivedDataBackend::DidSimulateMiss(const TCHAR* InKey)
+{
+	if (DebugOptions.RandomMissRate == 0 || DebugOptions.SimulateMissTypes.Num() == 0)
+	{
+		return false;
+	}
+	FScopeLock Lock(&MissedKeysCS);
+	return DebugMissedKeys.Contains(FName(InKey));
+}
+
+bool FMemoryDerivedDataBackend::ShouldSimulateMiss(const TCHAR* InKey)
+{
+	// once missed, always missed
+	if (DidSimulateMiss(InKey))
+	{
+		return true;
+	}
+
+	if (DebugOptions.ShouldSimulateMiss(InKey))
+	{
+		FScopeLock Lock(&MissedKeysCS);
+		UE_LOG(LogDerivedDataCache, Verbose, TEXT("Simulating miss in %s for %s"), *GetName(), InKey);
+		DebugMissedKeys.Add(FName(InKey));
+		return true;
+	}
+
+	return true;
+}
\ No newline at end of file
diff --git a/Engine/Source/Developer/DerivedDataCache/Private/MemoryDerivedDataBackend.h b/Engine/Source/Developer/DerivedDataCache/Private/MemoryDerivedDataBackend.h
--- a/Engine/Source/Developer/DerivedDataCache/Private/MemoryDerivedDataBackend.h
+++ b/Engine/Source/Developer/DerivedDataCache/Private/MemoryDerivedDataBackend.h
@@ -84,6 +84,11 @@ public:
 	 *  Determines if we would cache the provided data
 	 */
 	virtual bool WouldCache(const TCHAR* CacheKey, TArrayView<const uint8> InData) override;
+
+	/**
+	 * Apply debug options
+	 */
+	bool ApplyDebugOptions(FBackendDebugOptions& InOptions) override;
 	
 
 private:
@@ -136,5 +141,18 @@ private:
 									  + sizeof(int64)	// Size
 									  + sizeof(uint32), // CRC
 	};
+
+protected:
+
+	/* Debug helpers */
+	bool DidSimulateMiss(const TCHAR* InKey);
+	bool ShouldSimulateMiss(const TCHAR* InKey);
+
+	/** Debug Options */
+	FBackendDebugOptions DebugOptions;
+
+	/** Keys we ignored due to miss rate settings */
+	FCriticalSection MissedKeysCS;
+	TSet<FName> DebugMissedKeys;
 };
 
diff --git a/Engine/Source/Developer/DerivedDataCache/Private/PakFileDerivedDataBackend.h b/Engine/Source/Developer/DerivedDataCache/Private/PakFileDerivedDataBackend.h
--- a/Engine/Source/Developer/DerivedDataCache/Private/PakFileDerivedDataBackend.h
+++ b/Engine/Source/Developer/DerivedDataCache/Private/PakFileDerivedDataBackend.h
@@ -96,6 +96,8 @@ public:
 
 	virtual bool WouldCache(const TCHAR* CacheKey, TArrayView<const uint8> InData) override { return true; }
 
+	virtual bool ApplyDebugOptions(FBackendDebugOptions& InOptions) override { return false; }
+
 private:
 	FDerivedDataCacheUsageStats UsageStats;
 
diff --git a/Engine/Source/Developer/DerivedDataCache/Private/S3DerivedDataBackend.cpp b/Engine/Source/Developer/DerivedDataCache/Private/S3DerivedDataBackend.cpp
--- a/Engine/Source/Developer/DerivedDataCache/Private/S3DerivedDataBackend.cpp
+++ b/Engine/Source/Developer/DerivedDataCache/Private/S3DerivedDataBackend.cpp
@@ -815,6 +815,12 @@ bool FS3DerivedDataBackend::CachedDataProbablyExists(const TCHAR* CacheKey)
 {
 	const FBundle* Bundle;
 	const FBundleEntry* BundleEntry;
+
+	if (ShouldSimulateMiss(CacheKey))
+	{
+		return false;
+	}
+
 	if (!FindBundleEntry(CacheKey, Bundle, BundleEntry))
 	{
 		UE_LOG(LogDerivedDataCache, Verbose, TEXT("S3DerivedDataBackend: Cache miss on %s (probably)"), CacheKey);
@@ -829,6 +835,11 @@ bool FS3DerivedDataBackend::GetCachedData(const TCHAR* CacheKey, TArray<uint8>&
 	TRACE_COUNTER_ADD(S3DDC_Get, int64(1));
 	COOK_STAT(auto Timer = UsageStats.TimeGet());
 
+	if (ShouldSimulateMiss(CacheKey))
+	{
+		return false;
+	}
+
 	const FBundle* Bundle;
 	const FBundleEntry* BundleEntry;
 	if (FindBundleEntry(CacheKey, Bundle, BundleEntry))
@@ -1078,4 +1089,39 @@ bool FS3DerivedDataBackend::FindBundleEntry(const TCHAR* CacheKey, const FBundle
 	return false;
 }
 
+bool FS3DerivedDataBackend::ApplyDebugOptions(FBackendDebugOptions& InOptions)
+{
+	DebugOptions = InOptions;
+	return true;
+}
+
+bool FS3DerivedDataBackend::DidSimulateMiss(const TCHAR* InKey)
+{
+	if (DebugOptions.RandomMissRate == 0 || DebugOptions.SimulateMissTypes.Num() == 0)
+	{
+		return false;
+	}
+	FScopeLock Lock(&MissedKeysCS);
+	return DebugMissedKeys.Contains(FName(InKey));
+}
+
+bool FS3DerivedDataBackend::ShouldSimulateMiss(const TCHAR* InKey)
+{
+	// once missed, always missed
+	if (DidSimulateMiss(InKey))
+	{
+		return true;
+	}
+
+	if (DebugOptions.ShouldSimulateMiss(InKey))
+	{
+		FScopeLock Lock(&MissedKeysCS);
+		UE_LOG(LogDerivedDataCache, Verbose, TEXT("Simulating miss in %s for %s"), *GetName(), InKey);
+		DebugMissedKeys.Add(FName(InKey));
+		return true;
+	}
+
+	return true;
+}
+
 #endif
diff --git a/Engine/Source/Developer/DerivedDataCache/Private/S3DerivedDataBackend.h b/Engine/Source/Developer/DerivedDataCache/Private/S3DerivedDataBackend.h
--- a/Engine/Source/Developer/DerivedDataCache/Private/S3DerivedDataBackend.h
+++ b/Engine/Source/Developer/DerivedDataCache/Private/S3DerivedDataBackend.h
@@ -51,6 +51,8 @@ public:
 	bool TryToPrefetch(const TCHAR* CacheKey) override;
 	bool WouldCache(const TCHAR* CacheKey, TArrayView<const uint8> InData) override;
 
+	bool ApplyDebugOptions(FBackendDebugOptions& InOptions) override;
+
 private:
 	struct FBundle;
 	struct FBundleEntry;
@@ -73,6 +75,17 @@ private:
 	void RemoveUnusedBundles();
 	void ReadBundle(FBundle& Bundle);
 	bool FindBundleEntry(const TCHAR* CacheKey, const FBundle*& OutBundle, const FBundleEntry*& OutBundleEntry) const;
+
+	/* Debug helpers */
+	bool DidSimulateMiss(const TCHAR* InKey);
+	bool ShouldSimulateMiss(const TCHAR* InKey);
+
+	/** Debug Options */
+	FBackendDebugOptions DebugOptions;
+
+	/** Keys we ignored due to miss rate settings */
+	FCriticalSection MissedKeysCS;
+	TSet<FName> DebugMissedKeys;
 };
 
 #endif
diff --git a/Engine/Source/Developer/DerivedDataCache/Public/DerivedDataBackendInterface.h b/Engine/Source/Developer/DerivedDataCache/Public/DerivedDataBackendInterface.h
--- a/Engine/Source/Developer/DerivedDataCache/Public/DerivedDataBackendInterface.h
+++ b/Engine/Source/Developer/DerivedDataCache/Public/DerivedDataBackendInterface.h
@@ -40,6 +40,35 @@ public:
 		Local			/* Little to no impact from seek times and extremly fast reads */
 	};
 
+	/* Debug options that can be applied to backends to simulate different behavior */
+	struct FBackendDebugOptions
+	{
+		/* Percentage of requests that should result in random misses */
+		int					RandomMissRate;
+
+		/* Apply behavior of this speed class */
+		ESpeedClass			SpeedClass;		
+
+		/* Types of DDC entries that should always be a miss */
+		TArray<FString>		SimulateMissTypes;
+
+		FBackendDebugOptions()
+			: RandomMissRate(0)
+			, SpeedClass(ESpeedClass::Unknown)
+		{
+		}
+
+		/* Fill in the provided structure based on the name of the node (e.g. 'shared') and the provided token stream */
+		static bool ParseFromTokens(FBackendDebugOptions& OutOptions, const TCHAR* InNodeName, const TCHAR* InTokens);
+
+		/* 
+			Returns true if, according to the properties of this struct, the provided key should be treated as a miss.
+			Implementing that miss and accounting for any behaviour impact (e.g. skipping a subsequent put) is left to
+			each backend.
+		*/
+		bool ShouldSimulateMiss(const TCHAR* InCacheKey);
+	};
+
 	virtual ~FDerivedDataBackendInterface()
 	{
 	}
@@ -117,6 +146,12 @@ public:
 	 * a file size limit, etc.
 	 */
 	virtual bool WouldCache(const TCHAR* CacheKey, TArrayView<const uint8> InData) = 0;
+
+	/**
+	 *  Ask a backend to apply debug behavior to simulate different conditions. Backends that don't support these options should return 
+		false which will result in a warning if an attempt is made to apply these options.
+	 */
+	virtual bool ApplyDebugOptions(FBackendDebugOptions& InOptions) = 0;
 };
 
 class FDerivedDataBackend
@@ -163,22 +198,6 @@ public:
 
 };
 
-
-inline const TCHAR* LexToString(FDerivedDataBackendInterface::ESpeedClass SpeedClass)
-{
-	switch (SpeedClass)
-	{
-	case FDerivedDataBackendInterface::ESpeedClass::Unknown:
-		return TEXT("Unknown");
-	case FDerivedDataBackendInterface::ESpeedClass::Slow:
-		return TEXT("Slow");
-	case FDerivedDataBackendInterface::ESpeedClass::Ok:
-		return TEXT("Ok");
-	case FDerivedDataBackendInterface::ESpeedClass::Fast:
-		return TEXT("Fast");
-	case FDerivedDataBackendInterface::ESpeedClass::Local:
-		return TEXT("Local");
-	}
-
-	return TEXT("Unknow value! (Update LexToString!)");
-}
\ No newline at end of file
+/* Lexical conversions from and to enums */
+const TCHAR* LexToString(FDerivedDataBackendInterface::ESpeedClass SpeedClass);
+void LexFromString(FDerivedDataBackendInterface::ESpeedClass& OutValue, const TCHAR* Buffer);
\ No newline at end of file
