Moved some DDC debug options into a common class that can be provided to a backend to simulate certain behavior.

E.g.

'-ddc-local-missrate=50' simulates a 50% miss rate of keys on the local backend,
'-ddc-shared-misstypes=AnimSeq' simulates misses for any AnimSeq key. Use + to combine keys. E.g. AnimSeq+Audio

Memory DDC, FileSystemDDC and S3DDC support these options as do the classes that wrap other backends (e.g. AsyncPut).

#tests lots of warm and cold DDC
#swarm https://p4-swarm.epicgames.net/reviews/12458696

diff --git a/Engine/Source/Developer/DerivedDataCache/Private/DerivedDataBackends.cpp b/Engine/Source/Developer/DerivedDataCache/Private/DerivedDataBackends.cpp
--- a/Engine/Source/Developer/DerivedDataCache/Private/DerivedDataBackends.cpp
+++ b/Engine/Source/Developer/DerivedDataCache/Private/DerivedDataBackends.cpp
@@ -666,7 +666,7 @@ public:
 		FS3DerivedDataBackend* Backend = new FS3DerivedDataBackend(*ManifestPath, *BaseUrl, *Region, *CanaryObjectKey, *AccessKey, *SecretKey);
 		return new FDerivedDataBackendCorruptionWrapper(Backend);
 #else
-		UE_LOG(LogDerivedDataCache, Warning, TEXT("S3 backend is not supported in the current build configuration."));
+		UE_LOG(LogDerivedDataCache, Log, TEXT("S3 backend is not supported on the current platform."));
 		return nullptr;
 #endif
 	}
diff --git a/Engine/Source/Developer/DerivedDataCache/Private/FileSystemDerivedDataBackend.cpp b/Engine/Source/Developer/DerivedDataCache/Private/FileSystemDerivedDataBackend.cpp
--- a/Engine/Source/Developer/DerivedDataCache/Private/FileSystemDerivedDataBackend.cpp
+++ b/Engine/Source/Developer/DerivedDataCache/Private/FileSystemDerivedDataBackend.cpp
@@ -823,7 +823,7 @@ private:
 			return true;
 		}
 
-		return true;
+		return false;
 	}
 
 };
diff --git a/Engine/Source/Developer/DerivedDataCache/Private/MemoryDerivedDataBackend.cpp b/Engine/Source/Developer/DerivedDataCache/Private/MemoryDerivedDataBackend.cpp
--- a/Engine/Source/Developer/DerivedDataCache/Private/MemoryDerivedDataBackend.cpp
+++ b/Engine/Source/Developer/DerivedDataCache/Private/MemoryDerivedDataBackend.cpp
@@ -357,5 +357,5 @@ bool FMemoryDerivedDataBackend::ShouldSimulateMiss(const TCHAR* InKey)
 		return true;
 	}
 
-	return true;
+	return false;
 }
\ No newline at end of file
diff --git a/Engine/Source/Developer/DerivedDataCache/Private/S3DerivedDataBackend.cpp b/Engine/Source/Developer/DerivedDataCache/Private/S3DerivedDataBackend.cpp
--- a/Engine/Source/Developer/DerivedDataCache/Private/S3DerivedDataBackend.cpp
+++ b/Engine/Source/Developer/DerivedDataCache/Private/S3DerivedDataBackend.cpp
@@ -1121,7 +1121,7 @@ bool FS3DerivedDataBackend::ShouldSimulateMiss(const TCHAR* InKey)
 		return true;
 	}
 
-	return true;
+	return false;
 }
 
 #endif
