Fixed a... typo? bad resolve? that was causing new data not to be written to DDC caches

diff --git a/Engine/Source/Developer/DerivedDataCache/Private/FileSystemDerivedDataBackend.cpp b/Engine/Source/Developer/DerivedDataCache/Private/FileSystemDerivedDataBackend.cpp
--- a/Engine/Source/Developer/DerivedDataCache/Private/FileSystemDerivedDataBackend.cpp
+++ b/Engine/Source/Developer/DerivedDataCache/Private/FileSystemDerivedDataBackend.cpp
@@ -612,7 +612,7 @@ public:
 			}
 
 			// don't put anything we pretended didn't exist
-			if ((CacheKey))
+			if (ShouldSimulateMiss(CacheKey))
 			{
 				return;
 			}
