Allow disabling the S3 backend by setting the UE-S3DataCachePath environment variable to "None"

diff --git a/Engine/Source/Developer/DerivedDataCache/Private/DerivedDataBackends.cpp b/Engine/Source/Developer/DerivedDataCache/Private/DerivedDataBackends.cpp
--- a/Engine/Source/Developer/DerivedDataCache/Private/DerivedDataBackends.cpp
+++ b/Engine/Source/Developer/DerivedDataCache/Private/DerivedDataBackends.cpp
@@ -670,8 +670,16 @@ public:
 			FString FilesystemCachePathEnv = FPlatformMisc::GetEnvironmentVariable(*EnvPathOverride);
 			if (FilesystemCachePathEnv.Len() > 0)
 			{
-				CachePath = FilesystemCachePathEnv;
-				UE_LOG(LogDerivedDataCache, Log, TEXT("Found environment variable %s=%s"), *EnvPathOverride, *CachePath);
+				if (FilesystemCachePathEnv == TEXT("None"))
+				{
+					UE_LOG(LogDerivedDataCache, Log, TEXT("Node %s disabled due to %s=None"), *FilesystemCachePathEnv);
+					return nullptr;
+				}
+				else
+				{
+					CachePath = FilesystemCachePathEnv;
+					UE_LOG(LogDerivedDataCache, Log, TEXT("Found environment variable %s=%s"), *EnvPathOverride, *CachePath);
+				}
 			}
 		}
 
