Fixed crash when verbose DDC logging is enabled

diff --git a/Engine/Source/Developer/DerivedDataCache/Private/FileSystemDerivedDataBackend.cpp b/Engine/Source/Developer/DerivedDataCache/Private/FileSystemDerivedDataBackend.cpp
--- a/Engine/Source/Developer/DerivedDataCache/Private/FileSystemDerivedDataBackend.cpp
+++ b/Engine/Source/Developer/DerivedDataCache/Private/FileSystemDerivedDataBackend.cpp
@@ -527,7 +527,7 @@ public:
 			}
 		}
 
-		UE_LOG(LogDerivedDataCache, Verbose, TEXT("%s CachedDataProbablyExists=%d for %s"), *GetName(), FileStat.bIsValid, *CacheKey);
+		UE_LOG(LogDerivedDataCache, Verbose, TEXT("%s CachedDataProbablyExists=%d for %s"), *GetName(), FileStat.bIsValid, CacheKey);
 
 		return FileStat.bIsValid;
 	}	
