Changed the S3 DDC backend so it elects not to write to lower cache levels.

(This value is only checked on backends that have a cache of the data, so data that is not in the S3 DDC will still be written back to lower levels if the S3 cache misses).

Changed the speed class for backends to be elligible to be backfilled to fast.

diff --git a/Engine/Source/Developer/DerivedDataCache/Private/HierarchicalDerivedDataBackend.h b/Engine/Source/Developer/DerivedDataCache/Private/HierarchicalDerivedDataBackend.h
--- a/Engine/Source/Developer/DerivedDataCache/Private/HierarchicalDerivedDataBackend.h
+++ b/Engine/Source/Developer/DerivedDataCache/Private/HierarchicalDerivedDataBackend.h
@@ -9,6 +9,7 @@
 #include "DerivedDataBackendAsyncPutWrapper.h"
 #include "Templates/UniquePtr.h"
 
+PRAGMA_DISABLE_OPTIMIZATION
 
 /** 
  * A backend wrapper that implements a cache hierarchy of backends. 
@@ -179,6 +180,8 @@ public:
 		return bSuccess;
 	}
 
+
+
 	/**
 	 * Synchronous retrieve of a cache item
 	 *
@@ -233,9 +236,8 @@ public:
 									break; //do not write things that are already in the read only pak file
 								}
 							}
-							else if (PutBackend->GetSpeedClass() >= FDerivedDataBackendInterface::ESpeedClass::Ok && PutBackend->WouldCache(CacheKey, OutData))
-							{
-								
+							else if (PutBackend->GetSpeedClass() >= FDerivedDataBackendInterface::ESpeedClass::Fast && PutBackend->WouldCache(CacheKey, OutData))
+							{								
 								AsyncPutInnerBackends[PutCacheIndex]->PutCachedData(CacheKey, OutData, false); // we do not need to force a put here
 								UE_LOG(LogDerivedDataCache, Verbose, TEXT("Back-filling cache %s with: %s (%d bytes) (force=%d)"), *PutBackend->GetName(), CacheKey, OutData.Num(), false);
 							}
@@ -322,6 +324,8 @@ public:
 		});
 	}
 
+
+
 private:
 	FDerivedDataCacheUsageStats UsageStats;
 
@@ -332,3 +336,5 @@ private:
 	/** As an optimization, we check our writable status at contruction **/
 	bool bIsWritable;
 };
+
+PRAGMA_ENABLE_OPTIMIZATION
\ No newline at end of file
diff --git a/Engine/Source/Developer/DerivedDataCache/Private/S3DerivedDataBackend.h b/Engine/Source/Developer/DerivedDataCache/Private/S3DerivedDataBackend.h
--- a/Engine/Source/Developer/DerivedDataCache/Private/S3DerivedDataBackend.h
+++ b/Engine/Source/Developer/DerivedDataCache/Private/S3DerivedDataBackend.h
@@ -40,7 +40,13 @@ public:
 	 * @return true if usable
 	 */
 	bool IsUsable() const;
+
+	/* S3 Cache cannot be written to*/
 	bool IsWritable() override { return false; }
+
+	/* S3 Cache does not try to write back to lower caches (e.g. Shared DDC) */
+	bool BackfillLowerCacheLevels() override { return false; }
+
 	bool CachedDataProbablyExists(const TCHAR* CacheKey) override;
 	bool GetCachedData(const TCHAR* CacheKey, TArray<uint8>& OutData) override;
 	void PutCachedData(const TCHAR* CacheKey, TArrayView<const uint8> InData, bool bPutEvenIfExists) override;
