Removed some deoptimization pragmas that were left in accidentally.

diff --git a/Engine/Source/Developer/DerivedDataCache/Private/HierarchicalDerivedDataBackend.h b/Engine/Source/Developer/DerivedDataCache/Private/HierarchicalDerivedDataBackend.h
--- a/Engine/Source/Developer/DerivedDataCache/Private/HierarchicalDerivedDataBackend.h
+++ b/Engine/Source/Developer/DerivedDataCache/Private/HierarchicalDerivedDataBackend.h
@@ -9,8 +9,6 @@
 #include "DerivedDataBackendAsyncPutWrapper.h"
 #include "Templates/UniquePtr.h"
 
-PRAGMA_DISABLE_OPTIMIZATION
-
 /** 
  * A backend wrapper that implements a cache hierarchy of backends. 
 **/
@@ -337,4 +335,3 @@ private:
 	bool bIsWritable;
 };
 
-PRAGMA_ENABLE_OPTIMIZATION
\ No newline at end of file
