Fix for shader compilers showing message about unmapped drives

diff --git a/Engine/Source/Developer/DerivedDataCache/Private/DerivedDataBackends.cpp b/Engine/Source/Developer/DerivedDataCache/Private/DerivedDataBackends.cpp
--- a/Engine/Source/Developer/DerivedDataCache/Private/DerivedDataBackends.cpp
+++ b/Engine/Source/Developer/DerivedDataCache/Private/DerivedDataBackends.cpp
@@ -603,7 +603,7 @@ public:
 					UE_LOG(LogDerivedDataCache, Warning, TEXT("%s"), *Message);
 
 					// Give the user a chance to retry incase they need to connect a network drive or something.
-					if (!FApp::IsUnattended())
+					if (!FApp::IsUnattended() && !IS_PROGRAM)
 					{
 						Message += FString::Printf(TEXT("\n\nRetry connection to %s?"), *Path);
 						EAppReturnType::Type MessageReturn = FPlatformMisc::MessageBoxExt(EAppMsgType::YesNo, *Message, TEXT("Could not access DDC"));
