- Moved some engine modules from NotForLicensees/NoRedist to Restricted

diff --git a/Engine/Source/Developer/DerivedDataCache/DerivedDataCache.Build.cs b/Engine/Source/Developer/DerivedDataCache/DerivedDataCache.Build.cs
--- a/Engine/Source/Developer/DerivedDataCache/DerivedDataCache.Build.cs
+++ b/Engine/Source/Developer/DerivedDataCache/DerivedDataCache.Build.cs
@@ -16,7 +16,7 @@ public class DerivedDataCache : ModuleRules
 		AddEngineThirdPartyPrivateStaticDependencies(Target, "OpenSSL");
 
 		// Internal (NotForLicensees) module
-		var DDCUtilsModule = Path.Combine("Developer", "NotForLicensees", "DDCUtils", "DDCUtils.Build.cs");
+		var DDCUtilsModule = Path.Combine(EngineDirectory, "Restricted/NotForLicensees/Engine/Source/Developer/DDCUtils/DDCUtils.Build.cs");
 		if (File.Exists(DDCUtilsModule))
 		{
 			DynamicallyLoadedModuleNames.Add("DDCUtils");
