Added DDC options to engine test to investigate corruption issue

- VerifyDDC option is now correctly passed to the editor for editor-based tests and the cooker for cooked tests
- DDC= option added that allows kicked tests to use a different DDC backend (e.g. noshared). Defaults to 'default' which will use the default DDC set up.

diff --git a/Engine/Source/Developer/DerivedDataCache/Private/DerivedDataCache.cpp b/Engine/Source/Developer/DerivedDataCache/Private/DerivedDataCache.cpp
--- a/Engine/Source/Developer/DerivedDataCache/Private/DerivedDataCache.cpp
+++ b/Engine/Source/Developer/DerivedDataCache/Private/DerivedDataCache.cpp
@@ -284,6 +284,8 @@ public:
 		FDerivedDataBackend::Get(); // we need to make sure this starts before we all us to start
 
 		GVerifyDDC = FParse::Param(FCommandLine::Get(), TEXT("VerifyDDC"));
+
+		UE_CLOG(GVerifyDDC, LogDerivedDataCache, Display, TEXT("Items retrieved from the DDC will be verified (-VerifyDDC)"));
 	}
 
 	/** Destructor, flushes all sync tasks **/
