Add some more verbose DDC logging

diff --git a/Engine/Source/Developer/DerivedDataCache/Private/DerivedDataCache.cpp b/Engine/Source/Developer/DerivedDataCache/Private/DerivedDataCache.cpp
--- a/Engine/Source/Developer/DerivedDataCache/Private/DerivedDataCache.cpp
+++ b/Engine/Source/Developer/DerivedDataCache/Private/DerivedDataCache.cpp
@@ -323,7 +323,7 @@ public:
 		FScopeLock ScopeLock(&SynchronizationObject);
 		const uint32 Handle = NextHandle();
 		FString CacheKey = FDerivedDataCache::BuildCacheKey(DataDeriver);
-		UE_LOG(LogDerivedDataCache, VeryVerbose, TEXT("GetAsynchronous %s from '%s'"), *CacheKey, *DataDeriver->GetDebugContextString());
+		UE_LOG(LogDerivedDataCache, VeryVerbose, TEXT("GetAsynchronous %s from '%s', Handle %d"), *CacheKey, *DataDeriver->GetDebugContextString(), Handle);
 		const bool bSync = !DataDeriver->IsBuildThreadsafe();
 		FAsyncTask<FBuildAsyncWorker>* AsyncTask = new FAsyncTask<FBuildAsyncWorker>(DataDeriver, *CacheKey, bSync);
 		check(!PendingTasks.Contains(Handle));
@@ -367,11 +367,12 @@ public:
 			}
 			check(AsyncTask);
 			AsyncTask->EnsureCompletion();
+			UE_LOG(LogDerivedDataCache, Verbose, TEXT("WaitAsynchronousCompletion, Handle %d"), Handle);
 		}
 		INC_FLOAT_STAT_BY(STAT_DDC_ASyncWaitTime,(float)ThisTime);
 	}
 
-	virtual bool GetAsynchronousResults(uint32 Handle, TArray<uint8>& OutData, bool* bDataWasBuilt = nullptr) override
+	virtual bool GetAsynchronousResults(uint32 Handle, TArray<uint8>& OutData, bool* bOutDataWasBuilt = nullptr) override
 	{
 		DDC_SCOPE_CYCLE_COUNTER(DDC_GetAsynchronousResults);
 		FAsyncTask<FBuildAsyncWorker>* AsyncTask = NULL;
@@ -380,15 +381,19 @@ public:
 			PendingTasks.RemoveAndCopyValue(Handle,AsyncTask);
 		}
 		check(AsyncTask);
-		if (bDataWasBuilt)
+		const bool bDataWasBuilt = AsyncTask->GetTask().bDataWasBuilt;
+		if (bOutDataWasBuilt)
 		{
-			*bDataWasBuilt = AsyncTask->GetTask().bDataWasBuilt;
+			*bOutDataWasBuilt = bDataWasBuilt;
 		}
 		if (!AsyncTask->GetTask().bSuccess)
 		{
+			UE_LOG(LogDerivedDataCache, Verbose, TEXT("GetAsynchronousResults, bDataWasBuilt: %d, Handle %d, FAILED"), (int32)bDataWasBuilt, Handle);
 			delete AsyncTask;
 			return false;
 		}
+
+		UE_LOG(LogDerivedDataCache, Verbose, TEXT("GetAsynchronousResults, bDataWasBuilt: %d, Handle %d, SUCCESS"), (int32)bDataWasBuilt, Handle);
 		OutData = MoveTemp(AsyncTask->GetTask().Data);
 		delete AsyncTask;
 		check(OutData.Num());
@@ -410,8 +415,8 @@ public:
 	{
 		DDC_SCOPE_CYCLE_COUNTER(DDC_GetAsynchronous_Handle);
 		FScopeLock ScopeLock(&SynchronizationObject);
-		UE_LOG(LogDerivedDataCache, VeryVerbose, TEXT("GetAsynchronous %s from '%.*s'"), CacheKey, DataContext.Len(), DataContext.GetData());
 		const uint32 Handle = NextHandle();
+		UE_LOG(LogDerivedDataCache, VeryVerbose, TEXT("GetAsynchronous %s from '%.*s', Handle %d"), CacheKey, DataContext.Len(), DataContext.GetData(), Handle);
 		FAsyncTask<FBuildAsyncWorker>* AsyncTask = new FAsyncTask<FBuildAsyncWorker>((FDerivedDataPluginInterface*)NULL, CacheKey, false);
 		check(!PendingTasks.Contains(Handle));
 		PendingTasks.Add(Handle, AsyncTask);
