Edit-grated improved DDC behavior logging from Main to get better info around failing cooked tests

Temporarily turning on verbose DDC logging for EngineTest project

diff --git a/Engine/Source/Developer/DerivedDataCache/Private/DerivedDataBackendAsyncPutWrapper.cpp b/Engine/Source/Developer/DerivedDataCache/Private/DerivedDataBackendAsyncPutWrapper.cpp
--- a/Engine/Source/Developer/DerivedDataCache/Private/DerivedDataBackendAsyncPutWrapper.cpp
+++ b/Engine/Source/Developer/DerivedDataCache/Private/DerivedDataBackendAsyncPutWrapper.cpp
@@ -41,6 +41,8 @@ public:
 	/** Call the inner backend and when that completes, remove the memory cache */
 	void DoWork()
 	{
+		FPlatformProcess::Sleep(1);
+
 		COOK_STAT(auto Timer = UsageStats.TimePut());
 		bool bOk = true;
 		bool bDidTry = false;
@@ -67,6 +69,7 @@ public:
 		}
 		FilesInFlight->Remove(CacheKey);
 		FDerivedDataBackend::Get().AddToAsyncCompletionCounter(-1);
+		UE_LOG(LogDerivedDataCache, Log, TEXT("%s: Completed AsyncPut of %s."), *InnerBackend->GetName(), *CacheKey);
 	}
 
 	FORCEINLINE TStatId GetStatId() const
@@ -89,6 +92,7 @@ public:
 		}
 		FilesInFlight->Remove(CacheKey);
 		FDerivedDataBackend::Get().AddToAsyncCompletionCounter(-1);
+		UE_LOG(LogDerivedDataCache, Log, TEXT("%s: Abandoned AsyncPut of %s."), *InnerBackend->GetName(), *CacheKey);
 	}
 };
 
@@ -121,6 +125,8 @@ bool FDerivedDataBackendAsyncPutWrapper::CachedDataProbablyExists(const TCHAR* C
 	COOK_STAT(auto Timer = UsageStats.TimeProbablyExists());
 	bool Result = (InflightCache && InflightCache->CachedDataProbablyExists(CacheKey)) || InnerBackend->CachedDataProbablyExists(CacheKey);
 	COOK_STAT(if (Result) {	Timer.AddHit(0); });
+
+	UE_LOG(LogDerivedDataCache, Verbose, TEXT("%s CachedDataProbablyExists=%d for %s"), *GetName(), Result, CacheKey);
 	return Result;
 }
 
@@ -160,13 +166,19 @@ bool FDerivedDataBackendAsyncPutWrapper::GetCachedData(const TCHAR* CacheKey, TA
 	if (InflightCache && InflightCache->GetCachedData(CacheKey, OutData))
 	{
 		COOK_STAT(Timer.AddHit(OutData.Num()));
+		UE_LOG(LogDerivedDataCache, Verbose, TEXT("%s CacheHit from InFlightCache on %s"), *GetName(), CacheKey);
 		return true;
 	}
 	bool bSuccess = InnerBackend->GetCachedData(CacheKey, OutData);
 	if (bSuccess)
 	{
+		UE_LOG(LogDerivedDataCache, Verbose, TEXT("%s Cache hit on %s"), *GetName(), CacheKey);
 		COOK_STAT(Timer.AddHit(OutData.Num()));
 	}
+	else
+	{
+		UE_LOG(LogDerivedDataCache, Verbose, TEXT("%s Cache miss on %s"), *GetName(), CacheKey);
+	}
 	return bSuccess;
 }
 
@@ -186,11 +198,15 @@ void FDerivedDataBackendAsyncPutWrapper::PutCachedData(const TCHAR* CacheKey, TA
 	{
 		if (InflightCache->CachedDataProbablyExists(CacheKey))
 		{
+			UE_LOG(LogDerivedDataCache, Verbose, TEXT("%s skipping out of key already in in-flight cache %s"), *GetName(), CacheKey);
 			return; // if it is already on its way, we don't need to send it again
 		}
 		InflightCache->PutCachedData(CacheKey, InData, true); // temp copy stored in memory while the async task waits to complete
 		COOK_STAT(Timer.AddHit(InData.Num()));
 	}
+
+	UE_LOG(LogDerivedDataCache, Verbose, TEXT("%s queueing %s for put"), *GetName(), CacheKey);
+
 	FDerivedDataBackend::Get().AddToAsyncCompletionCounter(1);
 	(new FAutoDeleteAsyncTask<FCachePutAsyncWorker>(CacheKey, InData, InnerBackend, bPutEvenIfExists, InflightCache.Get(), &FilesInFlight, UsageStats))->StartBackgroundTask();
 }
@@ -210,6 +226,8 @@ void FDerivedDataBackendAsyncPutWrapper::RemoveCachedData(const TCHAR* CacheKey,
 		InflightCache->RemoveCachedData(CacheKey, bTransient);
 	}
 	InnerBackend->RemoveCachedData(CacheKey, bTransient);
+
+	UE_LOG(LogDerivedDataCache, Verbose, TEXT("%s removed %s"), *GetName(), CacheKey)
 }
 
 void FDerivedDataBackendAsyncPutWrapper::GatherUsageStats(TMap<FString, FDerivedDataCacheUsageStats>& UsageStatsMap, FString&& GraphPath)
diff --git a/Engine/Source/Developer/DerivedDataCache/Private/DerivedDataBackends.cpp b/Engine/Source/Developer/DerivedDataCache/Private/DerivedDataBackends.cpp
--- a/Engine/Source/Developer/DerivedDataCache/Private/DerivedDataBackends.cpp
+++ b/Engine/Source/Developer/DerivedDataCache/Private/DerivedDataBackends.cpp
@@ -764,7 +764,7 @@ public:
 		FString Filename;
 
 		FParse::Value( Entry, TEXT("Filename="), Filename );
-		Cache = new FMemoryDerivedDataBackend(TEXT("Memory"));
+		Cache = new FMemoryDerivedDataBackend(NodeName);
 		if( Cache && Filename.Len() )
 		{
 			if( Cache->LoadCache( *Filename ) )
diff --git a/Engine/Source/Developer/DerivedDataCache/Private/FileSystemDerivedDataBackend.cpp b/Engine/Source/Developer/DerivedDataCache/Private/FileSystemDerivedDataBackend.cpp
--- a/Engine/Source/Developer/DerivedDataCache/Private/FileSystemDerivedDataBackend.cpp
+++ b/Engine/Source/Developer/DerivedDataCache/Private/FileSystemDerivedDataBackend.cpp
@@ -655,7 +655,7 @@ public:
 							}
 							else
 							{
-								UE_LOG(LogDerivedDataCache, Verbose, TEXT("%s: Successful cache put to %s"),*GetName(), *Filename);
+								UE_LOG(LogDerivedDataCache, Verbose, TEXT("%s: Successful cache put of %s to %s"),*GetName(), CacheKey, *Filename);
 							}
 						}
 					}
@@ -680,6 +680,7 @@ public:
 			else
 			{
 				COOK_STAT(Timer.AddMiss(Data.Num()));
+				UE_LOG(LogDerivedDataCache, Verbose, TEXT("%s skipping put to existing file %s"), *GetName(), CacheKey);
 			}
 
 			// If not using a shared cache, update estimated build time
@@ -719,7 +720,10 @@ public:
 				}
 
 			}
-
+		}
+		else
+		{
+			UE_LOG(LogDerivedDataCache, Verbose, TEXT("%s is read only. Skipping put of %s"), *GetName(), CacheKey);
 		}
 	}
 
diff --git a/Engine/Source/Developer/DerivedDataCache/Private/HierarchicalDerivedDataBackend.h b/Engine/Source/Developer/DerivedDataCache/Private/HierarchicalDerivedDataBackend.h
--- a/Engine/Source/Developer/DerivedDataCache/Private/HierarchicalDerivedDataBackend.h
+++ b/Engine/Source/Developer/DerivedDataCache/Private/HierarchicalDerivedDataBackend.h
@@ -99,6 +99,15 @@ public:
 				COOK_STAT(Timer.AddHit(0));
 				return true;
 			}
+			else
+			{
+				extern bool GVerifyDDC;
+
+				if (GVerifyDDC)
+				{
+					ensureMsgf(!AsyncPutInnerBackends[CacheIndex]->CachedDataProbablyExists(CacheKey), TEXT("%s did not exist in sync interface for CachedDataProbablyExists but was found in async wrapper"), CacheKey);
+				}
+			}
 		}
 		return false;
 	}
@@ -254,6 +263,16 @@ public:
 				COOK_STAT(Timer.AddHit(OutData.Num()));
 				return true;
 			}
+			else
+			{
+				extern bool GVerifyDDC;
+
+				if (GVerifyDDC)
+				{
+					TArray<uint8> TempData;
+					ensureMsgf(!AsyncPutInnerBackends[CacheIndex]->GetCachedData(CacheKey, TempData), TEXT("CacheKey %s did not exist in sync interface for GetCachedData but was found in async wrapper"), CacheKey);
+				}
+			}
 		}
 		return false;
 	}
