Removed some debugging code that was accidentally submitted.
Droped some debug logging to VeryVerbose
Set EngineTest back to regular DDC logging

diff --git a/Engine/Source/Developer/DerivedDataCache/Private/DerivedDataBackendAsyncPutWrapper.cpp b/Engine/Source/Developer/DerivedDataCache/Private/DerivedDataBackendAsyncPutWrapper.cpp
--- a/Engine/Source/Developer/DerivedDataCache/Private/DerivedDataBackendAsyncPutWrapper.cpp
+++ b/Engine/Source/Developer/DerivedDataCache/Private/DerivedDataBackendAsyncPutWrapper.cpp
@@ -41,8 +41,6 @@ public:
 	/** Call the inner backend and when that completes, remove the memory cache */
 	void DoWork()
 	{
-		FPlatformProcess::Sleep(1);
-
 		COOK_STAT(auto Timer = UsageStats.TimePut());
 		bool bOk = true;
 		bool bDidTry = false;
@@ -59,7 +57,7 @@ public:
 			InnerBackend->PutCachedData(*CacheKey, Data, false);
 			if (!InnerBackend->CachedDataProbablyExists(*CacheKey))
 			{
-				UE_LOG(LogDerivedDataCache, Log, TEXT("%s: Put failed, keeping in memory copy %s."),*InnerBackend->GetName(), *CacheKey);
+				UE_LOG(LogDerivedDataCache, Warning, TEXT("%s: Put failed, keeping in memory copy %s."),*InnerBackend->GetName(), *CacheKey);
 				bOk = false;
 			}
 		}
@@ -69,7 +67,7 @@ public:
 		}
 		FilesInFlight->Remove(CacheKey);
 		FDerivedDataBackend::Get().AddToAsyncCompletionCounter(-1);
-		UE_LOG(LogDerivedDataCache, Log, TEXT("%s: Completed AsyncPut of %s."), *InnerBackend->GetName(), *CacheKey);
+		UE_LOG(LogDerivedDataCache, VeryVerbose, TEXT("%s: Completed AsyncPut of %s."), *InnerBackend->GetName(), *CacheKey);
 	}
 
 	FORCEINLINE TStatId GetStatId() const
@@ -92,7 +90,7 @@ public:
 		}
 		FilesInFlight->Remove(CacheKey);
 		FDerivedDataBackend::Get().AddToAsyncCompletionCounter(-1);
-		UE_LOG(LogDerivedDataCache, Log, TEXT("%s: Abandoned AsyncPut of %s."), *InnerBackend->GetName(), *CacheKey);
+		UE_LOG(LogDerivedDataCache, VeryVerbose, TEXT("%s: Abandoned AsyncPut of %s."), *InnerBackend->GetName(), *CacheKey);
 	}
 };
 
