Editgrate fix for warning about -ddc=default not being a valid backend

diff --git a/Engine/Source/Developer/DerivedDataCache/Private/DerivedDataBackends.cpp b/Engine/Source/Developer/DerivedDataCache/Private/DerivedDataBackends.cpp
--- a/Engine/Source/Developer/DerivedDataCache/Private/DerivedDataBackends.cpp
+++ b/Engine/Source/Developer/DerivedDataCache/Private/DerivedDataBackends.cpp
@@ -64,8 +64,8 @@ public:
 		RootCache = NULL;
 		TMap<FString, FDerivedDataBackendInterface*> ParsedNodes;		
 
-		// Create the graph using ini settings
-		if( FParse::Value( FCommandLine::Get(), TEXT("-DDC="), GraphName ) )
+		// Create the graph using ini settings. The string "default" forwards creation to use the default graph.
+		if( FParse::Value( FCommandLine::Get(), TEXT("-DDC="), GraphName ) && FCString::Stricmp(*GraphName, TEXT("default")) != 0 )
 		{
 			if( GraphName.Len() > 0 )
 			{
@@ -77,7 +77,7 @@ public:
 				// Remove references to any backend instances that might have been created
 				ParsedNodes.Empty();
 				DestroyCreatedBackends();
-				UE_LOG( LogDerivedDataCache, Warning, TEXT("FDerivedDataBackendGraph:  Unable to create backend graph using the specified graph settings (%s). Reverting to default."), *GraphName );
+				UE_LOG( LogDerivedDataCache, Warning, TEXT("FDerivedDataBackendGraph: Unable to create backend graph using the specified graph settings (\"%s\"). Reverting to the default graph."), *GraphName );
 			}
 		}
 
