Demote no-access to the DDC file share to be a warning instead of an error.

Making this an error breaks progression of BuildCookRun scenarios under UAT.

I wasn't able to find a way to communicate arguments from UAT's command line to it's invocation of the cooker. The default behavior of performance liabilities shouldn't be a failed build. Ideally it would be optional, and there would be a feature to promote this warning to an error where it is specifically desired, perhaps when running in a data center.

diff --git a/Engine/Source/Developer/DerivedDataCache/Private/FileSystemDerivedDataBackend.cpp b/Engine/Source/Developer/DerivedDataCache/Private/FileSystemDerivedDataBackend.cpp
--- a/Engine/Source/Developer/DerivedDataCache/Private/FileSystemDerivedDataBackend.cpp
+++ b/Engine/Source/Developer/DerivedDataCache/Private/FileSystemDerivedDataBackend.cpp
@@ -145,7 +145,7 @@ public:
 		if (!SkipSpeedTest && !RunSpeedTest(ConsiderSlowAtMS * 2, SeekTimeMS, ReadSpeedMBs, WriteSpeedMBs))
 		{
 			bDisabled = true;
-			UE_LOG(LogDerivedDataCache, Error, TEXT("No read or write access to %s"), *CachePath);
+			UE_LOG(LogDerivedDataCache, Warning, TEXT("No read or write access to %s"), *CachePath);
 		}
 		else
 		{
