Fix stutters in the game-thread while streaming texture mips from DDC.

This will fix a good bunch of ForceRebuildPlatformData in UTexture2D::UpdateStreamingStatus because of the Boot DDC telling it has some mips cached effectively preventing other cache layers from receiving them.
When the Boot DDC is disabled after the editor boot, all the mips will go into oblivion, and nobody will have them anymore, causing stutters in the game-thread while streaming because a second full rebuild is now required.

Any DDC that plans on disappearing should never be alone in the cache hierarchy and should give other levels the opportunity to properly cache data before its disapperance.

diff --git a/Engine/Source/Developer/DerivedDataCache/Private/DerivedDataBackends.cpp b/Engine/Source/Developer/DerivedDataCache/Private/DerivedDataBackends.cpp
--- a/Engine/Source/Developer/DerivedDataCache/Private/DerivedDataBackends.cpp
+++ b/Engine/Source/Developer/DerivedDataCache/Private/DerivedDataBackends.cpp
@@ -772,7 +772,7 @@ public:
 			MaxCacheSize = FMath::Min(MaxCacheSize, MaxSupportedCacheSize);
 
 			UE_LOG( LogDerivedDataCache, Display, TEXT("Max Cache Size: %d MB"), MaxCacheSize);
-			Cache = new FMemoryDerivedDataBackend(TEXT("Boot"), MaxCacheSize * 1024 * 1024);
+			Cache = new FMemoryDerivedDataBackend(TEXT("Boot"), MaxCacheSize * 1024 * 1024, true /* bCanBeDisabled */);
 
 			if( Cache && Filename.Len() )
 			{
diff --git a/Engine/Source/Developer/DerivedDataCache/Private/MemoryDerivedDataBackend.cpp b/Engine/Source/Developer/DerivedDataCache/Private/MemoryDerivedDataBackend.cpp
--- a/Engine/Source/Developer/DerivedDataCache/Private/MemoryDerivedDataBackend.cpp
+++ b/Engine/Source/Developer/DerivedDataCache/Private/MemoryDerivedDataBackend.cpp
@@ -3,17 +3,19 @@
 #include "MemoryDerivedDataBackend.h"
 #include "Templates/UniquePtr.h"
 
-FMemoryDerivedDataBackend::FMemoryDerivedDataBackend(const TCHAR* InName, int64 InMaxCacheSize)
+FMemoryDerivedDataBackend::FMemoryDerivedDataBackend(const TCHAR* InName, int64 InMaxCacheSize, bool bInCanBeDisabled)
 	: Name(InName)
 	, MaxCacheSize(InMaxCacheSize)
 	, bDisabled( false )
 	, CurrentCacheSize( SerializationSpecificDataSize )
 	, bMaxSizeExceeded(false)
+	, bCanBeDisabled(bInCanBeDisabled)
 {
 }
 
 FMemoryDerivedDataBackend::~FMemoryDerivedDataBackend()
 {
+	bShuttingDown = true;
 	Disable();
 }
 
@@ -30,6 +32,12 @@ FDerivedDataBackendInterface::ESpeedClass FMemoryDerivedDataBackend::GetSpeedCla
 
 bool FMemoryDerivedDataBackend::CachedDataProbablyExists(const TCHAR* CacheKey)
 {
+	// See comments on the declaration of bCanBeDisabled variable.
+	if (bCanBeDisabled)
+	{
+		return false;
+	}
+
 	COOK_STAT(auto Timer = UsageStats.TimeProbablyExists());
 
 	if (ShouldSimulateMiss(CacheKey))
@@ -309,6 +317,7 @@ bool FMemoryDerivedDataBackend::LoadCache(const TCHAR* Filename)
 
 void FMemoryDerivedDataBackend::Disable()
 {
+	check(bCanBeDisabled || bShuttingDown);
 	FScopeLock ScopeLock(&SynchronizationObject);
 	bDisabled = true;
 	for (TMap<FString,FCacheValue*>::TIterator It(CacheItems); It; ++It )
diff --git a/Engine/Source/Developer/DerivedDataCache/Private/MemoryDerivedDataBackend.h b/Engine/Source/Developer/DerivedDataCache/Private/MemoryDerivedDataBackend.h
--- a/Engine/Source/Developer/DerivedDataCache/Private/MemoryDerivedDataBackend.h
+++ b/Engine/Source/Developer/DerivedDataCache/Private/MemoryDerivedDataBackend.h
@@ -17,7 +17,7 @@ class Error;
 class FMemoryDerivedDataBackend : public FFileBackedDerivedDataBackend
 {
 public:
-	explicit FMemoryDerivedDataBackend(const TCHAR* InName, int64 InMaxCacheSize = -1);
+	explicit FMemoryDerivedDataBackend(const TCHAR* InName, int64 InMaxCacheSize = -1, bool bCanBeDisabled = false);
 	~FMemoryDerivedDataBackend();
 
 	/** Return a name for this interface */
@@ -128,6 +128,26 @@ private:
 	/** Indicates that the cache max size has been exceeded. This is used to avoid
 		warning spam after the size has reached the limit. */
 	bool bMaxSizeExceeded;
+
+	/** When a memory cache can be disabled, it won't return true for CachedDataProbablyExists calls.
+	  * This is to avoid having the Boot DDC tells it has some resources that will suddenly disappear after the boot.
+	  * Get() requests will still get fulfilled and other cache level will be properly back-filled 
+	  * offering the speed benefit of the boot cache while maintaining coherency at all cache levels.
+	  * 
+	  * The problem is that most asset types (audio/staticmesh/texture) will always verify if their different LODS/Chunks can be found in the cache using CachedDataProbablyExists.
+	  * If any of the LOD/MIP can't be found, a build of the asset is triggered, otherwise they skip asset compilation altogether.
+	  * However, we should not skip the compilation based on the CachedDataProbablyExists result of the boot cache because it is a lie and will disappear at some point.
+	  * When the boot cache disappears and the streamer tries to fetch a LOD that it has been told was cached, it will fail and will then have no choice but to rebuild the asset synchronously.
+	  * This obviously causes heavy game-thread stutters.
+
+	  * However, if the bootcache returns false during CachedDataProbablyExists. The async compilation will be triggered and data will be put in the both the boot.ddc and the local cache.
+	  * This way, no more heavy game-thread stutters during streaming...
+
+	  * This can be reproed when you clear the local cache but do not clear the boot.ddc file, but even if it's a corner case, I stumbled upon it enough times that I though it was worth to fix so the caches are coherent.
+	  */
+	bool bCanBeDisabled = false;
+	bool bShuttingDown  = false;
+
 	enum 
 	{
 		/** Magic number to use in header */
