Revert the reentrancy guard from cl 14528803 which was uneeded and actually causing local ddc from being properly filled
Added profiling scope for Async DDC put

diff --git a/Engine/Source/Developer/DerivedDataCache/Private/DerivedDataBackendAsyncPutWrapper.cpp b/Engine/Source/Developer/DerivedDataCache/Private/DerivedDataBackendAsyncPutWrapper.cpp
--- a/Engine/Source/Developer/DerivedDataCache/Private/DerivedDataBackendAsyncPutWrapper.cpp
+++ b/Engine/Source/Developer/DerivedDataCache/Private/DerivedDataBackendAsyncPutWrapper.cpp
@@ -41,6 +41,7 @@ public:
 	/** Call the inner backend and when that completes, remove the memory cache */
 	void DoWork()
 	{
+		TRACE_CPUPROFILER_EVENT_SCOPE(DDCPut_DoWork);
 		COOK_STAT(auto Timer = UsageStats.TimePut());
 		bool bOk = true;
 		bool bDidTry = false;
@@ -166,9 +167,6 @@ bool FDerivedDataBackendAsyncPutWrapper::ApplyDebugOptions(FBackendDebugOptions&
 	return InnerBackend->ApplyDebugOptions(InOptions);
 }
 
-// TLS for async put reentrant guard
-static thread_local bool bReentrantAsyncPut = false;
-
 bool FDerivedDataBackendAsyncPutWrapper::GetCachedData(const TCHAR* CacheKey, TArray<uint8>& OutData)
 {
 	COOK_STAT(auto Timer = UsageStats.TimeGet());
@@ -179,8 +177,6 @@ bool FDerivedDataBackendAsyncPutWrapper::GetCachedData(const TCHAR* CacheKey, TA
 		return true;
 	}
 
-	// Async Put might be triggered from the inner backend, the reentrant guard will ignore async put request triggered from an async put
-	TGuardValue<bool> ReentrantPutGuard(bReentrantAsyncPut, true);
 	bool bSuccess = InnerBackend->GetCachedData(CacheKey, OutData);
 	if (bSuccess)
 	{
@@ -198,11 +194,6 @@ void FDerivedDataBackendAsyncPutWrapper::PutCachedData(const TCHAR* CacheKey, TA
 {
 	COOK_STAT(auto Timer = PutSyncUsageStats.TimePut());
 
-	// this is an async put triggered as a backfill from a get on another async put, just early out
-	if (bReentrantAsyncPut)
-	{
-		return;
-	}
 	if (!InnerBackend->IsWritable())
 	{
 		return; // no point in continuing down the chain
