Fix for Jupiter token expiring

For very long cooking sessions the Jupiter token could expire multiple times. That could cause the failed login counter to disable further login attempts. Fix this by resetting the counter on every successful login. Also reduced verbosity of 401 errors so that CI doesn't pick up expired tokens as errors.

diff --git a/Engine/Source/Developer/DerivedDataCache/Private/HttpDerivedDataBackend.cpp b/Engine/Source/Developer/DerivedDataCache/Private/HttpDerivedDataBackend.cpp
--- a/Engine/Source/Developer/DerivedDataCache/Private/HttpDerivedDataBackend.cpp
+++ b/Engine/Source/Developer/DerivedDataCache/Private/HttpDerivedDataBackend.cpp
@@ -473,15 +473,33 @@ private:
 				FString Response = GetAnsiBufferAsString(ResponseBuffer.Num() > 0 ? ResponseBuffer : ResponseHeader);
 				Response.ReplaceCharInline('\n', ' ');
 				Response.ReplaceCharInline('\r', ' ');
-				UE_LOG(
-					LogDerivedDataCache, 
-					Error, 
-					TEXT("Failed %s HTTP cache entry (response %d) from %s. Response: %s"),
-					VerbStr, 
-					ResponseCode, 
-					Uri,
-					*Response
-				);
+				// Dont log access denied as error, since tokens can expire mid session
+				if (ResponseCode == 401)
+				{
+					UE_LOG(
+						LogDerivedDataCache,
+						Verbose,
+						TEXT("Failed %s HTTP cache entry (response %d) from %s. Response: %s"),
+						VerbStr,
+						ResponseCode,
+						Uri,
+						*Response
+					);
+				}
+				else
+				{ 
+					UE_LOG(
+						LogDerivedDataCache,
+						Error,
+						TEXT("Failed %s HTTP cache entry (response %d) from %s. Response: %s"),
+						VerbStr,
+						ResponseCode,
+						Uri,
+						*Response
+					);
+				}
+
+
 			}
 		}
 		else if(bLogErrors)
@@ -995,6 +1013,7 @@ bool FHttpDerivedDataBackend::IsServiceReady()
 
 bool FHttpDerivedDataBackend::AcquireAccessToken()
 {
+	// Avoid spamming the this if the service is down
 	if (FailedLoginAttempts > UE_HTTPDDC_MAX_FAILED_LOGIN_ATTEMPTS)
 	{
 		return false;
@@ -1084,7 +1103,8 @@ bool FHttpDerivedDataBackend::AcquireAccessToken()
 							}
 						), ExpiryTimeSeconds - 20.0f);
 					}
-
+					// Reset failed login attempts, the service is indeed alive.
+					FailedLoginAttempts = 0;
 					return true;
 				}
 			}
