Bumped http request pool size to attempt to reduce amount of time we spend waiting on the pool when fetching from Jupiter

diff --git a/Engine/Source/Developer/DerivedDataCache/Private/HttpDerivedDataBackend.cpp b/Engine/Source/Developer/DerivedDataCache/Private/HttpDerivedDataBackend.cpp
--- a/Engine/Source/Developer/DerivedDataCache/Private/HttpDerivedDataBackend.cpp
+++ b/Engine/Source/Developer/DerivedDataCache/Private/HttpDerivedDataBackend.cpp
@@ -28,7 +28,7 @@
 #define UE_HTTPDDC_HTTP_REQUEST_TIMEOUT_SECONDS 30L
 #define UE_HTTPDDC_HTTP_REQUEST_TIMOUT_ENABLED 1
 #define UE_HTTPDDC_HTTP_DEBUG 0
-#define UE_HTTPDDC_REQUEST_POOL_SIZE 16
+#define UE_HTTPDDC_REQUEST_POOL_SIZE 64
 #define UE_HTTPDDC_MAX_FAILED_LOGIN_ATTEMPTS 16
 #define UE_HTTPDDC_MAX_ATTEMPTS 4
 #define UE_HTTPDDC_MAX_BUFFER_RESERVE 104857600u
