Jupiter DDC head requests will no longer give errors if the response code is 204 (which is valid)

- The Io service currently returns response 204 when a HEAD request finds an entry, however the Europa service has been returning 200. The goal here is to update the DDC HEAD requests so that they will work with a response of 200 OR 204 so that we can unify the Jupiter apis.
- Added a utility method to FRequest that can be used to test if the response code is in the range that we consider a success.
- Updated FHttpDerivedDataBackend::CachedDataProbablyExists to use this.
- Updated the FRequest logging code to use this as well to avoid reporting successful operations as errors.

#rnx

diff --git a/Engine/Source/Developer/DerivedDataCache/Private/HttpDerivedDataBackend.cpp b/Engine/Source/Developer/DerivedDataCache/Private/HttpDerivedDataBackend.cpp
--- a/Engine/Source/Developer/DerivedDataCache/Private/HttpDerivedDataBackend.cpp
+++ b/Engine/Source/Developer/DerivedDataCache/Private/HttpDerivedDataBackend.cpp
@@ -333,6 +333,13 @@ public:
 		return JsonObject;
 	}
 
+	/** Will return true if the response code is considered a success */
+	static bool IsSuccessResponse(long ResponseCode)
+	{
+		// We consider anything in the 1XX or 2XX range a success
+		return ResponseCode >= 100 && ResponseCode < 300;
+	}
+
 private:
 
 	CURL*					Curl;
@@ -431,26 +438,26 @@ private:
 			switch (Verb)
 			{
 			case Head:
-				bSuccess = (ResponseCode == 400 || ResponseCode == 200);
+				bSuccess = (ResponseCode == 400 || IsSuccessResponse(ResponseCode));
 				VerbStr = TEXT("querying");
 				break;
 			case Get:
-				bSuccess = (ResponseCode == 400 || ResponseCode == 200);
+				bSuccess = (ResponseCode == 400 || IsSuccessResponse(ResponseCode));
 				VerbStr = TEXT("fetching");
 				AdditionalInfo = FString::Printf(TEXT("Received: %d bytes."), BytesReceived);
 				break;
 			case Put:
-				bSuccess = ResponseCode == 200;
+				bSuccess = IsSuccessResponse(ResponseCode);
 				VerbStr = TEXT("updating");
 				AdditionalInfo = FString::Printf(TEXT("Sent: %d bytes."), BytesSent);
 				break;
 			case Post:
 			case PostJson:
-				bSuccess = ResponseCode == 200;
+				bSuccess = IsSuccessResponse(ResponseCode);
 				VerbStr = TEXT("posting");
 				break;
 			case Delete:
-				bSuccess = ResponseCode == 200;
+				bSuccess = IsSuccessResponse(ResponseCode);
 				VerbStr = TEXT("deleting");
 				break;
 			}
@@ -1151,9 +1158,9 @@ bool FHttpDerivedDataBackend::CachedDataProbablyExists(const TCHAR* CacheKey)
 		FRequest::Result Result = Request->PerformBlockingQuery<FRequest::Head>(*Uri);
 		ResponseCode = Request->GetResponseCode();
 
-		if (ResponseCode == 200 || ResponseCode == 400)
+		if (FRequest::IsSuccessResponse(ResponseCode) || ResponseCode == 400)
 		{
-			const bool bIsHit = (Result == FRequest::Success && ResponseCode == 200);
+			const bool bIsHit = (Result == FRequest::Success && FRequest::IsSuccessResponse(ResponseCode));
 			if (bIsHit)
 			{
 				TRACE_COUNTER_ADD(HttpDDC_ExistHit, int64(1));
