Considering Jupiter backend as a fast backend, in the same way the SMB share is configured. Works around problematic assumptions in the hirarchal backend were it will not accept puts that primarilly live in slow backends (e.g. if no fast backends are available)

diff --git a/Engine/Source/Developer/DerivedDataCache/Private/HttpDerivedDataBackend.cpp b/Engine/Source/Developer/DerivedDataCache/Private/HttpDerivedDataBackend.cpp
--- a/Engine/Source/Developer/DerivedDataCache/Private/HttpDerivedDataBackend.cpp
+++ b/Engine/Source/Developer/DerivedDataCache/Private/HttpDerivedDataBackend.cpp
@@ -991,7 +991,7 @@ bool FHttpDerivedDataBackend::WouldCache(const TCHAR* CacheKey, TArrayView<const
 
 FHttpDerivedDataBackend::ESpeedClass FHttpDerivedDataBackend::GetSpeedClass()
 {
-	return ESpeedClass::Slow;
+	return ESpeedClass::Fast;
 }
 
 bool FHttpDerivedDataBackend::ApplyDebugOptions(FBackendDebugOptions& InOptions)
