DDC: Added variants of CachedDataProbablyExists that check multiple keys

diff --git a/Engine/Source/Developer/DerivedDataCache/Private/DerivedDataBackendAsyncPutWrapper.cpp b/Engine/Source/Developer/DerivedDataCache/Private/DerivedDataBackendAsyncPutWrapper.cpp
--- a/Engine/Source/Developer/DerivedDataCache/Private/DerivedDataBackendAsyncPutWrapper.cpp
+++ b/Engine/Source/Developer/DerivedDataCache/Private/DerivedDataBackendAsyncPutWrapper.cpp
@@ -121,12 +121,6 @@ FDerivedDataBackendInterface::ESpeedClass FDerivedDataBackendAsyncPutWrapper::Ge
 	return InnerBackend->GetSpeedClass();
 }
 
-/**
- * Synchronous test for the existence of a cache item
- *
- * @param	CacheKey	Alphanumeric+underscore key of this cache item
- * @return				true if the data probably will be found, this can't be guaranteed because of concurrency in the backends, corruption, etc
- */
 bool FDerivedDataBackendAsyncPutWrapper::CachedDataProbablyExists(const TCHAR* CacheKey)
 {
 	COOK_STAT(auto Timer = UsageStats.TimeProbablyExists());
@@ -137,6 +131,33 @@ bool FDerivedDataBackendAsyncPutWrapper::CachedDataProbablyExists(const TCHAR* C
 	return Result;
 }
 
+TBitArray<> FDerivedDataBackendAsyncPutWrapper::CachedDataProbablyExistsBatch(TConstArrayView<FString> CacheKeys)
+{
+	COOK_STAT(auto Timer = UsageStats.TimeProbablyExists());
+
+	TBitArray<> Result;
+	if (InflightCache)
+	{
+		Result = InflightCache->CachedDataProbablyExistsBatch(CacheKeys);
+		check(Result.Num() == CacheKeys.Num());
+		if (Result.CountSetBits() < CacheKeys.Num())
+		{
+			TBitArray<> InnerResult = InnerBackend->CachedDataProbablyExistsBatch(CacheKeys);
+			check(InnerResult.Num() == CacheKeys.Num());
+			Result.CombineWithBitwiseOR(InnerResult, EBitwiseOperatorFlags::MaintainSize);
+		}
+	}
+	else
+	{
+		Result = InnerBackend->CachedDataProbablyExistsBatch(CacheKeys);
+		check(Result.Num() == CacheKeys.Num());
+	}
+
+	COOK_STAT(if (Result.CountSetBits() == CacheKeys.Num()) { Timer.AddHit(0); });
+	UE_LOG(LogDerivedDataCache, Verbose, TEXT("%s CachedDataProbablyExists found %d/%d keys"), *GetName(), Result.CountSetBits(), CacheKeys.Num());
+	return Result;
+}
+
 bool FDerivedDataBackendAsyncPutWrapper::TryToPrefetch(const TCHAR* CacheKey)
 {
 	COOK_STAT(auto Timer = UsageStats.TimePrefetch());
diff --git a/Engine/Source/Developer/DerivedDataCache/Private/DerivedDataBackendAsyncPutWrapper.h b/Engine/Source/Developer/DerivedDataCache/Private/DerivedDataBackendAsyncPutWrapper.h
--- a/Engine/Source/Developer/DerivedDataCache/Private/DerivedDataBackendAsyncPutWrapper.h
+++ b/Engine/Source/Developer/DerivedDataCache/Private/DerivedDataBackendAsyncPutWrapper.h
@@ -83,6 +83,14 @@ public:
 	 */
 	virtual bool CachedDataProbablyExists(const TCHAR* CacheKey) override;
 
+	/**
+	 * Synchronous test for the existence of multiple cache items
+	 *
+	 * @param	CacheKeys	Alphanumeric+underscore key of the cache items
+	 * @return				A bit array with bits indicating whether the data for the corresponding key will probably be found
+	 */
+	virtual TBitArray<> CachedDataProbablyExistsBatch(TConstArrayView<FString> CacheKeys) override;
+
 	/**
 	 * Attempts to make sure the cached data will be available as optimally as possible. This is left up to the implementation to do
 	 * @param	CacheKey	Alphanumeric+underscore key of this cache item
diff --git a/Engine/Source/Developer/DerivedDataCache/Private/DerivedDataBackendCorruptionWrapper.h b/Engine/Source/Developer/DerivedDataCache/Private/DerivedDataBackendCorruptionWrapper.h
--- a/Engine/Source/Developer/DerivedDataCache/Private/DerivedDataBackendCorruptionWrapper.h
+++ b/Engine/Source/Developer/DerivedDataCache/Private/DerivedDataBackendCorruptionWrapper.h
@@ -115,6 +115,24 @@ public:
 		}
 		return Result;
 	}
+
+	/**
+	 * Synchronous test for the existence of multiple cache items
+	 *
+	 * @param	CacheKeys	Alphanumeric+underscore key of the cache items
+	 * @return				A bit array with bits indicating whether the data for the corresponding key will probably be found
+	 */
+	virtual TBitArray<> CachedDataProbablyExistsBatch(TConstArrayView<FString> CacheKeys) override
+	{
+		COOK_STAT(auto Timer = UsageStats.TimeProbablyExists());
+		TBitArray<> Result = InnerBackend->CachedDataProbablyExistsBatch(CacheKeys);
+		if (Result.CountSetBits() == CacheKeys.Num())
+		{
+			COOK_STAT(Timer.AddHit(0));
+		}
+		return Result;
+	}
+
 	/**
 	 * Synchronous retrieve of a cache item
 	 *
diff --git a/Engine/Source/Developer/DerivedDataCache/Private/DerivedDataCache.cpp b/Engine/Source/Developer/DerivedDataCache/Private/DerivedDataCache.cpp
--- a/Engine/Source/Developer/DerivedDataCache/Private/DerivedDataCache.cpp
+++ b/Engine/Source/Developer/DerivedDataCache/Private/DerivedDataCache.cpp
@@ -494,6 +494,29 @@ public:
 		return bResult;
 	}
 
+	virtual TBitArray<> CachedDataProbablyExistsBatch(TConstArrayView<FString> CacheKeys) override
+	{
+		DDC_SCOPE_CYCLE_COUNTER(DDC_CachedDataProbablyExistsBatch);
+		TBitArray<> Result;
+		if (CacheKeys.Num())
+		{
+			INC_DWORD_STAT(STAT_DDC_NumExist);
+			STAT(double ThisTime = 0);
+			{
+				SCOPE_SECONDS_COUNTER(ThisTime);
+				Result = FDerivedDataBackend::Get().GetRoot().CachedDataProbablyExistsBatch(CacheKeys);
+				check(Result.Num() == CacheKeys.Num());
+			}
+			INC_FLOAT_STAT_BY(STAT_DDC_ExistTime, (float)ThisTime);
+		}
+		return Result;
+	}
+
+	virtual bool AllCachedDataProbablyExists(TConstArrayView<FString> CacheKeys) override
+	{
+		return CacheKeys.Num() == 0 || CachedDataProbablyExistsBatch(CacheKeys).CountSetBits() == CacheKeys.Num();
+	}
+
 	void NotifyBootComplete() override
 	{
 		DDC_SCOPE_CYCLE_COUNTER(DDC_NotifyBootComplete);
diff --git a/Engine/Source/Developer/DerivedDataCache/Private/DerivedDataLimitKeyLengthWrapper.h b/Engine/Source/Developer/DerivedDataCache/Private/DerivedDataLimitKeyLengthWrapper.h
--- a/Engine/Source/Developer/DerivedDataCache/Private/DerivedDataLimitKeyLengthWrapper.h
+++ b/Engine/Source/Developer/DerivedDataCache/Private/DerivedDataLimitKeyLengthWrapper.h
@@ -64,6 +64,29 @@ public:
 		return Result;
 	}
 
+	/**
+	 * Synchronous test for the existence of multiple cache items
+	 *
+	 * @param	CacheKeys	Alphanumeric+underscore key of the cache items
+	 * @return				A bit array with bits indicating whether the data for the corresponding key will probably be found
+	 */
+	virtual TBitArray<> CachedDataProbablyExistsBatch(TConstArrayView<FString> CacheKeys) override
+	{
+		COOK_STAT(auto Timer = UsageStats.TimeProbablyExists());
+		TArray<FString> NewKeys;
+		NewKeys.Reserve(CacheKeys.Num());
+		for (const FString& CacheKey : CacheKeys)
+		{
+			ShortenKey(*CacheKey, NewKeys.Emplace_GetRef());
+		}
+		TBitArray<> Result = InnerBackend->CachedDataProbablyExistsBatch(NewKeys);
+		if (Result.CountSetBits() == CacheKeys.Num())
+		{
+			COOK_STAT(Timer.AddHit(0));
+		}
+		return Result;
+	}
+
 	/**
 	 * Attempts to make sure the cached data will be available as optimally as possible. This is left up to the implementation to do
 	 * @param	CacheKey	Alphanumeric+underscore key of this cache item
diff --git a/Engine/Source/Developer/DerivedDataCache/Private/HierarchicalDerivedDataBackend.h b/Engine/Source/Developer/DerivedDataCache/Private/HierarchicalDerivedDataBackend.h
--- a/Engine/Source/Developer/DerivedDataCache/Private/HierarchicalDerivedDataBackend.h
+++ b/Engine/Source/Developer/DerivedDataCache/Private/HierarchicalDerivedDataBackend.h
@@ -122,6 +122,81 @@ public:
 		return false;
 	}
 
+	/**
+	 * Synchronous test for the existence of multiple cache items
+	 *
+	 * @param	CacheKeys	Alphanumeric+underscore key of the cache items
+	 * @return				A bit array with bits indicating whether the data for the corresponding key will probably be found
+	 */
+	virtual TBitArray<> CachedDataProbablyExistsBatch(TConstArrayView<FString> CacheKeys) override
+	{
+		COOK_STAT(auto Timer = UsageStats.TimeProbablyExists());
+		check(InnerBackends.Num() > 0);
+
+		TBitArray<> Result;
+
+		for (int32 CacheIndex = 0; CacheIndex < InnerBackends.Num(); ++CacheIndex)
+		{
+			const int32 MissingKeys = CacheKeys.Num() - Result.CountSetBits();
+			if (MissingKeys == 0)
+			{
+				break;
+			}
+
+			// Skip slow caches because the primary users of this function assume that
+			// they will have fast access to the data if this returns true. It will be
+			// better in those cases to rebuild the data locally than to block on slow
+			// fetch operations because the build may be, and often is, asynchronous.
+			bool bFastCache = InnerBackends[CacheIndex]->GetSpeedClass() >= ESpeedClass::Fast;
+			if (!bFastCache)
+			{
+				continue;
+			}
+
+			if (MissingKeys == CacheKeys.Num())
+			{
+				Result = InnerBackends[CacheIndex]->CachedDataProbablyExistsBatch(CacheKeys);
+				check(Result.Num() == CacheKeys.Num());
+			}
+			else
+			{
+				TArray<FString> RemainingKeys;
+				for (int32 KeyIndex = 0; KeyIndex < CacheKeys.Num(); ++KeyIndex)
+				{
+					if (!Result[KeyIndex])
+					{
+						RemainingKeys.Add(CacheKeys[KeyIndex]);
+					}
+				}
+
+				TBitArray<>::FIterator ResultIt(Result);
+				TBitArray<> NewResult = InnerBackends[CacheIndex]->CachedDataProbablyExistsBatch(RemainingKeys);
+				check(NewResult.Num() == RemainingKeys.Num());
+				for (TBitArray<>::FConstIterator NewIt(NewResult); NewIt; ++NewIt)
+				{
+					// Skip to the bit for the next remaining key.
+					while (ResultIt.GetValue())
+					{
+						++ResultIt;
+					}
+					ResultIt.GetValue() = NewIt.GetValue();
+					++ResultIt;
+				}
+			}
+		}
+
+		if (Result.IsEmpty())
+		{
+			Result.Add(false, CacheKeys.Num());
+		}
+
+		if (Result.CountSetBits() == CacheKeys.Num())
+		{
+			COOK_STAT(Timer.AddHit(0));
+		}
+		return Result;
+	}
+
 	/**
 	 * Attempts to make sure the cached data will be available as optimally as possible. This is left up to the implementation to do
 	 * @param	CacheKey	Alphanumeric+underscore key of this cache item
diff --git a/Engine/Source/Developer/DerivedDataCache/Public/DerivedDataBackendInterface.h b/Engine/Source/Developer/DerivedDataCache/Public/DerivedDataBackendInterface.h
--- a/Engine/Source/Developer/DerivedDataCache/Public/DerivedDataBackendInterface.h
+++ b/Engine/Source/Developer/DerivedDataCache/Public/DerivedDataBackendInterface.h
@@ -4,6 +4,7 @@
 #pragma once
 
 #include "CoreMinimal.h"
+#include "Containers/BitArray.h"
 #include "Stats/Stats.h"
 
 class FDerivedDataCacheUsageStats;
@@ -100,6 +101,24 @@ public:
 	 * @return				true if the data probably will be found, this can't be guaranteed because of concurrency in the backends, corruption, etc
 	 */
 	virtual bool CachedDataProbablyExists(const TCHAR* CacheKey)=0;
+
+	/**
+	 * Synchronous test for the existence of multiple cache items
+	 *
+	 * @param	CacheKeys	Alphanumeric+underscore key of the cache items
+	 * @return				A bit array with bits indicating whether the data for the corresponding key will probably be found
+	 */
+	virtual TBitArray<> CachedDataProbablyExistsBatch(TConstArrayView<FString> CacheKeys)
+	{
+		TBitArray<> Result;
+		Result.Reserve(CacheKeys.Num());
+		for (const FString& Key : CacheKeys)
+		{
+			Result.Add(CachedDataProbablyExists(*Key));
+		}
+		return Result;
+	}
+
 	/**
 	 * Synchronous retrieve of a cache item
 	 *
diff --git a/Engine/Source/Developer/DerivedDataCache/Public/DerivedDataCacheInterface.h b/Engine/Source/Developer/DerivedDataCache/Public/DerivedDataCacheInterface.h
--- a/Engine/Source/Developer/DerivedDataCache/Public/DerivedDataCacheInterface.h
+++ b/Engine/Source/Developer/DerivedDataCache/Public/DerivedDataCacheInterface.h
@@ -4,6 +4,7 @@
 #pragma once
 
 #include "CoreMinimal.h"
+#include "Containers/BitArray.h"
 #include "Containers/StringView.h"
 #include "Modules/ModuleInterface.h"
 
@@ -195,6 +196,21 @@ public:
 	 */
 	virtual bool CachedDataProbablyExists(const TCHAR* CacheKey) = 0;
 
+	/**
+	 * Returns whether the data associated with each key is likely to exist in the cache.
+	 * Even if this function returns true, a get for one of the keys may still fail!
+	 * @param	CacheKeys	Keys to see if data probably exists.
+	 * @return				A bit array with bits indicating whether the data for the corresponding key will probably be found
+	 */
+	virtual TBitArray<> CachedDataProbablyExistsBatch(TConstArrayView<FString> CacheKeys) = 0;
+
+	/**
+	 * Returns true if the data associated with each key is likely to exist in the cache.
+	 * Even if this function returns true, a get for one of the keys may still fail!
+	 * @param	CacheKeys	Keys to see if data probably exists.
+	 */
+	virtual bool AllCachedDataProbablyExists(TConstArrayView<FString> CacheKeys) = 0;
+
 	//--------------------
 	// System Interface
 	//--------------------
