DDC: Return a status from PutCachedData

The async put backend now uses this new status to avoid calling CachedDataProbablyExists after every put, which has a high cost for certain backends.

diff --git a/Engine/Source/Developer/DerivedDataCache/Private/DerivedDataBackendAsyncPutWrapper.cpp b/Engine/Source/Developer/DerivedDataCache/Private/DerivedDataBackendAsyncPutWrapper.cpp
--- a/Engine/Source/Developer/DerivedDataCache/Private/DerivedDataBackendAsyncPutWrapper.cpp
+++ b/Engine/Source/Developer/DerivedDataCache/Private/DerivedDataBackendAsyncPutWrapper.cpp
@@ -43,37 +43,61 @@ public:
 	{
 		TRACE_CPUPROFILER_EVENT_SCOPE(DDCPut_DoWork);
 		COOK_STAT(auto Timer = UsageStats.TimePut());
-		bool bOk = true;
-		bool bDidTry = false;
-		if (bPutEvenIfExists || !InnerBackend->CachedDataProbablyExists(*CacheKey))
+
+		using EPutStatus = FDerivedDataBackendInterface::EPutStatus;
+		EPutStatus Status = EPutStatus::NotCached;
+
+		if (!bPutEvenIfExists && InnerBackend->CachedDataProbablyExists(*CacheKey))
+		{
+			Status = EPutStatus::Cached;
+		}
+		else
 		{
-			bDidTry = true;
-			InnerBackend->PutCachedData(*CacheKey, Data, bPutEvenIfExists);
+			Status = InnerBackend->PutCachedData(*CacheKey, Data, bPutEvenIfExists);
 			COOK_STAT(Timer.AddHit(Data.Num()));
 		}
-		// if we tried to put it up there but it isn't there now, retry
-		if (InflightCache && bDidTry && !InnerBackend->CachedDataProbablyExists(*CacheKey))
+
+		if (InflightCache)
 		{
-			// retry after a brief wait
-			FPlatformProcess::SleepNoStats(0.2f);
-			InnerBackend->PutCachedData(*CacheKey, Data, false);
+			// if the data was not cached synchronously, retry
+			if (Status != EPutStatus::Cached)
+			{
+				// retry after a brief wait
+				FPlatformProcess::SleepNoStats(0.2f);
 
-			if (!InnerBackend->CachedDataProbablyExists(*CacheKey))
+				if (Status == EPutStatus::Executing && InnerBackend->CachedDataProbablyExists(*CacheKey))
+				{
+					Status = EPutStatus::Cached;
+				}
+				else
+				{
+					Status = InnerBackend->PutCachedData(*CacheKey, Data, /*bPutEvenIfExists*/ false);
+				}
+			}
+
+			switch (Status)
 			{
-				UE_LOG(LogDerivedDataCache, Display, TEXT("%s: Put failed, keeping in memory copy %s."),*InnerBackend->GetName(), *CacheKey);
-
-				// log the filesystem error
-				uint32 ErrorCode = FPlatformMisc::GetLastError();
-				TCHAR ErrorBuffer[1024];
-				FPlatformMisc::GetSystemErrorMessage(ErrorBuffer, 1024, ErrorCode);
-				UE_LOG(LogDerivedDataCache, Display, TEXT("Failed to write %s to %s. Error: %u (%s)"), *CacheKey, *InnerBackend->GetName(), ErrorCode, ErrorBuffer);
-				bOk = false;
+			case EPutStatus::Cached:
+				// remove this from the in-flight cache because the inner cache contains the data
+				InflightCache->RemoveCachedData(*CacheKey, /*bTransient*/ false);
+				break;
+			case EPutStatus::NotCached:
+				UE_LOG(LogDerivedDataCache, Display, TEXT("%s: Put failed, keeping in memory copy %s."), *InnerBackend->GetName(), *CacheKey);
+				if (uint32 ErrorCode = FPlatformMisc::GetLastError())
+				{
+					TCHAR ErrorBuffer[1024];
+					FPlatformMisc::GetSystemErrorMessage(ErrorBuffer, 1024, ErrorCode);
+					UE_LOG(LogDerivedDataCache, Display, TEXT("Failed to write %s to %s. Error: %u (%s)"), *CacheKey, *InnerBackend->GetName(), ErrorCode, ErrorBuffer);
+				}
+				break;
+			case EPutStatus::Executing:
+				UE_LOG(LogDerivedDataCache, Display, TEXT("%s: Put not finished executing, keeping in memory copy %s."), *InnerBackend->GetName(), *CacheKey);
+				break;
+			default:
+				break;
 			}
 		}
-		if (bOk && InflightCache)
-		{
-			InflightCache->RemoveCachedData(*CacheKey, /*bTransient=*/ false); // we can remove this from the temp cache, since the real cache will hit now
-		}
+
 		FilesInFlight->Remove(CacheKey);
 		FDerivedDataBackend::Get().AddToAsyncCompletionCounter(-1);
 		UE_LOG(LogDerivedDataCache, VeryVerbose, TEXT("%s: Completed AsyncPut of %s."), *InnerBackend->GetName(), *CacheKey);
@@ -211,25 +235,25 @@ bool FDerivedDataBackendAsyncPutWrapper::GetCachedData(const TCHAR* CacheKey, TA
 	return bSuccess;
 }
 
-void FDerivedDataBackendAsyncPutWrapper::PutCachedData(const TCHAR* CacheKey, TArrayView<const uint8> InData, bool bPutEvenIfExists)
+FDerivedDataBackendInterface::EPutStatus FDerivedDataBackendAsyncPutWrapper::PutCachedData(const TCHAR* CacheKey, TArrayView<const uint8> InData, bool bPutEvenIfExists)
 {
 	COOK_STAT(auto Timer = PutSyncUsageStats.TimePut());
 
 	if (!InnerBackend->IsWritable())
 	{
-		return; // no point in continuing down the chain
+		return EPutStatus::NotCached; // no point in continuing down the chain
 	}
 	const bool bAdded = FilesInFlight.AddIfNotExists(CacheKey);
 	if (!bAdded)
 	{
-		return; // if it is already on its way, we don't need to send it again
+		return EPutStatus::Executing; // if it is already on its way, we don't need to send it again
 	}
 	if (InflightCache)
 	{
 		if (InflightCache->CachedDataProbablyExists(CacheKey))
 		{
 			UE_LOG(LogDerivedDataCache, Verbose, TEXT("%s skipping out of key already in in-flight cache %s"), *GetName(), CacheKey);
-			return; // if it is already on its way, we don't need to send it again
+			return EPutStatus::Executing; // if it is already on its way, we don't need to send it again
 		}
 		InflightCache->PutCachedData(CacheKey, InData, true); // temp copy stored in memory while the async task waits to complete
 		COOK_STAT(Timer.AddHit(InData.Num()));
@@ -239,6 +263,8 @@ void FDerivedDataBackendAsyncPutWrapper::PutCachedData(const TCHAR* CacheKey, TA
 
 	FDerivedDataBackend::Get().AddToAsyncCompletionCounter(1);
 	(new FAutoDeleteAsyncTask<FCachePutAsyncWorker>(CacheKey, InData, InnerBackend, bPutEvenIfExists, InflightCache.Get(), &FilesInFlight, UsageStats))->StartBackgroundTask();
+
+	return EPutStatus::Executing;
 }
 
 void FDerivedDataBackendAsyncPutWrapper::RemoveCachedData(const TCHAR* CacheKey, bool bTransient)
diff --git a/Engine/Source/Developer/DerivedDataCache/Private/DerivedDataBackendAsyncPutWrapper.h b/Engine/Source/Developer/DerivedDataCache/Private/DerivedDataBackendAsyncPutWrapper.h
--- a/Engine/Source/Developer/DerivedDataCache/Private/DerivedDataBackendAsyncPutWrapper.h
+++ b/Engine/Source/Developer/DerivedDataCache/Private/DerivedDataBackendAsyncPutWrapper.h
@@ -120,7 +120,7 @@ public:
 	 * @param	InData		Buffer containing the data to cache, can be destroyed after the call returns, immediately
 	 * @param	bPutEvenIfExists	If true, then do not attempt skip the put even if CachedDataProbablyExists returns true
 	 */
-	virtual void PutCachedData(const TCHAR* CacheKey, TArrayView<const uint8> InData, bool bPutEvenIfExists) override;
+	virtual EPutStatus PutCachedData(const TCHAR* CacheKey, TArrayView<const uint8> InData, bool bPutEvenIfExists) override;
 
 	virtual void RemoveCachedData(const TCHAR* CacheKey, bool bTransient) override;
 
diff --git a/Engine/Source/Developer/DerivedDataCache/Private/DerivedDataBackendCorruptionWrapper.h b/Engine/Source/Developer/DerivedDataCache/Private/DerivedDataBackendCorruptionWrapper.h
--- a/Engine/Source/Developer/DerivedDataCache/Private/DerivedDataBackendCorruptionWrapper.h
+++ b/Engine/Source/Developer/DerivedDataCache/Private/DerivedDataBackendCorruptionWrapper.h
@@ -190,12 +190,12 @@ public:
 	 * @param	InData		Buffer containing the data to cache, can be destroyed after the call returns, immediately
 	 * @param	bPutEvenIfExists	If true, then do not attempt skip the put even if CachedDataProbablyExists returns true
 	 */
-	virtual void PutCachedData(const TCHAR* CacheKey, TArrayView<const uint8> InData, bool bPutEvenIfExists) override
+	virtual EPutStatus PutCachedData(const TCHAR* CacheKey, TArrayView<const uint8> InData, bool bPutEvenIfExists) override
 	{
 		COOK_STAT(auto Timer = UsageStats.TimePut());
 		if (!InnerBackend->IsWritable())
 		{
-			return; // no point in continuing down the chain
+			return EPutStatus::NotCached; // no point in continuing down the chain
 		}
 		COOK_STAT(Timer.AddHit(InData.Num()));
 		// Get rid of the double copy!
@@ -206,7 +206,7 @@ public:
 		FDerivedDataTrailer Trailer(Data);
 		Data.AddUninitialized(sizeof(FDerivedDataTrailer));
 		FMemory::Memcpy(&Data[Data.Num() - sizeof(FDerivedDataTrailer)], &Trailer, sizeof(FDerivedDataTrailer));
-		InnerBackend->PutCachedData(CacheKey, Data, bPutEvenIfExists);
+		return InnerBackend->PutCachedData(CacheKey, Data, bPutEvenIfExists);
 	}
 	
 	virtual void RemoveCachedData(const TCHAR* CacheKey, bool bTransient) override
diff --git a/Engine/Source/Developer/DerivedDataCache/Private/DerivedDataBackendVerifyWrapper.h b/Engine/Source/Developer/DerivedDataCache/Private/DerivedDataBackendVerifyWrapper.h
--- a/Engine/Source/Developer/DerivedDataCache/Private/DerivedDataBackendVerifyWrapper.h
+++ b/Engine/Source/Developer/DerivedDataCache/Private/DerivedDataBackendVerifyWrapper.h
@@ -83,51 +83,46 @@ public:
 		return false;
 	}
 
-	virtual void PutCachedData(const TCHAR* CacheKey, TArrayView<const uint8> InData, bool bPutEvenIfExists) override
+	virtual EPutStatus PutCachedData(const TCHAR* CacheKey, TArrayView<const uint8> InData, bool bPutEvenIfExists) override
 	{
 		COOK_STAT(auto Timer = UsageStats.TimePut());
-		bool bAlreadyTested = false;
 		{
 			FScopeLock ScopeLock(&SynchronizationObject);
 			if (AlreadyTested.Contains(FString(CacheKey)))
 			{
-				bAlreadyTested = true;
-			}
-			else
-			{
-				AlreadyTested.Add(FString(CacheKey));
+				return EPutStatus::Cached;
 			}
+			AlreadyTested.Add(FString(CacheKey));
 		}
-		if (!bAlreadyTested)
+
+		COOK_STAT(Timer.AddHit(InData.Num()));
+		TArray<uint8> OutData;
+		bool bSuccess = InnerBackend->GetCachedData(CacheKey, OutData);
+		if (bSuccess)
 		{
-			COOK_STAT(Timer.AddHit(InData.Num()));
-			TArray<uint8> OutData;
-			bool bSuccess = InnerBackend->GetCachedData(CacheKey, OutData);
-			if (bSuccess)
+			if (OutData != InData)
 			{
-				if (OutData != InData)
-				{
-					UE_LOG(LogDerivedDataCache, Error, TEXT("Verify: Cached data differs from newly generated data %s."), CacheKey);
-					FString Cache = FPaths::ProjectSavedDir() / TEXT("VerifyDDC") / CacheKey + TEXT(".fromcache");
-					FFileHelper::SaveArrayToFile(OutData, *Cache);
-					FString Verify = FPaths::ProjectSavedDir() / TEXT("VerifyDDC") / CacheKey + TEXT(".verify");;
-					FFileHelper::SaveArrayToFile(InData, *Verify);
-					if (bFixProblems)
-					{
-						UE_LOG(LogDerivedDataCache, Display, TEXT("Verify: Wrote newly generated data to the cache."), CacheKey);
-						InnerBackend->PutCachedData(CacheKey, InData, true);
-					}
-				}
-				else
+				UE_LOG(LogDerivedDataCache, Error, TEXT("Verify: Cached data differs from newly generated data %s."), CacheKey);
+				FString Cache = FPaths::ProjectSavedDir() / TEXT("VerifyDDC") / CacheKey + TEXT(".fromcache");
+				FFileHelper::SaveArrayToFile(OutData, *Cache);
+				FString Verify = FPaths::ProjectSavedDir() / TEXT("VerifyDDC") / CacheKey + TEXT(".verify");;
+				FFileHelper::SaveArrayToFile(InData, *Verify);
+				if (bFixProblems)
 				{
-					UE_LOG(LogDerivedDataCache, Log, TEXT("Verify: Cached data exists and matched %s."), CacheKey);
+					UE_LOG(LogDerivedDataCache, Display, TEXT("Verify: Wrote newly generated data to the cache."), CacheKey);
+					InnerBackend->PutCachedData(CacheKey, InData, true);
 				}
 			}
 			else
 			{
-				UE_LOG(LogDerivedDataCache, Warning, TEXT("Verify: Cached data didn't exist %s."), CacheKey);
-				InnerBackend->PutCachedData(CacheKey, InData, bPutEvenIfExists);
+				UE_LOG(LogDerivedDataCache, Log, TEXT("Verify: Cached data exists and matched %s."), CacheKey);
 			}
+			return EPutStatus::Cached;
+		}
+		else
+		{
+			UE_LOG(LogDerivedDataCache, Warning, TEXT("Verify: Cached data didn't exist %s."), CacheKey);
+			return InnerBackend->PutCachedData(CacheKey, InData, bPutEvenIfExists);
 		}
 	}
 	virtual void RemoveCachedData(const TCHAR* CacheKey, bool bTransient) override
diff --git a/Engine/Source/Developer/DerivedDataCache/Private/DerivedDataLimitKeyLengthWrapper.h b/Engine/Source/Developer/DerivedDataCache/Private/DerivedDataLimitKeyLengthWrapper.h
--- a/Engine/Source/Developer/DerivedDataCache/Private/DerivedDataLimitKeyLengthWrapper.h
+++ b/Engine/Source/Developer/DerivedDataCache/Private/DerivedDataLimitKeyLengthWrapper.h
@@ -187,19 +187,18 @@ public:
 	 * @param	InData		Buffer containing the data to cache, can be destroyed after the call returns, immediately
 	 * @param	bPutEvenIfExists	If true, then do not attempt skip the put even if CachedDataProbablyExists returns true
 	 */
-	virtual void PutCachedData(const TCHAR* CacheKey, TArrayView<const uint8> InData, bool bPutEvenIfExists) override
+	virtual EPutStatus PutCachedData(const TCHAR* CacheKey, TArrayView<const uint8> InData, bool bPutEvenIfExists) override
 	{
 		COOK_STAT(auto Timer = UsageStats.TimePut());
 		if (!InnerBackend->IsWritable())
 		{
-			return; // no point in continuing down the chain
+			return EPutStatus::NotCached; // no point in continuing down the chain
 		}
 		COOK_STAT(Timer.AddHit(InData.Num()));
 		FString NewKey;
 		if (!ShortenKey(CacheKey, NewKey))
 		{
-			InnerBackend->PutCachedData(CacheKey, InData, bPutEvenIfExists);
-			return;
+			return InnerBackend->PutCachedData(CacheKey, InData, bPutEvenIfExists);
 		}
 		TArray<uint8> Data(InData.GetData(), InData.Num());
 		check(Data.Num());
@@ -207,7 +206,7 @@ public:
 		Data.AddUninitialized(KeyLen);
 		FCStringAnsi::Strcpy((char*)&Data[Data.Num() - KeyLen], KeyLen, TCHAR_TO_ANSI(CacheKey));
 		check(Data.Last()==0);
-		InnerBackend->PutCachedData(*NewKey, Data, bPutEvenIfExists);
+		return InnerBackend->PutCachedData(*NewKey, Data, bPutEvenIfExists);
 	}
 
 	virtual void RemoveCachedData(const TCHAR* CacheKey, bool bTransient) override
diff --git a/Engine/Source/Developer/DerivedDataCache/Private/FileSystemDerivedDataBackend.cpp b/Engine/Source/Developer/DerivedDataCache/Private/FileSystemDerivedDataBackend.cpp
--- a/Engine/Source/Developer/DerivedDataCache/Private/FileSystemDerivedDataBackend.cpp
+++ b/Engine/Source/Developer/DerivedDataCache/Private/FileSystemDerivedDataBackend.cpp
@@ -607,7 +607,7 @@ public:
 	 * @param	OutData		Buffer containing the data to cache, can be destroyed after the call returns, immediately
 	 * @param	bPutEvenIfExists	If true, then do not attempt skip the put even if CachedDataProbablyExists returns true
 	 */
-	virtual void PutCachedData(const TCHAR* CacheKey, TArrayView<const uint8> Data, bool bPutEvenIfExists) override
+	virtual EPutStatus PutCachedData(const TCHAR* CacheKey, TArrayView<const uint8> Data, bool bPutEvenIfExists) override
 	{
 		COOK_STAT(auto Timer = UsageStats.TimePut());
 		check(IsUsable());
@@ -622,9 +622,11 @@ public:
 			// don't put anything we pretended didn't exist
 			if (ShouldSimulateMiss(CacheKey))
 			{
-				return;
+				return EPutStatus::NotCached;
 			}
 
+			EPutStatus Status = EPutStatus::NotCached;
+
 			if (bPutEvenIfExists || !CachedDataProbablyExists(CacheKey))
 			{
 				COOK_STAT(Timer.AddHit(Data.Num()));
@@ -656,8 +658,13 @@ public:
 							else
 							{
 								UE_LOG(LogDerivedDataCache, Verbose, TEXT("%s: Successful cache put of %s to %s"),*GetName(), CacheKey, *Filename);
+								Status = EPutStatus::Cached;
 							}
 						}
+						else
+						{
+							Status = EPutStatus::Cached;
+						}
 					}
 					else
 					{
@@ -681,6 +688,7 @@ public:
 			{
 				COOK_STAT(Timer.AddMiss(Data.Num()));
 				UE_LOG(LogDerivedDataCache, Verbose, TEXT("%s skipping put to existing file %s"), *GetName(), CacheKey);
+				Status = EPutStatus::Cached;
 			}
 
 			// If not using a shared cache, update estimated build time
@@ -720,10 +728,13 @@ public:
 				}
 
 			}
+
+			return Status;
 		}
 		else
 		{
 			UE_LOG(LogDerivedDataCache, Verbose, TEXT("%s is read only. Skipping put of %s"), *GetName(), CacheKey);
+			return EPutStatus::NotCached;
 		}
 	}
 
diff --git a/Engine/Source/Developer/DerivedDataCache/Private/HierarchicalDerivedDataBackend.h b/Engine/Source/Developer/DerivedDataCache/Private/HierarchicalDerivedDataBackend.h
--- a/Engine/Source/Developer/DerivedDataCache/Private/HierarchicalDerivedDataBackend.h
+++ b/Engine/Source/Developer/DerivedDataCache/Private/HierarchicalDerivedDataBackend.h
@@ -377,14 +377,14 @@ public:
 	 * @param	InData		Buffer containing the data to cache, can be destroyed after the call returns, immediately
 	 * @param	bPutEvenIfExists	If true, then do not attempt skip the put even if CachedDataProbablyExists returns true
 	 */
-	virtual void PutCachedData(const TCHAR* CacheKey, TArrayView<const uint8> InData, bool bPutEvenIfExists) override
+	virtual EPutStatus PutCachedData(const TCHAR* CacheKey, TArrayView<const uint8> InData, bool bPutEvenIfExists) override
 	{
 		COOK_STAT(auto Timer = UsageStats.TimePut());
+		EPutStatus Status = EPutStatus::NotCached;
 		if (!bIsWritable)
 		{
-			return; // no point in continuing down the chain
+			return Status; // no point in continuing down the chain
 		}
-		bool bSynchronousPutPeformed = false;  // we must do at least one synchronous put to a writable cache before we return
 		for (int32 PutCacheIndex = 0; PutCacheIndex < InnerBackends.Num(); PutCacheIndex++)
 		{
 			if (!InnerBackends[PutCacheIndex]->IsWritable() && !InnerBackends[PutCacheIndex]->BackfillLowerCacheLevels() && InnerBackends[PutCacheIndex]->CachedDataProbablyExists(CacheKey))
@@ -394,10 +394,10 @@ public:
 			if (InnerBackends[PutCacheIndex]->IsWritable())
 			{
 				COOK_STAT(Timer.AddHit(InData.Num()));
-				if (!bSynchronousPutPeformed)
+				// we must do at least one synchronous put to a writable cache before we return
+				if (Status != EPutStatus::Cached)
 				{
-					InnerBackends[PutCacheIndex]->PutCachedData(CacheKey, InData, bPutEvenIfExists);
-					bSynchronousPutPeformed = true;
+					Status = InnerBackends[PutCacheIndex]->PutCachedData(CacheKey, InData, bPutEvenIfExists);
 				}
 				else
 				{
@@ -405,6 +405,7 @@ public:
 				}
 			}
 		}
+		return Status;
 	}
 
 	virtual void RemoveCachedData(const TCHAR* CacheKey, bool bTransient) override
diff --git a/Engine/Source/Developer/DerivedDataCache/Private/HttpDerivedDataBackend.cpp b/Engine/Source/Developer/DerivedDataCache/Private/HttpDerivedDataBackend.cpp
--- a/Engine/Source/Developer/DerivedDataCache/Private/HttpDerivedDataBackend.cpp
+++ b/Engine/Source/Developer/DerivedDataCache/Private/HttpDerivedDataBackend.cpp
@@ -1218,7 +1218,7 @@ bool FHttpDerivedDataBackend::GetCachedData(const TCHAR* CacheKey, TArray<uint8>
 	return false;
 }
 
-void FHttpDerivedDataBackend::PutCachedData(const TCHAR* CacheKey, TArrayView<const uint8> InData, bool bPutEvenIfExists)
+FDerivedDataBackendInterface::EPutStatus FHttpDerivedDataBackend::PutCachedData(const TCHAR* CacheKey, TArrayView<const uint8> InData, bool bPutEvenIfExists)
 {
 	TRACE_CPUPROFILER_EVENT_SCOPE(HttpDDC_Put);
 	COOK_STAT(auto Timer = UsageStats.TimePut());
@@ -1242,17 +1242,19 @@ void FHttpDerivedDataBackend::PutCachedData(const TCHAR* CacheKey, TArrayView<co
 			{
 				TRACE_COUNTER_ADD(HttpDDC_BytesSent, int64(Request->GetBytesSent()));
 				COOK_STAT(Timer.AddHit(Request->GetBytesSent()));
-				return;
+				return EPutStatus::Cached;
 			}
 
 			if (!ShouldRetryOnError(ResponseCode))
 			{
-				return;
+				return EPutStatus::NotCached;
 			}
 
 			ResponseCode = 0;
 		}
 	}
+
+	return EPutStatus::NotCached;
 }
 
 void FHttpDerivedDataBackend::RemoveCachedData(const TCHAR* CacheKey, bool bTransient)
diff --git a/Engine/Source/Developer/DerivedDataCache/Private/HttpDerivedDataBackend.h b/Engine/Source/Developer/DerivedDataCache/Private/HttpDerivedDataBackend.h
--- a/Engine/Source/Developer/DerivedDataCache/Private/HttpDerivedDataBackend.h
+++ b/Engine/Source/Developer/DerivedDataCache/Private/HttpDerivedDataBackend.h
@@ -52,7 +52,7 @@ public:
 	virtual bool IsWritable() override { return true; }
 	virtual bool CachedDataProbablyExists(const TCHAR* CacheKey) override;
 	virtual bool GetCachedData(const TCHAR* CacheKey, TArray<uint8>& OutData) override;
-	virtual void PutCachedData(const TCHAR* CacheKey, TArrayView<const uint8> InData, bool bPutEvenIfExists) override;
+	virtual EPutStatus PutCachedData(const TCHAR* CacheKey, TArrayView<const uint8> InData, bool bPutEvenIfExists) override;
 	virtual void RemoveCachedData(const TCHAR* CacheKey, bool bTransient) override;
 	virtual void GatherUsageStats(TMap<FString, FDerivedDataCacheUsageStats>& UsageStatsMap, FString&& GraphPath) override;
 	
diff --git a/Engine/Source/Developer/DerivedDataCache/Private/MemoryDerivedDataBackend.cpp b/Engine/Source/Developer/DerivedDataCache/Private/MemoryDerivedDataBackend.cpp
--- a/Engine/Source/Developer/DerivedDataCache/Private/MemoryDerivedDataBackend.cpp
+++ b/Engine/Source/Developer/DerivedDataCache/Private/MemoryDerivedDataBackend.cpp
@@ -101,21 +101,21 @@ bool FMemoryDerivedDataBackend::WouldCache(const TCHAR* CacheKey, TArrayView<con
 	return true;
 }
 
-void FMemoryDerivedDataBackend::PutCachedData(const TCHAR* CacheKey, TArrayView<const uint8> InData, bool bPutEvenIfExists)
+FDerivedDataBackendInterface::EPutStatus FMemoryDerivedDataBackend::PutCachedData(const TCHAR* CacheKey, TArrayView<const uint8> InData, bool bPutEvenIfExists)
 {
 	COOK_STAT(auto Timer = UsageStats.TimePut());
 	FScopeLock ScopeLock(&SynchronizationObject);
 
 	if (DidSimulateMiss(CacheKey))
 	{
-		return;
+		return EPutStatus::NotCached;
 	}
 	
 	// Should never hit this as higher level code should be checking..
 	if (!WouldCache(CacheKey, InData))
 	{
 		//UE_LOG(LogDerivedDataCache, Warning, TEXT("WouldCache was not called prior to attempted Put!"));
-		return;
+		return EPutStatus::NotCached;
 	}
 	
 	FString Key(CacheKey);
@@ -123,6 +123,7 @@ void FMemoryDerivedDataBackend::PutCachedData(const TCHAR* CacheKey, TArrayView<
 	if (Item)
 	{
 		//check(Item->Data == InData); // any second attempt to push data should be identical data
+		return EPutStatus::Cached;
 	}
 	else
 	{
@@ -135,6 +136,7 @@ void FMemoryDerivedDataBackend::PutCachedData(const TCHAR* CacheKey, TArrayView<
 			delete Val;
 			UE_LOG(LogDerivedDataCache, Display, TEXT("Failed to cache data. Maximum cache size reached. CurrentSize %d kb / MaxSize: %d kb"), CurrentCacheSize / 1024, MaxCacheSize / 1024);
 			bMaxSizeExceeded = true;
+			return EPutStatus::NotCached;
 		}
 		else
 		{
@@ -143,6 +145,7 @@ void FMemoryDerivedDataBackend::PutCachedData(const TCHAR* CacheKey, TArrayView<
 			CalcCacheValueSize(Key, *Val);
 
 			CurrentCacheSize += CacheValueSize;
+			return EPutStatus::Cached;
 		}
 	}
 }
diff --git a/Engine/Source/Developer/DerivedDataCache/Private/MemoryDerivedDataBackend.h b/Engine/Source/Developer/DerivedDataCache/Private/MemoryDerivedDataBackend.h
--- a/Engine/Source/Developer/DerivedDataCache/Private/MemoryDerivedDataBackend.h
+++ b/Engine/Source/Developer/DerivedDataCache/Private/MemoryDerivedDataBackend.h
@@ -53,7 +53,7 @@ public:
 	 * @param	InData		Buffer containing the data to cache, can be destroyed after the call returns, immediately
 	 * @param	bPutEvenIfExists	If true, then do not attempt skip the put even if CachedDataProbablyExists returns true
 	 */
-	virtual void PutCachedData(const TCHAR* CacheKey, TArrayView<const uint8> InData, bool bPutEvenIfExists) override;
+	virtual EPutStatus PutCachedData(const TCHAR* CacheKey, TArrayView<const uint8> InData, bool bPutEvenIfExists) override;
 
 	virtual void RemoveCachedData(const TCHAR* CacheKey, bool bTransient) override;
 
diff --git a/Engine/Source/Developer/DerivedDataCache/Private/PakFileDerivedDataBackend.cpp b/Engine/Source/Developer/DerivedDataCache/Private/PakFileDerivedDataBackend.cpp
--- a/Engine/Source/Developer/DerivedDataCache/Private/PakFileDerivedDataBackend.cpp
+++ b/Engine/Source/Developer/DerivedDataCache/Private/PakFileDerivedDataBackend.cpp
@@ -137,12 +137,12 @@ bool FPakFileDerivedDataBackend::GetCachedData(const TCHAR* CacheKey, TArray<uin
 	return false;
 }
 
-void FPakFileDerivedDataBackend::PutCachedData(const TCHAR* CacheKey, TArrayView<const uint8> InData, bool bPutEvenIfExists)
+FDerivedDataBackendInterface::EPutStatus FPakFileDerivedDataBackend::PutCachedData(const TCHAR* CacheKey, TArrayView<const uint8> InData, bool bPutEvenIfExists)
 {
 	COOK_STAT(auto Timer = UsageStats.TimePut());
 	if (!bWriting || bClosed)
 	{
-		return;
+		return EPutStatus::NotCached;
 	}
 	{
 		FScopeLock ScopeLock(&SynchronizationObject);
@@ -160,7 +160,8 @@ void FPakFileDerivedDataBackend::PutCachedData(const TCHAR* CacheKey, TArrayView
 			{
 				CacheItems.Empty();
 				FileHandle.Reset();
-				bClosed = true;
+				UE_LOG(LogDerivedDataCache, Fatal, TEXT("Could not write pak file...out of disk space?"));
+				return EPutStatus::NotCached;
 			}
 			else
 			{
@@ -169,12 +170,13 @@ void FPakFileDerivedDataBackend::PutCachedData(const TCHAR* CacheKey, TArrayView
 				FileHandle->Serialize(const_cast<uint8*>(InData.GetData()), int64(InData.Num()));
 				UE_LOG(LogDerivedDataCache, Verbose, TEXT("FPakFileDerivedDataBackend: Put %s"), CacheKey);
 				CacheItems.Add(Key,FCacheValue(Offset, InData.Num(), Crc));
+				return EPutStatus::Cached;
 			}
 		}
-	}
-	if (bClosed)
-	{
-		UE_LOG(LogDerivedDataCache, Fatal, TEXT("Could not write pak file...out of disk space?"));
+		else
+		{
+			return EPutStatus::Cached;
+		}
 	}
 }
 
@@ -410,7 +412,7 @@ FCompressedPakFileDerivedDataBackend::FCompressedPakFileDerivedDataBackend(const
 {
 }
 
-void FCompressedPakFileDerivedDataBackend::PutCachedData(const TCHAR* CacheKey, TArrayView<const uint8> InData, bool bPutEvenIfExists)
+FDerivedDataBackendInterface::EPutStatus FCompressedPakFileDerivedDataBackend::PutCachedData(const TCHAR* CacheKey, TArrayView<const uint8> InData, bool bPutEvenIfExists)
 {
 	int32 UncompressedSize = InData.Num();
 	int32 CompressedSize = FCompression::CompressMemoryBound(CompressionFormat, UncompressedSize, CompressionFlags);
@@ -422,7 +424,7 @@ void FCompressedPakFileDerivedDataBackend::PutCachedData(const TCHAR* CacheKey,
 	verify(FCompression::CompressMemory(CompressionFormat, CompressedData.GetData() + sizeof(UncompressedSize), CompressedSize, InData.GetData(), InData.Num(), CompressionFlags));
 	CompressedData.SetNum(CompressedSize + sizeof(UncompressedSize), false);
 
-	FPakFileDerivedDataBackend::PutCachedData(CacheKey, CompressedData, bPutEvenIfExists);
+	return FPakFileDerivedDataBackend::PutCachedData(CacheKey, CompressedData, bPutEvenIfExists);
 }
 
 bool FCompressedPakFileDerivedDataBackend::GetCachedData(const TCHAR* CacheKey, TArray<uint8>& OutData)
diff --git a/Engine/Source/Developer/DerivedDataCache/Private/PakFileDerivedDataBackend.h b/Engine/Source/Developer/DerivedDataCache/Private/PakFileDerivedDataBackend.h
--- a/Engine/Source/Developer/DerivedDataCache/Private/PakFileDerivedDataBackend.h
+++ b/Engine/Source/Developer/DerivedDataCache/Private/PakFileDerivedDataBackend.h
@@ -62,7 +62,7 @@ public:
 	 * @param	InData		Buffer containing the data to cache, can be destroyed after the call returns, immediately
 	 * @param	bPutEvenIfExists	If true, then do not attempt skip the put even if CachedDataProbablyExists returns true
 	 */
-	virtual void PutCachedData(const TCHAR* CacheKey, TArrayView<const uint8> InData, bool bPutEvenIfExists) override;
+	virtual EPutStatus PutCachedData(const TCHAR* CacheKey, TArrayView<const uint8> InData, bool bPutEvenIfExists) override;
 	virtual void RemoveCachedData(const TCHAR* CacheKey, bool bTransient) override;
 
 	/**
@@ -138,7 +138,7 @@ class FCompressedPakFileDerivedDataBackend : public FPakFileDerivedDataBackend
 public:
 	FCompressedPakFileDerivedDataBackend(const TCHAR* InFilename, bool bInWriting);
 
-	virtual void PutCachedData(const TCHAR* CacheKey, TArrayView<const uint8> InData, bool bPutEvenIfExists) override;
+	virtual EPutStatus PutCachedData(const TCHAR* CacheKey, TArrayView<const uint8> InData, bool bPutEvenIfExists) override;
 	virtual bool GetCachedData(const TCHAR* CacheKey, TArray<uint8>& OutData) override;
 
 	/** Returns a class of speed for this interface **/
diff --git a/Engine/Source/Developer/DerivedDataCache/Private/S3DerivedDataBackend.cpp b/Engine/Source/Developer/DerivedDataCache/Private/S3DerivedDataBackend.cpp
--- a/Engine/Source/Developer/DerivedDataCache/Private/S3DerivedDataBackend.cpp
+++ b/Engine/Source/Developer/DerivedDataCache/Private/S3DerivedDataBackend.cpp
@@ -921,9 +921,10 @@ bool FS3DerivedDataBackend::GetCachedData(const TCHAR* CacheKey, TArray<uint8>&
 	return false;
 }
 
-void FS3DerivedDataBackend::PutCachedData(const TCHAR* CacheKey, TArrayView<const uint8> InData, bool bPutEvenIfExists)
+FDerivedDataBackendInterface::EPutStatus FS3DerivedDataBackend::PutCachedData(const TCHAR* CacheKey, TArrayView<const uint8> InData, bool bPutEvenIfExists)
 {
 	// Not implemented
+	return EPutStatus::NotCached;
 }
 
 void FS3DerivedDataBackend::RemoveCachedData(const TCHAR* CacheKey, bool bTransient)
diff --git a/Engine/Source/Developer/DerivedDataCache/Private/S3DerivedDataBackend.h b/Engine/Source/Developer/DerivedDataCache/Private/S3DerivedDataBackend.h
--- a/Engine/Source/Developer/DerivedDataCache/Private/S3DerivedDataBackend.h
+++ b/Engine/Source/Developer/DerivedDataCache/Private/S3DerivedDataBackend.h
@@ -47,7 +47,7 @@ public:
 
 	bool CachedDataProbablyExists(const TCHAR* CacheKey) override;
 	bool GetCachedData(const TCHAR* CacheKey, TArray<uint8>& OutData) override;
-	void PutCachedData(const TCHAR* CacheKey, TArrayView<const uint8> InData, bool bPutEvenIfExists) override;
+	EPutStatus PutCachedData(const TCHAR* CacheKey, TArrayView<const uint8> InData, bool bPutEvenIfExists) override;
 	void RemoveCachedData(const TCHAR* CacheKey, bool bTransient) override;
 	void GatherUsageStats(TMap<FString, FDerivedDataCacheUsageStats>& UsageStatsMap, FString&& GraphPath) override;
 
diff --git a/Engine/Source/Developer/DerivedDataCache/Public/DerivedDataBackendInterface.h b/Engine/Source/Developer/DerivedDataCache/Public/DerivedDataBackendInterface.h
--- a/Engine/Source/Developer/DerivedDataCache/Public/DerivedDataBackendInterface.h
+++ b/Engine/Source/Developer/DerivedDataCache/Public/DerivedDataBackendInterface.h
@@ -29,6 +29,17 @@ class FDerivedDataBackendInterface
 {
 public:
 
+	/** Status of a put operation. */
+	enum class EPutStatus
+	{
+		/** The put is executing asynchronously. */
+		Executing,
+		/** The put completed synchronously and the data was not cached. */
+		NotCached,
+		/** The put completed synchronously and the data was cached. */
+		Cached,
+	};
+
 	/*
 		Speed classes. Higher values are faster so > / < comparisons make sense.
 	*/
@@ -134,7 +145,7 @@ public:
 	 * @param	InData				Buffer containing the data to cache, can be destroyed after the call returns, immediately
 	 * @param	bPutEvenIfExists	If true, then do not attempt skip the put even if CachedDataProbablyExists returns true
 	 */
-	virtual void PutCachedData(const TCHAR* CacheKey, TArrayView<const uint8> InData, bool bPutEvenIfExists) = 0;
+	virtual EPutStatus PutCachedData(const TCHAR* CacheKey, TArrayView<const uint8> InData, bool bPutEvenIfExists) = 0;
 
 	/**
 	 * Remove data from cache (used in the event that corruption is detected at a higher level and possibly house keeping)
