DDC: Changed CachedDataProbablyExistsBatch to call CachedDataProbablyExists for a batch size of 1

#rnx

diff --git a/Engine/Source/Developer/DerivedDataCache/Private/DerivedDataCache.cpp b/Engine/Source/Developer/DerivedDataCache/Private/DerivedDataCache.cpp
--- a/Engine/Source/Developer/DerivedDataCache/Private/DerivedDataCache.cpp
+++ b/Engine/Source/Developer/DerivedDataCache/Private/DerivedDataCache.cpp
@@ -496,10 +496,10 @@ public:
 
 	virtual TBitArray<> CachedDataProbablyExistsBatch(TConstArrayView<FString> CacheKeys) override
 	{
-		DDC_SCOPE_CYCLE_COUNTER(DDC_CachedDataProbablyExistsBatch);
 		TBitArray<> Result;
-		if (CacheKeys.Num())
+		if (CacheKeys.Num() > 1)
 		{
+			DDC_SCOPE_CYCLE_COUNTER(DDC_CachedDataProbablyExistsBatch);
 			INC_DWORD_STAT(STAT_DDC_NumExist);
 			STAT(double ThisTime = 0);
 			{
@@ -509,6 +509,10 @@ public:
 			}
 			INC_FLOAT_STAT_BY(STAT_DDC_ExistTime, (float)ThisTime);
 		}
+		else if (CacheKeys.Num() == 1)
+		{
+			Result.Add(CachedDataProbablyExists(*CacheKeys[0]));
+		}
 		return Result;
 	}
 
