DDC: Implemented batched CachedDataProbablyExists requests in the Jupiter backend

diff --git a/Engine/Source/Developer/DerivedDataCache/Private/HttpDerivedDataBackend.cpp b/Engine/Source/Developer/DerivedDataCache/Private/HttpDerivedDataBackend.cpp
--- a/Engine/Source/Developer/DerivedDataCache/Private/HttpDerivedDataBackend.cpp
+++ b/Engine/Source/Developer/DerivedDataCache/Private/HttpDerivedDataBackend.cpp
@@ -19,6 +19,7 @@
 #include "Misc/FileHelper.h"
 #include "Misc/ScopeLock.h"
 #include "Misc/SecureHash.h"
+#include "Misc/StringBuilder.h"
 #include "ProfilingDebugging/CpuProfilerTrace.h"
 #include "ProfilingDebugging/CountersTrace.h"
 #include "Serialization/JsonReader.h"
@@ -333,6 +334,20 @@ public:
 		return JsonObject;
 	}
 
+	/**
+	 * Tries to parse the response buffer as a JsonArray. Return empty array if
+	 * parse error occurs.
+	 */
+	TArray<TSharedPtr<FJsonValue>> GetResponseAsJsonArray() const
+	{
+		FString Response = GetAnsiBufferAsString(ResponseBuffer);
+
+		TArray<TSharedPtr<FJsonValue>> JsonArray;
+		TSharedRef<TJsonReader<>> JsonReader = TJsonReaderFactory<>::Create(Response);
+		FJsonSerializer::Deserialize(JsonReader, JsonArray);
+		return JsonArray;
+	}
+
 	/** Will return true if the response code is considered a success */
 	static bool IsSuccessResponse(long ResponseCode)
 	{
@@ -1149,14 +1164,13 @@ bool FHttpDerivedDataBackend::CachedDataProbablyExists(const TCHAR* CacheKey)
 	COOK_STAT(auto Timer = UsageStats.TimeProbablyExists());
 	
 	FString Uri = FString::Printf(TEXT("api/v1/c/ddc/%s/%s/%s"), *Namespace, *DefaultBucket, CacheKey);
-	long ResponseCode = 0; uint32 Attempts = 0;
 
 	// Retry request until we get an accepted response or exhaust allowed number of attempts.
-	while (ResponseCode == 0 && ++Attempts < UE_HTTPDDC_MAX_ATTEMPTS)
+	for (int32 Attempts = 0; Attempts < UE_HTTPDDC_MAX_ATTEMPTS; ++Attempts)
 	{
 		FScopedRequestPtr Request(RequestPool.Get());
-		FRequest::Result Result = Request->PerformBlockingQuery<FRequest::Head>(*Uri);
-		ResponseCode = Request->GetResponseCode();
+		const FRequest::Result Result = Request->PerformBlockingQuery<FRequest::Head>(*Uri);
+		const int64 ResponseCode = Request->GetResponseCode();
 
 		if (FRequest::IsSuccessResponse(ResponseCode) || ResponseCode == 400)
 		{
@@ -1164,21 +1178,82 @@ bool FHttpDerivedDataBackend::CachedDataProbablyExists(const TCHAR* CacheKey)
 			if (bIsHit)
 			{
 				TRACE_COUNTER_ADD(HttpDDC_ExistHit, int64(1));
+				COOK_STAT(Timer.AddHit(0));
 			}
 			return bIsHit;
 		}
 
 		if (!ShouldRetryOnError(ResponseCode))
 		{
-			return false;
+			break;
 		}
-
-		ResponseCode = 0;
 	}
 
 	return false;
 }
 
+TBitArray<> FHttpDerivedDataBackend::CachedDataProbablyExistsBatch(TConstArrayView<FString> CacheKeys)
+{
+	TRACE_CPUPROFILER_EVENT_SCOPE(HttpDDC_Exist);
+	TRACE_COUNTER_ADD(HttpDDC_Exist, int64(1));
+	COOK_STAT(auto Timer = UsageStats.TimeProbablyExists());
+
+	const TCHAR* const Uri = TEXT("api/v1/c/ddc-rpc");
+
+	TAnsiStringBuilder<512> Body;
+	const FTCHARToUTF8 AnsiNamespace(*Namespace);
+	const FTCHARToUTF8 AnsiBucket(*DefaultBucket);
+	Body << "{\"Operations\":[";
+	for (const FString& CacheKey : CacheKeys)
+	{
+		Body << "{\"Namespace\":\"" << AnsiNamespace.Get() << "\",\"Bucket\":\"" << AnsiBucket.Get() << "\",";
+		Body << "\"Id\":\"" << FTCHARToUTF8(*CacheKey).Get() << "\",\"Op\":\"HEAD\"},";
+	}
+	Body.RemoveSuffix(1);
+	Body << "]}";
+
+	TConstArrayView<uint8> BodyView(reinterpret_cast<const uint8*>(Body.ToString()), Body.Len());
+
+	// Retry request until we get an accepted response or exhaust allowed number of attempts.
+	for (int32 Attempts = 0; Attempts < UE_HTTPDDC_MAX_ATTEMPTS; ++Attempts)
+	{
+		FScopedRequestPtr Request(RequestPool.Get());
+		const FRequest::Result Result = Request->PerformBlockingUpload<FRequest::PostJson>(Uri, BodyView);
+		const int64 ResponseCode = Request->GetResponseCode();
+
+		if (Result == FRequest::Success && ResponseCode == 200)
+		{
+			TArray<TSharedPtr<FJsonValue>> ResponseArray = Request->GetResponseAsJsonArray();
+			if (ResponseArray.Num() == CacheKeys.Num())
+			{
+				TBitArray<> Exists;
+				Exists.Reserve(CacheKeys.Num());
+				const FString* CacheKeyIt = CacheKeys.GetData();
+				for (const TSharedPtr<FJsonValue>& Response : ResponseArray)
+				{
+					FString Key;
+					Exists.Add(Response->TryGetString(Key) && Key == *CacheKeyIt);
+					++CacheKeyIt;
+				}
+				if (Exists.CountSetBits() == CacheKeys.Num())
+				{
+					TRACE_COUNTER_ADD(HttpDDC_ExistHit, int64(1));
+					COOK_STAT(Timer.AddHit(0));
+				}
+				return Exists;
+			}
+			break;
+		}
+
+		if (!ShouldRetryOnError(ResponseCode))
+		{
+			break;
+		}
+	}
+
+	return TBitArray<>(false, CacheKeys.Num());
+}
+
 bool FHttpDerivedDataBackend::GetCachedData(const TCHAR* CacheKey, TArray<uint8>& OutData)
 {
 	TRACE_CPUPROFILER_EVENT_SCOPE(HttpDDC_Get);
diff --git a/Engine/Source/Developer/DerivedDataCache/Private/HttpDerivedDataBackend.h b/Engine/Source/Developer/DerivedDataCache/Private/HttpDerivedDataBackend.h
--- a/Engine/Source/Developer/DerivedDataCache/Private/HttpDerivedDataBackend.h
+++ b/Engine/Source/Developer/DerivedDataCache/Private/HttpDerivedDataBackend.h
@@ -51,6 +51,7 @@ public:
 
 	virtual bool IsWritable() override { return true; }
 	virtual bool CachedDataProbablyExists(const TCHAR* CacheKey) override;
+	virtual TBitArray<> CachedDataProbablyExistsBatch(TConstArrayView<FString> CacheKeys) override;
 	virtual bool GetCachedData(const TCHAR* CacheKey, TArray<uint8>& OutData) override;
 	virtual EPutStatus PutCachedData(const TCHAR* CacheKey, TArrayView<const uint8> InData, bool bPutEvenIfExists) override;
 	virtual void RemoveCachedData(const TCHAR* CacheKey, bool bTransient) override;
