DDC: Added the new cache interface UE::DerivedData::ICache

ICache will eventually replace the existing cache access functions on FDerivedDataCacheInterface. The notable differences relative to the existing interface are:

- It is asynchronous with completion callbacks by default, and allows waiting when necessary.
- It allows batch requests containing multiple cache records by default.
- It uses a structured cache key made up of a bucket name and a hash.
- It supports prioritization and cancellation of requests.
- It supports structured cache records with metadata.
- It supports attachments to allow multi-part cache records.
- It uses a typed request object instead of an opaque integer handle.
- It provides cache policy flags to control how to query and store cache records, and which parts of a cache record to fetch.

This initial implementation of ICache is totally synchronous and uses the existing cache backends for storage.

diff --git a/Engine/Source/Developer/DerivedDataCache/Private/DerivedDataCache.cpp b/Engine/Source/Developer/DerivedDataCache/Private/DerivedDataCache.cpp
--- a/Engine/Source/Developer/DerivedDataCache/Private/DerivedDataCache.cpp
+++ b/Engine/Source/Developer/DerivedDataCache/Private/DerivedDataCache.cpp
@@ -1,5 +1,7 @@
 // Copyright Epic Games, Inc. All Rights Reserved.
 
+#include "DerivedDataCache.h"
+#include "DerivedDataCacheInterface.h"
 
 #include "CoreMinimal.h"
 #include "Misc/CommandLine.h"
@@ -13,12 +15,21 @@
 #include "Serialization/MemoryWriter.h"
 #include "Modules/ModuleManager.h"
 
-#include "DerivedDataCacheInterface.h"
 #include "DerivedDataBackendInterface.h"
 #include "DerivedDataPluginInterface.h"
 #include "DDCCleanup.h"
 #include "ProfilingDebugging/CookStats.h"
 
+#include "Algo/AllOf.h"
+#include "Algo/Transform.h"
+#include "DerivedDataBackendInterface.h"
+#include "Misc/CoreMisc.h"
+#include "Misc/ScopeExit.h"
+#include "Misc/StringBuilder.h"
+#include "Serialization/CompactBinary.h"
+#include "Serialization/CompactBinaryValidation.h"
+#include "Serialization/CompactBinaryWriter.h"
+
 #include <atomic>
 
 DEFINE_STAT(STAT_DDC_NumGets);
@@ -130,6 +141,16 @@ namespace DerivedDataCacheCookStats
 }
 #endif
 
+namespace UE
+{
+namespace DerivedData
+{
+
+ICache* CreateCache();
+
+} // DerivedData
+} // UE
+
 /** Whether we want to verify the DDC (pass in -VerifyDDC on the command line)*/
 bool GVerifyDDC = false;
 
@@ -317,6 +338,7 @@ public:
 	/** Constructor, called once to cereate a singleton **/
 	FDerivedDataCache()
 		: CurrentHandle(19248) // we will skip some potential handles to catch errors
+		, Cache(UE::DerivedData::CreateCache())
 	{
 		FDerivedDataBackend::Get(); // we need to make sure this starts before we all us to start
 
@@ -337,6 +359,11 @@ public:
 		PendingTasks.Empty();
 	}
 
+	virtual UE::DerivedData::ICache& GetCache()
+	{
+		return *Cache;
+	}
+
 	virtual bool GetSynchronous(FDerivedDataPluginInterface* DataDeriver, TArray<uint8>& OutData, bool* bDataWasBuilt = nullptr) override
 	{
 		DDC_SCOPE_CYCLE_COUNTER(DDC_GetSynchronous);
@@ -595,6 +622,8 @@ private:
 
 	/** Cache notification delegate */
 	FOnDDCNotification DDCNotificationEvent;
+
+	TUniquePtr<UE::DerivedData::ICache> Cache;
 };
 
 /**
@@ -621,3 +650,533 @@ public:
 };
 
 IMPLEMENT_MODULE( FDerivedDataCacheModule, DerivedDataCache);
+
+///////////////////////////////////////////////////////////////////////////////////////////////////////////////////////
+
+namespace UE
+{
+namespace DerivedData
+{
+
+///////////////////////////////////////////////////////////////////////////////////////////////////////////////////////
+
+class FCache final : public ICache
+{
+public:
+	virtual ~FCache() = default;
+
+	virtual FCacheRequest Get(
+		TConstArrayView<FCacheKey> Keys,
+		FStringView Context,
+		ECachePolicy Policy,
+		ECachePriority Priority,
+		FOnCacheGetComplete&& Callback) final;
+
+	virtual FCacheRequest Put(
+		TConstArrayView<FCacheRecord> Records,
+		FStringView Context,
+		ECachePolicy Policy,
+		ECachePriority Priority,
+		FOnCachePutComplete&& Callback) final;
+
+	virtual FCacheRequest GetAttachments(
+		TConstArrayView<FCacheAttachmentKey> Keys,
+		FStringView Context,
+		ECachePolicy Policy,
+		ECachePriority Priority,
+		FOnCacheGetAttachmentComplete&& Callback) final;
+
+	virtual void CancelAll() final;
+
+private:
+	void Get(
+		const FCacheKey& Key,
+		FStringView Context,
+		ECachePolicy Policy,
+		ECachePriority Priority,
+		FOnCacheGetComplete& Callback);
+
+	void Put(
+		const FCacheRecord& Record,
+		FStringView Context,
+		ECachePolicy Policy,
+		ECachePriority Priority,
+		FOnCachePutComplete& Callback);
+
+	void GetAttachment(
+		const FCacheAttachmentKey& Key,
+		FStringView Context,
+		ECachePolicy Policy,
+		ECachePriority Priority,
+		FOnCacheGetAttachmentComplete& Callback);
+
+	template <int32 BufferSize>
+	class TToString
+	{
+	public:
+		template <typename T>
+		explicit TToString(const T& Input)
+		{
+			Buffer << Input;
+		}
+
+		inline const TCHAR* operator*() const { return Buffer.ToString(); }
+		inline const TCHAR* ToString() const { return Buffer.ToString(); }
+		inline operator const TCHAR*() const { return Buffer.ToString(); }
+
+	private:
+		TStringBuilder<BufferSize> Buffer;
+	};
+
+	static FString MakeRecordKey(const FCacheKey& Key)
+	{
+		check(Key);
+		TStringBuilder<96> Out;
+		Out << Key.GetBucket() << TEXT("_") << Key.GetHash();
+		return FDerivedDataCacheInterface::SanitizeCacheKey(Out.ToString());
+	}
+
+	static FString MakeContentKey(const FCacheAttachmentKey& AttachmentKey)
+	{
+		check(AttachmentKey);
+		TStringBuilder<128> Out;
+		Out << AttachmentKey.GetKey().GetBucket() << TEXT("_BLAKE3_") << AttachmentKey.GetHash();
+		return FDerivedDataCacheInterface::SanitizeCacheKey(Out.ToString());
+	}
+};
+
+///////////////////////////////////////////////////////////////////////////////////////////////////////////////////////
+
+void FCache::Get(
+	const FCacheKey& Key,
+	FStringView Context,
+	ECachePolicy Policy,
+	ECachePriority Priority,
+	FOnCacheGetComplete& Callback)
+{
+	FCacheGetCompleteParams Params;
+	Params.Record.SetKey(Key);
+
+	ON_SCOPE_EXIT
+	{
+		if (Callback)
+		{
+			Callback(MoveTemp(Params));
+		}
+	};
+
+	// Skip the request if querying the cache is disabled.
+	if (!EnumHasAnyFlags(Policy, ECachePolicy::Query))
+	{
+		UE_LOG(LogDerivedDataCache, VeryVerbose, TEXT("Cache: Get skipped for %s from '%.*s'"),
+			*TToString<96>(Key), Context.Len(), Context.GetData());
+		return;
+	}
+
+	// Request the metadata and references from storage.
+	TArray<uint8> Data;
+	if (!GetDerivedDataCacheRef().GetSynchronous(*MakeRecordKey(Key), Data, Context))
+	{
+		UE_LOG(LogDerivedDataCache, Verbose, TEXT("Cache: Get cache miss for %s from '%.*s'"),
+			*TToString<96>(Key), Context.Len(), Context.GetData());
+		return;
+	}
+
+	// Validate that Data can be read as compact binary without crashing.
+	if (ValidateCompactBinaryRange(MakeMemoryView(Data), ECbValidateMode::Default) != ECbValidateError::None)
+	{
+		UE_LOG(LogDerivedDataCache, Verbose, TEXT("Cache: Get cache miss with corrupted record for %s from '%.*s'"),
+			*TToString<96>(Key), Context.Len(), Context.GetData());
+		return;
+	}
+
+	// Read the record from its compact binary fields.
+	FCbObject RecordObject(Data.GetData());
+	FCbField MetaField = RecordObject["Meta"_ASV];
+	FCbField ValueField = RecordObject["Value"_ASV];
+	FCbField AttachmentsField = RecordObject["Attachments"_ASV];
+	FCbArray AttachmentsArray = AttachmentsField.AsArray();
+
+	// Validate that the record was serialized in the expected format.
+	if ((MetaField && !MetaField.IsObject()) ||
+		(ValueField && !ValueField.IsAnyReference()) ||
+		(AttachmentsField && !AttachmentsField.IsArray()) ||
+		!Algo::AllOf(AttachmentsArray, &FCbField::IsAnyReference))
+	{
+		UE_LOG(LogDerivedDataCache, Verbose, TEXT("Cache: Get cache miss with invalid format for %s from '%.*s'"),
+			*TToString<96>(Key), Context.Len(), Context.GetData());
+		return;
+	}
+
+	// Check for existence of the value and attachments if they are being skipped.
+	{
+		TArray<FBlake3Hash, TInlineAllocator<1>> KeysToCheck;
+		if (EnumHasAnyFlags(Policy, ECachePolicy::SkipValue) && ValueField.IsAnyReference())
+		{
+			KeysToCheck.Add(ValueField.AsAnyReference());
+			if (KeysToCheck.Last().IsZero())
+			{
+				KeysToCheck.Pop();
+			}
+		}
+		if (EnumHasAnyFlags(Policy, ECachePolicy::SkipAttachments))
+		{
+			KeysToCheck.Reserve(KeysToCheck.Num() + AttachmentsArray.Num());
+			for (FCbField AttachmentField : AttachmentsArray)
+			{
+				KeysToCheck.Add(AttachmentField.AsAnyReference());
+			}
+		}
+		if (KeysToCheck.Num())
+		{
+			TArray<FString, TInlineAllocator<1>> ContentKeysToCheck;
+			ContentKeysToCheck.Reserve(KeysToCheck.Num());
+			Algo::Transform(KeysToCheck, ContentKeysToCheck,
+				[&Key](const FBlake3Hash& Hash) -> FString { return MakeContentKey(FCacheAttachmentKey(Key, Hash)); });
+			if (!GetDerivedDataCacheRef().AllCachedDataProbablyExists(ContentKeysToCheck))
+			{
+				UE_LOG(LogDerivedDataCache, Verbose, TEXT("Cache: Get cache miss with missing content for %s from '%.*s'"),
+					*TToString<96>(Key), Context.Len(), Context.GetData());
+				return;
+			}
+		}
+	}
+
+	// Read the value and attachments if they have been requested.
+	const auto GetContent = [&Key, Context](FCbField ReferenceField) -> FCbAttachment
+	{
+		TArray<uint8> ContentData;
+		const FBlake3Hash Hash = ReferenceField.AsAnyReference();
+		check(!ReferenceField.HasError());
+		const FCacheAttachmentKey ContentKey(Key, Hash);
+		if (GetDerivedDataCacheRef().GetSynchronous(*MakeContentKey(ContentKey), ContentData, Context))
+		{
+			if (FBlake3::HashBuffer(MakeMemoryView(ContentData)) == Hash)
+			{
+				FSharedBuffer ContentBuffer = FSharedBuffer::Clone(MakeMemoryView(ContentData));
+				if (ReferenceField.IsReference())
+				{
+					return FCbAttachment(FCbFieldRefIterator::MakeRange(MoveTemp(ContentBuffer)), Hash);
+				}
+				else
+				{
+					return FCbAttachment(MoveTemp(ContentBuffer), Hash);
+				}
+			}
+			else
+			{
+				UE_LOG(LogDerivedDataCache, Verbose, TEXT("Cache: Get cache miss with corrupted content for %s from '%.*s'"),
+					*TToString<160>(ContentKey), Context.Len(), Context.GetData());
+			}
+		}
+		else
+		{
+			UE_LOG(LogDerivedDataCache, Verbose, TEXT("Cache: Get cache miss with missing content for %s from '%.*s'"),
+				*TToString<160>(ContentKey), Context.Len(), Context.GetData());
+		}
+		return FCbAttachment();
+	};
+
+	FCbAttachment Value;
+	if (!EnumHasAnyFlags(Policy, ECachePolicy::SkipValue) && ValueField.IsAnyReference() && !(Value = GetContent(ValueField)))
+	{
+		return;
+	}
+
+	TArray<FCbAttachment> Attachments;
+	if (!EnumHasAnyFlags(Policy, ECachePolicy::SkipAttachments) && AttachmentsField.IsArray())
+	{
+		Attachments.Reserve(AttachmentsArray.Num());
+		for (FCbField AttachmentField : AttachmentsArray)
+		{
+			Attachments.Add(GetContent(AttachmentField));
+			if (!Attachments.Last())
+			{
+				return;
+			}
+		}
+	}
+
+	// Package
+	if (AttachmentsField.IsArray() && !EnumHasAllFlags(Policy, ECachePolicy::SkipValue | ECachePolicy::SkipAttachments))
+	{
+		FCbPackage Package;
+		if (Value && Value.IsCompactBinary())
+		{
+			Package.SetObject(Value.AsCompactBinary().AsObjectRef(), Value.GetHash());
+		}
+		for (FCbAttachment Attachment : Attachments)
+		{
+			Package.AddAttachment(MoveTemp(Attachment));
+		}
+		Params.Record.SetPackage(MoveTemp(Package));
+	}
+	else if (Value)
+	{
+		// Object
+		if (Value.IsCompactBinary())
+		{
+			Params.Record.SetObject(Value.AsCompactBinary().AsObjectRef());
+		}
+		// Binary
+		else
+		{
+			Params.Record.SetBinary(Value.AsBinary());
+		}
+	}
+
+	// Meta
+	if (!EnumHasAnyFlags(Policy, ECachePolicy::SkipMeta) && MetaField.IsObject())
+	{
+		Params.Record.SetMeta(FCbObjectRef::Clone(MetaField.AsObject()));
+	}
+
+	UE_LOG(LogDerivedDataCache, Verbose, TEXT("Cache: Get cache hit for %s from '%.*s'"),
+		*TToString<96>(Key), Context.Len(), Context.GetData());
+	Params.Status = ECacheStatus::Cached;
+}
+
+FCacheRequest FCache::Get(
+	TConstArrayView<FCacheKey> Keys,
+	FStringView Context,
+	ECachePolicy Policy,
+	ECachePriority Priority,
+	FOnCacheGetComplete&& Callback)
+{
+	for (const FCacheKey& Key : Keys)
+	{
+		Get(Key, Context, Policy, Priority, Callback);
+	}
+	return FCacheRequest();
+}
+
+void FCache::Put(
+	const FCacheRecord& Record,
+	FStringView Context,
+	ECachePolicy Policy,
+	ECachePriority Priority,
+	FOnCachePutComplete& Callback)
+{
+	FCachePutCompleteParams Params;
+	Params.Key = Record.GetKey();
+
+	ON_SCOPE_EXIT
+	{
+		if (Callback)
+		{
+			Callback(MoveTemp(Params));
+		}
+	};
+
+	if (!EnumHasAnyFlags(Policy, ECachePolicy::Store))
+	{
+		UE_LOG(LogDerivedDataCache, VeryVerbose, TEXT("Cache: Put skipped for %s from '%.*s'"),
+			*TToString<96>(Params.Key), Context.Len(), Context.GetData());
+		Params.Status = ECacheStatus::NotCached;
+		return;
+	}
+
+	FCbAttachment Value;
+	TConstArrayView<FCbAttachment> Attachments;
+
+	switch (Record.GetType())
+	{
+	case ECacheRecordType::None:
+		break;
+	case ECacheRecordType::Binary:
+		Value = FCbAttachment(Record.AsBinary());
+		break;
+	case ECacheRecordType::Object:
+		Value = FCbAttachment(FCbFieldRefIterator::MakeSingle(Record.AsObject().AsFieldRef()));
+		break;
+	case ECacheRecordType::Package:
+		{
+			const FCbPackage& Package = Record.AsPackage();
+			Value = FCbAttachment(FCbFieldRefIterator::MakeSingle(Package.GetObject().AsFieldRef()));
+			Attachments = Package.GetAttachments();
+		}
+		break;
+	default:
+		checkNoEntry();
+		break;
+	}
+
+	const auto PutContent = [&Key = Params.Key, Context](const FSharedBuffer& Buffer, const FBlake3Hash& Hash)
+	{
+		check(Buffer.GetSize() <= MAX_int32);
+		TConstArrayView<uint8> BufferView = MakeArrayView(
+			static_cast<const uint8*>(Buffer.GetData()),
+			static_cast<int32>(Buffer.GetSize()));
+		GetDerivedDataCacheRef().Put(*MakeContentKey(FCacheAttachmentKey(Key, Hash)), BufferView, Context);
+	};
+
+	if (Value)
+	{
+		PutContent(Value.AsBinary(), Value.GetHash());
+	}
+	for (const FCbAttachment& Attachment : Attachments)
+	{
+		PutContent(Attachment.AsBinary(), Attachment.GetHash());
+	}
+
+	FCbWriter Writer;
+	Writer.BeginObject();
+	{
+		if (Value.IsCompactBinary())
+		{
+			Writer.Name("Value"_ASV).Reference(Value.GetHash());
+		}
+		else if (Value.IsBinary())
+		{
+			Writer.Name("Value"_ASV).BinaryReference(Value.GetHash());
+		}
+
+		if (Record.GetType() == ECacheRecordType::Package)
+		{
+			Writer.Name("Attachments"_ASV);
+			Writer.BeginArray();
+			for (const FCbAttachment& Attachment : Attachments)
+			{
+				if (Attachment.IsCompactBinary())
+				{
+					Writer.Reference(Attachment.GetHash());
+				}
+				else
+				{
+					Writer.BinaryReference(Attachment.GetHash());
+				}
+			}
+			Writer.EndArray();
+		}
+
+		const FCbObjectRef& Meta = Record.GetMeta();
+		if (Meta.CreateIterator())
+		{
+			Writer.Name("Meta"_ASV).Object(Meta);
+		}
+	}
+	Writer.EndObject();
+
+	TArray<uint8> Data;
+	const uint64 SaveSize = Writer.GetSaveSize();
+	check(SaveSize <= MAX_int32);
+	Data.SetNumUninitialized(static_cast<int32>(SaveSize));
+	Writer.Save(MakeMemoryView(Data));
+
+	const bool bPutEvenIfExists = EnumHasAnyFlags(Policy, ECachePolicy::Overwrite);
+	GetDerivedDataCacheRef().Put(*MakeRecordKey(Params.Key), Data, Context, bPutEvenIfExists);
+
+	UE_LOG(LogDerivedDataCache, VeryVerbose, TEXT("Cache: Put for %s from '%.*s'"),
+		*TToString<96>(Params.Key), Context.Len(), Context.GetData());
+	Params.Status = ECacheStatus::Cached;
+}
+
+FCacheRequest FCache::Put(
+	TConstArrayView<FCacheRecord> Records,
+	FStringView Context,
+	ECachePolicy Policy,
+	ECachePriority Priority,
+	FOnCachePutComplete&& Callback)
+{
+	for (const FCacheRecord& Record : Records)
+	{
+		Put(Record, Context, Policy, Priority, Callback);
+	}
+	return FCacheRequest();
+}
+
+void FCache::GetAttachment(
+	const FCacheAttachmentKey& Key,
+	FStringView Context,
+	ECachePolicy Policy,
+	ECachePriority Priority,
+	FOnCacheGetAttachmentComplete& Callback)
+{
+	FCacheGetAttachmentCompleteParams Params;
+	Params.Key = Key;
+
+	ON_SCOPE_EXIT
+	{
+		if (Callback)
+		{
+			Callback(MoveTemp(Params));
+		}
+	};
+
+	// Skip the request if querying the cache is disabled.
+	if (!EnumHasAnyFlags(Policy, ECachePolicy::Query))
+	{
+		UE_LOG(LogDerivedDataCache, VeryVerbose, TEXT("Cache: GetAttachment skipped on content for %s from '%.*s'"),
+			*TToString<160>(Key), Context.Len(), Context.GetData());
+		return;
+	}
+
+	// Check for existence of the attachment if it is being skipped.
+	if (EnumHasAnyFlags(Policy, ECachePolicy::SkipAttachments))
+	{
+		if (GetDerivedDataCacheRef().CachedDataProbablyExists(*MakeContentKey(Key)))
+		{
+			UE_LOG(LogDerivedDataCache, Verbose, TEXT("Cache: GetAttachment (exists) cache hit on content for %s from '%.*s'"),
+				*TToString<160>(Key), Context.Len(), Context.GetData());
+			Params.Status = ECacheStatus::Cached;
+		}
+		else
+		{
+			UE_LOG(LogDerivedDataCache, Verbose, TEXT("Cache: GetAttachment (exists) cache miss on content for %s from '%.*s'"),
+				*TToString<160>(Key), Context.Len(), Context.GetData());
+		}
+		return;
+	}
+
+	// Request the attachment from storage.
+	TArray<uint8> Data;
+	if (!GetDerivedDataCacheRef().GetSynchronous(*MakeContentKey(Key), Data, Context))
+	{
+		UE_LOG(LogDerivedDataCache, Verbose, TEXT("Cache: GetAttachment cache miss on content for %s from '%.*s'"),
+			*TToString<160>(Key), Context.Len(), Context.GetData());
+		return;
+	}
+
+	if (FBlake3::HashBuffer(MakeMemoryView(Data)) != Key.GetHash())
+	{
+		UE_LOG(LogDerivedDataCache, Verbose, TEXT("Cache: GetAttachment cache miss with corrupted content for %s from '%.*s'"),
+			*TToString<160>(Key), Context.Len(), Context.GetData());
+		return;
+	}
+
+	Params.Value = FSharedBuffer::Clone(MakeMemoryView(Data));
+
+	UE_LOG(LogDerivedDataCache, Verbose, TEXT("Cache: GetAttachment cache hit on content for %s from '%.*s'"),
+		*TToString<160>(Key), Context.Len(), Context.GetData());
+	Params.Status = ECacheStatus::Cached;
+}
+
+FCacheRequest FCache::GetAttachments(
+	TConstArrayView<FCacheAttachmentKey> Keys,
+	FStringView Context,
+	ECachePolicy Policy,
+	ECachePriority Priority,
+	FOnCacheGetAttachmentComplete&& Callback)
+{
+	for (const FCacheAttachmentKey& Key : Keys)
+	{
+		GetAttachment(Key, Context, Policy, Priority, Callback);
+	}
+	return FCacheRequest();
+}
+
+void FCache::CancelAll()
+{
+}
+
+///////////////////////////////////////////////////////////////////////////////////////////////////////////////////////
+
+ICache* CreateCache()
+{
+	return new FCache();
+}
+
+///////////////////////////////////////////////////////////////////////////////////////////////////////////////////////
+
+} // DerivedData
+} // UE
diff --git a/Engine/Source/Developer/DerivedDataCache/Private/DerivedDataCacheKey.cpp b/Engine/Source/Developer/DerivedDataCache/Private/DerivedDataCacheKey.cpp
new file mode 100644
--- /dev/null
+++ b/Engine/Source/Developer/DerivedDataCache/Private/DerivedDataCacheKey.cpp
@@ -0,0 +1,72 @@
+// Copyright Epic Games, Inc. All Rights Reserved.
+
+#include "DerivedDataCacheKey.h"
+
+#include "Algo/AllOf.h"
+#include "Misc/StringBuilder.h"
+#include "UObject/NameTypes.h"
+
+namespace UE
+{
+namespace DerivedData
+{
+
+///////////////////////////////////////////////////////////////////////////////////////////////////////////////////////
+
+FCacheBucket::FCacheBucket(FStringView Name)
+	: Index(FName(Name).GetDisplayIndex().ToUnstableInt())
+{
+	check(Algo::AllOf(Name, FChar::IsAlnum));
+}
+
+void FCacheBucket::ToString(FAnsiStringBuilderBase& Builder) const
+{
+	verify(FName::CreateFromDisplayId(FNameEntryId::FromUnstableInt(Index), 0).TryAppendAnsiString(Builder));
+}
+
+void FCacheBucket::ToString(FWideStringBuilderBase& Builder) const
+{
+	FName::CreateFromDisplayId(FNameEntryId::FromUnstableInt(Index), 0).AppendString(Builder);
+}
+
+bool FCacheBucketLexicalLess::operator()(FCacheBucket A, FCacheBucket B) const
+{
+	return FNameEntryId::FromUnstableInt(A.ToIndex()).LexicalLess(FNameEntryId::FromUnstableInt(B.ToIndex()));
+}
+
+///////////////////////////////////////////////////////////////////////////////////////////////////////////////////////
+
+void FCacheKey::ToString(FAnsiStringBuilderBase& Builder) const
+{
+	Builder << Bucket;
+	if (!Bucket.IsNull())
+	{
+		Builder << "/" << Hash;
+	}
+}
+
+void FCacheKey::ToString(FWideStringBuilderBase& Builder) const
+{
+	Builder << Bucket;
+	if (!Bucket.IsNull())
+	{
+		Builder << TEXT("/") << Hash;
+	}
+}
+
+///////////////////////////////////////////////////////////////////////////////////////////////////////////////////////
+
+void FCacheAttachmentKey::ToString(FAnsiStringBuilderBase& Builder) const
+{
+	Builder << Key << '/' << Hash;
+}
+
+void FCacheAttachmentKey::ToString(FWideStringBuilderBase& Builder) const
+{
+	Builder << Key << TEXT('/') << Hash;
+}
+
+///////////////////////////////////////////////////////////////////////////////////////////////////////////////////////
+
+} // DerivedData
+} // UE
diff --git a/Engine/Source/Developer/DerivedDataCache/Private/DerivedDataCacheRecord.cpp b/Engine/Source/Developer/DerivedDataCache/Private/DerivedDataCacheRecord.cpp
new file mode 100644
--- /dev/null
+++ b/Engine/Source/Developer/DerivedDataCache/Private/DerivedDataCacheRecord.cpp
@@ -0,0 +1,46 @@
+// Copyright Epic Games, Inc. All Rights Reserved.
+
+#include "DerivedDataCacheRecord.h"
+
+namespace UE
+{
+namespace DerivedData
+{
+
+///////////////////////////////////////////////////////////////////////////////////////////////////////////////////////
+
+FCacheRecord::FCacheRecord() = default;
+
+void FCacheRecord::Reset()
+{
+	*this = FCacheRecord();
+}
+
+void FCacheRecord::SetBinary(FSharedBuffer Value)
+{
+	Object.Reset();
+	Package.Reset();
+	Binary = MoveTemp(Value);
+	Type = ECacheRecordType::Binary;
+}
+
+void FCacheRecord::SetObject(FCbObjectRef Value)
+{
+	Binary.Reset();
+	Package.Reset();
+	Object = MoveTemp(Value);
+	Type = ECacheRecordType::Object;
+}
+
+void FCacheRecord::SetPackage(FCbPackage Value)
+{
+	Binary.Reset();
+	Object.Reset();
+	Package = MoveTemp(Value);
+	Type = ECacheRecordType::Package;
+}
+
+///////////////////////////////////////////////////////////////////////////////////////////////////////////////////////
+
+} // DerivedData
+} // UE
diff --git a/Engine/Source/Developer/DerivedDataCache/Public/DerivedDataCache.h b/Engine/Source/Developer/DerivedDataCache/Public/DerivedDataCache.h
new file mode 100644
--- /dev/null
+++ b/Engine/Source/Developer/DerivedDataCache/Public/DerivedDataCache.h
@@ -0,0 +1,320 @@
+// Copyright Epic Games, Inc. All Rights Reserved.
+
+#pragma once
+
+#include "CoreTypes.h"
+#include "DerivedDataCacheKey.h"
+#include "DerivedDataCacheRecord.h"
+#include "Memory/MemoryFwd.h"
+#include "Misc/EnumClassFlags.h"
+#include "Templates/RefCounting.h"
+
+namespace UE
+{
+namespace DerivedData
+{
+
+///////////////////////////////////////////////////////////////////////////////////////////////////////////////////////
+
+/** Flags to control the behavior of cache requests. */
+enum class ECachePolicy : uint8
+{
+	/** A value without any flags set. */
+	None            = 0,
+
+	/** Allow a cache request to query local caches. */
+	QueryLocal      = 1 << 0,
+	/** Allow a cache request to query remote caches. */
+	QueryRemote     = 1 << 1,
+	/** Allow requests to query available local and remote caches. */
+	Query           = QueryLocal | QueryRemote,
+
+	/** Allow cache records to be stored in local caches. */
+	StoreLocal      = 1 << 2,
+	/** Allow cache records to be stored in remote caches. */
+	StoreRemote     = 1 << 3,
+	/** Allow cache records to be stored in local and remote caches. */
+	Store           = StoreLocal | StoreRemote,
+
+	/** Skip fetching the metadata for get requests. */
+	SkipMeta        = 1 << 4,
+	/** Skip fetching the value (binary, object, or the object of the package) for get requests. */
+	SkipValue       = 1 << 5,
+	/** Skip fetching the attachments for get requests. */
+	SkipAttachments = 1 << 6,
+	/** Skip fetching any data for get requests, which will only check for existence. */
+	SkipAll         = SkipMeta | SkipValue | SkipAttachments,
+
+	/** Allow a put request to overwrite existing cache records. */
+	Overwrite       = 1 << 7,
+
+	/** Allow requests to query and store records in local caches. */
+	Local           = QueryLocal | StoreLocal,
+	/** Allow requests to query and store records in remote caches. */
+	Remote          = QueryRemote | StoreRemote,
+
+	/** Allow requests to query and store records in local and remote caches. */
+	Default         = Query | Store,
+
+	/** Do not query or store records in the cache. */
+	Disable         = None,
+};
+
+ENUM_CLASS_FLAGS(ECachePolicy);
+
+///////////////////////////////////////////////////////////////////////////////////////////////////////////////////////
+
+/** Priority for scheduling cache requests. */
+enum class ECachePriority : uint8
+{
+	Highest,
+	High,
+	Normal,
+	Low,
+	Lowest,
+};
+
+///////////////////////////////////////////////////////////////////////////////////////////////////////////////////////
+
+/** The status of a cache operation. */
+enum class ECacheStatus : uint8
+{
+	/** A cache miss for a get operation, or a failure for a put operation. */
+	NotCached,
+	/** A cache hit for a get operation, or a success for a put operation. */
+	Cached,
+	/** The operation was canceled before it completed. */
+	Canceled,
+};
+
+///////////////////////////////////////////////////////////////////////////////////////////////////////////////////////
+
+/** Interface for a cache request. Provides functionality common to every request type. */
+class ICacheRequest : public IRefCountedObject
+{
+public:
+	virtual ~ICacheRequest() = default;
+
+	/** Set the priority of the request. */
+	virtual void SetPriority(ECachePriority Priority) = 0;
+
+	/** Cancel the request and invoke its callback. */
+	virtual void Cancel() = 0;
+
+	/** Block the calling thread until the request is complete. */
+	virtual void Wait() = 0;
+
+	/** Poll the request and return true if complete and false otherwise. */
+	virtual bool Poll() = 0;
+};
+
+///////////////////////////////////////////////////////////////////////////////////////////////////////////////////////
+
+/** A wrapper for the cache request interface that manages lifetime with reference counting. */
+class FCacheRequest
+{
+public:
+	/** Construct a null cache request. */
+	FCacheRequest() = default;
+
+	/** Construct a cache request from . */
+	inline explicit FCacheRequest(ICacheRequest* InRequest)
+		: Request(InRequest)
+	{
+	}
+
+	/** Reset this to null. */
+	inline void Reset() { Request.SafeRelease(); }
+
+	/** Whether the cache request is null. */
+	inline bool IsNull() const { return !Request; }
+	/** Whether the cache request is not null. */
+	inline explicit operator bool() const { return !IsNull(); }
+
+	/**
+	 * Set the priority of the request.
+	 *
+	 * @param Priority The new priority for the request. May be higher or lower than the original.
+	 */
+	inline void SetPriority(ECachePriority Priority)
+	{
+		if (Request)
+		{
+			Request->SetPriority(Priority);
+		}
+	}
+
+	/**
+	 * Cancel the request and invoke its callback.
+	 *
+	 * Does not return until the callback for the request has finished executing.
+	 */
+	inline void Cancel() const
+	{
+		if (Request)
+		{
+			Request->Cancel();
+		}
+	}
+
+	/**
+	 * Block the calling thread until the request is complete.
+	 *
+	 * Avoid waiting if possible and use callbacks to manage control flow instead.
+	 * Does not return until the callback for the request has finished executing.
+	 */
+	inline void Wait() const
+	{
+		if (Request)
+		{
+			Request->Wait();
+		}
+	}
+
+	/**
+	 * Poll the request and return true if complete and false otherwise.
+	 *
+	 * Avoid polling if possible and use callbacks to manage control flow instead.
+	 * A request is not considered to be complete until its callback has finished executing.
+	 */
+	inline bool Poll() const
+	{
+		return !Request || Request->Poll();
+	}
+
+private:
+	TRefCountPtr<ICacheRequest> Request;
+};
+
+///////////////////////////////////////////////////////////////////////////////////////////////////////////////////////
+
+/** Parameters for the completion callback for cache get requests. */
+struct FCacheGetCompleteParams
+{
+	/**
+	 * A cache record produced by the cache request.
+	 *
+	 * The key is always set, and the other fields of the record are set according to the flags used
+	 * when the cache request was made.
+	 */
+	FCacheRecord Record;
+
+	/** Status of the cache request. */
+	ECacheStatus Status = ECacheStatus::NotCached;
+};
+
+/** Parameters for the completion callback for cache put requests. */
+struct FCachePutCompleteParams
+{
+	/** Key used by the put request that completed. */
+	FCacheKey Key;
+
+	/** Status of the cache request. */
+	ECacheStatus Status = ECacheStatus::NotCached;
+};
+
+/** Parameters for the completion callback for cache attachment requests. */
+struct FCacheGetAttachmentCompleteParams
+{
+	/** Key used by the attachment request that completed. */
+	FCacheAttachmentKey Key;
+
+	/** Value of the attachment that was requested, if it was available. */
+	FSharedBuffer Value;
+
+	/** Status of the cache request. */
+	ECacheStatus Status = ECacheStatus::NotCached;
+};
+
+/** Callback for completion of a get request. */
+using FOnCacheGetComplete = TUniqueFunction<void (FCacheGetCompleteParams&& Params)>;
+
+/** Callback for completion of a put request. */
+using FOnCachePutComplete = TUniqueFunction<void (FCachePutCompleteParams&& Params)>;
+
+/** Callback for completion of an attachment request. */
+using FOnCacheGetAttachmentComplete = TUniqueFunction<void (FCacheGetAttachmentCompleteParams&& Params)>;
+
+///////////////////////////////////////////////////////////////////////////////////////////////////////////////////////
+
+/**
+ * Interface to the cache.
+ *
+ * Functions on this interface may be called from any thread. When a callback is provided, it may
+ * be invoked from a different thread than the request was made on, and may be invoked before the
+ * request function returns, both of which must be supported by the callback.
+ */
+class ICache
+{
+public:
+	virtual ~ICache() = default;
+
+	/**
+	 * Asynchronous request to get cache records according to the policy.
+	 *
+	 * The callback will always be called for every key, and may be called from an arbitrary thread.
+	 * Records may become available in any order, and from multiple threads concurrently.
+	 *
+	 * @param Keys The keys identifying the cache records to query.
+	 * @param Context A description of the request. An object path is typically sufficient.
+	 * @param Policy Flags to control the behavior of the request. See ECachePolicy.
+	 * @param Priority A priority to consider when scheduling the request. See ECachePriority.
+	 * @param Callback A callback invoked for every key in the batch as it completes or is canceled.
+	 */
+	virtual FCacheRequest Get(
+		TConstArrayView<FCacheKey> Keys,
+		FStringView Context,
+		ECachePolicy Policy,
+		ECachePriority Priority,
+		FOnCacheGetComplete&& Callback) = 0;
+
+	/**
+	 * Asynchronous request to put cache records according to the policy.
+	 *
+	 * The callback will always be called for every key, and may be called from an arbitrary thread.
+	 * Records may finish storing in any order, and from multiple threads concurrently.
+	 *
+	 * @param Records The cache records to store. Must have a key and value.
+	 * @param Context A description of the request. An object path is typically sufficient.
+	 * @param Policy Flags to control the behavior of the request. See ECachePolicy.
+	 * @param Priority A priority to consider when scheduling the request. See ECachePriority.
+	 * @param Callback A callback invoked for every key in the batch as it completes or is canceled.
+	 */
+	virtual FCacheRequest Put(
+		TConstArrayView<FCacheRecord> Records,
+		FStringView Context,
+		ECachePolicy Policy = ECachePolicy::Default,
+		ECachePriority Priority = ECachePriority::Normal,
+		FOnCachePutComplete&& Callback = FOnCachePutComplete()) = 0;
+
+	/**
+	 * Asynchronous request to get cache record attachments according to the policy.
+	 *
+	 * The callback will always be called for every key, and may be called from an arbitrary thread.
+	 * Attachments may become available in any order, and from multiple threads concurrently.
+	 *
+	 * @param Keys The keys identifying the cache record attachments to query.
+	 * @param Context A description of the request. An object path is typically sufficient.
+	 * @param Policy Flags to control the behavior of the request. See ECachePolicy.
+	 * @param Priority A priority to consider when scheduling the request. See ECachePriority.
+	 * @param Callback A callback invoked for every key in the batch as it completes or is canceled.
+	 */
+	virtual FCacheRequest GetAttachments(
+		TConstArrayView<FCacheAttachmentKey> Keys,
+		FStringView Context,
+		ECachePolicy Policy,
+		ECachePriority Priority,
+		FOnCacheGetAttachmentComplete&& Callback) = 0;
+
+	/**
+	 * Cancel all queued and active cache requests and invoke their callbacks.
+	 *
+	 * This is used to finish requests before destroying the cache.
+	 */
+	virtual void CancelAll() = 0;
+};
+
+///////////////////////////////////////////////////////////////////////////////////////////////////////////////////////
+
+} // DerivedData
+} // UE
diff --git a/Engine/Source/Developer/DerivedDataCache/Public/DerivedDataCacheInterface.h b/Engine/Source/Developer/DerivedDataCache/Public/DerivedDataCacheInterface.h
--- a/Engine/Source/Developer/DerivedDataCache/Public/DerivedDataCacheInterface.h
+++ b/Engine/Source/Developer/DerivedDataCache/Public/DerivedDataCacheInterface.h
@@ -10,6 +10,8 @@
 
 class FDerivedDataCacheUsageStats;
 
+namespace UE { namespace DerivedData { class ICache; } }
+
 /** 
  * Interface for the derived data cache
  * This API is fully threadsafe (with the possible exception of the system interface: NotfiyBootComplete, etc).
@@ -21,6 +23,9 @@ public:
 	{
 	}
 
+	/** Returns the interface to the cache. */
+	virtual UE::DerivedData::ICache& GetCache() = 0;
+
 	//--------------------
 	// High Level Interface
 	//--------------------
diff --git a/Engine/Source/Developer/DerivedDataCache/Public/DerivedDataCacheKey.h b/Engine/Source/Developer/DerivedDataCache/Public/DerivedDataCacheKey.h
new file mode 100644
--- /dev/null
+++ b/Engine/Source/Developer/DerivedDataCache/Public/DerivedDataCacheKey.h
@@ -0,0 +1,253 @@
+// Copyright Epic Games, Inc. All Rights Reserved.
+
+#pragma once
+
+#include "CoreTypes.h"
+#include "Containers/StringFwd.h"
+#include "Containers/StringView.h"
+#include "Misc/Blake3.h"
+#include "Templates/TypeHash.h"
+
+#define UE_API DERIVEDDATACACHE_API
+
+namespace UE
+{
+namespace DerivedData
+{
+
+///////////////////////////////////////////////////////////////////////////////////////////////////////////////////////
+
+/** An alphanumeric identifier that groups related cache records. */
+class FCacheBucket
+{
+public:
+	/** Construct a null cache bucket. */
+	FCacheBucket() = default;
+
+	/** Construct a cache bucket from the alphanumeric name. Names are case-insensitive. */
+	UE_API explicit FCacheBucket(FStringView Name);
+
+	/** Append the name of the cache bucket to the string builder. */
+	UE_API void ToString(FAnsiStringBuilderBase& Builder) const;
+	UE_API void ToString(FWideStringBuilderBase& Builder) const;
+
+	/** Whether the cache bucket is null. */
+	inline bool IsNull() const { return !Index; }
+	/** Whether the cache bucket is not null. */
+	inline explicit operator bool() const { return !IsNull(); }
+
+	/** Returns the index that is only stable within this process. */
+	inline uint32 ToIndex() const { return Index; }
+
+	/** Returns the bucket for the index produced within this process. */
+	static inline FCacheBucket FromIndex(uint32 InIndex)
+	{
+		FCacheBucket Bucket;
+		Bucket.Index = InIndex;
+		return Bucket;
+	}
+
+	/** Reset this to null. */
+	inline void Reset() { *this = FCacheBucket(); }
+
+private:
+	uint32 Index = 0;
+};
+
+inline bool operator==(FCacheBucket A, FCacheBucket B)
+{
+	return A.ToIndex() == B.ToIndex();
+}
+
+inline uint32 GetTypeHash(FCacheBucket Bucket)
+{
+	return ::GetTypeHash(Bucket.ToIndex());
+}
+
+/** Append the cache bucket name to the string builder. */
+template <typename CharType>
+inline TStringBuilderBase<CharType>& operator<<(TStringBuilderBase<CharType>& Builder, const FCacheBucket& Bucket)
+{
+	Bucket.ToString(Builder);
+	return Builder;
+}
+
+/** Fast less with an order that is only stable within the lifetime of this process. */
+struct FCacheBucketFastLess
+{
+	inline bool operator()(FCacheBucket A, FCacheBucket B) const
+	{
+		return A.ToIndex() < B.ToIndex();
+	}
+};
+
+/** Slow less with lexical order that is stable between processes. */
+struct FCacheBucketLexicalLess
+{
+	UE_API bool operator()(FCacheBucket A, FCacheBucket B) const;
+};
+
+///////////////////////////////////////////////////////////////////////////////////////////////////////////////////////
+
+/** A key that uniquely identifies a cache record. */
+class FCacheKey
+{
+public:
+	/** Construct a null cache key. */
+	FCacheKey() = default;
+
+	/** Construct a cache key from the non-null bucket and non-zero hash. */
+	inline FCacheKey(FCacheBucket InBucket, const FBlake3Hash& InHash)
+		: Bucket(InBucket)
+		, Hash(InHash)
+	{
+	}
+
+	/** Returns the bucket component of the cache key. */
+	inline FCacheBucket GetBucket() const { return Bucket; }
+
+	/** Returns the hash component of the cache key. */
+	inline const FBlake3Hash& GetHash() const { return Hash; }
+
+	/** Append the cache key to the string builder. */
+	UE_API void ToString(FAnsiStringBuilderBase& Builder) const;
+	UE_API void ToString(FWideStringBuilderBase& Builder) const;
+
+	/** Whether the cache key is null. */
+	inline bool IsNull() const { return Bucket.IsNull(); }
+	/** Whether the cache key is not null. */
+	inline explicit operator bool() const { return !IsNull(); }
+
+	/** Reset this to null. */
+	void Reset() { *this = FCacheKey(); }
+
+private:
+	FCacheBucket Bucket;
+	FBlake3Hash Hash;
+};
+
+inline bool operator==(const FCacheKey& A, const FCacheKey& B)
+{
+	return A.GetBucket() == B.GetBucket() && A.GetHash() == B.GetHash();
+}
+
+inline uint32 GetTypeHash(const FCacheKey& Key)
+{
+	return HashCombine(GetTypeHash(Key.GetBucket()), GetTypeHash(Key.GetHash()));
+}
+
+/** Append the cache key to the string builder. */
+template <typename CharType>
+inline TStringBuilderBase<CharType>& operator<<(TStringBuilderBase<CharType>& Builder, const FCacheKey& Key)
+{
+	Key.ToString(Builder);
+	return Builder;
+}
+
+/** Fast less with an order that is only stable within the lifetime of this process. */
+struct FCacheKeyFastLess
+{
+	inline bool operator()(const FCacheKey& A, const FCacheKey& B) const
+	{
+		const FCacheBucket BucketA = A.GetBucket();
+		const FCacheBucket BucketB = B.GetBucket();
+		return BucketA == BucketB ? A.GetHash() < B.GetHash() : FCacheBucketFastLess()(BucketA, BucketB);
+	}
+};
+
+/** Slow less with lexical order that is stable between processes. */
+struct FCacheKeyLexicalLess
+{
+	inline bool operator()(const FCacheKey& A, const FCacheKey& B) const
+	{
+		const FCacheBucket BucketA = A.GetBucket();
+		const FCacheBucket BucketB = B.GetBucket();
+		return BucketA == BucketB ? A.GetHash() < B.GetHash() : FCacheBucketLexicalLess()(BucketA, BucketB);
+	}
+};
+
+///////////////////////////////////////////////////////////////////////////////////////////////////////////////////////
+
+/** A key that uniquely identifies an attachment within a cache record. */
+class FCacheAttachmentKey
+{
+public:
+	/** Construct a null cache attachment key. */
+	FCacheAttachmentKey() = default;
+
+	/** Construct a cache attachment key from the non-null key and non-zero hash. */
+	inline FCacheAttachmentKey(const FCacheKey& InKey, const FBlake3Hash& InHash)
+		: Key(InKey)
+		, Hash(InHash)
+	{
+	}
+
+	/** Returns the cache key for the record containing the attachment. */
+	inline const FCacheKey& GetKey() const { return Key; }
+
+	/** Returns the hash of the attachment. */
+	inline const FBlake3Hash& GetHash() const { return Hash; }
+
+	/** Append the cache attachment key to the string builder. */
+	UE_API void ToString(FAnsiStringBuilderBase& Builder) const;
+	UE_API void ToString(FWideStringBuilderBase& Builder) const;
+
+	/** Whether the cache attachment key is null. */
+	inline bool IsNull() const { return Key.IsNull(); }
+	/** Whether the cache attachment key is not null. */
+	inline explicit operator bool() const { return !IsNull(); }
+
+	/** Reset this to null. */
+	void Reset() { *this = FCacheAttachmentKey(); }
+
+private:
+	FCacheKey Key;
+	FBlake3Hash Hash;
+};
+
+inline bool operator==(const FCacheAttachmentKey& A, const FCacheAttachmentKey& B)
+{
+	return A.GetKey() == B.GetKey() && A.GetHash() == B.GetHash();
+}
+
+inline uint32 GetTypeHash(const FCacheAttachmentKey& Key)
+{
+	return HashCombine(GetTypeHash(Key.GetKey()), GetTypeHash(Key.GetHash()));
+}
+
+/** Append the cache attachment key to the string builder. */
+template <typename CharType>
+inline TStringBuilderBase<CharType>& operator<<(TStringBuilderBase<CharType>& Builder, const FCacheAttachmentKey& Key)
+{
+	Key.ToString(Builder);
+	return Builder;
+}
+
+/** Fast less with an order that is only stable within the lifetime of this process. */
+struct FCacheAttachmentKeyFastLess
+{
+	inline bool operator()(const FCacheAttachmentKey& A, const FCacheAttachmentKey& B) const
+	{
+		const FCacheKey& KeyA = A.GetKey();
+		const FCacheKey& KeyB = B.GetKey();
+		return KeyA == KeyB ? A.GetHash() < B.GetHash() : FCacheKeyFastLess()(KeyA, KeyB);
+	}
+};
+
+/** Slow less with lexical order that is stable between processes. */
+struct FCacheAttachmentKeyLexicalLess
+{
+	inline bool operator()(const FCacheAttachmentKey& A, const FCacheAttachmentKey& B) const
+	{
+		const FCacheKey& KeyA = A.GetKey();
+		const FCacheKey& KeyB = B.GetKey();
+		return KeyA == KeyB ? A.GetHash() < B.GetHash() : FCacheKeyLexicalLess()(KeyA, KeyB);
+	}
+};
+
+///////////////////////////////////////////////////////////////////////////////////////////////////////////////////////
+
+} // DerivedData
+} // UE
+
+#undef UE_API
diff --git a/Engine/Source/Developer/DerivedDataCache/Public/DerivedDataCacheRecord.h b/Engine/Source/Developer/DerivedDataCache/Public/DerivedDataCacheRecord.h
new file mode 100644
--- /dev/null
+++ b/Engine/Source/Developer/DerivedDataCache/Public/DerivedDataCacheRecord.h
@@ -0,0 +1,82 @@
+// Copyright Epic Games, Inc. All Rights Reserved.
+
+#pragma once
+
+#include "CoreTypes.h"
+#include "DerivedDataCacheKey.h"
+#include "Memory/SharedBuffer.h"
+#include "Serialization/CompactBinary.h"
+#include "Serialization/CompactBinaryPackage.h"
+
+#define UE_API DERIVEDDATACACHE_API
+
+namespace UE
+{
+namespace DerivedData
+{
+
+///////////////////////////////////////////////////////////////////////////////////////////////////////////////////////
+
+/** The type of value stored in a cache record. */
+enum class ECacheRecordType : uint8
+{
+	None,
+	Binary,
+	Object,
+	Package,
+};
+
+///////////////////////////////////////////////////////////////////////////////////////////////////////////////////////
+
+/**
+ * A cache record is a key, metadata, and a value.
+ *
+ * The value may be one of binary, object, or package.
+ */
+class FCacheRecord
+{
+public:
+	/** Construct an empty cache record. */
+	UE_API FCacheRecord();
+
+	/** Reset to an empty record. */
+	UE_API void Reset();
+
+	inline const FCacheKey& GetKey() const { return Key; }
+	inline void SetKey(const FCacheKey& InKey) { Key = InKey; }
+
+	inline const FCbObjectRef& GetMeta() const { return Meta; }
+	inline void SetMeta(FCbObjectRef InMeta) { Meta = MoveTemp(InMeta); }
+
+	inline ECacheRecordType GetType() const { return Type; }
+
+	/** Access the value as binary. Returns a null buffer if not binary. */
+	inline const FSharedBuffer& AsBinary() const { return Binary; }
+	/** Access the value as an object. Returns an empty object if not an object. */
+	inline const FCbObjectRef& AsObject() const { return Object; }
+	/** Access the value as a package. Returns a null package if not a package. */
+	inline const FCbPackage& AsPackage() const { return Package; }
+
+	/** Set the value as binary. Removes any existing value of another type. */
+	UE_API void SetBinary(FSharedBuffer Value);
+	/** Set the value as an object. Removes any existing value of another type. */
+	UE_API void SetObject(FCbObjectRef Value);
+	/** Set the value as a package. Removes any existing value of another type. */
+	UE_API void SetPackage(FCbPackage Value);
+
+private:
+	FCacheKey Key;
+	FCbObjectRef Meta;
+
+	FSharedBuffer Binary;
+	FCbObjectRef Object;
+	FCbPackage Package;
+	ECacheRecordType Type = ECacheRecordType::None;
+};
+
+///////////////////////////////////////////////////////////////////////////////////////////////////////////////////////
+
+} // DerivedData
+} // UE
+
+#undef UE_API
