Do not delete transient data from Jupiter, it will clean itself up as needed.
Also made Jupiter a slow backend, as this generally is closer to the typical speed people have to it. We should make this automatic based on the latency we find from the health check call.

diff --git a/Engine/Source/Developer/DerivedDataCache/Private/HttpDerivedDataBackend.cpp b/Engine/Source/Developer/DerivedDataCache/Private/HttpDerivedDataBackend.cpp
--- a/Engine/Source/Developer/DerivedDataCache/Private/HttpDerivedDataBackend.cpp
+++ b/Engine/Source/Developer/DerivedDataCache/Private/HttpDerivedDataBackend.cpp
@@ -1006,7 +1006,7 @@ bool FHttpDerivedDataBackend::WouldCache(const TCHAR* CacheKey, TArrayView<const
 
 FHttpDerivedDataBackend::ESpeedClass FHttpDerivedDataBackend::GetSpeedClass()
 {
-	return ESpeedClass::Fast;
+	return ESpeedClass::Slow;
 }
 
 bool FHttpDerivedDataBackend::ApplyDebugOptions(FBackendDebugOptions& InOptions)
@@ -1334,6 +1334,10 @@ FDerivedDataBackendInterface::EPutStatus FHttpDerivedDataBackend::PutCachedData(
 
 void FHttpDerivedDataBackend::RemoveCachedData(const TCHAR* CacheKey, bool bTransient)
 {
+	// do not remove transient data as Jupiter
+	if (bTransient)
+		return;
+	
 	TRACE_CPUPROFILER_EVENT_SCOPE(HttpDDC_Remove);
 	FString Uri = FString::Printf(TEXT("api/v1/c/ddc/%s/%s/%s"), *Namespace, *DefaultBucket, CacheKey);
 	int64 ResponseCode = 0; uint32 Attempts = 0;
