Clarified comment on why transient data can be disabled.

diff --git a/Engine/Source/Developer/DerivedDataCache/Private/HttpDerivedDataBackend.cpp b/Engine/Source/Developer/DerivedDataCache/Private/HttpDerivedDataBackend.cpp
--- a/Engine/Source/Developer/DerivedDataCache/Private/HttpDerivedDataBackend.cpp
+++ b/Engine/Source/Developer/DerivedDataCache/Private/HttpDerivedDataBackend.cpp
@@ -1334,7 +1334,7 @@ FDerivedDataBackendInterface::EPutStatus FHttpDerivedDataBackend::PutCachedData(
 
 void FHttpDerivedDataBackend::RemoveCachedData(const TCHAR* CacheKey, bool bTransient)
 {
-	// do not remove transient data as Jupiter
+	// do not remove transient data as Jupiter does its own verification of the content and cleans itself up
 	if (bTransient)
 		return;
 	
