DDC: Removed ECachePolicy::Overwrite because there is not a valid use case for it

#rnx

diff --git a/Engine/Source/Developer/DerivedDataCache/Private/DerivedDataCache.cpp b/Engine/Source/Developer/DerivedDataCache/Private/DerivedDataCache.cpp
--- a/Engine/Source/Developer/DerivedDataCache/Private/DerivedDataCache.cpp
+++ b/Engine/Source/Developer/DerivedDataCache/Private/DerivedDataCache.cpp
@@ -1063,8 +1063,7 @@ void FCache::Put(
 	Data.SetNumUninitialized(static_cast<int32>(SaveSize));
 	Writer.Save(MakeMemoryView(Data));
 
-	const bool bPutEvenIfExists = EnumHasAnyFlags(Policy, ECachePolicy::Overwrite);
-	GetDerivedDataCacheRef().Put(*MakeRecordKey(Params.Key), Data, Context, bPutEvenIfExists);
+	GetDerivedDataCacheRef().Put(*MakeRecordKey(Params.Key), Data, Context);
 
 	UE_LOG(LogDerivedDataCache, VeryVerbose, TEXT("Cache: Put for %s from '%.*s'"),
 		*TToString<96>(Params.Key), Context.Len(), Context.GetData());
diff --git a/Engine/Source/Developer/DerivedDataCache/Public/DerivedDataCache.h b/Engine/Source/Developer/DerivedDataCache/Public/DerivedDataCache.h
--- a/Engine/Source/Developer/DerivedDataCache/Public/DerivedDataCache.h
+++ b/Engine/Source/Developer/DerivedDataCache/Public/DerivedDataCache.h
@@ -45,9 +45,6 @@ enum class ECachePolicy : uint8
 	/** Skip fetching any data for get requests, which will only check for existence. */
 	SkipAll         = SkipMeta | SkipValue | SkipAttachments,
 
-	/** Allow a put request to overwrite existing cache records. */
-	Overwrite       = 1 << 7,
-
 	/** Allow requests to query and store records in local caches. */
 	Local           = QueryLocal | StoreLocal,
 	/** Allow requests to query and store records in remote caches. */
