Added stats for Cloud based DCC ( ie Jupiter ) to DDC.Summary
Added BackEnd name to DDC.Summary
Added HasLocalDDC, HasSharedDDC, HasCloudDDC to DDC.Summary

diff --git a/Engine/Source/Developer/DerivedDataCache/Private/DerivedDataCache.cpp b/Engine/Source/Developer/DerivedDataCache/Private/DerivedDataCache.cpp
--- a/Engine/Source/Developer/DerivedDataCache/Private/DerivedDataCache.cpp
+++ b/Engine/Source/Developer/DerivedDataCache/Private/DerivedDataCache.cpp
@@ -27,6 +27,7 @@
 #include "Misc/CoreMisc.h"
 #include "Misc/ScopeExit.h"
 #include "Misc/StringBuilder.h"
+#include "Misc/CommandLine.h"
 #include "Serialization/CompactBinary.h"
 #include "Serialization/CompactBinaryValidation.h"
 #include "Serialization/CompactBinaryWriter.h"
@@ -84,6 +85,9 @@ namespace DerivedDataCacheCookStats
 			FString* LocalDDCKey = Keys.FindByPredicate([](const FString& Key) {return Key.Contains(TEXT(": FileSystem.")) && !Key.Contains(TEXT("//")); });
 			// look for a UNC path
 			FString* SharedDDCKey = Keys.FindByPredicate([](const FString& Key) {return Key.Contains(TEXT(": FileSystem.//")); });
+			// look for a Cloud path
+			FString* CloudDDCKey = Keys.FindByPredicate([](const FString& Key) {return Key.Contains(TEXT("0: HTTP")); });
+
 			if (RootKey)
 			{
 				const FDerivedDataCacheUsageStats& RootStats = DDCStats[*RootKey];
@@ -112,6 +116,14 @@ namespace DerivedDataCacheCookStats
 						SharedDDCStats.GetStats.GetAccumulatedValue(FCookStats::CallStats::EHitOrMiss::Hit, FCookStats::CallStats::EStatType::Counter, true) +
 						SharedDDCStats.GetStats.GetAccumulatedValue(FCookStats::CallStats::EHitOrMiss::Hit, FCookStats::CallStats::EStatType::Counter, false);
 				}
+				int64 CloudHits = 0;
+				if (CloudDDCKey)
+				{
+					const FDerivedDataCacheUsageStats& CloudDDCStats = DDCStats[*CloudDDCKey];
+					CloudHits =
+						CloudDDCStats.GetStats.GetAccumulatedValue(FCookStats::CallStats::EHitOrMiss::Hit, FCookStats::CallStats::EStatType::Counter, true) +
+						CloudDDCStats.GetStats.GetAccumulatedValue(FCookStats::CallStats::EHitOrMiss::Hit, FCookStats::CallStats::EStatType::Counter, false);
+				}
 
 				int64 TotalPutHits =
 					RootStats.PutStats.GetAccumulatedValue(FCookStats::CallStats::EHitOrMiss::Hit, FCookStats::CallStats::EStatType::Counter, true) +
@@ -122,11 +134,16 @@ namespace DerivedDataCacheCookStats
 				int64 TotalPuts = TotalPutHits + TotalPutMisses;
 
 				AddStat(TEXT("DDC.Summary"), FCookStatsManager::CreateKeyValueArray(
+					TEXT("BackEnd"), FDerivedDataBackend::Get().GetGraphName(),
+					TEXT("HasLocalCache"), LocalDDCKey != nullptr,
+					TEXT("HasSharedCache"), SharedDDCKey!=nullptr,
+					TEXT("HasCloudCache"), CloudDDCKey !=nullptr,
 					TEXT("TotalGetHits"), TotalGetHits,
 					TEXT("TotalGets"), TotalGets,
 					TEXT("TotalGetHitPct"), SafeDivide(TotalGetHits, TotalGets),
 					TEXT("LocalGetHitPct"), SafeDivide(LocalHits, TotalGets),
 					TEXT("SharedGetHitPct"), SafeDivide(SharedHits, TotalGets),
+					TEXT("CloudGetHitPct"), SafeDivide(CloudHits, TotalGets),
 					TEXT("OtherGetHitPct"), SafeDivide((TotalGetHits - LocalHits - SharedHits), TotalGets),
 					TEXT("GetMissPct"), SafeDivide(TotalGetMisses, TotalGets),
 					TEXT("TotalPutHits"), TotalPutHits,
