DDC: Added the value hash to FCacheRecord and made it easier to access object values

#rnx

diff --git a/Engine/Source/Developer/DerivedDataCache/Private/DerivedDataCache.cpp b/Engine/Source/Developer/DerivedDataCache/Private/DerivedDataCache.cpp
--- a/Engine/Source/Developer/DerivedDataCache/Private/DerivedDataCache.cpp
+++ b/Engine/Source/Developer/DerivedDataCache/Private/DerivedDataCache.cpp
@@ -932,7 +932,7 @@ void FCache::Get(
 		FCbPackage Package;
 		if (Value && Value.IsCompactBinary())
 		{
-			Package.SetObject(Value.AsCompactBinary().AsObjectRef(), Value.GetHash());
+			Package.SetObject(Value.AsCompactBinary().AsObjectRef(), ValueField.AsReference());
 		}
 		for (FCbAttachment Attachment : Attachments)
 		{
@@ -945,12 +945,12 @@ void FCache::Get(
 		// Object
 		if (Value.IsCompactBinary())
 		{
-			Params.Record.SetObject(Value.AsCompactBinary().AsObjectRef());
+			Params.Record.SetObject(Value.AsCompactBinary().AsObjectRef(), ValueField.AsReference());
 		}
 		// Binary
 		else
 		{
-			Params.Record.SetBinary(Value.AsBinary());
+			Params.Record.SetBinary(Value.AsBinary(), ValueField.AsReference());
 		}
 	}
 
@@ -1013,15 +1013,13 @@ void FCache::Put(
 	case ECacheRecordType::None:
 		break;
 	case ECacheRecordType::Binary:
-		Value = FCbAttachment(Record.AsBinary());
+		Value = FCbAttachment(Record.AsBinary(), Record.GetValueHash());
 		break;
 	case ECacheRecordType::Object:
-		Value = FCbAttachment(FCbFieldRefIterator::MakeSingle(Record.AsObject().AsFieldRef()));
-		break;
 	case ECacheRecordType::Package:
 		{
 			const FCbPackage& Package = Record.AsPackage();
-			Value = FCbAttachment(FCbFieldRefIterator::MakeSingle(Package.GetObject().AsFieldRef()));
+			Value = FCbAttachment(FCbFieldRefIterator::MakeSingle(Package.GetObject().AsFieldRef()), Package.GetObjectHash());
 			Attachments = Package.GetAttachments();
 		}
 		break;
diff --git a/Engine/Source/Developer/DerivedDataCache/Private/DerivedDataCacheRecord.cpp b/Engine/Source/Developer/DerivedDataCache/Private/DerivedDataCacheRecord.cpp
--- a/Engine/Source/Developer/DerivedDataCache/Private/DerivedDataCacheRecord.cpp
+++ b/Engine/Source/Developer/DerivedDataCache/Private/DerivedDataCacheRecord.cpp
@@ -18,24 +18,50 @@ void FCacheRecord::Reset()
 
 void FCacheRecord::SetBinary(FSharedBuffer Value)
 {
-	Object.Reset();
+	check(Value);
 	Package.Reset();
 	Binary = MoveTemp(Value);
+	BinaryHash = Binary.GetSize() ? FIoHash::HashBuffer(Value) : FIoHash();
+	Type = ECacheRecordType::Binary;
+}
+
+void FCacheRecord::SetBinary(FSharedBuffer Value, const FIoHash& ValueHash)
+{
+	check(Value);
+	if (Value.GetSize())
+	{
+		checkSlow(ValueHash == FIoHash::HashBuffer(Value));
+	}
+	else
+	{
+		checkfSlow(ValueHash.IsZero(), TEXT("A null or empty buffer must use a hash of zero."));
+	}
+	Package.Reset();
+	Binary = MoveTemp(Value);
+	BinaryHash = ValueHash;
 	Type = ECacheRecordType::Binary;
 }
 
 void FCacheRecord::SetObject(FCbObjectRef Value)
 {
 	Binary.Reset();
-	Package.Reset();
-	Object = MoveTemp(Value);
+	BinaryHash.Reset();
+	Package = FCbPackage(MoveTemp(Value));
+	Type = ECacheRecordType::Object;
+}
+
+void FCacheRecord::SetObject(FCbObjectRef Value, const FIoHash& ValueHash)
+{
+	Binary.Reset();
+	BinaryHash.Reset();
+	Package = FCbPackage(MoveTemp(Value), ValueHash);
 	Type = ECacheRecordType::Object;
 }
 
 void FCacheRecord::SetPackage(FCbPackage Value)
 {
 	Binary.Reset();
-	Object.Reset();
+	BinaryHash.Reset();
 	Package = MoveTemp(Value);
 	Type = ECacheRecordType::Package;
 }
diff --git a/Engine/Source/Developer/DerivedDataCache/Public/DerivedDataCacheRecord.h b/Engine/Source/Developer/DerivedDataCache/Public/DerivedDataCacheRecord.h
--- a/Engine/Source/Developer/DerivedDataCache/Public/DerivedDataCacheRecord.h
+++ b/Engine/Source/Developer/DerivedDataCache/Public/DerivedDataCacheRecord.h
@@ -29,9 +29,20 @@ enum class ECacheRecordType : uint8
 ///////////////////////////////////////////////////////////////////////////////////////////////////////////////////////
 
 /**
- * A cache record is a key, metadata, and a value.
+ * A cache record is a key, a value, attachments, and metadata.
  *
- * The value may be one of binary, object, or package.
+ * The key is expected to correspond uniquely to its value and attachments. The key should not be
+ * used with any other value and attachments. The metadata does not have the same requirement and
+ * may be used to persist details such as the machine that created the value or how much time was
+ * spent creating the value.
+ *
+ * The value may be binary or an object. Only object values may have attachments, and attachments
+ * are consumed by accessing the cache record as a package. It is valid to access an object as if
+ * it is a package with no attachments, or to access a package as though it is an object. This is
+ * not valid for binary values, which may not have attachments or be interpreted as another type.
+ *
+ * It is invalid for a package to include attachments that are not referenced by the object or by
+ * one of its referenced attachments.
  */
 class FCacheRecord
 {
@@ -52,30 +63,46 @@ public:
 
 	/** Access the value as binary. Returns a null buffer if not binary. */
 	inline const FSharedBuffer& AsBinary() const { return Binary; }
-	/** Access the value as an object. Returns an empty object if not an object. */
-	inline const FCbObjectRef& AsObject() const { return Object; }
-	/** Access the value as a package. Returns a null package if not a package. */
+	/** Access the value as an object. Returns an empty object if not an object or package. */
+	inline const FCbObjectRef& AsObject() const { return Package.GetObject(); }
+	/** Access the value as a package. Returns a null package if not an object or package. */
 	inline const FCbPackage& AsPackage() const { return Package; }
 
-	/** Set the value as binary. Removes any existing value of another type. */
+	/** Returns the hash of the value, or the zero hash if there is no value. */
+	inline const FIoHash& GetValueHash() const;
+
+	/** Set the value as binary. Removes any existing value and attachments. */
 	UE_API void SetBinary(FSharedBuffer Value);
-	/** Set the value as an object. Removes any existing value of another type. */
+	/** Set the value as binary with its hash. Removes any existing value and attachments. */
+	UE_API void SetBinary(FSharedBuffer Value, const FIoHash& ValueHash);
+
+	/** Set the value as an object. Removes any existing value and attachments. */
 	UE_API void SetObject(FCbObjectRef Value);
-	/** Set the value as a package. Removes any existing value of another type. */
+	/** Set the value as an object with its hash. Removes any existing value and attachments. */
+	UE_API void SetObject(FCbObjectRef Value, const FIoHash& ValueHash);
+
+	/** Set the value as a package. Removes any existing value and attachments. */
 	UE_API void SetPackage(FCbPackage Value);
 
 private:
 	FCacheKey Key;
 	FCbObjectRef Meta;
 
-	FSharedBuffer Binary;
-	FCbObjectRef Object;
 	FCbPackage Package;
+	FSharedBuffer Binary;
+	FIoHash BinaryHash;
 	ECacheRecordType Type = ECacheRecordType::None;
 };
 
 ///////////////////////////////////////////////////////////////////////////////////////////////////////////////////////
 
+inline const FIoHash& FCacheRecord::GetValueHash() const
+{
+	return Type == ECacheRecordType::Binary ? BinaryHash : Package.GetObjectHash();
+}
+
+///////////////////////////////////////////////////////////////////////////////////////////////////////////////////////
+
 } // DerivedData
 } // UE
 
