DDC: Fixed an error in FCacheRecord::SetBinary

#rnx

diff --git a/Engine/Source/Developer/DerivedDataCache/Private/DerivedDataCacheRecord.cpp b/Engine/Source/Developer/DerivedDataCache/Private/DerivedDataCacheRecord.cpp
--- a/Engine/Source/Developer/DerivedDataCache/Private/DerivedDataCacheRecord.cpp
+++ b/Engine/Source/Developer/DerivedDataCache/Private/DerivedDataCacheRecord.cpp
@@ -21,7 +21,7 @@ void FCacheRecord::SetBinary(FSharedBuffer Value)
 	check(Value);
 	Package.Reset();
 	Binary = MoveTemp(Value);
-	BinaryHash = Binary.GetSize() ? FIoHash::HashBuffer(Value) : FIoHash();
+	BinaryHash = Binary.GetSize() ? FIoHash::HashBuffer(Binary) : FIoHash();
 	Type = ECacheRecordType::Binary;
 }
 
