Add the DDC Long -> Short key in the log so we can easily find the mapping between both when analysing a DDC problem

diff --git a/Engine/Source/Developer/DerivedDataCache/Private/DerivedDataLimitKeyLengthWrapper.h b/Engine/Source/Developer/DerivedDataCache/Private/DerivedDataLimitKeyLengthWrapper.h
--- a/Engine/Source/Developer/DerivedDataCache/Private/DerivedDataLimitKeyLengthWrapper.h
+++ b/Engine/Source/Developer/DerivedDataCache/Private/DerivedDataLimitKeyLengthWrapper.h
@@ -269,6 +269,7 @@ private:
 		int32 OriginalPart = MaxKeyLength - HashStringSize - 2;
 		Result = Result.Left(OriginalPart) + TEXT("__") + HashString;
 		check(Result.Len() == MaxKeyLength && Result.Len() > 0);
+		UE_LOG(LogDerivedDataCache, VeryVerbose, TEXT("ShortenKey %s -> %s"), CacheKey, *Result);
 		return true;
 	}
 
