DDC: Added TryToPrefetch to check multiple keys for existence and prefetch the data into a fast cache backend

diff --git a/Engine/Source/Developer/DerivedDataCache/Private/DerivedDataBackendAsyncPutWrapper.cpp b/Engine/Source/Developer/DerivedDataCache/Private/DerivedDataBackendAsyncPutWrapper.cpp
--- a/Engine/Source/Developer/DerivedDataCache/Private/DerivedDataBackendAsyncPutWrapper.cpp
+++ b/Engine/Source/Developer/DerivedDataCache/Private/DerivedDataBackendAsyncPutWrapper.cpp
@@ -182,21 +182,18 @@ TBitArray<> FDerivedDataBackendAsyncPutWrapper::CachedDataProbablyExistsBatch(TC
 	return Result;
 }
 
-bool FDerivedDataBackendAsyncPutWrapper::TryToPrefetch(const TCHAR* CacheKey)
+bool FDerivedDataBackendAsyncPutWrapper::TryToPrefetch(TConstArrayView<FString> CacheKeys)
 {
 	COOK_STAT(auto Timer = UsageStats.TimePrefetch());
 
-	bool SkipCheck = (!InflightCache && InflightCache->CachedDataProbablyExists(CacheKey));
-	
-	bool Hit = false;
-
-	if (!SkipCheck)
+	if ((InflightCache && InflightCache->CachedDataProbablyExistsBatch(CacheKeys).CountSetBits() == CacheKeys.Num()) ||
+		InnerBackend->TryToPrefetch(CacheKeys))
 	{
-		Hit = InnerBackend->TryToPrefetch(CacheKey);
+		COOK_STAT(Timer.AddHit(0));
+		return true;
 	}
 
-	COOK_STAT(if (Hit) { Timer.AddHit(0); });
-	return Hit;
+	return false;
 }
 
 /*
diff --git a/Engine/Source/Developer/DerivedDataCache/Private/DerivedDataBackendAsyncPutWrapper.h b/Engine/Source/Developer/DerivedDataCache/Private/DerivedDataBackendAsyncPutWrapper.h
--- a/Engine/Source/Developer/DerivedDataCache/Private/DerivedDataBackendAsyncPutWrapper.h
+++ b/Engine/Source/Developer/DerivedDataCache/Private/DerivedDataBackendAsyncPutWrapper.h
@@ -92,11 +92,12 @@ public:
 	virtual TBitArray<> CachedDataProbablyExistsBatch(TConstArrayView<FString> CacheKeys) override;
 
 	/**
-	 * Attempts to make sure the cached data will be available as optimally as possible. This is left up to the implementation to do
-	 * @param	CacheKey	Alphanumeric+underscore key of this cache item
-	 * @return				true if any steps were performed to optimize future retrieval
+	 * Attempt to make sure the cached data will be available as optimally as possible.
+	 *
+	 * @param	CacheKeys	Alphanumeric+underscore keys of the cache items
+	 * @return				true if the data will probably be found in a fast backend on a future request.
 	 */
-	virtual bool TryToPrefetch(const TCHAR* CacheKey) override;
+	virtual bool TryToPrefetch(TConstArrayView<FString> CacheKeys) override;
 
 	/**
 	 * Allows the DDC backend to determine if it wants to cache the provided data. Reasons for returning false could be a slow connection,
diff --git a/Engine/Source/Developer/DerivedDataCache/Private/DerivedDataBackendCorruptionWrapper.h b/Engine/Source/Developer/DerivedDataCache/Private/DerivedDataBackendCorruptionWrapper.h
--- a/Engine/Source/Developer/DerivedDataCache/Private/DerivedDataBackendCorruptionWrapper.h
+++ b/Engine/Source/Developer/DerivedDataCache/Private/DerivedDataBackendCorruptionWrapper.h
@@ -230,9 +230,9 @@ public:
 		});
 	}
 
-	virtual bool TryToPrefetch(const TCHAR* CacheKey) override
+	virtual bool TryToPrefetch(TConstArrayView<FString> CacheKeys) override
 	{
-		return InnerBackend->TryToPrefetch(CacheKey);
+		return InnerBackend->TryToPrefetch(CacheKeys);
 	}
 
 	virtual bool WouldCache(const TCHAR* CacheKey, TArrayView<const uint8> InData) override
diff --git a/Engine/Source/Developer/DerivedDataCache/Private/DerivedDataBackendVerifyWrapper.h b/Engine/Source/Developer/DerivedDataCache/Private/DerivedDataBackendVerifyWrapper.h
--- a/Engine/Source/Developer/DerivedDataCache/Private/DerivedDataBackendVerifyWrapper.h
+++ b/Engine/Source/Developer/DerivedDataCache/Private/DerivedDataBackendVerifyWrapper.h
@@ -142,9 +142,9 @@ public:
 		});
 	}
 
-	virtual bool TryToPrefetch(const TCHAR* CacheKey) override
+	virtual bool TryToPrefetch(TConstArrayView<FString> CacheKeys) override
 	{
-		return InnerBackend->TryToPrefetch(CacheKey);
+		return InnerBackend->TryToPrefetch(CacheKeys);
 	}
 
 	virtual bool WouldCache(const TCHAR* CacheKey, TArrayView<const uint8> InData) override
diff --git a/Engine/Source/Developer/DerivedDataCache/Private/DerivedDataCache.cpp b/Engine/Source/Developer/DerivedDataCache/Private/DerivedDataCache.cpp
--- a/Engine/Source/Developer/DerivedDataCache/Private/DerivedDataCache.cpp
+++ b/Engine/Source/Developer/DerivedDataCache/Private/DerivedDataCache.cpp
@@ -572,6 +572,18 @@ public:
 		return CacheKeys.Num() == 0 || CachedDataProbablyExistsBatch(CacheKeys).CountSetBits() == CacheKeys.Num();
 	}
 
+	virtual bool TryToPrefetch(TConstArrayView<FString> CacheKeys, FStringView DebugContext) override
+	{
+		if (CacheKeys.Num() > 1)
+		{
+			DDC_SCOPE_CYCLE_COUNTER(DDC_TryToPrefetch);
+			UE_LOG(LogDerivedDataCache, VeryVerbose, TEXT("TryToPrefetch %d keys including %s from '%.*s'"),
+				CacheKeys.Num(), *CacheKeys[0], DebugContext.Len(), DebugContext.GetData());
+			return FDerivedDataBackend::Get().GetRoot().TryToPrefetch(CacheKeys);
+		}
+		return true;
+	}
+
 	void NotifyBootComplete() override
 	{
 		DDC_SCOPE_CYCLE_COUNTER(DDC_NotifyBootComplete);
diff --git a/Engine/Source/Developer/DerivedDataCache/Private/DerivedDataLimitKeyLengthWrapper.h b/Engine/Source/Developer/DerivedDataCache/Private/DerivedDataLimitKeyLengthWrapper.h
--- a/Engine/Source/Developer/DerivedDataCache/Private/DerivedDataLimitKeyLengthWrapper.h
+++ b/Engine/Source/Developer/DerivedDataCache/Private/DerivedDataLimitKeyLengthWrapper.h
@@ -92,12 +92,16 @@ public:
 	 * @param	CacheKey	Alphanumeric+underscore key of this cache item
 	 * @return				true if any steps were performed to optimize future retrieval
 	 */
-	virtual bool TryToPrefetch(const TCHAR* CacheKey) override
+	virtual bool TryToPrefetch(TConstArrayView<FString> CacheKeys) override
 	{
 		COOK_STAT(auto Timer = UsageStats.TimePrefetch());
-		FString NewKey;
-		ShortenKey(CacheKey, NewKey);
-		bool Result = InnerBackend->TryToPrefetch(*NewKey);
+		TArray<FString> NewKeys;
+		NewKeys.Reserve(CacheKeys.Num());
+		for (const FString& CacheKey : CacheKeys)
+		{
+			ShortenKey(*CacheKey, NewKeys.Emplace_GetRef());
+		}
+		bool Result = InnerBackend->TryToPrefetch(NewKeys);
 		if (Result)
 		{
 			COOK_STAT(Timer.AddHit(0));
diff --git a/Engine/Source/Developer/DerivedDataCache/Private/FileSystemDerivedDataBackend.cpp b/Engine/Source/Developer/DerivedDataCache/Private/FileSystemDerivedDataBackend.cpp
--- a/Engine/Source/Developer/DerivedDataCache/Private/FileSystemDerivedDataBackend.cpp
+++ b/Engine/Source/Developer/DerivedDataCache/Private/FileSystemDerivedDataBackend.cpp
@@ -756,7 +756,7 @@ public:
 		COOK_STAT(UsageStatsMap.Add(FString::Printf(TEXT("%s: %s.%s"), *GraphPath, TEXT("FileSystem"), *CachePath), UsageStats));
 	}
 
-	bool TryToPrefetch(const TCHAR* CacheKey) override
+	bool TryToPrefetch(TConstArrayView<FString> CacheKeys) override
 	{
 		return false;
 	}
diff --git a/Engine/Source/Developer/DerivedDataCache/Private/HierarchicalDerivedDataBackend.h b/Engine/Source/Developer/DerivedDataCache/Private/HierarchicalDerivedDataBackend.h
--- a/Engine/Source/Developer/DerivedDataCache/Private/HierarchicalDerivedDataBackend.h
+++ b/Engine/Source/Developer/DerivedDataCache/Private/HierarchicalDerivedDataBackend.h
@@ -198,48 +198,66 @@ public:
 	}
 
 	/**
-	 * Attempts to make sure the cached data will be available as optimally as possible. This is left up to the implementation to do
-	 * @param	CacheKey	Alphanumeric+underscore key of this cache item
-	 * @return				true if any steps were performed to optimize future retrieval
+	 * Attempt to make sure the cached data will be available as optimally as possible.
+	 *
+	 * @param	CacheKeys	Alphanumeric+underscore keys of the cache items
+	 * @return				true if the data will probably be found in a fast backend on a future request.
 	 */
-	virtual bool TryToPrefetch(const TCHAR* CacheKey) override
+	virtual bool TryToPrefetch(TConstArrayView<FString> CacheKeys) override
 	{
 		COOK_STAT(auto Timer = UsageStats.TimePrefetch());
 
-		// Search all backends for this key. If it can be moved into a faster class then we'll do so.
-		bool WorthFetching = false;
+		TArray<FString, TInlineAllocator<16>> SearchKeys(CacheKeys);
 
-		FDerivedDataBackendInterface* LastMissedInterface = nullptr;
+		bool bHasFastBackendToWrite = false;
+		bool bHasSlowBackend = false;
 
-		for (int32 CacheIndex = 0; CacheIndex < InnerBackends.Num(); CacheIndex++)
+		for (FDerivedDataBackendInterface* Interface : InnerBackends)
 		{
-			FDerivedDataBackendInterface* Interface = InnerBackends[CacheIndex];
-
-			if (!Interface->CachedDataProbablyExists(CacheKey) && Interface->IsWritable())
+			if (Interface->GetSpeedClass() < ESpeedClass::Fast)
 			{
-				LastMissedInterface = Interface;
+				bHasSlowBackend = true;
 			}
 			else
 			{
-				// if we have an interface that's writable and faster, lets get it
-				if (LastMissedInterface && LastMissedInterface->GetSpeedClass() > Interface->GetSpeedClass())
+				bHasFastBackendToWrite = bHasFastBackendToWrite || Interface->IsWritable();
+
+				// Remove keys that exist in a fast backend because we already have fast access to them.
+				TBitArray<> Hits = Interface->CachedDataProbablyExistsBatch(SearchKeys);
+				TArray<FString, TInlineAllocator<16>>::TIterator It = SearchKeys.CreateIterator();
+				It.SetToEnd();
+				for (--It; It; --It)
 				{
-					WorthFetching = true;
+					if (Hits[It.GetIndex()])
+					{
+						It.RemoveCurrent();
+					}
+				}
+
+				// No fetch is necessary if every key already exists in a fast backend.
+				if (SearchKeys.IsEmpty())
+				{
+					COOK_STAT(Timer.AddHit(0));
+					return true;
 				}
 			}
 		}
-		
-		// If it's remote then fetch it. We don't care about the data but we 
-		// Need to read a copy from the remote store anyway to fill the caches
-		if (WorthFetching)
+
+		// Try to fetch remaining keys, which will fill them from slow backends into writable fast backends.
+		bool bHit = true;
+		int64 BytesFetched = 0;
+		if (bHasSlowBackend && bHasFastBackendToWrite)
 		{
-			TArray<uint8> DontCare;
-			GetCachedData(CacheKey, DontCare);
-			COOK_STAT(Timer.AddHit(0));
-		}			
+			for (const FString& CacheKey : SearchKeys)
+			{
+				TArray<uint8> Ignored;
+				bHit &= GetCachedData(*CacheKey, Ignored);
+				BytesFetched += Ignored.Num();
+			}
+		}
 
-		// Return true if we did anything
-		return WorthFetching;
+		COOK_STAT(if (bHit) { Timer.AddHit(BytesFetched); });
+		return bHit;
 	}
 
 	/*
diff --git a/Engine/Source/Developer/DerivedDataCache/Private/HttpDerivedDataBackend.cpp b/Engine/Source/Developer/DerivedDataCache/Private/HttpDerivedDataBackend.cpp
--- a/Engine/Source/Developer/DerivedDataCache/Private/HttpDerivedDataBackend.cpp
+++ b/Engine/Source/Developer/DerivedDataCache/Private/HttpDerivedDataBackend.cpp
@@ -994,7 +994,7 @@ FString FHttpDerivedDataBackend::GetName() const
 	return Domain;
 }
 
-bool FHttpDerivedDataBackend::TryToPrefetch(const TCHAR* CacheKey)
+bool FHttpDerivedDataBackend::TryToPrefetch(TConstArrayView<FString> CacheKeys)
 {
 	return false;
 }
diff --git a/Engine/Source/Developer/DerivedDataCache/Private/HttpDerivedDataBackend.h b/Engine/Source/Developer/DerivedDataCache/Private/HttpDerivedDataBackend.h
--- a/Engine/Source/Developer/DerivedDataCache/Private/HttpDerivedDataBackend.h
+++ b/Engine/Source/Developer/DerivedDataCache/Private/HttpDerivedDataBackend.h
@@ -58,7 +58,7 @@ public:
 	virtual void GatherUsageStats(TMap<FString, FDerivedDataCacheUsageStats>& UsageStatsMap, FString&& GraphPath) override;
 	
 	virtual FString GetName() const override;
-	virtual bool TryToPrefetch(const TCHAR* CacheKey) override;
+	virtual bool TryToPrefetch(TConstArrayView<FString> CacheKeys) override;
 	virtual bool WouldCache(const TCHAR* CacheKey, TArrayView<const uint8> InData) override;
 	virtual ESpeedClass GetSpeedClass() override;
 	virtual bool ApplyDebugOptions(FBackendDebugOptions& InOptions) override;
diff --git a/Engine/Source/Developer/DerivedDataCache/Private/MemoryDerivedDataBackend.cpp b/Engine/Source/Developer/DerivedDataCache/Private/MemoryDerivedDataBackend.cpp
--- a/Engine/Source/Developer/DerivedDataCache/Private/MemoryDerivedDataBackend.cpp
+++ b/Engine/Source/Developer/DerivedDataCache/Private/MemoryDerivedDataBackend.cpp
@@ -86,7 +86,7 @@ bool FMemoryDerivedDataBackend::GetCachedData(const TCHAR* CacheKey, TArray<uint
 	return false;
 }
 
-bool FMemoryDerivedDataBackend::TryToPrefetch(const TCHAR* CacheKey)
+bool FMemoryDerivedDataBackend::TryToPrefetch(TConstArrayView<FString> CacheKeys)
 {
 	return false;
 }
diff --git a/Engine/Source/Developer/DerivedDataCache/Private/MemoryDerivedDataBackend.h b/Engine/Source/Developer/DerivedDataCache/Private/MemoryDerivedDataBackend.h
--- a/Engine/Source/Developer/DerivedDataCache/Private/MemoryDerivedDataBackend.h
+++ b/Engine/Source/Developer/DerivedDataCache/Private/MemoryDerivedDataBackend.h
@@ -78,7 +78,7 @@ public:
 
 	virtual void GatherUsageStats(TMap<FString, FDerivedDataCacheUsageStats>& UsageStatsMap, FString&& GraphPath) override;
 
-	virtual bool TryToPrefetch(const TCHAR* CacheKey) override;
+	virtual bool TryToPrefetch(TConstArrayView<FString> CacheKeys) override;
 
 	/**
 	 *  Determines if we would cache the provided data
diff --git a/Engine/Source/Developer/DerivedDataCache/Private/PakFileDerivedDataBackend.h b/Engine/Source/Developer/DerivedDataCache/Private/PakFileDerivedDataBackend.h
--- a/Engine/Source/Developer/DerivedDataCache/Private/PakFileDerivedDataBackend.h
+++ b/Engine/Source/Developer/DerivedDataCache/Private/PakFileDerivedDataBackend.h
@@ -92,7 +92,7 @@ public:
 
 	virtual void GatherUsageStats(TMap<FString, FDerivedDataCacheUsageStats>& UsageStatsMap, FString&& GraphPath) override;
 
-	virtual bool TryToPrefetch(const TCHAR* CacheKey) override { return false; }
+	virtual bool TryToPrefetch(TConstArrayView<FString> CacheKeys) override { return false; }
 
 	virtual bool WouldCache(const TCHAR* CacheKey, TArrayView<const uint8> InData) override { return true; }
 
diff --git a/Engine/Source/Developer/DerivedDataCache/Private/S3DerivedDataBackend.cpp b/Engine/Source/Developer/DerivedDataCache/Private/S3DerivedDataBackend.cpp
--- a/Engine/Source/Developer/DerivedDataCache/Private/S3DerivedDataBackend.cpp
+++ b/Engine/Source/Developer/DerivedDataCache/Private/S3DerivedDataBackend.cpp
@@ -947,7 +947,7 @@ FDerivedDataBackendInterface::ESpeedClass FS3DerivedDataBackend::GetSpeedClass()
 	return ESpeedClass::Local;
 }
 
-bool FS3DerivedDataBackend::TryToPrefetch(const TCHAR* CacheKey)
+bool FS3DerivedDataBackend::TryToPrefetch(TConstArrayView<FString> CacheKeys)
 {
 	return false;
 }
diff --git a/Engine/Source/Developer/DerivedDataCache/Private/S3DerivedDataBackend.h b/Engine/Source/Developer/DerivedDataCache/Private/S3DerivedDataBackend.h
--- a/Engine/Source/Developer/DerivedDataCache/Private/S3DerivedDataBackend.h
+++ b/Engine/Source/Developer/DerivedDataCache/Private/S3DerivedDataBackend.h
@@ -53,7 +53,7 @@ public:
 
 	FString GetName() const override;
 	ESpeedClass GetSpeedClass() override;
-	bool TryToPrefetch(const TCHAR* CacheKey) override;
+	bool TryToPrefetch(TConstArrayView<FString> CacheKeys) override;
 	bool WouldCache(const TCHAR* CacheKey, TArrayView<const uint8> InData) override;
 
 	bool ApplyDebugOptions(FBackendDebugOptions& InOptions) override;
diff --git a/Engine/Source/Developer/DerivedDataCache/Public/DerivedDataBackendInterface.h b/Engine/Source/Developer/DerivedDataCache/Public/DerivedDataBackendInterface.h
--- a/Engine/Source/Developer/DerivedDataCache/Public/DerivedDataBackendInterface.h
+++ b/Engine/Source/Developer/DerivedDataCache/Public/DerivedDataBackendInterface.h
@@ -164,12 +164,12 @@ public:
 	virtual void GatherUsageStats(TMap<FString, FDerivedDataCacheUsageStats>& UsageStatsMap, FString&& GraphPath) = 0;
 
 	/**
-	 * Synchronous attempt to make sure the cached data will be available as optimally as possible. This is left up to the implementation 
-	 * to implement or just skip.
-	 * @param	CacheKey	Alphanumeric+underscore key of this cache item
-	 * @return				true if any steps were performed to optimize future retrieval
+	 * Synchronous attempt to make sure the cached data will be available as optimally as possible.
+	 *
+	 * @param	CacheKeys	Alphanumeric+underscore keys of the cache items
+	 * @return				true if the data will probably be found in a fast backend on a future request.
 	 */
-	virtual bool TryToPrefetch(const TCHAR* CacheKey) = 0;
+	virtual bool TryToPrefetch(TConstArrayView<FString> CacheKeys) = 0;
 
 	/**
 	 * Allows the DDC backend to determine if it wants to cache the provided data. Reasons for returning false could be a slow connection,
diff --git a/Engine/Source/Developer/DerivedDataCache/Public/DerivedDataCacheInterface.h b/Engine/Source/Developer/DerivedDataCache/Public/DerivedDataCacheInterface.h
--- a/Engine/Source/Developer/DerivedDataCache/Public/DerivedDataCacheInterface.h
+++ b/Engine/Source/Developer/DerivedDataCache/Public/DerivedDataCacheInterface.h
@@ -216,6 +216,15 @@ public:
 	 */
 	virtual bool AllCachedDataProbablyExists(TConstArrayView<FString> CacheKeys) = 0;
 
+	/**
+	 * Synchronous attempt to make sure the cached data will be available as optimally as possible.
+	 *
+	 * @param	CacheKeys		Keys to identify the data.
+	 * @param	DebugContext	A string used to describe the data being generated. Typically the path to the object that it is generated from is sufficient.
+	 * @return					true if the data will probably be found in a fast backend on a future request.
+	 */
+	virtual bool TryToPrefetch(TConstArrayView<FString> CacheKeys, FStringView DebugContext) = 0;
+
 	//--------------------
 	// System Interface
 	//--------------------
