Move static object out of global namespace.

#rnx

diff --git a/Engine/Source/Developer/DerivedDataCache/Private/DerivedDataCache.cpp b/Engine/Source/Developer/DerivedDataCache/Private/DerivedDataCache.cpp
--- a/Engine/Source/Developer/DerivedDataCache/Private/DerivedDataCache.cpp
+++ b/Engine/Source/Developer/DerivedDataCache/Private/DerivedDataCache.cpp
@@ -675,19 +675,29 @@ class FDerivedDataCacheModule : public IDerivedDataCacheModule
 public:
 	virtual FDerivedDataCacheInterface& GetDDC() override
 	{
-		static FDerivedDataCache SingletonInstance;
-		return SingletonInstance;
+		check(Instance);
+		return *Instance;
 	}
 
 	virtual void StartupModule() override
 	{
-		GetDDC();
+		if (!Instance)
+		{
+			Instance = MakeUnique<FDerivedDataCache>();
+		}
+		check(Instance);
 	}
 
 	virtual void ShutdownModule() override
 	{
 		FDDCCleanup::Shutdown();
+
+		check(Instance);
+		Instance.Release();
 	}
+
+private:
+	TUniquePtr<FDerivedDataCache> Instance;
 };
 
 IMPLEMENT_MODULE( FDerivedDataCacheModule, DerivedDataCache);
