DDC: Fixed TryToPrefetch of one key, which was skipped due to an incorrect copy from CachedDataProbablyExistsBatch

#rnx

diff --git a/Engine/Source/Developer/DerivedDataCache/Private/DerivedDataCache.cpp b/Engine/Source/Developer/DerivedDataCache/Private/DerivedDataCache.cpp
--- a/Engine/Source/Developer/DerivedDataCache/Private/DerivedDataCache.cpp
+++ b/Engine/Source/Developer/DerivedDataCache/Private/DerivedDataCache.cpp
@@ -574,7 +574,7 @@ public:
 
 	virtual bool TryToPrefetch(TConstArrayView<FString> CacheKeys, FStringView DebugContext) override
 	{
-		if (CacheKeys.Num() > 1)
+		if (!CacheKeys.IsEmpty())
 		{
 			DDC_SCOPE_CYCLE_COUNTER(DDC_TryToPrefetch);
 			UE_LOG(LogDerivedDataCache, VeryVerbose, TEXT("TryToPrefetch %d keys including %s from '%.*s'"),
