DDC: Updated the new cache interface to expose compression of cache records

- FCacheRecord must now be created by FCacheRecordBuilder.
- FCacheRecord values are always an opaque buffer now, and ECacheRecordType has been removed.
- FCachePayload has been added as a common type for the cache record value and attachments.
- FCachePayloadId has been added as a 12-byte identifier for cache payloads.
- ICache::GetAttachments is now GetPayloads and can fetch both values and attachments.
- ECachePriority::Blocking has been added for requests that the caller will block on.

#rnx

diff --git a/Engine/Source/Developer/DerivedDataCache/Private/DerivedDataCache.cpp b/Engine/Source/Developer/DerivedDataCache/Private/DerivedDataCache.cpp
--- a/Engine/Source/Developer/DerivedDataCache/Private/DerivedDataCache.cpp
+++ b/Engine/Source/Developer/DerivedDataCache/Private/DerivedDataCache.cpp
@@ -716,50 +716,50 @@ class FCache final : public ICache
 public:
 	virtual ~FCache() = default;
 
-	virtual FCacheRequest Get(
-		TConstArrayView<FCacheKey> Keys,
+	virtual FCacheRequest Put(
+		TArrayView<FCacheRecord> Records,
 		FStringView Context,
 		ECachePolicy Policy,
 		ECachePriority Priority,
-		FOnCacheGetComplete&& Callback) final;
+		FOnCachePutComplete&& Callback) final;
 
-	virtual FCacheRequest Put(
-		TConstArrayView<FCacheRecord> Records,
+	virtual FCacheRequest Get(
+		TConstArrayView<FCacheKey> Keys,
 		FStringView Context,
 		ECachePolicy Policy,
 		ECachePriority Priority,
-		FOnCachePutComplete&& Callback) final;
+		FOnCacheGetComplete&& Callback) final;
 
-	virtual FCacheRequest GetAttachments(
-		TConstArrayView<FCacheAttachmentKey> Keys,
+	virtual FCacheRequest GetPayloads(
+		TConstArrayView<FCachePayloadKey> Keys,
 		FStringView Context,
 		ECachePolicy Policy,
 		ECachePriority Priority,
-		FOnCacheGetAttachmentComplete&& Callback) final;
+		FOnCacheGetPayloadComplete&& Callback) final;
 
 	virtual void CancelAll() final;
 
 private:
-	void Get(
-		const FCacheKey& Key,
+	void Put(
+		FCacheRecord Record,
 		FStringView Context,
 		ECachePolicy Policy,
 		ECachePriority Priority,
-		FOnCacheGetComplete& Callback);
+		FOnCachePutComplete& Callback);
 
-	void Put(
-		const FCacheRecord& Record,
+	void Get(
+		const FCacheKey& Key,
 		FStringView Context,
 		ECachePolicy Policy,
 		ECachePriority Priority,
-		FOnCachePutComplete& Callback);
+		FOnCacheGetComplete& Callback);
 
-	void GetAttachment(
-		const FCacheAttachmentKey& Key,
+	void GetPayload(
+		const FCachePayloadKey& Key,
 		FStringView Context,
 		ECachePolicy Policy,
 		ECachePriority Priority,
-		FOnCacheGetAttachmentComplete& Callback);
+		FOnCacheGetPayloadComplete& Callback);
 
 	template <int32 BufferSize>
 	class TToString
@@ -781,23 +781,119 @@ private:
 
 	static FString MakeRecordKey(const FCacheKey& Key)
 	{
-		check(Key);
+		check(Key.IsValid());
 		TStringBuilder<96> Out;
 		Out << Key.GetBucket() << TEXT("_") << Key.GetHash();
 		return FDerivedDataCacheInterface::SanitizeCacheKey(Out.ToString());
 	}
 
-	static FString MakeContentKey(const FCacheAttachmentKey& AttachmentKey)
+	static FString MakePayloadKey(const FCachePayloadKey& PayloadKey)
 	{
-		check(AttachmentKey);
-		TStringBuilder<128> Out;
-		Out << AttachmentKey.GetKey().GetBucket() << TEXT("_CAS_") << AttachmentKey.GetHash();
+		check(PayloadKey.IsValid());
+		TStringBuilder<96> Out;
+		Out << PayloadKey.GetKey().GetBucket() << TEXT("_") << PayloadKey.GetKey().GetHash() << TEXT("_") << PayloadKey.GetId();
 		return FDerivedDataCacheInterface::SanitizeCacheKey(Out.ToString());
 	}
 };
 
 ///////////////////////////////////////////////////////////////////////////////////////////////////////////////////////
 
+void FCache::Put(
+	FCacheRecord Record,
+	FStringView Context,
+	ECachePolicy Policy,
+	ECachePriority Priority,
+	FOnCachePutComplete& Callback)
+{
+	FCachePutCompleteParams Params;
+	Params.Key = Record.GetKey();
+
+	ON_SCOPE_EXIT
+	{
+		if (Callback)
+		{
+			Callback(MoveTemp(Params));
+		}
+	};
+
+	if (!EnumHasAnyFlags(Policy, ECachePolicy::Store))
+	{
+		UE_LOG(LogDerivedDataCache, VeryVerbose, TEXT("Cache: Put skipped for %s from '%.*s'"),
+			*TToString<96>(Params.Key), Context.Len(), Context.GetData());
+		Params.Status = ECacheStatus::NotCached;
+		return;
+	}
+
+	const auto PutPayload = [&Key = Params.Key, Context](const FCachePayload& Payload)
+	{
+		const FSharedBuffer& Buffer = Payload.GetCompressedData();
+		check(Buffer && Buffer.GetSize() <= MAX_int32);
+		TConstArrayView<uint8> BufferView = MakeArrayView(
+			static_cast<const uint8*>(Buffer.GetData()),
+			static_cast<int32>(Buffer.GetSize()));
+		GetDerivedDataCacheRef().Put(*MakePayloadKey(FCachePayloadKey(Key, Payload.GetId())), BufferView, Context);
+	};
+
+	FCbWriter Writer;
+	Writer.BeginObject();
+	{
+		if (const FCachePayload& Value = Record.GetValuePayload())
+		{
+			PutPayload(Value);
+			Writer.BeginObject("Value"_ASV);
+			Writer.AddObjectId("Id"_ASV, Value.GetId().ToObjectId());
+			Writer.AddBinaryAttachment("Hash"_ASV, Value.GetCompressedDataHash());
+			Writer.EndObject();
+		}
+		if (!Record.GetAttachmentPayloads().IsEmpty())
+		{
+			Writer.BeginArray("Attachments"_ASV);
+			for (const FCachePayload& Attachment : Record.GetAttachmentPayloads())
+			{
+				PutPayload(Attachment);
+				Writer.BeginObject();
+				Writer.AddObjectId("Id"_ASV, Attachment.GetId().ToObjectId());
+				Writer.AddBinaryAttachment("Hash"_ASV, Attachment.GetCompressedDataHash());
+				Writer.EndObject();
+			}
+			Writer.EndArray();
+		}
+
+		const FCbObjectRef& Meta = Record.GetMeta();
+		if (Meta.CreateIterator())
+		{
+			Writer.AddObject("Meta"_ASV, Meta);
+		}
+	}
+	Writer.EndObject();
+
+	TArray<uint8> Data;
+	const uint64 SaveSize = Writer.GetSaveSize();
+	check(SaveSize <= MAX_int32);
+	Data.SetNumUninitialized(static_cast<int32>(SaveSize));
+	Writer.Save(MakeMemoryView(Data));
+
+	GetDerivedDataCacheRef().Put(*MakeRecordKey(Params.Key), Data, Context);
+
+	UE_LOG(LogDerivedDataCache, VeryVerbose, TEXT("Cache: Put for %s from '%.*s'"),
+		*TToString<96>(Params.Key), Context.Len(), Context.GetData());
+	Params.Status = ECacheStatus::Cached;
+}
+
+FCacheRequest FCache::Put(
+	TArrayView<FCacheRecord> Records,
+	FStringView Context,
+	ECachePolicy Policy,
+	ECachePriority Priority,
+	FOnCachePutComplete&& Callback)
+{
+	for (FCacheRecord& Record : Records)
+	{
+		Put(MoveTemp(Record), Context, Policy, Priority, Callback);
+	}
+	return FCacheRequest();
+}
+
 void FCache::Get(
 	const FCacheKey& Key,
 	FStringView Context,
@@ -805,13 +901,18 @@ void FCache::Get(
 	ECachePriority Priority,
 	FOnCacheGetComplete& Callback)
 {
-	FCacheGetCompleteParams Params;
-	Params.Record.SetKey(Key);
+	FCachePayload Value;
+	TArray<FCachePayload> Attachments;
+	TArray<FCachePayloadId, TInlineAllocator<1>> SkippedPayloads;
+	FCacheRecordBuilder Builder;
+	Builder.SetKey(Key);
 
+	FCacheGetCompleteParams Params;
 	ON_SCOPE_EXIT
 	{
 		if (Callback)
 		{
+			Params.Record = Builder.Build();
 			Callback(MoveTemp(Params));
 		}
 	};
@@ -824,17 +925,17 @@ void FCache::Get(
 		return;
 	}
 
-	// Request the metadata and references from storage.
-	TArray<uint8> Data;
-	if (!GetDerivedDataCacheRef().GetSynchronous(*MakeRecordKey(Key), Data, Context))
+	// Request the record from storage.
+	TArray<uint8> RecordData;
+	if (!GetDerivedDataCacheRef().GetSynchronous(*MakeRecordKey(Key), RecordData, Context))
 	{
 		UE_LOG(LogDerivedDataCache, Verbose, TEXT("Cache: Get cache miss for %s from '%.*s'"),
 			*TToString<96>(Key), Context.Len(), Context.GetData());
 		return;
 	}
 
-	// Validate that Data can be read as compact binary without crashing.
-	if (ValidateCompactBinaryRange(MakeMemoryView(Data), ECbValidateMode::Default) != ECbValidateError::None)
+	// Validate that the record can be read as compact binary without crashing.
+	if (ValidateCompactBinaryRange(MakeMemoryView(RecordData), ECbValidateMode::Default) != ECbValidateError::None)
 	{
 		UE_LOG(LogDerivedDataCache, Verbose, TEXT("Cache: Get cache miss with corrupted record for %s from '%.*s'"),
 			*TToString<96>(Key), Context.Len(), Context.GetData());
@@ -842,144 +943,90 @@ void FCache::Get(
 	}
 
 	// Read the record from its compact binary fields.
-	FCbObject RecordObject(Data.GetData());
-	FCbField MetaField = RecordObject["Meta"_ASV];
-	FCbField ValueField = RecordObject["Value"_ASV];
-	FCbField AttachmentsField = RecordObject["Attachments"_ASV];
-	FCbArray AttachmentsArray = AttachmentsField.AsArray();
-
-	// Validate that the record was serialized in the expected format.
-	if ((MetaField && !MetaField.IsObject()) ||
-		(ValueField && !ValueField.IsAttachment()) ||
-		(AttachmentsField && !AttachmentsField.IsArray()) ||
-		!Algo::AllOf(AttachmentsArray, &FCbField::IsAttachment))
-	{
-		UE_LOG(LogDerivedDataCache, Verbose, TEXT("Cache: Get cache miss with invalid format for %s from '%.*s'"),
-			*TToString<96>(Key), Context.Len(), Context.GetData());
-		return;
-	}
+	const FCbObject RecordObject(RecordData.GetData());
+	const FCbObject MetaObject = RecordObject["Meta"_ASV].AsObject();
+	const FCbObject ValueObject = RecordObject["Value"_ASV].AsObject();
+	const FCbArray AttachmentsArray = RecordObject["Attachments"_ASV].AsArray();
 
-	// Check for existence of the value and attachments if they are being skipped.
+	// Read the value and attachments if they have been requested.
+	const auto GetPayload = [&SkippedPayloads, &Key, Policy, Context](FCbObject PayloadObject, ECachePolicy SkipFlag) -> FCachePayload
 	{
-		TArray<FIoHash, TInlineAllocator<1>> KeysToCheck;
-		if (EnumHasAnyFlags(Policy, ECachePolicy::SkipValue) && ValueField.IsAttachment())
-		{
-			KeysToCheck.Add(ValueField.AsAttachment());
-			if (KeysToCheck.Last().IsZero())
-			{
-				KeysToCheck.Pop();
-			}
-		}
-		if (EnumHasAnyFlags(Policy, ECachePolicy::SkipAttachments))
-		{
-			KeysToCheck.Reserve(KeysToCheck.Num() + AttachmentsArray.Num());
-			for (FCbField AttachmentField : AttachmentsArray)
-			{
-				KeysToCheck.Add(AttachmentField.AsAttachment());
-			}
-		}
-		if (KeysToCheck.Num())
+		TArray<uint8> PayloadData;
+		const FCachePayloadId PayloadId = FCachePayloadId::FromObjectId(PayloadObject["Id"_ASV].AsObjectId());
+		const FIoHash PayloadHash = PayloadObject["Hash"_ASV].AsBinaryAttachment();
+		const FCachePayloadKey PayloadKey(Key, PayloadId);
+		if (EnumHasAnyFlags(Policy, SkipFlag))
 		{
-			TArray<FString, TInlineAllocator<1>> ContentKeysToCheck;
-			ContentKeysToCheck.Reserve(KeysToCheck.Num());
-			Algo::Transform(KeysToCheck, ContentKeysToCheck,
-				[&Key](const FIoHash& Hash) -> FString { return MakeContentKey(FCacheAttachmentKey(Key, Hash)); });
-			if (!GetDerivedDataCacheRef().AllCachedDataProbablyExists(ContentKeysToCheck))
-			{
-				UE_LOG(LogDerivedDataCache, Verbose, TEXT("Cache: Get cache miss with missing content for %s from '%.*s'"),
-					*TToString<96>(Key), Context.Len(), Context.GetData());
-				return;
-			}
+			SkippedPayloads.Add(PayloadId);
+			return FCachePayload(PayloadId, PayloadHash);
 		}
-	}
-
-	// Read the value and attachments if they have been requested.
-	const auto GetContent = [&Key, Context](FCbField AttachmentField) -> FCbAttachment
-	{
-		TArray<uint8> ContentData;
-		const FIoHash Hash = AttachmentField.AsAttachment();
-		check(!AttachmentField.HasError());
-		const FCacheAttachmentKey ContentKey(Key, Hash);
-		if (GetDerivedDataCacheRef().GetSynchronous(*MakeContentKey(ContentKey), ContentData, Context))
+		else if (GetDerivedDataCacheRef().GetSynchronous(*MakePayloadKey(PayloadKey), PayloadData, Context))
 		{
-			if (FIoHash::HashBuffer(MakeMemoryView(ContentData)) == Hash)
+			if (FIoHash::HashBuffer(MakeMemoryView(PayloadData)) == PayloadHash)
 			{
-				FSharedBuffer ContentBuffer = FSharedBuffer::Clone(MakeMemoryView(ContentData));
-				if (AttachmentField.IsCompactBinaryAttachment())
-				{
-					return FCbAttachment(FCbFieldRefIterator::MakeRange(MoveTemp(ContentBuffer)), Hash);
-				}
-				else
-				{
-					return FCbAttachment(MoveTemp(ContentBuffer), Hash);
-				}
+				FSharedBuffer PayloadBuffer = FSharedBuffer::Clone(MakeMemoryView(PayloadData));
+				return FCachePayload(PayloadId, MoveTemp(PayloadBuffer), PayloadHash);
 			}
 			else
 			{
 				UE_LOG(LogDerivedDataCache, Verbose, TEXT("Cache: Get cache miss with corrupted content for %s from '%.*s'"),
-					*TToString<160>(ContentKey), Context.Len(), Context.GetData());
+					*TToString<96>(PayloadKey), Context.Len(), Context.GetData());
+				return FCachePayload();
 			}
 		}
 		else
 		{
 			UE_LOG(LogDerivedDataCache, Verbose, TEXT("Cache: Get cache miss with missing content for %s from '%.*s'"),
-				*TToString<160>(ContentKey), Context.Len(), Context.GetData());
+				*TToString<96>(PayloadKey), Context.Len(), Context.GetData());
+			return FCachePayload();
 		}
-		return FCbAttachment();
 	};
 
-	FCbAttachment Value;
-	if (!EnumHasAnyFlags(Policy, ECachePolicy::SkipValue) && ValueField.IsAttachment() && !(Value = GetContent(ValueField)))
+	if (ValueObject.CreateIterator() && !(Value = GetPayload(ValueObject, ECachePolicy::SkipValue)))
 	{
 		return;
 	}
 
-	TArray<FCbAttachment> Attachments;
-	if (!EnumHasAnyFlags(Policy, ECachePolicy::SkipAttachments) && AttachmentsField.IsArray())
+	Attachments.Reserve(AttachmentsArray.Num());
+	for (FCbField AttachmentField : AttachmentsArray)
 	{
-		Attachments.Reserve(AttachmentsArray.Num());
-		for (FCbField AttachmentField : AttachmentsArray)
+		Attachments.Add(GetPayload(AttachmentField.AsObject(), ECachePolicy::SkipAttachments));
+		if (!Attachments.Last())
 		{
-			Attachments.Add(GetContent(AttachmentField));
-			if (!Attachments.Last())
-			{
-				return;
-			}
+			return;
 		}
 	}
 
-	// Package
-	if (AttachmentsField.IsArray() && !EnumHasAllFlags(Policy, ECachePolicy::SkipValue | ECachePolicy::SkipAttachments))
+	// Check for existence of the value and attachments if they were skipped.
+	if (!SkippedPayloads.IsEmpty())
 	{
-		FCbPackage Package;
-		if (Value && Value.IsCompactBinary())
-		{
-			Package.SetObject(Value.AsCompactBinary().AsObjectRef(), ValueField.AsAttachment());
-		}
-		for (FCbAttachment Attachment : Attachments)
+		TArray<FString, TInlineAllocator<1>> SkippedPayloadKeys;
+		SkippedPayloadKeys.Reserve(SkippedPayloads.Num());
+		Algo::Transform(SkippedPayloads, SkippedPayloadKeys,
+			[&Key](const FCachePayloadId& Id) -> FString { return MakePayloadKey(FCachePayloadKey(Key, Id)); });
+		const bool bPayloadsExist = EnumHasAnyFlags(Policy, ECachePolicy::StoreLocal)
+			? GetDerivedDataCacheRef().TryToPrefetch(SkippedPayloadKeys, Context)
+			: GetDerivedDataCacheRef().AllCachedDataProbablyExists(SkippedPayloadKeys);
+		if (!bPayloadsExist)
 		{
-			Package.AddAttachment(MoveTemp(Attachment));
+			UE_LOG(LogDerivedDataCache, Verbose, TEXT("Cache: Get cache miss with missing content for %s from '%.*s'"),
+				*TToString<96>(Key), Context.Len(), Context.GetData());
+			return;
 		}
-		Params.Record.SetPackage(MoveTemp(Package));
 	}
-	else if (Value)
+
+	// Populate the builder now that any error conditions have been handled.
+	if (!EnumHasAnyFlags(Policy, ECachePolicy::SkipMeta) && MetaObject.CreateIterator())
 	{
-		// Object
-		if (Value.IsCompactBinary())
-		{
-			Params.Record.SetObject(Value.AsCompactBinary().AsObjectRef(), ValueField.AsAttachment());
-		}
-		// Binary
-		else
-		{
-			Params.Record.SetBinary(Value.AsBinary(), ValueField.AsAttachment());
-		}
+		Builder.SetMeta(FCbObjectRef::Clone(MetaObject));
 	}
-
-	// Meta
-	if (!EnumHasAnyFlags(Policy, ECachePolicy::SkipMeta) && MetaField.IsObject())
+	if (Value)
 	{
-		Params.Record.SetMeta(FCbObjectRef::Clone(MetaField.AsObject()));
+		Builder.SetValue(MoveTemp(Value));
+	}
+	for (FCachePayload& Attachment : Attachments)
+	{
+		Builder.AddAttachment(MoveTemp(Attachment));
 	}
 
 	UE_LOG(LogDerivedDataCache, Verbose, TEXT("Cache: Get cache hit for %s from '%.*s'"),
@@ -1001,214 +1048,131 @@ FCacheRequest FCache::Get(
 	return FCacheRequest();
 }
 
-void FCache::Put(
-	const FCacheRecord& Record,
+void FCache::GetPayload(
+	const FCachePayloadKey& Key,
 	FStringView Context,
 	ECachePolicy Policy,
 	ECachePriority Priority,
-	FOnCachePutComplete& Callback)
+	FOnCacheGetPayloadComplete& Callback)
 {
-	FCachePutCompleteParams Params;
-	Params.Key = Record.GetKey();
+	FCacheGetPayloadCompleteParams Params;
+	Params.Key = Key.GetKey();
 
 	ON_SCOPE_EXIT
 	{
 		if (Callback)
 		{
+			if (!Params.Payload)
+			{
+				Params.Payload = FCachePayload(Key.GetId(), FIoHash());
+			}
 			Callback(MoveTemp(Params));
 		}
 	};
 
-	if (!EnumHasAnyFlags(Policy, ECachePolicy::Store))
+	// Skip the request if querying the cache is disabled.
+	if (!EnumHasAnyFlags(Policy, ECachePolicy::Query))
 	{
-		UE_LOG(LogDerivedDataCache, VeryVerbose, TEXT("Cache: Put skipped for %s from '%.*s'"),
-			*TToString<96>(Params.Key), Context.Len(), Context.GetData());
-		Params.Status = ECacheStatus::NotCached;
+		UE_LOG(LogDerivedDataCache, VeryVerbose, TEXT("Cache: GetPayload skipped %s from '%.*s'"),
+			*TToString<96>(Key), Context.Len(), Context.GetData());
 		return;
 	}
 
-	FCbAttachment Value;
-	TConstArrayView<FCbAttachment> Attachments;
-
-	switch (Record.GetType())
+	// Check for existence of the payload if it is being skipped.
+	if (EnumHasAnyFlags(Policy, ECachePolicy::SkipValue))
 	{
-	case ECacheRecordType::None:
-		break;
-	case ECacheRecordType::Binary:
-		Value = FCbAttachment(Record.AsBinary(), Record.GetValueHash());
-		break;
-	case ECacheRecordType::Object:
-	case ECacheRecordType::Package:
+		const bool bPayloadExists = EnumHasAnyFlags(Policy, ECachePolicy::StoreLocal)
+			? GetDerivedDataCacheRef().TryToPrefetch({MakePayloadKey(Key)}, Context)
+			: GetDerivedDataCacheRef().CachedDataProbablyExists(*MakePayloadKey(Key));
+		if (bPayloadExists)
 		{
-			const FCbPackage& Package = Record.AsPackage();
-			Value = FCbAttachment(FCbFieldRefIterator::MakeSingle(Package.GetObject().AsFieldRef()), Package.GetObjectHash());
-			Attachments = Package.GetAttachments();
-		}
-		break;
-	default:
-		checkNoEntry();
-		break;
-	}
-
-	const auto PutContent = [&Key = Params.Key, Context](const FSharedBuffer& Buffer, const FIoHash& Hash)
-	{
-		check(Buffer.GetSize() <= MAX_int32);
-		TConstArrayView<uint8> BufferView = MakeArrayView(
-			static_cast<const uint8*>(Buffer.GetData()),
-			static_cast<int32>(Buffer.GetSize()));
-		GetDerivedDataCacheRef().Put(*MakeContentKey(FCacheAttachmentKey(Key, Hash)), BufferView, Context);
-	};
+			// Request the record from storage.
+			TArray<uint8> RecordData;
+			if (!GetDerivedDataCacheRef().GetSynchronous(*MakeRecordKey(Key.GetKey()), RecordData, Context))
+			{
+				UE_LOG(LogDerivedDataCache, Verbose, TEXT("Cache: Get cache miss with missing record for %s from '%.*s'"),
+					*TToString<96>(Key), Context.Len(), Context.GetData());
+				return;
+			}
 
-	if (Value)
-	{
-		PutContent(Value.AsBinary(), Value.GetHash());
-	}
-	for (const FCbAttachment& Attachment : Attachments)
-	{
-		PutContent(Attachment.AsBinary(), Attachment.GetHash());
-	}
+			// Validate that the record can be read as compact binary without crashing.
+			if (ValidateCompactBinaryRange(MakeMemoryView(RecordData), ECbValidateMode::Default) != ECbValidateError::None)
+			{
+				UE_LOG(LogDerivedDataCache, Verbose, TEXT("Cache: Get cache miss with corrupted record for %s from '%.*s'"),
+					*TToString<96>(Key), Context.Len(), Context.GetData());
+				return;
+			}
 
-	FCbWriter Writer;
-	Writer.BeginObject();
-	{
-		if (Value.IsCompactBinary())
-		{
-			Writer.SetName("Value"_ASV).AddCompactBinaryAttachment(Value.GetHash());
-		}
-		else if (Value.IsBinary())
-		{
-			Writer.SetName("Value"_ASV).AddBinaryAttachment(Value.GetHash());
-		}
+			// Read the record from its compact binary fields.
+			const FCbObject RecordObject(RecordData.GetData());
+			const FCbObject ValueObject = RecordObject["Value"_ASV].AsObject();
+			const FCbArray AttachmentsArray = RecordObject["Attachments"_ASV].AsArray();
 
-		if (Record.GetType() == ECacheRecordType::Package)
-		{
-			Writer.SetName("Attachments"_ASV);
-			Writer.BeginArray();
-			for (const FCbAttachment& Attachment : Attachments)
+			// Find the hash of the payload in the record.
+			const FCbObjectId KeyObjectId = Key.GetId().ToObjectId();
+			if (ValueObject["Id"_ASV].AsObjectId() == KeyObjectId)
 			{
-				if (Attachment.IsCompactBinary())
-				{
-					Writer.AddCompactBinaryAttachment(Attachment.GetHash());
-				}
-				else
+				Params.Payload = FCachePayload(Key.GetId(), ValueObject["Hash"_ASV].AsBinaryAttachment());
+			}
+			else
+			{
+				for (FCbField AttachmentField : AttachmentsArray)
 				{
-					Writer.AddBinaryAttachment(Attachment.GetHash());
+					if (AttachmentField.AsObject()["Id"_ASV].AsObjectId() == KeyObjectId)
+					{
+						Params.Payload = FCachePayload(Key.GetId(), AttachmentField.AsObject()["Hash"_ASV].AsBinaryAttachment());
+						break;
+					}
 				}
 			}
-			Writer.EndArray();
-		}
 
-		const FCbObjectRef& Meta = Record.GetMeta();
-		if (Meta.CreateIterator())
-		{
-			Writer.SetName("Meta"_ASV).AddObject(Meta);
-		}
-	}
-	Writer.EndObject();
-
-	TArray<uint8> Data;
-	const uint64 SaveSize = Writer.GetSaveSize();
-	check(SaveSize <= MAX_int32);
-	Data.SetNumUninitialized(static_cast<int32>(SaveSize));
-	Writer.Save(MakeMemoryView(Data));
-
-	GetDerivedDataCacheRef().Put(*MakeRecordKey(Params.Key), Data, Context);
-
-	UE_LOG(LogDerivedDataCache, VeryVerbose, TEXT("Cache: Put for %s from '%.*s'"),
-		*TToString<96>(Params.Key), Context.Len(), Context.GetData());
-	Params.Status = ECacheStatus::Cached;
-}
-
-FCacheRequest FCache::Put(
-	TConstArrayView<FCacheRecord> Records,
-	FStringView Context,
-	ECachePolicy Policy,
-	ECachePriority Priority,
-	FOnCachePutComplete&& Callback)
-{
-	for (const FCacheRecord& Record : Records)
-	{
-		Put(Record, Context, Policy, Priority, Callback);
-	}
-	return FCacheRequest();
-}
-
-void FCache::GetAttachment(
-	const FCacheAttachmentKey& Key,
-	FStringView Context,
-	ECachePolicy Policy,
-	ECachePriority Priority,
-	FOnCacheGetAttachmentComplete& Callback)
-{
-	FCacheGetAttachmentCompleteParams Params;
-	Params.Key = Key;
-
-	ON_SCOPE_EXIT
-	{
-		if (Callback)
-		{
-			Callback(MoveTemp(Params));
-		}
-	};
-
-	// Skip the request if querying the cache is disabled.
-	if (!EnumHasAnyFlags(Policy, ECachePolicy::Query))
-	{
-		UE_LOG(LogDerivedDataCache, VeryVerbose, TEXT("Cache: GetAttachment skipped on content for %s from '%.*s'"),
-			*TToString<160>(Key), Context.Len(), Context.GetData());
-		return;
-	}
-
-	// Check for existence of the attachment if it is being skipped.
-	if (EnumHasAnyFlags(Policy, ECachePolicy::SkipAttachments))
-	{
-		if (GetDerivedDataCacheRef().CachedDataProbablyExists(*MakeContentKey(Key)))
-		{
-			UE_LOG(LogDerivedDataCache, Verbose, TEXT("Cache: GetAttachment (exists) cache hit on content for %s from '%.*s'"),
-				*TToString<160>(Key), Context.Len(), Context.GetData());
-			Params.Status = ECacheStatus::Cached;
+			if (Params.Payload)
+			{
+				UE_LOG(LogDerivedDataCache, Verbose, TEXT("Cache: GetPayload (exists) cache hit for %s from '%.*s'"),
+					*TToString<96>(Key), Context.Len(), Context.GetData());
+				Params.Status = ECacheStatus::Cached;
+			}
+			else
+			{
+				UE_LOG(LogDerivedDataCache, Verbose,
+					TEXT("Cache: GetPayload (exists) cache miss with missing payload in record for %s from '%.*s'"),
+					*TToString<96>(Key), Context.Len(), Context.GetData());
+			}
 		}
 		else
 		{
-			UE_LOG(LogDerivedDataCache, Verbose, TEXT("Cache: GetAttachment (exists) cache miss on content for %s from '%.*s'"),
-				*TToString<160>(Key), Context.Len(), Context.GetData());
+			UE_LOG(LogDerivedDataCache, Verbose, TEXT("Cache: GetPayload (exists) cache miss for %s from '%.*s'"),
+				*TToString<96>(Key), Context.Len(), Context.GetData());
 		}
 		return;
 	}
 
-	// Request the attachment from storage.
+	// Request the payload from storage.
 	TArray<uint8> Data;
-	if (!GetDerivedDataCacheRef().GetSynchronous(*MakeContentKey(Key), Data, Context))
+	if (!GetDerivedDataCacheRef().GetSynchronous(*MakePayloadKey(Key), Data, Context))
 	{
-		UE_LOG(LogDerivedDataCache, Verbose, TEXT("Cache: GetAttachment cache miss on content for %s from '%.*s'"),
-			*TToString<160>(Key), Context.Len(), Context.GetData());
-		return;
-	}
-
-	if (FIoHash::HashBuffer(MakeMemoryView(Data)) != Key.GetHash())
-	{
-		UE_LOG(LogDerivedDataCache, Verbose, TEXT("Cache: GetAttachment cache miss with corrupted content for %s from '%.*s'"),
-			*TToString<160>(Key), Context.Len(), Context.GetData());
+		UE_LOG(LogDerivedDataCache, Verbose, TEXT("Cache: GetPayload cache miss for %s from '%.*s'"),
+			*TToString<96>(Key), Context.Len(), Context.GetData());
 		return;
 	}
 
-	Params.Value = FSharedBuffer::Clone(MakeMemoryView(Data));
+	Params.Payload = FCachePayload(Key.GetId(), FSharedBuffer::Clone(MakeMemoryView(Data)));
 
-	UE_LOG(LogDerivedDataCache, Verbose, TEXT("Cache: GetAttachment cache hit on content for %s from '%.*s'"),
-		*TToString<160>(Key), Context.Len(), Context.GetData());
+	UE_LOG(LogDerivedDataCache, Verbose, TEXT("Cache: GetPayload cache hit for %s from '%.*s'"),
+		*TToString<96>(Key), Context.Len(), Context.GetData());
 	Params.Status = ECacheStatus::Cached;
 }
 
-FCacheRequest FCache::GetAttachments(
-	TConstArrayView<FCacheAttachmentKey> Keys,
+FCacheRequest FCache::GetPayloads(
+	TConstArrayView<FCachePayloadKey> Keys,
 	FStringView Context,
 	ECachePolicy Policy,
 	ECachePriority Priority,
-	FOnCacheGetAttachmentComplete&& Callback)
+	FOnCacheGetPayloadComplete&& Callback)
 {
-	for (const FCacheAttachmentKey& Key : Keys)
+	for (const FCachePayloadKey& Key : Keys)
 	{
-		GetAttachment(Key, Context, Policy, Priority, Callback);
+		GetPayload(Key, Context, Policy, Priority, Callback);
 	}
 	return FCacheRequest();
 }
diff --git a/Engine/Source/Developer/DerivedDataCache/Private/DerivedDataCacheKey.cpp b/Engine/Source/Developer/DerivedDataCache/Private/DerivedDataCacheKey.cpp
--- a/Engine/Source/Developer/DerivedDataCache/Private/DerivedDataCacheKey.cpp
+++ b/Engine/Source/Developer/DerivedDataCache/Private/DerivedDataCacheKey.cpp
@@ -3,7 +3,9 @@
 #include "DerivedDataCacheKey.h"
 
 #include "Algo/AllOf.h"
+#include "Containers/StringConv.h"
 #include "Misc/StringBuilder.h"
+#include "Serialization/CompactBinary.h"
 #include "UObject/NameTypes.h"
 
 namespace UE
@@ -29,6 +31,13 @@ void FCacheBucket::ToString(FWideStringBuilderBase& Builder) const
 	FName::CreateFromDisplayId(FNameEntryId::FromUnstableInt(Index), 0).AppendString(Builder);
 }
 
+FString FCacheBucket::ToString() const
+{
+	TStringBuilder<32> Out;
+	ToString(Out);
+	return FString(Out);
+}
+
 bool FCacheBucketLexicalLess::operator()(FCacheBucket A, FCacheBucket B) const
 {
 	return FNameEntryId::FromUnstableInt(A.ToIndex()).LexicalLess(FNameEntryId::FromUnstableInt(B.ToIndex()));
@@ -54,16 +63,77 @@ void FCacheKey::ToString(FWideStringBuilderBase& Builder) const
 	}
 }
 
+FString FCacheKey::ToString() const
+{
+	TStringBuilder<72> Out;
+	ToString(Out);
+	return FString(Out);
+}
+
+///////////////////////////////////////////////////////////////////////////////////////////////////////////////////////
+
+FCachePayloadId FCachePayloadId::FromName(FAnsiStringView Name)
+{
+	checkf(!Name.IsEmpty(), TEXT("Payload ID requires a non-empty name."));
+	return FCachePayloadId::FromHash(FIoHash::HashBuffer(Name.GetData(), Name.Len()));
+}
+
+FCachePayloadId FCachePayloadId::FromName(FWideStringView Name)
+{
+	checkf(!Name.IsEmpty(), TEXT("Payload ID requires a non-empty name."));
+	TAnsiStringBuilder<128> Name8;
+	const int32 Len = int32(FTCHARToUTF8_Convert::ConvertedLength(Name.GetData(), Name.Len()));
+	const int32 Offset = Name8.AddUninitialized(Len);
+	const int32 WrittenLen = FTCHARToUTF8_Convert::Convert(Name8.GetData() + Offset, Len, Name.GetData(), Name.Len());
+	checkf(Len == WrittenLen, TEXT("Name '%.*s' failed to convert to a payload ID."), Name.Len(), Name.GetData());
+	return FCachePayloadId::FromName(Name8);
+}
+
+FCachePayloadId FCachePayloadId::FromObjectId(const FCbObjectId& ObjectId)
+{
+	checkf(ObjectId != FCbObjectId(), TEXT("Payload ID requires a non-zero ObjectId."));
+	return FCachePayloadId(ObjectId.GetView());
+}
+
+FCbObjectId FCachePayloadId::ToObjectId() const
+{
+	return FCbObjectId(GetView());
+}
+
+void FCachePayloadId::ToString(FAnsiStringBuilderBase& Builder) const
+{
+	UE::String::BytesToHexLower(Bytes, Builder);
+}
+
+void FCachePayloadId::ToString(FWideStringBuilderBase& Builder) const
+{
+	UE::String::BytesToHexLower(Bytes, Builder);
+}
+
+FString FCachePayloadId::ToString() const
+{
+	TStringBuilder<16> Out;
+	ToString(Out);
+	return FString(Out);
+}
+
 ///////////////////////////////////////////////////////////////////////////////////////////////////////////////////////
 
-void FCacheAttachmentKey::ToString(FAnsiStringBuilderBase& Builder) const
+void FCachePayloadKey::ToString(FAnsiStringBuilderBase& Builder) const
+{
+	Builder << Key << '/' << Id;
+}
+
+void FCachePayloadKey::ToString(FWideStringBuilderBase& Builder) const
 {
-	Builder << Key << '/' << Hash;
+	Builder << Key << TEXT('/') << Id;
 }
 
-void FCacheAttachmentKey::ToString(FWideStringBuilderBase& Builder) const
+FString FCachePayloadKey::ToString() const
 {
-	Builder << Key << TEXT('/') << Hash;
+	TStringBuilder<84> Out;
+	ToString(Out);
+	return FString(Out);
 }
 
 ///////////////////////////////////////////////////////////////////////////////////////////////////////////////////////
diff --git a/Engine/Source/Developer/DerivedDataCache/Private/DerivedDataCacheRecord.cpp b/Engine/Source/Developer/DerivedDataCache/Private/DerivedDataCacheRecord.cpp
--- a/Engine/Source/Developer/DerivedDataCache/Private/DerivedDataCacheRecord.cpp
+++ b/Engine/Source/Developer/DerivedDataCache/Private/DerivedDataCacheRecord.cpp
@@ -2,6 +2,9 @@
 
 #include "DerivedDataCacheRecord.h"
 
+#include "Async/Future.h"
+#include "Misc/StringBuilder.h"
+
 namespace UE
 {
 namespace DerivedData
@@ -9,61 +12,261 @@ namespace DerivedData
 
 ///////////////////////////////////////////////////////////////////////////////////////////////////////////////////////
 
-FCacheRecord::FCacheRecord() = default;
-
-void FCacheRecord::Reset()
+FCachePayload::FCachePayload(const FCachePayloadId& InId, const FIoHash& InCompressedDataHash)
+	: Id(InId)
+	, CompressedDataHash(InCompressedDataHash)
 {
-	*this = FCacheRecord();
+	checkf(Id.IsValid(), TEXT("A valid ID is required to construct a cache payload."));
 }
 
-void FCacheRecord::SetBinary(FSharedBuffer Value)
+FCachePayload::FCachePayload(const FCachePayloadId& InId, FSharedBuffer InCompressedData)
+	: Id(InId)
+	, CompressedData(MoveTemp(InCompressedData))
 {
-	check(Value);
-	Package.Reset();
-	Binary = MoveTemp(Value);
-	BinaryHash = Binary.GetSize() ? FIoHash::HashBuffer(Binary) : FIoHash();
-	Type = ECacheRecordType::Binary;
+	checkf(Id.IsValid(), TEXT("A valid ID is required to construct a cache payload."));
+	if (CompressedData.GetSize())
+	{
+		CompressedDataHash = FIoHash::HashBuffer(CompressedData);
+	}
 }
 
-void FCacheRecord::SetBinary(FSharedBuffer Value, const FIoHash& ValueHash)
+FCachePayload::FCachePayload(const FCachePayloadId& InId, FSharedBuffer InCompressedData, const FIoHash& InCompressedDataHash)
+	: Id(InId)
+	, CompressedData(MoveTemp(InCompressedData))
+	, CompressedDataHash(InCompressedDataHash)
 {
-	check(Value);
-	if (Value.GetSize())
+	checkf(Id.IsValid(), TEXT("A valid ID is required to construct a cache payload."));
+	if (CompressedData.GetSize())
 	{
-		checkSlow(ValueHash == FIoHash::HashBuffer(Value));
+		checkfSlow(CompressedDataHash == FIoHash::HashBuffer(CompressedData),
+			TEXT("Provided compressed data hash %s does not match calculated hash %s"),
+			*CompressedDataHash.ToString(), *FIoHash::HashBuffer(CompressedData).ToString());
+		checkf(!CompressedDataHash.IsZero(), TEXT("A non-empty compressed data buffer must have a non-zero hash."));
 	}
 	else
 	{
-		checkfSlow(ValueHash.IsZero(), TEXT("A null or empty buffer must use a hash of zero."));
+		checkf(CompressedDataHash.IsZero(), TEXT("A null or empty compressed data buffer must use a hash of zero."));
+	}
+}
+
+FSharedBuffer FCachePayload::Compress(FSharedBuffer Buffer)
+{
+	// DDC-TODO: Switch to the new compression interface when it is ready. For now, match the "no compression" method.
+	FUniqueBuffer CompressedBuffer = FUniqueBuffer::Alloc(Buffer.GetSize() + 1);
+	FMemory::Memset(CompressedBuffer.GetData(), 0, 1);
+	FMemory::Memcpy(CompressedBuffer.GetView().RightChop(1).GetData(), Buffer.GetData(), Buffer.GetSize());
+	return FSharedBuffer(MoveTemp(CompressedBuffer));
+}
+
+FSharedBuffer FCachePayload::Decompress(FSharedBuffer Buffer)
+{
+	// DDC-TODO: Switch to the new compression interface when it is ready. For now, match the "no compression" method.
+	if (Buffer)
+	{
+		const FMemoryView View = Buffer.GetView() + 1;
+		return FSharedBuffer::MakeView(View, MoveTemp(Buffer));
 	}
-	Package.Reset();
-	Binary = MoveTemp(Value);
-	BinaryHash = ValueHash;
-	Type = ECacheRecordType::Binary;
+	return FSharedBuffer();
+}
+
+///////////////////////////////////////////////////////////////////////////////////////////////////////////////////////
+
+struct FCacheRecordBuilderData
+{
+	FCacheKey Key;
+	FCbObjectRef Meta;
+	FCachePayload Value;
+	TArray<FCachePayload> Attachments;
+};
+
+///////////////////////////////////////////////////////////////////////////////////////////////////////////////////////
+
+struct FCacheRecordData
+{
+	FCacheRecordData() = default;
+	explicit FCacheRecordData(FCacheRecordBuilderData&& Builder);
+
+	FCacheKey Key;
+	FCbObjectRef Meta;
+	FCachePayload Value;
+	FSharedBuffer ValueCache;
+	TArray<FCachePayload> Attachments;
+	TArray<FSharedBuffer> AttachmentsCache;
+};
+
+FCacheRecordData::FCacheRecordData(FCacheRecordBuilderData&& Builder)
+	: Key(Builder.Key)
+	, Meta(MoveTemp(Builder.Meta))
+	, Value(MoveTemp(Builder.Value))
+	, Attachments(MoveTemp(Builder.Attachments))
+{
 }
 
-void FCacheRecord::SetObject(FCbObjectRef Value)
+///////////////////////////////////////////////////////////////////////////////////////////////////////////////////////
+
+static const FCachePayload& GetEmptyCachePayload()
+{
+	static const FCachePayload Empty;
+	return Empty;
+}
+
+static const FCacheRecordData& GetEmptyCacheRecord()
+{
+	static const FCacheRecordData Empty;
+	return Empty;
+}
+
+static FCachePayloadId GetOrCreatePayloadId(const FCachePayloadId& Id, const FSharedBuffer& Buffer)
+{
+	return Id.IsValid() ? Id : FCachePayloadId::FromHash(FIoHash::HashBuffer(Buffer));
+}
+
+///////////////////////////////////////////////////////////////////////////////////////////////////////////////////////
+
+FCacheRecord::FCacheRecord(const FCacheRecordData& InData)
+	: Data(MakePimpl<FCacheRecordData>(InData))
+{
+}
+
+FCacheRecord::FCacheRecord(FCacheRecordBuilderData&& Builder)
+	: Data(MakePimpl<FCacheRecordData>(MoveTemp(Builder)))
+{
+}
+
+FCacheRecord FCacheRecord::Clone() const
+{
+	return Data ? FCacheRecord(*Data) : FCacheRecord();
+}
+
+const FCacheKey& FCacheRecord::GetKey() const
 {
-	Binary.Reset();
-	BinaryHash.Reset();
-	Package = FCbPackage(MoveTemp(Value));
-	Type = ECacheRecordType::Object;
+	return (Data ? *Data : GetEmptyCacheRecord()).Key;
 }
 
-void FCacheRecord::SetObject(FCbObjectRef Value, const FIoHash& ValueHash)
+const FCbObjectRef& FCacheRecord::GetMeta() const
 {
-	Binary.Reset();
-	BinaryHash.Reset();
-	Package = FCbPackage(MoveTemp(Value), ValueHash);
-	Type = ECacheRecordType::Object;
+	return (Data ? *Data : GetEmptyCacheRecord()).Meta;
+}
+
+FSharedBuffer FCacheRecord::GetValue() const
+{
+	if (Data)
+	{
+		if (Data->ValueCache.IsNull() && Data->Value.IsValid())
+		{
+			Data->ValueCache = Data->Value.Decompress();
+		}
+		return Data->ValueCache;
+	}
+	return FSharedBuffer();
+}
+
+const FCachePayload& FCacheRecord::GetValuePayload() const
+{
+	return (Data ? *Data : GetEmptyCacheRecord()).Value;
+}
+
+FSharedBuffer FCacheRecord::GetAttachment(const FCachePayloadId& Id) const
+{
+	TArray<FCachePayload>& Attachments = Data->Attachments;
+	const int32 Index = Algo::LowerBound(Attachments, Id, FCachePayloadLessById());
+	if (Attachments.IsValidIndex(Index) && FCachePayloadEqualById()(Attachments[Index], Id))
+	{
+		TArray<FSharedBuffer>& AttachmentsCache = Data->AttachmentsCache;
+		if (AttachmentsCache.IsEmpty())
+		{
+			AttachmentsCache.SetNum(Attachments.Num());
+		}
+		FSharedBuffer& DataCache = AttachmentsCache[Index];
+		if (!DataCache)
+		{
+			DataCache = Attachments[Index].Decompress();
+		}
+		return DataCache;
+	}
+	return FSharedBuffer();
+}
+
+const FCachePayload& FCacheRecord::GetAttachmentPayload(const FCachePayloadId& Id) const
+{
+	TArray<FCachePayload>& Attachments = Data->Attachments;
+	const int32 Index = Algo::LowerBound(Attachments, Id, FCachePayloadLessById());
+	if (Attachments.IsValidIndex(Index) && FCachePayloadEqualById()(Attachments[Index], Id))
+	{
+		return Attachments[Index];
+	}
+	return GetEmptyCachePayload();
+}
+
+TConstArrayView<FCachePayload> FCacheRecord::GetAttachmentPayloads() const
+{
+	return (Data ? *Data : GetEmptyCacheRecord()).Attachments;
+}
+
+///////////////////////////////////////////////////////////////////////////////////////////////////////////////////////
+
+FCacheRecordBuilder::FCacheRecordBuilder()
+	: Data(MakePimpl<FCacheRecordBuilderData>())
+{
+}
+
+void FCacheRecordBuilder::SetKey(const FCacheKey& Key)
+{
+	Data->Key = Key;
+}
+
+void FCacheRecordBuilder::SetMeta(FCbObjectRef Meta)
+{
+	Meta.MakeOwned();
+	Data->Meta = MoveTemp(Meta);
+}
+
+FCachePayloadId FCacheRecordBuilder::SetValue(FSharedBuffer Buffer, FCachePayloadId Id)
+{
+	Id = GetOrCreatePayloadId(Id, Buffer);
+	return SetValue(FCachePayload(Id, FCachePayload::Compress(MoveTemp(Buffer))));
+}
+
+FCachePayloadId FCacheRecordBuilder::SetValue(FCachePayload&& Payload)
+{
+	FCachePayload& Value = Data->Value;
+	checkf(Value.IsNull(), TEXT("Failed to set value with ID %s because an there is an existing value with ID %s"),
+		*Payload.GetId().ToString(), *Value.GetId().ToString());
+	Value = MoveTemp(Payload);
+	return Value.GetId();
+}
+
+FCachePayloadId FCacheRecordBuilder::AddAttachment(FSharedBuffer Buffer, FCachePayloadId Id)
+{
+	Id = GetOrCreatePayloadId(Id, Buffer);
+	return AddAttachment(FCachePayload(Id, MoveTemp(Buffer)));
+}
+
+FCachePayloadId FCacheRecordBuilder::AddAttachment(FCachePayload&& Payload)
+{
+	TArray<FCachePayload>& Attachments = Data->Attachments;
+	const int32 Index = Algo::LowerBound(Attachments, Payload, FCachePayloadLessById());
+	checkf(!Attachments.IsValidIndex(Index) || !FCachePayloadEqualById()(Attachments[Index], Payload),
+		TEXT("Failed to add attachment with ID %s because an existing attachment is using that ID"),
+		*Payload.GetId().ToString());
+	Attachments.Insert(MoveTemp(Payload), Index);
+	return Attachments[Index].GetId();
+}
+
+FCacheRecord FCacheRecordBuilder::Build()
+{
+	if (TPimplPtr<FCacheRecordBuilderData> LocalData = MoveTemp(Data))
+	{
+		return FCacheRecord(MoveTemp(*LocalData));
+	}
+	return FCacheRecord();
 }
 
-void FCacheRecord::SetPackage(FCbPackage Value)
+TFuture<FCacheRecord> FCacheRecordBuilder::BuildAsync()
 {
-	Binary.Reset();
-	BinaryHash.Reset();
-	Package = MoveTemp(Value);
-	Type = ECacheRecordType::Package;
+	TPromise<FCacheRecord> Promise;
+	Promise.SetValue(Build());
+	return Promise.GetFuture();
 }
 
 ///////////////////////////////////////////////////////////////////////////////////////////////////////////////////////
diff --git a/Engine/Source/Developer/DerivedDataCache/Public/DerivedDataCache.h b/Engine/Source/Developer/DerivedDataCache/Public/DerivedDataCache.h
--- a/Engine/Source/Developer/DerivedDataCache/Public/DerivedDataCache.h
+++ b/Engine/Source/Developer/DerivedDataCache/Public/DerivedDataCache.h
@@ -16,44 +16,65 @@ namespace DerivedData
 
 ///////////////////////////////////////////////////////////////////////////////////////////////////////////////////////
 
-/** Flags to control the behavior of cache requests. */
+/**
+ * Flags to control the behavior of cache requests.
+ *
+ * The cache policy flags can be combined to support a variety of usage patterns. Examples:
+ *
+ * Get(Default): read from any cache; put to writable caches if missing.
+ * Get(Remote): read any remote cache; put to writable remote caches if missing.
+ * Get(Local): read any local cache; put to writable local caches if missing.
+ * Get(Query | StoreRemote): read from any cache; put to writable remote caches if missing.
+ * Get(Query | StoreLocal): read from any cache; put to writable local caches if missing.
+ * Get(Query | SkipData): check for existence in any cache; do not modify any cache.
+ * Get(Default | SkipData): check for existence in any cache; put to writable caches if missing.
+ * Get(Default | SkipLocalCopy): read from any cache; put to writable remote caches if missing,
+ *                               and to writable local caches if not present in any local cache.
+ *
+ * Put(Default): write to every writable cache.
+ * Put(Remote): write to every writable remote cache, skipping local caches.
+ * Put(Local): write to every writable local cache, skipping remote caches.
+ */
 enum class ECachePolicy : uint8
 {
 	/** A value without any flags set. */
 	None            = 0,
 
-	/** Allow a cache request to query local caches. */
+	/** Allow a cache get request or payload request to query local caches. */
 	QueryLocal      = 1 << 0,
-	/** Allow a cache request to query remote caches. */
+	/** Allow a cache get request or payload request to query remote caches. */
 	QueryRemote     = 1 << 1,
-	/** Allow requests to query available local and remote caches. */
+	/** Allow a cache get request or payload request to query local and remote caches. */
 	Query           = QueryLocal | QueryRemote,
 
-	/** Allow cache records to be stored in local caches. */
+	/** Allow cache records to be stored in local caches by any cache request. */
 	StoreLocal      = 1 << 2,
-	/** Allow cache records to be stored in remote caches. */
+	/** Allow cache records to be stored in remote caches by any cache request. */
 	StoreRemote     = 1 << 3,
-	/** Allow cache records to be stored in local and remote caches. */
+	/** Allow cache records to be stored in local and remote caches by any cache request. */
 	Store           = StoreLocal | StoreRemote,
 
 	/** Skip fetching the metadata for get requests. */
 	SkipMeta        = 1 << 4,
-	/** Skip fetching the value (binary, object, or the object of the package) for get requests. */
+	/** Skip fetching the value for get requests and the payload for payload requests. */
 	SkipValue       = 1 << 5,
 	/** Skip fetching the attachments for get requests. */
 	SkipAttachments = 1 << 6,
-	/** Skip fetching any data for get requests, which will only check for existence. */
-	SkipAll         = SkipMeta | SkipValue | SkipAttachments,
+	/** Skip fetching the metadata or payloads for get requests or payload requests. */
+	SkipData        = SkipMeta | SkipValue | SkipAttachments,
 
-	/** Allow requests to query and store records in local caches. */
+	/** Skip copying records to local caches for get requests if they already exist locally. */
+	SkipLocalCopy   = 1 << 7,
+
+	/** Allow cache requests to query and store records in local caches. */
 	Local           = QueryLocal | StoreLocal,
-	/** Allow requests to query and store records in remote caches. */
+	/** Allow cache requests to query and store records in remote caches. */
 	Remote          = QueryRemote | StoreRemote,
 
-	/** Allow requests to query and store records in local and remote caches. */
+	/** Allow cache requests to query and store records in local and remote caches. */
 	Default         = Query | Store,
 
-	/** Do not query or store records in the cache. */
+	/** Do not allow cache requests to query or store records in local or remote caches. */
 	Disable         = None,
 };
 
@@ -64,16 +85,43 @@ ENUM_CLASS_FLAGS(ECachePolicy);
 /** Priority for scheduling cache requests. */
 enum class ECachePriority : uint8
 {
+	/**
+	 * Blocking is to be used only when the thread making the request will wait on completion of the
+	 * request before doing any other work. Requests at this priority level will be processed before
+	 * any request at a lower priority level. The cache may choose to process part of the request on
+	 * the thread making the request. Waiting on a request may increase its priority to this level.
+	 */
+	Blocking,
+	/**
+	 * Highest is the maximum priority for asynchronous cache requests, and is intended for requests
+	 * that are required to maintain interactivity of the program.
+	 */
 	Highest,
+	/**
+	 * High is intended for requests that are on the critical path, but are not required to maintain
+	 * interactivity of the program.
+	 */
 	High,
+	/**
+	 * Normal is intended as the default request priority.
+	 */
 	Normal,
+	/**
+	 * Low is intended for requests that are below the default priority, but are used for operations
+	 * that the program will execute now rather than at an unknown time in the future.
+	 */
 	Low,
+	/**
+	 * Lowest is the minimum priority for asynchronous cache requests, and is primarily intended for
+	 * requests that prefetch from the cache to optimize a future operation, while minimizing impact
+	 * on other cache requests that are in flight.
+	 */
 	Lowest,
 };
 
 ///////////////////////////////////////////////////////////////////////////////////////////////////////////////////////
 
-/** The status of a cache operation. */
+/** Status of a cache operation. */
 enum class ECacheStatus : uint8
 {
 	/** A cache miss for a get operation, or a failure for a put operation. */
@@ -86,7 +134,7 @@ enum class ECacheStatus : uint8
 
 ///////////////////////////////////////////////////////////////////////////////////////////////////////////////////////
 
-/** Interface for a cache request. Provides functionality common to every request type. */
+/** Interface for a cache request. Provides functionality common to every type of request. */
 class ICacheRequest : public IRefCountedObject
 {
 public:
@@ -98,7 +146,7 @@ public:
 	/** Cancel the request and invoke its callback. */
 	virtual void Cancel() = 0;
 
-	/** Block the calling thread until the request is complete. */
+	/** Block the calling thread until the request and its callback complete. */
 	virtual void Wait() = 0;
 
 	/** Poll the request and return true if complete and false otherwise. */
@@ -185,14 +233,26 @@ private:
 
 ///////////////////////////////////////////////////////////////////////////////////////////////////////////////////////
 
+/** Parameters for the completion callback for cache put requests. */
+struct FCachePutCompleteParams
+{
+	/** Key for the part of the put request that completed or was canceled. */
+	FCacheKey Key;
+
+	/** Status of the cache request. */
+	ECacheStatus Status = ECacheStatus::NotCached;
+};
+
 /** Parameters for the completion callback for cache get requests. */
 struct FCacheGetCompleteParams
 {
 	/**
-	 * A cache record produced by the cache request.
+	 * Record for the part of the get request that completed or was canceled.
+	 *
+	 * The key is always populated. The remainder of the record is populated when Status is Cached.
 	 *
-	 * The key is always set, and the other fields of the record are set according to the flags used
-	 * when the cache request was made.
+	 * The value, attachments, and metadata may be skipped based on cache policy flags. When a value
+	 * or attachment has been skipped, it will have a payload but its buffers will be null.
 	 */
 	FCacheRecord Record;
 
@@ -200,37 +260,31 @@ struct FCacheGetCompleteParams
 	ECacheStatus Status = ECacheStatus::NotCached;
 };
 
-/** Parameters for the completion callback for cache put requests. */
-struct FCachePutCompleteParams
+/** Parameters for the completion callback for cache payload requests. */
+struct FCacheGetPayloadCompleteParams
 {
-	/** Key used by the put request that completed. */
+	/** Key for the part of the payload request that completed or was canceled. See Payload for ID. */
 	FCacheKey Key;
 
-	/** Status of the cache request. */
-	ECacheStatus Status = ECacheStatus::NotCached;
-};
-
-/** Parameters for the completion callback for cache attachment requests. */
-struct FCacheGetAttachmentCompleteParams
-{
-	/** Key used by the attachment request that completed. */
-	FCacheAttachmentKey Key;
-
-	/** Value of the attachment that was requested, if it was available. */
-	FSharedBuffer Value;
+	/**
+	 * Payload for the part of the payload request that completed or was canceled.
+	 *
+	 * The ID is always populated. The remainder of the payload is populated when Status is Cached.
+	 */
+	FCachePayload Payload;
 
 	/** Status of the cache request. */
 	ECacheStatus Status = ECacheStatus::NotCached;
 };
 
-/** Callback for completion of a get request. */
-using FOnCacheGetComplete = TUniqueFunction<void (FCacheGetCompleteParams&& Params)>;
-
 /** Callback for completion of a put request. */
 using FOnCachePutComplete = TUniqueFunction<void (FCachePutCompleteParams&& Params)>;
 
-/** Callback for completion of an attachment request. */
-using FOnCacheGetAttachmentComplete = TUniqueFunction<void (FCacheGetAttachmentCompleteParams&& Params)>;
+/** Callback for completion of a get request. */
+using FOnCacheGetComplete = TUniqueFunction<void (FCacheGetCompleteParams&& Params)>;
+
+/** Callback for completion of a payload request. */
+using FOnCacheGetPayloadComplete = TUniqueFunction<void (FCacheGetPayloadCompleteParams&& Params)>;
 
 ///////////////////////////////////////////////////////////////////////////////////////////////////////////////////////
 
@@ -247,66 +301,69 @@ public:
 	virtual ~ICache() = default;
 
 	/**
-	 * Asynchronous request to get cache records according to the policy.
+	 * Asynchronous request to put cache records according to the policy.
 	 *
 	 * The callback will always be called for every key, and may be called from an arbitrary thread.
-	 * Records may become available in any order, and from multiple threads concurrently.
+	 * Records may finish storing in any order, and from multiple threads concurrently.
 	 *
-	 * @param Keys The keys identifying the cache records to query.
+	 * @param Records The cache records to store. Must have a key. Records are reset to null.
 	 * @param Context A description of the request. An object path is typically sufficient.
 	 * @param Policy Flags to control the behavior of the request. See ECachePolicy.
 	 * @param Priority A priority to consider when scheduling the request. See ECachePriority.
 	 * @param Callback A callback invoked for every key in the batch as it completes or is canceled.
 	 */
-	virtual FCacheRequest Get(
-		TConstArrayView<FCacheKey> Keys,
+	virtual FCacheRequest Put(
+		TArrayView<FCacheRecord> Records,
 		FStringView Context,
-		ECachePolicy Policy,
-		ECachePriority Priority,
-		FOnCacheGetComplete&& Callback) = 0;
+		ECachePolicy Policy = ECachePolicy::Default,
+		ECachePriority Priority = ECachePriority::Normal,
+		FOnCachePutComplete&& Callback = FOnCachePutComplete()) = 0;
 
 	/**
-	 * Asynchronous request to put cache records according to the policy.
+	 * Asynchronous request to get cache records according to the policy.
 	 *
 	 * The callback will always be called for every key, and may be called from an arbitrary thread.
-	 * Records may finish storing in any order, and from multiple threads concurrently.
+	 * Records may become available in any order, and from multiple threads concurrently.
 	 *
-	 * @param Records The cache records to store. Must have a key and value.
+	 * @param Keys The keys identifying the cache records to query.
 	 * @param Context A description of the request. An object path is typically sufficient.
 	 * @param Policy Flags to control the behavior of the request. See ECachePolicy.
 	 * @param Priority A priority to consider when scheduling the request. See ECachePriority.
 	 * @param Callback A callback invoked for every key in the batch as it completes or is canceled.
 	 */
-	virtual FCacheRequest Put(
-		TConstArrayView<FCacheRecord> Records,
+	virtual FCacheRequest Get(
+		TConstArrayView<FCacheKey> Keys,
 		FStringView Context,
-		ECachePolicy Policy = ECachePolicy::Default,
-		ECachePriority Priority = ECachePriority::Normal,
-		FOnCachePutComplete&& Callback = FOnCachePutComplete()) = 0;
+		ECachePolicy Policy,
+		ECachePriority Priority,
+		FOnCacheGetComplete&& Callback) = 0;
 
 	/**
-	 * Asynchronous request to get cache record attachments according to the policy.
+	 * Asynchronous request to get cache record payloads according to the policy.
 	 *
 	 * The callback will always be called for every key, and may be called from an arbitrary thread.
-	 * Attachments may become available in any order, and from multiple threads concurrently.
+	 * Payloads may become available in any order, and from multiple threads concurrently.
 	 *
-	 * @param Keys The keys identifying the cache record attachments to query.
+	 * @param Keys The keys identifying the cache record payloads to query.
 	 * @param Context A description of the request. An object path is typically sufficient.
 	 * @param Policy Flags to control the behavior of the request. See ECachePolicy.
 	 * @param Priority A priority to consider when scheduling the request. See ECachePriority.
 	 * @param Callback A callback invoked for every key in the batch as it completes or is canceled.
 	 */
-	virtual FCacheRequest GetAttachments(
-		TConstArrayView<FCacheAttachmentKey> Keys,
+	virtual FCacheRequest GetPayloads(
+		TConstArrayView<FCachePayloadKey> Keys,
 		FStringView Context,
 		ECachePolicy Policy,
 		ECachePriority Priority,
-		FOnCacheGetAttachmentComplete&& Callback) = 0;
+		FOnCacheGetPayloadComplete&& Callback) = 0;
 
 	/**
 	 * Cancel all queued and active cache requests and invoke their callbacks.
 	 *
-	 * This is used to finish requests before destroying the cache.
+	 * This is meant to be called before destroying the cache. Any new requests that are made during
+	 * execution of this function will also be canceled, because callbacks might queue requests when
+	 * they are invoked. This does not block any future cache requests, which must be handled by the
+	 * owner of the cache if future requests are meant to be avoided.
 	 */
 	virtual void CancelAll() = 0;
 };
diff --git a/Engine/Source/Developer/DerivedDataCache/Public/DerivedDataCacheKey.h b/Engine/Source/Developer/DerivedDataCache/Public/DerivedDataCacheKey.h
--- a/Engine/Source/Developer/DerivedDataCache/Public/DerivedDataCacheKey.h
+++ b/Engine/Source/Developer/DerivedDataCache/Public/DerivedDataCacheKey.h
@@ -10,6 +10,8 @@
 
 #define UE_API DERIVEDDATACACHE_API
 
+class FCbObjectId;
+
 namespace UE
 {
 namespace DerivedData
@@ -27,14 +29,15 @@ public:
 	/** Construct a cache bucket from the alphanumeric name. Names are case-insensitive. */
 	UE_API explicit FCacheBucket(FStringView Name);
 
-	/** Append the name of the cache bucket to the string builder. */
+	/** Convert the name of the cache bucket to a string. */
 	UE_API void ToString(FAnsiStringBuilderBase& Builder) const;
 	UE_API void ToString(FWideStringBuilderBase& Builder) const;
+	UE_API FString ToString() const;
 
-	/** Whether the cache bucket is null. */
+	/** Whether this is null. */
 	inline bool IsNull() const { return !Index; }
-	/** Whether the cache bucket is not null. */
-	inline explicit operator bool() const { return !IsNull(); }
+	/** Whether this is not null. */
+	inline bool IsValid() const { return !IsNull(); }
 
 	/** Returns the index that is only stable within this process. */
 	inline uint32 ToIndex() const { return Index; }
@@ -64,7 +67,6 @@ inline uint32 GetTypeHash(FCacheBucket Bucket)
 	return ::GetTypeHash(Bucket.ToIndex());
 }
 
-/** Append the cache bucket name to the string builder. */
 template <typename CharType>
 inline TStringBuilderBase<CharType>& operator<<(TStringBuilderBase<CharType>& Builder, const FCacheBucket& Bucket)
 {
@@ -109,17 +111,18 @@ public:
 	/** Returns the hash component of the cache key. */
 	inline const FIoHash& GetHash() const { return Hash; }
 
-	/** Append the cache key to the string builder. */
+	/** Convert the cache key to a string. */
 	UE_API void ToString(FAnsiStringBuilderBase& Builder) const;
 	UE_API void ToString(FWideStringBuilderBase& Builder) const;
+	UE_API FString ToString() const;
 
-	/** Whether the cache key is null. */
+	/** Whether this is null. */
 	inline bool IsNull() const { return Bucket.IsNull(); }
-	/** Whether the cache key is not null. */
-	inline explicit operator bool() const { return !IsNull(); }
+	/** Whether this is not null. */
+	inline bool IsValid() const { return !IsNull(); }
 
 	/** Reset this to null. */
-	void Reset() { *this = FCacheKey(); }
+	inline void Reset() { *this = FCacheKey(); }
 
 private:
 	FCacheBucket Bucket;
@@ -136,7 +139,6 @@ inline uint32 GetTypeHash(const FCacheKey& Key)
 	return HashCombine(GetTypeHash(Key.GetBucket()), GetTypeHash(Key.GetHash()));
 }
 
-/** Append the cache key to the string builder. */
 template <typename CharType>
 inline TStringBuilderBase<CharType>& operator<<(TStringBuilderBase<CharType>& Builder, const FCacheKey& Key)
 {
@@ -168,80 +170,171 @@ struct FCacheKeyLexicalLess
 
 ///////////////////////////////////////////////////////////////////////////////////////////////////////////////////////
 
-/** A key that uniquely identifies an attachment within a cache record. */
-class FCacheAttachmentKey
+/** A 12-byte value that uniquely identifies a cache payload in the scope of a cache record. */
+struct FCachePayloadId
 {
 public:
-	/** Construct a null cache attachment key. */
-	FCacheAttachmentKey() = default;
+	using ByteArray = uint8[12];
+
+	/** Construct a null ID. */
+	FCachePayloadId() = default;
+
+	/** Construct an ID from a view of 12 bytes. */
+	inline explicit FCachePayloadId(FMemoryView Id);
+
+	/** Construct an ID from a non-zero hash. The hash must be unique in the outer cache record. */
+	static inline FCachePayloadId FromHash(const FIoHash& Hash);
+
+	/** Construct an ID from a non-empty name. The name must be unique in the outer cache record. */
+	UE_API static FCachePayloadId FromName(FAnsiStringView Name);
+	UE_API static FCachePayloadId FromName(FWideStringView Name);
+
+	/** Construct an ID from an ObjectId. The ObjectId must be unique in the outer cache record. */
+	UE_API static FCachePayloadId FromObjectId(const FCbObjectId& ObjectId);
+
+	/** Convert the ID to an ObjectId. */
+	UE_API FCbObjectId ToObjectId() const;
+
+	/** Convert the ID to a 24-character hex string. */
+	UE_API void ToString(FAnsiStringBuilderBase& Builder) const;
+	UE_API void ToString(FWideStringBuilderBase& Builder) const;
+	UE_API FString ToString() const;
+
+	/** Returns a view of the raw byte array for the ID. */
+	constexpr inline FMemoryView GetView() const { return MakeMemoryView(Bytes); }
+
+	/** Whether this is null. */
+	inline bool IsNull() const;
+	/** Whether this is not null. */
+	inline bool IsValid() const { return !IsNull(); }
+
+	/** Reset this to null. */
+	inline void Reset() { *this = FCachePayloadId(); }
+
+private:
+	alignas(uint32) ByteArray Bytes{};
+};
+
+inline bool operator==(const FCachePayloadId& A, const FCachePayloadId& B)
+{
+	return A.GetView().EqualBytes(B.GetView());
+}
+
+inline bool operator!=(const FCachePayloadId& A, const FCachePayloadId& B)
+{
+	return !A.GetView().EqualBytes(B.GetView());
+}
+
+inline bool operator<(const FCachePayloadId& A, const FCachePayloadId& B)
+{
+	return A.GetView().CompareBytes(B.GetView()) < 0;
+}
 
-	/** Construct a cache attachment key from the non-null key and non-zero hash. */
-	inline FCacheAttachmentKey(const FCacheKey& InKey, const FIoHash& InHash)
+inline uint32 GetTypeHash(const FCachePayloadId& Id)
+{
+	return *reinterpret_cast<const uint32*>(Id.GetView().GetData());
+}
+
+inline FCachePayloadId::FCachePayloadId(const FMemoryView InId)
+{
+	checkf(InId.GetSize() == sizeof(ByteArray),
+		TEXT("FCachePayloadId cannot be constructed from a view of %" UINT64_FMT " bytes."), InId.GetSize());
+	FMemory::Memcpy(Bytes, InId.GetData(), sizeof(ByteArray));
+}
+
+inline FCachePayloadId FCachePayloadId::FromHash(const FIoHash& Hash)
+{
+	return FCachePayloadId(MakeMemoryView(Hash.GetBytes()).Left(sizeof(ByteArray)));
+}
+
+inline bool FCachePayloadId::IsNull() const
+{
+	return *this == FCachePayloadId();
+}
+
+template <typename CharType>
+inline TStringBuilderBase<CharType>& operator<<(TStringBuilderBase<CharType>& Builder, const FCachePayloadId& Id)
+{
+	Id.ToString(Builder);
+	return Builder;
+}
+
+///////////////////////////////////////////////////////////////////////////////////////////////////////////////////////
+
+/** A key that uniquely identifies a payload within a cache record. */
+class FCachePayloadKey
+{
+public:
+	/** Construct a null cache payload key. */
+	FCachePayloadKey() = default;
+
+	/** Construct a cache payload key from the non-null key and non-null ID. */
+	inline FCachePayloadKey(const FCacheKey& InKey, const FCachePayloadId& InId)
 		: Key(InKey)
-		, Hash(InHash)
+		, Id(InId)
 	{
 	}
 
-	/** Returns the cache key for the record containing the attachment. */
+	/** Returns the cache key for the record containing the payload. */
 	inline const FCacheKey& GetKey() const { return Key; }
 
-	/** Returns the hash of the attachment. */
-	inline const FIoHash& GetHash() const { return Hash; }
+	/** Returns the ID of the payload. */
+	inline const FCachePayloadId& GetId() const { return Id; }
 
-	/** Append the cache attachment key to the string builder. */
+	/** Convert the cache payload key to a string. */
 	UE_API void ToString(FAnsiStringBuilderBase& Builder) const;
 	UE_API void ToString(FWideStringBuilderBase& Builder) const;
+	UE_API FString ToString() const;
 
-	/** Whether the cache attachment key is null. */
+	/** Whether this is null. */
 	inline bool IsNull() const { return Key.IsNull(); }
-	/** Whether the cache attachment key is not null. */
-	inline explicit operator bool() const { return !IsNull(); }
+	/** Whether this is not null. */
+	inline bool IsValid() const { return !IsNull(); }
 
 	/** Reset this to null. */
-	void Reset() { *this = FCacheAttachmentKey(); }
+	inline void Reset() { *this = FCachePayloadKey(); }
 
 private:
 	FCacheKey Key;
-	FIoHash Hash;
+	FCachePayloadId Id;
 };
 
-inline bool operator==(const FCacheAttachmentKey& A, const FCacheAttachmentKey& B)
+inline bool operator==(const FCachePayloadKey& A, const FCachePayloadKey& B)
 {
-	return A.GetKey() == B.GetKey() && A.GetHash() == B.GetHash();
+	return A.GetKey() == B.GetKey() && A.GetId() == B.GetId();
 }
 
-inline uint32 GetTypeHash(const FCacheAttachmentKey& Key)
+inline uint32 GetTypeHash(const FCachePayloadKey& Key)
 {
-	return HashCombine(GetTypeHash(Key.GetKey()), GetTypeHash(Key.GetHash()));
+	return HashCombine(GetTypeHash(Key.GetKey()), GetTypeHash(Key.GetId()));
 }
 
-/** Append the cache attachment key to the string builder. */
 template <typename CharType>
-inline TStringBuilderBase<CharType>& operator<<(TStringBuilderBase<CharType>& Builder, const FCacheAttachmentKey& Key)
+inline TStringBuilderBase<CharType>& operator<<(TStringBuilderBase<CharType>& Builder, const FCachePayloadKey& Key)
 {
 	Key.ToString(Builder);
 	return Builder;
 }
 
 /** Fast less with an order that is only stable within the lifetime of this process. */
-struct FCacheAttachmentKeyFastLess
+struct FCachePayloadKeyFastLess
 {
-	inline bool operator()(const FCacheAttachmentKey& A, const FCacheAttachmentKey& B) const
+	inline bool operator()(const FCachePayloadKey& A, const FCachePayloadKey& B) const
 	{
 		const FCacheKey& KeyA = A.GetKey();
 		const FCacheKey& KeyB = B.GetKey();
-		return KeyA == KeyB ? A.GetHash() < B.GetHash() : FCacheKeyFastLess()(KeyA, KeyB);
+		return KeyA == KeyB ? A.GetId() < B.GetId() : FCacheKeyFastLess()(KeyA, KeyB);
 	}
 };
 
 /** Slow less with lexical order that is stable between processes. */
-struct FCacheAttachmentKeyLexicalLess
+struct FCachePayloadKeyLexicalLess
 {
-	inline bool operator()(const FCacheAttachmentKey& A, const FCacheAttachmentKey& B) const
+	inline bool operator()(const FCachePayloadKey& A, const FCachePayloadKey& B) const
 	{
 		const FCacheKey& KeyA = A.GetKey();
 		const FCacheKey& KeyB = B.GetKey();
-		return KeyA == KeyB ? A.GetHash() < B.GetHash() : FCacheKeyLexicalLess()(KeyA, KeyB);
+		return KeyA == KeyB ? A.GetId() < B.GetId() : FCacheKeyLexicalLess()(KeyA, KeyB);
 	}
 };
 
diff --git a/Engine/Source/Developer/DerivedDataCache/Public/DerivedDataCacheRecord.h b/Engine/Source/Developer/DerivedDataCache/Public/DerivedDataCacheRecord.h
--- a/Engine/Source/Developer/DerivedDataCache/Public/DerivedDataCacheRecord.h
+++ b/Engine/Source/Developer/DerivedDataCache/Public/DerivedDataCacheRecord.h
@@ -6,24 +6,82 @@
 #include "DerivedDataCacheKey.h"
 #include "Memory/SharedBuffer.h"
 #include "Serialization/CompactBinary.h"
-#include "Serialization/CompactBinaryPackage.h"
+#include "Templates/PimplPtr.h"
 
 #define UE_API DERIVEDDATACACHE_API
 
+template <typename ResultType> class TFuture;
+
 namespace UE
 {
 namespace DerivedData
 {
 
+class FCacheRecordBuilder;
+struct FCacheRecordData;
+struct FCacheRecordBuilderData;
+
 ///////////////////////////////////////////////////////////////////////////////////////////////////////////////////////
 
-/** The type of value stored in a cache record. */
-enum class ECacheRecordType : uint8
+/**
+ * A cache payload is a compressed value or attachment for a cache record.
+ */
+class FCachePayload
 {
-	None,
-	Binary,
-	Object,
-	Package,
+public:
+	/** Construct a null cache payload. */
+	FCachePayload() = default;
+
+	/** Construct a cache payload from the compressed data hash. */
+	UE_API FCachePayload(const FCachePayloadId& Id, const FIoHash& CompressedDataHash);
+	/** Construct a cache payload from the compressed data. See FCachePayload::Compress. */
+	UE_API FCachePayload(const FCachePayloadId& Id, FSharedBuffer CompressedData);
+	/** Construct a cache payload from the compressed data its hash. See FCachePayload::Compress. */
+	UE_API FCachePayload(const FCachePayloadId& Id, FSharedBuffer CompressedData, const FIoHash& CompressedDataHash);
+
+	/** Returns a payload ID that uniquely identifies it within the scope of its cache record. */
+	inline const FCachePayloadId& GetId() const { return Id; }
+
+	/**
+	 * Returns the compressed data for the payload.
+	 *
+	 * This will be null when the originating cache request skips the payload.
+	 */
+	inline const FSharedBuffer& GetCompressedData() const { return CompressedData; }
+
+	/** Returns the hash of the compressed data for the payload. */
+	inline const FIoHash& GetCompressedDataHash() const { return CompressedDataHash; }
+
+	inline FSharedBuffer Decompress() const { return Decompress(CompressedData); }
+
+	/** Whether this is null. */
+	inline bool IsNull() const { return Id.IsNull(); }
+	/** Whether this is not null. */
+	inline bool IsValid() const { return !IsNull(); }
+	/** Whether this is not null. */
+	inline explicit operator bool() const { return IsValid(); }
+
+	/** Reset this to null. */
+	inline void Reset() { *this = FCachePayload(); }
+
+	/**
+	 * Compress the input buffer with the default settings for a cache payload.
+	 *
+	 * FIoCompression::Compress and its variants may be used instead to access more options.
+	 */
+	UE_API static FSharedBuffer Compress(FSharedBuffer Data);
+
+	/**
+	 * Decompress the input buffer.
+	 *
+	 * FIoCompression::Decompress and its variants may be used instead to access more options.
+	 */
+	UE_API static FSharedBuffer Decompress(FSharedBuffer CompressedData);
+
+private:
+	FCachePayloadId Id;
+	FSharedBuffer CompressedData;
+	FIoHash CompressedDataHash;
 };
 
 ///////////////////////////////////////////////////////////////////////////////////////////////////////////////////////
@@ -36,70 +94,202 @@ enum class ECacheRecordType : uint8
  * may be used to persist details such as the machine that created the value or how much time was
  * spent creating the value.
  *
- * The value may be binary or an object. Only object values may have attachments, and attachments
- * are consumed by accessing the cache record as a package. It is valid to access an object as if
- * it is a package with no attachments, or to access a package as though it is an object. This is
- * not valid for binary values, which may not have attachments or be interpreted as another type.
+ * The value and its attachments are compressed, and both the compressed and uncompressed formats
+ * of the payloads are exposed by the cache record.
  *
- * It is invalid for a package to include attachments that are not referenced by the object or by
- * one of its referenced attachments.
+ * The value, attachments, and metadata are optional and can be skipped when requesting a record.
+ * When the value or attachments have been skipped, the record will contain a payload with the ID
+ * and hash, but no buffer containing the associated data.
  */
 class FCacheRecord
 {
 public:
-	/** Construct an empty cache record. */
-	UE_API FCacheRecord();
+	/** Construct a null cache record. */
+	FCacheRecord() = default;
+
+	/** Clone this cache record. Prefer to move records, and clone only when a copy is necessary. */
+	UE_API FCacheRecord Clone() const;
+
+	/** Returns the key that identifies this record in the cache. Always available in a non-null cache record. */
+	UE_API const FCacheKey& GetKey() const;
+
+	/** Returns the metadata for the cache record unless it was skipped in the request. */
+	UE_API const FCbObjectRef& GetMeta() const;
 
-	/** Reset to an empty record. */
-	UE_API void Reset();
+	/** Returns the value. Null if the value was skipped in the request. */
+	UE_API FSharedBuffer GetValue() const;
 
-	inline const FCacheKey& GetKey() const { return Key; }
-	inline void SetKey(const FCacheKey& InKey) { Key = InKey; }
+	/** Returns the value payload. Always available in a non-null cache record. */
+	UE_API const FCachePayload& GetValuePayload() const;
 
-	inline const FCbObjectRef& GetMeta() const { return Meta; }
-	inline void SetMeta(FCbObjectRef InMeta) { Meta = MoveTemp(InMeta); }
+	/** Returns the attachment that matches the ID. Null if no match or attachments were skipped in the request. */
+	UE_API FSharedBuffer GetAttachment(const FCachePayloadId& Id) const;
 
-	inline ECacheRecordType GetType() const { return Type; }
+	/** Returns the attachment payload that matches the ID. Null if no match or the cache record is null. */
+	UE_API const FCachePayload& GetAttachmentPayload(const FCachePayloadId& Id) const;
 
-	/** Access the value as binary. Returns a null buffer if not binary. */
-	inline const FSharedBuffer& AsBinary() const { return Binary; }
-	/** Access the value as an object. Returns an empty object if not an object or package. */
-	inline const FCbObjectRef& AsObject() const { return Package.GetObject(); }
-	/** Access the value as a package. Returns a null package if not an object or package. */
-	inline const FCbPackage& AsPackage() const { return Package; }
+	/** Returns a view of the attachments. Always available in a non-null cache record, but data may be skipped. */
+	UE_API TConstArrayView<FCachePayload> GetAttachmentPayloads() const;
 
-	/** Returns the hash of the value, or the zero hash if there is no value. */
-	inline const FIoHash& GetValueHash() const;
+	/** Whether this is null. */
+	inline bool IsNull() const { return !Data; }
+	/** Whether this is not null. */
+	inline bool IsValid() const { return !IsNull(); }
+	/** Whether this is not null. */
+	inline explicit operator bool() const { return IsValid(); }
 
-	/** Set the value as binary. Removes any existing value and attachments. */
-	UE_API void SetBinary(FSharedBuffer Value);
-	/** Set the value as binary with its hash. Removes any existing value and attachments. */
-	UE_API void SetBinary(FSharedBuffer Value, const FIoHash& ValueHash);
+	/** Reset this to null. */
+	inline void Reset() { *this = FCacheRecord(); }
 
-	/** Set the value as an object. Removes any existing value and attachments. */
-	UE_API void SetObject(FCbObjectRef Value);
-	/** Set the value as an object with its hash. Removes any existing value and attachments. */
-	UE_API void SetObject(FCbObjectRef Value, const FIoHash& ValueHash);
+public:
+	// Internal API
+
+	/** Construct a cache record by cloning a cache record. Use Clone(). */
+	explicit FCacheRecord(const FCacheRecordData& Data);
 
-	/** Set the value as a package. Removes any existing value and attachments. */
-	UE_API void SetPackage(FCbPackage Value);
+	/** Construct a cache record from a builder. Use Builder.Build() or Builder.BuildAsync(). */
+	explicit FCacheRecord(FCacheRecordBuilderData&& Builder);
 
 private:
-	FCacheKey Key;
-	FCbObjectRef Meta;
+	TPimplPtr<FCacheRecordData> Data;
+};
+
+///////////////////////////////////////////////////////////////////////////////////////////////////////////////////////
+
+/**
+ * A cache record builder is used to construct a cache record.
+ *
+ * A key must be set before using the this to construct a cache record. The value and attachments
+ * can be provided as buffers or as payloads which are already compressed and have an identifier.
+ *
+ * @see FCacheRecord
+ */
+class FCacheRecordBuilder
+{
+public:
+	/** Construct an empty cache record builder. */
+	UE_API FCacheRecordBuilder();
+
+	/**
+	 * Set the key for the cache record.
+	 *
+	 * The key must uniquely correspond to the value and attachments for the cache record, while the
+	 * optional metadata may vary and may be non-deterministic.
+	 */
+	UE_API void SetKey(const FCacheKey& Key);
+
+	/**
+	 * Set the metadata for the cache record.
+	 *
+	 * The metadata is cloned if not owned.
+	 */
+	UE_API void SetMeta(FCbObjectRef Meta);
+
+	/**
+	 * Set the value for the cache record.
+	 *
+	 * @param Buffer The value, which is cloned if not owned, and is compressed by the builder.
+	 * @param Id An ID for the value that is unique within the scope of the cache record. If omitted
+	 *           the builder will create an ID by hashing the buffer.
+	 * @return The payload ID that was provided or created.
+	 */
+	UE_API FCachePayloadId SetValue(FSharedBuffer Buffer, FCachePayloadId Id = FCachePayloadId());
+
+	/**
+	 * Set the value for the cache record.
+	 *
+	 * @param Payload The value payload, which is reset to null.
+	 * @return The payload ID that was provided. Unique within the scope of the cache record.
+	 */
+	UE_API FCachePayloadId SetValue(FCachePayload&& Payload);
+
+	/**
+	 * Add an attachment to the cache record.
+	 *
+	 * @param Buffer The attachment, which is cloned if not owned, and is compressed by the builder.
+	 * @param Id An ID for the attachment that is unique within the scope of the cache record. If it
+	 *           is omitted the builder will create an ID by hashing the buffer.
+	 * @return The payload ID that was provided or created.
+	 */
+	UE_API FCachePayloadId AddAttachment(FSharedBuffer Buffer, FCachePayloadId Id = FCachePayloadId());
+
+	/**
+	 * Add an attachment to the cache record.
+	 *
+	 * @param Payload The attachment payload, which is reset to null.
+	 * @return The payload ID that was provided. Unique within the scope of the cache record.
+	 */
+	UE_API FCachePayloadId AddAttachment(FCachePayload&& Payload);
+
+	/**
+	 * Build a cache record, resetting this builder.
+	 *
+	 * Prefer BuildAsync() when the value or attachments are added from a buffer, as this must block
+	 * on compression of those buffers before it can construct a cache record.
+	 */
+	UE_API FCacheRecord Build();
 
-	FCbPackage Package;
-	FSharedBuffer Binary;
-	FIoHash BinaryHash;
-	ECacheRecordType Type = ECacheRecordType::None;
+	/**
+	 * Build a cache record asynchronously, resetting this builder.
+	 *
+	 * Prefer Build() when the value and attachments are added by payload, as compression is already
+	 * complete and BuildAsync() will complete immediately in that case.
+	 */
+	UE_API TFuture<FCacheRecord> BuildAsync();
+
+	/** Whether this is null. */
+	inline bool IsNull() const { return !Data; }
+	/** Whether this is not null. */
+	inline bool IsValid() const { return !IsNull(); }
+	/** Whether this is not null. */
+	inline explicit operator bool() const { return IsValid(); }
+
+	/** Reset this to empty. */
+	inline void Reset() { *this = FCacheRecordBuilder(); }
+
+private:
+	TPimplPtr<FCacheRecordBuilderData> Data;
 };
 
 ///////////////////////////////////////////////////////////////////////////////////////////////////////////////////////
 
-inline const FIoHash& FCacheRecord::GetValueHash() const
+/** Compare cache payloads by their ID. */
+struct FCachePayloadEqualById
+{
+	inline bool operator()(const FCachePayload& Payload1, const FCachePayload& Payload2) const
+	{
+		return Payload1.GetId() == Payload2.GetId();
+	}
+
+	inline bool operator()(const FCachePayload& Payload1, const FCachePayloadId& PayloadId2) const
+	{
+		return Payload1.GetId() == PayloadId2;
+	}
+
+	inline bool operator()(const FCachePayloadId& PayloadId1, const FCachePayload& Payload2) const
+	{
+		return PayloadId1 == Payload2.GetId();
+	}
+};
+
+/** Compare cache payloads by their ID. */
+struct FCachePayloadLessById
 {
-	return Type == ECacheRecordType::Binary ? BinaryHash : Package.GetObjectHash();
-}
+	inline bool operator()(const FCachePayload& Payload1, const FCachePayload& Payload2) const
+	{
+		return Payload1.GetId() < Payload2.GetId();
+	}
+
+	inline bool operator()(const FCachePayload& Payload1, const FCachePayloadId& PayloadId2) const
+	{
+		return Payload1.GetId() < PayloadId2;
+	}
+
+	inline bool operator()(const FCachePayloadId& PayloadId1, const FCachePayload& Payload2) const
+	{
+		return PayloadId1 < Payload2.GetId();
+	}
+};
 
 ///////////////////////////////////////////////////////////////////////////////////////////////////////////////////////
 
diff --git a/Engine/Source/Editor/UnrealEd/Private/EditorDomain/EditorDomainArchive.cpp b/Engine/Source/Editor/UnrealEd/Private/EditorDomain/EditorDomainArchive.cpp
--- a/Engine/Source/Editor/UnrealEd/Private/EditorDomain/EditorDomainArchive.cpp
+++ b/Engine/Source/Editor/UnrealEd/Private/EditorDomain/EditorDomainArchive.cpp
@@ -185,7 +185,7 @@ void FEditorDomainReadArchive::OnCacheRequestComplete(UE::DerivedData::FCacheGet
 	{
 		const FCbObjectRef& MetaData = Params.Record.GetMeta();
 		int64 FileSize = MetaData["FileSize"].AsInt64(-1);
-		const FSharedBuffer& LocalBytes = Params.Record.AsBinary();
+		FSharedBuffer LocalBytes = Params.Record.GetValue();
 		if (static_cast<int64>(LocalBytes.GetSize()) != FileSize)
 		{
 			UE_LOG(LogEditorDomain, Warning, TEXT("Package %s received invalid record from EditorDomainPackage table ")
@@ -197,8 +197,8 @@ void FEditorDomainReadArchive::OnCacheRequestComplete(UE::DerivedData::FCacheGet
 		{
 			AsyncSource = ESource::Bytes;
 			Size = FileSize;
-			Bytes = LocalBytes;
-			Bytes.MakeOwned();
+			LocalBytes.MakeOwned();
+			Bytes = MoveTemp(LocalBytes);
 			PackageFormat = EPackageFormat::Binary;
 			PackageSource->Source = FEditorDomain::EPackageSource::Editor;
 		}
@@ -414,7 +414,7 @@ void FEditorDomainAsyncReadFileHandle::OnCacheRequestComplete(UE::DerivedData::F
 	{
 		const FCbObjectRef& MetaData = Params.Record.GetMeta();
 		int64 FileSize = MetaData["FileSize"].AsInt64(-1);
-		const FSharedBuffer& LocalBytes = Params.Record.AsBinary();
+		FSharedBuffer LocalBytes = Params.Record.GetValue();
 		if (static_cast<int64>(LocalBytes.GetSize()) != FileSize)
 		{
 			UE_LOG(LogEditorDomain, Warning, TEXT("Package %s received invalid record from EditorDomainPackage ")
@@ -425,8 +425,8 @@ void FEditorDomainAsyncReadFileHandle::OnCacheRequestComplete(UE::DerivedData::F
 		else
 		{
 			AsyncSource = ESource::Bytes;
-			Bytes = LocalBytes;
-			Bytes.MakeOwned();
+			LocalBytes.MakeOwned();
+			Bytes = MoveTemp(LocalBytes);
 			PackageSource->Source = FEditorDomain::EPackageSource::Editor;
 		}
 	}
diff --git a/Engine/Source/Editor/UnrealEd/Private/EditorDomain/EditorDomainSave.cpp b/Engine/Source/Editor/UnrealEd/Private/EditorDomain/EditorDomainSave.cpp
--- a/Engine/Source/Editor/UnrealEd/Private/EditorDomain/EditorDomainSave.cpp
+++ b/Engine/Source/Editor/UnrealEd/Private/EditorDomain/EditorDomainSave.cpp
@@ -758,15 +758,16 @@ bool FEditorDomainSaveServer::TrySavePackage(const FPackagePath& PackagePath, FS
 		return false;
 	}
 
-	UE::DerivedData::FCacheRecord Record;
-	Record.SetKey(GetEditorDomainPackageKey(PackageDigest));
+	UE::DerivedData::FCacheRecordBuilder RecordBuilder;
+	RecordBuilder.SetKey(GetEditorDomainPackageKey(PackageDigest));
 	TCbWriter<256> MetaData;
 	MetaData.BeginObject();
 	MetaData << "FileSize" << TempBytes.Num();
 	MetaData.EndObject();
-	Record.SetMeta(MetaData.Save().AsObjectRef());
-	Record.SetBinary(FSharedBuffer::MakeView(TempBytes.GetData(), TempBytes.Num()));
-	GetDerivedDataCacheRef().GetCache().Put({ Record }, PackagePath.GetDebugName());
+	RecordBuilder.SetMeta(MetaData.Save().AsObjectRef());
+	RecordBuilder.SetValue(FSharedBuffer::MakeView(TempBytes.GetData(), TempBytes.Num()));
+	UE::DerivedData::FCacheRecord Record = RecordBuilder.Build();
+	GetDerivedDataCacheRef().GetCache().Put(MakeArrayView(&Record, 1), PackagePath.GetDebugName());
 	return true;
 }
 
diff --git a/Engine/Source/Editor/UnrealEd/Private/EditorDomain/EditorDomainUtils.cpp b/Engine/Source/Editor/UnrealEd/Private/EditorDomain/EditorDomainUtils.cpp
--- a/Engine/Source/Editor/UnrealEd/Private/EditorDomain/EditorDomainUtils.cpp
+++ b/Engine/Source/Editor/UnrealEd/Private/EditorDomain/EditorDomainUtils.cpp
@@ -19,7 +19,7 @@ namespace EditorDomain
 {
 
 // Change to a new guid when EditorDomain needs to be invalidated
-const TCHAR* EditorDomainVersion = TEXT("D3317196-098A-449B-B740-631D94CD7C93");
+const TCHAR* EditorDomainVersion = TEXT("34A0B57B-32D5-4229-A6A3-DA0FB7BE0440");
 // Identifier of the CacheBucket for EditorDomainPackages
 const TCHAR* EditorDomainPackageBucketName = TEXT("EditorDomainPackage");
 
