Disable intermittent DDC warning when running on build machines.

#rnx

diff --git a/Engine/Source/Developer/DerivedDataCache/Private/FileSystemDerivedDataBackend.cpp b/Engine/Source/Developer/DerivedDataCache/Private/FileSystemDerivedDataBackend.cpp
--- a/Engine/Source/Developer/DerivedDataCache/Private/FileSystemDerivedDataBackend.cpp
+++ b/Engine/Source/Developer/DerivedDataCache/Private/FileSystemDerivedDataBackend.cpp
@@ -177,7 +177,14 @@ public:
 
 			if (SpeedClass <= FDerivedDataBackendInterface::ESpeedClass::Slow && !bReadOnly)
 			{
-				UE_LOG(LogDerivedDataCache, Warning, TEXT("Access to %s appears to be slow. 'Touch' will be disabled and queries/writes will be limited."), *CachePath);
+				if (GIsBuildMachine)
+				{
+					UE_LOG(LogDerivedDataCache, Display, TEXT("Access to %s appears to be slow. 'Touch' will be disabled and queries/writes will be limited."), *CachePath);
+				}
+				else
+				{
+					UE_LOG(LogDerivedDataCache, Warning, TEXT("Access to %s appears to be slow. 'Touch' will be disabled and queries/writes will be limited."), *CachePath);
+				}
 				bTouch = false;
 				//bReadOnly = true;
 			}
