HttpDerivedBackend - Allow IoHash header as well
* Removed pointless double hashing for sha1 hashes.
* Fixed spelling recieved -> received

diff --git a/Engine/Source/Developer/DerivedDataCache/Private/HttpDerivedDataBackend.cpp b/Engine/Source/Developer/DerivedDataCache/Private/HttpDerivedDataBackend.cpp
--- a/Engine/Source/Developer/DerivedDataCache/Private/HttpDerivedDataBackend.cpp
+++ b/Engine/Source/Developer/DerivedDataCache/Private/HttpDerivedDataBackend.cpp
@@ -16,6 +16,7 @@
 #include "Dom/JsonObject.h"
 #include "GenericPlatform/GenericPlatformFile.h"
 #include "HAL/PlatformFileManager.h"
+#include "IO/IoHash.h"
 #include "Misc/FileHelper.h"
 #include "Misc/ScopeLock.h"
 #include "Misc/SecureHash.h"
@@ -872,9 +873,9 @@ static CURLcode sslctx_function(CURL * curl, void * sslctx, void * parm)
 
 
 /**
- * Verifies the integrity of the recieved data using supplied checksum.
- * @param Hash Recieved hash value.
- * @param Payload Payload recieved.
+ * Verifies the integrity of the received data using supplied checksum.
+ * @param Hash received hash value.
+ * @param Payload Payload received.
  * @return True if the data is correct, false if checksums doesn't match.
  */
 bool VerifyPayload(FSHAHash Hash, const TArray<uint8>& Payload)
@@ -886,7 +887,31 @@ bool VerifyPayload(FSHAHash Hash, const TArray<uint8>& Payload)
 	{
 		UE_LOG(LogDerivedDataCache,
 			Warning,
-			TEXT("Checksum from server did not match recieved data (%s vs %s). Discarding cached result."),
+			TEXT("Checksum from server did not match received data (%s vs %s). Discarding cached result."),
+			*Hash.ToString(),
+			*PayloadHash.ToString()
+		);
+		return false;
+	}
+
+	return true;
+}
+
+/**
+ * Verifies the integrity of the received data using supplied checksum.
+ * @param Hash received hash value.
+ * @param Payload Payload received.
+ * @return True if the data is correct, false if checksums doesn't match.
+ */
+bool VerifyPayload(FIoHash Hash, const TArray<uint8>& Payload)
+{
+	FIoHash PayloadHash = FIoHash::HashBuffer(Payload.GetData(), Payload.Num());
+
+	if (Hash != PayloadHash)
+	{
+		UE_LOG(LogDerivedDataCache,
+			Warning,
+			TEXT("Checksum from server did not match received data (%s vs %s). Discarding cached result."),
 			*Hash.ToString(),
 			*PayloadHash.ToString()
 		);
@@ -898,21 +923,24 @@ bool VerifyPayload(FSHAHash Hash, const TArray<uint8>& Payload)
 
 
 /**
- * Verifies the integrity of the recieved data using supplied checksum.
- * @param Request Request that the data was be recieved with.
- * @param Payload Payload recieved.
+ * Verifies the integrity of the received data using supplied checksum.
+ * @param Request Request that the data was be received with.
+ * @param Payload Payload received.
  * @return True if the data is correct, false if checksums doesn't match.
  */
 bool VerifyRequest(const FRequest* Request, const TArray<uint8>& Payload)
 {
-	FString RecievedHashStr;
-	FSHAHash PayloadHash;
-	FSHA1::HashBuffer(Payload.GetData(), Payload.Num(), PayloadHash.Hash);
-	if (Request->GetHeader("X-Jupiter-Sha1", RecievedHashStr))
+	FString receivedHashStr;
+	if (Request->GetHeader("X-Jupiter-Sha1", receivedHashStr))
+	{
+		FSHAHash receivedHash;
+		receivedHash.FromString(receivedHashStr);
+		return VerifyPayload(receivedHash, Payload);
+	}
+	if (Request->GetHeader("X-Jupiter-IoHash", receivedHashStr))
 	{
-		FSHAHash RecievedHash;
-		RecievedHash.FromString(RecievedHashStr);
-		return VerifyPayload(RecievedHash, Payload);
+		FIoHash receivedHash(receivedHashStr);
+		return VerifyPayload(receivedHash, Payload);
 	}
 	UE_LOG(LogDerivedDataCache, Warning, TEXT("HTTP server did not send a content hash. Wrong server version?"));
 	return true;
@@ -920,7 +948,7 @@ bool VerifyRequest(const FRequest* Request, const TArray<uint8>& Payload)
 
 /**
  * Adds a checksum (as request header) for a given payload. Jupiter will use this to verify the integrity
- * of the recieved data.
+ * of the received data.
  * @param Request Request that the data will be sent with.
  * @param Payload Payload that will be sent.
  * @return True on success, false on failure.
