DDC: Capture cache value arrays into shared buffers instead of copying them

#rnx

diff --git a/Engine/Source/Developer/DerivedDataCache/Private/DerivedDataCache.cpp b/Engine/Source/Developer/DerivedDataCache/Private/DerivedDataCache.cpp
--- a/Engine/Source/Developer/DerivedDataCache/Private/DerivedDataCache.cpp
+++ b/Engine/Source/Developer/DerivedDataCache/Private/DerivedDataCache.cpp
@@ -926,16 +926,20 @@ void FCache::Get(
 	}
 
 	// Request the record from storage.
-	TArray<uint8> RecordData;
-	if (!GetDerivedDataCacheRef().GetSynchronous(*MakeRecordKey(Key), RecordData, Context))
+	FSharedBuffer RecordBuffer;
 	{
-		UE_LOG(LogDerivedDataCache, Verbose, TEXT("Cache: Get cache miss for %s from '%.*s'"),
-			*TToString<96>(Key), Context.Len(), Context.GetData());
-		return;
+		TArray<uint8> RecordData;
+		if (!GetDerivedDataCacheRef().GetSynchronous(*MakeRecordKey(Key), RecordData, Context))
+		{
+			UE_LOG(LogDerivedDataCache, Verbose, TEXT("Cache: Get cache miss for %s from '%.*s'"),
+				*TToString<96>(Key), Context.Len(), Context.GetData());
+			return;
+		}
+		RecordBuffer = MakeSharedBufferFromArray(MoveTemp(RecordData));
 	}
 
 	// Validate that the record can be read as compact binary without crashing.
-	if (ValidateCompactBinaryRange(MakeMemoryView(RecordData), ECbValidateMode::Default) != ECbValidateError::None)
+	if (ValidateCompactBinaryRange(RecordBuffer, ECbValidateMode::Default) != ECbValidateError::None)
 	{
 		UE_LOG(LogDerivedDataCache, Verbose, TEXT("Cache: Get cache miss with corrupted record for %s from '%.*s'"),
 			*TToString<96>(Key), Context.Len(), Context.GetData());
@@ -943,8 +947,8 @@ void FCache::Get(
 	}
 
 	// Read the record from its compact binary fields.
-	const FCbObjectView RecordObject(RecordData.GetData());
-	const FCbObjectView MetaObject = RecordObject["Meta"_ASV].AsObjectView();
+	const FCbObject RecordObject(MoveTemp(RecordBuffer));
+	const FCbObject MetaObject = RecordObject.Find("Meta"_ASV).AsObject();
 	const FCbObjectView ValueObject = RecordObject["Value"_ASV].AsObjectView();
 	const FCbArrayView AttachmentsArray = RecordObject["Attachments"_ASV].AsArrayView();
 
@@ -964,8 +968,7 @@ void FCache::Get(
 		{
 			if (FIoHash::HashBuffer(MakeMemoryView(PayloadData)) == PayloadHash)
 			{
-				FSharedBuffer PayloadBuffer = FSharedBuffer::Clone(MakeMemoryView(PayloadData));
-				return FCachePayload(PayloadId, MoveTemp(PayloadBuffer), PayloadHash);
+				return FCachePayload(PayloadId, MakeSharedBufferFromArray(MoveTemp(PayloadData)), PayloadHash);
 			}
 			else
 			{
@@ -1018,7 +1021,7 @@ void FCache::Get(
 	// Populate the builder now that any error conditions have been handled.
 	if (!EnumHasAnyFlags(Policy, ECachePolicy::SkipMeta) && MetaObject.CreateViewIterator())
 	{
-		Builder.SetMeta(FCbObject::Clone(MetaObject));
+		Builder.SetMeta(MetaObject);
 	}
 	if (Value)
 	{
@@ -1087,16 +1090,20 @@ void FCache::GetPayload(
 		if (bPayloadExists)
 		{
 			// Request the record from storage.
-			TArray<uint8> RecordData;
-			if (!GetDerivedDataCacheRef().GetSynchronous(*MakeRecordKey(Key.GetKey()), RecordData, Context))
+			FSharedBuffer RecordBuffer;
 			{
-				UE_LOG(LogDerivedDataCache, Verbose, TEXT("Cache: Get cache miss with missing record for %s from '%.*s'"),
-					*TToString<96>(Key), Context.Len(), Context.GetData());
-				return;
+				TArray<uint8> RecordData;
+				if (!GetDerivedDataCacheRef().GetSynchronous(*MakeRecordKey(Key.GetKey()), RecordData, Context))
+				{
+					UE_LOG(LogDerivedDataCache, Verbose, TEXT("Cache: Get cache miss with missing record for %s from '%.*s'"),
+						*TToString<96>(Key), Context.Len(), Context.GetData());
+					return;
+				}
+				RecordBuffer = MakeSharedBufferFromArray(MoveTemp(RecordData));
 			}
 
 			// Validate that the record can be read as compact binary without crashing.
-			if (ValidateCompactBinaryRange(MakeMemoryView(RecordData), ECbValidateMode::Default) != ECbValidateError::None)
+			if (ValidateCompactBinaryRange(RecordBuffer, ECbValidateMode::Default) != ECbValidateError::None)
 			{
 				UE_LOG(LogDerivedDataCache, Verbose, TEXT("Cache: Get cache miss with corrupted record for %s from '%.*s'"),
 					*TToString<96>(Key), Context.Len(), Context.GetData());
@@ -1104,7 +1111,7 @@ void FCache::GetPayload(
 			}
 
 			// Read the record from its compact binary fields.
-			const FCbObjectView RecordObject(RecordData.GetData());
+			const FCbObject RecordObject(MoveTemp(RecordBuffer));
 			const FCbObjectView ValueObject = RecordObject["Value"_ASV].AsObjectView();
 			const FCbArrayView AttachmentsArray = RecordObject["Attachments"_ASV].AsArrayView();
 
@@ -1156,7 +1163,7 @@ void FCache::GetPayload(
 		return;
 	}
 
-	Params.Payload = FCachePayload(Key.GetId(), FSharedBuffer::Clone(MakeMemoryView(Data)));
+	Params.Payload = FCachePayload(Key.GetId(), MakeSharedBufferFromArray(MoveTemp(Data)));
 
 	UE_LOG(LogDerivedDataCache, Verbose, TEXT("Cache: GetPayload cache hit for %s from '%.*s'"),
 		*TToString<96>(Key), Context.Len(), Context.GetData());
