DDC: Refactored the Cache API to prepare for the Build API

- Removed the "Cache" prefix from ECachePriority, FCacheRequest, ICacheRequest, FCachePayload, FCachePayloadId to allow reuse in the Build API.
- Added FCompressedBuffer as the starting point for the API for creating and accessing compressed buffers.

#rnx

diff --git a/Engine/Source/Developer/DerivedDataCache/Private/DerivedDataCache.cpp b/Engine/Source/Developer/DerivedDataCache/Private/DerivedDataCache.cpp
--- a/Engine/Source/Developer/DerivedDataCache/Private/DerivedDataCache.cpp
+++ b/Engine/Source/Developer/DerivedDataCache/Private/DerivedDataCache.cpp
@@ -716,25 +716,25 @@ class FCache final : public ICache
 public:
 	virtual ~FCache() = default;
 
-	virtual FCacheRequest Put(
+	virtual FRequest Put(
 		TArrayView<FCacheRecord> Records,
 		FStringView Context,
 		ECachePolicy Policy,
-		ECachePriority Priority,
+		EPriority Priority,
 		FOnCachePutComplete&& Callback) final;
 
-	virtual FCacheRequest Get(
+	virtual FRequest Get(
 		TConstArrayView<FCacheKey> Keys,
 		FStringView Context,
 		ECachePolicy Policy,
-		ECachePriority Priority,
+		EPriority Priority,
 		FOnCacheGetComplete&& Callback) final;
 
-	virtual FCacheRequest GetPayloads(
+	virtual FRequest GetPayloads(
 		TConstArrayView<FCachePayloadKey> Keys,
 		FStringView Context,
 		ECachePolicy Policy,
-		ECachePriority Priority,
+		EPriority Priority,
 		FOnCacheGetPayloadComplete&& Callback) final;
 
 	virtual void CancelAll() final;
@@ -744,21 +744,21 @@ private:
 		FCacheRecord Record,
 		FStringView Context,
 		ECachePolicy Policy,
-		ECachePriority Priority,
+		EPriority Priority,
 		FOnCachePutComplete& Callback);
 
 	void Get(
 		const FCacheKey& Key,
 		FStringView Context,
 		ECachePolicy Policy,
-		ECachePriority Priority,
+		EPriority Priority,
 		FOnCacheGetComplete& Callback);
 
 	void GetPayload(
 		const FCachePayloadKey& Key,
 		FStringView Context,
 		ECachePolicy Policy,
-		ECachePriority Priority,
+		EPriority Priority,
 		FOnCacheGetPayloadComplete& Callback);
 
 	template <int32 BufferSize>
@@ -802,7 +802,7 @@ void FCache::Put(
 	FCacheRecord Record,
 	FStringView Context,
 	ECachePolicy Policy,
-	ECachePriority Priority,
+	EPriority Priority,
 	FOnCachePutComplete& Callback)
 {
 	FCachePutCompleteParams Params;
@@ -824,9 +824,9 @@ void FCache::Put(
 		return;
 	}
 
-	const auto PutPayload = [&Key = Params.Key, Context](const FCachePayload& Payload)
+	const auto PutPayload = [&Key = Params.Key, Context](const FPayload& Payload)
 	{
-		const FSharedBuffer& Buffer = Payload.GetCompressedData();
+		const FSharedBuffer Buffer = Payload.GetBuffer().GetCompressed();
 		check(Buffer && Buffer.GetSize() <= MAX_int32);
 		TConstArrayView<uint8> BufferView = MakeArrayView(
 			static_cast<const uint8*>(Buffer.GetData()),
@@ -837,23 +837,23 @@ void FCache::Put(
 	FCbWriter Writer;
 	Writer.BeginObject();
 	{
-		if (const FCachePayload& Value = Record.GetValuePayload())
+		if (const FPayload& Value = Record.GetValuePayload())
 		{
 			PutPayload(Value);
 			Writer.BeginObject("Value"_ASV);
 			Writer.AddObjectId("Id"_ASV, Value.GetId().ToObjectId());
-			Writer.AddBinaryAttachment("Hash"_ASV, Value.GetCompressedDataHash());
+			Writer.AddBinaryAttachment("Hash"_ASV, Value.GetHash());
 			Writer.EndObject();
 		}
 		if (!Record.GetAttachmentPayloads().IsEmpty())
 		{
 			Writer.BeginArray("Attachments"_ASV);
-			for (const FCachePayload& Attachment : Record.GetAttachmentPayloads())
+			for (const FPayload& Attachment : Record.GetAttachmentPayloads())
 			{
 				PutPayload(Attachment);
 				Writer.BeginObject();
 				Writer.AddObjectId("Id"_ASV, Attachment.GetId().ToObjectId());
-				Writer.AddBinaryAttachment("Hash"_ASV, Attachment.GetCompressedDataHash());
+				Writer.AddBinaryAttachment("Hash"_ASV, Attachment.GetHash());
 				Writer.EndObject();
 			}
 			Writer.EndArray();
@@ -880,30 +880,30 @@ void FCache::Put(
 	Params.Status = ECacheStatus::Cached;
 }
 
-FCacheRequest FCache::Put(
+FRequest FCache::Put(
 	TArrayView<FCacheRecord> Records,
 	FStringView Context,
 	ECachePolicy Policy,
-	ECachePriority Priority,
+	EPriority Priority,
 	FOnCachePutComplete&& Callback)
 {
 	for (FCacheRecord& Record : Records)
 	{
 		Put(MoveTemp(Record), Context, Policy, Priority, Callback);
 	}
-	return FCacheRequest();
+	return FRequest();
 }
 
 void FCache::Get(
 	const FCacheKey& Key,
 	FStringView Context,
 	ECachePolicy Policy,
-	ECachePriority Priority,
+	EPriority Priority,
 	FOnCacheGetComplete& Callback)
 {
-	FCachePayload Value;
-	TArray<FCachePayload> Attachments;
-	TArray<FCachePayloadId, TInlineAllocator<1>> SkippedPayloads;
+	FPayload Value;
+	TArray<FPayload> Attachments;
+	TArray<FPayloadId, TInlineAllocator<1>> SkippedPayloads;
 	FCacheRecordBuilder Builder;
 	Builder.SetKey(Key);
 
@@ -953,35 +953,35 @@ void FCache::Get(
 	const FCbArrayView AttachmentsArray = RecordObject["Attachments"_ASV].AsArrayView();
 
 	// Read the value and attachments if they have been requested.
-	const auto GetPayload = [&SkippedPayloads, &Key, Policy, Context](FCbObjectView PayloadObject, ECachePolicy SkipFlag) -> FCachePayload
+	const auto GetPayload = [&SkippedPayloads, &Key, Policy, Context](FCbObjectView PayloadObject, ECachePolicy SkipFlag) -> FPayload
 	{
 		TArray<uint8> PayloadData;
-		const FCachePayloadId PayloadId = FCachePayloadId::FromObjectId(PayloadObject["Id"_ASV].AsObjectId());
+		const FPayloadId PayloadId = FPayloadId::FromObjectId(PayloadObject["Id"_ASV].AsObjectId());
 		const FIoHash PayloadHash = PayloadObject["Hash"_ASV].AsBinaryAttachment();
 		const FCachePayloadKey PayloadKey(Key, PayloadId);
 		if (EnumHasAnyFlags(Policy, SkipFlag))
 		{
 			SkippedPayloads.Add(PayloadId);
-			return FCachePayload(PayloadId, PayloadHash);
+			return FPayload(PayloadId, PayloadHash);
 		}
 		else if (GetDerivedDataCacheRef().GetSynchronous(*MakePayloadKey(PayloadKey), PayloadData, Context))
 		{
 			if (FIoHash::HashBuffer(MakeMemoryView(PayloadData)) == PayloadHash)
 			{
-				return FCachePayload(PayloadId, MakeSharedBufferFromArray(MoveTemp(PayloadData)), PayloadHash);
+				return FPayload(PayloadId, PayloadHash, FCompressedBuffer(MakeSharedBufferFromArray(MoveTemp(PayloadData))));
 			}
 			else
 			{
 				UE_LOG(LogDerivedDataCache, Verbose, TEXT("Cache: Get cache miss with corrupted content for %s from '%.*s'"),
 					*TToString<96>(PayloadKey), Context.Len(), Context.GetData());
-				return FCachePayload();
+				return FPayload();
 			}
 		}
 		else
 		{
 			UE_LOG(LogDerivedDataCache, Verbose, TEXT("Cache: Get cache miss with missing content for %s from '%.*s'"),
 				*TToString<96>(PayloadKey), Context.Len(), Context.GetData());
-			return FCachePayload();
+			return FPayload();
 		}
 	};
 
@@ -1006,7 +1006,7 @@ void FCache::Get(
 		TArray<FString, TInlineAllocator<1>> SkippedPayloadKeys;
 		SkippedPayloadKeys.Reserve(SkippedPayloads.Num());
 		Algo::Transform(SkippedPayloads, SkippedPayloadKeys,
-			[&Key](const FCachePayloadId& Id) -> FString { return MakePayloadKey(FCachePayloadKey(Key, Id)); });
+			[&Key](const FPayloadId& Id) -> FString { return MakePayloadKey(FCachePayloadKey(Key, Id)); });
 		const bool bPayloadsExist = EnumHasAnyFlags(Policy, ECachePolicy::StoreLocal)
 			? GetDerivedDataCacheRef().TryToPrefetch(SkippedPayloadKeys, Context)
 			: GetDerivedDataCacheRef().AllCachedDataProbablyExists(SkippedPayloadKeys);
@@ -1027,7 +1027,7 @@ void FCache::Get(
 	{
 		Builder.SetValue(MoveTemp(Value));
 	}
-	for (FCachePayload& Attachment : Attachments)
+	for (FPayload& Attachment : Attachments)
 	{
 		Builder.AddAttachment(MoveTemp(Attachment));
 	}
@@ -1037,25 +1037,25 @@ void FCache::Get(
 	Params.Status = ECacheStatus::Cached;
 }
 
-FCacheRequest FCache::Get(
+FRequest FCache::Get(
 	TConstArrayView<FCacheKey> Keys,
 	FStringView Context,
 	ECachePolicy Policy,
-	ECachePriority Priority,
+	EPriority Priority,
 	FOnCacheGetComplete&& Callback)
 {
 	for (const FCacheKey& Key : Keys)
 	{
 		Get(Key, Context, Policy, Priority, Callback);
 	}
-	return FCacheRequest();
+	return FRequest();
 }
 
 void FCache::GetPayload(
 	const FCachePayloadKey& Key,
 	FStringView Context,
 	ECachePolicy Policy,
-	ECachePriority Priority,
+	EPriority Priority,
 	FOnCacheGetPayloadComplete& Callback)
 {
 	FCacheGetPayloadCompleteParams Params;
@@ -1067,7 +1067,7 @@ void FCache::GetPayload(
 		{
 			if (!Params.Payload)
 			{
-				Params.Payload = FCachePayload(Key.GetId(), FIoHash());
+				Params.Payload = FPayload(Key.GetId(), FIoHash());
 			}
 			Callback(MoveTemp(Params));
 		}
@@ -1119,7 +1119,7 @@ void FCache::GetPayload(
 			const FCbObjectId KeyObjectId = Key.GetId().ToObjectId();
 			if (ValueObject["Id"_ASV].AsObjectId() == KeyObjectId)
 			{
-				Params.Payload = FCachePayload(Key.GetId(), ValueObject["Hash"_ASV].AsBinaryAttachment());
+				Params.Payload = FPayload(Key.GetId(), ValueObject["Hash"_ASV].AsBinaryAttachment());
 			}
 			else
 			{
@@ -1127,7 +1127,7 @@ void FCache::GetPayload(
 				{
 					if (AttachmentField.AsObjectView()["Id"_ASV].AsObjectId() == KeyObjectId)
 					{
-						Params.Payload = FCachePayload(Key.GetId(), AttachmentField.AsObjectView()["Hash"_ASV].AsBinaryAttachment());
+						Params.Payload = FPayload(Key.GetId(), AttachmentField.AsObjectView()["Hash"_ASV].AsBinaryAttachment());
 						break;
 					}
 				}
@@ -1163,25 +1163,25 @@ void FCache::GetPayload(
 		return;
 	}
 
-	Params.Payload = FCachePayload(Key.GetId(), MakeSharedBufferFromArray(MoveTemp(Data)));
+	Params.Payload = FPayload(Key.GetId(), FCompressedBuffer(MakeSharedBufferFromArray(MoveTemp(Data))));
 
 	UE_LOG(LogDerivedDataCache, Verbose, TEXT("Cache: GetPayload cache hit for %s from '%.*s'"),
 		*TToString<96>(Key), Context.Len(), Context.GetData());
 	Params.Status = ECacheStatus::Cached;
 }
 
-FCacheRequest FCache::GetPayloads(
+FRequest FCache::GetPayloads(
 	TConstArrayView<FCachePayloadKey> Keys,
 	FStringView Context,
 	ECachePolicy Policy,
-	ECachePriority Priority,
+	EPriority Priority,
 	FOnCacheGetPayloadComplete&& Callback)
 {
 	for (const FCachePayloadKey& Key : Keys)
 	{
 		GetPayload(Key, Context, Policy, Priority, Callback);
 	}
-	return FCacheRequest();
+	return FRequest();
 }
 
 void FCache::CancelAll()
diff --git a/Engine/Source/Developer/DerivedDataCache/Private/DerivedDataCacheKey.cpp b/Engine/Source/Developer/DerivedDataCache/Private/DerivedDataCacheKey.cpp
--- a/Engine/Source/Developer/DerivedDataCache/Private/DerivedDataCacheKey.cpp
+++ b/Engine/Source/Developer/DerivedDataCache/Private/DerivedDataCacheKey.cpp
@@ -3,9 +3,7 @@
 #include "DerivedDataCacheKey.h"
 
 #include "Algo/AllOf.h"
-#include "Containers/StringConv.h"
 #include "Misc/StringBuilder.h"
-#include "Serialization/CompactBinary.h"
 #include "UObject/NameTypes.h"
 
 namespace UE
@@ -72,53 +70,6 @@ FString FCacheKey::ToString() const
 
 ///////////////////////////////////////////////////////////////////////////////////////////////////////////////////////
 
-FCachePayloadId FCachePayloadId::FromName(FAnsiStringView Name)
-{
-	checkf(!Name.IsEmpty(), TEXT("Payload ID requires a non-empty name."));
-	return FCachePayloadId::FromHash(FIoHash::HashBuffer(Name.GetData(), Name.Len()));
-}
-
-FCachePayloadId FCachePayloadId::FromName(FWideStringView Name)
-{
-	checkf(!Name.IsEmpty(), TEXT("Payload ID requires a non-empty name."));
-	TAnsiStringBuilder<128> Name8;
-	const int32 Len = int32(FTCHARToUTF8_Convert::ConvertedLength(Name.GetData(), Name.Len()));
-	const int32 Offset = Name8.AddUninitialized(Len);
-	const int32 WrittenLen = FTCHARToUTF8_Convert::Convert(Name8.GetData() + Offset, Len, Name.GetData(), Name.Len());
-	checkf(Len == WrittenLen, TEXT("Name '%.*s' failed to convert to a payload ID."), Name.Len(), Name.GetData());
-	return FCachePayloadId::FromName(Name8);
-}
-
-FCachePayloadId FCachePayloadId::FromObjectId(const FCbObjectId& ObjectId)
-{
-	checkf(ObjectId != FCbObjectId(), TEXT("Payload ID requires a non-zero ObjectId."));
-	return FCachePayloadId(ObjectId.GetView());
-}
-
-FCbObjectId FCachePayloadId::ToObjectId() const
-{
-	return FCbObjectId(GetView());
-}
-
-void FCachePayloadId::ToString(FAnsiStringBuilderBase& Builder) const
-{
-	UE::String::BytesToHexLower(Bytes, Builder);
-}
-
-void FCachePayloadId::ToString(FWideStringBuilderBase& Builder) const
-{
-	UE::String::BytesToHexLower(Bytes, Builder);
-}
-
-FString FCachePayloadId::ToString() const
-{
-	TStringBuilder<16> Out;
-	ToString(Out);
-	return FString(Out);
-}
-
-///////////////////////////////////////////////////////////////////////////////////////////////////////////////////////
-
 void FCachePayloadKey::ToString(FAnsiStringBuilderBase& Builder) const
 {
 	Builder << Key << '/' << Id;
diff --git a/Engine/Source/Developer/DerivedDataCache/Private/DerivedDataCacheRecord.cpp b/Engine/Source/Developer/DerivedDataCache/Private/DerivedDataCacheRecord.cpp
--- a/Engine/Source/Developer/DerivedDataCache/Private/DerivedDataCacheRecord.cpp
+++ b/Engine/Source/Developer/DerivedDataCache/Private/DerivedDataCacheRecord.cpp
@@ -4,6 +4,7 @@
 
 #include "Async/Future.h"
 #include "Misc/StringBuilder.h"
+#include "Serialization/CompactBinary.h"
 
 namespace UE
 {
@@ -12,71 +13,12 @@ namespace DerivedData
 
 ///////////////////////////////////////////////////////////////////////////////////////////////////////////////////////
 
-FCachePayload::FCachePayload(const FCachePayloadId& InId, const FIoHash& InCompressedDataHash)
-	: Id(InId)
-	, CompressedDataHash(InCompressedDataHash)
-{
-	checkf(Id.IsValid(), TEXT("A valid ID is required to construct a cache payload."));
-}
-
-FCachePayload::FCachePayload(const FCachePayloadId& InId, FSharedBuffer InCompressedData)
-	: Id(InId)
-	, CompressedData(MoveTemp(InCompressedData))
-{
-	checkf(Id.IsValid(), TEXT("A valid ID is required to construct a cache payload."));
-	if (CompressedData.GetSize())
-	{
-		CompressedDataHash = FIoHash::HashBuffer(CompressedData);
-	}
-}
-
-FCachePayload::FCachePayload(const FCachePayloadId& InId, FSharedBuffer InCompressedData, const FIoHash& InCompressedDataHash)
-	: Id(InId)
-	, CompressedData(MoveTemp(InCompressedData))
-	, CompressedDataHash(InCompressedDataHash)
-{
-	checkf(Id.IsValid(), TEXT("A valid ID is required to construct a cache payload."));
-	if (CompressedData.GetSize())
-	{
-		checkfSlow(CompressedDataHash == FIoHash::HashBuffer(CompressedData),
-			TEXT("Provided compressed data hash %s does not match calculated hash %s"),
-			*CompressedDataHash.ToString(), *FIoHash::HashBuffer(CompressedData).ToString());
-		checkf(!CompressedDataHash.IsZero(), TEXT("A non-empty compressed data buffer must have a non-zero hash."));
-	}
-	else
-	{
-		checkf(CompressedDataHash.IsZero(), TEXT("A null or empty compressed data buffer must use a hash of zero."));
-	}
-}
-
-FSharedBuffer FCachePayload::Compress(FSharedBuffer Buffer)
-{
-	// DDC-TODO: Switch to the new compression interface when it is ready. For now, match the "no compression" method.
-	FUniqueBuffer CompressedBuffer = FUniqueBuffer::Alloc(Buffer.GetSize() + 1);
-	FMemory::Memset(CompressedBuffer.GetData(), 0, 1);
-	FMemory::Memcpy(CompressedBuffer.GetView().RightChop(1).GetData(), Buffer.GetData(), Buffer.GetSize());
-	return FSharedBuffer(MoveTemp(CompressedBuffer));
-}
-
-FSharedBuffer FCachePayload::Decompress(FSharedBuffer Buffer)
-{
-	// DDC-TODO: Switch to the new compression interface when it is ready. For now, match the "no compression" method.
-	if (Buffer)
-	{
-		const FMemoryView View = Buffer.GetView() + 1;
-		return FSharedBuffer::MakeView(View, MoveTemp(Buffer));
-	}
-	return FSharedBuffer();
-}
-
-///////////////////////////////////////////////////////////////////////////////////////////////////////////////////////
-
 struct FCacheRecordBuilderData
 {
 	FCacheKey Key;
 	FCbObject Meta;
-	FCachePayload Value;
-	TArray<FCachePayload> Attachments;
+	FPayload Value;
+	TArray<FPayload> Attachments;
 };
 
 ///////////////////////////////////////////////////////////////////////////////////////////////////////////////////////
@@ -88,9 +30,9 @@ struct FCacheRecordData
 
 	FCacheKey Key;
 	FCbObject Meta;
-	FCachePayload Value;
+	FPayload Value;
 	FSharedBuffer ValueCache;
-	TArray<FCachePayload> Attachments;
+	TArray<FPayload> Attachments;
 	TArray<FSharedBuffer> AttachmentsCache;
 };
 
@@ -104,9 +46,9 @@ FCacheRecordData::FCacheRecordData(FCacheRecordBuilderData&& Builder)
 
 ///////////////////////////////////////////////////////////////////////////////////////////////////////////////////////
 
-static const FCachePayload& GetEmptyCachePayload()
+static const FPayload& GetEmptyCachePayload()
 {
-	static const FCachePayload Empty;
+	static const FPayload Empty;
 	return Empty;
 }
 
@@ -116,9 +58,9 @@ static const FCacheRecordData& GetEmptyCacheRecord()
 	return Empty;
 }
 
-static FCachePayloadId GetOrCreatePayloadId(const FCachePayloadId& Id, const FSharedBuffer& Buffer)
+static FPayloadId GetOrCreatePayloadId(const FPayloadId& Id, const FSharedBuffer& Buffer)
 {
-	return Id.IsValid() ? Id : FCachePayloadId::FromHash(FIoHash::HashBuffer(Buffer));
+	return Id.IsValid() ? Id : FPayloadId::FromHash(FIoHash::HashBuffer(Buffer));
 }
 
 ///////////////////////////////////////////////////////////////////////////////////////////////////////////////////////
@@ -154,23 +96,23 @@ FSharedBuffer FCacheRecord::GetValue() const
 	{
 		if (Data->ValueCache.IsNull() && Data->Value.IsValid())
 		{
-			Data->ValueCache = Data->Value.Decompress();
+			Data->ValueCache = Data->Value.GetBuffer().Decompress();
 		}
 		return Data->ValueCache;
 	}
 	return FSharedBuffer();
 }
 
-const FCachePayload& FCacheRecord::GetValuePayload() const
+const FPayload& FCacheRecord::GetValuePayload() const
 {
 	return (Data ? *Data : GetEmptyCacheRecord()).Value;
 }
 
-FSharedBuffer FCacheRecord::GetAttachment(const FCachePayloadId& Id) const
+FSharedBuffer FCacheRecord::GetAttachment(const FPayloadId& Id) const
 {
-	TArray<FCachePayload>& Attachments = Data->Attachments;
-	const int32 Index = Algo::LowerBound(Attachments, Id, FCachePayloadLessById());
-	if (Attachments.IsValidIndex(Index) && FCachePayloadEqualById()(Attachments[Index], Id))
+	TArray<FPayload>& Attachments = Data->Attachments;
+	const int32 Index = Algo::LowerBound(Attachments, Id, FPayloadLessById());
+	if (Attachments.IsValidIndex(Index) && FPayloadEqualById()(Attachments[Index], Id))
 	{
 		TArray<FSharedBuffer>& AttachmentsCache = Data->AttachmentsCache;
 		if (AttachmentsCache.IsEmpty())
@@ -180,25 +122,25 @@ FSharedBuffer FCacheRecord::GetAttachment(const FCachePayloadId& Id) const
 		FSharedBuffer& DataCache = AttachmentsCache[Index];
 		if (!DataCache)
 		{
-			DataCache = Attachments[Index].Decompress();
+			DataCache = Attachments[Index].GetBuffer().Decompress();
 		}
 		return DataCache;
 	}
 	return FSharedBuffer();
 }
 
-const FCachePayload& FCacheRecord::GetAttachmentPayload(const FCachePayloadId& Id) const
+const FPayload& FCacheRecord::GetAttachmentPayload(const FPayloadId& Id) const
 {
-	TArray<FCachePayload>& Attachments = Data->Attachments;
-	const int32 Index = Algo::LowerBound(Attachments, Id, FCachePayloadLessById());
-	if (Attachments.IsValidIndex(Index) && FCachePayloadEqualById()(Attachments[Index], Id))
+	TArray<FPayload>& Attachments = Data->Attachments;
+	const int32 Index = Algo::LowerBound(Attachments, Id, FPayloadLessById());
+	if (Attachments.IsValidIndex(Index) && FPayloadEqualById()(Attachments[Index], Id))
 	{
 		return Attachments[Index];
 	}
 	return GetEmptyCachePayload();
 }
 
-TConstArrayView<FCachePayload> FCacheRecord::GetAttachmentPayloads() const
+TConstArrayView<FPayload> FCacheRecord::GetAttachmentPayloads() const
 {
 	return (Data ? *Data : GetEmptyCacheRecord()).Attachments;
 }
@@ -221,32 +163,32 @@ void FCacheRecordBuilder::SetMeta(FCbObject Meta)
 	Data->Meta = MoveTemp(Meta);
 }
 
-FCachePayloadId FCacheRecordBuilder::SetValue(FSharedBuffer Buffer, FCachePayloadId Id)
+FPayloadId FCacheRecordBuilder::SetValue(const FSharedBuffer& Buffer, FPayloadId Id)
 {
 	Id = GetOrCreatePayloadId(Id, Buffer);
-	return SetValue(FCachePayload(Id, FCachePayload::Compress(MoveTemp(Buffer))));
+	return SetValue(FPayload(Id, FCompressedBuffer::Compress(Buffer)));
 }
 
-FCachePayloadId FCacheRecordBuilder::SetValue(FCachePayload&& Payload)
+FPayloadId FCacheRecordBuilder::SetValue(FPayload&& Payload)
 {
-	FCachePayload& Value = Data->Value;
+	FPayload& Value = Data->Value;
 	checkf(Value.IsNull(), TEXT("Failed to set value with ID %s because an there is an existing value with ID %s"),
 		*Payload.GetId().ToString(), *Value.GetId().ToString());
 	Value = MoveTemp(Payload);
 	return Value.GetId();
 }
 
-FCachePayloadId FCacheRecordBuilder::AddAttachment(FSharedBuffer Buffer, FCachePayloadId Id)
+FPayloadId FCacheRecordBuilder::AddAttachment(const FSharedBuffer& Buffer, FPayloadId Id)
 {
 	Id = GetOrCreatePayloadId(Id, Buffer);
-	return AddAttachment(FCachePayload(Id, MoveTemp(Buffer)));
+	return AddAttachment(FPayload(Id, FCompressedBuffer::Compress(Buffer)));
 }
 
-FCachePayloadId FCacheRecordBuilder::AddAttachment(FCachePayload&& Payload)
+FPayloadId FCacheRecordBuilder::AddAttachment(FPayload&& Payload)
 {
-	TArray<FCachePayload>& Attachments = Data->Attachments;
-	const int32 Index = Algo::LowerBound(Attachments, Payload, FCachePayloadLessById());
-	checkf(!Attachments.IsValidIndex(Index) || !FCachePayloadEqualById()(Attachments[Index], Payload),
+	TArray<FPayload>& Attachments = Data->Attachments;
+	const int32 Index = Algo::LowerBound(Attachments, Payload, FPayloadLessById());
+	checkf(!Attachments.IsValidIndex(Index) || !FPayloadEqualById()(Attachments[Index], Payload),
 		TEXT("Failed to add attachment with ID %s because an existing attachment is using that ID"),
 		*Payload.GetId().ToString());
 	Attachments.Insert(MoveTemp(Payload), Index);
diff --git a/Engine/Source/Developer/DerivedDataCache/Private/DerivedDataPayload.cpp b/Engine/Source/Developer/DerivedDataCache/Private/DerivedDataPayload.cpp
new file mode 100644
--- /dev/null
+++ b/Engine/Source/Developer/DerivedDataCache/Private/DerivedDataPayload.cpp
@@ -0,0 +1,110 @@
+// Copyright Epic Games, Inc. All Rights Reserved.
+
+#include "DerivedDataPayload.h"
+
+#include "Containers/StringConv.h"
+#include "Containers/UnrealString.h"
+#include "Misc/StringBuilder.h"
+#include "Serialization/CompactBinary.h"
+
+namespace UE
+{
+namespace DerivedData
+{
+
+///////////////////////////////////////////////////////////////////////////////////////////////////////////////////////
+
+FPayloadId FPayloadId::FromHash(const FIoHash& Hash)
+{
+	checkf(!Hash.IsZero(), TEXT("PayloadId requires a non-zero hash."));
+	return FPayloadId(MakeMemoryView(Hash.GetBytes()).Left(sizeof(ByteArray)));
+}
+
+FPayloadId FPayloadId::FromName(const FAnsiStringView Name)
+{
+	checkf(!Name.IsEmpty(), TEXT("PayloadId requires a non-empty name."));
+	return FPayloadId::FromHash(FIoHash::HashBuffer(Name.GetData(), Name.Len()));
+}
+
+FPayloadId FPayloadId::FromName(const FWideStringView Name)
+{
+	checkf(!Name.IsEmpty(), TEXT("PayloadId requires a non-empty name."));
+	TAnsiStringBuilder<128> Name8;
+	const int32 Len = int32(FTCHARToUTF8_Convert::ConvertedLength(Name.GetData(), Name.Len()));
+	const int32 Offset = Name8.AddUninitialized(Len);
+	const int32 WrittenLen = FTCHARToUTF8_Convert::Convert(Name8.GetData() + Offset, Len, Name.GetData(), Name.Len());
+	checkf(Len == WrittenLen, TEXT("Name '%.*s' failed to convert to a PayloadId."), Name.Len(), Name.GetData());
+	return FPayloadId::FromName(Name8);
+}
+
+FPayloadId FPayloadId::FromObjectId(const FCbObjectId& ObjectId)
+{
+	checkf(ObjectId != FCbObjectId(), TEXT("PayloadId requires a non-zero ObjectId."));
+	return FPayloadId(ObjectId.GetView());
+}
+
+FCbObjectId FPayloadId::ToObjectId() const
+{
+	return FCbObjectId(GetView());
+}
+
+void FPayloadId::ToString(FAnsiStringBuilderBase& Builder) const
+{
+	UE::String::BytesToHexLower(Bytes, Builder);
+}
+
+void FPayloadId::ToString(FWideStringBuilderBase& Builder) const
+{
+	UE::String::BytesToHexLower(Bytes, Builder);
+}
+
+FString FPayloadId::ToString() const
+{
+	TStringBuilder<16> Out;
+	ToString(Out);
+	return FString(Out);
+}
+
+///////////////////////////////////////////////////////////////////////////////////////////////////////////////////////
+
+FPayload::FPayload(const FPayloadId& InId, const FIoHash& InHash)
+	: Id(InId)
+	, Hash(InHash)
+{
+	checkf(Id.IsValid(), TEXT("A valid PayloadId is required to construct a payload."));
+}
+
+FPayload::FPayload(const FPayloadId& InId, const FIoHash& InHash, FCompressedBuffer InBuffer)
+	: Id(InId)
+	, Hash(InHash)
+	, Buffer(MoveTemp(InBuffer))
+{
+	checkf(Id.IsValid(), TEXT("A valid PayloadId is required to construct a payload."));
+	if (Buffer.GetCompressedSize())
+	{
+		checkfSlow(Hash == FIoHash::HashBuffer(Buffer.GetCompressed()),
+			TEXT("Provided compressed data hash %s does not match calculated hash %s"),
+			*Hash.ToString(), *FIoHash::HashBuffer(Buffer.GetCompressed()).ToString());
+		checkf(!Hash.IsZero(), TEXT("A non-empty compressed data buffer must have a non-zero hash."));
+	}
+	else
+	{
+		checkf(Hash.IsZero(), TEXT("A null or empty compressed data buffer must use a hash of zero."));
+	}
+}
+
+FPayload::FPayload(const FPayloadId& InId, FCompressedBuffer InBuffer)
+	: Id(InId)
+	, Buffer(MoveTemp(InBuffer))
+{
+	checkf(Id.IsValid(), TEXT("A valid PayloadId is required to construct a payload."));
+	if (Buffer.GetCompressedSize())
+	{
+		Hash = FIoHash::HashBuffer(Buffer.GetCompressed());
+	}
+}
+
+///////////////////////////////////////////////////////////////////////////////////////////////////////////////////////
+
+} // DerivedData
+} // UE
diff --git a/Engine/Source/Developer/DerivedDataCache/Public/DerivedDataCache.h b/Engine/Source/Developer/DerivedDataCache/Public/DerivedDataCache.h
--- a/Engine/Source/Developer/DerivedDataCache/Public/DerivedDataCache.h
+++ b/Engine/Source/Developer/DerivedDataCache/Public/DerivedDataCache.h
@@ -5,9 +5,8 @@
 #include "CoreTypes.h"
 #include "DerivedDataCacheKey.h"
 #include "DerivedDataCacheRecord.h"
-#include "Memory/MemoryFwd.h"
+#include "DerivedDataRequest.h"
 #include "Misc/EnumClassFlags.h"
-#include "Templates/RefCounting.h"
 
 namespace UE
 {
@@ -82,45 +81,6 @@ ENUM_CLASS_FLAGS(ECachePolicy);
 
 ///////////////////////////////////////////////////////////////////////////////////////////////////////////////////////
 
-/** Priority for scheduling cache requests. */
-enum class ECachePriority : uint8
-{
-	/**
-	 * Blocking is to be used only when the thread making the request will wait on completion of the
-	 * request before doing any other work. Requests at this priority level will be processed before
-	 * any request at a lower priority level. The cache may choose to process part of the request on
-	 * the thread making the request. Waiting on a request may increase its priority to this level.
-	 */
-	Blocking,
-	/**
-	 * Highest is the maximum priority for asynchronous cache requests, and is intended for requests
-	 * that are required to maintain interactivity of the program.
-	 */
-	Highest,
-	/**
-	 * High is intended for requests that are on the critical path, but are not required to maintain
-	 * interactivity of the program.
-	 */
-	High,
-	/**
-	 * Normal is intended as the default request priority.
-	 */
-	Normal,
-	/**
-	 * Low is intended for requests that are below the default priority, but are used for operations
-	 * that the program will execute now rather than at an unknown time in the future.
-	 */
-	Low,
-	/**
-	 * Lowest is the minimum priority for asynchronous cache requests, and is primarily intended for
-	 * requests that prefetch from the cache to optimize a future operation, while minimizing impact
-	 * on other cache requests that are in flight.
-	 */
-	Lowest,
-};
-
-///////////////////////////////////////////////////////////////////////////////////////////////////////////////////////
-
 /** Status of a cache operation. */
 enum class ECacheStatus : uint8
 {
@@ -134,105 +94,6 @@ enum class ECacheStatus : uint8
 
 ///////////////////////////////////////////////////////////////////////////////////////////////////////////////////////
 
-/** Interface for a cache request. Provides functionality common to every type of request. */
-class ICacheRequest : public IRefCountedObject
-{
-public:
-	virtual ~ICacheRequest() = default;
-
-	/** Set the priority of the request. */
-	virtual void SetPriority(ECachePriority Priority) = 0;
-
-	/** Cancel the request and invoke its callback. */
-	virtual void Cancel() = 0;
-
-	/** Block the calling thread until the request and its callback complete. */
-	virtual void Wait() = 0;
-
-	/** Poll the request and return true if complete and false otherwise. */
-	virtual bool Poll() = 0;
-};
-
-///////////////////////////////////////////////////////////////////////////////////////////////////////////////////////
-
-/** A wrapper for the cache request interface that manages lifetime with reference counting. */
-class FCacheRequest
-{
-public:
-	/** Construct a null cache request. */
-	FCacheRequest() = default;
-
-	/** Construct a cache request from . */
-	inline explicit FCacheRequest(ICacheRequest* InRequest)
-		: Request(InRequest)
-	{
-	}
-
-	/** Reset this to null. */
-	inline void Reset() { Request.SafeRelease(); }
-
-	/** Whether the cache request is null. */
-	inline bool IsNull() const { return !Request; }
-	/** Whether the cache request is not null. */
-	inline explicit operator bool() const { return !IsNull(); }
-
-	/**
-	 * Set the priority of the request.
-	 *
-	 * @param Priority The new priority for the request. May be higher or lower than the original.
-	 */
-	inline void SetPriority(ECachePriority Priority)
-	{
-		if (Request)
-		{
-			Request->SetPriority(Priority);
-		}
-	}
-
-	/**
-	 * Cancel the request and invoke its callback.
-	 *
-	 * Does not return until the callback for the request has finished executing.
-	 */
-	inline void Cancel() const
-	{
-		if (Request)
-		{
-			Request->Cancel();
-		}
-	}
-
-	/**
-	 * Block the calling thread until the request is complete.
-	 *
-	 * Avoid waiting if possible and use callbacks to manage control flow instead.
-	 * Does not return until the callback for the request has finished executing.
-	 */
-	inline void Wait() const
-	{
-		if (Request)
-		{
-			Request->Wait();
-		}
-	}
-
-	/**
-	 * Poll the request and return true if complete and false otherwise.
-	 *
-	 * Avoid polling if possible and use callbacks to manage control flow instead.
-	 * A request is not considered to be complete until its callback has finished executing.
-	 */
-	inline bool Poll() const
-	{
-		return !Request || Request->Poll();
-	}
-
-private:
-	TRefCountPtr<ICacheRequest> Request;
-};
-
-///////////////////////////////////////////////////////////////////////////////////////////////////////////////////////
-
 /** Parameters for the completion callback for cache put requests. */
 struct FCachePutCompleteParams
 {
@@ -269,9 +130,9 @@ struct FCacheGetPayloadCompleteParams
 	/**
 	 * Payload for the part of the payload request that completed or was canceled.
 	 *
-	 * The ID is always populated. The remainder of the payload is populated when Status is Cached.
+	 * The ID is always populated. The hash and buffer are populated when Status is Cached.
 	 */
-	FCachePayload Payload;
+	FPayload Payload;
 
 	/** Status of the cache request. */
 	ECacheStatus Status = ECacheStatus::NotCached;
@@ -309,14 +170,14 @@ public:
 	 * @param Records The cache records to store. Must have a key. Records are reset to null.
 	 * @param Context A description of the request. An object path is typically sufficient.
 	 * @param Policy Flags to control the behavior of the request. See ECachePolicy.
-	 * @param Priority A priority to consider when scheduling the request. See ECachePriority.
+	 * @param Priority A priority to consider when scheduling the request. See EPriority.
 	 * @param Callback A callback invoked for every key in the batch as it completes or is canceled.
 	 */
-	virtual FCacheRequest Put(
+	virtual FRequest Put(
 		TArrayView<FCacheRecord> Records,
 		FStringView Context,
 		ECachePolicy Policy = ECachePolicy::Default,
-		ECachePriority Priority = ECachePriority::Normal,
+		EPriority Priority = EPriority::Normal,
 		FOnCachePutComplete&& Callback = FOnCachePutComplete()) = 0;
 
 	/**
@@ -328,14 +189,14 @@ public:
 	 * @param Keys The keys identifying the cache records to query.
 	 * @param Context A description of the request. An object path is typically sufficient.
 	 * @param Policy Flags to control the behavior of the request. See ECachePolicy.
-	 * @param Priority A priority to consider when scheduling the request. See ECachePriority.
+	 * @param Priority A priority to consider when scheduling the request. See EPriority.
 	 * @param Callback A callback invoked for every key in the batch as it completes or is canceled.
 	 */
-	virtual FCacheRequest Get(
+	virtual FRequest Get(
 		TConstArrayView<FCacheKey> Keys,
 		FStringView Context,
 		ECachePolicy Policy,
-		ECachePriority Priority,
+		EPriority Priority,
 		FOnCacheGetComplete&& Callback) = 0;
 
 	/**
@@ -347,14 +208,14 @@ public:
 	 * @param Keys The keys identifying the cache record payloads to query.
 	 * @param Context A description of the request. An object path is typically sufficient.
 	 * @param Policy Flags to control the behavior of the request. See ECachePolicy.
-	 * @param Priority A priority to consider when scheduling the request. See ECachePriority.
+	 * @param Priority A priority to consider when scheduling the request. See EPriority.
 	 * @param Callback A callback invoked for every key in the batch as it completes or is canceled.
 	 */
-	virtual FCacheRequest GetPayloads(
+	virtual FRequest GetPayloads(
 		TConstArrayView<FCachePayloadKey> Keys,
 		FStringView Context,
 		ECachePolicy Policy,
-		ECachePriority Priority,
+		EPriority Priority,
 		FOnCacheGetPayloadComplete&& Callback) = 0;
 
 	/**
diff --git a/Engine/Source/Developer/DerivedDataCache/Public/DerivedDataCacheKey.h b/Engine/Source/Developer/DerivedDataCache/Public/DerivedDataCacheKey.h
--- a/Engine/Source/Developer/DerivedDataCache/Public/DerivedDataCacheKey.h
+++ b/Engine/Source/Developer/DerivedDataCache/Public/DerivedDataCacheKey.h
@@ -5,6 +5,7 @@
 #include "CoreTypes.h"
 #include "Containers/StringFwd.h"
 #include "Containers/StringView.h"
+#include "DerivedDataPayload.h"
 #include "IO/IoHash.h"
 #include "Templates/TypeHash.h"
 
@@ -170,97 +171,6 @@ struct FCacheKeyLexicalLess
 
 ///////////////////////////////////////////////////////////////////////////////////////////////////////////////////////
 
-/** A 12-byte value that uniquely identifies a cache payload in the scope of a cache record. */
-struct FCachePayloadId
-{
-public:
-	using ByteArray = uint8[12];
-
-	/** Construct a null ID. */
-	FCachePayloadId() = default;
-
-	/** Construct an ID from a view of 12 bytes. */
-	inline explicit FCachePayloadId(FMemoryView Id);
-
-	/** Construct an ID from a non-zero hash. The hash must be unique in the outer cache record. */
-	static inline FCachePayloadId FromHash(const FIoHash& Hash);
-
-	/** Construct an ID from a non-empty name. The name must be unique in the outer cache record. */
-	UE_API static FCachePayloadId FromName(FAnsiStringView Name);
-	UE_API static FCachePayloadId FromName(FWideStringView Name);
-
-	/** Construct an ID from an ObjectId. The ObjectId must be unique in the outer cache record. */
-	UE_API static FCachePayloadId FromObjectId(const FCbObjectId& ObjectId);
-
-	/** Convert the ID to an ObjectId. */
-	UE_API FCbObjectId ToObjectId() const;
-
-	/** Convert the ID to a 24-character hex string. */
-	UE_API void ToString(FAnsiStringBuilderBase& Builder) const;
-	UE_API void ToString(FWideStringBuilderBase& Builder) const;
-	UE_API FString ToString() const;
-
-	/** Returns a view of the raw byte array for the ID. */
-	constexpr inline FMemoryView GetView() const { return MakeMemoryView(Bytes); }
-
-	/** Whether this is null. */
-	inline bool IsNull() const;
-	/** Whether this is not null. */
-	inline bool IsValid() const { return !IsNull(); }
-
-	/** Reset this to null. */
-	inline void Reset() { *this = FCachePayloadId(); }
-
-private:
-	alignas(uint32) ByteArray Bytes{};
-};
-
-inline bool operator==(const FCachePayloadId& A, const FCachePayloadId& B)
-{
-	return A.GetView().EqualBytes(B.GetView());
-}
-
-inline bool operator!=(const FCachePayloadId& A, const FCachePayloadId& B)
-{
-	return !A.GetView().EqualBytes(B.GetView());
-}
-
-inline bool operator<(const FCachePayloadId& A, const FCachePayloadId& B)
-{
-	return A.GetView().CompareBytes(B.GetView()) < 0;
-}
-
-inline uint32 GetTypeHash(const FCachePayloadId& Id)
-{
-	return *reinterpret_cast<const uint32*>(Id.GetView().GetData());
-}
-
-inline FCachePayloadId::FCachePayloadId(const FMemoryView InId)
-{
-	checkf(InId.GetSize() == sizeof(ByteArray),
-		TEXT("FCachePayloadId cannot be constructed from a view of %" UINT64_FMT " bytes."), InId.GetSize());
-	FMemory::Memcpy(Bytes, InId.GetData(), sizeof(ByteArray));
-}
-
-inline FCachePayloadId FCachePayloadId::FromHash(const FIoHash& Hash)
-{
-	return FCachePayloadId(MakeMemoryView(Hash.GetBytes()).Left(sizeof(ByteArray)));
-}
-
-inline bool FCachePayloadId::IsNull() const
-{
-	return *this == FCachePayloadId();
-}
-
-template <typename CharType>
-inline TStringBuilderBase<CharType>& operator<<(TStringBuilderBase<CharType>& Builder, const FCachePayloadId& Id)
-{
-	Id.ToString(Builder);
-	return Builder;
-}
-
-///////////////////////////////////////////////////////////////////////////////////////////////////////////////////////
-
 /** A key that uniquely identifies a payload within a cache record. */
 class FCachePayloadKey
 {
@@ -269,7 +179,7 @@ public:
 	FCachePayloadKey() = default;
 
 	/** Construct a cache payload key from the non-null key and non-null ID. */
-	inline FCachePayloadKey(const FCacheKey& InKey, const FCachePayloadId& InId)
+	inline FCachePayloadKey(const FCacheKey& InKey, const FPayloadId& InId)
 		: Key(InKey)
 		, Id(InId)
 	{
@@ -279,7 +189,7 @@ public:
 	inline const FCacheKey& GetKey() const { return Key; }
 
 	/** Returns the ID of the payload. */
-	inline const FCachePayloadId& GetId() const { return Id; }
+	inline const FPayloadId& GetId() const { return Id; }
 
 	/** Convert the cache payload key to a string. */
 	UE_API void ToString(FAnsiStringBuilderBase& Builder) const;
@@ -296,7 +206,7 @@ public:
 
 private:
 	FCacheKey Key;
-	FCachePayloadId Id;
+	FPayloadId Id;
 };
 
 inline bool operator==(const FCachePayloadKey& A, const FCachePayloadKey& B)
diff --git a/Engine/Source/Developer/DerivedDataCache/Public/DerivedDataCacheRecord.h b/Engine/Source/Developer/DerivedDataCache/Public/DerivedDataCacheRecord.h
--- a/Engine/Source/Developer/DerivedDataCache/Public/DerivedDataCacheRecord.h
+++ b/Engine/Source/Developer/DerivedDataCache/Public/DerivedDataCacheRecord.h
@@ -4,12 +4,13 @@
 
 #include "CoreTypes.h"
 #include "DerivedDataCacheKey.h"
-#include "Memory/SharedBuffer.h"
-#include "Serialization/CompactBinary.h"
+#include "DerivedDataPayload.h"
+#include "Memory/MemoryFwd.h"
 #include "Templates/PimplPtr.h"
 
 #define UE_API DERIVEDDATACACHE_API
 
+class FCbObject;
 template <typename ResultType> class TFuture;
 
 namespace UE
@@ -23,69 +24,6 @@ struct FCacheRecordBuilderData;
 
 ///////////////////////////////////////////////////////////////////////////////////////////////////////////////////////
 
-/**
- * A cache payload is a compressed value or attachment for a cache record.
- */
-class FCachePayload
-{
-public:
-	/** Construct a null cache payload. */
-	FCachePayload() = default;
-
-	/** Construct a cache payload from the compressed data hash. */
-	UE_API FCachePayload(const FCachePayloadId& Id, const FIoHash& CompressedDataHash);
-	/** Construct a cache payload from the compressed data. See FCachePayload::Compress. */
-	UE_API FCachePayload(const FCachePayloadId& Id, FSharedBuffer CompressedData);
-	/** Construct a cache payload from the compressed data its hash. See FCachePayload::Compress. */
-	UE_API FCachePayload(const FCachePayloadId& Id, FSharedBuffer CompressedData, const FIoHash& CompressedDataHash);
-
-	/** Returns a payload ID that uniquely identifies it within the scope of its cache record. */
-	inline const FCachePayloadId& GetId() const { return Id; }
-
-	/**
-	 * Returns the compressed data for the payload.
-	 *
-	 * This will be null when the originating cache request skips the payload.
-	 */
-	inline const FSharedBuffer& GetCompressedData() const { return CompressedData; }
-
-	/** Returns the hash of the compressed data for the payload. */
-	inline const FIoHash& GetCompressedDataHash() const { return CompressedDataHash; }
-
-	inline FSharedBuffer Decompress() const { return Decompress(CompressedData); }
-
-	/** Whether this is null. */
-	inline bool IsNull() const { return Id.IsNull(); }
-	/** Whether this is not null. */
-	inline bool IsValid() const { return !IsNull(); }
-	/** Whether this is not null. */
-	inline explicit operator bool() const { return IsValid(); }
-
-	/** Reset this to null. */
-	inline void Reset() { *this = FCachePayload(); }
-
-	/**
-	 * Compress the input buffer with the default settings for a cache payload.
-	 *
-	 * FIoCompression::Compress and its variants may be used instead to access more options.
-	 */
-	UE_API static FSharedBuffer Compress(FSharedBuffer Data);
-
-	/**
-	 * Decompress the input buffer.
-	 *
-	 * FIoCompression::Decompress and its variants may be used instead to access more options.
-	 */
-	UE_API static FSharedBuffer Decompress(FSharedBuffer CompressedData);
-
-private:
-	FCachePayloadId Id;
-	FSharedBuffer CompressedData;
-	FIoHash CompressedDataHash;
-};
-
-///////////////////////////////////////////////////////////////////////////////////////////////////////////////////////
-
 /**
  * A cache record is a key, a value, attachments, and metadata.
  *
@@ -120,16 +58,16 @@ public:
 	UE_API FSharedBuffer GetValue() const;
 
 	/** Returns the value payload. Always available in a non-null cache record. */
-	UE_API const FCachePayload& GetValuePayload() const;
+	UE_API const FPayload& GetValuePayload() const;
 
 	/** Returns the attachment that matches the ID. Null if no match or attachments were skipped in the request. */
-	UE_API FSharedBuffer GetAttachment(const FCachePayloadId& Id) const;
+	UE_API FSharedBuffer GetAttachment(const FPayloadId& Id) const;
 
 	/** Returns the attachment payload that matches the ID. Null if no match or the cache record is null. */
-	UE_API const FCachePayload& GetAttachmentPayload(const FCachePayloadId& Id) const;
+	UE_API const FPayload& GetAttachmentPayload(const FPayloadId& Id) const;
 
-	/** Returns a view of the attachments. Always available in a non-null cache record, but data may be skipped. */
-	UE_API TConstArrayView<FCachePayload> GetAttachmentPayloads() const;
+	/** Returns a view of the attachments. Always available in a non-null cache record, but buffer may be skipped. */
+	UE_API TConstArrayView<FPayload> GetAttachmentPayloads() const;
 
 	/** Whether this is null. */
 	inline bool IsNull() const { return !Data; }
@@ -191,17 +129,17 @@ public:
 	 * @param Buffer The value, which is cloned if not owned, and is compressed by the builder.
 	 * @param Id An ID for the value that is unique within the scope of the cache record. If omitted
 	 *           the builder will create an ID by hashing the buffer.
-	 * @return The payload ID that was provided or created.
+	 * @return The ID that was provided or created.
 	 */
-	UE_API FCachePayloadId SetValue(FSharedBuffer Buffer, FCachePayloadId Id = FCachePayloadId());
+	UE_API FPayloadId SetValue(const FSharedBuffer& Buffer, FPayloadId Id = FPayloadId());
 
 	/**
 	 * Set the value for the cache record.
 	 *
 	 * @param Payload The value payload, which is reset to null.
-	 * @return The payload ID that was provided. Unique within the scope of the cache record.
+	 * @return The ID that was provided. Unique within the scope of the cache record.
 	 */
-	UE_API FCachePayloadId SetValue(FCachePayload&& Payload);
+	UE_API FPayloadId SetValue(FPayload&& Payload);
 
 	/**
 	 * Add an attachment to the cache record.
@@ -209,17 +147,17 @@ public:
 	 * @param Buffer The attachment, which is cloned if not owned, and is compressed by the builder.
 	 * @param Id An ID for the attachment that is unique within the scope of the cache record. If it
 	 *           is omitted the builder will create an ID by hashing the buffer.
-	 * @return The payload ID that was provided or created.
+	 * @return The ID that was provided or created.
 	 */
-	UE_API FCachePayloadId AddAttachment(FSharedBuffer Buffer, FCachePayloadId Id = FCachePayloadId());
+	UE_API FPayloadId AddAttachment(const FSharedBuffer& Buffer, FPayloadId Id = FPayloadId());
 
 	/**
 	 * Add an attachment to the cache record.
 	 *
 	 * @param Payload The attachment payload, which is reset to null.
-	 * @return The payload ID that was provided. Unique within the scope of the cache record.
+	 * @return The ID that was provided. Unique within the scope of the cache record.
 	 */
-	UE_API FCachePayloadId AddAttachment(FCachePayload&& Payload);
+	UE_API FPayloadId AddAttachment(FPayload&& Payload);
 
 	/**
 	 * Build a cache record, resetting this builder.
@@ -253,46 +191,6 @@ private:
 
 ///////////////////////////////////////////////////////////////////////////////////////////////////////////////////////
 
-/** Compare cache payloads by their ID. */
-struct FCachePayloadEqualById
-{
-	inline bool operator()(const FCachePayload& Payload1, const FCachePayload& Payload2) const
-	{
-		return Payload1.GetId() == Payload2.GetId();
-	}
-
-	inline bool operator()(const FCachePayload& Payload1, const FCachePayloadId& PayloadId2) const
-	{
-		return Payload1.GetId() == PayloadId2;
-	}
-
-	inline bool operator()(const FCachePayloadId& PayloadId1, const FCachePayload& Payload2) const
-	{
-		return PayloadId1 == Payload2.GetId();
-	}
-};
-
-/** Compare cache payloads by their ID. */
-struct FCachePayloadLessById
-{
-	inline bool operator()(const FCachePayload& Payload1, const FCachePayload& Payload2) const
-	{
-		return Payload1.GetId() < Payload2.GetId();
-	}
-
-	inline bool operator()(const FCachePayload& Payload1, const FCachePayloadId& PayloadId2) const
-	{
-		return Payload1.GetId() < PayloadId2;
-	}
-
-	inline bool operator()(const FCachePayloadId& PayloadId1, const FCachePayload& Payload2) const
-	{
-		return PayloadId1 < Payload2.GetId();
-	}
-};
-
-///////////////////////////////////////////////////////////////////////////////////////////////////////////////////////
-
 } // DerivedData
 } // UE
 
diff --git a/Engine/Source/Developer/DerivedDataCache/Public/DerivedDataPayload.h b/Engine/Source/Developer/DerivedDataCache/Public/DerivedDataPayload.h
new file mode 100644
--- /dev/null
+++ b/Engine/Source/Developer/DerivedDataCache/Public/DerivedDataPayload.h
@@ -0,0 +1,210 @@
+// Copyright Epic Games, Inc. All Rights Reserved.
+
+#pragma once
+
+#include "CoreTypes.h"
+#include "Compression/CompressedBuffer.h"
+#include "Containers/StringFwd.h"
+#include "IO/IoHash.h"
+#include "Memory/MemoryView.h"
+
+#define UE_API DERIVEDDATACACHE_API
+
+class FCbObjectId;
+class FString;
+
+namespace UE
+{
+namespace DerivedData
+{
+
+///////////////////////////////////////////////////////////////////////////////////////////////////////////////////////
+
+/** A 12-byte value that uniquely identifies a payload in the context that it was created. */
+struct FPayloadId
+{
+public:
+	using ByteArray = uint8[12];
+
+	/** Construct a null ID. */
+	FPayloadId() = default;
+
+	/** Construct an ID from a view of 12 bytes. */
+	inline explicit FPayloadId(FMemoryView Id);
+
+	/** Construct an ID from a non-zero hash. */
+	UE_API static FPayloadId FromHash(const FIoHash& Hash);
+
+	/** Construct an ID from a non-empty name. */
+	UE_API static FPayloadId FromName(FAnsiStringView Name);
+	UE_API static FPayloadId FromName(FWideStringView Name);
+
+	/** Construct an ID from an ObjectId. */
+	UE_API static FPayloadId FromObjectId(const FCbObjectId& ObjectId);
+
+	/** Convert the ID to an ObjectId. */
+	UE_API FCbObjectId ToObjectId() const;
+
+	/** Convert the ID to a 24-character hex string. */
+	UE_API void ToString(FAnsiStringBuilderBase& Builder) const;
+	UE_API void ToString(FWideStringBuilderBase& Builder) const;
+	UE_API FString ToString() const;
+
+	/** Returns a view of the raw byte array for the ID. */
+	constexpr inline FMemoryView GetView() const { return MakeMemoryView(Bytes); }
+
+	/** Whether this is null. */
+	inline bool IsNull() const;
+	/** Whether this is not null. */
+	inline bool IsValid() const { return !IsNull(); }
+
+	/** Reset this to null. */
+	inline void Reset() { *this = FPayloadId(); }
+
+private:
+	alignas(uint32) ByteArray Bytes{};
+};
+
+inline bool operator==(const FPayloadId& A, const FPayloadId& B)
+{
+	return A.GetView().EqualBytes(B.GetView());
+}
+
+inline bool operator!=(const FPayloadId& A, const FPayloadId& B)
+{
+	return !A.GetView().EqualBytes(B.GetView());
+}
+
+inline bool operator<(const FPayloadId& A, const FPayloadId& B)
+{
+	return A.GetView().CompareBytes(B.GetView()) < 0;
+}
+
+inline uint32 GetTypeHash(const FPayloadId& Id)
+{
+	return *reinterpret_cast<const uint32*>(Id.GetView().GetData());
+}
+
+inline FPayloadId::FPayloadId(const FMemoryView InId)
+{
+	checkf(InId.GetSize() == sizeof(ByteArray),
+		TEXT("FPayloadId cannot be constructed from a view of %" UINT64_FMT " bytes."), InId.GetSize());
+	FMemory::Memcpy(Bytes, InId.GetData(), sizeof(ByteArray));
+}
+
+inline bool FPayloadId::IsNull() const
+{
+	return *this == FPayloadId();
+}
+
+template <typename CharType>
+inline TStringBuilderBase<CharType>& operator<<(TStringBuilderBase<CharType>& Builder, const FPayloadId& Id)
+{
+	Id.ToString(Builder);
+	return Builder;
+}
+
+///////////////////////////////////////////////////////////////////////////////////////////////////////////////////////
+
+/**
+ * A payload is described by a PayloadId and an IoHash of the compressed payload buffer.
+ *
+ * Payloads may be constructed with or without a compressed buffer. A payload without a buffer is
+ * used to uniquely describe a payload without the overhead of its buffer.
+ */
+class FPayload
+{
+public:
+	/** Construct a null payload. */
+	FPayload() = default;
+
+	/** Construct a payload from the hash of a compressed buffer. */
+	UE_API FPayload(const FPayloadId& Id, const FIoHash& BufferHash);
+	/** Construct a payload from a compressed buffer and its hash. See FCompressedBuffer::Compress. */
+	UE_API FPayload(const FPayloadId& Id, const FIoHash& BufferHash, FCompressedBuffer Buffer);
+	/** Construct a payload from a compressed buffer and calculates its hash. See FCompressedBuffer::Compress. */
+	UE_API FPayload(const FPayloadId& Id, FCompressedBuffer Buffer);
+
+	/** Returns the PayloadId for the payload. */
+	inline const FPayloadId& GetId() const { return Id; }
+
+	/** Returns the IoHash of the compressed buffer for the payload. */
+	inline const FIoHash& GetHash() const { return Hash; }
+
+	/** Returns the compressed buffer for the payload. May be null. */
+	inline const FCompressedBuffer& GetBuffer() const { return Buffer; }
+
+	/** Whether this is null. */
+	inline bool IsNull() const { return Id.IsNull(); }
+	/** Whether this is not null. */
+	inline bool IsValid() const { return !IsNull(); }
+	/** Whether this is not null. */
+	inline explicit operator bool() const { return IsValid(); }
+
+	/** Reset this to null. */
+	inline void Reset() { *this = FPayload(); }
+
+private:
+	FPayloadId Id;
+	FIoHash Hash;
+	FCompressedBuffer Buffer;
+};
+
+///////////////////////////////////////////////////////////////////////////////////////////////////////////////////////
+
+/** Compare payloads by their PayloadId and compressed buffer hash. */
+inline bool operator==(const FPayload& A, const FPayload& B)
+{
+	return A.GetId() == B.GetId() && A.GetHash() == B.GetHash();
+}
+
+/** Compare payloads by their PayloadId and compressed buffer hash. */
+inline bool operator<(const FPayload& A, const FPayload& B)
+{
+	return A.GetId() == B.GetId() && A.GetHash() == B.GetHash();
+}
+
+/** Compare payloads by their PayloadId. */
+struct FPayloadEqualById
+{
+	inline bool operator()(const FPayload& A, const FPayload& B) const
+	{
+		return A.GetId() == B.GetId();
+	}
+
+	inline bool operator()(const FPayload& A, const FPayloadId& B) const
+	{
+		return A.GetId() == B;
+	}
+
+	inline bool operator()(const FPayloadId& A, const FPayload& B) const
+	{
+		return A == B.GetId();
+	}
+};
+
+/** Compare payloads by their PayloadId. */
+struct FPayloadLessById
+{
+	inline bool operator()(const FPayload& A, const FPayload& B) const
+	{
+		return A.GetId() < B.GetId();
+	}
+
+	inline bool operator()(const FPayload& A, const FPayloadId& B) const
+	{
+		return A.GetId() < B;
+	}
+
+	inline bool operator()(const FPayloadId& A, const FPayload& B) const
+	{
+		return A < B.GetId();
+	}
+};
+
+///////////////////////////////////////////////////////////////////////////////////////////////////////////////////////
+
+} // DerivedData
+} // UE
+
+#undef UE_API
diff --git a/Engine/Source/Developer/DerivedDataCache/Public/DerivedDataRequest.h b/Engine/Source/Developer/DerivedDataCache/Public/DerivedDataRequest.h
new file mode 100644
--- /dev/null
+++ b/Engine/Source/Developer/DerivedDataCache/Public/DerivedDataRequest.h
@@ -0,0 +1,157 @@
+// Copyright Epic Games, Inc. All Rights Reserved.
+
+#pragma once
+
+#include "CoreTypes.h"
+#include "Templates/RefCounting.h"
+
+namespace UE
+{
+namespace DerivedData
+{
+
+///////////////////////////////////////////////////////////////////////////////////////////////////////////////////////
+
+/** Priority for scheduling requests. */
+enum class EPriority : uint8
+{
+	/**
+	 * Blocking is to be used only when the thread making the request will wait on completion of the
+	 * request before doing any other work. Requests at this priority level will be processed before
+	 * any request at a lower priority level. This priority permits a request to be processed on the
+	 * thread making the request. Waiting on a request may increase its priority to this level.
+	 */
+	Blocking,
+	/**
+	 * Highest is the maximum priority for asynchronous requests, and intended for requests that are
+	 * required to maintain interactivity of the program.
+	 */
+	Highest,
+	/**
+	 * High is intended for requests that are on the critical path, but are not required to maintain
+	 * interactivity of the program.
+	 */
+	High,
+	/**
+	 * Normal is intended as the default request priority.
+	 */
+	Normal,
+	/**
+	 * Low is intended for requests that are below the default priority, but are used for operations
+	 * that the program will execute now rather than at an unknown time in the future.
+	 */
+	Low,
+	/**
+	 * Lowest is the minimum priority for asynchronous requests, and primarily for requests that are
+	 * speculative in nature, while minimizing impact on other requests.
+	 */
+	Lowest,
+};
+
+///////////////////////////////////////////////////////////////////////////////////////////////////////////////////////
+
+/** Interface for a request. Provides functionality common to every type of request. */
+class IRequest : public IRefCountedObject
+{
+public:
+	virtual ~IRequest() = default;
+
+	/** Set the priority of the request. */
+	virtual void SetPriority(EPriority Priority) = 0;
+
+	/** Cancel the request and invoke its callback. */
+	virtual void Cancel() = 0;
+
+	/** Block the calling thread until the request and its callback complete. */
+	virtual void Wait() = 0;
+
+	/** Poll the request and return true if complete and false otherwise. */
+	virtual bool Poll() = 0;
+};
+
+///////////////////////////////////////////////////////////////////////////////////////////////////////////////////////
+
+/** A wrapper for the request interface that manages lifetime with reference counting. */
+class FRequest
+{
+public:
+	/** Construct a null request. */
+	FRequest() = default;
+
+	/** Construct a request from a raw pointer. */
+	inline explicit FRequest(IRequest* InRequest)
+		: Request(InRequest)
+	{
+	}
+
+	/** Reset this to null. */
+	inline void Reset() { Request.SafeRelease(); }
+
+	/** Whether the request is null. */
+	inline bool IsNull() const { return !Request; }
+	/** Whether the request is not null. */
+	inline bool IsValid() const { return !IsNull(); }
+	/** Whether the request is not null. */
+	inline explicit operator bool() const { return IsValid(); }
+
+	/**
+	 * Set the priority of the request.
+	 *
+	 * @param Priority The new priority for the request. May be higher or lower than the original.
+	 */
+	inline void SetPriority(EPriority Priority)
+	{
+		if (Request)
+		{
+			Request->SetPriority(Priority);
+		}
+	}
+
+	/**
+	 * Cancel the request and invoke any associated callback.
+	 *
+	 * Does not return until any callback for the request has finished executing.
+	 */
+	inline void Cancel() const
+	{
+		if (Request)
+		{
+			Request->Cancel();
+		}
+	}
+
+	/**
+	 * Block the calling thread until the request is complete.
+	 *
+	 * Avoid waiting if possible and use callbacks to manage control flow instead.
+	 * Does not return until any callback for the request has finished executing.
+	 */
+	inline void Wait() const
+	{
+		if (Request)
+		{
+			Request->Wait();
+		}
+	}
+
+	/**
+	 * Poll the request and return true if complete and false otherwise.
+	 *
+	 * Avoid polling if possible and use callbacks to manage control flow instead.
+	 * Does not return true until any callback for the request has finished executing.
+	 */
+	inline bool Poll() const
+	{
+		return !Request || Request->Poll();
+	}
+
+private:
+	TRefCountPtr<IRequest> Request;
+};
+
+///////////////////////////////////////////////////////////////////////////////////////////////////////////////////////
+
+} // DerivedData
+} // UE
+
+#undef UE_API
diff --git a/Engine/Source/Editor/UnrealEd/Private/EditorDomain/EditorDomain.cpp b/Engine/Source/Editor/UnrealEd/Private/EditorDomain/EditorDomain.cpp
--- a/Engine/Source/Editor/UnrealEd/Private/EditorDomain/EditorDomain.cpp
+++ b/Engine/Source/Editor/UnrealEd/Private/EditorDomain/EditorDomain.cpp
@@ -18,6 +18,7 @@
 #include "Misc/PackageSegment.h"
 #include "Misc/ScopeLock.h"
 #include "Serialization/Archive.h"
+#include "Serialization/CompactBinary.h"
 #include "Templates/RefCounting.h"
 #include "Templates/UniquePtr.h"
 #include "UObject/PackageResourceManagerFile.h"
@@ -166,7 +167,7 @@ int64 FEditorDomain::FileSize(const FPackagePath& PackagePath, EPackageSegment P
 		return Workspace->FileSize(PackagePath, PackageSegment, OutUpdatedPath);
 	}
 
-	UE::DerivedData::FCacheRequest Request;
+	UE::DerivedData::FRequest Request;
 	int64 FileSize = -1;
 	{
 		FScopeLock ScopeLock(&Locks->Lock);
@@ -186,8 +187,7 @@ int64 FEditorDomain::FileSize(const FPackagePath& PackagePath, EPackageSegment P
 				Params.Status == UE::DerivedData::ECacheStatus::Cached)
 			{
 				const FCbObject& MetaData = Params.Record.GetMeta();
-				FCbFieldView FileSizeField = MetaData["FileSize"];
-				FileSize = FileSizeField.AsInt64();
+				FileSize = MetaData["FileSize"].AsInt64();
 				PackageSource->Source = EPackageSource::Editor;
 			}
 			else
@@ -211,7 +211,7 @@ int64 FEditorDomain::FileSize(const FPackagePath& PackagePath, EPackageSegment P
 			}
 		};
 		Request = RequestEditorDomainPackage(PackagePath, PackageSource->Digest,
-			UE::DerivedData::ECachePriority::Highest, MoveTemp(MetaDataGetComplete));
+			UE::DerivedData::EPriority::Highest, MoveTemp(MetaDataGetComplete));
 	}
 	Request.Wait();
 	return FileSize;
@@ -234,8 +234,8 @@ FOpenPackageResult FEditorDomain::OpenReadPackage(const FPackagePath& PackagePat
 	}
 
 	FEditorDomainReadArchive* Result = new FEditorDomainReadArchive(Locks, PackagePath, PackageSource);
-	UE::DerivedData::FCacheRequest Request = RequestEditorDomainPackage(PackagePath, PackageSource->Digest,
-		UE::DerivedData::ECachePriority::Normal,
+	UE::DerivedData::FRequest Request = RequestEditorDomainPackage(PackagePath, PackageSource->Digest,
+		UE::DerivedData::EPriority::Normal,
 		[Result](UE::DerivedData::FCacheGetCompleteParams&& Params)
 		{
 			// Note that ~FEditorDomainReadArchive waits for this callback to be called, so Result cannot dangle
@@ -279,8 +279,8 @@ IAsyncReadFileHandle* FEditorDomain::OpenAsyncReadPackage(const FPackagePath& Pa
 	}
 
 	FEditorDomainAsyncReadFileHandle* Result = new FEditorDomainAsyncReadFileHandle(Locks, PackagePath, PackageSource);
-	UE::DerivedData::FCacheRequest Request = RequestEditorDomainPackage(PackagePath, PackageSource->Digest,
-		UE::DerivedData::ECachePriority::Normal,
+	UE::DerivedData::FRequest Request = RequestEditorDomainPackage(PackagePath, PackageSource->Digest,
+		UE::DerivedData::EPriority::Normal,
 		[Result](UE::DerivedData::FCacheGetCompleteParams&& Params)
 		{
 			// Note that ~FEditorDomainAsyncReadFileHandle waits for this callback to be called, so Result cannot dangle
diff --git a/Engine/Source/Editor/UnrealEd/Private/EditorDomain/EditorDomainArchive.cpp b/Engine/Source/Editor/UnrealEd/Private/EditorDomain/EditorDomainArchive.cpp
--- a/Engine/Source/Editor/UnrealEd/Private/EditorDomain/EditorDomainArchive.cpp
+++ b/Engine/Source/Editor/UnrealEd/Private/EditorDomain/EditorDomainArchive.cpp
@@ -10,6 +10,7 @@
 #include "Misc/AssertionMacros.h"
 #include "Misc/PackageSegment.h"
 #include "Misc/ScopeLock.h"
+#include "Serialization/CompactBinary.h"
 
 FEditorDomainReadArchive::FEditorDomainReadArchive(const TRefCountPtr<FEditorDomain::FLocks>& InLocks,
 	const FPackagePath& InPackagePath, const TRefCountPtr<FEditorDomain::FPackageSource>& InPackageSource)
@@ -26,7 +27,7 @@ FEditorDomainReadArchive::~FEditorDomainReadArchive()
 	Close();
 }
 
-void FEditorDomainReadArchive::SetRequest(const UE::DerivedData::FCacheRequest& InRequest)
+void FEditorDomainReadArchive::SetRequest(const UE::DerivedData::FRequest& InRequest)
 {
 	Request = InRequest;
 }
@@ -351,7 +352,7 @@ FEditorDomainAsyncReadFileHandle::~FEditorDomainAsyncReadFileHandle()
 	Source = ESource::Closed;
 }
 
-void FEditorDomainAsyncReadFileHandle::SetRequest(const UE::DerivedData::FCacheRequest& InRequest)
+void FEditorDomainAsyncReadFileHandle::SetRequest(const UE::DerivedData::FRequest& InRequest)
 {
 	Request = InRequest;
 }
diff --git a/Engine/Source/Editor/UnrealEd/Private/EditorDomain/EditorDomainArchive.h b/Engine/Source/Editor/UnrealEd/Private/EditorDomain/EditorDomainArchive.h
--- a/Engine/Source/Editor/UnrealEd/Private/EditorDomain/EditorDomainArchive.h
+++ b/Engine/Source/Editor/UnrealEd/Private/EditorDomain/EditorDomainArchive.h
@@ -30,7 +30,7 @@ public:
 	~FEditorDomainReadArchive();
 
 	/** Set the CacheRequest handle that will feed this archive. */
-	void SetRequest(const UE::DerivedData::FCacheRequest& InRequest);
+	void SetRequest(const UE::DerivedData::FRequest& InRequest);
 	/** Callback from the CacheRequest; set whether we're reading from EditorDomain bytes or WorkspaceDomain archive. */
 	void OnCacheRequestComplete(UE::DerivedData::FCacheGetCompleteParams&& Params);
 	/** Get the PackageFormat, which depends on the domain the data is read from. */
@@ -64,7 +64,7 @@ private:
 
 	// Data in this section is either read-only, or is read and written only on the public interface thread.
 	TRefCountPtr<FEditorDomain::FLocks> EditorDomainLocks; // Pointer is read-only, *pointer has internal lock
-	UE::DerivedData::FCacheRequest Request; // Read-only
+	UE::DerivedData::FRequest Request; // Read-only
 	FPackagePath PackagePath; // Read-only
 	int64 Pos = 0; // Interface-thread-only
 	mutable ESource Source = ESource::Uninitialized; // Interface-thread-only
@@ -91,7 +91,7 @@ public:
 	~FEditorDomainAsyncReadFileHandle();
 
 	/** Set the CacheRequest handle that will feed this archive. */
-	void SetRequest(const UE::DerivedData::FCacheRequest& InRequest);
+	void SetRequest(const UE::DerivedData::FRequest& InRequest);
 	/** Callback from the CacheRequest; set whether we're reading from EditorDomain bytes or WorkspaceDomain archive. */
 	void OnCacheRequestComplete(UE::DerivedData::FCacheGetCompleteParams&& Params);
 
@@ -119,7 +119,7 @@ private:
 
 	// Data in this section is either read-only, or is read and written only on the public interface thread.
 	TRefCountPtr<FEditorDomain::FLocks> EditorDomainLocks; // Pointer is read-only, *pointer has internal lock
-	UE::DerivedData::FCacheRequest Request; // Read-only
+	UE::DerivedData::FRequest Request; // Read-only
 	FPackagePath PackagePath; // Read-only
 	mutable ESource Source = ESource::Uninitialized; // Interface-thread-only
 	TRefCountPtr<FEditorDomain::FPackageSource> PackageSource; // Read-only, *pointer requires EditorDomainLocks.Lock
diff --git a/Engine/Source/Editor/UnrealEd/Private/EditorDomain/EditorDomainUtils.cpp b/Engine/Source/Editor/UnrealEd/Private/EditorDomain/EditorDomainUtils.cpp
--- a/Engine/Source/Editor/UnrealEd/Private/EditorDomain/EditorDomainUtils.cpp
+++ b/Engine/Source/Editor/UnrealEd/Private/EditorDomain/EditorDomainUtils.cpp
@@ -74,8 +74,8 @@ UE::DerivedData::FCacheKey GetEditorDomainPackageKey(const FPackageDigest& Packa
 	return UE::DerivedData::FCacheKey(EditorDomainPackageCacheBucket, PackageDigest);
 }
 
-UE::DerivedData::FCacheRequest RequestEditorDomainPackage(const FPackagePath& PackagePath,
-	const FPackageDigest& PackageDigest, UE::DerivedData::ECachePriority CachePriority,
+UE::DerivedData::FRequest RequestEditorDomainPackage(const FPackagePath& PackagePath,
+	const FPackageDigest& PackageDigest, UE::DerivedData::EPriority CachePriority,
 	UE::DerivedData::FOnCacheGetComplete&& Callback)
 {
 	UE::DerivedData::ICache& Cache = GetDerivedDataCacheRef().GetCache();
diff --git a/Engine/Source/Editor/UnrealEd/Private/EditorDomain/EditorDomainUtils.h b/Engine/Source/Editor/UnrealEd/Private/EditorDomain/EditorDomainUtils.h
--- a/Engine/Source/Editor/UnrealEd/Private/EditorDomain/EditorDomainUtils.h
+++ b/Engine/Source/Editor/UnrealEd/Private/EditorDomain/EditorDomainUtils.h
@@ -41,8 +41,8 @@ EPackageDigestResult GetPackageDigest(IAssetRegistry& AssetRegistry, FName Packa
 UE::DerivedData::FCacheKey GetEditorDomainPackageKey(const FPackageDigest& PackageDigest);
 
 /** Get the CacheRequest for the given package from the EditorDomain cache bucket. */
-UE::DerivedData::FCacheRequest RequestEditorDomainPackage(const FPackagePath& PackagePath,
-	const FPackageDigest& PackageDigest, UE::DerivedData::ECachePriority CachePriority,
+UE::DerivedData::FRequest RequestEditorDomainPackage(const FPackagePath& PackagePath,
+	const FPackageDigest& PackageDigest, UE::DerivedData::EPriority CachePriority,
 	UE::DerivedData::FOnCacheGetComplete&& Callback);
 
 }
diff --git a/Engine/Source/Runtime/Core/Private/Compression/CompressedBuffer.cpp b/Engine/Source/Runtime/Core/Private/Compression/CompressedBuffer.cpp
new file mode 100644
--- /dev/null
+++ b/Engine/Source/Runtime/Core/Private/Compression/CompressedBuffer.cpp
@@ -0,0 +1,23 @@
+// Copyright Epic Games, Inc. All Rights Reserved.
+
+#include "Compression/CompressedBuffer.h"
+
+FCompressedBuffer FCompressedBuffer::Compress(const FSharedBuffer& Buffer)
+{
+	// DDC-TODO: Switch to the new compression interface when it is ready. For now, match the "no compression" method.
+	FUniqueBuffer CompressedBuffer = FUniqueBuffer::Alloc(Buffer.GetSize() + 1);
+	FMemory::Memset(CompressedBuffer.GetData(), 0, 1);
+	FMemory::Memcpy(CompressedBuffer.GetView().RightChop(1).GetData(), Buffer.GetData(), Buffer.GetSize());
+	return FCompressedBuffer(FSharedBuffer(MoveTemp(CompressedBuffer)));
+}
+
+FSharedBuffer FCompressedBuffer::Decompress(const FCompressedBuffer& Buffer)
+{
+	// DDC-TODO: Switch to the new compression interface when it is ready. For now, match the "no compression" method.
+	if (FSharedBuffer Compressed = Buffer.GetCompressed())
+	{
+		const FMemoryView View = Compressed.GetView() + 1;
+		return FSharedBuffer::MakeView(View, MoveTemp(Compressed));
+	}
+	return FSharedBuffer();
+}
diff --git a/Engine/Source/Runtime/Core/Public/Compression/CompressedBuffer.h b/Engine/Source/Runtime/Core/Public/Compression/CompressedBuffer.h
new file mode 100644
--- /dev/null
+++ b/Engine/Source/Runtime/Core/Public/Compression/CompressedBuffer.h
@@ -0,0 +1,49 @@
+// Copyright Epic Games, Inc. All Rights Reserved.
+
+#pragma once
+
+#include "Memory/SharedBuffer.h"
+
+/**
+ * A compressed buffer stores data in a self-contained format.
+ *
+ * A buffer is self-contained in the sense that it can be decompressed without external knowledge
+ * of the compression format or the size of the uncompressed data.
+ */
+class FCompressedBuffer
+{
+public:
+	/** Construct a null compressed buffer. */
+	FCompressedBuffer() = default;
+
+	/** Construct from a compressed shared buffer. */
+	inline explicit FCompressedBuffer(FSharedBuffer InCompressed)
+		: Compressed(MoveTemp(InCompressed))
+	{
+	}
+
+	/** Returns true if the compressed buffer is not null. */
+	inline explicit operator bool() const { return !IsNull(); }
+
+	/** Returns true if the compressed buffer is null. */
+	inline bool IsNull() const { return Compressed.IsNull(); }
+
+	/** Returns the shared buffer containing the compressed data. May be null. */
+	inline FSharedBuffer GetCompressed() const & { return Compressed; }
+	inline FSharedBuffer GetCompressed() && { return MoveTemp(Compressed); }
+
+	/** Returns the size of the compressed data. */
+	inline uint64 GetCompressedSize() const { return Compressed.GetSize(); }
+
+	/** Returns an owned shared buffer containing the uncompresesd data. May be null. */
+	inline FSharedBuffer Decompress() const { return Decompress(*this); }
+
+	/** Compress the buffer using default compression settings. Returns null if input is null. */
+	CORE_API static FCompressedBuffer Compress(const FSharedBuffer& Buffer);
+
+	/** Decompress the compressed buffer. Returns null if input is null. */
+	CORE_API static FSharedBuffer Decompress(const FCompressedBuffer& Buffer);
+
+private:
+	FSharedBuffer Compressed;
+};
