DDC: Simplified comparison of FCacheBucket and the keys that depend on it

#rnx

diff --git a/Engine/Source/Developer/DerivedDataCache/Private/DerivedDataCacheKey.cpp b/Engine/Source/Developer/DerivedDataCache/Private/DerivedDataCacheKey.cpp
--- a/Engine/Source/Developer/DerivedDataCache/Private/DerivedDataCacheKey.cpp
+++ b/Engine/Source/Developer/DerivedDataCache/Private/DerivedDataCacheKey.cpp
@@ -3,8 +3,11 @@
 #include "DerivedDataCacheKey.h"
 
 #include "Algo/AllOf.h"
+#include "Containers/Set.h"
+#include "Containers/StringConv.h"
+#include "HAL/CriticalSection.h"
+#include "Misc/ScopeRWLock.h"
 #include "Misc/StringBuilder.h"
-#include "UObject/NameTypes.h"
 
 namespace UE
 {
@@ -13,32 +16,111 @@ namespace DerivedData
 
 ///////////////////////////////////////////////////////////////////////////////////////////////////////////////////////
 
-FCacheBucket::FCacheBucket(FStringView Name)
-	: Index(FName(Name).GetDisplayIndex().ToUnstableInt())
+class FCacheBuckets
 {
-	check(Algo::AllOf(Name, FChar::IsAlnum));
-}
-
-void FCacheBucket::ToString(FAnsiStringBuilderBase& Builder) const
-{
-	verify(FName::CreateFromDisplayId(FNameEntryId::FromUnstableInt(Index), 0).TryAppendAnsiString(Builder));
-}
+public:
+	inline const TCHAR* FindOrAdd(FStringView Name)
+	{
+		const uint32 Hash = GetTypeHash(Name);
+		{
+			FReadScopeLock ReadLock(Lock);
+			if (const FBucket* Bucket = Buckets.FindByHash(Hash, Name))
+			{
+				return Bucket->GetName();
+			}
+		}
+		{
+			FWriteScopeLock WriteLock(Lock);
+			const FBucket& Bucket = Buckets.FindOrAddByHash(Hash, FBucket(Name));
+			return Bucket.GetName();
+		}
+	}
 
-void FCacheBucket::ToString(FWideStringBuilderBase& Builder) const
+private:
+	class FBucket
+	{
+	public:
+		inline explicit FBucket(FStringView Bucket)
+		{
+			const int32 BucketLen = Bucket.Len();
+			TCHAR* Buffer = new TCHAR[BucketLen + 2];
+			*reinterpret_cast<uint8*>(Buffer++) = static_cast<uint8>(BucketLen);
+			Bucket.CopyString(Buffer, BucketLen);
+			Buffer[BucketLen] = TEXT('\0');
+			Name = Buffer;
+		}
+
+		FBucket(const FBucket&) = delete;
+		FBucket& operator=(const FBucket&) = delete;
+
+		inline FBucket(FBucket&& Bucket)
+			: Name(Bucket.Name)
+		{
+			Bucket.Name = nullptr;
+		}
+
+		FBucket& operator=(FBucket&& Bucket)
+		{
+			Name = Bucket.Name;
+			Bucket.Name = nullptr;
+			return *this;
+		}
+
+		inline ~FBucket()
+		{
+			if (Name)
+			{
+				delete[] (Name - 1);
+			}
+		}
+
+		inline FStringView ToString() const
+		{
+			return FStringView(Name, *reinterpret_cast<const uint8*>(Name - 1));
+		}
+
+		inline bool operator==(const FBucket& Bucket) const
+		{
+			return Name == Bucket.Name;
+		}
+
+		inline bool operator==(FStringView Bucket) const
+		{
+			return ToString().Equals(Bucket, ESearchCase::IgnoreCase);
+		}
+
+		friend inline uint32 GetTypeHash(const FBucket& Bucket)
+		{
+			return GetTypeHash(Bucket.ToString());
+		}
+
+		inline const TCHAR* GetName() const { return Name; }
+
+	private:
+		const TCHAR* Name;
+	};
+
+	FRWLock Lock;
+	TSet<FBucket> Buckets;
+};
+
+const TCHAR* FCacheBucket::FindOrAddCacheBucket(FStringView Bucket)
 {
-	FName::CreateFromDisplayId(FNameEntryId::FromUnstableInt(Index), 0).AppendString(Builder);
+	checkf(!Bucket.IsEmpty() && Bucket.Len() < 256 && Algo::AllOf(Bucket, FChar::IsAlnum),
+		TEXT("Invalid cache bucket name '%.*s' must be alphanumeric, non-empty, and contain fewer than 256 characters."),
+		Bucket.Len(), Bucket.GetData());
+	static FCacheBuckets Buckets;
+	return Buckets.FindOrAdd(Bucket);
 }
 
-FString FCacheBucket::ToString() const
+void FCacheBucket::ToString(FAnsiStringBuilderBase& Builder) const
 {
-	TStringBuilder<32> Out;
-	ToString(Out);
-	return FString(Out);
+	Builder << TCHAR_TO_UTF8(Name);
 }
 
-bool FCacheBucketLexicalLess::operator()(FCacheBucket A, FCacheBucket B) const
+void FCacheBucket::ToString(FWideStringBuilderBase& Builder) const
 {
-	return FNameEntryId::FromUnstableInt(A.ToIndex()).LexicalLess(FNameEntryId::FromUnstableInt(B.ToIndex()));
+	Builder << ToString();
 }
 
 ///////////////////////////////////////////////////////////////////////////////////////////////////////////////////////
diff --git a/Engine/Source/Developer/DerivedDataCache/Public/DerivedDataCacheKey.h b/Engine/Source/Developer/DerivedDataCache/Public/DerivedDataCacheKey.h
--- a/Engine/Source/Developer/DerivedDataCache/Public/DerivedDataCacheKey.h
+++ b/Engine/Source/Developer/DerivedDataCache/Public/DerivedDataCacheKey.h
@@ -27,45 +27,48 @@ public:
 	/** Construct a null cache bucket. */
 	FCacheBucket() = default;
 
-	/** Construct a cache bucket from the alphanumeric name. Names are case-insensitive. */
-	UE_API explicit FCacheBucket(FStringView Name);
+	/**
+	 * Construct a cache bucket from the name.
+	 *
+	 * A cache bucket name must be alphanumeric, non-empty, and contain fewer than 256 characters.
+	 */
+	inline explicit FCacheBucket(FStringView Bucket)
+		: Name(FindOrAddCacheBucket(Bucket))
+	{
+	}
 
 	/** Convert the name of the cache bucket to a string. */
 	UE_API void ToString(FAnsiStringBuilderBase& Builder) const;
 	UE_API void ToString(FWideStringBuilderBase& Builder) const;
-	UE_API FString ToString() const;
+	inline FStringView ToString() const;
 
 	/** Whether this is null. */
-	inline bool IsNull() const { return !Index; }
+	inline bool IsNull() const { return !Name; }
 	/** Whether this is not null. */
 	inline bool IsValid() const { return !IsNull(); }
 
-	/** Returns the index that is only stable within this process. */
-	inline uint32 ToIndex() const { return Index; }
+	/** Reset this to null. */
+	inline void Reset() { *this = FCacheBucket(); }
+
+	inline bool operator==(FCacheBucket Other) const { return Name == Other.Name; }
+	inline bool operator!=(FCacheBucket Other) const { return Name != Other.Name; }
 
-	/** Returns the bucket for the index produced within this process. */
-	static inline FCacheBucket FromIndex(uint32 InIndex)
+	inline bool operator<(FCacheBucket Other) const
 	{
-		FCacheBucket Bucket;
-		Bucket.Index = InIndex;
-		return Bucket;
+		return Name != Other.Name && ToString().Compare(Other.ToString(), ESearchCase::IgnoreCase) < 0;
 	}
 
-	/** Reset this to null. */
-	inline void Reset() { *this = FCacheBucket(); }
+	friend inline uint32 GetTypeHash(FCacheBucket Bucket) { return ::GetTypeHash(reinterpret_cast<UPTRINT>(Bucket.Name)); }
 
 private:
-	uint32 Index = 0;
-};
+	UE_API static const TCHAR* FindOrAddCacheBucket(FStringView Bucket);
 
-inline bool operator==(FCacheBucket A, FCacheBucket B)
-{
-	return A.ToIndex() == B.ToIndex();
-}
+	const TCHAR* Name = nullptr;
+};
 
-inline uint32 GetTypeHash(FCacheBucket Bucket)
+inline FStringView FCacheBucket::ToString() const
 {
-	return ::GetTypeHash(Bucket.ToIndex());
+	return Name ? FStringView(Name, *reinterpret_cast<const uint8*>(Name - 1)) : FStringView();
 }
 
 template <typename CharType>
@@ -75,21 +78,6 @@ inline TStringBuilderBase<CharType>& operator<<(TStringBuilderBase<CharType>& Bu
 	return Builder;
 }
 
-/** Fast less with an order that is only stable within the lifetime of this process. */
-struct FCacheBucketFastLess
-{
-	inline bool operator()(FCacheBucket A, FCacheBucket B) const
-	{
-		return A.ToIndex() < B.ToIndex();
-	}
-};
-
-/** Slow less with lexical order that is stable between processes. */
-struct FCacheBucketLexicalLess
-{
-	UE_API bool operator()(FCacheBucket A, FCacheBucket B) const;
-};
-
 ///////////////////////////////////////////////////////////////////////////////////////////////////////////////////////
 
 /** A key that uniquely identifies a cache record. */
@@ -135,6 +123,18 @@ inline bool operator==(const FCacheKey& A, const FCacheKey& B)
 	return A.GetBucket() == B.GetBucket() && A.GetHash() == B.GetHash();
 }
 
+inline bool operator!=(const FCacheKey& A, const FCacheKey& B)
+{
+	return A.GetBucket() != B.GetBucket() || A.GetHash() != B.GetHash();
+}
+
+inline bool operator<(const FCacheKey& A, const FCacheKey& B)
+{
+	const FCacheBucket& BucketA = A.GetBucket();
+	const FCacheBucket& BucketB = B.GetBucket();
+	return BucketA == BucketB ? A.GetHash() < B.GetHash() : BucketA < BucketB;
+}
+
 inline uint32 GetTypeHash(const FCacheKey& Key)
 {
 	return HashCombine(GetTypeHash(Key.GetBucket()), GetTypeHash(Key.GetHash()));
@@ -147,28 +147,6 @@ inline TStringBuilderBase<CharType>& operator<<(TStringBuilderBase<CharType>& Bu
 	return Builder;
 }
 
-/** Fast less with an order that is only stable within the lifetime of this process. */
-struct FCacheKeyFastLess
-{
-	inline bool operator()(const FCacheKey& A, const FCacheKey& B) const
-	{
-		const FCacheBucket BucketA = A.GetBucket();
-		const FCacheBucket BucketB = B.GetBucket();
-		return BucketA == BucketB ? A.GetHash() < B.GetHash() : FCacheBucketFastLess()(BucketA, BucketB);
-	}
-};
-
-/** Slow less with lexical order that is stable between processes. */
-struct FCacheKeyLexicalLess
-{
-	inline bool operator()(const FCacheKey& A, const FCacheKey& B) const
-	{
-		const FCacheBucket BucketA = A.GetBucket();
-		const FCacheBucket BucketB = B.GetBucket();
-		return BucketA == BucketB ? A.GetHash() < B.GetHash() : FCacheBucketLexicalLess()(BucketA, BucketB);
-	}
-};
-
 ///////////////////////////////////////////////////////////////////////////////////////////////////////////////////////
 
 /** A key that uniquely identifies a payload within a cache record. */
@@ -214,6 +192,18 @@ inline bool operator==(const FCachePayloadKey& A, const FCachePayloadKey& B)
 	return A.GetKey() == B.GetKey() && A.GetId() == B.GetId();
 }
 
+inline bool operator!=(const FCachePayloadKey& A, const FCachePayloadKey& B)
+{
+	return A.GetKey() != B.GetKey() || A.GetId() != B.GetId();
+}
+
+inline bool operator<(const FCachePayloadKey& A, const FCachePayloadKey& B)
+{
+	const FCacheKey& KeyA = A.GetKey();
+	const FCacheKey& KeyB = B.GetKey();
+	return KeyA == KeyB ? A.GetId() < B.GetId() : KeyA < KeyB;
+}
+
 inline uint32 GetTypeHash(const FCachePayloadKey& Key)
 {
 	return HashCombine(GetTypeHash(Key.GetKey()), GetTypeHash(Key.GetId()));
@@ -226,28 +216,6 @@ inline TStringBuilderBase<CharType>& operator<<(TStringBuilderBase<CharType>& Bu
 	return Builder;
 }
 
-/** Fast less with an order that is only stable within the lifetime of this process. */
-struct FCachePayloadKeyFastLess
-{
-	inline bool operator()(const FCachePayloadKey& A, const FCachePayloadKey& B) const
-	{
-		const FCacheKey& KeyA = A.GetKey();
-		const FCacheKey& KeyB = B.GetKey();
-		return KeyA == KeyB ? A.GetId() < B.GetId() : FCacheKeyFastLess()(KeyA, KeyB);
-	}
-};
-
-/** Slow less with lexical order that is stable between processes. */
-struct FCachePayloadKeyLexicalLess
-{
-	inline bool operator()(const FCachePayloadKey& A, const FCachePayloadKey& B) const
-	{
-		const FCacheKey& KeyA = A.GetKey();
-		const FCacheKey& KeyB = B.GetKey();
-		return KeyA == KeyB ? A.GetId() < B.GetId() : FCacheKeyLexicalLess()(KeyA, KeyB);
-	}
-};
-
 ///////////////////////////////////////////////////////////////////////////////////////////////////////////////////////
 
 } // DerivedData
