Adding command line override for Jupiter Host in ShooterGame, DDCHttpHost
Changed the Dev toggle for CookCompare to use this.
This can also be used to locally toggle Jupiter to run against a local deployment.

diff --git a/Engine/Source/Developer/DerivedDataCache/Private/DerivedDataBackends.cpp b/Engine/Source/Developer/DerivedDataCache/Private/DerivedDataBackends.cpp
--- a/Engine/Source/Developer/DerivedDataCache/Private/DerivedDataBackends.cpp
+++ b/Engine/Source/Developer/DerivedDataCache/Private/DerivedDataBackends.cpp
@@ -700,6 +700,18 @@ public:
 			return nullptr;
 		}
 
+		// Allow Host to be overriden via CommandLineHostOverride
+		FString CommandLineOverride;
+		if( FParse::Value( Entry, TEXT("CommandLineHostOverride="), CommandLineOverride ) )
+		{
+			FString Value;
+			if (FParse::Value(FCommandLine::Get(), *(CommandLineOverride + TEXT("=")), Value))
+			{
+				ServiceUrl = Value;
+				UE_LOG(LogDerivedDataCache, Warning, TEXT("Found command line override for Host %s=%s"), *CommandLineOverride, *ServiceUrl);
+			}
+		}
+
 		FString Namespace;
 		if (!FParse::Value(Entry, TEXT("Namespace="), Namespace))
 		{
