Reduced Error output to Display in HttpBackend as this is a recoverable error.

diff --git a/Engine/Source/Developer/DerivedDataCache/Private/HttpDerivedDataBackend.cpp b/Engine/Source/Developer/DerivedDataCache/Private/HttpDerivedDataBackend.cpp
--- a/Engine/Source/Developer/DerivedDataCache/Private/HttpDerivedDataBackend.cpp
+++ b/Engine/Source/Developer/DerivedDataCache/Private/HttpDerivedDataBackend.cpp
@@ -513,7 +513,7 @@ private:
 				{ 
 					UE_LOG(
 						LogDerivedDataCache,
-						Error,
+						Display,
 						TEXT("Failed %s HTTP cache entry (response %d) from %s. Response: %s"),
 						VerbStr,
 						ResponseCode,
@@ -529,7 +529,7 @@ private:
 		{
 			UE_LOG(
 				LogDerivedDataCache, 
-				Error, 
+				Display, 
 				TEXT("Error while connecting to %s: %s"), 
 				*Domain, 
 				ANSI_TO_TCHAR(curl_easy_strerror(Result))
