DDC: Add generic EStatus and replace ECacheStatus with it

#rnx

diff --git a/Engine/Source/Developer/DerivedDataCache/Private/DerivedDataCache.cpp b/Engine/Source/Developer/DerivedDataCache/Private/DerivedDataCache.cpp
--- a/Engine/Source/Developer/DerivedDataCache/Private/DerivedDataCache.cpp
+++ b/Engine/Source/Developer/DerivedDataCache/Private/DerivedDataCache.cpp
@@ -737,7 +737,7 @@ public:
 		EPriority Priority,
 		FOnCacheGetComplete&& Callback) final;
 
-	virtual FRequest GetPayloads(
+	virtual FRequest GetPayload(
 		TConstArrayView<FCachePayloadKey> Keys,
 		FStringView Context,
 		ECachePolicy Policy,
@@ -827,7 +827,6 @@ void FCache::Put(
 	{
 		UE_LOG(LogDerivedDataCache, VeryVerbose, TEXT("Cache: Put skipped for %s from '%.*s'"),
 			*TToString<96>(Params.Key), Context.Len(), Context.GetData());
-		Params.Status = ECacheStatus::NotCached;
 		return;
 	}
 
@@ -884,7 +883,7 @@ void FCache::Put(
 
 	UE_LOG(LogDerivedDataCache, VeryVerbose, TEXT("Cache: Put for %s from '%.*s'"),
 		*TToString<96>(Params.Key), Context.Len(), Context.GetData());
-	Params.Status = ECacheStatus::Cached;
+	Params.Status = EStatus::Ok;
 }
 
 FRequest FCache::Put(
@@ -1041,7 +1040,7 @@ void FCache::Get(
 
 	UE_LOG(LogDerivedDataCache, Verbose, TEXT("Cache: Get cache hit for %s from '%.*s'"),
 		*TToString<96>(Key), Context.Len(), Context.GetData());
-	Params.Status = ECacheStatus::Cached;
+	Params.Status = EStatus::Ok;
 }
 
 FRequest FCache::Get(
@@ -1144,7 +1143,7 @@ void FCache::GetPayload(
 			{
 				UE_LOG(LogDerivedDataCache, Verbose, TEXT("Cache: GetPayload (exists) cache hit for %s from '%.*s'"),
 					*TToString<96>(Key), Context.Len(), Context.GetData());
-				Params.Status = ECacheStatus::Cached;
+				Params.Status = EStatus::Ok;
 			}
 			else
 			{
@@ -1174,10 +1173,10 @@ void FCache::GetPayload(
 
 	UE_LOG(LogDerivedDataCache, Verbose, TEXT("Cache: GetPayload cache hit for %s from '%.*s'"),
 		*TToString<96>(Key), Context.Len(), Context.GetData());
-	Params.Status = ECacheStatus::Cached;
+	Params.Status = EStatus::Ok;
 }
 
-FRequest FCache::GetPayloads(
+FRequest FCache::GetPayload(
 	TConstArrayView<FCachePayloadKey> Keys,
 	FStringView Context,
 	ECachePolicy Policy,
diff --git a/Engine/Source/Developer/DerivedDataCache/Public/DerivedDataCache.h b/Engine/Source/Developer/DerivedDataCache/Public/DerivedDataCache.h
--- a/Engine/Source/Developer/DerivedDataCache/Public/DerivedDataCache.h
+++ b/Engine/Source/Developer/DerivedDataCache/Public/DerivedDataCache.h
@@ -13,6 +13,14 @@
 namespace DerivedData
 {
 
+struct FCacheGetCompleteParams;
+struct FCacheGetPayloadCompleteParams;
+struct FCachePutCompleteParams;
+
+using FOnCachePutComplete = TUniqueFunction<void (FCachePutCompleteParams&& Params)>;
+using FOnCacheGetComplete = TUniqueFunction<void (FCacheGetCompleteParams&& Params)>;
+using FOnCacheGetPayloadComplete = TUniqueFunction<void (FCacheGetPayloadCompleteParams&& Params)>;
+
 ///////////////////////////////////////////////////////////////////////////////////////////////////////////////////////
 
 /**
@@ -81,74 +89,6 @@
 
 ///////////////////////////////////////////////////////////////////////////////////////////////////////////////////////
 
-/** Status of a cache operation. */
-enum class ECacheStatus : uint8
-{
-	/** A cache miss for a get operation, or a failure for a put operation. */
-	NotCached,
-	/** A cache hit for a get operation, or a success for a put operation. */
-	Cached,
-	/** The operation was canceled before it completed. */
-	Canceled,
-};
-
-///////////////////////////////////////////////////////////////////////////////////////////////////////////////////////
-
-/** Parameters for the completion callback for cache put requests. */
-struct FCachePutCompleteParams
-{
-	/** Key for the part of the put request that completed or was canceled. */
-	FCacheKey Key;
-
-	/** Status of the cache request. */
-	ECacheStatus Status = ECacheStatus::NotCached;
-};
-
-/** Parameters for the completion callback for cache get requests. */
-struct FCacheGetCompleteParams
-{
-	/**
-	 * Record for the part of the get request that completed or was canceled.
-	 *
-	 * The key is always populated. The remainder of the record is populated when Status is Cached.
-	 *
-	 * The value, attachments, and metadata may be skipped based on cache policy flags. When a value
-	 * or attachment has been skipped, it will have a payload but its buffers will be null.
-	 */
-	FCacheRecord Record;
-
-	/** Status of the cache request. */
-	ECacheStatus Status = ECacheStatus::NotCached;
-};
-
-/** Parameters for the completion callback for cache payload requests. */
-struct FCacheGetPayloadCompleteParams
-{
-	/** Key for the part of the payload request that completed or was canceled. See Payload for ID. */
-	FCacheKey Key;
-
-	/**
-	 * Payload for the part of the payload request that completed or was canceled.
-	 *
-	 * The ID is always populated. The hash and buffer are populated when Status is Cached.
-	 */
-	FPayload Payload;
-
-	/** Status of the cache request. */
-	ECacheStatus Status = ECacheStatus::NotCached;
-};
-
-/** Callback for completion of a put request. */
-using FOnCachePutComplete = TUniqueFunction<void (FCachePutCompleteParams&& Params)>;
-
-/** Callback for completion of a get request. */
-using FOnCacheGetComplete = TUniqueFunction<void (FCacheGetCompleteParams&& Params)>;
-
-/** Callback for completion of a payload request. */
-using FOnCacheGetPayloadComplete = TUniqueFunction<void (FCacheGetPayloadCompleteParams&& Params)>;
-
-///////////////////////////////////////////////////////////////////////////////////////////////////////////////////////
-
 /**
  * Interface to the cache.
  *
@@ -211,7 +151,7 @@
 	 * @param Priority A priority to consider when scheduling the request. See EPriority.
 	 * @param Callback A callback invoked for every key in the batch as it completes or is canceled.
 	 */
-	virtual FRequest GetPayloads(
+	virtual FRequest GetPayload(
 		TConstArrayView<FCachePayloadKey> Keys,
 		FStringView Context,
 		ECachePolicy Policy,
@@ -231,5 +171,51 @@
 
 ///////////////////////////////////////////////////////////////////////////////////////////////////////////////////////
 
+/** Parameters for the completion callback for cache put requests. */
+struct FCachePutCompleteParams
+{
+	/** Key for the part of the put request that completed or was canceled. */
+	FCacheKey Key;
+
+	/** Status of the cache request. */
+	EStatus Status = EStatus::Error;
+};
+
+/** Parameters for the completion callback for cache get requests. */
+struct FCacheGetCompleteParams
+{
+	/**
+	 * Record for the part of the get request that completed or was canceled.
+	 *
+	 * The key is always populated. The remainder of the record is populated when Status is Cached.
+	 *
+	 * The value, attachments, and metadata may be skipped based on cache policy flags. When a value
+	 * or attachment has been skipped, it will have a payload but its buffers will be null.
+	 */
+	FCacheRecord Record;
+
+	/** Status of the cache request. */
+	EStatus Status = EStatus::Error;
+};
+
+/** Parameters for the completion callback for cache payload requests. */
+struct FCacheGetPayloadCompleteParams
+{
+	/** Key for the part of the payload request that completed or was canceled. See Payload for ID. */
+	FCacheKey Key;
+
+	/**
+	 * Payload for the part of the payload request that completed or was canceled.
+	 *
+	 * The ID is always populated. The hash and buffer are populated when Status is Cached.
+	 */
+	FPayload Payload;
+
+	/** Status of the cache request. */
+	EStatus Status = EStatus::Error;
+};
+
+///////////////////////////////////////////////////////////////////////////////////////////////////////////////////////
+
 } // DerivedData
 } // UE
diff --git a/Engine/Source/Developer/DerivedDataCache/Public/DerivedDataRequest.h b/Engine/Source/Developer/DerivedDataCache/Public/DerivedDataRequest.h
--- a/Engine/Source/Developer/DerivedDataCache/Public/DerivedDataRequest.h
+++ b/Engine/Source/Developer/DerivedDataCache/Public/DerivedDataRequest.h
@@ -12,7 +12,7 @@ namespace DerivedData
 
 ///////////////////////////////////////////////////////////////////////////////////////////////////////////////////////
 
-/** Priority for scheduling requests. */
+/** Priority for scheduling a request. */
 enum class EPriority : uint8
 {
 	/**
@@ -50,7 +50,20 @@ enum class EPriority : uint8
 
 ///////////////////////////////////////////////////////////////////////////////////////////////////////////////////////
 
-/** Interface for a request. Provides functionality common to every type of request. */
+/** Status of a request that has completed. */
+enum class EStatus : uint8
+{
+	/** The request completed successfully. Any requested data is available. */
+	Ok,
+	/** The request completed unsuccessfully. Any requested data is not available. */
+	Error,
+	/** The request was canceled before it completed. Any requested data is not available. */
+	Canceled,
+};
+
+///////////////////////////////////////////////////////////////////////////////////////////////////////////////////////
+
+/** Interface to a request. Use via FRequest. */
 class IRequest : public IRefCountedObject
 {
 public:
@@ -71,7 +84,7 @@ public:
 
 ///////////////////////////////////////////////////////////////////////////////////////////////////////////////////////
 
-/** A wrapper for the request interface that manages lifetime with reference counting. */
+/** Handle to a request that manages its lifetime with reference counting. */
 class FRequest
 {
 public:
diff --git a/Engine/Source/Editor/UnrealEd/Private/EditorDomain/EditorDomain.cpp b/Engine/Source/Editor/UnrealEd/Private/EditorDomain/EditorDomain.cpp
--- a/Engine/Source/Editor/UnrealEd/Private/EditorDomain/EditorDomain.cpp
+++ b/Engine/Source/Editor/UnrealEd/Private/EditorDomain/EditorDomain.cpp
@@ -184,7 +184,7 @@ int64 FEditorDomain::FileSize(const FPackagePath& PackagePath, EPackageSegment P
 		{
 			FScopeLock ScopeLock(&Locks->Lock);
 			if ((PackageSource->Source == FEditorDomain::EPackageSource::Undecided || PackageSource->Source == FEditorDomain::EPackageSource::Editor) &&
-				Params.Status == UE::DerivedData::ECacheStatus::Cached)
+				Params.Status == UE::DerivedData::EStatus::Ok)
 			{
 				const FCbObject& MetaData = Params.Record.GetMeta();
 				FileSize = MetaData["FileSize"].AsInt64();
diff --git a/Engine/Source/Editor/UnrealEd/Private/EditorDomain/EditorDomainArchive.cpp b/Engine/Source/Editor/UnrealEd/Private/EditorDomain/EditorDomainArchive.cpp
--- a/Engine/Source/Editor/UnrealEd/Private/EditorDomain/EditorDomainArchive.cpp
+++ b/Engine/Source/Editor/UnrealEd/Private/EditorDomain/EditorDomainArchive.cpp
@@ -182,7 +182,7 @@ void FEditorDomainReadArchive::OnCacheRequestComplete(UE::DerivedData::FCacheGet
 	check(AsyncSource == ESource::Uninitialized);
 	FScopeLock DomainScopeLock(&EditorDomainLocks->Lock);
 	if ((PackageSource->Source == FEditorDomain::EPackageSource::Undecided || PackageSource->Source == FEditorDomain::EPackageSource::Editor) &&
-		Params.Status == UE::DerivedData::ECacheStatus::Cached)
+		Params.Status == UE::DerivedData::EStatus::Ok)
 	{
 		const FCbObject& MetaData = Params.Record.GetMeta();
 		int64 FileSize = MetaData["FileSize"].AsInt64(-1);
@@ -411,7 +411,7 @@ void FEditorDomainAsyncReadFileHandle::OnCacheRequestComplete(UE::DerivedData::F
 	check(AsyncSource == ESource::Uninitialized);
 	FScopeLock DomainScopeLock(&EditorDomainLocks->Lock);
 	if ((PackageSource->Source == FEditorDomain::EPackageSource::Undecided || PackageSource->Source == FEditorDomain::EPackageSource::Editor) &&
-		Params.Status == UE::DerivedData::ECacheStatus::Cached)
+		Params.Status == UE::DerivedData::EStatus::Ok)
 	{
 		const FCbObject& MetaData = Params.Record.GetMeta();
 		int64 FileSize = MetaData["FileSize"].AsInt64(-1);
