Batched Jupiter queries

Adds a helper class to make CachedDataProbablyExists, GetCachedData and PutCachedData operations. The helper checks the availble connections in the pool and decides if the thread should make a regular request or batch together with other requests. This will reduce the number of concurrent connections on each client.

Disabled by default.

diff --git a/Engine/Source/Developer/DerivedDataCache/Private/HttpDerivedDataBackend.cpp b/Engine/Source/Developer/DerivedDataCache/Private/HttpDerivedDataBackend.cpp
--- a/Engine/Source/Developer/DerivedDataCache/Private/HttpDerivedDataBackend.cpp
+++ b/Engine/Source/Developer/DerivedDataCache/Private/HttpDerivedDataBackend.cpp
@@ -11,6 +11,8 @@
 #if PLATFORM_WINDOWS || PLATFORM_HOLOLENS
 #include "Windows/HideWindowsPlatformTypes.h"
 #endif
+#include "Algo/Transform.h"
+#include "Algo/Find.h"
 #include "Containers/StaticArray.h"
 #include "Containers/Ticker.h"
 #include "Dom/JsonObject.h"
@@ -23,17 +25,38 @@
 #include "Misc/StringBuilder.h"
 #include "ProfilingDebugging/CpuProfilerTrace.h"
 #include "ProfilingDebugging/CountersTrace.h"
+#include "Serialization/BufferArchive.h"
 #include "Serialization/JsonReader.h"
+#include "Serialization/JsonWriter.h"
 #include "Serialization/JsonSerializer.h"
+#include "Policies/CondensedJsonPrintPolicy.h"
+
+// Enables data request helpers that internally
+// batch requests to reduce the number of concurrent
+// connections.
+#ifndef WITH_DATAREQUEST_HELPER
+	#define WITH_DATAREQUEST_HELPER 0
+#endif
 
 #define UE_HTTPDDC_BACKEND_WAIT_INTERVAL 0.01f
 #define UE_HTTPDDC_HTTP_REQUEST_TIMEOUT_SECONDS 30L
 #define UE_HTTPDDC_HTTP_REQUEST_TIMOUT_ENABLED 1
 #define UE_HTTPDDC_HTTP_DEBUG 0
-#define UE_HTTPDDC_REQUEST_POOL_SIZE 64
+#if WITH_DATAREQUEST_HELPER
+	#define UE_HTTPDDC_GET_REQUEST_POOL_SIZE 4 
+	#define UE_HTTPDDC_PUT_REQUEST_POOL_SIZE 2
+#else
+	#define UE_HTTPDDC_GET_REQUEST_POOL_SIZE 48
+	#define UE_HTTPDDC_PUT_REQUEST_POOL_SIZE 16
+#endif
 #define UE_HTTPDDC_MAX_FAILED_LOGIN_ATTEMPTS 16
 #define UE_HTTPDDC_MAX_ATTEMPTS 4
 #define UE_HTTPDDC_MAX_BUFFER_RESERVE 104857600u
+#define UE_HTTPDDC_BATCH_SIZE 12
+#define UE_HTTPDDC_BATCH_NUM 64
+#define UE_HTTPDDC_BATCH_GET_WEIGHT 4
+#define UE_HTTPDDC_BATCH_HEAD_WEIGHT 1
+#define UE_HTTPDDC_BATCH_WEIGHT_HINT 12
 
 TRACE_DECLARE_INT_COUNTER(HttpDDC_Exist, TEXT("HttpDDC Exist"));
 TRACE_DECLARE_INT_COUNTER(HttpDDC_ExistHit, TEXT("HttpDDC Exist Hit"));
@@ -689,23 +712,35 @@ private:
 
 };
 
+
+//----------------------------------------------------------------------------------------------------------
+// Forward declarations
+//----------------------------------------------------------------------------------------------------------
+bool VerifyPayload(const FSHAHash& Hash, const TArray<uint8>& Payload);
+bool VerifyPayload(const FIoHash& Hash, const TArray<uint8>& Payload);
+bool VerifyRequest(const class FRequest* Request, const TArray<uint8>& Payload);
+bool HashPayload(class FRequest* Request, const TArrayView<const uint8> Payload);
+
 //----------------------------------------------------------------------------------------------------------
 // Request pool
 //----------------------------------------------------------------------------------------------------------
 
+
 /**
  * Pool that manages a fixed set of requests. Users are required to release requests that have been 
  * acquired. Usable with \ref FScopedRequestPtr which handles this automatically.
  */
 struct FRequestPool
 {
-	FRequestPool(const TCHAR* InServiceUrl, FHttpAccessToken* InAuthorizationToken)
+	FRequestPool(const TCHAR* InServiceUrl, FHttpAccessToken* InAuthorizationToken, uint32 PoolSize)
 	{
+		Pool.AddUninitialized(PoolSize);
 		for (uint8 i = 0; i < Pool.Num(); ++i)
 		{
 			Pool[i].Usage = 0u;
 			Pool[i].Request = new FRequest(InServiceUrl, InAuthorizationToken, true);
 		}
+		
 	}
 
 	~FRequestPool()
@@ -713,11 +748,33 @@ struct FRequestPool
 		for (uint8 i = 0; i < Pool.Num(); ++i)
 		{
 			// No requests should be in use by now.
-			check(Pool[i].Usage.Load(EMemoryOrder::Relaxed) == 0u);
+			check(Pool[i].Usage.load(std::memory_order_acquire) == 0u);
 			delete Pool[i].Request;
 		}
 	}
 
+	/**
+	 * Attempts to get a request is free. Once a request has been returned it is
+	 * "owned by the caller and need to release it to the pool when work has been completed.
+	 * @return Usable request instance if one is available, otherwise null.
+	 */
+	FRequest* GetFreeRequest()
+	{
+		for (uint8 i = 0; i < Pool.Num(); ++i)
+		{
+			if (!Pool[i].Usage.load(std::memory_order_relaxed))
+			{
+				uint8 Expected = 0u;
+				if (Pool[i].Usage.compare_exchange_strong(Expected, 1u))
+				{
+					Pool[i].Request->Reset();
+					return Pool[i].Request;
+				}
+			}
+		}
+		return nullptr;
+	}
+
 	/**
 	 * Block until a request is free. Once a request has been returned it is 
 	 * "owned by the caller and need to release it to the pool when work has been completed.
@@ -726,21 +783,15 @@ struct FRequestPool
 	FRequest* WaitForFreeRequest()
 	{
 		TRACE_CPUPROFILER_EVENT_SCOPE(HttpDDC_WaitForConnPool);
+		FRequest* Request = nullptr;
 		while (true)
 		{
-			for (uint8 i = 0; i < Pool.Num(); ++i)
-			{
-				if (!Pool[i].Usage.Load(EMemoryOrder::Relaxed))
-				{
-					uint8 Expected = 0u;
-					if (Pool[i].Usage.CompareExchange(Expected, 1u))
-					{
-						return Pool[i].Request;
-					}
-				}
-			}
+			Request = GetFreeRequest();
+			if (Request != nullptr)
+				break;
 			FPlatformProcess::Sleep(UE_HTTPDDC_BACKEND_WAIT_INTERVAL);
 		}
+		return Request;
 	}
 
 	/**
@@ -753,9 +804,26 @@ struct FRequestPool
 		{
 			if (Pool[i].Request == Request)
 			{
-				Request->Reset();
-				uint8 Expected = 1u;
-				Pool[i].Usage.CompareExchange(Expected, 0u);
+				
+				Pool[i].Usage--;
+				return;
+			}
+		}
+		check(false);
+	}
+
+	/**
+	 * While holding a request, make it shared across many users.
+	 */
+	void MakeRequestShared(FRequest* Request, uint8 Users)
+	{
+		check(Users != 0);
+		for (uint8 i = 0; i < Pool.Num(); ++i)
+		{
+			if (Pool[i].Request == Request)
+			{
+
+				Pool[i].Usage = Users;
 				return;
 			}
 		}
@@ -766,13 +834,13 @@ private:
 
 	struct FEntry
 	{
-		TAtomic<uint8> Usage;
+		std::atomic<uint8> Usage;
 		FRequest* Request;
 	};
-	
-	TStaticArray<FEntry, UE_HTTPDDC_REQUEST_POOL_SIZE> Pool;
 
-	FRequestPool() {}
+	TArray<FEntry> Pool;
+
+	FRequestPool() = delete;
 };
 
 //----------------------------------------------------------------------------------------------------------
@@ -818,6 +886,614 @@ private:
 };
 
 
+#if WITH_DATAREQUEST_HELPER
+
+//----------------------------------------------------------------------------------------------------------
+// FDataRequestHelper
+//----------------------------------------------------------------------------------------------------------
+/**
+ * Helper class for requesting data. Will batch requests once the number of concurrent requests reach a threshold.
+ */
+struct FDataRequestHelper
+{
+	FDataRequestHelper(FRequestPool* InPool, const TCHAR* InNamespace, const TCHAR* InBucket, const TCHAR* InCacheKey, TArray<uint8>* OutData)
+		: Request(nullptr)
+		, Pool(InPool)
+		, bVerified(false, 1)
+	{
+		Request = Pool->GetFreeRequest();
+		if (Request && OutData != nullptr)
+		{
+			// We are below the threshold, make the connection immediately. OutData is set so this is a get.
+			FString Uri = FString::Printf(TEXT("api/v1/c/ddc/%s/%s/%s.raw"), InNamespace, InBucket, InCacheKey);
+			const FRequest::Result Result = Request->PerformBlockingDownload(*Uri, OutData);
+			if (FRequest::IsSuccessResponse(Request->GetResponseCode()))
+			{
+				if (VerifyRequest(Request, *OutData))
+				{
+					TRACE_COUNTER_ADD(HttpDDC_GetHit, int64(1));
+					TRACE_COUNTER_ADD(HttpDDC_BytesReceived, int64(Request->GetBytesReceived()));
+					bVerified[0] = true;
+				}
+			}
+		}
+		else if (Request)
+		{
+			// We are below the threshold, make the connection immediately. OutData is missing so this is a head.
+			FString Uri = FString::Printf(TEXT("api/v1/c/ddc/%s/%s/%s"), InNamespace, InBucket, InCacheKey);
+			const FRequest::Result Result = Request->PerformBlockingQuery<FRequest::Head>(*Uri);
+			if (FRequest::IsSuccessResponse(Request->GetResponseCode()))
+			{
+				TRACE_COUNTER_ADD(HttpDDC_ExistHit, int64(1));
+				bVerified[0] = true;
+			}
+		}
+		else
+		{
+			// We have exceeded the threshold for concurrent connections, start or add this request
+			// to a batched request.
+			Request = QueueBatchRequest(
+				InPool, 
+				InNamespace,
+				InBucket, 
+				TConstArrayView<const TCHAR*>({ InCacheKey }),
+				OutData ? TConstArrayView<TArray<uint8>*>({ OutData }) : TConstArrayView<TArray<uint8>*>(),
+				bVerified
+			);
+		}
+	}
+
+	// Constructor specifically for batched head queries
+	FDataRequestHelper(FRequestPool* InPool, const TCHAR* InNamespace, const TCHAR* InBucket, TConstArrayView<FString> InCacheKeys)
+		: Request(nullptr)
+		, Pool(InPool)
+		, bVerified(false, InCacheKeys.Num())
+	{
+		// Transform the FString array to char pointers
+		TArray<const TCHAR*> CacheKeys;
+		Algo::Transform(InCacheKeys, CacheKeys, [](const FString& Key) { return *Key; });
+		
+		Request = Pool->GetFreeRequest();
+		if (Request || InCacheKeys.Num() > UE_HTTPDDC_BATCH_SIZE)
+		{
+			// If the request is too big for existing batches, wait for a free connection and create our own.
+			if (!Request)
+			{
+				Request = Pool->WaitForFreeRequest();
+			}
+
+			FQueuedBatchEntry Entry{
+				InNamespace, 
+				InBucket,
+				CacheKeys,
+				TConstArrayView<TArray<uint8>*>(),
+				FRequest::RequestVerb::Head,
+				&bVerified
+			};
+
+			PerformBatchQuery(Request, TArrayView<FQueuedBatchEntry>(&Entry, 1));
+		}
+		else
+		{
+			Request = QueueBatchRequest(
+				InPool, 
+				InNamespace, 
+				InBucket, 
+				CacheKeys, 
+				TConstArrayView<TArray<uint8>*>(), 
+				bVerified
+			);
+		}
+	}
+
+	~FDataRequestHelper()
+	{
+		if (Request)
+		{
+			Pool->ReleaseRequestToPool(Request);
+		}
+	}
+
+	static void StaticInitialize()
+	{
+		static bool Initialized = false;
+		check(!Initialized);
+		for (int32 i = 0; i < UE_HTTPDDC_BATCH_NUM; i++)
+		{
+			Batches[i].Reserved = 0;
+			Batches[i].Ready = 0;
+			Batches[i].Complete = FPlatformProcess::CreateSynchEvent(true); //todo: leaks
+		}
+		Initialized = true;
+	}
+
+	bool IsSuccess() const
+	{
+		return bVerified[0];
+	}
+
+	const TBitArray<>& IsBatchSuccess() const
+	{
+		return bVerified;
+	}
+
+	int64 GetResponseCode() const
+	{
+		return Request ? Request->GetResponseCode() : 0;
+	}
+
+private:
+
+	struct FQueuedBatchEntry
+	{
+		const TCHAR* Namespace;
+		const TCHAR* Bucket;
+		TConstArrayView<const TCHAR*> CacheKeys;
+		TConstArrayView<TArray<uint8>*> OutDatas;
+		FRequest::RequestVerb Verb;
+		TBitArray<>* bSuccess;
+	};
+
+	struct FBatch
+	{
+		FQueuedBatchEntry Entries[UE_HTTPDDC_BATCH_SIZE];
+		std::atomic<uint32> Reserved;
+		std::atomic<uint32> Ready;
+		std::atomic<uint32> WeightHint;
+		FRequest* Request;
+		FEvent* Complete;
+	};
+
+	FRequest* Request;
+	FRequestPool* Pool;
+	TBitArray<> bVerified;
+	static std::atomic<uint32> FirstAvailableBatch;
+	static TStaticArray<FBatch, UE_HTTPDDC_BATCH_NUM> Batches;
+
+	/**
+	 * Queues up a request to be batched. Blocks until the query is made.
+	 */
+	static FRequest* QueueBatchRequest(FRequestPool* InPool, 
+		const TCHAR* InNamespace, 
+		const TCHAR* InBucket, 
+		TConstArrayView<const TCHAR*> InCacheKeys,
+		TConstArrayView<TArray<uint8>*> OutDatas, 
+		TBitArray<>& bOutVerified)
+	{
+		TRACE_CPUPROFILER_EVENT_SCOPE(HttpDDC_BatchQuery);
+		check(InCacheKeys.Num() == OutDatas.Num() || OutDatas.Num() == 0);
+		const uint32 RequestNum = InCacheKeys.Num();
+		const uint32 RequestWeight = InCacheKeys.Num() * (OutDatas.Num() ? UE_HTTPDDC_BATCH_HEAD_WEIGHT : UE_HTTPDDC_BATCH_GET_WEIGHT);
+
+		for (int32 i = 0; i < Batches.Num(); i++)
+		{
+			uint32 Index = (FirstAvailableBatch.load(std::memory_order_relaxed) + i) % Batches.Num();
+			FBatch& Batch = Batches[Index];
+
+			//Assign different weights to head vs. get queries
+			if (Batch.WeightHint.load(std::memory_order_acquire) + RequestWeight > UE_HTTPDDC_BATCH_WEIGHT_HINT)
+			{
+				continue;
+			}
+
+			// Attempt to reserve a spot in the batch
+			const uint32 Reserve = Batch.Reserved.fetch_add(1, std::memory_order_acquire);
+			if (Reserve >= UE_HTTPDDC_BATCH_SIZE)
+			{
+				// We didn't manage to snag a valid reserve index try next batch
+				continue;
+			}
+
+			// Add our weight to the batch. Note we are treating it as a hint, so don't syncronize.
+			const uint32 ActualWeight = Batch.WeightHint.fetch_add(RequestWeight, std::memory_order_release);
+
+			TAnsiStringBuilder<64> BatchString;
+			BatchString << "HttpDDC_Batch" << Index;
+			TRACE_CPUPROFILER_EVENT_SCOPE_TEXT(*BatchString);
+
+			if (Reserve == (UE_HTTPDDC_BATCH_SIZE - 1))
+			{
+				FirstAvailableBatch++;
+			}
+
+			Batch.Entries[Reserve] = FQueuedBatchEntry{
+				InNamespace,
+				InBucket,
+				InCacheKeys,
+				OutDatas,
+				OutDatas.Num() ? FRequest::RequestVerb::Get : FRequest::RequestVerb::Head,
+				&bOutVerified
+			};
+
+			// Signal we are ready for batch to be submitted
+			Batch.Ready.fetch_add(1u, std::memory_order_release);
+
+			FRequest* Request = nullptr;
+
+			// The first to reserve a slot is the "driver" of the batch
+			if (Reserve == 0)
+			{
+				Batch.Request = InPool->WaitForFreeRequest();
+
+				// Make sure no new requests are added
+				const uint32 Reserved = FMath::Min((uint32)UE_HTTPDDC_BATCH_SIZE, Batch.Reserved.fetch_add(UE_HTTPDDC_BATCH_SIZE, std::memory_order_acquire));
+
+				// Give other threads time to copy their data to batch
+				while (Batch.Ready.load(std::memory_order_acquire) < Reserved) {}
+
+				// Increment request ref count to reflect all waiting threads
+				InPool->MakeRequestShared(Batch.Request, Reserved);
+
+				// Do the actual query and write response to respective target arrays
+				PerformBatchQuery(Batch.Request, TArrayView<FQueuedBatchEntry>(Batch.Entries, Batch.Ready));
+
+				// Signal to waiting threads the batch is complete
+				Batch.Complete->Trigger();
+
+				// Store away the request and wait until other threads have too
+				Request = Batch.Request;
+				while (Batch.Ready.load(std::memory_order_acquire) > 1) {}
+
+				//Reset batch for next use
+				Batch.Complete->Reset();
+				Batch.WeightHint.store(0, std::memory_order_release);
+				Batch.Ready.store(0, std::memory_order_release);
+				Batch.Reserved.store(0, std::memory_order_release);
+			}
+			else
+			{
+				// Wait until "driver" has done query
+				Batch.Complete->Wait(~0);
+
+				// Store away request and signal we are done
+				Request = Batch.Request;
+				Batch.Ready.fetch_sub(1u, std::memory_order_release);
+			}
+
+			return Request;
+		}
+
+		return nullptr;
+	}
+
+
+	/**
+	 * Creates request uri and headers and submits the request
+	 */
+	static void PerformBatchQuery(FRequest* Request, TArrayView<FQueuedBatchEntry> Entries)
+	{
+		TRACE_CPUPROFILER_EVENT_SCOPE(HttpDDC_BatchGet);
+		const TCHAR* Uri(TEXT("api/v1/c/ddc-rpc/batchget"));
+		int64 ResponseCode = 0; uint32 Attempts = 0;
+
+		//Prepare request object
+		TArray<TSharedPtr<FJsonValue>> Operations;
+		for (const FQueuedBatchEntry& Entry : Entries)
+		{
+			for (int32 KeyIdx = 0; KeyIdx < Entry.CacheKeys.Num(); KeyIdx++)
+			{
+				TSharedPtr<FJsonObject> Object = MakeShared<FJsonObject>();
+				Object->SetField(TEXT("bucket"), MakeShared<FJsonValueString>(Entry.Bucket));
+				Object->SetField(TEXT("key"), MakeShared<FJsonValueString>(Entry.CacheKeys[KeyIdx]));
+				if (Entry.Verb == FRequest::RequestVerb::Head)
+				{
+					Object->SetField(TEXT("verb"), MakeShared<FJsonValueString>(TEXT("HEAD")));
+				}
+				Operations.Add(MakeShared<FJsonValueObject>(Object));
+			}
+		}
+		TSharedPtr<FJsonObject> RequestObject = MakeShared<FJsonObject>();
+		RequestObject->SetField(TEXT("namespace"), MakeShared<FJsonValueString>(Entries[0].Namespace));
+		RequestObject->SetField(TEXT("operations"), MakeShared<FJsonValueArray>(Operations));
+
+		//Serialize to a buffer
+		FBufferArchive RequestData;
+		if (FJsonSerializer::Serialize(RequestObject.ToSharedRef(), TJsonWriterFactory<ANSICHAR, TCondensedJsonPrintPolicy<ANSICHAR>>::Create(&RequestData)))
+		{
+			Request->PerformBlockingUpload<FRequest::PostJson>(Uri, MakeArrayView(RequestData));
+			ResponseCode = Request->GetResponseCode();
+
+			if (ResponseCode == 200)
+			{
+				const TArray<uint8>& ResponseBuffer = Request->GetResponseBuffer();
+				const uint8* Response = ResponseBuffer.GetData();
+				const int32 ResponseSize = ResponseBuffer.Num();
+
+				// Parse the response and move the data to the target requests.
+				if (ParseBatchedResponse(Response, ResponseSize, Entries))
+				{
+					UE_LOG(LogDerivedDataCache, VeryVerbose, TEXT("Batch query with %d operations completed."), Entries.Num());
+					return;
+				}
+			}
+		}
+		
+		// If we get here the request failed.
+		UE_LOG(LogDerivedDataCache, Display, TEXT("Batch query failed. Query: %s"), ANSI_TO_TCHAR((ANSICHAR*)RequestData.GetData()));
+
+		// Set all batch operations to failures
+		for (FQueuedBatchEntry Entry : Entries)
+		{
+			Entry.bSuccess->SetRange(0, Entry.CacheKeys.Num(), false);
+		}
+	}
+
+	/**
+	 * Parses a batched response stream, moves the data to target requests and marks them with result.
+	 * @param Response Pointer to Response buffer
+	 * @param ResponseSize Size of response buffer
+	 * @param Requests Requests that will be filled with data.
+	 * @return True if response was successfully parsed, false otherwise.
+	 */
+	static bool ParseBatchedResponse(const uint8* ResponseStart, const int32 ResponseSize, TArrayView<FQueuedBatchEntry> Requests)
+	{
+		// The expected data stream is structured accordingly
+		// {"JPTR"} {PayloadCount:uint32} {{"JPEE"} {Name:cstr} {Result:uint8} {Hash:IoHash} {Size:uint64} {Payload...}} ...
+
+		// Above result value
+		enum OpResult : uint8
+		{
+			Ok = 0,			// Op finished succesfully
+			Error = 1,		// Error during op
+			NotFound = 2,	// Key was not found
+			Exists = 3		// Used to indicate head op success
+		};
+
+		const TCHAR ResponseErrorMessage[] = TEXT("Malformed response from server.");
+		const ANSICHAR* ProtocolMagic = "JPTR";
+		const ANSICHAR* PayloadMagic = "JPEE";
+		const uint32 MagicSize = 4;
+		const uint8* Response = ResponseStart;
+		const uint8* ResponseEnd = Response + ResponseSize;
+
+		// Check that the stream starts with the protocol magic
+		if (FMemory::Memcmp(ProtocolMagic, Response, MagicSize) != 0)
+		{
+			UE_LOG(LogDerivedDataCache, Display, ResponseErrorMessage);
+			return false;
+		}
+		Response += MagicSize;
+
+		// Number of payloads recieved
+		uint32 PayloadCount = *(uint32*)Response;
+		Response += sizeof(uint32);
+
+		uint32 PayloadIdx = 0; 	// Current processed result
+		int32 EntryIdx = 0; 	// Current Entry index
+		int32 KeyIdx = 0; 		// Current Key index for current Entry
+
+		while (Response < ResponseEnd && FMemory::Memcmp(PayloadMagic, Response, MagicSize) == 0)
+		{
+			PayloadIdx++;
+			Response += MagicSize;
+
+			const ANSICHAR* PayloadNameA = (const ANSICHAR*)Response;
+			Response += FCStringAnsi::Strlen(PayloadNameA) + 1; //String and zero termination
+			const ANSICHAR* CacheKeyA = FCStringAnsi::Strrchr(PayloadNameA, '.') + 1; // "namespace.bucket.cachekey"
+
+			// Find the payload among the requests. Assume payloads are recieved in the same order they
+			// were requested.
+			FUTF8ToTCHAR CacheKey(CacheKeyA);
+			if (Requests.Num() < EntryIdx || Requests[EntryIdx].CacheKeys.Num() < KeyIdx || FCString::Stricmp(Requests[EntryIdx].CacheKeys[KeyIdx], CacheKey.Get()) != 0)
+			{
+				UE_LOG(LogDerivedDataCache, Error, ResponseErrorMessage);
+				return false;
+			}
+
+			FQueuedBatchEntry& RequestOp = Requests[EntryIdx];
+			TBitArray<>& bSuccess = *RequestOp.bSuccess;
+
+			// Result of the operation
+			uint8 PayloadResult = *Response;
+			Response += sizeof(uint8);
+
+			switch (PayloadResult)
+			{
+			
+			case OpResult::Ok:
+				{
+					TArray<uint8>* OutData = RequestOp.OutDatas[KeyIdx];
+
+					// Payload hash of the following payload data
+					FIoHash PayloadHash = *(FIoHash*)Response;
+					Response += sizeof(FIoHash);
+
+					// Size of the following payload data
+					const uint64 PayloadSize = *(uint64*)Response;
+					Response += sizeof(uint64);
+
+					if (PayloadSize > 0)
+					{
+						OutData->Append(Response, PayloadSize);
+						Response += PayloadSize;
+						// Verify the recieved and parsed payload
+						if (VerifyPayload(PayloadHash, *OutData))
+						{
+							TRACE_COUNTER_ADD(HttpDDC_GetHit, int64(1));
+							TRACE_COUNTER_ADD(HttpDDC_BytesReceived, int64(PayloadSize));
+							
+							bSuccess[KeyIdx] = true;
+						}
+						else
+						{
+							OutData->Empty();
+							bSuccess[KeyIdx] = false;
+						}
+					}
+					else
+					{
+						bSuccess[KeyIdx] = false;
+					}
+				}
+				break;
+
+			case OpResult::Exists:
+				{
+					TRACE_COUNTER_ADD(HttpDDC_ExistHit, int64(1));
+					bSuccess[KeyIdx] = true;
+				}
+				break;
+
+			default:
+			case OpResult::Error:
+				UE_LOG(LogDerivedDataCache, Display, TEXT("Server error while getting %s"), CacheKey.Get());
+				// intentional falltrough
+
+			case OpResult::NotFound:
+				bSuccess[KeyIdx] = false;
+				break;
+
+			}
+
+
+			// Advance the request indices
+			if (++KeyIdx >= RequestOp.CacheKeys.Num())
+			{
+				EntryIdx++;
+				KeyIdx = 0;
+			}
+			
+		}
+
+		// Have we parsed all the payloads from the message?
+		if (PayloadIdx != PayloadCount)
+		{
+			UE_LOG(LogDerivedDataCache, Display, TEXT("%s: Found %d payloads but %d was reported."), ResponseErrorMessage, PayloadIdx, PayloadCount);
+		}
+
+		return true;
+	}
+};
+
+TStaticArray<FDataRequestHelper::FBatch, UE_HTTPDDC_BATCH_NUM> FDataRequestHelper::Batches;
+std::atomic<uint32> FDataRequestHelper::FirstAvailableBatch;
+
+//----------------------------------------------------------------------------------------------------------
+// FDataUploadHelper
+//----------------------------------------------------------------------------------------------------------
+struct FDataUploadHelper
+{
+	FDataUploadHelper(FRequestPool* InPool, 
+		const TCHAR* InNamespace, 
+		const TCHAR* InBucket, 
+		const TCHAR* InCacheKey, 
+		const TArrayView<const uint8>& InData,
+		FDerivedDataCacheUsageStats& InUsageStats)
+		: ResponseCode(0)
+		, bSuccess(false)
+		, bQueued(false)
+	{
+		FRequest* Request = InPool->GetFreeRequest();
+		if (Request)
+		{
+			ResponseCode = PerformPut(Request, InNamespace, InBucket, InCacheKey, InData, InUsageStats);
+			bSuccess = FRequest::IsSuccessResponse(Request->GetResponseCode());
+
+			ProcessQueuedPutsAndReleaseRequest(InPool, Request, InUsageStats);
+		}
+		else
+		{
+			FQueuedEntry* Entry = new FQueuedEntry(InNamespace, InBucket, InCacheKey, InData);
+			QueuedPuts.Push(Entry);
+			bSuccess = true;
+			bQueued = true;
+			
+			// A request may have been released while the entry was being queued.
+			Request = InPool->GetFreeRequest();
+			if (Request)
+			{
+				ProcessQueuedPutsAndReleaseRequest(InPool, Request, InUsageStats);
+			}
+		}
+	}
+
+	bool IsSuccess() const
+	{
+		return bSuccess;
+	}
+
+	int64 GetResponseCode() const
+	{
+		return ResponseCode;
+	}
+
+	bool IsQueued() const
+	{
+		return bQueued;
+	}
+
+private:
+
+	struct FQueuedEntry
+	{
+		FString Namespace;
+		FString Bucket;
+		FString CacheKey;
+		TArray<uint8> Data;
+
+		FQueuedEntry(const TCHAR* InNamespace, const TCHAR* InBucket, const TCHAR* InCacheKey, const TArrayView<const uint8> InData)
+			: Namespace(InNamespace)
+			, Bucket(InBucket)
+			, CacheKey(InCacheKey)
+			, Data(InData) // Copies the data!
+		{}
+	};
+
+	static TLockFreePointerListUnordered<FQueuedEntry, PLATFORM_CACHE_LINE_SIZE> QueuedPuts;
+
+	int64 ResponseCode;
+	bool bSuccess;
+	bool bQueued;
+
+	static void ProcessQueuedPutsAndReleaseRequest(FRequestPool* Pool, FRequest* Request, FDerivedDataCacheUsageStats& UsageStats)
+	{
+		while (Request)
+		{
+			while (FQueuedEntry* Entry = QueuedPuts.Pop())
+			{
+				Request->Reset();
+				PerformPut(Request, *Entry->Namespace, *Entry->Bucket, *Entry->CacheKey, Entry->Data, UsageStats);
+				delete Entry;
+			}
+
+			Pool->ReleaseRequestToPool(Request);
+
+			// An entry may have been queued while the request was being released.
+			if (QueuedPuts.IsEmpty())
+			{
+				break;
+			}
+
+			// Process the queue again if a request is free, otherwise the thread that got the request will process it.
+			Request = Pool->GetFreeRequest();
+		}
+	}
+
+	static int64 PerformPut(FRequest* Request, const TCHAR* Namespace, const TCHAR* Bucket, const TCHAR* CacheKey, const TArrayView<const uint8> Data, FDerivedDataCacheUsageStats& UsageStats)
+	{
+		COOK_STAT(auto Timer = UsageStats.TimePut());
+
+		HashPayload(Request, Data);
+
+		TStringBuilder<256> Uri;
+		Uri.Appendf(TEXT("api/v1/c/ddc/%s/%s/%s"), Namespace, Bucket, CacheKey);
+
+		Request->PerformBlockingUpload<FRequest::Put>(*Uri, Data);
+
+		const int64 ResponseCode = Request->GetResponseCode();
+		if (FRequest::IsSuccessResponse(ResponseCode))
+		{
+			TRACE_COUNTER_ADD(HttpDDC_BytesSent, int64(Request->GetBytesSent()));
+			COOK_STAT(Timer.AddHit(Request->GetBytesSent()));
+		}
+
+		return Request->GetResponseCode();
+	}
+};
+
+TLockFreePointerListUnordered<FDataUploadHelper::FQueuedEntry, PLATFORM_CACHE_LINE_SIZE> FDataUploadHelper::QueuedPuts;
+
+#endif // WITH_DATAREQUEST_HELPER
 
 //----------------------------------------------------------------------------------------------------------
 // Certificate checking
@@ -871,14 +1547,13 @@ static CURLcode sslctx_function(CURL * curl, void * sslctx, void * parm)
 // Content parsing and checking
 //----------------------------------------------------------------------------------------------------------
 
-
 /**
  * Verifies the integrity of the received data using supplied checksum.
  * @param Hash received hash value.
  * @param Payload Payload received.
  * @return True if the data is correct, false if checksums doesn't match.
  */
-bool VerifyPayload(FSHAHash Hash, const TArray<uint8>& Payload)
+bool VerifyPayload(const FSHAHash& Hash, const TArray<uint8>& Payload)
 {
 	FSHAHash PayloadHash;
 	FSHA1::HashBuffer(Payload.GetData(), Payload.Num(), PayloadHash.Hash);
@@ -903,7 +1578,7 @@ bool VerifyPayload(FSHAHash Hash, const TArray<uint8>& Payload)
  * @param Payload Payload received.
  * @return True if the data is correct, false if checksums doesn't match.
  */
-bool VerifyPayload(FIoHash Hash, const TArray<uint8>& Payload)
+bool VerifyPayload(const FIoHash& Hash, const TArray<uint8>& Payload)
 {
 	FIoHash PayloadHash = FIoHash::HashBuffer(Payload.GetData(), Payload.Num());
 
@@ -930,17 +1605,17 @@ bool VerifyPayload(FIoHash Hash, const TArray<uint8>& Payload)
  */
 bool VerifyRequest(const FRequest* Request, const TArray<uint8>& Payload)
 {
-	FString receivedHashStr;
-	if (Request->GetHeader("X-Jupiter-Sha1", receivedHashStr))
+	FString ReceivedHashStr;
+	if (Request->GetHeader("X-Jupiter-Sha1", ReceivedHashStr))
 	{
-		FSHAHash receivedHash;
-		receivedHash.FromString(receivedHashStr);
-		return VerifyPayload(receivedHash, Payload);
+		FSHAHash ReceivedHash;
+		ReceivedHash.FromString(ReceivedHashStr);
+		return VerifyPayload(ReceivedHash, Payload);
 	}
-	if (Request->GetHeader("X-Jupiter-IoHash", receivedHashStr))
+	if (Request->GetHeader("X-Jupiter-IoHash", ReceivedHashStr))
 	{
-		FIoHash receivedHash(receivedHashStr);
-		return VerifyPayload(receivedHash, Payload);
+		FIoHash ReceivedHash(ReceivedHashStr);
+		return VerifyPayload(ReceivedHash, Payload);
 	}
 	UE_LOG(LogDerivedDataCache, Warning, TEXT("HTTP server did not send a content hash. Wrong server version?"));
 	return true;
@@ -1006,9 +1681,13 @@ FHttpDerivedDataBackend::FHttpDerivedDataBackend(
 	, bIsUsable(false)
 	, FailedLoginAttempts(0)
 {
+#if WITH_DATAREQUEST_HELPER
+	FDataRequestHelper::StaticInitialize();
+#endif
 	if (IsServiceReady() && AcquireAccessToken())
 	{
-		RequestPool = MakeUnique<FRequestPool>(InServiceUrl, Access.Get());
+		GetRequestPool = MakeUnique<FRequestPool>(InServiceUrl, Access.Get(), UE_HTTPDDC_GET_REQUEST_POOL_SIZE);
+		PutRequestPool = MakeUnique<FRequestPool>(InServiceUrl, Access.Get(), UE_HTTPDDC_PUT_REQUEST_POOL_SIZE);
 		bIsUsable = true;
 	}
 }
@@ -1190,13 +1869,32 @@ bool FHttpDerivedDataBackend::CachedDataProbablyExists(const TCHAR* CacheKey)
 	TRACE_CPUPROFILER_EVENT_SCOPE(HttpDDC_Exist);
 	TRACE_COUNTER_ADD(HttpDDC_Exist, int64(1));
 	COOK_STAT(auto Timer = UsageStats.TimeProbablyExists());
-	
+
+#if WITH_DATAREQUEST_HELPER
+	// Retry request until we get an accepted response or exhaust allowed number of attempts.
+	for (int32 Attempts = 0; Attempts < UE_HTTPDDC_MAX_ATTEMPTS; ++Attempts)
+	{
+		FDataRequestHelper RequestHelper(GetRequestPool.Get(), *Namespace, *DefaultBucket, CacheKey, nullptr);
+		const int64 ResponseCode = RequestHelper.GetResponseCode();
+
+		if (FRequest::IsSuccessResponse(ResponseCode) && RequestHelper.IsSuccess())
+		{
+			COOK_STAT(Timer.AddHit(0));
+			return true;
+		}
+
+		if (!ShouldRetryOnError(ResponseCode))
+		{
+			return false;
+		}
+	}
+#else
 	FString Uri = FString::Printf(TEXT("api/v1/c/ddc/%s/%s/%s"), *Namespace, *DefaultBucket, CacheKey);
 
 	// Retry request until we get an accepted response or exhaust allowed number of attempts.
 	for (int32 Attempts = 0; Attempts < UE_HTTPDDC_MAX_ATTEMPTS; ++Attempts)
 	{
-		FScopedRequestPtr Request(RequestPool.Get());
+		FScopedRequestPtr Request(GetRequestPool.Get());
 		const FRequest::Result Result = Request->PerformBlockingQuery<FRequest::Head>(*Uri);
 		const int64 ResponseCode = Request->GetResponseCode();
 
@@ -1216,6 +1914,7 @@ bool FHttpDerivedDataBackend::CachedDataProbablyExists(const TCHAR* CacheKey)
 			break;
 		}
 	}
+#endif
 
 	return false;
 }
@@ -1225,7 +1924,24 @@ TBitArray<> FHttpDerivedDataBackend::CachedDataProbablyExistsBatch(TConstArrayVi
 	TRACE_CPUPROFILER_EVENT_SCOPE(HttpDDC_Exist);
 	TRACE_COUNTER_ADD(HttpDDC_Exist, int64(1));
 	COOK_STAT(auto Timer = UsageStats.TimeProbablyExists());
+#if WITH_DATAREQUEST_HELPER
+	for (int32 Attempts = 0; Attempts < UE_HTTPDDC_MAX_ATTEMPTS; ++Attempts)
+	{
+		FDataRequestHelper RequestHelper(GetRequestPool.Get(), *Namespace, *DefaultBucket, CacheKeys);
+		const int64 ResponseCode = RequestHelper.GetResponseCode();
 
+		if (FRequest::IsSuccessResponse(ResponseCode) && RequestHelper.IsSuccess())
+		{
+			COOK_STAT(Timer.AddHit(0));
+			return RequestHelper.IsBatchSuccess();
+		}
+
+		if (!ShouldRetryOnError(ResponseCode))
+		{
+			return RequestHelper.IsBatchSuccess();
+		}
+	}
+#else
 	const TCHAR* const Uri = TEXT("api/v1/c/ddc-rpc");
 
 	TAnsiStringBuilder<512> Body;
@@ -1245,7 +1961,7 @@ TBitArray<> FHttpDerivedDataBackend::CachedDataProbablyExistsBatch(TConstArrayVi
 	// Retry request until we get an accepted response or exhaust allowed number of attempts.
 	for (int32 Attempts = 0; Attempts < UE_HTTPDDC_MAX_ATTEMPTS; ++Attempts)
 	{
-		FScopedRequestPtr Request(RequestPool.Get());
+		FScopedRequestPtr Request(GetRequestPool.Get());
 		const FRequest::Result Result = Request->PerformBlockingUpload<FRequest::PostJson>(Uri, BodyView);
 		const int64 ResponseCode = Request->GetResponseCode();
 
@@ -1278,6 +1994,7 @@ TBitArray<> FHttpDerivedDataBackend::CachedDataProbablyExistsBatch(TConstArrayVi
 			break;
 		}
 	}
+#endif
 
 	return TBitArray<>(false, CacheKeys.Num());
 }
@@ -1288,20 +2005,39 @@ bool FHttpDerivedDataBackend::GetCachedData(const TCHAR* CacheKey, TArray<uint8>
 	TRACE_COUNTER_ADD(HttpDDC_Get, int64(1));
 	COOK_STAT(auto Timer = UsageStats.TimeGet());
 
+#if WITH_DATAREQUEST_HELPER
+	// Retry request until we get an accepted response or exhaust allowed number of attempts.
+	for (int32 Attempts = 0; Attempts < UE_HTTPDDC_MAX_ATTEMPTS; ++Attempts)
+	{
+		FDataRequestHelper RequestHelper(GetRequestPool.Get(), *Namespace, *DefaultBucket, CacheKey, &OutData);
+		const int64 ResponseCode = RequestHelper.GetResponseCode();
+
+		if (FRequest::IsSuccessResponse(ResponseCode) && RequestHelper.IsSuccess())
+		{
+			COOK_STAT(Timer.AddHit(OutData.Num()));
+			check(OutData.Num() > 0);
+			return true;
+		}
+
+		if (!ShouldRetryOnError(ResponseCode))
+		{
+			return false;
+		}
+	}
+#else 
 	FString Uri = FString::Printf(TEXT("api/v1/c/ddc/%s/%s/%s.raw"), *Namespace, *DefaultBucket, CacheKey);
-	int64 ResponseCode = 0; uint32 Attempts = 0;
 
 	// Retry request until we get an accepted response or exhaust allowed number of attempts.
-	while (ResponseCode == 0 && ++Attempts < UE_HTTPDDC_MAX_ATTEMPTS)
+	for (int32 Attempts = 0; Attempts < UE_HTTPDDC_MAX_ATTEMPTS; ++Attempts)
 	{
-		FScopedRequestPtr Request(RequestPool.Get());
+		FScopedRequestPtr Request(GetRequestPool.Get());
 		if (Request.IsValid())
 		{
 			FRequest::Result Result = Request->PerformBlockingDownload(*Uri, &OutData);
-			ResponseCode = Request->GetResponseCode();
+			const uint64 ResponseCode = Request->GetResponseCode();
 
 			// Request was successful, make sure we got all the expected data.
-			if (ResponseCode == 200 && VerifyRequest(Request.Get(), OutData))
+			if (FRequest::IsSuccessResponse(ResponseCode) && VerifyRequest(Request.Get(), OutData))
 			{
 				TRACE_COUNTER_ADD(HttpDDC_GetHit, int64(1));
 				TRACE_COUNTER_ADD(HttpDDC_BytesReceived, int64(Request->GetBytesReceived()));
@@ -1313,10 +2049,9 @@ bool FHttpDerivedDataBackend::GetCachedData(const TCHAR* CacheKey, TArray<uint8>
 			{
 				return false;
 			}
-
-			ResponseCode = 0;
 		}
 	}
+#endif
 
 	return false;
 }
@@ -1324,6 +2059,24 @@ bool FHttpDerivedDataBackend::GetCachedData(const TCHAR* CacheKey, TArray<uint8>
 FDerivedDataBackendInterface::EPutStatus FHttpDerivedDataBackend::PutCachedData(const TCHAR* CacheKey, TArrayView<const uint8> InData, bool bPutEvenIfExists)
 {
 	TRACE_CPUPROFILER_EVENT_SCOPE(HttpDDC_Put);
+
+#if WITH_DATAREQUEST_HELPER
+	for (int32 Attempts = 0; Attempts < UE_HTTPDDC_MAX_ATTEMPTS; ++Attempts)
+	{
+		FDataUploadHelper Request(PutRequestPool.Get(), *Namespace, *DefaultBucket,	CacheKey, InData, UsageStats);
+		const int64 ResponseCode = Request.GetResponseCode();
+
+		if (Request.IsSuccess() && (Request.IsQueued() || FRequest::IsSuccessResponse(ResponseCode)))
+		{
+			return Request.IsQueued() ? EPutStatus::Executing : EPutStatus::Cached;
+		}
+
+		if (!ShouldRetryOnError(ResponseCode))
+		{
+			return EPutStatus::NotCached;
+		}
+	}
+#else
 	COOK_STAT(auto Timer = UsageStats.TimePut());
 
 	FString Uri = FString::Printf(TEXT("api/v1/c/ddc/%s/%s/%s"), *Namespace, *DefaultBucket, CacheKey);
@@ -1332,7 +2085,7 @@ FDerivedDataBackendInterface::EPutStatus FHttpDerivedDataBackend::PutCachedData(
 	// Retry request until we get an accepted response or exhaust allowed number of attempts.
 	while (ResponseCode == 0 && ++Attempts < UE_HTTPDDC_MAX_ATTEMPTS)
 	{
-		FScopedRequestPtr Request(RequestPool.Get());
+		FScopedRequestPtr Request(PutRequestPool.Get());
 		if (Request.IsValid())
 		{
 			// Append the content hash to the header
@@ -1356,6 +2109,7 @@ FDerivedDataBackendInterface::EPutStatus FHttpDerivedDataBackend::PutCachedData(
 			ResponseCode = 0;
 		}
 	}
+#endif // WITH_DATAREQUEST_HELPER
 
 	return EPutStatus::NotCached;
 }
@@ -1373,7 +2127,7 @@ void FHttpDerivedDataBackend::RemoveCachedData(const TCHAR* CacheKey, bool bTran
 	// Retry request until we get an accepted response or exhaust allowed number of attempts.
 	while (ResponseCode == 0 && ++Attempts < UE_HTTPDDC_MAX_ATTEMPTS)
 	{
-		FScopedRequestPtr Request(RequestPool.Get());
+		FScopedRequestPtr Request(PutRequestPool.Get());
 		if (Request.IsValid())
 		{
 			FRequest::Result Result = Request->PerformBlockingQuery<FRequest::Delete>(*Uri);
diff --git a/Engine/Source/Developer/DerivedDataCache/Private/HttpDerivedDataBackend.h b/Engine/Source/Developer/DerivedDataCache/Private/HttpDerivedDataBackend.h
--- a/Engine/Source/Developer/DerivedDataCache/Private/HttpDerivedDataBackend.h
+++ b/Engine/Source/Developer/DerivedDataCache/Private/HttpDerivedDataBackend.h
@@ -73,7 +73,8 @@ private:
 	FString OAuthSecret;
 	FCriticalSection AccessCs;
 	FDerivedDataCacheUsageStats UsageStats;
-	TUniquePtr<struct FRequestPool> RequestPool;
+	TUniquePtr<struct FRequestPool> GetRequestPool;
+	TUniquePtr<struct FRequestPool> PutRequestPool;
 	TUniquePtr<struct FHttpAccessToken> Access;
 	bool bIsUsable;
 	uint32 FailedLoginAttempts;
