DDC: Refactored the new Cache API to support its use through only an interface dependency

#rnx

diff --git a/Engine/Source/Developer/DerivedDataCache/Private/DerivedDataCache.cpp b/Engine/Source/Developer/DerivedDataCache/Private/DerivedDataCache.cpp
--- a/Engine/Source/Developer/DerivedDataCache/Private/DerivedDataCache.cpp
+++ b/Engine/Source/Developer/DerivedDataCache/Private/DerivedDataCache.cpp
@@ -17,13 +17,13 @@
 #include "Modules/ModuleManager.h"
 
 #include "DerivedDataBackendInterface.h"
+#include "DerivedDataCachePrivate.h"
 #include "DerivedDataPluginInterface.h"
 #include "DDCCleanup.h"
 #include "ProfilingDebugging/CookStats.h"
 
 #include "Algo/AllOf.h"
 #include "Algo/Transform.h"
-#include "DerivedDataBackendInterface.h"
 #include "Misc/CoreMisc.h"
 #include "Misc/ScopeExit.h"
 #include "Misc/StringBuilder.h"
@@ -723,6 +723,10 @@ class FCache final : public ICache
 public:
 	virtual ~FCache() = default;
 
+	virtual FCacheBucket CreateBucket(FStringView Name) final { return CreateCacheBucket(Name); }
+
+	virtual FCacheRecordBuilder CreateRecord(const FCacheKey& Key) final { return CreateCacheRecordBuilder(Key); }
+
 	virtual FRequest Put(
 		TArrayView<FCacheRecord> Records,
 		FStringView Context,
@@ -768,37 +772,19 @@ private:
 		EPriority Priority,
 		FOnCacheGetPayloadComplete& Callback);
 
-	template <int32 BufferSize>
-	class TToString
-	{
-	public:
-		template <typename T>
-		explicit TToString(const T& Input)
-		{
-			Buffer << Input;
-		}
-
-		inline const TCHAR* operator*() const { return Buffer.ToString(); }
-		inline const TCHAR* ToString() const { return Buffer.ToString(); }
-		inline operator const TCHAR*() const { return Buffer.ToString(); }
-
-	private:
-		TStringBuilder<BufferSize> Buffer;
-	};
-
 	static FString MakeRecordKey(const FCacheKey& Key)
 	{
-		check(Key.IsValid());
+		check(Key.Bucket.IsValid() && !Key.Hash.IsZero());
 		TStringBuilder<96> Out;
-		Out << Key.GetBucket() << TEXT("_") << Key.GetHash();
+		Out << Key.Bucket << TEXT("_") << Key.Hash;
 		return FDerivedDataCacheInterface::SanitizeCacheKey(Out.ToString());
 	}
 
 	static FString MakePayloadKey(const FCachePayloadKey& PayloadKey)
 	{
-		check(PayloadKey.IsValid());
+		check(PayloadKey.Record.Bucket.IsValid() && !PayloadKey.Record.Hash.IsZero() && PayloadKey.Payload.IsValid());
 		TStringBuilder<96> Out;
-		Out << PayloadKey.GetKey().GetBucket() << TEXT("_") << PayloadKey.GetKey().GetHash() << TEXT("_") << PayloadKey.GetId();
+		Out << PayloadKey.Record.Bucket << TEXT("_") << PayloadKey.Record.Hash << TEXT("_") << PayloadKey.Payload;
 		return FDerivedDataCacheInterface::SanitizeCacheKey(Out.ToString());
 	}
 };
@@ -826,7 +812,7 @@ void FCache::Put(
 	if (!EnumHasAnyFlags(Policy, ECachePolicy::Store))
 	{
 		UE_LOG(LogDerivedDataCache, VeryVerbose, TEXT("Cache: Put skipped for %s from '%.*s'"),
-			*TToString<96>(Params.Key), Context.Len(), Context.GetData());
+			*WriteToString<96>(Params.Key), Context.Len(), Context.GetData());
 		return;
 	}
 
@@ -837,7 +823,7 @@ void FCache::Put(
 		TConstArrayView<uint8> BufferView = MakeArrayView(
 			static_cast<const uint8*>(Buffer.GetData()),
 			static_cast<int32>(Buffer.GetSize()));
-		GetDerivedDataCacheRef().Put(*MakePayloadKey(FCachePayloadKey(Key, Payload.GetId())), BufferView, Context);
+		GetDerivedDataCacheRef().Put(*MakePayloadKey(FCachePayloadKey{Key, Payload.GetId()}), BufferView, Context);
 	};
 
 	FCbWriter Writer;
@@ -847,7 +833,7 @@ void FCache::Put(
 		{
 			PutPayload(Value);
 			Writer.BeginObject("Value"_ASV);
-			Writer.AddObjectId("Id"_ASV, Value.GetId().ToObjectId());
+			Writer.AddObjectId("Id"_ASV, FCbObjectId(Value.GetId().GetView()));
 			Writer.AddBinaryAttachment("Hash"_ASV, Value.GetHash());
 			Writer.EndObject();
 		}
@@ -858,7 +844,7 @@ void FCache::Put(
 			{
 				PutPayload(Attachment);
 				Writer.BeginObject();
-				Writer.AddObjectId("Id"_ASV, Attachment.GetId().ToObjectId());
+				Writer.AddObjectId("Id"_ASV, FCbObjectId(Attachment.GetId().GetView()));
 				Writer.AddBinaryAttachment("Hash"_ASV, Attachment.GetHash());
 				Writer.EndObject();
 			}
@@ -882,7 +868,7 @@ void FCache::Put(
 	GetDerivedDataCacheRef().Put(*MakeRecordKey(Params.Key), Data, Context);
 
 	UE_LOG(LogDerivedDataCache, VeryVerbose, TEXT("Cache: Put for %s from '%.*s'"),
-		*TToString<96>(Params.Key), Context.Len(), Context.GetData());
+		*WriteToString<96>(Params.Key), Context.Len(), Context.GetData());
 	Params.Status = EStatus::Ok;
 }
 
@@ -910,8 +896,7 @@ void FCache::Get(
 	FPayload Value;
 	TArray<FPayload> Attachments;
 	TArray<FPayloadId, TInlineAllocator<1>> SkippedPayloads;
-	FCacheRecordBuilder Builder;
-	Builder.SetKey(Key);
+	FCacheRecordBuilder Builder = CreateRecord(Key);
 
 	FCacheGetCompleteParams Params;
 	ON_SCOPE_EXIT
@@ -927,7 +912,7 @@ void FCache::Get(
 	if (!EnumHasAnyFlags(Policy, ECachePolicy::Query))
 	{
 		UE_LOG(LogDerivedDataCache, VeryVerbose, TEXT("Cache: Get skipped for %s from '%.*s'"),
-			*TToString<96>(Key), Context.Len(), Context.GetData());
+			*WriteToString<96>(Key), Context.Len(), Context.GetData());
 		return;
 	}
 
@@ -938,7 +923,7 @@ void FCache::Get(
 		if (!GetDerivedDataCacheRef().GetSynchronous(*MakeRecordKey(Key), RecordData, Context))
 		{
 			UE_LOG(LogDerivedDataCache, Verbose, TEXT("Cache: Get cache miss for %s from '%.*s'"),
-				*TToString<96>(Key), Context.Len(), Context.GetData());
+				*WriteToString<96>(Key), Context.Len(), Context.GetData());
 			return;
 		}
 		RecordBuffer = MakeSharedBufferFromArray(MoveTemp(RecordData));
@@ -948,13 +933,13 @@ void FCache::Get(
 	if (ValidateCompactBinaryRange(RecordBuffer, ECbValidateMode::Default) != ECbValidateError::None)
 	{
 		UE_LOG(LogDerivedDataCache, Verbose, TEXT("Cache: Get cache miss with corrupted record for %s from '%.*s'"),
-			*TToString<96>(Key), Context.Len(), Context.GetData());
+			*WriteToString<96>(Key), Context.Len(), Context.GetData());
 		return;
 	}
 
 	// Read the record from its compact binary fields.
 	const FCbObject RecordObject(MoveTemp(RecordBuffer));
-	const FCbObject MetaObject = RecordObject.Find("Meta"_ASV).AsObject();
+	FCbObject MetaObject = RecordObject.Find("Meta"_ASV).AsObject();
 	const FCbObjectView ValueObject = RecordObject["Value"_ASV].AsObjectView();
 	const FCbArrayView AttachmentsArray = RecordObject["Attachments"_ASV].AsArrayView();
 
@@ -962,9 +947,9 @@ void FCache::Get(
 	const auto GetPayload = [&SkippedPayloads, &Key, Policy, Context](FCbObjectView PayloadObject, ECachePolicy SkipFlag) -> FPayload
 	{
 		TArray<uint8> PayloadData;
-		const FPayloadId PayloadId = FPayloadId::FromObjectId(PayloadObject["Id"_ASV].AsObjectId());
+		const FPayloadId PayloadId = FPayloadId(PayloadObject["Id"_ASV].AsObjectId().GetView());
 		const FIoHash PayloadHash = PayloadObject["Hash"_ASV].AsBinaryAttachment();
-		const FCachePayloadKey PayloadKey(Key, PayloadId);
+		const FCachePayloadKey PayloadKey{Key, PayloadId};
 		if (EnumHasAnyFlags(Policy, SkipFlag))
 		{
 			SkippedPayloads.Add(PayloadId);
@@ -979,14 +964,14 @@ void FCache::Get(
 			else
 			{
 				UE_LOG(LogDerivedDataCache, Verbose, TEXT("Cache: Get cache miss with corrupted content for %s from '%.*s'"),
-					*TToString<96>(PayloadKey), Context.Len(), Context.GetData());
+					*WriteToString<96>(PayloadKey), Context.Len(), Context.GetData());
 				return FPayload();
 			}
 		}
 		else
 		{
 			UE_LOG(LogDerivedDataCache, Verbose, TEXT("Cache: Get cache miss with missing content for %s from '%.*s'"),
-				*TToString<96>(PayloadKey), Context.Len(), Context.GetData());
+				*WriteToString<96>(PayloadKey), Context.Len(), Context.GetData());
 			return FPayload();
 		}
 	};
@@ -1012,14 +997,14 @@ void FCache::Get(
 		TArray<FString, TInlineAllocator<1>> SkippedPayloadKeys;
 		SkippedPayloadKeys.Reserve(SkippedPayloads.Num());
 		Algo::Transform(SkippedPayloads, SkippedPayloadKeys,
-			[&Key](const FPayloadId& Id) -> FString { return MakePayloadKey(FCachePayloadKey(Key, Id)); });
+			[&Key](const FPayloadId& Id) -> FString { return MakePayloadKey(FCachePayloadKey{Key, Id}); });
 		const bool bPayloadsExist = EnumHasAnyFlags(Policy, ECachePolicy::StoreLocal)
 			? GetDerivedDataCacheRef().TryToPrefetch(SkippedPayloadKeys, Context)
 			: GetDerivedDataCacheRef().AllCachedDataProbablyExists(SkippedPayloadKeys);
 		if (!bPayloadsExist)
 		{
 			UE_LOG(LogDerivedDataCache, Verbose, TEXT("Cache: Get cache miss with missing content for %s from '%.*s'"),
-				*TToString<96>(Key), Context.Len(), Context.GetData());
+				*WriteToString<96>(Key), Context.Len(), Context.GetData());
 			return;
 		}
 	}
@@ -1027,7 +1012,7 @@ void FCache::Get(
 	// Populate the builder now that any error conditions have been handled.
 	if (!EnumHasAnyFlags(Policy, ECachePolicy::SkipMeta) && MetaObject.CreateViewIterator())
 	{
-		Builder.SetMeta(MetaObject);
+		Builder.SetMeta(MoveTemp(MetaObject));
 	}
 	if (Value)
 	{
@@ -1039,7 +1024,7 @@ void FCache::Get(
 	}
 
 	UE_LOG(LogDerivedDataCache, Verbose, TEXT("Cache: Get cache hit for %s from '%.*s'"),
-		*TToString<96>(Key), Context.Len(), Context.GetData());
+		*WriteToString<96>(Key), Context.Len(), Context.GetData());
 	Params.Status = EStatus::Ok;
 }
 
@@ -1065,7 +1050,7 @@ void FCache::GetPayload(
 	FOnCacheGetPayloadComplete& Callback)
 {
 	FCacheGetPayloadCompleteParams Params;
-	Params.Key = Key.GetKey();
+	Params.Key = Key.Record;
 
 	ON_SCOPE_EXIT
 	{
@@ -1073,7 +1058,7 @@ void FCache::GetPayload(
 		{
 			if (!Params.Payload)
 			{
-				Params.Payload = FPayload(Key.GetId(), FIoHash());
+				Params.Payload = FPayload(Key.Payload, FIoHash());
 			}
 			Callback(MoveTemp(Params));
 		}
@@ -1083,7 +1068,7 @@ void FCache::GetPayload(
 	if (!EnumHasAnyFlags(Policy, ECachePolicy::Query))
 	{
 		UE_LOG(LogDerivedDataCache, VeryVerbose, TEXT("Cache: GetPayload skipped %s from '%.*s'"),
-			*TToString<96>(Key), Context.Len(), Context.GetData());
+			*WriteToString<96>(Key), Context.Len(), Context.GetData());
 		return;
 	}
 
@@ -1099,10 +1084,10 @@ void FCache::GetPayload(
 			FSharedBuffer RecordBuffer;
 			{
 				TArray<uint8> RecordData;
-				if (!GetDerivedDataCacheRef().GetSynchronous(*MakeRecordKey(Key.GetKey()), RecordData, Context))
+				if (!GetDerivedDataCacheRef().GetSynchronous(*MakeRecordKey(Key.Record), RecordData, Context))
 				{
 					UE_LOG(LogDerivedDataCache, Verbose, TEXT("Cache: Get cache miss with missing record for %s from '%.*s'"),
-						*TToString<96>(Key), Context.Len(), Context.GetData());
+						*WriteToString<96>(Key), Context.Len(), Context.GetData());
 					return;
 				}
 				RecordBuffer = MakeSharedBufferFromArray(MoveTemp(RecordData));
@@ -1112,7 +1097,7 @@ void FCache::GetPayload(
 			if (ValidateCompactBinaryRange(RecordBuffer, ECbValidateMode::Default) != ECbValidateError::None)
 			{
 				UE_LOG(LogDerivedDataCache, Verbose, TEXT("Cache: Get cache miss with corrupted record for %s from '%.*s'"),
-					*TToString<96>(Key), Context.Len(), Context.GetData());
+					*WriteToString<96>(Key), Context.Len(), Context.GetData());
 				return;
 			}
 
@@ -1122,10 +1107,10 @@ void FCache::GetPayload(
 			const FCbArrayView AttachmentsArray = RecordObject["Attachments"_ASV].AsArrayView();
 
 			// Find the hash of the payload in the record.
-			const FCbObjectId KeyObjectId = Key.GetId().ToObjectId();
+			const FCbObjectId KeyObjectId(Key.Payload.GetView());
 			if (ValueObject["Id"_ASV].AsObjectId() == KeyObjectId)
 			{
-				Params.Payload = FPayload(Key.GetId(), ValueObject["Hash"_ASV].AsBinaryAttachment());
+				Params.Payload = FPayload(Key.Payload, ValueObject["Hash"_ASV].AsBinaryAttachment());
 			}
 			else
 			{
@@ -1133,7 +1118,7 @@ void FCache::GetPayload(
 				{
 					if (AttachmentField.AsObjectView()["Id"_ASV].AsObjectId() == KeyObjectId)
 					{
-						Params.Payload = FPayload(Key.GetId(), AttachmentField.AsObjectView()["Hash"_ASV].AsBinaryAttachment());
+						Params.Payload = FPayload(Key.Payload, AttachmentField.AsObjectView()["Hash"_ASV].AsBinaryAttachment());
 						break;
 					}
 				}
@@ -1142,20 +1127,20 @@ void FCache::GetPayload(
 			if (Params.Payload)
 			{
 				UE_LOG(LogDerivedDataCache, Verbose, TEXT("Cache: GetPayload (exists) cache hit for %s from '%.*s'"),
-					*TToString<96>(Key), Context.Len(), Context.GetData());
+					*WriteToString<96>(Key), Context.Len(), Context.GetData());
 				Params.Status = EStatus::Ok;
 			}
 			else
 			{
 				UE_LOG(LogDerivedDataCache, Verbose,
 					TEXT("Cache: GetPayload (exists) cache miss with missing payload in record for %s from '%.*s'"),
-					*TToString<96>(Key), Context.Len(), Context.GetData());
+					*WriteToString<96>(Key), Context.Len(), Context.GetData());
 			}
 		}
 		else
 		{
 			UE_LOG(LogDerivedDataCache, Verbose, TEXT("Cache: GetPayload (exists) cache miss for %s from '%.*s'"),
-				*TToString<96>(Key), Context.Len(), Context.GetData());
+				*WriteToString<96>(Key), Context.Len(), Context.GetData());
 		}
 		return;
 	}
@@ -1165,14 +1150,14 @@ void FCache::GetPayload(
 	if (!GetDerivedDataCacheRef().GetSynchronous(*MakePayloadKey(Key), Data, Context))
 	{
 		UE_LOG(LogDerivedDataCache, Verbose, TEXT("Cache: GetPayload cache miss for %s from '%.*s'"),
-			*TToString<96>(Key), Context.Len(), Context.GetData());
+			*WriteToString<96>(Key), Context.Len(), Context.GetData());
 		return;
 	}
 
-	Params.Payload = FPayload(Key.GetId(), FCompressedBuffer(MakeSharedBufferFromArray(MoveTemp(Data))));
+	Params.Payload = FPayload(Key.Payload, FCompressedBuffer(MakeSharedBufferFromArray(MoveTemp(Data))));
 
 	UE_LOG(LogDerivedDataCache, Verbose, TEXT("Cache: GetPayload cache hit for %s from '%.*s'"),
-		*TToString<96>(Key), Context.Len(), Context.GetData());
+		*WriteToString<96>(Key), Context.Len(), Context.GetData());
 	Params.Status = EStatus::Ok;
 }
 
diff --git a/Engine/Source/Developer/DerivedDataCache/Private/DerivedDataCacheKey.cpp b/Engine/Source/Developer/DerivedDataCache/Private/DerivedDataCacheKey.cpp
--- a/Engine/Source/Developer/DerivedDataCache/Private/DerivedDataCacheKey.cpp
+++ b/Engine/Source/Developer/DerivedDataCache/Private/DerivedDataCacheKey.cpp
@@ -4,10 +4,7 @@
 
 #include "Algo/AllOf.h"
 #include "Containers/Set.h"
-#include "Containers/StringConv.h"
-#include "HAL/CriticalSection.h"
 #include "Misc/ScopeRWLock.h"
-#include "Misc/StringBuilder.h"
 
 namespace UE
 {
@@ -104,69 +101,13 @@ private:
 	TSet<FBucket> Buckets;
 };
 
-const TCHAR* FCacheBucket::FindOrAddCacheBucket(FStringView Bucket)
+FCacheBucket CreateCacheBucket(FStringView Name)
 {
-	checkf(!Bucket.IsEmpty() && Bucket.Len() < 256 && Algo::AllOf(Bucket, FChar::IsAlnum),
+	checkf(!Name.IsEmpty() && Name.Len() < 256 && Algo::AllOf(Name, FChar::IsAlnum),
 		TEXT("Invalid cache bucket name '%.*s' must be alphanumeric, non-empty, and contain fewer than 256 characters."),
-		Bucket.Len(), Bucket.GetData());
+		Name.Len(), Name.GetData());
 	static FCacheBuckets Buckets;
-	return Buckets.FindOrAdd(Bucket);
-}
-
-void FCacheBucket::ToString(FAnsiStringBuilderBase& Builder) const
-{
-	Builder << TCHAR_TO_UTF8(Name);
-}
-
-void FCacheBucket::ToString(FWideStringBuilderBase& Builder) const
-{
-	Builder << ToString();
-}
-
-///////////////////////////////////////////////////////////////////////////////////////////////////////////////////////
-
-void FCacheKey::ToString(FAnsiStringBuilderBase& Builder) const
-{
-	Builder << Bucket;
-	if (!Bucket.IsNull())
-	{
-		Builder << "/" << Hash;
-	}
-}
-
-void FCacheKey::ToString(FWideStringBuilderBase& Builder) const
-{
-	Builder << Bucket;
-	if (!Bucket.IsNull())
-	{
-		Builder << TEXT("/") << Hash;
-	}
-}
-
-FString FCacheKey::ToString() const
-{
-	TStringBuilder<72> Out;
-	ToString(Out);
-	return FString(Out);
-}
-
-///////////////////////////////////////////////////////////////////////////////////////////////////////////////////////
-
-void FCachePayloadKey::ToString(FAnsiStringBuilderBase& Builder) const
-{
-	Builder << Key << '/' << Id;
-}
-
-void FCachePayloadKey::ToString(FWideStringBuilderBase& Builder) const
-{
-	Builder << Key << TEXT('/') << Id;
-}
-
-FString FCachePayloadKey::ToString() const
-{
-	TStringBuilder<84> Out;
-	ToString(Out);
-	return FString(Out);
+	return FCacheBucket(Private::FCacheBucketName{Buckets.FindOrAdd(Name)});
 }
 
 ///////////////////////////////////////////////////////////////////////////////////////////////////////////////////////
diff --git a/Engine/Source/Developer/DerivedDataCache/Private/DerivedDataCachePrivate.h b/Engine/Source/Developer/DerivedDataCache/Private/DerivedDataCachePrivate.h
new file mode 100644
--- /dev/null
+++ b/Engine/Source/Developer/DerivedDataCache/Private/DerivedDataCachePrivate.h
@@ -0,0 +1,27 @@
+// Copyright Epic Games, Inc. All Rights Reserved.
+
+#pragma once
+
+#include "Containers/StringFwd.h"
+
+namespace UE
+{
+namespace DerivedData
+{
+
+class FCacheBucket;
+class FCacheRecordBuilder;
+struct FCacheKey;
+
+///////////////////////////////////////////////////////////////////////////////////////////////////////////////////////
+
+// Implemented in DerivedDataCacheKey.cpp
+FCacheBucket CreateCacheBucket(FStringView Name);
+
+// Implemented in DerivedDataCacheRecord.cpp
+FCacheRecordBuilder CreateCacheRecordBuilder(const FCacheKey& Key);
+
+///////////////////////////////////////////////////////////////////////////////////////////////////////////////////////
+
+} // DerivedData
+} // UE
diff --git a/Engine/Source/Developer/DerivedDataCache/Private/DerivedDataCacheRecord.cpp b/Engine/Source/Developer/DerivedDataCache/Private/DerivedDataCacheRecord.cpp
--- a/Engine/Source/Developer/DerivedDataCache/Private/DerivedDataCacheRecord.cpp
+++ b/Engine/Source/Developer/DerivedDataCache/Private/DerivedDataCacheRecord.cpp
@@ -2,19 +2,37 @@
 
 #include "DerivedDataCacheRecord.h"
 
-#include "Async/Future.h"
-#include "Misc/StringBuilder.h"
+#include "Containers/UnrealString.h"
+#include "DerivedDataCacheKey.h"
+#include "DerivedDataCachePrivate.h"
 #include "Serialization/CompactBinary.h"
 
 namespace UE
 {
 namespace DerivedData
 {
+namespace Private
+{
 
 ///////////////////////////////////////////////////////////////////////////////////////////////////////////////////////
 
-struct FCacheRecordBuilderData
+class FCacheRecordBuilderInternal final : public ICacheRecordBuilderInternal
 {
+public:
+	explicit FCacheRecordBuilderInternal(const FCacheKey& Key);
+	virtual ~FCacheRecordBuilderInternal() = default;
+
+	virtual void SetMeta(FCbObject&& Meta) final;
+
+	virtual FPayloadId SetValue(const FSharedBuffer& Buffer, const FPayloadId& Id) final;
+	virtual FPayloadId SetValue(FPayload&& Payload) final;
+
+	virtual FPayloadId AddAttachment(const FSharedBuffer& Buffer, const FPayloadId& Id) final;
+	virtual FPayloadId AddAttachment(FPayload&& Payload) final;
+
+	virtual FCacheRecord Build() final;
+	virtual FRequest BuildAsync(FOnCacheRecordComplete&& Callback, EPriority Priority) final;
+
 	FCacheKey Key;
 	FCbObject Meta;
 	FPayload Value;
@@ -23,27 +41,34 @@ struct FCacheRecordBuilderData
 
 ///////////////////////////////////////////////////////////////////////////////////////////////////////////////////////
 
-struct FCacheRecordData
+class FCacheRecordInternal final : public ICacheRecordInternal
 {
-	FCacheRecordData() = default;
-	explicit FCacheRecordData(FCacheRecordBuilderData&& Builder);
+public:
+	FCacheRecordInternal() = default;
+	explicit FCacheRecordInternal(FCacheRecordBuilderInternal&& Builder);
+
+	virtual ~FCacheRecordInternal() = default;
+
+	virtual FCacheRecord Clone() const final;
+
+	virtual const FCacheKey& GetKey() const final;
+	virtual const FCbObject& GetMeta() const final;
+
+	virtual FSharedBuffer GetValue() const final;
+	virtual const FPayload& GetValuePayload() const final;
+
+	virtual FSharedBuffer GetAttachment(const FPayloadId& Id) const final;
+	virtual const FPayload& GetAttachmentPayload(const FPayloadId& Id) const final;
+	virtual TConstArrayView<FPayload> GetAttachmentPayloads() const final;
 
 	FCacheKey Key;
 	FCbObject Meta;
 	FPayload Value;
-	FSharedBuffer ValueCache;
 	TArray<FPayload> Attachments;
-	TArray<FSharedBuffer> AttachmentsCache;
+	mutable FSharedBuffer ValueCache;
+	mutable TArray<FSharedBuffer> AttachmentsCache;
 };
 
-FCacheRecordData::FCacheRecordData(FCacheRecordBuilderData&& Builder)
-	: Key(Builder.Key)
-	, Meta(MoveTemp(Builder.Meta))
-	, Value(MoveTemp(Builder.Value))
-	, Attachments(MoveTemp(Builder.Attachments))
-{
-}
-
 ///////////////////////////////////////////////////////////////////////////////////////////////////////////////////////
 
 static const FPayload& GetEmptyCachePayload()
@@ -52,12 +77,6 @@ static const FPayload& GetEmptyCachePayload()
 	return Empty;
 }
 
-static const FCacheRecordData& GetEmptyCacheRecord()
-{
-	static const FCacheRecordData Empty;
-	return Empty;
-}
-
 static FPayloadId GetOrCreatePayloadId(const FPayloadId& Id, const FSharedBuffer& Buffer)
 {
 	return Id.IsValid() ? Id : FPayloadId::FromHash(FIoHash::HashBuffer(Buffer));
@@ -65,56 +84,48 @@ static FPayloadId GetOrCreatePayloadId(const FPayloadId& Id, const FSharedBuffer
 
 ///////////////////////////////////////////////////////////////////////////////////////////////////////////////////////
 
-FCacheRecord::FCacheRecord(const FCacheRecordData& InData)
-	: Data(MakePimpl<FCacheRecordData>(InData))
-{
-}
-
-FCacheRecord::FCacheRecord(FCacheRecordBuilderData&& Builder)
-	: Data(MakePimpl<FCacheRecordData>(MoveTemp(Builder)))
+FCacheRecordInternal::FCacheRecordInternal(FCacheRecordBuilderInternal&& Builder)
+	: Key(Builder.Key)
+	, Meta(MoveTemp(Builder.Meta))
+	, Value(MoveTemp(Builder.Value))
+	, Attachments(MoveTemp(Builder.Attachments))
 {
 }
 
-FCacheRecord FCacheRecord::Clone() const
+FCacheRecord FCacheRecordInternal::Clone() const
 {
-	return Data ? FCacheRecord(*Data) : FCacheRecord();
+	return FCacheRecord(new FCacheRecordInternal(*this));
 }
 
-const FCacheKey& FCacheRecord::GetKey() const
+const FCacheKey& FCacheRecordInternal::GetKey() const
 {
-	return (Data ? *Data : GetEmptyCacheRecord()).Key;
+	return Key;
 }
 
-const FCbObject& FCacheRecord::GetMeta() const
+const FCbObject& FCacheRecordInternal::GetMeta() const
 {
-	return (Data ? *Data : GetEmptyCacheRecord()).Meta;
+	return Meta;
 }
 
-FSharedBuffer FCacheRecord::GetValue() const
+FSharedBuffer FCacheRecordInternal::GetValue() const
 {
-	if (Data)
+	if (ValueCache.IsNull() && Value.IsValid())
 	{
-		if (Data->ValueCache.IsNull() && Data->Value.IsValid())
-		{
-			Data->ValueCache = Data->Value.GetBuffer().Decompress();
-		}
-		return Data->ValueCache;
+		ValueCache = Value.GetBuffer().Decompress();
 	}
-	return FSharedBuffer();
+	return ValueCache;
 }
 
-const FPayload& FCacheRecord::GetValuePayload() const
+const FPayload& FCacheRecordInternal::GetValuePayload() const
 {
-	return (Data ? *Data : GetEmptyCacheRecord()).Value;
+	return Value;
 }
 
-FSharedBuffer FCacheRecord::GetAttachment(const FPayloadId& Id) const
+FSharedBuffer FCacheRecordInternal::GetAttachment(const FPayloadId& Id) const
 {
-	TArray<FPayload>& Attachments = Data->Attachments;
 	const int32 Index = Algo::LowerBound(Attachments, Id, FPayloadLessById());
 	if (Attachments.IsValidIndex(Index) && FPayloadEqualById()(Attachments[Index], Id))
 	{
-		TArray<FSharedBuffer>& AttachmentsCache = Data->AttachmentsCache;
 		if (AttachmentsCache.IsEmpty())
 		{
 			AttachmentsCache.SetNum(Attachments.Num());
@@ -129,9 +140,8 @@ FSharedBuffer FCacheRecord::GetAttachment(const FPayloadId& Id) const
 	return FSharedBuffer();
 }
 
-const FPayload& FCacheRecord::GetAttachmentPayload(const FPayloadId& Id) const
+const FPayload& FCacheRecordInternal::GetAttachmentPayload(const FPayloadId& Id) const
 {
-	TArray<FPayload>& Attachments = Data->Attachments;
 	const int32 Index = Algo::LowerBound(Attachments, Id, FPayloadLessById());
 	if (Attachments.IsValidIndex(Index) && FPayloadEqualById()(Attachments[Index], Id))
 	{
@@ -140,78 +150,76 @@ const FPayload& FCacheRecord::GetAttachmentPayload(const FPayloadId& Id) const
 	return GetEmptyCachePayload();
 }
 
-TConstArrayView<FPayload> FCacheRecord::GetAttachmentPayloads() const
+TConstArrayView<FPayload> FCacheRecordInternal::GetAttachmentPayloads() const
 {
-	return (Data ? *Data : GetEmptyCacheRecord()).Attachments;
+	return Attachments;
 }
 
 ///////////////////////////////////////////////////////////////////////////////////////////////////////////////////////
 
-FCacheRecordBuilder::FCacheRecordBuilder()
-	: Data(MakePimpl<FCacheRecordBuilderData>())
-{
-}
-
-void FCacheRecordBuilder::SetKey(const FCacheKey& Key)
+FCacheRecordBuilderInternal::FCacheRecordBuilderInternal(const FCacheKey& InKey)
+	: Key(InKey)
 {
-	Data->Key = Key;
 }
 
-void FCacheRecordBuilder::SetMeta(FCbObject Meta)
+void FCacheRecordBuilderInternal::SetMeta(FCbObject&& InMeta)
 {
+	Meta = MoveTemp(InMeta);
 	Meta.MakeOwned();
-	Data->Meta = MoveTemp(Meta);
 }
 
-FPayloadId FCacheRecordBuilder::SetValue(const FSharedBuffer& Buffer, FPayloadId Id)
+FPayloadId FCacheRecordBuilderInternal::SetValue(const FSharedBuffer& Buffer, const FPayloadId& Id)
 {
-	Id = GetOrCreatePayloadId(Id, Buffer);
-	return SetValue(FPayload(Id, FCompressedBuffer::Compress(Buffer)));
+	const FPayloadId ValueId = GetOrCreatePayloadId(Id, Buffer);
+	return SetValue(FPayload(ValueId, FCompressedBuffer::Compress(Buffer)));
 }
 
-FPayloadId FCacheRecordBuilder::SetValue(FPayload&& Payload)
+FPayloadId FCacheRecordBuilderInternal::SetValue(FPayload&& Payload)
 {
-	FPayload& Value = Data->Value;
 	checkf(Value.IsNull(), TEXT("Failed to set value with ID %s because an there is an existing value with ID %s"),
-		*Payload.GetId().ToString(), *Value.GetId().ToString());
+		*WriteToString<32>(Payload.GetId()), *WriteToString<32>(Value.GetId()));
 	Value = MoveTemp(Payload);
 	return Value.GetId();
 }
 
-FPayloadId FCacheRecordBuilder::AddAttachment(const FSharedBuffer& Buffer, FPayloadId Id)
+FPayloadId FCacheRecordBuilderInternal::AddAttachment(const FSharedBuffer& Buffer, const FPayloadId& Id)
 {
-	Id = GetOrCreatePayloadId(Id, Buffer);
-	return AddAttachment(FPayload(Id, FCompressedBuffer::Compress(Buffer)));
+	const FPayloadId AttachmentId = GetOrCreatePayloadId(Id, Buffer);
+	return AddAttachment(FPayload(AttachmentId, FCompressedBuffer::Compress(Buffer)));
 }
 
-FPayloadId FCacheRecordBuilder::AddAttachment(FPayload&& Payload)
+FPayloadId FCacheRecordBuilderInternal::AddAttachment(FPayload&& Payload)
 {
-	TArray<FPayload>& Attachments = Data->Attachments;
 	const int32 Index = Algo::LowerBound(Attachments, Payload, FPayloadLessById());
 	checkf(!Attachments.IsValidIndex(Index) || !FPayloadEqualById()(Attachments[Index], Payload),
 		TEXT("Failed to add attachment with ID %s because an existing attachment is using that ID"),
-		*Payload.GetId().ToString());
+		*WriteToString<32>(Payload.GetId()));
 	Attachments.Insert(MoveTemp(Payload), Index);
 	return Attachments[Index].GetId();
 }
 
-FCacheRecord FCacheRecordBuilder::Build()
+FCacheRecord FCacheRecordBuilderInternal::Build()
 {
-	if (TPimplPtr<FCacheRecordBuilderData> LocalData = MoveTemp(Data))
-	{
-		return FCacheRecord(MoveTemp(*LocalData));
-	}
-	return FCacheRecord();
+	return FCacheRecord(new FCacheRecordInternal(MoveTemp(*this)));
 }
 
-TFuture<FCacheRecord> FCacheRecordBuilder::BuildAsync()
+FRequest FCacheRecordBuilderInternal::BuildAsync(FOnCacheRecordComplete&& Callback, EPriority Priority)
 {
-	TPromise<FCacheRecord> Promise;
-	Promise.SetValue(Build());
-	return Promise.GetFuture();
+	Callback(Build());
+	return FRequest();
 }
 
 ///////////////////////////////////////////////////////////////////////////////////////////////////////////////////////
 
+} // Private
 } // DerivedData
 } // UE
+
+///////////////////////////////////////////////////////////////////////////////////////////////////////////////////////
+
+UE::DerivedData::FCacheRecordBuilder UE::DerivedData::CreateCacheRecordBuilder(const FCacheKey& Key)
+{
+	return FCacheRecordBuilder(new Private::FCacheRecordBuilderInternal(Key));
+}
+
+///////////////////////////////////////////////////////////////////////////////////////////////////////////////////////
diff --git a/Engine/Source/Developer/DerivedDataCache/Private/DerivedDataPayload.cpp b/Engine/Source/Developer/DerivedDataCache/Private/DerivedDataPayload.cpp
deleted file mode 100644
diff --git a/Engine/Source/Developer/DerivedDataCache/Public/DerivedDataCache.h b/Engine/Source/Developer/DerivedDataCache/Public/DerivedDataCache.h
--- a/Engine/Source/Developer/DerivedDataCache/Public/DerivedDataCache.h
+++ b/Engine/Source/Developer/DerivedDataCache/Public/DerivedDataCache.h
@@ -13,6 +13,8 @@ namespace UE
 namespace DerivedData
 {
 
+class FCacheBucket;
+class FCacheRecordBuilder;
 struct FCacheGetCompleteParams;
 struct FCacheGetPayloadCompleteParams;
 struct FCachePutCompleteParams;
@@ -101,6 +103,18 @@ class ICache
 public:
 	virtual ~ICache() = default;
 
+	/**
+	 * Create a cache bucket from a name.
+	 *
+	 * A cache bucket name must be alphanumeric, non-empty, and contain fewer than 256 characters.
+	 */
+	virtual FCacheBucket CreateBucket(FStringView Name) = 0;
+
+	/**
+	 * Create a cache record builder from a cache key.
+	 */
+	virtual FCacheRecordBuilder CreateRecord(const FCacheKey& Key) = 0;
+
 	/**
 	 * Asynchronous request to put cache records according to the policy.
 	 *
diff --git a/Engine/Source/Developer/DerivedDataCache/Public/DerivedDataCacheKey.h b/Engine/Source/Developer/DerivedDataCache/Public/DerivedDataCacheKey.h
--- a/Engine/Source/Developer/DerivedDataCache/Public/DerivedDataCacheKey.h
+++ b/Engine/Source/Developer/DerivedDataCache/Public/DerivedDataCacheKey.h
@@ -3,16 +3,12 @@
 #pragma once
 
 #include "CoreTypes.h"
-#include "Containers/StringFwd.h"
 #include "Containers/StringView.h"
 #include "DerivedDataPayload.h"
 #include "IO/IoHash.h"
+#include "Misc/StringBuilder.h"
 #include "Templates/TypeHash.h"
 
-#define UE_API DERIVEDDATACACHE_API
-
-class FCbObjectId;
-
 namespace UE
 {
 namespace DerivedData
@@ -20,26 +16,27 @@ namespace DerivedData
 
 ///////////////////////////////////////////////////////////////////////////////////////////////////////////////////////
 
-/** An alphanumeric identifier that groups related cache records. */
+namespace Private
+{
+
+struct FCacheBucketName { const TCHAR* Name; };
+
+}
+
+///////////////////////////////////////////////////////////////////////////////////////////////////////////////////////
+
+/**
+ * An alphanumeric identifier that groups related cache records.
+ *
+ * Create using ICache::CreateBucket.
+ */
 class FCacheBucket
 {
 public:
 	/** Construct a null cache bucket. */
 	FCacheBucket() = default;
 
-	/**
-	 * Construct a cache bucket from the name.
-	 *
-	 * A cache bucket name must be alphanumeric, non-empty, and contain fewer than 256 characters.
-	 */
-	inline explicit FCacheBucket(FStringView Bucket)
-		: Name(FindOrAddCacheBucket(Bucket))
-	{
-	}
-
-	/** Convert the name of the cache bucket to a string. */
-	UE_API void ToString(FAnsiStringBuilderBase& Builder) const;
-	UE_API void ToString(FWideStringBuilderBase& Builder) const;
+	/** Get the name of the cache bucket as a string. */
 	inline FStringView ToString() const;
 
 	/** Whether this is null. */
@@ -58,11 +55,21 @@ public:
 		return Name != Other.Name && ToString().Compare(Other.ToString(), ESearchCase::IgnoreCase) < 0;
 	}
 
-	friend inline uint32 GetTypeHash(FCacheBucket Bucket) { return ::GetTypeHash(reinterpret_cast<UPTRINT>(Bucket.Name)); }
+	friend inline uint32 GetTypeHash(FCacheBucket Bucket)
+	{
+		return ::GetTypeHash(reinterpret_cast<UPTRINT>(Bucket.Name));
+	}
+
+public:
+	// Internal API
 
-private:
-	UE_API static const TCHAR* FindOrAddCacheBucket(FStringView Bucket);
+	explicit FCacheBucket(Private::FCacheBucketName InName)
+		: Name(InName.Name)
+	{
+	}
 
+private:
+	/** Name with a one-byte length stored in the preceding character. */
 	const TCHAR* Name = nullptr;
 };
 
@@ -71,154 +78,84 @@ inline FStringView FCacheBucket::ToString() const
 	return Name ? FStringView(Name, *reinterpret_cast<const uint8*>(Name - 1)) : FStringView();
 }
 
-template <typename CharType>
-inline TStringBuilderBase<CharType>& operator<<(TStringBuilderBase<CharType>& Builder, const FCacheBucket& Bucket)
+inline FStringBuilderBase& operator<<(FStringBuilderBase& Builder, const FCacheBucket& Bucket)
 {
-	Bucket.ToString(Builder);
-	return Builder;
+	return Builder << Bucket.ToString();
 }
 
 ///////////////////////////////////////////////////////////////////////////////////////////////////////////////////////
 
 /** A key that uniquely identifies a cache record. */
-class FCacheKey
+struct FCacheKey
 {
-public:
-	/** Construct a null cache key. */
-	FCacheKey() = default;
-
-	/** Construct a cache key from the non-null bucket and non-zero hash. */
-	inline FCacheKey(FCacheBucket InBucket, const FIoHash& InHash)
-		: Bucket(InBucket)
-		, Hash(InHash)
-	{
-	}
-
-	/** Returns the bucket component of the cache key. */
-	inline FCacheBucket GetBucket() const { return Bucket; }
-
-	/** Returns the hash component of the cache key. */
-	inline const FIoHash& GetHash() const { return Hash; }
-
-	/** Convert the cache key to a string. */
-	UE_API void ToString(FAnsiStringBuilderBase& Builder) const;
-	UE_API void ToString(FWideStringBuilderBase& Builder) const;
-	UE_API FString ToString() const;
-
-	/** Whether this is null. */
-	inline bool IsNull() const { return Bucket.IsNull(); }
-	/** Whether this is not null. */
-	inline bool IsValid() const { return !IsNull(); }
-
-	/** Reset this to null. */
-	inline void Reset() { *this = FCacheKey(); }
-
-private:
 	FCacheBucket Bucket;
 	FIoHash Hash;
 };
 
 inline bool operator==(const FCacheKey& A, const FCacheKey& B)
 {
-	return A.GetBucket() == B.GetBucket() && A.GetHash() == B.GetHash();
+	return A.Bucket == B.Bucket && A.Hash == B.Hash;
 }
 
 inline bool operator!=(const FCacheKey& A, const FCacheKey& B)
 {
-	return A.GetBucket() != B.GetBucket() || A.GetHash() != B.GetHash();
+	return A.Bucket != B.Bucket || A.Hash != B.Hash;
 }
 
 inline bool operator<(const FCacheKey& A, const FCacheKey& B)
 {
-	const FCacheBucket& BucketA = A.GetBucket();
-	const FCacheBucket& BucketB = B.GetBucket();
-	return BucketA == BucketB ? A.GetHash() < B.GetHash() : BucketA < BucketB;
+	const FCacheBucket& BucketA = A.Bucket;
+	const FCacheBucket& BucketB = B.Bucket;
+	return BucketA == BucketB ? A.Hash < B.Hash : BucketA < BucketB;
 }
 
 inline uint32 GetTypeHash(const FCacheKey& Key)
 {
-	return HashCombine(GetTypeHash(Key.GetBucket()), GetTypeHash(Key.GetHash()));
+	return HashCombine(GetTypeHash(Key.Bucket), GetTypeHash(Key.Hash));
 }
 
-template <typename CharType>
-inline TStringBuilderBase<CharType>& operator<<(TStringBuilderBase<CharType>& Builder, const FCacheKey& Key)
+inline FStringBuilderBase& operator<<(FStringBuilderBase& Builder, const FCacheKey& Key)
 {
-	Key.ToString(Builder);
-	return Builder;
+	return Builder << Key.Bucket << TEXT('/') << Key.Hash;
 }
 
 ///////////////////////////////////////////////////////////////////////////////////////////////////////////////////////
 
 /** A key that uniquely identifies a payload within a cache record. */
-class FCachePayloadKey
+struct FCachePayloadKey
 {
-public:
-	/** Construct a null cache payload key. */
-	FCachePayloadKey() = default;
-
-	/** Construct a cache payload key from the non-null key and non-null ID. */
-	inline FCachePayloadKey(const FCacheKey& InKey, const FPayloadId& InId)
-		: Key(InKey)
-		, Id(InId)
-	{
-	}
-
-	/** Returns the cache key for the record containing the payload. */
-	inline const FCacheKey& GetKey() const { return Key; }
-
-	/** Returns the ID of the payload. */
-	inline const FPayloadId& GetId() const { return Id; }
-
-	/** Convert the cache payload key to a string. */
-	UE_API void ToString(FAnsiStringBuilderBase& Builder) const;
-	UE_API void ToString(FWideStringBuilderBase& Builder) const;
-	UE_API FString ToString() const;
-
-	/** Whether this is null. */
-	inline bool IsNull() const { return Key.IsNull(); }
-	/** Whether this is not null. */
-	inline bool IsValid() const { return !IsNull(); }
-
-	/** Reset this to null. */
-	inline void Reset() { *this = FCachePayloadKey(); }
-
-private:
-	FCacheKey Key;
-	FPayloadId Id;
+	FCacheKey Record;
+	FPayloadId Payload;
 };
 
 inline bool operator==(const FCachePayloadKey& A, const FCachePayloadKey& B)
 {
-	return A.GetKey() == B.GetKey() && A.GetId() == B.GetId();
+	return A.Record == B.Record && A.Payload == B.Payload;
 }
 
 inline bool operator!=(const FCachePayloadKey& A, const FCachePayloadKey& B)
 {
-	return A.GetKey() != B.GetKey() || A.GetId() != B.GetId();
+	return A.Record != B.Record || A.Payload != B.Payload;
 }
 
 inline bool operator<(const FCachePayloadKey& A, const FCachePayloadKey& B)
 {
-	const FCacheKey& KeyA = A.GetKey();
-	const FCacheKey& KeyB = B.GetKey();
-	return KeyA == KeyB ? A.GetId() < B.GetId() : KeyA < KeyB;
+	const FCacheKey& KeyA = A.Record;
+	const FCacheKey& KeyB = B.Record;
+	return KeyA == KeyB ? A.Payload < B.Payload : KeyA < KeyB;
 }
 
 inline uint32 GetTypeHash(const FCachePayloadKey& Key)
 {
-	return HashCombine(GetTypeHash(Key.GetKey()), GetTypeHash(Key.GetId()));
+	return HashCombine(GetTypeHash(Key.Record), GetTypeHash(Key.Payload));
 }
 
-template <typename CharType>
-inline TStringBuilderBase<CharType>& operator<<(TStringBuilderBase<CharType>& Builder, const FCachePayloadKey& Key)
+inline FStringBuilderBase& operator<<(FStringBuilderBase& Builder, const FCachePayloadKey& Key)
 {
-	Key.ToString(Builder);
-	return Builder;
+	return Builder << Key.Record << TEXT('/') << Key.Payload;
 }
 
 ///////////////////////////////////////////////////////////////////////////////////////////////////////////////////////
 
 } // DerivedData
 } // UE
-
-#undef UE_API
diff --git a/Engine/Source/Developer/DerivedDataCache/Public/DerivedDataCacheRecord.h b/Engine/Source/Developer/DerivedDataCache/Public/DerivedDataCacheRecord.h
--- a/Engine/Source/Developer/DerivedDataCache/Public/DerivedDataCacheRecord.h
+++ b/Engine/Source/Developer/DerivedDataCache/Public/DerivedDataCacheRecord.h
@@ -3,24 +3,57 @@
 #pragma once
 
 #include "CoreTypes.h"
-#include "DerivedDataCacheKey.h"
 #include "DerivedDataPayload.h"
+#include "DerivedDataRequest.h"
 #include "Memory/MemoryFwd.h"
-#include "Templates/PimplPtr.h"
-
-#define UE_API DERIVEDDATACACHE_API
+#include "Templates/UniquePtr.h"
 
 class FCbObject;
-template <typename ResultType> class TFuture;
+template <typename FuncType> class TUniqueFunction;
 
 namespace UE
 {
 namespace DerivedData
 {
 
-class FCacheRecordBuilder;
-struct FCacheRecordData;
-struct FCacheRecordBuilderData;
+class FCacheRecord;
+struct FCacheKey;
+
+using FOnCacheRecordComplete = TUniqueFunction<void (FCacheRecord&& Record)>;
+
+///////////////////////////////////////////////////////////////////////////////////////////////////////////////////////
+
+namespace Private
+{
+
+class ICacheRecordInternal
+{
+public:
+	virtual ~ICacheRecordInternal() = default;
+	virtual FCacheRecord Clone() const = 0;
+	virtual const FCacheKey& GetKey() const = 0;
+	virtual const FCbObject& GetMeta() const = 0;
+	virtual FSharedBuffer GetValue() const = 0;
+	virtual const FPayload& GetValuePayload() const = 0;
+	virtual FSharedBuffer GetAttachment(const FPayloadId& Id) const = 0;
+	virtual const FPayload& GetAttachmentPayload(const FPayloadId& Id) const = 0;
+	virtual TConstArrayView<FPayload> GetAttachmentPayloads() const = 0;
+};
+
+class ICacheRecordBuilderInternal
+{
+public:
+	virtual ~ICacheRecordBuilderInternal() = default;
+	virtual void SetMeta(FCbObject&& Meta) = 0;
+	virtual FPayloadId SetValue(const FSharedBuffer& Buffer, const FPayloadId& Id) = 0;
+	virtual FPayloadId SetValue(FPayload&& Payload) = 0;
+	virtual FPayloadId AddAttachment(const FSharedBuffer& Buffer, const FPayloadId& Id) = 0;
+	virtual FPayloadId AddAttachment(FPayload&& Payload) = 0;
+	virtual FCacheRecord Build() = 0;
+	virtual FRequest BuildAsync(FOnCacheRecordComplete&& Callback, EPriority Priority) = 0;
+};
+
+} // Private
 
 ///////////////////////////////////////////////////////////////////////////////////////////////////////////////////////
 
@@ -43,34 +76,34 @@ class FCacheRecord
 {
 public:
 	/** Construct a null cache record. */
-	FCacheRecord() = default;
+	explicit FCacheRecord() = default;
 
 	/** Clone this cache record. Prefer to move records, and clone only when a copy is necessary. */
-	UE_API FCacheRecord Clone() const;
+	inline FCacheRecord Clone() const { return Record->Clone(); }
 
 	/** Returns the key that identifies this record in the cache. Always available in a non-null cache record. */
-	UE_API const FCacheKey& GetKey() const;
+	inline const FCacheKey& GetKey() const { return Record->GetKey(); }
 
 	/** Returns the metadata for the cache record unless it was skipped in the request. */
-	UE_API const FCbObject& GetMeta() const;
+	inline const FCbObject& GetMeta() const { return Record->GetMeta(); }
 
 	/** Returns the value. Null if the value was skipped in the request. */
-	UE_API FSharedBuffer GetValue() const;
+	inline FSharedBuffer GetValue() const { return Record->GetValue(); }
 
 	/** Returns the value payload. Always available in a non-null cache record. */
-	UE_API const FPayload& GetValuePayload() const;
+	inline const FPayload& GetValuePayload() const { return Record->GetValuePayload(); }
 
 	/** Returns the attachment that matches the ID. Null if no match or attachments were skipped in the request. */
-	UE_API FSharedBuffer GetAttachment(const FPayloadId& Id) const;
+	inline FSharedBuffer GetAttachment(const FPayloadId& Id) const { return Record->GetAttachment(Id); }
 
 	/** Returns the attachment payload that matches the ID. Null if no match or the cache record is null. */
-	UE_API const FPayload& GetAttachmentPayload(const FPayloadId& Id) const;
+	inline const FPayload& GetAttachmentPayload(const FPayloadId& Id) const { return Record->GetAttachmentPayload(Id); }
 
 	/** Returns a view of the attachments. Always available in a non-null cache record, but buffer may be skipped. */
-	UE_API TConstArrayView<FPayload> GetAttachmentPayloads() const;
+	inline TConstArrayView<FPayload> GetAttachmentPayloads() const { return Record->GetAttachmentPayloads(); }
 
 	/** Whether this is null. */
-	inline bool IsNull() const { return !Data; }
+	inline bool IsNull() const { return !Record; }
 	/** Whether this is not null. */
 	inline bool IsValid() const { return !IsNull(); }
 	/** Whether this is not null. */
@@ -82,14 +115,14 @@ public:
 public:
 	// Internal API
 
-	/** Construct a cache record by cloning a cache record. Use Clone(). */
-	explicit FCacheRecord(const FCacheRecordData& Data);
-
-	/** Construct a cache record from a builder. Use Builder.Build() or Builder.BuildAsync(). */
-	explicit FCacheRecord(FCacheRecordBuilderData&& Builder);
+	/** Construct a cache record. Use Record.Clone() or Builder.Build[Async](). */
+	inline explicit FCacheRecord(Private::ICacheRecordInternal* InRecord)
+		: Record(InRecord)
+	{
+	}
 
 private:
-	TPimplPtr<FCacheRecordData> Data;
+	TUniquePtr<Private::ICacheRecordInternal> Record;
 };
 
 ///////////////////////////////////////////////////////////////////////////////////////////////////////////////////////
@@ -97,31 +130,26 @@ private:
 /**
  * A cache record builder is used to construct a cache record.
  *
- * A key must be set before using the this to construct a cache record. The value and attachments
- * can be provided as buffers or as payloads which are already compressed and have an identifier.
+ * Create using ICache::CreateRecord, which must be given a key uniquely corresponds to the value
+ * and attachments for the cache record. Optional metadata may vary and may be non-deterministic.
+ *
+ * The value and attachments can be provided as buffers, which will be compressed, or as payloads
+ * which were previously compressed and have an identifier assigned.
  *
  * @see FCacheRecord
  */
 class FCacheRecordBuilder
 {
 public:
-	/** Construct an empty cache record builder. */
-	UE_API FCacheRecordBuilder();
-
-	/**
-	 * Set the key for the cache record.
-	 *
-	 * The key must uniquely correspond to the value and attachments for the cache record, while the
-	 * optional metadata may vary and may be non-deterministic.
-	 */
-	UE_API void SetKey(const FCacheKey& Key);
-
 	/**
 	 * Set the metadata for the cache record.
 	 *
 	 * The metadata is cloned if not owned.
 	 */
-	UE_API void SetMeta(FCbObject Meta);
+	inline void SetMeta(FCbObject&& Meta)
+	{
+		return Builder->SetMeta(MoveTemp(Meta));
+	}
 
 	/**
 	 * Set the value for the cache record.
@@ -131,7 +159,10 @@ public:
 	 *           the builder will create an ID by hashing the buffer.
 	 * @return The ID that was provided or created.
 	 */
-	UE_API FPayloadId SetValue(const FSharedBuffer& Buffer, FPayloadId Id = FPayloadId());
+	inline FPayloadId SetValue(const FSharedBuffer& Buffer, const FPayloadId& Id = FPayloadId())
+	{
+		return Builder->SetValue(Buffer, Id);
+	}
 
 	/**
 	 * Set the value for the cache record.
@@ -139,7 +170,7 @@ public:
 	 * @param Payload The value payload, which is reset to null.
 	 * @return The ID that was provided. Unique within the scope of the cache record.
 	 */
-	UE_API FPayloadId SetValue(FPayload&& Payload);
+	inline FPayloadId SetValue(FPayload&& Payload) { return Builder->SetValue(MoveTemp(Payload)); }
 
 	/**
 	 * Add an attachment to the cache record.
@@ -149,7 +180,10 @@ public:
 	 *           is omitted the builder will create an ID by hashing the buffer.
 	 * @return The ID that was provided or created.
 	 */
-	UE_API FPayloadId AddAttachment(const FSharedBuffer& Buffer, FPayloadId Id = FPayloadId());
+	inline FPayloadId AddAttachment(const FSharedBuffer& Buffer, const FPayloadId& Id = FPayloadId())
+	{
+		return Builder->AddAttachment(Buffer, Id);
+	}
 
 	/**
 	 * Add an attachment to the cache record.
@@ -157,41 +191,53 @@ public:
 	 * @param Payload The attachment payload, which is reset to null.
 	 * @return The ID that was provided. Unique within the scope of the cache record.
 	 */
-	UE_API FPayloadId AddAttachment(FPayload&& Payload);
+	inline FPayloadId AddAttachment(FPayload&& Payload)
+	{
+		return Builder->AddAttachment(MoveTemp(Payload));
+	}
 
 	/**
-	 * Build a cache record, resetting this builder.
+	 * Build a cache record, which leaves this builder null.
 	 *
 	 * Prefer BuildAsync() when the value or attachments are added from a buffer, as this must block
 	 * on compression of those buffers before it can construct a cache record.
 	 */
-	UE_API FCacheRecord Build();
+	inline FCacheRecord Build()
+	{
+		return Builder->Build();
+	}
 
 	/**
-	 * Build a cache record asynchronously, resetting this builder.
+	 * Build a cache record asynchronously, which leaves this builder null.
 	 *
 	 * Prefer Build() when the value and attachments are added by payload, as compression is already
 	 * complete and BuildAsync() will complete immediately in that case.
 	 */
-	UE_API TFuture<FCacheRecord> BuildAsync();
+	inline FRequest BuildAsync(FOnCacheRecordComplete&& Callback, EPriority Priority)
+	{
+		return Builder->BuildAsync(MoveTemp(Callback), Priority);
+	}
 
 	/** Whether this is null. */
-	inline bool IsNull() const { return !Data; }
+	inline bool IsNull() const { return !Builder; }
 	/** Whether this is not null. */
 	inline bool IsValid() const { return !IsNull(); }
 	/** Whether this is not null. */
 	inline explicit operator bool() const { return IsValid(); }
 
-	/** Reset this to empty. */
-	inline void Reset() { *this = FCacheRecordBuilder(); }
+public:
+	// Internal API
+
+	inline explicit FCacheRecordBuilder(Private::ICacheRecordBuilderInternal* InBuilder)
+		: Builder(InBuilder)
+	{
+	}
 
 private:
-	TPimplPtr<FCacheRecordBuilderData> Data;
+	TUniquePtr<Private::ICacheRecordBuilderInternal> Builder;
 };
 
 ///////////////////////////////////////////////////////////////////////////////////////////////////////////////////////
 
 } // DerivedData
 } // UE
-
-#undef UE_API
diff --git a/Engine/Source/Developer/DerivedDataCache/Public/DerivedDataPayload.h b/Engine/Source/Developer/DerivedDataCache/Public/DerivedDataPayload.h
--- a/Engine/Source/Developer/DerivedDataCache/Public/DerivedDataPayload.h
+++ b/Engine/Source/Developer/DerivedDataCache/Public/DerivedDataPayload.h
@@ -4,14 +4,12 @@
 
 #include "CoreTypes.h"
 #include "Compression/CompressedBuffer.h"
-#include "Containers/StringFwd.h"
+#include "Containers/StringConv.h"
+#include "Containers/StringView.h"
 #include "IO/IoHash.h"
 #include "Memory/MemoryView.h"
-
-#define UE_API DERIVEDDATACACHE_API
-
-class FCbObjectId;
-class FString;
+#include "Misc/StringBuilder.h"
+#include "String/BytesToHex.h"
 
 namespace UE
 {
@@ -33,25 +31,17 @@ public:
 	inline explicit FPayloadId(FMemoryView Id);
 
 	/** Construct an ID from a non-zero hash. */
-	UE_API static FPayloadId FromHash(const FIoHash& Hash);
+	static inline FPayloadId FromHash(const FIoHash& Hash);
 
 	/** Construct an ID from a non-empty name. */
-	UE_API static FPayloadId FromName(FAnsiStringView Name);
-	UE_API static FPayloadId FromName(FWideStringView Name);
-
-	/** Construct an ID from an ObjectId. */
-	UE_API static FPayloadId FromObjectId(const FCbObjectId& ObjectId);
-
-	/** Convert the ID to an ObjectId. */
-	UE_API FCbObjectId ToObjectId() const;
+	static inline FPayloadId FromName(FAnsiStringView Name);
+	static inline FPayloadId FromName(FWideStringView Name);
 
-	/** Convert the ID to a 24-character hex string. */
-	UE_API void ToString(FAnsiStringBuilderBase& Builder) const;
-	UE_API void ToString(FWideStringBuilderBase& Builder) const;
-	UE_API FString ToString() const;
+	/** Returns a reference to the raw byte array for the ID. */
+	inline const ByteArray& GetBytes() const { return Bytes; }
 
 	/** Returns a view of the raw byte array for the ID. */
-	constexpr inline FMemoryView GetView() const { return MakeMemoryView(Bytes); }
+	inline FMemoryView GetView() const { return MakeMemoryView(Bytes); }
 
 	/** Whether this is null. */
 	inline bool IsNull() const;
@@ -65,45 +55,6 @@ private:
 	alignas(uint32) ByteArray Bytes{};
 };
 
-inline bool operator==(const FPayloadId& A, const FPayloadId& B)
-{
-	return A.GetView().EqualBytes(B.GetView());
-}
-
-inline bool operator!=(const FPayloadId& A, const FPayloadId& B)
-{
-	return !A.GetView().EqualBytes(B.GetView());
-}
-
-inline bool operator<(const FPayloadId& A, const FPayloadId& B)
-{
-	return A.GetView().CompareBytes(B.GetView()) < 0;
-}
-
-inline uint32 GetTypeHash(const FPayloadId& Id)
-{
-	return *reinterpret_cast<const uint32*>(Id.GetView().GetData());
-}
-
-inline FPayloadId::FPayloadId(const FMemoryView InId)
-{
-	checkf(InId.GetSize() == sizeof(ByteArray),
-		TEXT("FPayloadId cannot be constructed from a view of %" UINT64_FMT " bytes."), InId.GetSize());
-	FMemory::Memcpy(Bytes, InId.GetData(), sizeof(ByteArray));
-}
-
-inline bool FPayloadId::IsNull() const
-{
-	return *this == FPayloadId();
-}
-
-template <typename CharType>
-inline TStringBuilderBase<CharType>& operator<<(TStringBuilderBase<CharType>& Builder, const FPayloadId& Id)
-{
-	Id.ToString(Builder);
-	return Builder;
-}
-
 ///////////////////////////////////////////////////////////////////////////////////////////////////////////////////////
 
 /**
@@ -119,11 +70,11 @@ public:
 	FPayload() = default;
 
 	/** Construct a payload from the hash of a compressed buffer. */
-	UE_API FPayload(const FPayloadId& Id, const FIoHash& BufferHash);
+	inline FPayload(const FPayloadId& Id, const FIoHash& BufferHash);
 	/** Construct a payload from a compressed buffer and its hash. See FCompressedBuffer::Compress. */
-	UE_API FPayload(const FPayloadId& Id, const FIoHash& BufferHash, FCompressedBuffer Buffer);
+	inline FPayload(const FPayloadId& Id, const FIoHash& BufferHash, FCompressedBuffer Buffer);
 	/** Construct a payload from a compressed buffer and calculates its hash. See FCompressedBuffer::Compress. */
-	UE_API FPayload(const FPayloadId& Id, FCompressedBuffer Buffer);
+	inline FPayload(const FPayloadId& Id, FCompressedBuffer Buffer);
 
 	/** Returns the PayloadId for the payload. */
 	inline const FPayloadId& GetId() const { return Id; }
@@ -152,6 +103,106 @@ private:
 
 ///////////////////////////////////////////////////////////////////////////////////////////////////////////////////////
 
+inline FPayloadId::FPayloadId(const FMemoryView InId)
+{
+	checkf(InId.GetSize() == sizeof(ByteArray),
+		TEXT("FPayloadId cannot be constructed from a view of %" UINT64_FMT " bytes."), InId.GetSize());
+	FMemory::Memcpy(Bytes, InId.GetData(), sizeof(ByteArray));
+}
+
+inline FPayloadId FPayloadId::FromHash(const FIoHash& Hash)
+{
+	checkf(!Hash.IsZero(), TEXT("PayloadId requires a non-zero hash."));
+	return FPayloadId(MakeMemoryView(Hash.GetBytes()).Left(sizeof(ByteArray)));
+}
+
+inline FPayloadId FPayloadId::FromName(const FAnsiStringView Name)
+{
+	checkf(!Name.IsEmpty(), TEXT("PayloadId requires a non-empty name."));
+	return FPayloadId::FromHash(FIoHash::HashBuffer(Name.GetData(), Name.Len()));
+}
+
+inline FPayloadId FPayloadId::FromName(const FWideStringView Name)
+{
+	return FPayloadId::FromName(FTCHARToUTF8(Name));
+}
+
+///////////////////////////////////////////////////////////////////////////////////////////////////////////////////////
+
+inline FPayload::FPayload(const FPayloadId& InId, const FIoHash& InHash)
+	: Id(InId)
+	, Hash(InHash)
+{
+	checkf(Id.IsValid(), TEXT("A valid PayloadId is required to construct a payload."));
+}
+
+inline FPayload::FPayload(const FPayloadId& InId, const FIoHash& InHash, FCompressedBuffer InBuffer)
+	: Id(InId)
+	, Hash(InHash)
+	, Buffer(MoveTemp(InBuffer))
+{
+	checkf(Id.IsValid(), TEXT("A valid PayloadId is required to construct a payload."));
+	if (Buffer.GetCompressedSize())
+	{
+		checkfSlow(Hash == FIoHash::HashBuffer(Buffer.GetCompressed()),
+			TEXT("Provided compressed data hash %s does not match calculated hash %s"),
+			*Hash.ToString(), *FIoHash::HashBuffer(Buffer.GetCompressed()).ToString());
+		checkf(!Hash.IsZero(), TEXT("A non-empty compressed data buffer must have a non-zero hash."));
+	}
+	else
+	{
+		checkf(Hash.IsZero(), TEXT("A null or empty compressed data buffer must use a hash of zero."));
+	}
+}
+
+inline FPayload::FPayload(const FPayloadId& InId, FCompressedBuffer InBuffer)
+	: Id(InId)
+	, Buffer(MoveTemp(InBuffer))
+{
+	checkf(Id.IsValid(), TEXT("A valid PayloadId is required to construct a payload."));
+	if (Buffer.GetCompressedSize())
+	{
+		Hash = FIoHash::HashBuffer(Buffer.GetCompressed());
+	}
+}
+
+///////////////////////////////////////////////////////////////////////////////////////////////////////////////////////
+
+inline bool operator==(const FPayloadId& A, const FPayloadId& B)
+{
+	return A.GetView().EqualBytes(B.GetView());
+}
+
+inline bool operator!=(const FPayloadId& A, const FPayloadId& B)
+{
+	return !A.GetView().EqualBytes(B.GetView());
+}
+
+inline bool operator<(const FPayloadId& A, const FPayloadId& B)
+{
+	return A.GetView().CompareBytes(B.GetView()) < 0;
+}
+
+inline uint32 GetTypeHash(const FPayloadId& Id)
+{
+	return *reinterpret_cast<const uint32*>(Id.GetView().GetData());
+}
+
+/** Convert the ID to a 24-character hex string. */
+template <typename CharType>
+inline TStringBuilderBase<CharType>& operator<<(TStringBuilderBase<CharType>& Builder, const FPayloadId& Id)
+{
+	UE::String::BytesToHexLower(Id.GetBytes(), Builder);
+	return Builder;
+}
+
+inline bool FPayloadId::IsNull() const
+{
+	return *this == FPayloadId();
+}
+
+///////////////////////////////////////////////////////////////////////////////////////////////////////////////////////
+
 /** Compare payloads by their PayloadId and compressed buffer hash. */
 inline bool operator==(const FPayload& A, const FPayload& B)
 {
@@ -206,5 +257,3 @@ struct FPayloadLessById
 
 } // DerivedData
 } // UE
-
-#undef UE_API
diff --git a/Engine/Source/Developer/DerivedDataCache/Public/DerivedDataRequest.h b/Engine/Source/Developer/DerivedDataCache/Public/DerivedDataRequest.h
--- a/Engine/Source/Developer/DerivedDataCache/Public/DerivedDataRequest.h
+++ b/Engine/Source/Developer/DerivedDataCache/Public/DerivedDataRequest.h
@@ -166,5 +166,3 @@ private:
 
 } // DerivedData
 } // UE
-
-#undef UE_API
diff --git a/Engine/Source/Editor/UnrealEd/Private/EditorDomain/EditorDomainSave.cpp b/Engine/Source/Editor/UnrealEd/Private/EditorDomain/EditorDomainSave.cpp
--- a/Engine/Source/Editor/UnrealEd/Private/EditorDomain/EditorDomainSave.cpp
+++ b/Engine/Source/Editor/UnrealEd/Private/EditorDomain/EditorDomainSave.cpp
@@ -761,8 +761,8 @@ bool FEditorDomainSaveServer::TrySavePackage(const FPackagePath& PackagePath, FS
 		return false;
 	}
 
-	UE::DerivedData::FCacheRecordBuilder RecordBuilder;
-	RecordBuilder.SetKey(GetEditorDomainPackageKey(PackageDigest));
+	UE::DerivedData::ICache& Cache = GetDerivedDataCacheRef().GetCache();
+	UE::DerivedData::FCacheRecordBuilder RecordBuilder = Cache.CreateRecord(GetEditorDomainPackageKey(PackageDigest));
 	TCbWriter<256> MetaData;
 	MetaData.BeginObject();
 	MetaData << "FileSize" << TempBytes.Num();
@@ -770,7 +770,7 @@ bool FEditorDomainSaveServer::TrySavePackage(const FPackagePath& PackagePath, FS
 	RecordBuilder.SetMeta(MetaData.Save().AsObject());
 	RecordBuilder.SetValue(FSharedBuffer::MakeView(TempBytes.GetData(), TempBytes.Num()));
 	UE::DerivedData::FCacheRecord Record = RecordBuilder.Build();
-	GetDerivedDataCacheRef().GetCache().Put(MakeArrayView(&Record, 1), PackagePath.GetDebugName());
+	Cache.Put(MakeArrayView(&Record, 1), PackagePath.GetDebugName());
 	return true;
 }
 
diff --git a/Engine/Source/Editor/UnrealEd/Private/EditorDomain/EditorDomainUtils.cpp b/Engine/Source/Editor/UnrealEd/Private/EditorDomain/EditorDomainUtils.cpp
--- a/Engine/Source/Editor/UnrealEd/Private/EditorDomain/EditorDomainUtils.cpp
+++ b/Engine/Source/Editor/UnrealEd/Private/EditorDomain/EditorDomainUtils.cpp
@@ -70,8 +70,9 @@ EPackageDigestResult GetPackageDigest(IAssetRegistry& AssetRegistry, FName Packa
 
 UE::DerivedData::FCacheKey GetEditorDomainPackageKey(const FPackageDigest& PackageDigest)
 {
-	static UE::DerivedData::FCacheBucket EditorDomainPackageCacheBucket(EditorDomainPackageBucketName);
-	return UE::DerivedData::FCacheKey(EditorDomainPackageCacheBucket, PackageDigest);
+	static UE::DerivedData::FCacheBucket EditorDomainPackageCacheBucket =
+		GetDerivedDataCacheRef().GetCache().CreateBucket(EditorDomainPackageBucketName);
+	return UE::DerivedData::FCacheKey{EditorDomainPackageCacheBucket, PackageDigest};
 }
 
 UE::DerivedData::FRequest RequestEditorDomainPackage(const FPackagePath& PackagePath,
diff --git a/Engine/Source/Editor/UnrealEd/Private/EditorDomain/EditorDomainUtils.h b/Engine/Source/Editor/UnrealEd/Private/EditorDomain/EditorDomainUtils.h
--- a/Engine/Source/Editor/UnrealEd/Private/EditorDomain/EditorDomainUtils.h
+++ b/Engine/Source/Editor/UnrealEd/Private/EditorDomain/EditorDomainUtils.h
@@ -14,7 +14,7 @@ namespace UE
 {
 namespace DerivedData
 {
-	class FCacheKey;
+struct FCacheKey;
 }
 }
 
