DDC: Clean up a few details in the new Cache API and remove the last exported type from DDC

#rnx

diff --git a/Engine/Source/Developer/DerivedDataCache/Private/DerivedDataCache.cpp b/Engine/Source/Developer/DerivedDataCache/Private/DerivedDataCache.cpp
--- a/Engine/Source/Developer/DerivedDataCache/Private/DerivedDataCache.cpp
+++ b/Engine/Source/Developer/DerivedDataCache/Private/DerivedDataCache.cpp
@@ -617,6 +617,11 @@ public:
 		FDerivedDataBackend::Get().GetDirectories(OutResults);
 	}
 
+	virtual IDDCCleanup* GetCleanup() const override
+	{
+		return FDDCCleanup::Get();
+	}
+
 	virtual void GatherUsageStats(TMap<FString, FDerivedDataCacheUsageStats>& UsageStatsMap) override
 	{
 		FDerivedDataBackend::Get().GatherUsageStats(UsageStatsMap);
@@ -723,9 +728,9 @@ class FCache final : public ICache
 public:
 	virtual ~FCache() = default;
 
-	virtual FCacheBucket CreateBucket(FStringView Name) final { return CreateCacheBucket(Name); }
+	virtual FCacheBucket CreateBucket(FStringView Name) final { return Private::CreateCacheBucket(Name); }
 
-	virtual FCacheRecordBuilder CreateRecord(const FCacheKey& Key) final { return CreateCacheRecordBuilder(Key); }
+	virtual FCacheRecordBuilder CreateRecord(const FCacheKey& Key) final { return Private::CreateCacheRecordBuilder(Key); }
 
 	virtual FRequest Put(
 		TArrayView<FCacheRecord> Records,
diff --git a/Engine/Source/Developer/DerivedDataCache/Private/DerivedDataCacheKey.cpp b/Engine/Source/Developer/DerivedDataCache/Private/DerivedDataCacheKey.cpp
--- a/Engine/Source/Developer/DerivedDataCache/Private/DerivedDataCacheKey.cpp
+++ b/Engine/Source/Developer/DerivedDataCache/Private/DerivedDataCacheKey.cpp
@@ -10,6 +10,8 @@ namespace UE
 {
 namespace DerivedData
 {
+namespace Private
+{
 
 ///////////////////////////////////////////////////////////////////////////////////////////////////////////////////////
 
@@ -107,10 +109,11 @@ FCacheBucket CreateCacheBucket(FStringView Name)
 		TEXT("Invalid cache bucket name '%.*s' must be alphanumeric, non-empty, and contain fewer than 256 characters."),
 		Name.Len(), Name.GetData());
 	static FCacheBuckets Buckets;
-	return FCacheBucket(Private::FCacheBucketName{Buckets.FindOrAdd(Name)});
+	return FCacheBucket(FCacheBucketName{Buckets.FindOrAdd(Name)});
 }
 
 ///////////////////////////////////////////////////////////////////////////////////////////////////////////////////////
 
+} // Private
 } // DerivedData
 } // UE
diff --git a/Engine/Source/Developer/DerivedDataCache/Private/DerivedDataCachePrivate.h b/Engine/Source/Developer/DerivedDataCache/Private/DerivedDataCachePrivate.h
--- a/Engine/Source/Developer/DerivedDataCache/Private/DerivedDataCachePrivate.h
+++ b/Engine/Source/Developer/DerivedDataCache/Private/DerivedDataCachePrivate.h
@@ -13,6 +13,9 @@ class FCacheBucket;
 class FCacheRecordBuilder;
 struct FCacheKey;
 
+namespace Private
+{
+
 ///////////////////////////////////////////////////////////////////////////////////////////////////////////////////////
 
 // Implemented in DerivedDataCacheKey.cpp
@@ -23,5 +26,6 @@ FCacheRecordBuilder CreateCacheRecordBuilder(const FCacheKey& Key);
 
 ///////////////////////////////////////////////////////////////////////////////////////////////////////////////////////
 
+} // Private
 } // DerivedData
 } // UE
diff --git a/Engine/Source/Developer/DerivedDataCache/Private/DerivedDataCacheRecord.cpp b/Engine/Source/Developer/DerivedDataCache/Private/DerivedDataCacheRecord.cpp
--- a/Engine/Source/Developer/DerivedDataCache/Private/DerivedDataCacheRecord.cpp
+++ b/Engine/Source/Developer/DerivedDataCache/Private/DerivedDataCacheRecord.cpp
@@ -211,15 +211,13 @@ FRequest FCacheRecordBuilderInternal::BuildAsync(FOnCacheRecordComplete&& Callba
 
 ///////////////////////////////////////////////////////////////////////////////////////////////////////////////////////
 
-} // Private
-} // DerivedData
-} // UE
-
-///////////////////////////////////////////////////////////////////////////////////////////////////////////////////////
-
-UE::DerivedData::FCacheRecordBuilder UE::DerivedData::CreateCacheRecordBuilder(const FCacheKey& Key)
+FCacheRecordBuilder CreateCacheRecordBuilder(const FCacheKey& Key)
 {
-	return FCacheRecordBuilder(new Private::FCacheRecordBuilderInternal(Key));
+	return FCacheRecordBuilder(new FCacheRecordBuilderInternal(Key));
 }
 
 ///////////////////////////////////////////////////////////////////////////////////////////////////////////////////////
+
+} // Private
+} // DerivedData
+} // UE
diff --git a/Engine/Source/Developer/DerivedDataCache/Public/DDCCleanup.h b/Engine/Source/Developer/DerivedDataCache/Public/DDCCleanup.h
--- a/Engine/Source/Developer/DerivedDataCache/Public/DDCCleanup.h
+++ b/Engine/Source/Developer/DerivedDataCache/Public/DDCCleanup.h
@@ -8,10 +8,18 @@
 
 struct FFilesystemInfo;
 
+class IDDCCleanup
+{
+public:
+	virtual ~IDDCCleanup() = default;
+	virtual bool IsFinished() const = 0;
+	virtual void WaitBetweenDeletes(bool bWait) = 0;
+};
+
 /** 
  * DDC Filesystem Cache cleanup thread.
  */
-class DERIVEDDATACACHE_API FDDCCleanup : public FRunnable
+class FDDCCleanup final : public FRunnable, public IDDCCleanup
 {
 	/** Singleton instance */
 	static FDDCCleanup* Runnable;
@@ -93,13 +101,13 @@ public:
 	static void Shutdown();
 
 	/** Sets whether cleanup should give up time to other threads between deletes */
-	void WaitBetweenDeletes(bool bWait)
+	virtual void WaitBetweenDeletes(bool bWait) final
 	{
 		bDontWaitBetweenDeletes = !bWait;
 	}
 
 	/** Checks if the cleanup thread is done deleting files */
-	bool IsFinished() const
+	virtual bool IsFinished() const final
 	{
 		return !CleanupList.Num();
 	}
diff --git a/Engine/Source/Developer/DerivedDataCache/Public/DerivedDataCacheInterface.h b/Engine/Source/Developer/DerivedDataCache/Public/DerivedDataCacheInterface.h
--- a/Engine/Source/Developer/DerivedDataCache/Public/DerivedDataCacheInterface.h
+++ b/Engine/Source/Developer/DerivedDataCache/Public/DerivedDataCacheInterface.h
@@ -9,6 +9,7 @@
 #include "Modules/ModuleInterface.h"
 
 class FDerivedDataCacheUsageStats;
+class IDDCCleanup;
 
 namespace UE { namespace DerivedData { class ICache; } }
 
@@ -259,6 +260,11 @@ public:
 	 */
 	virtual const TCHAR* GetGraphName() const = 0;
 
+	/**
+	 * Retrieve the interface to the background cache cleanup.
+	 */
+	virtual IDDCCleanup* GetCleanup() const = 0;
+
 	//--------------------
 	// UsageStats Interface
 	//--------------------
diff --git a/Engine/Source/Developer/DerivedDataCache/Public/DerivedDataCacheRecord.h b/Engine/Source/Developer/DerivedDataCache/Public/DerivedDataCacheRecord.h
--- a/Engine/Source/Developer/DerivedDataCache/Public/DerivedDataCacheRecord.h
+++ b/Engine/Source/Developer/DerivedDataCache/Public/DerivedDataCacheRecord.h
@@ -2,7 +2,6 @@
 
 #pragma once
 
-#include "CoreTypes.h"
 #include "DerivedDataPayload.h"
 #include "DerivedDataRequest.h"
 #include "Memory/MemoryFwd.h"
@@ -197,7 +196,7 @@ public:
 	}
 
 	/**
-	 * Build a cache record, which leaves this builder null.
+	 * Build a cache record, which makes this builder subsequently unusable.
 	 *
 	 * Prefer BuildAsync() when the value or attachments are added from a buffer, as this must block
 	 * on compression of those buffers before it can construct a cache record.
@@ -208,7 +207,7 @@ public:
 	}
 
 	/**
-	 * Build a cache record asynchronously, which leaves this builder null.
+	 * Build a cache record asynchronously, which makes this builder subsequently unusable.
 	 *
 	 * Prefer Build() when the value and attachments are added by payload, as compression is already
 	 * complete and BuildAsync() will complete immediately in that case.
@@ -218,13 +217,6 @@ public:
 		return Builder->BuildAsync(MoveTemp(Callback), Priority);
 	}
 
-	/** Whether this is null. */
-	inline bool IsNull() const { return !Builder; }
-	/** Whether this is not null. */
-	inline bool IsValid() const { return !IsNull(); }
-	/** Whether this is not null. */
-	inline explicit operator bool() const { return IsValid(); }
-
 public:
 	// Internal API
 
diff --git a/Engine/Source/Editor/UnrealEd/Private/Commandlets/DDCCleanupCommandlet.cpp b/Engine/Source/Editor/UnrealEd/Private/Commandlets/DDCCleanupCommandlet.cpp
--- a/Engine/Source/Editor/UnrealEd/Private/Commandlets/DDCCleanupCommandlet.cpp
+++ b/Engine/Source/Editor/UnrealEd/Private/Commandlets/DDCCleanupCommandlet.cpp
@@ -2,19 +2,20 @@
 
 #include "Commandlets/DDCCleanupCommandlet.h"
 #include "DDCCleanup.h"
+#include "DerivedDataCacheInterface.h"
 #include "HAL/PlatformProcess.h"
 
 int32 UDDCCleanupCommandlet::Main(const FString& Params)
 {
-	if (FDDCCleanup::Get())
+	if (IDDCCleanup* Cleanup = GetDerivedDataCacheRef().GetCleanup())
 	{
-		FDDCCleanup::GetNoInit()->WaitBetweenDeletes(false);
+		Cleanup->WaitBetweenDeletes(false);
 		do
-		{			
+		{
 			// Cleanup works from its own thread so just wait until it's done and flush logs once in a while
 			FPlatformProcess::SleepNoStats(1.0f);
 			GLog->Flush();
-		} while (FDDCCleanup::GetNoInit() && !FDDCCleanup::GetNoInit()->IsFinished());
+		} while (!Cleanup->IsFinished());
 	}
 
 	return 0;
