DDC: Clean up a few details in the new Cache API

#rnx

diff --git a/Engine/Source/Developer/DerivedDataCache/Private/DerivedDataCache.cpp b/Engine/Source/Developer/DerivedDataCache/Private/DerivedDataCache.cpp
--- a/Engine/Source/Developer/DerivedDataCache/Private/DerivedDataCache.cpp
+++ b/Engine/Source/Developer/DerivedDataCache/Private/DerivedDataCache.cpp
@@ -18,6 +18,7 @@
 
 #include "DerivedDataBackendInterface.h"
 #include "DerivedDataCachePrivate.h"
+#include "DerivedDataCacheRecord.h"
 #include "DerivedDataPluginInterface.h"
 #include "DDCCleanup.h"
 #include "ProfilingDebugging/CookStats.h"
@@ -787,9 +788,9 @@ private:
 
 	static FString MakePayloadKey(const FCachePayloadKey& PayloadKey)
 	{
-		check(PayloadKey.Record.Bucket.IsValid() && !PayloadKey.Record.Hash.IsZero() && PayloadKey.Payload.IsValid());
+		check(PayloadKey.CacheKey.Bucket.IsValid() && !PayloadKey.CacheKey.Hash.IsZero() && PayloadKey.Id.IsValid());
 		TStringBuilder<96> Out;
-		Out << PayloadKey.Record.Bucket << TEXT("_") << PayloadKey.Record.Hash << TEXT("_") << PayloadKey.Payload;
+		Out << PayloadKey.CacheKey.Bucket << TEXT("_") << PayloadKey.CacheKey.Hash << TEXT("_") << PayloadKey.Id;
 		return FDerivedDataCacheInterface::SanitizeCacheKey(Out.ToString());
 	}
 };
@@ -803,25 +804,25 @@ void FCache::Put(
 	EPriority Priority,
 	FOnCachePutComplete& Callback)
 {
-	FCachePutCompleteParams Params;
-	Params.Key = Record.GetKey();
+	const FCacheKey& Key = Record.GetKey();
+	EStatus Status = EStatus::Error;
 
 	ON_SCOPE_EXIT
 	{
 		if (Callback)
 		{
-			Callback(MoveTemp(Params));
+			Callback(FCachePutCompleteParams{Key, Status});
 		}
 	};
 
 	if (!EnumHasAnyFlags(Policy, ECachePolicy::Store))
 	{
 		UE_LOG(LogDerivedDataCache, VeryVerbose, TEXT("Cache: Put skipped for %s from '%.*s'"),
-			*WriteToString<96>(Params.Key), Context.Len(), Context.GetData());
+			*WriteToString<96>(Key), Context.Len(), Context.GetData());
 		return;
 	}
 
-	const auto PutPayload = [&Key = Params.Key, Context](const FPayload& Payload)
+	const auto PutPayload = [&Key, Context](const FPayload& Payload)
 	{
 		const FSharedBuffer Buffer = Payload.GetBuffer().GetCompressed();
 		check(Buffer && Buffer.GetSize() <= MAX_int32);
@@ -870,11 +871,11 @@ void FCache::Put(
 	Data.SetNumUninitialized(static_cast<int32>(SaveSize));
 	Writer.Save(MakeMemoryView(Data));
 
-	GetDerivedDataCacheRef().Put(*MakeRecordKey(Params.Key), Data, Context);
+	GetDerivedDataCacheRef().Put(*MakeRecordKey(Key), Data, Context);
 
 	UE_LOG(LogDerivedDataCache, VeryVerbose, TEXT("Cache: Put for %s from '%.*s'"),
-		*WriteToString<96>(Params.Key), Context.Len(), Context.GetData());
-	Params.Status = EStatus::Ok;
+		*WriteToString<96>(Key), Context.Len(), Context.GetData());
+	Status = EStatus::Ok;
 }
 
 FRequest FCache::Put(
@@ -901,15 +902,14 @@ void FCache::Get(
 	FPayload Value;
 	TArray<FPayload> Attachments;
 	TArray<FPayloadId, TInlineAllocator<1>> SkippedPayloads;
-	FCacheRecordBuilder Builder = CreateRecord(Key);
+	FCacheRecordBuilder RecordBuilder = CreateRecord(Key);
+	EStatus Status = EStatus::Error;
 
-	FCacheGetCompleteParams Params;
 	ON_SCOPE_EXIT
 	{
 		if (Callback)
 		{
-			Params.Record = Builder.Build();
-			Callback(MoveTemp(Params));
+			Callback(FCacheGetCompleteParams{RecordBuilder.Build(), Status});
 		}
 	};
 
@@ -1017,20 +1017,20 @@ void FCache::Get(
 	// Populate the builder now that any error conditions have been handled.
 	if (!EnumHasAnyFlags(Policy, ECachePolicy::SkipMeta) && MetaObject.CreateViewIterator())
 	{
-		Builder.SetMeta(MoveTemp(MetaObject));
+		RecordBuilder.SetMeta(MoveTemp(MetaObject));
 	}
 	if (Value)
 	{
-		Builder.SetValue(MoveTemp(Value));
+		RecordBuilder.SetValue(MoveTemp(Value));
 	}
 	for (FPayload& Attachment : Attachments)
 	{
-		Builder.AddAttachment(MoveTemp(Attachment));
+		RecordBuilder.AddAttachment(MoveTemp(Attachment));
 	}
 
 	UE_LOG(LogDerivedDataCache, Verbose, TEXT("Cache: Get cache hit for %s from '%.*s'"),
 		*WriteToString<96>(Key), Context.Len(), Context.GetData());
-	Params.Status = EStatus::Ok;
+	Status = EStatus::Ok;
 }
 
 FRequest FCache::Get(
@@ -1054,18 +1054,14 @@ void FCache::GetPayload(
 	EPriority Priority,
 	FOnCacheGetPayloadComplete& Callback)
 {
-	FCacheGetPayloadCompleteParams Params;
-	Params.Key = Key.Record;
+	FPayload Payload(Key.Id, FIoHash());
+	EStatus Status = EStatus::Error;
 
 	ON_SCOPE_EXIT
 	{
 		if (Callback)
 		{
-			if (!Params.Payload)
-			{
-				Params.Payload = FPayload(Key.Payload, FIoHash());
-			}
-			Callback(MoveTemp(Params));
+			Callback(FCacheGetPayloadCompleteParams{Key.CacheKey, MoveTemp(Payload), Status});
 		}
 	};
 
@@ -1089,7 +1085,7 @@ void FCache::GetPayload(
 			FSharedBuffer RecordBuffer;
 			{
 				TArray<uint8> RecordData;
-				if (!GetDerivedDataCacheRef().GetSynchronous(*MakeRecordKey(Key.Record), RecordData, Context))
+				if (!GetDerivedDataCacheRef().GetSynchronous(*MakeRecordKey(Key.CacheKey), RecordData, Context))
 				{
 					UE_LOG(LogDerivedDataCache, Verbose, TEXT("Cache: Get cache miss with missing record for %s from '%.*s'"),
 						*WriteToString<96>(Key), Context.Len(), Context.GetData());
@@ -1112,10 +1108,10 @@ void FCache::GetPayload(
 			const FCbArrayView AttachmentsArray = RecordObject["Attachments"_ASV].AsArrayView();
 
 			// Find the hash of the payload in the record.
-			const FCbObjectId KeyObjectId(Key.Payload.GetView());
+			const FCbObjectId KeyObjectId(Key.Id.GetView());
 			if (ValueObject["Id"_ASV].AsObjectId() == KeyObjectId)
 			{
-				Params.Payload = FPayload(Key.Payload, ValueObject["Hash"_ASV].AsBinaryAttachment());
+				Payload = FPayload(Key.Id, ValueObject["Hash"_ASV].AsBinaryAttachment());
 			}
 			else
 			{
@@ -1123,17 +1119,17 @@ void FCache::GetPayload(
 				{
 					if (AttachmentField.AsObjectView()["Id"_ASV].AsObjectId() == KeyObjectId)
 					{
-						Params.Payload = FPayload(Key.Payload, AttachmentField.AsObjectView()["Hash"_ASV].AsBinaryAttachment());
+						Payload = FPayload(Key.Id, AttachmentField.AsObjectView()["Hash"_ASV].AsBinaryAttachment());
 						break;
 					}
 				}
 			}
 
-			if (Params.Payload)
+			if (Payload)
 			{
 				UE_LOG(LogDerivedDataCache, Verbose, TEXT("Cache: GetPayload (exists) cache hit for %s from '%.*s'"),
 					*WriteToString<96>(Key), Context.Len(), Context.GetData());
-				Params.Status = EStatus::Ok;
+				Status = EStatus::Ok;
 			}
 			else
 			{
@@ -1159,11 +1155,11 @@ void FCache::GetPayload(
 		return;
 	}
 
-	Params.Payload = FPayload(Key.Payload, FCompressedBuffer(MakeSharedBufferFromArray(MoveTemp(Data))));
+	Payload = FPayload(Key.Id, FCompressedBuffer(MakeSharedBufferFromArray(MoveTemp(Data))));
 
 	UE_LOG(LogDerivedDataCache, Verbose, TEXT("Cache: GetPayload cache hit for %s from '%.*s'"),
 		*WriteToString<96>(Key), Context.Len(), Context.GetData());
-	Params.Status = EStatus::Ok;
+	Status = EStatus::Ok;
 }
 
 FRequest FCache::GetPayload(
diff --git a/Engine/Source/Developer/DerivedDataCache/Private/DerivedDataCacheKey.cpp b/Engine/Source/Developer/DerivedDataCache/Private/DerivedDataCacheKey.cpp
--- a/Engine/Source/Developer/DerivedDataCache/Private/DerivedDataCacheKey.cpp
+++ b/Engine/Source/Developer/DerivedDataCache/Private/DerivedDataCacheKey.cpp
@@ -4,6 +4,7 @@
 
 #include "Algo/AllOf.h"
 #include "Containers/Set.h"
+#include "Math/UnrealMathUtility.h"
 #include "Misc/ScopeRWLock.h"
 
 namespace UE
@@ -18,19 +19,19 @@ namespace Private
 class FCacheBuckets
 {
 public:
-	inline const TCHAR* FindOrAdd(FStringView Name)
+	inline const TCHAR* FindOrAdd(FStringView Name16, FAnsiStringView Name8)
 	{
-		const uint32 Hash = GetTypeHash(Name);
+		const uint32 Hash = GetTypeHash(Name16);
 		{
 			FReadScopeLock ReadLock(Lock);
-			if (const FBucket* Bucket = Buckets.FindByHash(Hash, Name))
+			if (const FBucket* Bucket = Buckets.FindByHash(Hash, Name16))
 			{
 				return Bucket->GetName();
 			}
 		}
 		{
 			FWriteScopeLock WriteLock(Lock);
-			const FBucket& Bucket = Buckets.FindOrAddByHash(Hash, FBucket(Name));
+			const FBucket& Bucket = Buckets.FindOrAddByHash(Hash, FBucket(Name16, Name8));
 			return Bucket.GetName();
 		}
 	}
@@ -39,14 +40,25 @@ private:
 	class FBucket
 	{
 	public:
-		inline explicit FBucket(FStringView Bucket)
+		inline FBucket(FStringView Bucket, FAnsiStringView Bucket8)
 		{
 			const int32 BucketLen = Bucket.Len();
-			TCHAR* Buffer = new TCHAR[BucketLen + 2];
-			*reinterpret_cast<uint8*>(Buffer++) = static_cast<uint8>(BucketLen);
-			Bucket.CopyString(Buffer, BucketLen);
-			Buffer[BucketLen] = TEXT('\0');
+			const int32 Bucket8Len = Bucket8.Len();
+			const int32 PrefixSize = FMath::DivideAndRoundUp<int32>(2, sizeof(TCHAR));
+			TCHAR* Buffer = new TCHAR[PrefixSize + BucketLen + 1 + FMath::DivideAndRoundUp<int32>(Bucket8Len + 1, sizeof(TCHAR))] + PrefixSize;
 			Name = Buffer;
+
+			reinterpret_cast<uint8*>(Buffer)[-2] = static_cast<uint8>(BucketLen);
+			reinterpret_cast<uint8*>(Buffer)[-1] = static_cast<uint8>(Bucket8Len);
+
+			Bucket.CopyString(Buffer, BucketLen);
+			Buffer += BucketLen;
+			*Buffer++ = TEXT('\0');
+
+			ANSICHAR* Buffer8 = reinterpret_cast<ANSICHAR*>(Buffer);
+			Bucket8.CopyString(Buffer8, Bucket8Len);
+			Buffer8 += Bucket8Len;
+			*Buffer8++ = '\0';
 		}
 
 		FBucket(const FBucket&) = delete;
@@ -69,13 +81,14 @@ private:
 		{
 			if (Name)
 			{
-				delete[] (Name - 1);
+				const int32 PrefixSize = FMath::DivideAndRoundUp<int32>(2, sizeof(TCHAR));
+				delete[] (Name - PrefixSize);
 			}
 		}
 
 		inline FStringView ToString() const
 		{
-			return FStringView(Name, *reinterpret_cast<const uint8*>(Name - 1));
+			return FStringView(Name, reinterpret_cast<const uint8*>(Name)[-2]);
 		}
 
 		inline bool operator==(const FBucket& Bucket) const
@@ -105,11 +118,12 @@ private:
 
 FCacheBucket CreateCacheBucket(FStringView Name)
 {
-	checkf(!Name.IsEmpty() && Name.Len() < 256 && Algo::AllOf(Name, FChar::IsAlnum),
-		TEXT("Invalid cache bucket name '%.*s' must be alphanumeric, non-empty, and contain fewer than 256 characters."),
+	FTCHARToUTF8 Name8(Name);
+	checkf(!Name.IsEmpty() && Name.Len() < 256 && Name8.Length() < 256 && Algo::AllOf(Name, FChar::IsAlnum),
+		TEXT("Invalid cache bucket name '%.*s' must be alphanumeric, non-empty, and contain fewer than 256 code units."),
 		Name.Len(), Name.GetData());
 	static FCacheBuckets Buckets;
-	return FCacheBucket(FCacheBucketName{Buckets.FindOrAdd(Name)});
+	return FCacheBucket(FCacheBucketName{Buckets.FindOrAdd(Name, Name8)});
 }
 
 ///////////////////////////////////////////////////////////////////////////////////////////////////////////////////////
diff --git a/Engine/Source/Developer/DerivedDataCache/Private/DerivedDataCacheRecord.cpp b/Engine/Source/Developer/DerivedDataCache/Private/DerivedDataCacheRecord.cpp
--- a/Engine/Source/Developer/DerivedDataCache/Private/DerivedDataCacheRecord.cpp
+++ b/Engine/Source/Developer/DerivedDataCache/Private/DerivedDataCacheRecord.cpp
@@ -4,7 +4,7 @@
 
 #include "Containers/UnrealString.h"
 #include "DerivedDataCacheKey.h"
-#include "DerivedDataCachePrivate.h"
+#include "Misc/StringBuilder.h"
 #include "Serialization/CompactBinary.h"
 
 namespace UE
@@ -45,7 +45,7 @@ class FCacheRecordInternal final : public ICacheRecordInternal
 {
 public:
 	FCacheRecordInternal() = default;
-	explicit FCacheRecordInternal(FCacheRecordBuilderInternal&& Builder);
+	explicit FCacheRecordInternal(FCacheRecordBuilderInternal&& RecordBuilder);
 
 	virtual ~FCacheRecordInternal() = default;
 
@@ -84,11 +84,11 @@ static FPayloadId GetOrCreatePayloadId(const FPayloadId& Id, const FSharedBuffer
 
 ///////////////////////////////////////////////////////////////////////////////////////////////////////////////////////
 
-FCacheRecordInternal::FCacheRecordInternal(FCacheRecordBuilderInternal&& Builder)
-	: Key(Builder.Key)
-	, Meta(MoveTemp(Builder.Meta))
-	, Value(MoveTemp(Builder.Value))
-	, Attachments(MoveTemp(Builder.Attachments))
+FCacheRecordInternal::FCacheRecordInternal(FCacheRecordBuilderInternal&& RecordBuilder)
+	: Key(RecordBuilder.Key)
+	, Meta(MoveTemp(RecordBuilder.Meta))
+	, Value(MoveTemp(RecordBuilder.Value))
+	, Attachments(MoveTemp(RecordBuilder.Attachments))
 {
 }
 
diff --git a/Engine/Source/Developer/DerivedDataCache/Public/DerivedDataCache.h b/Engine/Source/Developer/DerivedDataCache/Public/DerivedDataCache.h
--- a/Engine/Source/Developer/DerivedDataCache/Public/DerivedDataCache.h
+++ b/Engine/Source/Developer/DerivedDataCache/Public/DerivedDataCache.h
@@ -4,7 +4,6 @@
 
 #include "CoreTypes.h"
 #include "DerivedDataCacheKey.h"
-#include "DerivedDataCacheRecord.h"
 #include "DerivedDataRequest.h"
 #include "Misc/EnumClassFlags.h"
 
@@ -14,7 +13,9 @@ namespace DerivedData
 {
 
 class FCacheBucket;
+class FCacheRecord;
 class FCacheRecordBuilder;
+class FPayload;
 struct FCacheGetCompleteParams;
 struct FCacheGetPayloadCompleteParams;
 struct FCachePutCompleteParams;
@@ -106,7 +107,7 @@ public:
 	/**
 	 * Create a cache bucket from a name.
 	 *
-	 * A cache bucket name must be alphanumeric, non-empty, and contain fewer than 256 characters.
+	 * A cache bucket name must be alphanumeric, non-empty, and contain fewer than 256 code units.
 	 */
 	virtual FCacheBucket CreateBucket(FStringView Name) = 0;
 
@@ -206,7 +207,7 @@ struct FCacheGetCompleteParams
 	 * The value, attachments, and metadata may be skipped based on cache policy flags. When a value
 	 * or attachment has been skipped, it will have a payload but its buffers will be null.
 	 */
-	FCacheRecord Record;
+	FCacheRecord&& Record;
 
 	/** Status of the cache request. */
 	EStatus Status = EStatus::Error;
@@ -223,7 +224,7 @@ struct FCacheGetPayloadCompleteParams
 	 *
 	 * The ID is always populated. The hash and buffer are populated when Status is Cached.
 	 */
-	FPayload Payload;
+	FPayload&& Payload;
 
 	/** Status of the cache request. */
 	EStatus Status = EStatus::Error;
diff --git a/Engine/Source/Developer/DerivedDataCache/Public/DerivedDataCacheKey.h b/Engine/Source/Developer/DerivedDataCache/Public/DerivedDataCacheKey.h
--- a/Engine/Source/Developer/DerivedDataCache/Public/DerivedDataCacheKey.h
+++ b/Engine/Source/Developer/DerivedDataCache/Public/DerivedDataCacheKey.h
@@ -3,10 +3,10 @@
 #pragma once
 
 #include "CoreTypes.h"
+#include "Containers/StringFwd.h"
 #include "Containers/StringView.h"
 #include "DerivedDataPayload.h"
 #include "IO/IoHash.h"
-#include "Misc/StringBuilder.h"
 #include "Templates/TypeHash.h"
 
 namespace UE
@@ -36,9 +36,6 @@ public:
 	/** Construct a null cache bucket. */
 	FCacheBucket() = default;
 
-	/** Get the name of the cache bucket as a string. */
-	inline FStringView ToString() const;
-
 	/** Whether this is null. */
 	inline bool IsNull() const { return !Name; }
 	/** Whether this is not null. */
@@ -52,7 +49,7 @@ public:
 
 	inline bool operator<(FCacheBucket Other) const
 	{
-		return Name != Other.Name && ToString().Compare(Other.ToString(), ESearchCase::IgnoreCase) < 0;
+		return Name != Other.Name && ToString<TCHAR>().Compare(Other.ToString<TCHAR>(), ESearchCase::IgnoreCase) < 0;
 	}
 
 	friend inline uint32 GetTypeHash(FCacheBucket Bucket)
@@ -60,6 +57,14 @@ public:
 		return ::GetTypeHash(reinterpret_cast<UPTRINT>(Bucket.Name));
 	}
 
+	/** Get the name of the cache bucket as a string. */
+	template <typename CharType>
+	inline TStringView<CharType> ToString() const;
+	template <>
+	inline FStringView ToString<TCHAR>() const;
+	template <>
+	inline FAnsiStringView ToString<ANSICHAR>() const;
+
 public:
 	// Internal API
 
@@ -69,18 +74,33 @@ public:
 	}
 
 private:
-	/** Name with a one-byte length stored in the preceding character. */
+	/**
+	 * Name in TCHAR followed by UTF-8 in ANSICHAR, both null-terminated.
+	 *
+	 * The byte preceding the TCHAR name is the number of UTF-8 code units.
+	 * The byte preceding the UTF-8 name length is the number of TCHAR code units.
+	 */
 	const TCHAR* Name = nullptr;
 };
 
-inline FStringView FCacheBucket::ToString() const
+template <>
+inline FStringView FCacheBucket::ToString<TCHAR>() const
+{
+	return Name ? FStringView(Name, reinterpret_cast<const uint8*>(Name)[-2]) : FStringView();
+}
+
+template <>
+inline FAnsiStringView FCacheBucket::ToString<ANSICHAR>() const
 {
-	return Name ? FStringView(Name, *reinterpret_cast<const uint8*>(Name - 1)) : FStringView();
+	return Name ? FAnsiStringView(
+		reinterpret_cast<const ANSICHAR*>(Name + reinterpret_cast<const uint8*>(Name)[-2] + 1),
+		reinterpret_cast<const uint8*>(Name)[-1]) : FAnsiStringView();
 }
 
-inline FStringBuilderBase& operator<<(FStringBuilderBase& Builder, const FCacheBucket& Bucket)
+template <typename CharType>
+inline TStringBuilderBase<CharType>& operator<<(TStringBuilderBase<CharType>& Builder, const FCacheBucket& Bucket)
 {
-	return Builder << Bucket.ToString();
+	return Builder << Bucket.ToString<CharType>();
 }
 
 ///////////////////////////////////////////////////////////////////////////////////////////////////////////////////////
@@ -114,9 +134,10 @@ inline uint32 GetTypeHash(const FCacheKey& Key)
 	return HashCombine(GetTypeHash(Key.Bucket), GetTypeHash(Key.Hash));
 }
 
-inline FStringBuilderBase& operator<<(FStringBuilderBase& Builder, const FCacheKey& Key)
+template <typename CharType>
+inline TStringBuilderBase<CharType>& operator<<(TStringBuilderBase<CharType>& Builder, const FCacheKey& Key)
 {
-	return Builder << Key.Bucket << TEXT('/') << Key.Hash;
+	return Builder << Key.Bucket << CharType('/') << Key.Hash;
 }
 
 ///////////////////////////////////////////////////////////////////////////////////////////////////////////////////////
@@ -124,35 +145,36 @@ inline FStringBuilderBase& operator<<(FStringBuilderBase& Builder, const FCacheK
 /** A key that uniquely identifies a payload within a cache record. */
 struct FCachePayloadKey
 {
-	FCacheKey Record;
-	FPayloadId Payload;
+	FCacheKey CacheKey;
+	FPayloadId Id;
 };
 
 inline bool operator==(const FCachePayloadKey& A, const FCachePayloadKey& B)
 {
-	return A.Record == B.Record && A.Payload == B.Payload;
+	return A.CacheKey == B.CacheKey && A.Id == B.Id;
 }
 
 inline bool operator!=(const FCachePayloadKey& A, const FCachePayloadKey& B)
 {
-	return A.Record != B.Record || A.Payload != B.Payload;
+	return A.CacheKey != B.CacheKey || A.Id != B.Id;
 }
 
 inline bool operator<(const FCachePayloadKey& A, const FCachePayloadKey& B)
 {
-	const FCacheKey& KeyA = A.Record;
-	const FCacheKey& KeyB = B.Record;
-	return KeyA == KeyB ? A.Payload < B.Payload : KeyA < KeyB;
+	const FCacheKey& KeyA = A.CacheKey;
+	const FCacheKey& KeyB = B.CacheKey;
+	return KeyA == KeyB ? A.Id < B.Id : KeyA < KeyB;
 }
 
 inline uint32 GetTypeHash(const FCachePayloadKey& Key)
 {
-	return HashCombine(GetTypeHash(Key.Record), GetTypeHash(Key.Payload));
+	return HashCombine(GetTypeHash(Key.CacheKey), GetTypeHash(Key.Id));
 }
 
-inline FStringBuilderBase& operator<<(FStringBuilderBase& Builder, const FCachePayloadKey& Key)
+template <typename CharType>
+inline TStringBuilderBase<CharType>& operator<<(TStringBuilderBase<CharType>& Builder, const FCachePayloadKey& Key)
 {
-	return Builder << Key.Record << TEXT('/') << Key.Payload;
+	return Builder << Key.CacheKey << CharType('/') << Key.Id;
 }
 
 ///////////////////////////////////////////////////////////////////////////////////////////////////////////////////////
diff --git a/Engine/Source/Developer/DerivedDataCache/Public/DerivedDataCacheRecord.h b/Engine/Source/Developer/DerivedDataCache/Public/DerivedDataCacheRecord.h
--- a/Engine/Source/Developer/DerivedDataCache/Public/DerivedDataCacheRecord.h
+++ b/Engine/Source/Developer/DerivedDataCache/Public/DerivedDataCacheRecord.h
@@ -147,7 +147,7 @@ public:
 	 */
 	inline void SetMeta(FCbObject&& Meta)
 	{
-		return Builder->SetMeta(MoveTemp(Meta));
+		return RecordBuilder->SetMeta(MoveTemp(Meta));
 	}
 
 	/**
@@ -160,7 +160,7 @@ public:
 	 */
 	inline FPayloadId SetValue(const FSharedBuffer& Buffer, const FPayloadId& Id = FPayloadId())
 	{
-		return Builder->SetValue(Buffer, Id);
+		return RecordBuilder->SetValue(Buffer, Id);
 	}
 
 	/**
@@ -169,7 +169,10 @@ public:
 	 * @param Payload The value payload, which is reset to null.
 	 * @return The ID that was provided. Unique within the scope of the cache record.
 	 */
-	inline FPayloadId SetValue(FPayload&& Payload) { return Builder->SetValue(MoveTemp(Payload)); }
+	inline FPayloadId SetValue(FPayload&& Payload)
+	{
+		return RecordBuilder->SetValue(MoveTemp(Payload));
+	}
 
 	/**
 	 * Add an attachment to the cache record.
@@ -181,7 +184,7 @@ public:
 	 */
 	inline FPayloadId AddAttachment(const FSharedBuffer& Buffer, const FPayloadId& Id = FPayloadId())
 	{
-		return Builder->AddAttachment(Buffer, Id);
+		return RecordBuilder->AddAttachment(Buffer, Id);
 	}
 
 	/**
@@ -192,7 +195,7 @@ public:
 	 */
 	inline FPayloadId AddAttachment(FPayload&& Payload)
 	{
-		return Builder->AddAttachment(MoveTemp(Payload));
+		return RecordBuilder->AddAttachment(MoveTemp(Payload));
 	}
 
 	/**
@@ -203,7 +206,7 @@ public:
 	 */
 	inline FCacheRecord Build()
 	{
-		return Builder->Build();
+		return RecordBuilder->Build();
 	}
 
 	/**
@@ -214,19 +217,19 @@ public:
 	 */
 	inline FRequest BuildAsync(FOnCacheRecordComplete&& Callback, EPriority Priority)
 	{
-		return Builder->BuildAsync(MoveTemp(Callback), Priority);
+		return RecordBuilder->BuildAsync(MoveTemp(Callback), Priority);
 	}
 
 public:
 	// Internal API
 
-	inline explicit FCacheRecordBuilder(Private::ICacheRecordBuilderInternal* InBuilder)
-		: Builder(InBuilder)
+	inline explicit FCacheRecordBuilder(Private::ICacheRecordBuilderInternal* InRecordBuilder)
+		: RecordBuilder(InRecordBuilder)
 	{
 	}
 
 private:
-	TUniquePtr<Private::ICacheRecordBuilderInternal> Builder;
+	TUniquePtr<Private::ICacheRecordBuilderInternal> RecordBuilder;
 };
 
 ///////////////////////////////////////////////////////////////////////////////////////////////////////////////////////
diff --git a/Engine/Source/Developer/DerivedDataCache/Public/DerivedDataPayload.h b/Engine/Source/Developer/DerivedDataCache/Public/DerivedDataPayload.h
--- a/Engine/Source/Developer/DerivedDataCache/Public/DerivedDataPayload.h
+++ b/Engine/Source/Developer/DerivedDataCache/Public/DerivedDataPayload.h
@@ -5,10 +5,10 @@
 #include "CoreTypes.h"
 #include "Compression/CompressedBuffer.h"
 #include "Containers/StringConv.h"
+#include "Containers/StringFwd.h"
 #include "Containers/StringView.h"
 #include "IO/IoHash.h"
 #include "Memory/MemoryView.h"
-#include "Misc/StringBuilder.h"
 #include "String/BytesToHex.h"
 
 namespace UE
diff --git a/Engine/Source/Editor/UnrealEd/Private/EditorDomain/EditorDomain.cpp b/Engine/Source/Editor/UnrealEd/Private/EditorDomain/EditorDomain.cpp
--- a/Engine/Source/Editor/UnrealEd/Private/EditorDomain/EditorDomain.cpp
+++ b/Engine/Source/Editor/UnrealEd/Private/EditorDomain/EditorDomain.cpp
@@ -8,6 +8,7 @@
 #include "Containers/UnrealString.h"
 #include "DerivedDataCache.h"
 #include "DerivedDataCacheInterface.h"
+#include "DerivedDataCacheRecord.h"
 #include "EditorDomain/EditorDomainArchive.h"
 #include "EditorDomain/EditorDomainSave.h"
 #include "EditorDomain/EditorDomainUtils.h"
diff --git a/Engine/Source/Editor/UnrealEd/Private/EditorDomain/EditorDomain.h b/Engine/Source/Editor/UnrealEd/Private/EditorDomain/EditorDomain.h
--- a/Engine/Source/Editor/UnrealEd/Private/EditorDomain/EditorDomain.h
+++ b/Engine/Source/Editor/UnrealEd/Private/EditorDomain/EditorDomain.h
@@ -3,7 +3,6 @@
 #pragma once
 
 #include "AssetRegistry/AssetData.h"
-#include "DerivedDataCache.h"
 #include "HAL/CriticalSection.h"
 #include "HAL/Platform.h"
 #include "IO/IoHash.h"
diff --git a/Engine/Source/Editor/UnrealEd/Private/EditorDomain/EditorDomainArchive.cpp b/Engine/Source/Editor/UnrealEd/Private/EditorDomain/EditorDomainArchive.cpp
--- a/Engine/Source/Editor/UnrealEd/Private/EditorDomain/EditorDomainArchive.cpp
+++ b/Engine/Source/Editor/UnrealEd/Private/EditorDomain/EditorDomainArchive.cpp
@@ -4,6 +4,8 @@
 
 #include "AssetRegistry/AssetData.h"
 #include "Async/AsyncFileHandleNull.h"
+#include "DerivedDataCache.h"
+#include "DerivedDataCacheRecord.h"
 #include "EditorDomain/EditorDomainSave.h"
 #include "HAL/UnrealMemory.h"
 #include "Memory/SharedBuffer.h"
diff --git a/Engine/Source/Editor/UnrealEd/Private/EditorDomain/EditorDomainArchive.h b/Engine/Source/Editor/UnrealEd/Private/EditorDomain/EditorDomainArchive.h
--- a/Engine/Source/Editor/UnrealEd/Private/EditorDomain/EditorDomainArchive.h
+++ b/Engine/Source/Editor/UnrealEd/Private/EditorDomain/EditorDomainArchive.h
@@ -3,10 +3,11 @@
 #pragma once
 
 #include "Async/AsyncFileHandle.h"
-#include "DerivedDataCache.h"
+#include "DerivedDataRequest.h"
 #include "EditorDomain/EditorDomain.h"
 #include "HAL/CriticalSection.h"
 #include "HAL/Platform.h"
+#include "Memory/SharedBuffer.h"
 #include "Misc/PackagePath.h"
 #include "Serialization/Archive.h"
 #include "Templates/RefCounting.h"
@@ -14,6 +15,7 @@
 #include "UObject/PackageResourceManager.h"
 
 class FAssetPackageData;
+namespace UE { namespace DerivedData { struct FCacheGetCompleteParams; } }
 
 /**
  * An Archive that asynchronously waits for the cache request to complete, and reads either from the returned cache bytes
diff --git a/Engine/Source/Editor/UnrealEd/Private/EditorDomain/EditorDomainSave.cpp b/Engine/Source/Editor/UnrealEd/Private/EditorDomain/EditorDomainSave.cpp
--- a/Engine/Source/Editor/UnrealEd/Private/EditorDomain/EditorDomainSave.cpp
+++ b/Engine/Source/Editor/UnrealEd/Private/EditorDomain/EditorDomainSave.cpp
@@ -7,6 +7,7 @@
 #include "Commandlets/EditorDomainSaveCommandlet.h"
 #include "DerivedDataCache.h"
 #include "DerivedDataCacheInterface.h"
+#include "DerivedDataCacheRecord.h"
 #include "Dom/JsonObject.h"
 #include "EditorDomain/EditorDomainUtils.h"
 #include "HAL/FileManager.h"
diff --git a/Engine/Source/Editor/UnrealEd/Private/EditorDomain/EditorDomainUtils.h b/Engine/Source/Editor/UnrealEd/Private/EditorDomain/EditorDomainUtils.h
--- a/Engine/Source/Editor/UnrealEd/Private/EditorDomain/EditorDomainUtils.h
+++ b/Engine/Source/Editor/UnrealEd/Private/EditorDomain/EditorDomainUtils.h
@@ -3,6 +3,7 @@
 #pragma once
 
 #include "Containers/UnrealString.h"
+#include "DerivedDataCache.h"
 #include "EditorDomain/EditorDomain.h"
 
 DECLARE_LOG_CATEGORY_EXTERN(LogEditorDomainSave, Log, All);
@@ -10,13 +11,6 @@ DECLARE_LOG_CATEGORY_EXTERN(LogEditorDomainSave, Log, All);
 class FAssetPackageData;
 class FPackagePath;
 class IAssetRegistry;
-namespace UE
-{
-namespace DerivedData
-{
-struct FCacheKey;
-}
-}
 
 namespace UE
 {
