DDC: Assert that modifications made by PutCachedData will not exceed 2 GiB

#rnx

diff --git a/Engine/Source/Developer/DerivedDataCache/Private/DerivedDataBackendCorruptionWrapper.h b/Engine/Source/Developer/DerivedDataCache/Private/DerivedDataBackendCorruptionWrapper.h
--- a/Engine/Source/Developer/DerivedDataCache/Private/DerivedDataBackendCorruptionWrapper.h
+++ b/Engine/Source/Developer/DerivedDataCache/Private/DerivedDataBackendCorruptionWrapper.h
@@ -199,6 +199,8 @@ public:
 		}
 		COOK_STAT(Timer.AddHit(InData.Num()));
 		// Get rid of the double copy!
+		checkf(int64(InData.Num()) + sizeof(FDerivedDataTrailer) <= INT32_MAX,
+			TEXT("FDerivedDataBackendCorruptionWrapper: appending the trailer makes the data exceed 2 GiB for %s"), CacheKey);
 		TArray<uint8> Data;
 		Data.Reset( InData.Num() + sizeof(FDerivedDataTrailer) );
 		Data.AddUninitialized(InData.Num());
diff --git a/Engine/Source/Developer/DerivedDataCache/Private/DerivedDataLimitKeyLengthWrapper.h b/Engine/Source/Developer/DerivedDataCache/Private/DerivedDataLimitKeyLengthWrapper.h
--- a/Engine/Source/Developer/DerivedDataCache/Private/DerivedDataLimitKeyLengthWrapper.h
+++ b/Engine/Source/Developer/DerivedDataCache/Private/DerivedDataLimitKeyLengthWrapper.h
@@ -207,6 +207,8 @@ public:
 		TArray<uint8> Data(InData.GetData(), InData.Num());
 		check(Data.Num());
 		int32 KeyLen = FCString::Strlen(CacheKey) + 1;
+		checkf(int64(Data.Num()) + KeyLen <= INT32_MAX,
+			TEXT("FDerivedDataLimitKeyLengthWrapper: shortening the key makes the data exceed 2 GiB for %s"), CacheKey);
 		Data.AddUninitialized(KeyLen);
 		FCStringAnsi::Strcpy((char*)&Data[Data.Num() - KeyLen], KeyLen, TCHAR_TO_ANSI(CacheKey));
 		check(Data.Last()==0);
