DDC: Fixed a static analysis warning in the cache bucket constructor

#rnx

diff --git a/Engine/Source/Developer/DerivedDataCache/Private/DerivedDataCacheKey.cpp b/Engine/Source/Developer/DerivedDataCache/Private/DerivedDataCacheKey.cpp
--- a/Engine/Source/Developer/DerivedDataCache/Private/DerivedDataCacheKey.cpp
+++ b/Engine/Source/Developer/DerivedDataCache/Private/DerivedDataCacheKey.cpp
@@ -45,7 +45,8 @@ private:
 			const int32 BucketLen = Bucket.Len();
 			const int32 Bucket8Len = Bucket8.Len();
 			const int32 PrefixSize = FMath::DivideAndRoundUp<int32>(2, sizeof(TCHAR));
-			TCHAR* Buffer = new TCHAR[PrefixSize + BucketLen + 1 + FMath::DivideAndRoundUp<int32>(Bucket8Len + 1, sizeof(TCHAR))] + PrefixSize;
+			TCHAR* Buffer = new TCHAR[PrefixSize + BucketLen + 1 + FMath::DivideAndRoundUp<int32>(Bucket8Len + 1, sizeof(TCHAR))];
+			Buffer += PrefixSize;
 			Name = Buffer;
 
 			reinterpret_cast<uint8*>(Buffer)[-2] = static_cast<uint8>(BucketLen);
