Send DDC payloads with IoHash instead of SHA1

diff --git a/Engine/Source/Developer/DerivedDataCache/Private/HttpDerivedDataBackend.cpp b/Engine/Source/Developer/DerivedDataCache/Private/HttpDerivedDataBackend.cpp
--- a/Engine/Source/Developer/DerivedDataCache/Private/HttpDerivedDataBackend.cpp
+++ b/Engine/Source/Developer/DerivedDataCache/Private/HttpDerivedDataBackend.cpp
@@ -1630,9 +1630,8 @@ bool VerifyRequest(const FRequest* Request, const TArray<uint8>& Payload)
  */
 bool HashPayload(FRequest* Request, const TArrayView<const uint8> Payload)
 {
-	FSHAHash PayloadHash;
-	FSHA1::HashBuffer(Payload.GetData(), Payload.Num(), PayloadHash.Hash);
-	Request->SetHeader(TEXT("X-Jupiter-Sha1"), *PayloadHash.ToString());
+	FIoHash PayloadHash = FIoHash::HashBuffer(Payload.GetData(), Payload.Num());
+	Request->SetHeader(TEXT("X-Jupiter-IoHash"), *PayloadHash.ToString());
 	return true;
 }
 
