Enable DDC1 Jupiter backend GET/PUT pooling and added dedicated main thread pools

#rnx

diff --git a/Engine/Source/Developer/DerivedDataCache/Private/HttpDerivedDataBackend.cpp b/Engine/Source/Developer/DerivedDataCache/Private/HttpDerivedDataBackend.cpp
--- a/Engine/Source/Developer/DerivedDataCache/Private/HttpDerivedDataBackend.cpp
+++ b/Engine/Source/Developer/DerivedDataCache/Private/HttpDerivedDataBackend.cpp
@@ -35,7 +35,7 @@
 // batch requests to reduce the number of concurrent
 // connections.
 #ifndef WITH_DATAREQUEST_HELPER
-	#define WITH_DATAREQUEST_HELPER 0
+	#define WITH_DATAREQUEST_HELPER 1
 #endif
 
 #define UE_HTTPDDC_BACKEND_WAIT_INTERVAL 0.01f
@@ -43,7 +43,7 @@
 #define UE_HTTPDDC_HTTP_REQUEST_TIMOUT_ENABLED 1
 #define UE_HTTPDDC_HTTP_DEBUG 0
 #if WITH_DATAREQUEST_HELPER
-	#define UE_HTTPDDC_GET_REQUEST_POOL_SIZE 4 
+	#define UE_HTTPDDC_GET_REQUEST_POOL_SIZE 16 
 	#define UE_HTTPDDC_PUT_REQUEST_POOL_SIZE 2
 #else
 	#define UE_HTTPDDC_GET_REQUEST_POOL_SIZE 48
@@ -1685,8 +1685,10 @@ FHttpDerivedDataBackend::FHttpDerivedDataBackend(
 #endif
 	if (IsServiceReady() && AcquireAccessToken())
 	{
-		GetRequestPool = MakeUnique<FRequestPool>(InServiceUrl, Access.Get(), UE_HTTPDDC_GET_REQUEST_POOL_SIZE);
-		PutRequestPool = MakeUnique<FRequestPool>(InServiceUrl, Access.Get(), UE_HTTPDDC_PUT_REQUEST_POOL_SIZE);
+		GetRequestPools[0] = MakeUnique<FRequestPool>(InServiceUrl, Access.Get(), UE_HTTPDDC_GET_REQUEST_POOL_SIZE);
+		GetRequestPools[1] = MakeUnique<FRequestPool>(InServiceUrl, Access.Get(), 1);
+		PutRequestPools[0] = MakeUnique<FRequestPool>(InServiceUrl, Access.Get(), UE_HTTPDDC_PUT_REQUEST_POOL_SIZE);
+		PutRequestPools[1] = MakeUnique<FRequestPool>(InServiceUrl, Access.Get(), 1);
 		bIsUsable = true;
 	}
 }
@@ -1873,7 +1875,7 @@ bool FHttpDerivedDataBackend::CachedDataProbablyExists(const TCHAR* CacheKey)
 	// Retry request until we get an accepted response or exhaust allowed number of attempts.
 	for (int32 Attempts = 0; Attempts < UE_HTTPDDC_MAX_ATTEMPTS; ++Attempts)
 	{
-		FDataRequestHelper RequestHelper(GetRequestPool.Get(), *Namespace, *DefaultBucket, CacheKey, nullptr);
+		FDataRequestHelper RequestHelper(GetRequestPools[IsInGameThread()].Get(), *Namespace, *DefaultBucket, CacheKey, nullptr);
 		const int64 ResponseCode = RequestHelper.GetResponseCode();
 
 		if (FRequest::IsSuccessResponse(ResponseCode) && RequestHelper.IsSuccess())
@@ -1893,7 +1895,7 @@ bool FHttpDerivedDataBackend::CachedDataProbablyExists(const TCHAR* CacheKey)
 	// Retry request until we get an accepted response or exhaust allowed number of attempts.
 	for (int32 Attempts = 0; Attempts < UE_HTTPDDC_MAX_ATTEMPTS; ++Attempts)
 	{
-		FScopedRequestPtr Request(GetRequestPool.Get());
+		FScopedRequestPtr Request(GetRequestPools[IsInGameThread()].Get());
 		const FRequest::Result Result = Request->PerformBlockingQuery<FRequest::Head>(*Uri);
 		const int64 ResponseCode = Request->GetResponseCode();
 
@@ -1926,7 +1928,7 @@ TBitArray<> FHttpDerivedDataBackend::CachedDataProbablyExistsBatch(TConstArrayVi
 #if WITH_DATAREQUEST_HELPER
 	for (int32 Attempts = 0; Attempts < UE_HTTPDDC_MAX_ATTEMPTS; ++Attempts)
 	{
-		FDataRequestHelper RequestHelper(GetRequestPool.Get(), *Namespace, *DefaultBucket, CacheKeys);
+		FDataRequestHelper RequestHelper(GetRequestPools[IsInGameThread()].Get(), *Namespace, *DefaultBucket, CacheKeys);
 		const int64 ResponseCode = RequestHelper.GetResponseCode();
 
 		if (FRequest::IsSuccessResponse(ResponseCode) && RequestHelper.IsSuccess())
@@ -1960,7 +1962,7 @@ TBitArray<> FHttpDerivedDataBackend::CachedDataProbablyExistsBatch(TConstArrayVi
 	// Retry request until we get an accepted response or exhaust allowed number of attempts.
 	for (int32 Attempts = 0; Attempts < UE_HTTPDDC_MAX_ATTEMPTS; ++Attempts)
 	{
-		FScopedRequestPtr Request(GetRequestPool.Get());
+		FScopedRequestPtr Request(GetRequestPools[IsInGameThread()].Get());
 		const FRequest::Result Result = Request->PerformBlockingUpload<FRequest::PostJson>(Uri, BodyView);
 		const int64 ResponseCode = Request->GetResponseCode();
 
@@ -2008,7 +2010,7 @@ bool FHttpDerivedDataBackend::GetCachedData(const TCHAR* CacheKey, TArray<uint8>
 	// Retry request until we get an accepted response or exhaust allowed number of attempts.
 	for (int32 Attempts = 0; Attempts < UE_HTTPDDC_MAX_ATTEMPTS; ++Attempts)
 	{
-		FDataRequestHelper RequestHelper(GetRequestPool.Get(), *Namespace, *DefaultBucket, CacheKey, &OutData);
+		FDataRequestHelper RequestHelper(GetRequestPools[IsInGameThread()].Get(), *Namespace, *DefaultBucket, CacheKey, &OutData);
 		const int64 ResponseCode = RequestHelper.GetResponseCode();
 
 		if (FRequest::IsSuccessResponse(ResponseCode) && RequestHelper.IsSuccess())
@@ -2029,7 +2031,7 @@ bool FHttpDerivedDataBackend::GetCachedData(const TCHAR* CacheKey, TArray<uint8>
 	// Retry request until we get an accepted response or exhaust allowed number of attempts.
 	for (int32 Attempts = 0; Attempts < UE_HTTPDDC_MAX_ATTEMPTS; ++Attempts)
 	{
-		FScopedRequestPtr Request(GetRequestPool.Get());
+		FScopedRequestPtr Request(GetRequestPools[IsInGameThread()].Get());
 		if (Request.IsValid())
 		{
 			FRequest::Result Result = Request->PerformBlockingDownload(*Uri, &OutData);
@@ -2062,7 +2064,7 @@ FDerivedDataBackendInterface::EPutStatus FHttpDerivedDataBackend::PutCachedData(
 #if WITH_DATAREQUEST_HELPER
 	for (int32 Attempts = 0; Attempts < UE_HTTPDDC_MAX_ATTEMPTS; ++Attempts)
 	{
-		FDataUploadHelper Request(PutRequestPool.Get(), *Namespace, *DefaultBucket,	CacheKey, InData, UsageStats);
+		FDataUploadHelper Request(PutRequestPools[IsInGameThread()].Get(), *Namespace, *DefaultBucket,	CacheKey, InData, UsageStats);
 		const int64 ResponseCode = Request.GetResponseCode();
 
 		if (Request.IsSuccess() && (Request.IsQueued() || FRequest::IsSuccessResponse(ResponseCode)))
@@ -2084,7 +2086,7 @@ FDerivedDataBackendInterface::EPutStatus FHttpDerivedDataBackend::PutCachedData(
 	// Retry request until we get an accepted response or exhaust allowed number of attempts.
 	while (ResponseCode == 0 && ++Attempts < UE_HTTPDDC_MAX_ATTEMPTS)
 	{
-		FScopedRequestPtr Request(PutRequestPool.Get());
+		FScopedRequestPtr Request(PutRequestPools[IsInGameThread()].Get());
 		if (Request.IsValid())
 		{
 			// Append the content hash to the header
@@ -2126,7 +2128,7 @@ void FHttpDerivedDataBackend::RemoveCachedData(const TCHAR* CacheKey, bool bTran
 	// Retry request until we get an accepted response or exhaust allowed number of attempts.
 	while (ResponseCode == 0 && ++Attempts < UE_HTTPDDC_MAX_ATTEMPTS)
 	{
-		FScopedRequestPtr Request(PutRequestPool.Get());
+		FScopedRequestPtr Request(PutRequestPools[IsInGameThread()].Get());
 		if (Request.IsValid())
 		{
 			FRequest::Result Result = Request->PerformBlockingQuery<FRequest::Delete>(*Uri);
diff --git a/Engine/Source/Developer/DerivedDataCache/Private/HttpDerivedDataBackend.h b/Engine/Source/Developer/DerivedDataCache/Private/HttpDerivedDataBackend.h
--- a/Engine/Source/Developer/DerivedDataCache/Private/HttpDerivedDataBackend.h
+++ b/Engine/Source/Developer/DerivedDataCache/Private/HttpDerivedDataBackend.h
@@ -73,8 +73,8 @@ private:
 	FString OAuthSecret;
 	FCriticalSection AccessCs;
 	FDerivedDataCacheUsageStats UsageStats;
-	TUniquePtr<struct FRequestPool> GetRequestPool;
-	TUniquePtr<struct FRequestPool> PutRequestPool;
+	TUniquePtr<struct FRequestPool> GetRequestPools[2];
+	TUniquePtr<struct FRequestPool> PutRequestPools[2];
 	TUniquePtr<struct FHttpAccessToken> Access;
 	bool bIsUsable;
 	uint32 FailedLoginAttempts;
