Added option to override Jupiter speed class

diff --git a/Engine/Source/Developer/DerivedDataCache/Private/DerivedDataBackends.cpp b/Engine/Source/Developer/DerivedDataCache/Private/DerivedDataBackends.cpp
--- a/Engine/Source/Developer/DerivedDataCache/Private/DerivedDataBackends.cpp
+++ b/Engine/Source/Developer/DerivedDataCache/Private/DerivedDataBackends.cpp
@@ -712,6 +712,25 @@ public:
 			}
 		}
 
+		FDerivedDataBackendInterface::ESpeedClass ForceSpeedClass = FDerivedDataBackendInterface::ESpeedClass::Unknown;
+		FString ForceSpeedClassValue;
+		if (FParse::Value(FCommandLine::Get(), TEXT("HttpForceSpeedClass="), ForceSpeedClassValue))
+		{
+			UE_LOG(LogDerivedDataCache, Warning, TEXT("Jupiter speed class overridden due to HttpForceSpeedClass=%s"), *ForceSpeedClassValue);
+			if (ForceSpeedClassValue == TEXT("Slow"))
+			{
+				ForceSpeedClass = FDerivedDataBackendInterface::ESpeedClass::Slow;
+			}
+			else if (ForceSpeedClassValue == TEXT("Fast"))
+			{
+				ForceSpeedClass = FDerivedDataBackendInterface::ESpeedClass::Fast;
+			}
+			else if (ForceSpeedClassValue == TEXT("Local"))
+			{
+				ForceSpeedClass = FDerivedDataBackendInterface::ESpeedClass::Local;
+			}
+		}
+		
 		FString Namespace;
 		if (!FParse::Value(Entry, TEXT("Namespace="), Namespace))
 		{
@@ -741,6 +760,12 @@ public:
 		}
 
 		FHttpDerivedDataBackend* backend = new FHttpDerivedDataBackend(*ServiceUrl, *Namespace, *OAuthProvider, *OAuthClientId, *OAuthSecret);
+
+		if (ForceSpeedClass != FDerivedDataBackendInterface::ESpeedClass::Unknown)
+		{
+			backend->SetSpeedClass(ForceSpeedClass);
+		}
+
 		if (!backend->IsUsable())
 		{
 			UE_LOG(LogDerivedDataCache, Warning, TEXT("%s could not contact the service (%s), will not use it."), NodeName, *ServiceUrl);
diff --git a/Engine/Source/Developer/DerivedDataCache/Private/HttpDerivedDataBackend.cpp b/Engine/Source/Developer/DerivedDataCache/Private/HttpDerivedDataBackend.cpp
--- a/Engine/Source/Developer/DerivedDataCache/Private/HttpDerivedDataBackend.cpp
+++ b/Engine/Source/Developer/DerivedDataCache/Private/HttpDerivedDataBackend.cpp
@@ -1679,6 +1679,7 @@ FHttpDerivedDataBackend::FHttpDerivedDataBackend(
 	, Access(nullptr)
 	, bIsUsable(false)
 	, FailedLoginAttempts(0)
+	, SpeedClass(ESpeedClass::Slow)
 {
 #if WITH_DATAREQUEST_HELPER
 	FDataRequestHelper::StaticInitialize();
@@ -1714,7 +1715,7 @@ bool FHttpDerivedDataBackend::WouldCache(const TCHAR* CacheKey, TArrayView<const
 
 FHttpDerivedDataBackend::ESpeedClass FHttpDerivedDataBackend::GetSpeedClass()
 {
-	return ESpeedClass::Slow;
+	return SpeedClass;
 }
 
 bool FHttpDerivedDataBackend::ApplyDebugOptions(FBackendDebugOptions& InOptions)
diff --git a/Engine/Source/Developer/DerivedDataCache/Private/HttpDerivedDataBackend.h b/Engine/Source/Developer/DerivedDataCache/Private/HttpDerivedDataBackend.h
--- a/Engine/Source/Developer/DerivedDataCache/Private/HttpDerivedDataBackend.h
+++ b/Engine/Source/Developer/DerivedDataCache/Private/HttpDerivedDataBackend.h
@@ -63,6 +63,8 @@ public:
 	virtual ESpeedClass GetSpeedClass() override;
 	virtual bool ApplyDebugOptions(FBackendDebugOptions& InOptions) override;
 
+	void SetSpeedClass(ESpeedClass InSpeedClass) { SpeedClass = InSpeedClass; }
+
 private:
 	
 	FString Domain;
@@ -78,6 +80,7 @@ private:
 	TUniquePtr<struct FHttpAccessToken> Access;
 	bool bIsUsable;
 	uint32 FailedLoginAttempts;
+	ESpeedClass SpeedClass;
 
 	bool IsServiceReady();
 	bool AcquireAccessToken();
