DDC - Don't look up the store for DDC paths that don't have environment variable overrides.

diff --git a/Engine/Source/Developer/DerivedDataCache/Private/DerivedDataBackends.cpp b/Engine/Source/Developer/DerivedDataCache/Private/DerivedDataBackends.cpp
--- a/Engine/Source/Developer/DerivedDataCache/Private/DerivedDataBackends.cpp
+++ b/Engine/Source/Developer/DerivedDataCache/Private/DerivedDataBackends.cpp
@@ -520,6 +520,7 @@ public:
 			}
 		}
 
+		if (!EnvPathOverride.IsEmpty())
 		{
 			FString DDCPath;
 			if (FPlatformMisc::GetStoredValue(TEXT("Epic Games"), TEXT("GlobalDataCachePath"), *EnvPathOverride, DDCPath))
@@ -693,6 +694,7 @@ public:
 				}
 			}
 
+			if (!EnvPathOverride.IsEmpty())
 			{
 				FString DDCPath;
 				if (FPlatformMisc::GetStoredValue(TEXT("Epic Games"), TEXT("GlobalDataCachePath"), *EnvPathOverride, DDCPath))
