Fix HTTP/Jupiter DDC backend to account for the fact that Jupiter will return batched results in arbitrary order, not the same order that the requests/keys were listed.

diff --git a/Engine/Source/Developer/DerivedDataCache/Private/HttpDerivedDataBackend.cpp b/Engine/Source/Developer/DerivedDataCache/Private/HttpDerivedDataBackend.cpp
--- a/Engine/Source/Developer/DerivedDataCache/Private/HttpDerivedDataBackend.cpp
+++ b/Engine/Source/Developer/DerivedDataCache/Private/HttpDerivedDataBackend.cpp
@@ -1218,6 +1218,53 @@ private:
 		}
 	}
 
+	// Searches for potentially multiple key requests that are satisfied the given cache key result
+	// Search strategy is exhaustive forward search from the last found entry.  If the results come in ordered the same as the requests,
+	//  and there are no duplicates, the search will be somewhat efficient (still has to do exhaustive searching looking for duplicates).
+	//  If the results are unordered or there are duplicates, search will become more inefficient.
+	struct FRequestSearchHelper
+	{
+		FRequestSearchHelper(TArrayView<FQueuedBatchEntry> InRequests, const FUTF8ToTCHAR& InCacheKey, int32 InEntryIdx, int32 InKeyIdx)
+			: Requests(InRequests)
+			, CacheKey(InCacheKey)
+			, StartEntryIdx(InEntryIdx)
+			, StartKeyIdx(InKeyIdx)
+		{}
+
+		bool FindNext(int32& EntryIdx, int32& KeyIdx)
+		{
+			int32 CurrentEntryIdx = EntryIdx;
+			int32 CurrentKeyIdx = KeyIdx;
+			do
+			{
+				if (FCString::Stricmp(Requests[CurrentEntryIdx].CacheKeys[CurrentKeyIdx], CacheKey.Get()) == 0)
+				{
+					EntryIdx = CurrentEntryIdx;
+					KeyIdx = CurrentKeyIdx;
+					return true;
+				}
+			} while (AdvanceIndices(CurrentEntryIdx, CurrentKeyIdx));
+
+			return false;
+		}
+
+		bool AdvanceIndices(int32& EntryIdx, int32& KeyIdx)
+		{
+			if (++KeyIdx >= Requests[EntryIdx].CacheKeys.Num())
+			{
+				EntryIdx = (EntryIdx + 1) % Requests.Num();
+				KeyIdx = 0;
+			}
+
+			return !((EntryIdx == StartEntryIdx) && (KeyIdx == StartKeyIdx));
+		}
+
+		TArrayView<FQueuedBatchEntry> Requests;
+		const FUTF8ToTCHAR& CacheKey;
+		int32 StartEntryIdx;
+		int32 StartKeyIdx;
+	};
+
 	/**
 	 * Parses a batched response stream, moves the data to target requests and marks them with result.
 	 * @param Response Pointer to Response buffer
@@ -1271,88 +1318,101 @@ private:
 			Response += FCStringAnsi::Strlen(PayloadNameA) + 1; //String and zero termination
 			const ANSICHAR* CacheKeyA = FCStringAnsi::Strrchr(PayloadNameA, '.') + 1; // "namespace.bucket.cachekey"
 
-			// Find the payload among the requests. Assume payloads are recieved in the same order they
-			// were requested.
+			// Find the payload among the requests.  Payloads may be returned in any order and if the same cache key was part of two requests,
+			// a single payload may satisfy multiple cache keys in multiple requests.
 			FUTF8ToTCHAR CacheKey(CacheKeyA);
-			if (Requests.Num() < EntryIdx || Requests[EntryIdx].CacheKeys.Num() < KeyIdx || FCString::Stricmp(Requests[EntryIdx].CacheKeys[KeyIdx], CacheKey.Get()) != 0)
+			FRequestSearchHelper RequestSearch(Requests, CacheKey, EntryIdx, KeyIdx);
+			bool bFoundAny = false;
+
+			const uint8* ResponseRewindMark = Response;
+			while (RequestSearch.FindNext(EntryIdx, KeyIdx))
 			{
-				UE_LOG(LogDerivedDataCache, Error, ResponseErrorMessage);
-				return false;
-			}
+				Response = ResponseRewindMark;
+				bFoundAny = true;
 
-			FQueuedBatchEntry& RequestOp = Requests[EntryIdx];
-			TBitArray<>& bSuccess = *RequestOp.bSuccess;
+				FQueuedBatchEntry& RequestOp = Requests[EntryIdx];
+				TBitArray<>& bSuccess = *RequestOp.bSuccess;
 
-			// Result of the operation
-			uint8 PayloadResult = *Response;
-			Response += sizeof(uint8);
+				// Result of the operation
+				uint8 PayloadResult = *Response;
+				Response += sizeof(uint8);
 
-			switch (PayloadResult)
-			{
-			
-			case OpResult::Ok:
+				switch (PayloadResult)
 				{
-					TArray<uint8>* OutData = RequestOp.OutDatas[KeyIdx];
+				
+				case OpResult::Ok:
+					{
+						TArray<uint8>* OutData = RequestOp.OutDatas[KeyIdx];
 
-					// Payload hash of the following payload data
-					FIoHash PayloadHash = *(FIoHash*)Response;
-					Response += sizeof(FIoHash);
+						// Payload hash of the following payload data
+						FIoHash PayloadHash = *(FIoHash*)Response;
+						Response += sizeof(FIoHash);
 
-					// Size of the following payload data
-					const uint64 PayloadSize = *(uint64*)Response;
-					Response += sizeof(uint64);
+						// Size of the following payload data
+						const uint64 PayloadSize = *(uint64*)Response;
+						Response += sizeof(uint64);
 
-					if (PayloadSize > 0)
-					{
-						OutData->Append(Response, PayloadSize);
-						Response += PayloadSize;
-						// Verify the recieved and parsed payload
-						if (VerifyPayload(PayloadHash, *OutData))
+						if (PayloadSize > 0)
 						{
-							TRACE_COUNTER_ADD(HttpDDC_GetHit, int64(1));
-							TRACE_COUNTER_ADD(HttpDDC_BytesReceived, int64(PayloadSize));
-							
-							bSuccess[KeyIdx] = true;
+							if (bSuccess[KeyIdx])
+							{
+								Response += PayloadSize;
+							}
+							else
+							{
+								OutData->Append(Response, PayloadSize);
+								Response += PayloadSize;
+								// Verify the recieved and parsed payload
+								if (VerifyPayload(PayloadHash, *OutData))
+								{
+									TRACE_COUNTER_ADD(HttpDDC_GetHit, int64(1));
+									TRACE_COUNTER_ADD(HttpDDC_BytesReceived, int64(PayloadSize));
+									
+									bSuccess[KeyIdx] = true;
+								}
+								else
+								{
+									OutData->Empty();
+									bSuccess[KeyIdx] = false;
+								}
+							}
 						}
 						else
 						{
-							OutData->Empty();
 							bSuccess[KeyIdx] = false;
 						}
 					}
-					else
+					break;
+
+				case OpResult::Exists:
 					{
-						bSuccess[KeyIdx] = false;
+						TRACE_COUNTER_ADD(HttpDDC_ExistHit, int64(1));
+						bSuccess[KeyIdx] = true;
 					}
-				}
-				break;
+					break;
 
-			case OpResult::Exists:
-				{
-					TRACE_COUNTER_ADD(HttpDDC_ExistHit, int64(1));
-					bSuccess[KeyIdx] = true;
-				}
-				break;
+				default:
+				case OpResult::Error:
+					UE_LOG(LogDerivedDataCache, Display, TEXT("Server error while getting %s"), CacheKey.Get());
+					// intentional falltrough
 
-			default:
-			case OpResult::Error:
-				UE_LOG(LogDerivedDataCache, Display, TEXT("Server error while getting %s"), CacheKey.Get());
-				// intentional falltrough
+				case OpResult::NotFound:
+					bSuccess[KeyIdx] = false;
+					break;
 
-			case OpResult::NotFound:
-				bSuccess[KeyIdx] = false;
-				break;
+				}
 
+				if (!RequestSearch.AdvanceIndices(EntryIdx, KeyIdx))
+				{
+					break;
+				}
 			}
 
-
-			// Advance the request indices
-			if (++KeyIdx >= RequestOp.CacheKeys.Num())
+			if (!bFoundAny)
 			{
-				EntryIdx++;
-				KeyIdx = 0;
+				UE_LOG(LogDerivedDataCache, Error, ResponseErrorMessage);
+				return false;
 			}
-			
 		}
 
 		// Have we parsed all the payloads from the message?
@@ -1970,25 +2030,26 @@ TBitArray<> FHttpDerivedDataBackend::CachedDataProbablyExistsBatch(TConstArrayVi
 		if (Result == FRequest::Success && ResponseCode == 200)
 		{
 			TArray<TSharedPtr<FJsonValue>> ResponseArray = Request->GetResponseAsJsonArray();
-			if (ResponseArray.Num() == CacheKeys.Num())
+
+			TBitArray<> Exists;
+			Exists.Reserve(CacheKeys.Num());
+			for (const FString& CacheKey : CacheKeys)
 			{
-				TBitArray<> Exists;
-				Exists.Reserve(CacheKeys.Num());
-				const FString* CacheKeyIt = CacheKeys.GetData();
-				for (const TSharedPtr<FJsonValue>& Response : ResponseArray)
-				{
+				const TSharedPtr<FJsonValue>* FoundResponse = Algo::FindByPredicate(ResponseArray, [&CacheKey](const TSharedPtr<FJsonValue>& Response) {
 					FString Key;
-					Exists.Add(Response->TryGetString(Key) && Key == *CacheKeyIt);
-					++CacheKeyIt;
-				}
-				if (Exists.CountSetBits() == CacheKeys.Num())
-				{
-					TRACE_COUNTER_ADD(HttpDDC_ExistHit, int64(1));
-					COOK_STAT(Timer.AddHit(0));
-				}
-				return Exists;
+					Response->TryGetString(Key);
+					return Key == CacheKey;
+				});
+
+				Exists.Add(FoundResponse != nullptr);
 			}
-			break;
+
+			if (Exists.CountSetBits() == CacheKeys.Num())
+			{
+				TRACE_COUNTER_ADD(HttpDDC_ExistHit, int64(1));
+				COOK_STAT(Timer.AddHit(0));
+			}
+			return Exists;
 		}
 
 		if (!ShouldRetryOnError(ResponseCode))
