DDC: Switched to WriteToString in the Jupiter backend

#rnx

diff --git a/Engine/Source/Developer/DerivedDataCache/Private/HttpDerivedDataBackend.cpp b/Engine/Source/Developer/DerivedDataCache/Private/HttpDerivedDataBackend.cpp
--- a/Engine/Source/Developer/DerivedDataCache/Private/HttpDerivedDataBackend.cpp
+++ b/Engine/Source/Developer/DerivedDataCache/Private/HttpDerivedDataBackend.cpp
@@ -1623,8 +1623,8 @@ bool VerifyPayload(const FSHAHash& Hash, const TArray<uint8>& Payload)
 		UE_LOG(LogDerivedDataCache,
 			Warning,
 			TEXT("Checksum from server did not match received data (%s vs %s). Discarding cached result."),
-			*Hash.ToString(),
-			*PayloadHash.ToString()
+			*WriteToString<48>(Hash),
+			*WriteToString<48>(PayloadHash)
 		);
 		return false;
 	}
@@ -1647,8 +1647,8 @@ bool VerifyPayload(const FIoHash& Hash, const TArray<uint8>& Payload)
 		UE_LOG(LogDerivedDataCache,
 			Warning,
 			TEXT("Checksum from server did not match received data (%s vs %s). Discarding cached result."),
-			*Hash.ToString(),
-			*PayloadHash.ToString()
+			*WriteToString<48>(Hash),
+			*WriteToString<48>(PayloadHash)
 		);
 		return false;
 	}
@@ -1691,7 +1691,7 @@ bool VerifyRequest(const FRequest* Request, const TArray<uint8>& Payload)
 bool HashPayload(FRequest* Request, const TArrayView<const uint8> Payload)
 {
 	FIoHash PayloadHash = FIoHash::HashBuffer(Payload.GetData(), Payload.Num());
-	Request->SetHeader(TEXT("X-Jupiter-IoHash"), *PayloadHash.ToString());
+	Request->SetHeader(TEXT("X-Jupiter-IoHash"), *WriteToString<48>(PayloadHash));
 	return true;
 }
 
