DDC: Fixed a compile error inside a checkfSlow

#rnx

diff --git a/Engine/Source/Developer/DerivedDataCache/Public/DerivedDataPayload.h b/Engine/Source/Developer/DerivedDataCache/Public/DerivedDataPayload.h
--- a/Engine/Source/Developer/DerivedDataCache/Public/DerivedDataPayload.h
+++ b/Engine/Source/Developer/DerivedDataCache/Public/DerivedDataPayload.h
@@ -146,7 +146,7 @@ inline FPayload::FPayload(const FPayloadId& InId, const FIoHash& InHash, FCompre
 	{
 		checkfSlow(Hash == FIoHash::HashBuffer(Buffer.GetCompressed()),
 			TEXT("Provided compressed data hash %s does not match calculated hash %s"),
-			*Hash.ToString(), *FIoHash::HashBuffer(Buffer.GetCompressed()).ToString());
+			*LexToString(Hash), *LexToString(FIoHash::HashBuffer(Buffer.GetCompressed())));
 		checkf(!Hash.IsZero(), TEXT("A non-empty compressed data buffer must have a non-zero hash."));
 	}
 	else
