Ensure HTTP DDC backend doesn't delay process shutdown due to a queue of pending PUTs for normal users.
Still allows build machine runs to be thorough about completing PUTs.
Doesn't interrupt any in-flight PUTs, just prevents processing of any remaining queued PUTs.

diff --git a/Engine/Source/Developer/DerivedDataCache/Private/DerivedDataBackends.cpp b/Engine/Source/Developer/DerivedDataCache/Private/DerivedDataBackends.cpp
--- a/Engine/Source/Developer/DerivedDataCache/Private/DerivedDataBackends.cpp
+++ b/Engine/Source/Developer/DerivedDataCache/Private/DerivedDataBackends.cpp
@@ -30,6 +30,7 @@
 #include "Math/UnitConversion.h"
 #include "Internationalization/FastDecimalFormat.h"
 #include "Math/BasicMathExpressionEvaluator.h"
+#include <atomic>
 
 DEFINE_LOG_CATEGORY(LogDerivedDataCache);
 
@@ -56,6 +57,7 @@ public:
 		, KeyLengthWrapper(NULL)
 		, HierarchicalWrapper(NULL)
 		, bUsingSharedDDC(false)
+		, bIsShuttingDown(false)
 		, MountPakCommand(
 		TEXT( "DDC.MountPak" ),
 		*LOCTEXT("CommandText_DDCMountPak", "Mounts read-only pak file").ToString(),
@@ -922,6 +924,12 @@ public:
 	{
 		double StartTime = FPlatformTime::Seconds();
 		double LastPrint = StartTime;
+
+		if (bShutdown)
+		{
+			bIsShuttingDown.store(true, std::memory_order_relaxed);
+		}
+
 		while (AsyncCompletionCounter.GetValue())
 		{
 			check(AsyncCompletionCounter.GetValue() >= 0);
@@ -1011,6 +1019,11 @@ public:
 		return AsyncCompletionCounter.GetValue() > 0;
 	}
 
+	virtual bool IsShuttingDown() override
+	{
+		return bIsShuttingDown.load(std::memory_order_relaxed);
+	}
+
 	virtual void GetDirectories(TArray<FString>& OutResults) override
 	{
 		OutResults = Directories;
@@ -1137,6 +1150,9 @@ private:
 	/** Whether a shared cache is in use */
 	bool bUsingSharedDDC;
 
+	/** Whether a shutdown is pending */
+	std::atomic<bool> bIsShuttingDown;
+
 	/** MountPak console command */
 	FAutoConsoleCommand MountPakCommand;
 	/** UnmountPak console command */
diff --git a/Engine/Source/Developer/DerivedDataCache/Private/HttpDerivedDataBackend.cpp b/Engine/Source/Developer/DerivedDataCache/Private/HttpDerivedDataBackend.cpp
--- a/Engine/Source/Developer/DerivedDataCache/Private/HttpDerivedDataBackend.cpp
+++ b/Engine/Source/Developer/DerivedDataCache/Private/HttpDerivedDataBackend.cpp
@@ -720,6 +720,7 @@ bool VerifyPayload(const FSHAHash& Hash, const TArray<uint8>& Payload);
 bool VerifyPayload(const FIoHash& Hash, const TArray<uint8>& Payload);
 bool VerifyRequest(const class FRequest* Request, const TArray<uint8>& Payload);
 bool HashPayload(class FRequest* Request, const TArrayView<const uint8> Payload);
+bool ShouldAbortForShutdown();
 
 //----------------------------------------------------------------------------------------------------------
 // Request pool
@@ -1509,11 +1510,21 @@ private:
 	{
 		while (Request)
 		{
+			if (ShouldAbortForShutdown())
+			{
+				return;
+			}
+
 			while (FQueuedEntry* Entry = QueuedPuts.Pop())
 			{
 				Request->Reset();
 				PerformPut(Request, *Entry->Namespace, *Entry->Bucket, *Entry->CacheKey, Entry->Data, UsageStats);
 				delete Entry;
+
+				if (ShouldAbortForShutdown())
+				{
+					return;
+				}
 			}
 
 			Pool->ReleaseRequestToPool(Request);
@@ -1695,6 +1706,11 @@ bool HashPayload(FRequest* Request, const TArrayView<const uint8> Payload)
 	return true;
 }
 
+bool ShouldAbortForShutdown()
+{
+	return !GIsBuildMachine && FDerivedDataBackend::Get().IsShuttingDown();
+}
+
 //----------------------------------------------------------------------------------------------------------
 // FHttpAccessToken
 //----------------------------------------------------------------------------------------------------------
@@ -2127,6 +2143,12 @@ FDerivedDataBackendInterface::EPutStatus FHttpDerivedDataBackend::PutCachedData(
 	for (int32 Attempts = 0; Attempts < UE_HTTPDDC_MAX_ATTEMPTS; ++Attempts)
 	{
 		FDataUploadHelper Request(PutRequestPools[IsInGameThread()].Get(), *Namespace, *DefaultBucket,	CacheKey, InData, UsageStats);
+
+		if (ShouldAbortForShutdown())
+		{
+			return EPutStatus::NotCached;
+		}
+
 		const int64 ResponseCode = Request.GetResponseCode();
 
 		if (Request.IsSuccess() && (Request.IsQueued() || FRequest::IsSuccessResponse(ResponseCode)))
diff --git a/Engine/Source/Developer/DerivedDataCache/Public/DerivedDataBackendInterface.h b/Engine/Source/Developer/DerivedDataCache/Public/DerivedDataBackendInterface.h
--- a/Engine/Source/Developer/DerivedDataCache/Public/DerivedDataBackendInterface.h
+++ b/Engine/Source/Developer/DerivedDataCache/Public/DerivedDataBackendInterface.h
@@ -210,6 +210,7 @@ public:
 	virtual void NotifyBootComplete() = 0;
 	virtual void AddToAsyncCompletionCounter(int32 Addend) = 0;
 	virtual bool AnyAsyncRequestsRemaining() = 0;
+	virtual bool IsShuttingDown() = 0;
 	virtual void WaitForQuiescence(bool bShutdown = false) = 0;
 	virtual void GetDirectories(TArray<FString>& OutResults) = 0;
 	virtual bool GetUsingSharedDDC() const = 0;
