Ensure that the non-DataRequest-Helper codepath for putting DDC data into the cache also tears-down promptly during process shutdown if not running on a build machine.

diff --git a/Engine/Source/Developer/DerivedDataCache/Private/HttpDerivedDataBackend.cpp b/Engine/Source/Developer/DerivedDataCache/Private/HttpDerivedDataBackend.cpp
--- a/Engine/Source/Developer/DerivedDataCache/Private/HttpDerivedDataBackend.cpp
+++ b/Engine/Source/Developer/DerivedDataCache/Private/HttpDerivedDataBackend.cpp
@@ -2170,6 +2170,11 @@ FDerivedDataBackendInterface::EPutStatus FHttpDerivedDataBackend::PutCachedData(
 	// Retry request until we get an accepted response or exhaust allowed number of attempts.
 	while (ResponseCode == 0 && ++Attempts < UE_HTTPDDC_MAX_ATTEMPTS)
 	{
+		if (ShouldAbortForShutdown())
+		{
+			return EPutStatus::NotCached;
+		}
+
 		FScopedRequestPtr Request(PutRequestPools[IsInGameThread()].Get());
 		if (Request.IsValid())
 		{
