Ensure that at the time the FDerivedDataCache singleton is shutdown that we have waited on all internal backends to complete in-flight work.

diff --git a/Engine/Source/Developer/DerivedDataCache/Private/DerivedDataCache.cpp b/Engine/Source/Developer/DerivedDataCache/Private/DerivedDataCache.cpp
--- a/Engine/Source/Developer/DerivedDataCache/Private/DerivedDataCache.cpp
+++ b/Engine/Source/Developer/DerivedDataCache/Private/DerivedDataCache.cpp
@@ -369,6 +369,7 @@ public:
 	/** Destructor, flushes all sync tasks **/
 	~FDerivedDataCache()
 	{
+		WaitForQuiescence(true);
 		FScopeLock ScopeLock(&SynchronizationObject);
 		for (TMap<uint32,FAsyncTask<FBuildAsyncWorker>*>::TIterator It(PendingTasks); It; ++It)
 		{
