Since we are now waiting for quiescence from DDC backends during shutdown even in interactive sessions, reduce the sleep time before checking for task completion from 1s to 0.1s.

diff --git a/Engine/Source/Developer/DerivedDataCache/Private/DerivedDataBackends.cpp b/Engine/Source/Developer/DerivedDataCache/Private/DerivedDataBackends.cpp
--- a/Engine/Source/Developer/DerivedDataCache/Private/DerivedDataBackends.cpp
+++ b/Engine/Source/Developer/DerivedDataCache/Private/DerivedDataBackends.cpp
@@ -933,7 +933,7 @@ public:
 		while (AsyncCompletionCounter.GetValue())
 		{
 			check(AsyncCompletionCounter.GetValue() >= 0);
-			FPlatformProcess::Sleep(1.0f);
+			FPlatformProcess::Sleep(0.1f);
 			if (FPlatformTime::Seconds() - LastPrint > 5.0)
 			{
 				UE_LOG(LogDerivedDataCache, Log, TEXT("Waited %ds for derived data cache to finish..."), int32(FPlatformTime::Seconds() - StartTime));
