DDC: Iteration on ICache

- Identify payloads by the hash of the raw data.
- Added FCacheKey::Empty and FCachePayloadKey::Empty for convenience.
- Changed FCacheRecord and FCacheRecordBuilder constructors to be private.
- Changed storage in the legacy cache to store values inline for records with no attachments.
- Changed storage in the legacy cache to use CAS within each bucket.
- Improved wording and details in assert messages.

#rnx

diff --git a/Engine/Source/Developer/DerivedDataCache/Private/DerivedDataCache.cpp b/Engine/Source/Developer/DerivedDataCache/Private/DerivedDataCache.cpp
--- a/Engine/Source/Developer/DerivedDataCache/Private/DerivedDataCache.cpp
+++ b/Engine/Source/Developer/DerivedDataCache/Private/DerivedDataCache.cpp
@@ -160,15 +160,12 @@ namespace DerivedDataCacheCookStats
 }
 #endif
 
-namespace UE
-{
-namespace DerivedData
+namespace UE::DerivedData
 {
 
 ICache* CreateCache();
 
-} // DerivedData
-} // UE
+} // UE::DerivedData
 
 /** Whether we want to verify the DDC (pass in -VerifyDDC on the command line)*/
 bool GVerifyDDC = false;
@@ -694,12 +691,12 @@ PRAGMA_DISABLE_DEPRECATION_WARNINGS
 class FDerivedDataCacheModule final : public IDerivedDataCacheModule
 {
 public:
-	virtual FDerivedDataCacheInterface& GetDDC() final
+	FDerivedDataCacheInterface& GetDDC() final
 	{
 		return **CreateOrGetDDC();
 	}
 
-	virtual FDerivedDataCacheInterface* const* CreateOrGetDDC() final
+	FDerivedDataCacheInterface* const* CreateOrGetDDC() final
 	{
 		FScopeLock Lock(&CreateLock);
 		if (!GDerivedDataCacheInstance)
@@ -710,7 +707,7 @@ public:
 		return &GDerivedDataCacheInstance;
 	}
 
-	virtual void ShutdownModule() final
+	void ShutdownModule() final
 	{
 		FDDCCleanup::Shutdown();
 
@@ -728,9 +725,7 @@ IMPLEMENT_MODULE(FDerivedDataCacheModule, DerivedDataCache);
 
 ///////////////////////////////////////////////////////////////////////////////////////////////////////////////////////
 
-namespace UE
-{
-namespace DerivedData
+namespace UE::DerivedData
 {
 
 ///////////////////////////////////////////////////////////////////////////////////////////////////////////////////////
@@ -738,34 +733,34 @@ namespace DerivedData
 class FCache final : public ICache
 {
 public:
-	virtual ~FCache() = default;
+	~FCache() final = default;
 
-	virtual FCacheBucket CreateBucket(FStringView Name) final { return Private::CreateCacheBucket(Name); }
+	FCacheBucket CreateBucket(FStringView Name) final { return Private::CreateCacheBucket(Name); }
 
-	virtual FCacheRecordBuilder CreateRecord(const FCacheKey& Key) final { return Private::CreateCacheRecordBuilder(Key); }
+	FCacheRecordBuilder CreateRecord(const FCacheKey& Key) final { return Private::CreateCacheRecordBuilder(Key); }
 
-	virtual FRequest Put(
+	FRequest Put(
 		TArrayView<FCacheRecord> Records,
 		FStringView Context,
 		ECachePolicy Policy,
 		EPriority Priority,
 		FOnCachePutComplete&& Callback) final;
 
-	virtual FRequest Get(
+	FRequest Get(
 		TConstArrayView<FCacheKey> Keys,
 		FStringView Context,
 		ECachePolicy Policy,
 		EPriority Priority,
 		FOnCacheGetComplete&& Callback) final;
 
-	virtual FRequest GetPayload(
+	FRequest GetPayload(
 		TConstArrayView<FCachePayloadKey> Keys,
 		FStringView Context,
 		ECachePolicy Policy,
 		EPriority Priority,
 		FOnCacheGetPayloadComplete&& Callback) final;
 
-	virtual void CancelAll() final;
+	void CancelAll() final;
 
 private:
 	void Put(
@@ -791,17 +786,17 @@ private:
 
 	static FString MakeRecordKey(const FCacheKey& Key)
 	{
-		check(Key.Bucket.IsValid() && !Key.Hash.IsZero());
+		check(Key != FCacheKey::Empty);
 		TStringBuilder<96> Out;
 		Out << Key.Bucket << TEXT("_") << Key.Hash;
 		return FDerivedDataCacheInterface::SanitizeCacheKey(Out.ToString());
 	}
 
-	static FString MakePayloadKey(const FCachePayloadKey& PayloadKey)
+	static FString MakePayloadKey(const FCacheKey& CacheKey, const FIoHash& CompressedHash)
 	{
-		check(PayloadKey.CacheKey.Bucket.IsValid() && !PayloadKey.CacheKey.Hash.IsZero() && PayloadKey.Id.IsValid());
+		check(CacheKey != FCacheKey::Empty && !CompressedHash.IsZero());
 		TStringBuilder<96> Out;
-		Out << PayloadKey.CacheKey.Bucket << TEXT("_") << PayloadKey.CacheKey.Hash << TEXT("_") << PayloadKey.Id;
+		Out << CacheKey.Bucket << TEXT("_CAS_") << CompressedHash;
 		return FDerivedDataCacheInterface::SanitizeCacheKey(Out.ToString());
 	}
 };
@@ -833,52 +828,66 @@ void FCache::Put(
 		return;
 	}
 
-	const auto PutPayload = [&Key, Context](const FPayload& Payload)
-	{
-		const FSharedBuffer Buffer = Payload.GetBuffer().GetCompressed().Flatten();
-		check(Buffer && Buffer.GetSize() <= MAX_int32);
-		TConstArrayView<uint8> BufferView = MakeArrayView(
-			static_cast<const uint8*>(Buffer.GetData()),
-			static_cast<int32>(Buffer.GetSize()));
-		GetDerivedDataCacheRef().Put(*MakePayloadKey(FCachePayloadKey{Key, Payload.GetId()}), BufferView, Context);
-	};
-
 	FCbWriter Writer;
 	Writer.BeginObject();
+
+	const auto PutPayload = [&Writer, &Key, Context](const FPayload& Payload, bool bInline = false)
 	{
-		if (const FPayload& Value = Record.GetValuePayload())
+		const FCompositeBuffer CompressedBuffer = Payload.GetData().GetCompressed();
+		const FIoHash CompressedHash = FIoHash::HashBuffer(CompressedBuffer);
+		if (!bInline)
 		{
-			PutPayload(Value);
-			Writer.BeginObject("Value"_ASV);
-			Writer.AddObjectId("Id"_ASV, FCbObjectId(Value.GetId().GetView()));
-			Writer.AddBinaryAttachment("Hash"_ASV, Value.GetHash());
-			Writer.EndObject();
+			const FCachePayloadKey PayloadKey{Key, Payload.GetId()};
+			const FSharedBuffer Buffer = CompressedBuffer.Flatten();
+			checkf(Buffer.GetSize() <= MAX_int32,
+				TEXT("Cache: Put failed for %s from '%.*s' because payload is 2 GiB or greater."),
+				*WriteToString<96>(PayloadKey), Context.Len(), Context.GetData());
+			const TConstArrayView<uint8> BufferView = MakeArrayView(
+				static_cast<const uint8*>(Buffer.GetData()),
+				static_cast<int32>(Buffer.GetSize()));
+			GetDerivedDataCacheRef().Put(*MakePayloadKey(PayloadKey.CacheKey, CompressedHash), BufferView, Context);
 		}
-		if (!Record.GetAttachmentPayloads().IsEmpty())
+
+		Writer.BeginObject();
+		Writer.AddObjectId("Id"_ASV, FCbObjectId(Payload.GetId().GetView()));
+		Writer.AddBinaryAttachment("RawHash"_ASV, Payload.GetRawHash());
+		Writer.AddHash("CompressedHash"_ASV, CompressedHash);
+		if (bInline)
 		{
-			Writer.BeginArray("Attachments"_ASV);
-			for (const FPayload& Attachment : Record.GetAttachmentPayloads())
-			{
-				PutPayload(Attachment);
-				Writer.BeginObject();
-				Writer.AddObjectId("Id"_ASV, FCbObjectId(Attachment.GetId().GetView()));
-				Writer.AddBinaryAttachment("Hash"_ASV, Attachment.GetHash());
-				Writer.EndObject();
-			}
-			Writer.EndArray();
+			Writer.AddBinary("CompressedData"_ASV, CompressedBuffer);
 		}
+		Writer.EndObject();
+	};
 
-		const FCbObject& Meta = Record.GetMeta();
-		if (Meta.CreateViewIterator())
+	if (const FPayload& Value = Record.GetValuePayload())
+	{
+		const bool bInline = Record.GetAttachmentPayloads().IsEmpty();
+		Writer.SetName("Value"_ASV);
+		PutPayload(Value, bInline);
+	}
+	if (!Record.GetAttachmentPayloads().IsEmpty())
+	{
+		Writer.BeginArray("Attachments"_ASV);
+		for (const FPayload& Attachment : Record.GetAttachmentPayloads())
 		{
-			Writer.AddObject("Meta"_ASV, Meta);
+			PutPayload(Attachment);
 		}
+		Writer.EndArray();
 	}
+
+	const FCbObject& Meta = Record.GetMeta();
+	if (Meta.CreateViewIterator())
+	{
+		Writer.AddObject("Meta"_ASV, Meta);
+	}
+
 	Writer.EndObject();
 
 	TArray<uint8> Data;
 	const uint64 SaveSize = Writer.GetSaveSize();
-	check(SaveSize <= MAX_int32);
+	checkf(SaveSize <= MAX_int32,
+		TEXT("Cache: Put failed for %s from '%.*s' because record is 2 GiB or greater."),
+		*WriteToString<96>(Key), Context.Len(), Context.GetData());
 	Data.SetNumUninitialized(static_cast<int32>(SaveSize));
 	Writer.Save(MakeMemoryView(Data));
 
@@ -912,7 +921,7 @@ void FCache::Get(
 {
 	FPayload Value;
 	TArray<FPayload> Attachments;
-	TArray<FPayloadId, TInlineAllocator<1>> SkippedPayloads;
+	TArray<FIoHash, TInlineAllocator<1>> SkippedPayloads;
 	FCacheRecordBuilder RecordBuilder = CreateRecord(Key);
 	EStatus Status = EStatus::Error;
 
@@ -955,34 +964,36 @@ void FCache::Get(
 
 	// Read the record from its compact binary fields.
 	const FCbObject RecordObject(MoveTemp(RecordBuffer));
-	FCbObject MetaObject = RecordObject.Find("Meta"_ASV).AsObject();
-	const FCbObjectView ValueObject = RecordObject["Value"_ASV].AsObjectView();
-	const FCbArrayView AttachmentsArray = RecordObject["Attachments"_ASV].AsArrayView();
+	FCbObject MetaObject = RecordObject["Meta"_ASV].AsObject();
+	const FCbObject ValueObject = RecordObject["Value"_ASV].AsObject();
+	const FCbArray AttachmentsArray = RecordObject["Attachments"_ASV].AsArray();
 
 	// Read the value and attachments if they have been requested.
-	const auto GetPayload = [&SkippedPayloads, &Key, Policy, Context](FCbObjectView PayloadObject, ECachePolicy SkipFlag) -> FPayload
+	const auto GetPayload = [&SkippedPayloads, &Key, Policy, Context](FCbObject PayloadObject, ECachePolicy SkipFlag) -> FPayload
 	{
 		TArray<uint8> PayloadData;
-		const FPayloadId PayloadId = FPayloadId(PayloadObject["Id"_ASV].AsObjectId().GetView());
-		const FIoHash PayloadHash = PayloadObject["Hash"_ASV].AsBinaryAttachment();
+		const FPayloadId PayloadId = FPayloadId(PayloadObject.FindView("Id"_ASV).AsObjectId().GetView());
+		const FIoHash PayloadRawHash = PayloadObject.FindView("RawHash"_ASV).AsBinaryAttachment();
+		const FIoHash PayloadCompressedHash = PayloadObject.FindView("CompressedHash"_ASV).AsHash();
+		const FSharedBuffer PayloadCompressedData = PayloadObject["CompressedData"_ASV].AsBinary();
 		const FCachePayloadKey PayloadKey{Key, PayloadId};
 		if (EnumHasAnyFlags(Policy, SkipFlag))
 		{
-			SkippedPayloads.Add(PayloadId);
-			return FPayload(PayloadId, PayloadHash);
-		}
-		else if (GetDerivedDataCacheRef().GetSynchronous(*MakePayloadKey(PayloadKey), PayloadData, Context))
-		{
-			if (FIoHash::HashBuffer(MakeMemoryView(PayloadData)) == PayloadHash)
+			if (PayloadCompressedData.IsNull())
 			{
-				return FPayload(PayloadId, PayloadHash, FCompressedBuffer::FromCompressed(MakeSharedBufferFromArray(MoveTemp(PayloadData))));
-			}
-			else
-			{
-				UE_LOG(LogDerivedDataCache, Verbose, TEXT("Cache: Get cache miss with corrupted content for %s from '%.*s'"),
-					*WriteToString<96>(PayloadKey), Context.Len(), Context.GetData());
-				return FPayload();
+				SkippedPayloads.Add(PayloadCompressedHash);
 			}
+			return FPayload(PayloadId, PayloadRawHash);
+		}
+
+		FCompressedBuffer CompressedBuffer;
+		if (PayloadCompressedData)
+		{
+			CompressedBuffer = FCompressedBuffer::FromCompressed(PayloadCompressedData);
+		}
+		else if (GetDerivedDataCacheRef().GetSynchronous(*MakePayloadKey(PayloadKey.CacheKey, PayloadCompressedHash), PayloadData, Context))
+		{
+			CompressedBuffer = FCompressedBuffer::FromCompressed(MakeSharedBufferFromArray(MoveTemp(PayloadData)));
 		}
 		else
 		{
@@ -990,6 +1001,18 @@ void FCache::Get(
 				*WriteToString<96>(PayloadKey), Context.Len(), Context.GetData());
 			return FPayload();
 		}
+
+		if (CompressedBuffer.GetRawHash() == PayloadRawHash &&
+			FIoHash::HashBuffer(CompressedBuffer.GetCompressed()) == PayloadCompressedHash)
+		{
+			return FPayload(PayloadId, MoveTemp(CompressedBuffer));
+		}
+		else
+		{
+			UE_LOG(LogDerivedDataCache, Verbose, TEXT("Cache: Get cache miss with corrupted content for %s from '%.*s'"),
+				*WriteToString<96>(PayloadKey), Context.Len(), Context.GetData());
+			return FPayload();
+		}
 	};
 
 	if (ValueObject.CreateViewIterator() && !(Value = GetPayload(ValueObject, ECachePolicy::SkipValue)))
@@ -998,9 +1021,9 @@ void FCache::Get(
 	}
 
 	Attachments.Reserve(AttachmentsArray.Num());
-	for (FCbFieldView AttachmentField : AttachmentsArray)
+	for (FCbField AttachmentField : AttachmentsArray)
 	{
-		Attachments.Add(GetPayload(AttachmentField.AsObjectView(), ECachePolicy::SkipAttachments));
+		Attachments.Add(GetPayload(AttachmentField.AsObject(), ECachePolicy::SkipAttachments));
 		if (!Attachments.Last())
 		{
 			return;
@@ -1013,7 +1036,7 @@ void FCache::Get(
 		TArray<FString, TInlineAllocator<1>> SkippedPayloadKeys;
 		SkippedPayloadKeys.Reserve(SkippedPayloads.Num());
 		Algo::Transform(SkippedPayloads, SkippedPayloadKeys,
-			[&Key](const FPayloadId& Id) -> FString { return MakePayloadKey(FCachePayloadKey{Key, Id}); });
+			[&Key](const FIoHash& Hash) -> FString { return MakePayloadKey(Key, Hash); });
 		const bool bPayloadsExist = EnumHasAnyFlags(Policy, ECachePolicy::StoreLocal)
 			? GetDerivedDataCacheRef().TryToPrefetch(SkippedPayloadKeys, Context)
 			: GetDerivedDataCacheRef().AllCachedDataProbablyExists(SkippedPayloadKeys);
@@ -1084,70 +1107,75 @@ void FCache::GetPayload(
 		return;
 	}
 
-	// Check for existence of the payload if it is being skipped.
-	if (EnumHasAnyFlags(Policy, ECachePolicy::SkipValue))
+	// Request the record from storage.
+	FSharedBuffer RecordBuffer;
 	{
-		const bool bPayloadExists = EnumHasAnyFlags(Policy, ECachePolicy::StoreLocal)
-			? GetDerivedDataCacheRef().TryToPrefetch({MakePayloadKey(Key)}, Context)
-			: GetDerivedDataCacheRef().CachedDataProbablyExists(*MakePayloadKey(Key));
-		if (bPayloadExists)
+		TArray<uint8> RecordData;
+		if (!GetDerivedDataCacheRef().GetSynchronous(*MakeRecordKey(Key.CacheKey), RecordData, Context))
 		{
-			// Request the record from storage.
-			FSharedBuffer RecordBuffer;
-			{
-				TArray<uint8> RecordData;
-				if (!GetDerivedDataCacheRef().GetSynchronous(*MakeRecordKey(Key.CacheKey), RecordData, Context))
-				{
-					UE_LOG(LogDerivedDataCache, Verbose, TEXT("Cache: Get cache miss with missing record for %s from '%.*s'"),
-						*WriteToString<96>(Key), Context.Len(), Context.GetData());
-					return;
-				}
-				RecordBuffer = MakeSharedBufferFromArray(MoveTemp(RecordData));
-			}
-
-			// Validate that the record can be read as compact binary without crashing.
-			if (ValidateCompactBinaryRange(RecordBuffer, ECbValidateMode::Default) != ECbValidateError::None)
-			{
-				UE_LOG(LogDerivedDataCache, Verbose, TEXT("Cache: Get cache miss with corrupted record for %s from '%.*s'"),
-					*WriteToString<96>(Key), Context.Len(), Context.GetData());
-				return;
-			}
+			UE_LOG(LogDerivedDataCache, Verbose, TEXT("Cache: Get cache miss with missing record for %s from '%.*s'"),
+				*WriteToString<96>(Key), Context.Len(), Context.GetData());
+			return;
+		}
+		RecordBuffer = MakeSharedBufferFromArray(MoveTemp(RecordData));
+	}
 
-			// Read the record from its compact binary fields.
-			const FCbObject RecordObject(MoveTemp(RecordBuffer));
-			const FCbObjectView ValueObject = RecordObject["Value"_ASV].AsObjectView();
-			const FCbArrayView AttachmentsArray = RecordObject["Attachments"_ASV].AsArrayView();
+	// Validate that the record can be read as compact binary without crashing.
+	if (ValidateCompactBinaryRange(RecordBuffer, ECbValidateMode::Default) != ECbValidateError::None)
+	{
+		UE_LOG(LogDerivedDataCache, Verbose, TEXT("Cache: Get cache miss with corrupted record for %s from '%.*s'"),
+			*WriteToString<96>(Key), Context.Len(), Context.GetData());
+		return;
+	}
 
-			// Find the hash of the payload in the record.
-			const FCbObjectId KeyObjectId(Key.Id.GetView());
-			if (ValueObject["Id"_ASV].AsObjectId() == KeyObjectId)
-			{
-				Payload = FPayload(Key.Id, ValueObject["Hash"_ASV].AsBinaryAttachment());
-			}
-			else
+	// Read the record from its compact binary fields.
+	const FCbObject RecordObject(MoveTemp(RecordBuffer));
+	const FCbObject ValueObject = RecordObject["Value"_ASV].AsObject();
+	const FCbArrayView AttachmentsArray = RecordObject.FindView("Attachments"_ASV).AsArrayView();
+
+	// Find the hash of the payload in the record.
+	FIoHash PayloadCompressedHash;
+	FSharedBuffer PayloadCompressedData;
+	const FCbObjectId KeyObjectId(Key.Id.GetView());
+	if (ValueObject["Id"_ASV].AsObjectId() == KeyObjectId)
+	{
+		Payload = FPayload(Key.Id, ValueObject.FindView("RawHash"_ASV).AsBinaryAttachment());
+		PayloadCompressedHash = ValueObject.FindView("CompressedHash"_ASV).AsHash();
+		PayloadCompressedData = ValueObject["CompressedData"_ASV].AsBinary();
+	}
+	else
+	{
+		for (FCbFieldView AttachmentField : AttachmentsArray)
+		{
+			if (AttachmentField.AsObjectView()["Id"_ASV].AsObjectId() == KeyObjectId)
 			{
-				for (FCbFieldView AttachmentField : AttachmentsArray)
-				{
-					if (AttachmentField.AsObjectView()["Id"_ASV].AsObjectId() == KeyObjectId)
-					{
-						Payload = FPayload(Key.Id, AttachmentField.AsObjectView()["Hash"_ASV].AsBinaryAttachment());
-						break;
-					}
-				}
+				Payload = FPayload(Key.Id, AttachmentField.AsObjectView()["RawHash"_ASV].AsBinaryAttachment());
+				PayloadCompressedHash = AttachmentField.AsObjectView()["CompressedHash"_ASV].AsHash();
+				break;
 			}
+		}
+	}
 
-			if (Payload)
-			{
-				UE_LOG(LogDerivedDataCache, Verbose, TEXT("Cache: GetPayload (exists) cache hit for %s from '%.*s'"),
-					*WriteToString<96>(Key), Context.Len(), Context.GetData());
-				Status = EStatus::Ok;
-			}
-			else
-			{
-				UE_LOG(LogDerivedDataCache, Verbose,
-					TEXT("Cache: GetPayload (exists) cache miss with missing payload in record for %s from '%.*s'"),
-					*WriteToString<96>(Key), Context.Len(), Context.GetData());
-			}
+	// Verify that the payload was found in the record.
+	if (PayloadCompressedHash.IsZero())
+	{
+		UE_LOG(LogDerivedDataCache, Verbose,
+			TEXT("Cache: GetPayload cache miss with missing payload in record for %s from '%.*s'"),
+			*WriteToString<96>(Key), Context.Len(), Context.GetData());
+		return;
+	}
+
+	// Check for existence of the payload if it is being skipped.
+	if (EnumHasAnyFlags(Policy, ECachePolicy::SkipValue))
+	{
+		const bool bPayloadExists = PayloadCompressedData || EnumHasAnyFlags(Policy, ECachePolicy::StoreLocal)
+			? GetDerivedDataCacheRef().TryToPrefetch({MakePayloadKey(Key.CacheKey, PayloadCompressedHash)}, Context)
+			: GetDerivedDataCacheRef().CachedDataProbablyExists(*MakePayloadKey(Key.CacheKey, PayloadCompressedHash));
+		if (bPayloadExists)
+		{
+			UE_LOG(LogDerivedDataCache, Verbose, TEXT("Cache: GetPayload (exists) cache hit for %s from '%.*s'"),
+				*WriteToString<96>(Key), Context.Len(), Context.GetData());
+			Status = EStatus::Ok;
 		}
 		else
 		{
@@ -1158,19 +1186,32 @@ void FCache::GetPayload(
 	}
 
 	// Request the payload from storage.
-	TArray<uint8> Data;
-	if (!GetDerivedDataCacheRef().GetSynchronous(*MakePayloadKey(Key), Data, Context))
+	if (PayloadCompressedData.IsNull())
 	{
-		UE_LOG(LogDerivedDataCache, Verbose, TEXT("Cache: GetPayload cache miss for %s from '%.*s'"),
-			*WriteToString<96>(Key), Context.Len(), Context.GetData());
-		return;
+		TArray<uint8> Data;
+		if (!GetDerivedDataCacheRef().GetSynchronous(*MakePayloadKey(Key.CacheKey, PayloadCompressedHash), Data, Context))
+		{
+			UE_LOG(LogDerivedDataCache, Verbose, TEXT("Cache: GetPayload cache miss for %s from '%.*s'"),
+				*WriteToString<96>(Key), Context.Len(), Context.GetData());
+			return;
+		}
+		PayloadCompressedData = MakeSharedBufferFromArray(MoveTemp(Data));
 	}
 
-	Payload = FPayload(Key.Id, FCompressedBuffer::FromCompressed(MakeSharedBufferFromArray(MoveTemp(Data))));
-
-	UE_LOG(LogDerivedDataCache, Verbose, TEXT("Cache: GetPayload cache hit for %s from '%.*s'"),
-		*WriteToString<96>(Key), Context.Len(), Context.GetData());
-	Status = EStatus::Ok;
+	FCompressedBuffer CompressedBuffer = FCompressedBuffer::FromCompressed(PayloadCompressedData);
+	if (CompressedBuffer.GetRawHash() == Payload.GetRawHash() &&
+		FIoHash::HashBuffer(CompressedBuffer.GetCompressed()) == PayloadCompressedHash)
+	{
+		UE_LOG(LogDerivedDataCache, Verbose, TEXT("Cache: GetPayload cache hit for %s from '%.*s'"),
+			*WriteToString<96>(Key), Context.Len(), Context.GetData());
+		Payload = FPayload(Key.Id, MoveTemp(CompressedBuffer));
+		Status = EStatus::Ok;
+	}
+	else
+	{
+		UE_LOG(LogDerivedDataCache, Verbose, TEXT("Cache: GetPayload cache miss with corrupted content for %s from '%.*s'"),
+			*WriteToString<96>(Key), Context.Len(), Context.GetData());
+	}
 }
 
 FRequest FCache::GetPayload(
@@ -1200,5 +1241,4 @@ ICache* CreateCache()
 
 ///////////////////////////////////////////////////////////////////////////////////////////////////////////////////////
 
-} // DerivedData
-} // UE
+} // UE::DerivedData
diff --git a/Engine/Source/Developer/DerivedDataCache/Private/DerivedDataCacheRecord.cpp b/Engine/Source/Developer/DerivedDataCache/Private/DerivedDataCacheRecord.cpp
--- a/Engine/Source/Developer/DerivedDataCache/Private/DerivedDataCacheRecord.cpp
+++ b/Engine/Source/Developer/DerivedDataCache/Private/DerivedDataCacheRecord.cpp
@@ -8,11 +8,7 @@
 #include "Serialization/CompactBinary.h"
 #include "UObject/NameTypes.h"
 
-namespace UE
-{
-namespace DerivedData
-{
-namespace Private
+namespace UE::DerivedData::Private
 {
 
 ///////////////////////////////////////////////////////////////////////////////////////////////////////////////////////
@@ -21,18 +17,18 @@ class FCacheRecordBuilderInternal final : public ICacheRecordBuilderInternal
 {
 public:
 	explicit FCacheRecordBuilderInternal(const FCacheKey& Key);
-	virtual ~FCacheRecordBuilderInternal() = default;
+	~FCacheRecordBuilderInternal() final = default;
 
-	virtual void SetMeta(FCbObject&& Meta) final;
+	void SetMeta(FCbObject&& Meta) final;
 
-	virtual FPayloadId SetValue(const FSharedBuffer& Buffer, const FPayloadId& Id) final;
-	virtual FPayloadId SetValue(FPayload&& Payload) final;
+	FPayloadId SetValue(const FSharedBuffer& Buffer, const FPayloadId& Id) final;
+	FPayloadId SetValue(FPayload&& Payload) final;
 
-	virtual FPayloadId AddAttachment(const FSharedBuffer& Buffer, const FPayloadId& Id) final;
-	virtual FPayloadId AddAttachment(FPayload&& Payload) final;
+	FPayloadId AddAttachment(const FSharedBuffer& Buffer, const FPayloadId& Id) final;
+	FPayloadId AddAttachment(FPayload&& Payload) final;
 
-	virtual FCacheRecord Build() final;
-	virtual FRequest BuildAsync(FOnCacheRecordComplete&& Callback, EPriority Priority) final;
+	FCacheRecord Build() final;
+	FRequest BuildAsync(FOnCacheRecordComplete&& Callback, EPriority Priority) final;
 
 	FCacheKey Key;
 	FCbObject Meta;
@@ -48,19 +44,19 @@ public:
 	FCacheRecordInternal() = default;
 	explicit FCacheRecordInternal(FCacheRecordBuilderInternal&& RecordBuilder);
 
-	virtual ~FCacheRecordInternal() = default;
+	~FCacheRecordInternal() final = default;
 
-	virtual FCacheRecord Clone() const final;
+	FCacheRecord Clone() const final;
 
-	virtual const FCacheKey& GetKey() const final;
-	virtual const FCbObject& GetMeta() const final;
+	const FCacheKey& GetKey() const final;
+	const FCbObject& GetMeta() const final;
 
-	virtual FSharedBuffer GetValue() const final;
-	virtual const FPayload& GetValuePayload() const final;
+	FSharedBuffer GetValue() const final;
+	const FPayload& GetValuePayload() const final;
 
-	virtual FSharedBuffer GetAttachment(const FPayloadId& Id) const final;
-	virtual const FPayload& GetAttachmentPayload(const FPayloadId& Id) const final;
-	virtual TConstArrayView<FPayload> GetAttachmentPayloads() const final;
+	FSharedBuffer GetAttachment(const FPayloadId& Id) const final;
+	const FPayload& GetAttachmentPayload(const FPayloadId& Id) const final;
+	TConstArrayView<FPayload> GetAttachmentPayloads() const final;
 
 	FCacheKey Key;
 	FCbObject Meta;
@@ -78,9 +74,9 @@ static const FPayload& GetEmptyCachePayload()
 	return Empty;
 }
 
-static FPayloadId GetOrCreatePayloadId(const FPayloadId& Id, const FSharedBuffer& Buffer)
+static FPayloadId GetOrCreatePayloadId(const FPayloadId& Id, const FIoHash& RawHash)
 {
-	return Id.IsValid() ? Id : FPayloadId::FromHash(FIoHash::HashBuffer(Buffer));
+	return Id.IsValid() ? Id : FPayloadId::FromHash(RawHash);
 }
 
 ///////////////////////////////////////////////////////////////////////////////////////////////////////////////////////
@@ -95,7 +91,7 @@ FCacheRecordInternal::FCacheRecordInternal(FCacheRecordBuilderInternal&& RecordB
 
 FCacheRecord FCacheRecordInternal::Clone() const
 {
-	return FCacheRecord(new FCacheRecordInternal(*this));
+	return CreateCacheRecord(new FCacheRecordInternal(*this));
 }
 
 const FCacheKey& FCacheRecordInternal::GetKey() const
@@ -112,7 +108,7 @@ FSharedBuffer FCacheRecordInternal::GetValue() const
 {
 	if (ValueCache.IsNull() && Value.IsValid())
 	{
-		ValueCache = Value.GetBuffer().Decompress();
+		ValueCache = Value.GetData().Decompress();
 	}
 	return ValueCache;
 }
@@ -134,7 +130,7 @@ FSharedBuffer FCacheRecordInternal::GetAttachment(const FPayloadId& Id) const
 		FSharedBuffer& DataCache = AttachmentsCache[Index];
 		if (!DataCache)
 		{
-			DataCache = Attachments[Index].GetBuffer().Decompress();
+			DataCache = Attachments[Index].GetData().Decompress();
 		}
 		return DataCache;
 	}
@@ -171,37 +167,42 @@ void FCacheRecordBuilderInternal::SetMeta(FCbObject&& InMeta)
 
 FPayloadId FCacheRecordBuilderInternal::SetValue(const FSharedBuffer& Buffer, const FPayloadId& Id)
 {
-	const FPayloadId ValueId = GetOrCreatePayloadId(Id, Buffer);
-	return SetValue(FPayload(ValueId, FCompressedBuffer::Compress(NAME_Default, Buffer)));
+	FCompressedBuffer CompressedBuffer = FCompressedBuffer::Compress(NAME_Default, Buffer);
+	const FPayloadId ValueId = GetOrCreatePayloadId(Id, FIoHash(CompressedBuffer.GetRawHash()));
+	return SetValue(FPayload(ValueId, MoveTemp(CompressedBuffer)));
 }
 
 FPayloadId FCacheRecordBuilderInternal::SetValue(FPayload&& Payload)
 {
-	checkf(Value.IsNull(), TEXT("Failed to set value with ID %s because an there is an existing value with ID %s"),
-		*WriteToString<32>(Payload.GetId()), *WriteToString<32>(Value.GetId()));
+	checkf(Payload, TEXT("Failed to set value on %s because the payload is null."), *WriteToString<96>(Key));
+	checkf(Value.IsNull(),
+		TEXT("Cache: Failed to set value on %s with ID %s because it has an existing value with ID %s."),
+		*WriteToString<96>(Key), *WriteToString<32>(Payload.GetId()), *WriteToString<32>(Value.GetId()));
 	Value = MoveTemp(Payload);
 	return Value.GetId();
 }
 
 FPayloadId FCacheRecordBuilderInternal::AddAttachment(const FSharedBuffer& Buffer, const FPayloadId& Id)
 {
-	const FPayloadId AttachmentId = GetOrCreatePayloadId(Id, Buffer);
-	return AddAttachment(FPayload(AttachmentId, FCompressedBuffer::Compress(NAME_Default, Buffer)));
+	FCompressedBuffer CompressedBuffer = FCompressedBuffer::Compress(NAME_Default, Buffer);
+	const FPayloadId AttachmentId = GetOrCreatePayloadId(Id, FIoHash(CompressedBuffer.GetRawHash()));
+	return AddAttachment(FPayload(AttachmentId, MoveTemp(CompressedBuffer)));
 }
 
 FPayloadId FCacheRecordBuilderInternal::AddAttachment(FPayload&& Payload)
 {
+	checkf(Payload, TEXT("Failed to add attachment on %s because the payload is null."), *WriteToString<96>(Key));
 	const int32 Index = Algo::LowerBound(Attachments, Payload, FPayloadLessById());
 	checkf(!Attachments.IsValidIndex(Index) || !FPayloadEqualById()(Attachments[Index], Payload),
-		TEXT("Failed to add attachment with ID %s because an existing attachment is using that ID"),
-		*WriteToString<32>(Payload.GetId()));
+		TEXT("Failed to add attachment on %s with ID %s because it has an existing attachment with that ID."),
+		*WriteToString<96>(Key), *WriteToString<32>(Payload.GetId()));
 	Attachments.Insert(MoveTemp(Payload), Index);
 	return Attachments[Index].GetId();
 }
 
 FCacheRecord FCacheRecordBuilderInternal::Build()
 {
-	return FCacheRecord(new FCacheRecordInternal(MoveTemp(*this)));
+	return CreateCacheRecord(new FCacheRecordInternal(MoveTemp(*this)));
 }
 
 FRequest FCacheRecordBuilderInternal::BuildAsync(FOnCacheRecordComplete&& Callback, EPriority Priority)
@@ -212,13 +213,23 @@ FRequest FCacheRecordBuilderInternal::BuildAsync(FOnCacheRecordComplete&& Callba
 
 ///////////////////////////////////////////////////////////////////////////////////////////////////////////////////////
 
+FCacheRecord CreateCacheRecord(ICacheRecordInternal* Record)
+{
+	return FCacheRecord(Record);
+}
+
+FCacheRecordBuilder CreateCacheRecordBuilder(ICacheRecordBuilderInternal* RecordBuilder)
+{
+	return FCacheRecordBuilder(RecordBuilder);
+}
+
+///////////////////////////////////////////////////////////////////////////////////////////////////////////////////////
+
 FCacheRecordBuilder CreateCacheRecordBuilder(const FCacheKey& Key)
 {
-	return FCacheRecordBuilder(new FCacheRecordBuilderInternal(Key));
+	return CreateCacheRecordBuilder(new FCacheRecordBuilderInternal(Key));
 }
 
 ///////////////////////////////////////////////////////////////////////////////////////////////////////////////////////
 
-} // Private
-} // DerivedData
-} // UE
+} // UE::DerivedData::Private
diff --git a/Engine/Source/Developer/DerivedDataCache/Public/DerivedDataCacheKey.h b/Engine/Source/Developer/DerivedDataCache/Public/DerivedDataCacheKey.h
--- a/Engine/Source/Developer/DerivedDataCache/Public/DerivedDataCacheKey.h
+++ b/Engine/Source/Developer/DerivedDataCache/Public/DerivedDataCacheKey.h
@@ -9,19 +9,15 @@
 #include "IO/IoHash.h"
 #include "Templates/TypeHash.h"
 
-namespace UE
-{
-namespace DerivedData
-{
-
-///////////////////////////////////////////////////////////////////////////////////////////////////////////////////////
-
-namespace Private
+namespace UE::DerivedData::Private
 {
 
 struct FCacheBucketName { const TCHAR* Name; };
 
-}
+} // UE::DerivedData::Private
+
+namespace UE::DerivedData
+{
 
 ///////////////////////////////////////////////////////////////////////////////////////////////////////////////////////
 
@@ -83,6 +79,34 @@ private:
 	const TCHAR* Name = nullptr;
 };
 
+///////////////////////////////////////////////////////////////////////////////////////////////////////////////////////
+
+/** A key that uniquely identifies a cache record. */
+struct FCacheKey
+{
+	FCacheBucket Bucket;
+	FIoHash Hash;
+
+	/** A key with a null bucket and a zero hash. */
+	static const FCacheKey Empty;
+};
+
+inline const FCacheKey FCacheKey::Empty;
+
+/** A key that uniquely identifies a payload within a cache record. */
+struct FCachePayloadKey
+{
+	FCacheKey CacheKey;
+	FPayloadId Id;
+
+	/** A payload key with an empty cache key and a null payload identifier. */
+	static const FCachePayloadKey Empty;
+};
+
+inline const FCachePayloadKey FCachePayloadKey::Empty;
+
+///////////////////////////////////////////////////////////////////////////////////////////////////////////////////////
+
 template <>
 inline FStringView FCacheBucket::ToString<TCHAR>() const
 {
@@ -105,13 +129,6 @@ inline TStringBuilderBase<CharType>& operator<<(TStringBuilderBase<CharType>& Bu
 
 ///////////////////////////////////////////////////////////////////////////////////////////////////////////////////////
 
-/** A key that uniquely identifies a cache record. */
-struct FCacheKey
-{
-	FCacheBucket Bucket;
-	FIoHash Hash;
-};
-
 inline bool operator==(const FCacheKey& A, const FCacheKey& B)
 {
 	return A.Bucket == B.Bucket && A.Hash == B.Hash;
@@ -142,13 +159,6 @@ inline TStringBuilderBase<CharType>& operator<<(TStringBuilderBase<CharType>& Bu
 
 ///////////////////////////////////////////////////////////////////////////////////////////////////////////////////////
 
-/** A key that uniquely identifies a payload within a cache record. */
-struct FCachePayloadKey
-{
-	FCacheKey CacheKey;
-	FPayloadId Id;
-};
-
 inline bool operator==(const FCachePayloadKey& A, const FCachePayloadKey& B)
 {
 	return A.CacheKey == B.CacheKey && A.Id == B.Id;
@@ -179,5 +189,4 @@ inline TStringBuilderBase<CharType>& operator<<(TStringBuilderBase<CharType>& Bu
 
 ///////////////////////////////////////////////////////////////////////////////////////////////////////////////////////
 
-} // DerivedData
-} // UE
+} // UE::DerivedData
diff --git a/Engine/Source/Developer/DerivedDataCache/Public/DerivedDataCacheRecord.h b/Engine/Source/Developer/DerivedDataCache/Public/DerivedDataCacheRecord.h
--- a/Engine/Source/Developer/DerivedDataCache/Public/DerivedDataCacheRecord.h
+++ b/Engine/Source/Developer/DerivedDataCache/Public/DerivedDataCacheRecord.h
@@ -10,19 +10,14 @@
 class FCbObject;
 template <typename FuncType> class TUniqueFunction;
 
-namespace UE
-{
-namespace DerivedData
-{
-
-class FCacheRecord;
-struct FCacheKey;
-
-using FOnCacheRecordComplete = TUniqueFunction<void (FCacheRecord&& Record)>;
+namespace UE::DerivedData { class FCacheRecord; }
+namespace UE::DerivedData { class FCacheRecordBuilder; }
+namespace UE::DerivedData { struct FCacheKey; }
+namespace UE::DerivedData { using FOnCacheRecordComplete = TUniqueFunction<void (FCacheRecord&& Record)>; }
 
 ///////////////////////////////////////////////////////////////////////////////////////////////////////////////////////
 
-namespace Private
+namespace UE::DerivedData::Private
 {
 
 class ICacheRecordInternal
@@ -39,6 +34,8 @@ public:
 	virtual TConstArrayView<FPayload> GetAttachmentPayloads() const = 0;
 };
 
+FCacheRecord CreateCacheRecord(ICacheRecordInternal* Record);
+
 class ICacheRecordBuilderInternal
 {
 public:
@@ -52,10 +49,15 @@ public:
 	virtual FRequest BuildAsync(FOnCacheRecordComplete&& Callback, EPriority Priority) = 0;
 };
 
-} // Private
+FCacheRecordBuilder CreateCacheRecordBuilder(ICacheRecordBuilderInternal* RecordBuilder);
+
+} // UE::DerivedData::Private
 
 ///////////////////////////////////////////////////////////////////////////////////////////////////////////////////////
 
+namespace UE::DerivedData
+{
+
 /**
  * A cache record is a key, a value, attachments, and metadata.
  *
@@ -111,8 +113,8 @@ public:
 	/** Reset this to null. */
 	inline void Reset() { *this = FCacheRecord(); }
 
-public:
-	// Internal API
+private:
+	friend FCacheRecord Private::CreateCacheRecord(Private::ICacheRecordInternal* Record);
 
 	/** Construct a cache record. Use Record.Clone() or Builder.Build[Async](). */
 	inline explicit FCacheRecord(Private::ICacheRecordInternal* InRecord)
@@ -120,12 +122,9 @@ public:
 	{
 	}
 
-private:
 	TUniquePtr<Private::ICacheRecordInternal> Record;
 };
 
-///////////////////////////////////////////////////////////////////////////////////////////////////////////////////////
-
 /**
  * A cache record builder is used to construct a cache record.
  *
@@ -220,19 +219,16 @@ public:
 		return RecordBuilder->BuildAsync(MoveTemp(Callback), Priority);
 	}
 
-public:
-	// Internal API
+private:
+	friend FCacheRecordBuilder Private::CreateCacheRecordBuilder(Private::ICacheRecordBuilderInternal* RecordBuilder);
 
+	/** Construct a cache record builder. Use ICache::CreateRecord(). */
 	inline explicit FCacheRecordBuilder(Private::ICacheRecordBuilderInternal* InRecordBuilder)
 		: RecordBuilder(InRecordBuilder)
 	{
 	}
 
-private:
 	TUniquePtr<Private::ICacheRecordBuilderInternal> RecordBuilder;
 };
 
-///////////////////////////////////////////////////////////////////////////////////////////////////////////////////////
-
-} // DerivedData
-} // UE
+} // UE::DerivedData
diff --git a/Engine/Source/Developer/DerivedDataCache/Public/DerivedDataPayload.h b/Engine/Source/Developer/DerivedDataCache/Public/DerivedDataPayload.h
--- a/Engine/Source/Developer/DerivedDataCache/Public/DerivedDataPayload.h
+++ b/Engine/Source/Developer/DerivedDataCache/Public/DerivedDataPayload.h
@@ -11,9 +11,7 @@
 #include "Memory/MemoryView.h"
 #include "String/BytesToHex.h"
 
-namespace UE
-{
-namespace DerivedData
+namespace UE::DerivedData
 {
 
 ///////////////////////////////////////////////////////////////////////////////////////////////////////////////////////
@@ -31,11 +29,11 @@ public:
 	inline explicit FPayloadId(FMemoryView Id);
 
 	/** Construct an ID from a non-zero hash. */
-	static inline FPayloadId FromHash(const FIoHash& Hash);
+	[[nodiscard]] static inline FPayloadId FromHash(const FIoHash& Hash);
 
 	/** Construct an ID from a non-empty name. */
-	static inline FPayloadId FromName(FAnsiStringView Name);
-	static inline FPayloadId FromName(FWideStringView Name);
+	[[nodiscard]] static inline FPayloadId FromName(FAnsiStringView Name);
+	[[nodiscard]] static inline FPayloadId FromName(FWideStringView Name);
 
 	/** Returns a reference to the raw byte array for the ID. */
 	inline const ByteArray& GetBytes() const { return Bytes; }
@@ -58,7 +56,7 @@ private:
 ///////////////////////////////////////////////////////////////////////////////////////////////////////////////////////
 
 /**
- * A payload is described by a PayloadId and an IoHash of the compressed payload buffer.
+ * A payload is described by a PayloadId and an IoHash of its raw buffer (uncompressed).
  *
  * Payloads may be constructed with or without a compressed buffer. A payload without a buffer is
  * used to uniquely describe a payload without the overhead of its buffer.
@@ -69,21 +67,21 @@ public:
 	/** Construct a null payload. */
 	FPayload() = default;
 
-	/** Construct a payload from the hash of a compressed buffer. */
-	inline FPayload(const FPayloadId& Id, const FIoHash& BufferHash);
-	/** Construct a payload from a compressed buffer and its hash. See FCompressedBuffer::Compress. */
-	inline FPayload(const FPayloadId& Id, const FIoHash& BufferHash, FCompressedBuffer Buffer);
-	/** Construct a payload from a compressed buffer and calculates its hash. See FCompressedBuffer::Compress. */
-	inline FPayload(const FPayloadId& Id, FCompressedBuffer Buffer);
+	/** Construct a payload from a raw buffer hash. */
+	inline FPayload(const FPayloadId& Id, const FIoHash& RawHash);
+
+	/** Construct a payload from a compressed buffer. */
+	inline FPayload(const FPayloadId& Id, const FCompressedBuffer& Data);
+	inline FPayload(const FPayloadId& Id, FCompressedBuffer&& Data);
 
 	/** Returns the PayloadId for the payload. */
 	inline const FPayloadId& GetId() const { return Id; }
 
-	/** Returns the IoHash of the compressed buffer for the payload. */
-	inline const FIoHash& GetHash() const { return Hash; }
+	/** Returns the IoHash of the raw buffer (uncompressed) for the payload. */
+	inline const FIoHash& GetRawHash() const { return RawHash; }
 
 	/** Returns the compressed buffer for the payload. May be null. */
-	inline const FCompressedBuffer& GetBuffer() const { return Buffer; }
+	inline const FCompressedBuffer& GetData() const { return Data; }
 
 	/** Whether this is null. */
 	inline bool IsNull() const { return Id.IsNull(); }
@@ -97,8 +95,8 @@ public:
 
 private:
 	FPayloadId Id;
-	FIoHash Hash;
-	FCompressedBuffer Buffer;
+	FIoHash RawHash;
+	FCompressedBuffer Data;
 };
 
 ///////////////////////////////////////////////////////////////////////////////////////////////////////////////////////
@@ -112,13 +110,13 @@ inline FPayloadId::FPayloadId(const FMemoryView InId)
 
 inline FPayloadId FPayloadId::FromHash(const FIoHash& Hash)
 {
-	checkf(!Hash.IsZero(), TEXT("PayloadId requires a non-zero hash."));
+	checkf(!Hash.IsZero(), TEXT("FPayloadId requires a non-zero hash."));
 	return FPayloadId(MakeMemoryView(Hash.GetBytes()).Left(sizeof(ByteArray)));
 }
 
 inline FPayloadId FPayloadId::FromName(const FAnsiStringView Name)
 {
-	checkf(!Name.IsEmpty(), TEXT("PayloadId requires a non-empty name."));
+	checkf(!Name.IsEmpty(), TEXT("FPayloadId requires a non-empty name."));
 	return FPayloadId::FromHash(FIoHash::HashBuffer(Name.GetData(), Name.Len()));
 }
 
@@ -129,41 +127,27 @@ inline FPayloadId FPayloadId::FromName(const FWideStringView Name)
 
 ///////////////////////////////////////////////////////////////////////////////////////////////////////////////////////
 
-inline FPayload::FPayload(const FPayloadId& InId, const FIoHash& InHash)
+inline FPayload::FPayload(const FPayloadId& InId, const FIoHash& InRawHash)
 	: Id(InId)
-	, Hash(InHash)
+	, RawHash(InRawHash)
 {
-	checkf(Id.IsValid(), TEXT("A valid PayloadId is required to construct a payload."));
+	checkf(Id.IsValid(), TEXT("A valid ID is required to construct a payload."));
 }
 
-inline FPayload::FPayload(const FPayloadId& InId, const FIoHash& InHash, FCompressedBuffer InBuffer)
+inline FPayload::FPayload(const FPayloadId& InId, const FCompressedBuffer& InData)
 	: Id(InId)
-	, Hash(InHash)
-	, Buffer(MoveTemp(InBuffer))
+	, RawHash(InData.GetRawHash())
+	, Data(InData)
 {
-	checkf(Id.IsValid(), TEXT("A valid PayloadId is required to construct a payload."));
-	if (Buffer.GetCompressedSize())
-	{
-		checkfSlow(Hash == FIoHash::HashBuffer(Buffer.GetCompressed()),
-			TEXT("Provided compressed data hash %s does not match calculated hash %s"),
-			*LexToString(Hash), *LexToString(FIoHash::HashBuffer(Buffer.GetCompressed())));
-		checkf(!Hash.IsZero(), TEXT("A non-empty compressed data buffer must have a non-zero hash."));
-	}
-	else
-	{
-		checkf(Hash.IsZero(), TEXT("A null or empty compressed data buffer must use a hash of zero."));
-	}
+	checkf(Id.IsValid(), TEXT("A valid ID is required to construct a payload."));
 }
 
-inline FPayload::FPayload(const FPayloadId& InId, FCompressedBuffer InBuffer)
+inline FPayload::FPayload(const FPayloadId& InId, FCompressedBuffer&& InData)
 	: Id(InId)
-	, Buffer(MoveTemp(InBuffer))
+	, RawHash(InData.GetRawHash())
+	, Data(MoveTemp(InData))
 {
-	checkf(Id.IsValid(), TEXT("A valid PayloadId is required to construct a payload."));
-	if (Buffer.GetCompressedSize())
-	{
-		Hash = FIoHash::HashBuffer(Buffer.GetCompressed());
-	}
+	checkf(Id.IsValid(), TEXT("A valid ID is required to construct a payload."));
 }
 
 ///////////////////////////////////////////////////////////////////////////////////////////////////////////////////////
@@ -203,57 +187,34 @@ inline bool FPayloadId::IsNull() const
 
 ///////////////////////////////////////////////////////////////////////////////////////////////////////////////////////
 
-/** Compare payloads by their PayloadId and compressed buffer hash. */
+/** Compare payloads by their ID and raw buffer hash. */
 inline bool operator==(const FPayload& A, const FPayload& B)
 {
-	return A.GetId() == B.GetId() && A.GetHash() == B.GetHash();
+	return A.GetId() == B.GetId() && A.GetRawHash() == B.GetRawHash();
 }
 
-/** Compare payloads by their PayloadId and compressed buffer hash. */
+/** Compare payloads by their ID and raw buffer hash. */
 inline bool operator<(const FPayload& A, const FPayload& B)
 {
-	return A.GetId() == B.GetId() && A.GetHash() == B.GetHash();
+	return A.GetId() == B.GetId() && A.GetRawHash() == B.GetRawHash();
 }
 
-/** Compare payloads by their PayloadId. */
+/** Compare payloads by their ID. */
 struct FPayloadEqualById
 {
-	inline bool operator()(const FPayload& A, const FPayload& B) const
-	{
-		return A.GetId() == B.GetId();
-	}
-
-	inline bool operator()(const FPayload& A, const FPayloadId& B) const
-	{
-		return A.GetId() == B;
-	}
-
-	inline bool operator()(const FPayloadId& A, const FPayload& B) const
-	{
-		return A == B.GetId();
-	}
+	inline bool operator()(const FPayload& A, const FPayload& B) const { return A.GetId() == B.GetId(); }
+	inline bool operator()(const FPayload& A, const FPayloadId& B) const { return A.GetId() == B; }
+	inline bool operator()(const FPayloadId& A, const FPayload& B) const { return A == B.GetId(); }
 };
 
-/** Compare payloads by their PayloadId. */
+/** Compare payloads by their ID. */
 struct FPayloadLessById
 {
-	inline bool operator()(const FPayload& A, const FPayload& B) const
-	{
-		return A.GetId() < B.GetId();
-	}
-
-	inline bool operator()(const FPayload& A, const FPayloadId& B) const
-	{
-		return A.GetId() < B;
-	}
-
-	inline bool operator()(const FPayloadId& A, const FPayload& B) const
-	{
-		return A < B.GetId();
-	}
+	inline bool operator()(const FPayload& A, const FPayload& B) const { return A.GetId() < B.GetId(); }
+	inline bool operator()(const FPayload& A, const FPayloadId& B) const { return A.GetId() < B; }
+	inline bool operator()(const FPayloadId& A, const FPayload& B) const { return A < B.GetId(); }
 };
 
 ///////////////////////////////////////////////////////////////////////////////////////////////////////////////////////
 
-} // DerivedData
-} // UE
+} // UE::DerivedData
