Ensure that the pak writing DDC backend doesn't silently ignore overwrites of keys that are already in the pak file.  The writes can happen, but they will emit a warning and leave the old copy in the serialized file as an abandoned blob.  If the overwrite happens to be the same size as the original blob, the overwrite can successfully be done without wasting space.

diff --git a/Engine/Source/Developer/DerivedDataCache/Private/PakFileDerivedDataBackend.cpp b/Engine/Source/Developer/DerivedDataCache/Private/PakFileDerivedDataBackend.cpp
--- a/Engine/Source/Developer/DerivedDataCache/Private/PakFileDerivedDataBackend.cpp
+++ b/Engine/Source/Developer/DerivedDataCache/Private/PakFileDerivedDataBackend.cpp
@@ -147,34 +147,59 @@ FDerivedDataBackendInterface::EPutStatus FPakFileDerivedDataBackend::PutCachedDa
 	{
 		FScopeLock ScopeLock(&SynchronizationObject);
 		FString Key(CacheKey);
-		FCacheValue* Item = CacheItems.Find(FString(CacheKey));
-		if (!Item)
+		TOptional<uint32> Crc;
+		check(InData.Num());
+		check(Key.Len());
+		check(FileHandle);
+		check(FileHandle->IsSaving());
+
+		if (bPutEvenIfExists)
 		{
-			check(InData.Num());
-			check(Key.Len());
-			uint32 Crc = FCrc::MemCrc_DEPRECATED(InData.GetData(), InData.Num());
-			check(FileHandle);
-			check(FileHandle->IsSaving());
-			int64 Offset = FileHandle->Tell();
-			if (Offset < 0)
-			{
-				CacheItems.Empty();
-				FileHandle.Reset();
-				UE_LOG(LogDerivedDataCache, Fatal, TEXT("Could not write pak file...out of disk space?"));
-				return EPutStatus::NotCached;
-			}
-			else
+			if (FCacheValue* Item = CacheItems.Find(FString(CacheKey)))
 			{
-				COOK_STAT(Timer.AddHit(InData.Num()));
-				// NOTE: Gross that FArchive doesn't have a const version just for saving...
-				FileHandle->Serialize(const_cast<uint8*>(InData.GetData()), int64(InData.Num()));
-				UE_LOG(LogDerivedDataCache, Verbose, TEXT("FPakFileDerivedDataBackend: Put %s"), CacheKey);
-				CacheItems.Add(Key,FCacheValue(Offset, InData.Num(), Crc));
-				return EPutStatus::Cached;
+				// If there was an existing entry for this key, if it had the same contents, do nothing as the desired value is already stored.
+				// If the contents differ, replace it if the size hasn't changed, but if the size has changed, 
+				// remove the existing entry from the index but leave they actual data payload in place as it is too
+				// costly to go back and attempt to rewrite all offsets and shift all bytes that follow it in the file.
+				if (Item->Size == InData.Num())
+				{
+					COOK_STAT(Timer.AddHit(InData.Num()));
+					Crc = FCrc::MemCrc_DEPRECATED(InData.GetData(), InData.Num());
+					if (Crc.GetValue() != Item->Crc)
+					{
+						int64 Offset = FileHandle->Tell();
+						FileHandle->Seek(Item->Offset);
+						FileHandle->Serialize(const_cast<uint8*>(InData.GetData()), int64(InData.Num()));
+						Item->Crc = Crc.GetValue();
+						FileHandle->Seek(Offset);
+					}
+					return EPutStatus::Cached;
+				}
+
+				UE_LOG(LogDerivedDataCache, Warning, TEXT("FPakFileDerivedDataBackend: Repeated put of %s with different sized contents.  Multiple contents will be in the file, but only the last will be in the index.  This has wasted %" INT64_FMT " bytes in the file."), CacheKey, Item->Offset);
+				CacheItems.Remove(Key);
 			}
 		}
+
+		int64 Offset = FileHandle->Tell();
+		if (Offset < 0)
+		{
+			CacheItems.Empty();
+			FileHandle.Reset();
+			UE_LOG(LogDerivedDataCache, Fatal, TEXT("Could not write pak file...out of disk space?"));
+			return EPutStatus::NotCached;
+		}
 		else
 		{
+			COOK_STAT(Timer.AddHit(InData.Num()));
+			if (!Crc.IsSet())
+			{
+				Crc = FCrc::MemCrc_DEPRECATED(InData.GetData(), InData.Num());
+			}
+			// NOTE: Gross that FArchive doesn't have a const version just for saving...
+			FileHandle->Serialize(const_cast<uint8*>(InData.GetData()), int64(InData.Num()));
+			UE_LOG(LogDerivedDataCache, Verbose, TEXT("FPakFileDerivedDataBackend: Put %s"), CacheKey);
+			CacheItems.Add(Key,FCacheValue(Offset, InData.Num(), Crc.GetValue()));
 			return EPutStatus::Cached;
 		}
 	}
