Remove UE4 strings

diff --git a/Engine/Source/Developer/DerivedDataCache/Private/HttpDerivedDataBackend.cpp b/Engine/Source/Developer/DerivedDataCache/Private/HttpDerivedDataBackend.cpp
--- a/Engine/Source/Developer/DerivedDataCache/Private/HttpDerivedDataBackend.cpp
+++ b/Engine/Source/Developer/DerivedDataCache/Private/HttpDerivedDataBackend.cpp
@@ -408,7 +408,7 @@ private:
 	Result PerformBlocking(const TCHAR* Uri, RequestVerb Verb, uint32 ContentLength)
 	{
 		static const char* CommonHeaders[] = {
-			"User-Agent: UE4",
+			"User-Agent: Unreal Engine",
 			nullptr
 		};
 
