Avoid stalling during shutdown due to queued async PUTs to non-local backends  in DDC.

diff --git a/Engine/Source/Developer/DerivedDataCache/Private/DerivedDataBackendAsyncPutWrapper.cpp b/Engine/Source/Developer/DerivedDataCache/Private/DerivedDataBackendAsyncPutWrapper.cpp
--- a/Engine/Source/Developer/DerivedDataCache/Private/DerivedDataBackendAsyncPutWrapper.cpp
+++ b/Engine/Source/Developer/DerivedDataCache/Private/DerivedDataBackendAsyncPutWrapper.cpp
@@ -37,6 +37,17 @@ public:
 	{
 		check(InnerBackend);
 	}
+
+	bool ShouldAbortForShutdown()
+	{
+		using ESpeedClass = FDerivedDataBackendInterface::ESpeedClass;
+		ESpeedClass SpeedClass = InnerBackend->GetSpeedClass();
+		if (SpeedClass == ESpeedClass::Local)
+		{
+			return false;
+		}
+		return !GIsBuildMachine && FDerivedDataBackend::Get().IsShuttingDown();
+	}
 		
 	/** Call the inner backend and when that completes, remove the memory cache */
 	void DoWork()
@@ -44,6 +55,12 @@ public:
 		TRACE_CPUPROFILER_EVENT_SCOPE(DDCPut_DoWork);
 		COOK_STAT(auto Timer = UsageStats.TimePut());
 
+		if (ShouldAbortForShutdown())
+		{
+			Abandon();
+			return;
+		}
+
 		using EPutStatus = FDerivedDataBackendInterface::EPutStatus;
 		EPutStatus Status = EPutStatus::NotCached;
 
