DDC: Tweak names and namespace formatting for consistency

#rnx

diff --git a/Engine/Source/Developer/DerivedDataCache/Private/DerivedDataCachePrivate.h b/Engine/Source/Developer/DerivedDataCache/Private/DerivedDataCachePrivate.h
--- a/Engine/Source/Developer/DerivedDataCache/Private/DerivedDataCachePrivate.h
+++ b/Engine/Source/Developer/DerivedDataCache/Private/DerivedDataCachePrivate.h
@@ -4,28 +4,17 @@
 
 #include "Containers/StringFwd.h"
 
-namespace UE
-{
-namespace DerivedData
-{
-
-class FCacheBucket;
-class FCacheRecordBuilder;
-struct FCacheKey;
+namespace UE::DerivedData { class FCacheBucket; }
+namespace UE::DerivedData { class FCacheRecordBuilder; }
+namespace UE::DerivedData { struct FCacheKey; }
 
-namespace Private
+namespace UE::DerivedData::Private
 {
 
-///////////////////////////////////////////////////////////////////////////////////////////////////////////////////////
-
 // Implemented in DerivedDataCacheKey.cpp
 FCacheBucket CreateCacheBucket(FStringView Name);
 
 // Implemented in DerivedDataCacheRecord.cpp
 FCacheRecordBuilder CreateCacheRecordBuilder(const FCacheKey& Key);
 
-///////////////////////////////////////////////////////////////////////////////////////////////////////////////////////
-
-} // Private
-} // DerivedData
-} // UE
+} // UE::DerivedData::Private
diff --git a/Engine/Source/Developer/DerivedDataCache/Private/DerivedDataCacheRecord.cpp b/Engine/Source/Developer/DerivedDataCache/Private/DerivedDataCacheRecord.cpp
--- a/Engine/Source/Developer/DerivedDataCache/Private/DerivedDataCacheRecord.cpp
+++ b/Engine/Source/Developer/DerivedDataCache/Private/DerivedDataCacheRecord.cpp
@@ -28,7 +28,7 @@ public:
 	FPayloadId AddAttachment(FPayload&& Payload) final;
 
 	FCacheRecord Build() final;
-	FRequest BuildAsync(FOnCacheRecordComplete&& Callback, EPriority Priority) final;
+	FRequest BuildAsync(FOnCacheRecordComplete&& OnComplete, EPriority Priority) final;
 
 	FCacheKey Key;
 	FCbObject Meta;
@@ -205,9 +205,9 @@ FCacheRecord FCacheRecordBuilderInternal::Build()
 	return CreateCacheRecord(new FCacheRecordInternal(MoveTemp(*this)));
 }
 
-FRequest FCacheRecordBuilderInternal::BuildAsync(FOnCacheRecordComplete&& Callback, EPriority Priority)
+FRequest FCacheRecordBuilderInternal::BuildAsync(FOnCacheRecordComplete&& OnComplete, EPriority Priority)
 {
-	Callback(Build());
+	OnComplete(Build());
 	return FRequest();
 }
 
diff --git a/Engine/Source/Developer/DerivedDataCache/Public/DerivedDataCacheRecord.h b/Engine/Source/Developer/DerivedDataCache/Public/DerivedDataCacheRecord.h
--- a/Engine/Source/Developer/DerivedDataCache/Public/DerivedDataCacheRecord.h
+++ b/Engine/Source/Developer/DerivedDataCache/Public/DerivedDataCacheRecord.h
@@ -46,7 +46,7 @@ public:
 	virtual FPayloadId AddAttachment(const FSharedBuffer& Buffer, const FPayloadId& Id) = 0;
 	virtual FPayloadId AddAttachment(FPayload&& Payload) = 0;
 	virtual FCacheRecord Build() = 0;
-	virtual FRequest BuildAsync(FOnCacheRecordComplete&& Callback, EPriority Priority) = 0;
+	virtual FRequest BuildAsync(FOnCacheRecordComplete&& OnComplete, EPriority Priority) = 0;
 };
 
 FCacheRecordBuilder CreateCacheRecordBuilder(ICacheRecordBuilderInternal* RecordBuilder);
@@ -214,9 +214,9 @@ public:
 	 * Prefer Build() when the value and attachments are added by payload, as compression is already
 	 * complete and BuildAsync() will complete immediately in that case.
 	 */
-	inline FRequest BuildAsync(FOnCacheRecordComplete&& Callback, EPriority Priority)
+	inline FRequest BuildAsync(FOnCacheRecordComplete&& OnComplete, EPriority Priority)
 	{
-		return RecordBuilder->BuildAsync(MoveTemp(Callback), Priority);
+		return RecordBuilder->BuildAsync(MoveTemp(OnComplete), Priority);
 	}
 
 private:
