DDC: Simplify FCacheBucket

#rnx

diff --git a/Engine/Source/Developer/DerivedDataCache/Private/DerivedDataCacheKey.cpp b/Engine/Source/Developer/DerivedDataCache/Private/DerivedDataCacheKey.cpp
--- a/Engine/Source/Developer/DerivedDataCache/Private/DerivedDataCacheKey.cpp
+++ b/Engine/Source/Developer/DerivedDataCache/Private/DerivedDataCacheKey.cpp
@@ -7,128 +7,105 @@
 #include "Math/UnrealMathUtility.h"
 #include "Misc/ScopeRWLock.h"
 
-namespace UE
+namespace UE::DerivedData::Private
 {
-namespace DerivedData
+
+class FCacheBucketOwner : public FCacheBucket
 {
-namespace Private
+public:
+	static const TCHAR* AllocName(FStringView Bucket);
+
+	inline explicit FCacheBucketOwner(FStringView Bucket);
+	inline FCacheBucketOwner(FCacheBucketOwner&& Bucket);
+	inline ~FCacheBucketOwner();
+
+	FCacheBucketOwner(const FCacheBucketOwner&) = delete;
+	FCacheBucketOwner& operator=(const FCacheBucketOwner&) = delete;
+
+	using FCacheBucket::operator==;
+
+	inline bool operator==(FStringView Bucket) const { return ToString<TCHAR>().Equals(Bucket, ESearchCase::IgnoreCase); }
+};
+
+inline const TCHAR* FCacheBucketOwner::AllocName(FStringView Bucket)
 {
+	FTCHARToUTF8 Bucket8Conversion(Bucket);
+	FAnsiStringView Bucket8 = Bucket8Conversion;
 
-///////////////////////////////////////////////////////////////////////////////////////////////////////////////////////
+	const int32 BucketLen = Bucket.Len();
+	const int32 Bucket8Len = Bucket8.Len();
+	const int32 PrefixSize = FMath::DivideAndRoundUp<int32>(2, sizeof(TCHAR));
+	TCHAR* Buffer = new TCHAR[PrefixSize + BucketLen + 1 + FMath::DivideAndRoundUp<int32>(Bucket8Len + 1, sizeof(TCHAR))];
+	Buffer += PrefixSize;
+	const TCHAR* const Name = Buffer;
 
-class FCacheBuckets
+	reinterpret_cast<uint8*>(Buffer)[LengthOffset] = static_cast<uint8>(BucketLen);
+	reinterpret_cast<uint8*>(Buffer)[LengthOffset8] = static_cast<uint8>(Bucket8Len);
+
+	Bucket.CopyString(Buffer, BucketLen);
+	Buffer += BucketLen;
+	*Buffer++ = TEXT('\0');
+
+	ANSICHAR* Buffer8 = reinterpret_cast<ANSICHAR*>(Buffer);
+	Bucket8.CopyString(Buffer8, Bucket8Len);
+	Buffer8 += Bucket8Len;
+	*Buffer8++ = '\0';
+
+	return Name;
+}
+
+inline FCacheBucketOwner::FCacheBucketOwner(FStringView Bucket)
+	: FCacheBucket(AllocName(Bucket))
 {
-public:
-	inline const TCHAR* FindOrAdd(FStringView Name16, FAnsiStringView Name8)
+}
+
+inline FCacheBucketOwner::FCacheBucketOwner(FCacheBucketOwner&& Bucket)
+	: FCacheBucket(Bucket)
+{
+	Bucket.Reset();
+}
+
+inline FCacheBucketOwner::~FCacheBucketOwner()
+{
+	if (IsValid())
 	{
-		const uint32 Hash = GetTypeHash(Name16);
-		{
-			FReadScopeLock ReadLock(Lock);
-			if (const FBucket* Bucket = Buckets.FindByHash(Hash, Name16))
-			{
-				return Bucket->GetName();
-			}
-		}
-		{
-			FWriteScopeLock WriteLock(Lock);
-			const FBucket& Bucket = Buckets.FindOrAddByHash(Hash, FBucket(Name16, Name8));
-			return Bucket.GetName();
-		}
+		const int32 PrefixSize = FMath::DivideAndRoundUp<int32>(2, sizeof(TCHAR));
+		delete[] (ToString<TCHAR>().GetData() - PrefixSize);
 	}
+}
 
-private:
-	class FBucket
-	{
-	public:
-		inline FBucket(FStringView Bucket, FAnsiStringView Bucket8)
-		{
-			const int32 BucketLen = Bucket.Len();
-			const int32 Bucket8Len = Bucket8.Len();
-			const int32 PrefixSize = FMath::DivideAndRoundUp<int32>(2, sizeof(TCHAR));
-			TCHAR* Buffer = new TCHAR[PrefixSize + BucketLen + 1 + FMath::DivideAndRoundUp<int32>(Bucket8Len + 1, sizeof(TCHAR))];
-			Buffer += PrefixSize;
-			Name = Buffer;
-
-			reinterpret_cast<uint8*>(Buffer)[-2] = static_cast<uint8>(BucketLen);
-			reinterpret_cast<uint8*>(Buffer)[-1] = static_cast<uint8>(Bucket8Len);
-
-			Bucket.CopyString(Buffer, BucketLen);
-			Buffer += BucketLen;
-			*Buffer++ = TEXT('\0');
-
-			ANSICHAR* Buffer8 = reinterpret_cast<ANSICHAR*>(Buffer);
-			Bucket8.CopyString(Buffer8, Bucket8Len);
-			Buffer8 += Bucket8Len;
-			*Buffer8++ = '\0';
-		}
-
-		FBucket(const FBucket&) = delete;
-		FBucket& operator=(const FBucket&) = delete;
-
-		inline FBucket(FBucket&& Bucket)
-			: Name(Bucket.Name)
-		{
-			Bucket.Name = nullptr;
-		}
-
-		FBucket& operator=(FBucket&& Bucket)
-		{
-			Name = Bucket.Name;
-			Bucket.Name = nullptr;
-			return *this;
-		}
-
-		inline ~FBucket()
-		{
-			if (Name)
-			{
-				const int32 PrefixSize = FMath::DivideAndRoundUp<int32>(2, sizeof(TCHAR));
-				delete[] (Name - PrefixSize);
-			}
-		}
-
-		inline FStringView ToString() const
-		{
-			return FStringView(Name, reinterpret_cast<const uint8*>(Name)[-2]);
-		}
-
-		inline bool operator==(const FBucket& Bucket) const
-		{
-			return Name == Bucket.Name;
-		}
-
-		inline bool operator==(FStringView Bucket) const
-		{
-			return ToString().Equals(Bucket, ESearchCase::IgnoreCase);
-		}
-
-		friend inline uint32 GetTypeHash(const FBucket& Bucket)
-		{
-			return GetTypeHash(Bucket.ToString());
-		}
-
-		inline const TCHAR* GetName() const { return Name; }
-
-	private:
-		const TCHAR* Name;
-	};
+class FCacheBuckets
+{
+public:
+	inline FCacheBucket FindOrAdd(FStringView Name);
 
+private:
 	FRWLock Lock;
-	TSet<FBucket> Buckets;
+	TSet<FCacheBucketOwner> Buckets;
 };
 
-FCacheBucket CreateCacheBucket(FStringView Name)
+inline FCacheBucket FCacheBuckets::FindOrAdd(FStringView Name)
 {
-	FTCHARToUTF8 Name8(Name);
-	checkf(!Name.IsEmpty() && Name.Len() < 256 && Name8.Length() < 256 && Algo::AllOf(Name, FChar::IsAlnum),
+	const uint32 Hash = GetTypeHash(Name);
+	if (FReadScopeLock ReadLock(Lock); const FCacheBucketOwner* Bucket = Buckets.FindByHash(Hash, Name))
+	{
+		return *Bucket;
+	}
+
+	checkf(FCString::IsPureAnsi(Name.GetData(), Name.Len()) && Algo::AllOf(Name, FChar::IsAlnum) &&
+		!Name.IsEmpty() && Name.Len() < 256,
 		TEXT("Invalid cache bucket name '%.*s' must be alphanumeric, non-empty, and contain fewer than 256 code units."),
 		Name.Len(), Name.GetData());
-	static FCacheBuckets Buckets;
-	return FCacheBucket(FCacheBucketName{Buckets.FindOrAdd(Name, Name8)});
+
+	FCacheBucketOwner LocalBucket(Name);
+	FWriteScopeLock WriteLock(Lock);
+	return Buckets.FindOrAddByHash(Hash, MoveTemp(LocalBucket));
 }
 
-///////////////////////////////////////////////////////////////////////////////////////////////////////////////////////
+FCacheBucket CreateCacheBucket(FStringView Name)
+{
+	static FCacheBuckets Buckets;
+	return Buckets.FindOrAdd(Name);
+}
 
-} // Private
-} // DerivedData
-} // UE
+} // UE::DerivedData::Private
diff --git a/Engine/Source/Developer/DerivedDataCache/Public/DerivedDataCacheKey.h b/Engine/Source/Developer/DerivedDataCache/Public/DerivedDataCacheKey.h
--- a/Engine/Source/Developer/DerivedDataCache/Public/DerivedDataCacheKey.h
+++ b/Engine/Source/Developer/DerivedDataCache/Public/DerivedDataCacheKey.h
@@ -9,13 +9,6 @@
 #include "IO/IoHash.h"
 #include "Templates/TypeHash.h"
 
-namespace UE::DerivedData::Private
-{
-
-struct FCacheBucketName { const TCHAR* Name; };
-
-} // UE::DerivedData::Private
-
 namespace UE::DerivedData
 {
 
@@ -61,20 +54,21 @@ public:
 	template <>
 	inline FAnsiStringView ToString<ANSICHAR>() const;
 
-public:
-	// Internal API
-
-	explicit FCacheBucket(Private::FCacheBucketName InName)
-		: Name(InName.Name)
+protected:
+	explicit FCacheBucket(const TCHAR* InName)
+		: Name(InName)
 	{
 	}
 
+	static constexpr int32 LengthOffset = -1;
+	static constexpr int32 LengthOffset8 = -2;
+
 private:
 	/**
-	 * Name in TCHAR followed by UTF-8 in ANSICHAR, both null-terminated.
+	 * Name as TCHAR followed by the name as UTF-8, both null-terminated.
 	 *
-	 * The byte preceding the TCHAR name is the number of UTF-8 code units.
-	 * The byte preceding the UTF-8 name length is the number of TCHAR code units.
+	 * The byte preceding the TCHAR name is the number of TCHAR code units excluding the terminator.
+	 * The byte preceding the TCHAR name length is the number of UTF-8 code units.
 	 */
 	const TCHAR* Name = nullptr;
 };
@@ -110,15 +104,15 @@ inline const FCachePayloadKey FCachePayloadKey::Empty;
 template <>
 inline FStringView FCacheBucket::ToString<TCHAR>() const
 {
-	return Name ? FStringView(Name, reinterpret_cast<const uint8*>(Name)[-2]) : FStringView();
+	return Name ? FStringView(Name, reinterpret_cast<const uint8*>(Name)[LengthOffset]) : FStringView();
 }
 
 template <>
 inline FAnsiStringView FCacheBucket::ToString<ANSICHAR>() const
 {
 	return Name ? FAnsiStringView(
-		reinterpret_cast<const ANSICHAR*>(Name + reinterpret_cast<const uint8*>(Name)[-2] + 1),
-		reinterpret_cast<const uint8*>(Name)[-1]) : FAnsiStringView();
+		reinterpret_cast<const ANSICHAR*>(Name + reinterpret_cast<const uint8*>(Name)[LengthOffset] + 1),
+		reinterpret_cast<const uint8*>(Name)[LengthOffset8]) : FAnsiStringView();
 }
 
 template <typename CharType>
