Fix bug in the HttpDerivedDataBackend where an FRequest may not be returned to the FRequestPool if we're aborting queued PUTs due to a shutdown request.

diff --git a/Engine/Source/Developer/DerivedDataCache/Private/HttpDerivedDataBackend.cpp b/Engine/Source/Developer/DerivedDataCache/Private/HttpDerivedDataBackend.cpp
--- a/Engine/Source/Developer/DerivedDataCache/Private/HttpDerivedDataBackend.cpp
+++ b/Engine/Source/Developer/DerivedDataCache/Private/HttpDerivedDataBackend.cpp
@@ -1510,24 +1510,31 @@ private:
 	{
 		while (Request)
 		{
-			if (ShouldAbortForShutdown())
+			// Make sure that whether we early exit or execute past the end of this scope that
+			// the request is released back to the pool.
 			{
-				return;
-			}
-
-			while (FQueuedEntry* Entry = QueuedPuts.Pop())
-			{
-				Request->Reset();
-				PerformPut(Request, *Entry->Namespace, *Entry->Bucket, *Entry->CacheKey, Entry->Data, UsageStats);
-				delete Entry;
+				ON_SCOPE_EXIT
+				{
+					Pool->ReleaseRequestToPool(Request);
+				};
 
 				if (ShouldAbortForShutdown())
 				{
 					return;
 				}
-			}
 
-			Pool->ReleaseRequestToPool(Request);
+				while (FQueuedEntry* Entry = QueuedPuts.Pop())
+				{
+					Request->Reset();
+					PerformPut(Request, *Entry->Namespace, *Entry->Bucket, *Entry->CacheKey, Entry->Data, UsageStats);
+					delete Entry;
+
+					if (ShouldAbortForShutdown())
+					{
+						return;
+					}
+				}
+			}
 
 			// An entry may have been queued while the request was being released.
 			if (QueuedPuts.IsEmpty())
