Add support for ReadOnly={true|false} and EnvHostOverride={string} as configurable parameters on an HTTP DDC backend.

diff --git a/Engine/Source/Developer/DerivedDataCache/Private/DerivedDataBackends.cpp b/Engine/Source/Developer/DerivedDataCache/Private/DerivedDataBackends.cpp
--- a/Engine/Source/Developer/DerivedDataCache/Private/DerivedDataBackends.cpp
+++ b/Engine/Source/Developer/DerivedDataCache/Private/DerivedDataBackends.cpp
@@ -728,6 +728,18 @@ public:
 			return nullptr;
 		}
 
+		// Check the EnvHostOverride environment variable to allow persistent overriding of data cache path, eg for offsite workers.
+		FString EnvHostOverride;
+		if( FParse::Value( Entry, TEXT("EnvHostOverride="), EnvHostOverride ) )
+		{
+			FString ServiceUrlEnv = FPlatformMisc::GetEnvironmentVariable( *EnvHostOverride );
+			if( ServiceUrlEnv.Len() > 0 )
+			{
+				ServiceUrl = ServiceUrlEnv;
+				UE_LOG( LogDerivedDataCache, Log, TEXT("Found environment variable for Host %s=%s"), *EnvHostOverride, *ServiceUrl );
+			}
+		}
+
 		// Allow Host to be overriden via CommandLineHostOverride
 		FString CommandLineOverride;
 		if( FParse::Value( Entry, TEXT("CommandLineHostOverride="), CommandLineOverride ) )
@@ -740,6 +752,12 @@ public:
 			}
 		}
 
+		if (ServiceUrl == TEXT("None"))
+		{
+			UE_LOG( LogDerivedDataCache, Log, TEXT("Disabling %s data cache - host set to 'None'."), NodeName );
+			return nullptr;
+		}
+
 		FDerivedDataBackendInterface::ESpeedClass ForceSpeedClass = FDerivedDataBackendInterface::ESpeedClass::Unknown;
 		FString ForceSpeedClassValue;
 		if (FParse::Value(FCommandLine::Get(), TEXT("HttpForceSpeedClass="), ForceSpeedClassValue))
@@ -787,7 +805,9 @@ public:
 			return nullptr;
 		}
 
-		FHttpDerivedDataBackend* backend = new FHttpDerivedDataBackend(*ServiceUrl, *Namespace, *OAuthProvider, *OAuthClientId, *OAuthSecret);
+		const bool bReadOnly = GetParsedBool(Entry, TEXT("ReadOnly="));
+
+		FHttpDerivedDataBackend* backend = new FHttpDerivedDataBackend(*ServiceUrl, *Namespace, *OAuthProvider, *OAuthClientId, *OAuthSecret, bReadOnly);
 
 		if (ForceSpeedClass != FDerivedDataBackendInterface::ESpeedClass::Unknown)
 		{
diff --git a/Engine/Source/Developer/DerivedDataCache/Private/HttpDerivedDataBackend.cpp b/Engine/Source/Developer/DerivedDataCache/Private/HttpDerivedDataBackend.cpp
--- a/Engine/Source/Developer/DerivedDataCache/Private/HttpDerivedDataBackend.cpp
+++ b/Engine/Source/Developer/DerivedDataCache/Private/HttpDerivedDataBackend.cpp
@@ -1752,7 +1752,8 @@ FHttpDerivedDataBackend::FHttpDerivedDataBackend(
 	const TCHAR* InNamespace, 
 	const TCHAR* InOAuthProvider,
 	const TCHAR* InOAuthClientId,
-	const TCHAR* InOAuthSecret)
+	const TCHAR* InOAuthSecret,
+	const bool bInReadOnly)
 	: Domain(InServiceUrl)
 	, Namespace(InNamespace)
 	, DefaultBucket(TEXT("default"))
@@ -1761,6 +1762,7 @@ FHttpDerivedDataBackend::FHttpDerivedDataBackend(
 	, OAuthSecret(InOAuthSecret)
 	, Access(nullptr)
 	, bIsUsable(false)
+	, bReadOnly(bInReadOnly)
 	, FailedLoginAttempts(0)
 	, SpeedClass(ESpeedClass::Slow)
 {
@@ -1793,7 +1795,7 @@ bool FHttpDerivedDataBackend::TryToPrefetch(TConstArrayView<FString> CacheKeys)
 
 bool FHttpDerivedDataBackend::WouldCache(const TCHAR* CacheKey, TArrayView<const uint8> InData)
 {
-	return true;
+	return IsWritable();
 }
 
 FHttpDerivedDataBackend::ESpeedClass FHttpDerivedDataBackend::GetSpeedClass() const
@@ -2146,6 +2148,12 @@ FDerivedDataBackendInterface::EPutStatus FHttpDerivedDataBackend::PutCachedData(
 {
 	TRACE_CPUPROFILER_EVENT_SCOPE(HttpDDC_Put);
 
+	if (!IsWritable())
+	{
+		UE_LOG(LogDerivedDataCache, Verbose, TEXT("%s is read only. Skipping put of %s"), *GetName(), CacheKey);
+		return EPutStatus::NotCached;
+	}
+
 #if WITH_DATAREQUEST_HELPER
 	for (int32 Attempts = 0; Attempts < UE_HTTPDDC_MAX_ATTEMPTS; ++Attempts)
 	{
@@ -2214,7 +2222,7 @@ FDerivedDataBackendInterface::EPutStatus FHttpDerivedDataBackend::PutCachedData(
 void FHttpDerivedDataBackend::RemoveCachedData(const TCHAR* CacheKey, bool bTransient)
 {
 	// do not remove transient data as Jupiter does its own verification of the content and cleans itself up
-	if (bTransient)
+	if (!IsWritable() || bTransient)
 		return;
 	
 	TRACE_CPUPROFILER_EVENT_SCOPE(HttpDDC_Remove);
diff --git a/Engine/Source/Developer/DerivedDataCache/Private/HttpDerivedDataBackend.h b/Engine/Source/Developer/DerivedDataCache/Private/HttpDerivedDataBackend.h
--- a/Engine/Source/Developer/DerivedDataCache/Private/HttpDerivedDataBackend.h
+++ b/Engine/Source/Developer/DerivedDataCache/Private/HttpDerivedDataBackend.h
@@ -39,7 +39,8 @@ public:
 		const TCHAR* Namespace, 
 		const TCHAR* OAuthProvider, 
 		const TCHAR* OAuthClientId, 
-		const TCHAR* OAuthData);
+		const TCHAR* OAuthData,
+		bool bReadOnly);
 
 	~FHttpDerivedDataBackend();
 
@@ -49,7 +50,12 @@ public:
 	 */
 	bool IsUsable() const { return bIsUsable; }
 
-	virtual bool IsWritable() const override { return true; }
+	/** return true if this cache is writable **/
+	virtual bool IsWritable() const override
+	{
+		return !bReadOnly && bIsUsable;
+	}
+
 	virtual bool CachedDataProbablyExists(const TCHAR* CacheKey) override;
 	virtual TBitArray<> CachedDataProbablyExistsBatch(TConstArrayView<FString> CacheKeys) override;
 	virtual bool GetCachedData(const TCHAR* CacheKey, TArray<uint8>& OutData) override;
@@ -79,6 +85,7 @@ private:
 	TUniquePtr<struct FRequestPool> PutRequestPools[2];
 	TUniquePtr<struct FHttpAccessToken> Access;
 	bool bIsUsable;
+	bool bReadOnly;
 	uint32 FailedLoginAttempts;
 	ESpeedClass SpeedClass;
 
