DDC: Added FCacheRecord KeyFuncs, FCacheRecord::GetPayload, and FPayload(FPayloadId)

#rnx

diff --git a/Engine/Source/Developer/DerivedDataCache/Private/DerivedDataCacheRecord.cpp b/Engine/Source/Developer/DerivedDataCache/Private/DerivedDataCacheRecord.cpp
--- a/Engine/Source/Developer/DerivedDataCache/Private/DerivedDataCacheRecord.cpp
+++ b/Engine/Source/Developer/DerivedDataCache/Private/DerivedDataCacheRecord.cpp
@@ -58,6 +58,8 @@ public:
 	const FPayload& GetAttachmentPayload(const FPayloadId& Id) const final;
 	TConstArrayView<FPayload> GetAttachmentPayloads() const final;
 
+	const FPayload& GetPayload(const FPayloadId& Id) const final;
+
 	FCacheKey Key;
 	FCbObject Meta;
 	FPayload Value;
@@ -152,6 +154,11 @@ TConstArrayView<FPayload> FCacheRecordInternal::GetAttachmentPayloads() const
 	return Attachments;
 }
 
+const FPayload& FCacheRecordInternal::GetPayload(const FPayloadId& Id) const
+{
+	return Value.GetId() == Id ? Value : GetAttachmentPayload(Id);
+}
+
 ///////////////////////////////////////////////////////////////////////////////////////////////////////////////////////
 
 FCacheRecordBuilderInternal::FCacheRecordBuilderInternal(const FCacheKey& InKey)
diff --git a/Engine/Source/Developer/DerivedDataCache/Public/DerivedDataCacheRecord.h b/Engine/Source/Developer/DerivedDataCache/Public/DerivedDataCacheRecord.h
--- a/Engine/Source/Developer/DerivedDataCache/Public/DerivedDataCacheRecord.h
+++ b/Engine/Source/Developer/DerivedDataCache/Public/DerivedDataCacheRecord.h
@@ -2,6 +2,7 @@
 
 #pragma once
 
+#include "DerivedDataCacheKey.h"
 #include "DerivedDataPayload.h"
 #include "DerivedDataRequest.h"
 #include "Memory/MemoryFwd.h"
@@ -12,7 +13,6 @@ template <typename FuncType> class TUniqueFunction;
 
 namespace UE::DerivedData { class FCacheRecord; }
 namespace UE::DerivedData { class FCacheRecordBuilder; }
-namespace UE::DerivedData { struct FCacheKey; }
 namespace UE::DerivedData { using FOnCacheRecordComplete = TUniqueFunction<void (FCacheRecord&& Record)>; }
 
 ///////////////////////////////////////////////////////////////////////////////////////////////////////////////////////
@@ -32,6 +32,7 @@ public:
 	virtual FSharedBuffer GetAttachment(const FPayloadId& Id) const = 0;
 	virtual const FPayload& GetAttachmentPayload(const FPayloadId& Id) const = 0;
 	virtual TConstArrayView<FPayload> GetAttachmentPayloads() const = 0;
+	virtual const FPayload& GetPayload(const FPayloadId& Id) const = 0;
 };
 
 FCacheRecord CreateCacheRecord(ICacheRecordInternal* Record);
@@ -85,24 +86,27 @@ public:
 	/** Returns the key that identifies this record in the cache. Always available in a non-null cache record. */
 	inline const FCacheKey& GetKey() const { return Record->GetKey(); }
 
-	/** Returns the metadata for the cache record unless it was skipped in the request. */
+	/** Returns the metadata. Null when requested with ECachePolicy::SkipMeta. */
 	inline const FCbObject& GetMeta() const { return Record->GetMeta(); }
 
-	/** Returns the value. Null if the value was skipped in the request. */
+	/** Returns the value. Null if no value or requested with ECachePolicy::SkipValue. */
 	inline FSharedBuffer GetValue() const { return Record->GetValue(); }
 
-	/** Returns the value payload. Always available in a non-null cache record. */
+	/** Returns the value payload. Null if no value. Buffer is null if value was skipped. */
 	inline const FPayload& GetValuePayload() const { return Record->GetValuePayload(); }
 
-	/** Returns the attachment that matches the ID. Null if no match or attachments were skipped in the request. */
+	/** Returns the attachment matching the ID. Null if no match or requested with ECachePolicy::SkipAttachments. */
 	inline FSharedBuffer GetAttachment(const FPayloadId& Id) const { return Record->GetAttachment(Id); }
 
-	/** Returns the attachment payload that matches the ID. Null if no match or the cache record is null. */
+	/** Returns the attachment payload matching the ID. Null if no match. Buffer is null if attachments were skipped. */
 	inline const FPayload& GetAttachmentPayload(const FPayloadId& Id) const { return Record->GetAttachmentPayload(Id); }
 
 	/** Returns a view of the attachments. Always available in a non-null cache record, but buffer may be skipped. */
 	inline TConstArrayView<FPayload> GetAttachmentPayloads() const { return Record->GetAttachmentPayloads(); }
 
+	/** Returns the payload matching the ID, whether value or attachment. Null if no match. Buffer is null if skipped. */
+	inline const FPayload& GetPayload(const FPayloadId& Id) const { return Record->GetPayload(Id); }
+
 	/** Whether this is null. */
 	inline bool IsNull() const { return !Record; }
 	/** Whether this is not null. */
@@ -232,3 +236,24 @@ private:
 };
 
 } // UE::DerivedData
+
+///////////////////////////////////////////////////////////////////////////////////////////////////////////////////////
+
+namespace UE::DerivedData
+{
+
+/** An implementation of KeyFuncs to compare FCacheRecord by its FCacheKey. */
+struct FCacheRecordKeyFuncs
+{
+	using KeyType = FCacheKey;
+	using KeyInitType = const FCacheKey&;
+	using ElementInitType = const FCacheRecord&;
+
+	static constexpr bool bAllowDuplicateKeys = false;
+
+	static inline KeyInitType GetSetKey(ElementInitType Record) { return Record.GetKey(); }
+	static inline uint32 GetKeyHash(KeyInitType Key) { return GetTypeHash(Key); }
+	static inline bool Matches(KeyInitType A, KeyInitType B) { return A == B; }
+};
+
+} // UE::DerivedData
diff --git a/Engine/Source/Developer/DerivedDataCache/Public/DerivedDataPayload.h b/Engine/Source/Developer/DerivedDataCache/Public/DerivedDataPayload.h
--- a/Engine/Source/Developer/DerivedDataCache/Public/DerivedDataPayload.h
+++ b/Engine/Source/Developer/DerivedDataCache/Public/DerivedDataPayload.h
@@ -67,6 +67,9 @@ public:
 	/** Construct a null payload. */
 	FPayload() = default;
 
+	/** Construct a payload with no hash or buffer. */
+	inline explicit FPayload(const FPayloadId& Id);
+
 	/** Construct a payload from a raw buffer hash. */
 	inline FPayload(const FPayloadId& Id, const FIoHash& RawHash);
 
@@ -127,6 +130,12 @@ inline FPayloadId FPayloadId::FromName(const FWideStringView Name)
 
 ///////////////////////////////////////////////////////////////////////////////////////////////////////////////////////
 
+inline FPayload::FPayload(const FPayloadId& InId)
+	: Id(InId)
+{
+	checkf(Id.IsValid(), TEXT("A valid ID is required to construct a payload."));
+}
+
 inline FPayload::FPayload(const FPayloadId& InId, const FIoHash& InRawHash)
 	: Id(InId)
 	, RawHash(InRawHash)
