Put a minimum "transfer speed over duration" limit on curl operations done by the HTTP DDC backend.

diff --git a/Engine/Source/Developer/DerivedDataCache/Private/HttpDerivedDataBackend.cpp b/Engine/Source/Developer/DerivedDataCache/Private/HttpDerivedDataBackend.cpp
--- a/Engine/Source/Developer/DerivedDataCache/Private/HttpDerivedDataBackend.cpp
+++ b/Engine/Source/Developer/DerivedDataCache/Private/HttpDerivedDataBackend.cpp
@@ -175,6 +175,9 @@ public:
 		// Rewind method, handle special error case where request need to rewind data stream
 		curl_easy_setopt(Curl, CURLOPT_SEEKFUNCTION, StaticSeekFn);
 		curl_easy_setopt(Curl, CURLOPT_SEEKDATA, this);
+		// Set minimum speed behavior to allow operations to abort if the transfer speed is poor for the given duration (1kbps over a 30 second span)
+		curl_easy_setopt(Curl, CURLOPT_LOW_SPEED_TIME, 30L);
+		curl_easy_setopt(Curl, CURLOPT_LOW_SPEED_LIMIT, 1024L);
 		// Debug hooks
 #if UE_HTTPDDC_HTTP_DEBUG
 		curl_easy_setopt(Curl, CURLOPT_DEBUGDATA, this);
