Ensure that the DDC filesystem backend updates timestamps on files in GetCachedData (if within the last quarter of its cleanup lifetime), not solely from CachedDataProbablyExists.

diff --git a/Engine/Source/Developer/DerivedDataCache/Private/FileSystemDerivedDataBackend.cpp b/Engine/Source/Developer/DerivedDataCache/Private/FileSystemDerivedDataBackend.cpp
--- a/Engine/Source/Developer/DerivedDataCache/Private/FileSystemDerivedDataBackend.cpp
+++ b/Engine/Source/Developer/DerivedDataCache/Private/FileSystemDerivedDataBackend.cpp
@@ -564,6 +564,21 @@ public:
 
 		if (FFileHelper::LoadFileToArray(Data,*Filename,FILEREAD_Silent))
 		{
+			if (SpeedClass >= FDerivedDataBackendInterface::ESpeedClass::Fast && (!bReadOnly || bTouch))
+			{
+				FFileStatData FileStat = IFileManager::Get().GetStatData(*Filename);
+				if (FileStat.bIsValid)
+				{
+					FDateTime TimeStamp = FileStat.ModificationTime;
+					// Update file timestamp to prevent it from being deleted by DDC Cleanup.
+					if (bTouch ||
+						(!bReadOnly && (FDateTime::UtcNow() - TimeStamp).GetDays() > (DaysToDeleteUnusedFiles / 4)))
+					{
+						IFileManager::Get().SetTimeStamp(*Filename, FDateTime::UtcNow());
+					}
+				}
+			}
+
 			double ReadDuration = FPlatformTime::Seconds() - StartTime;
 			double ReadSpeed = (Data.Num() / ReadDuration) / (1024.0 * 1024.0);
 
