DDC: Split ICache into ICacheFactory and ICacheStore, expand FRequest into TRequest

#rnx

diff --git a/Engine/Source/Developer/DerivedDataCache/Public/DerivedDataCache.h b/Engine/Source/Developer/DerivedDataCache/Public/DerivedDataCache.h
--- a/Engine/Source/Developer/DerivedDataCache/Public/DerivedDataCache.h
+++ b/Engine/Source/Developer/DerivedDataCache/Public/DerivedDataCache.h
@@ -7,9 +7,7 @@
 #include "DerivedDataRequest.h"
 #include "Misc/EnumClassFlags.h"
 
-namespace UE
-{
-namespace DerivedData
+namespace UE::DerivedData
 {
 
 class FCacheBucket;
@@ -93,16 +91,12 @@ ENUM_CLASS_FLAGS(ECachePolicy);
 ///////////////////////////////////////////////////////////////////////////////////////////////////////////////////////
 
 /**
- * Interface to the cache.
- *
- * Functions on this interface may be called from any thread. When a callback is provided, it may
- * be invoked from a different thread than the request was made on, and may be invoked before the
- * request function returns, both of which must be supported by the callback.
+ * Interface to construct types related to the cache.
  */
-class ICache
+class ICacheFactory
 {
 public:
-	virtual ~ICache() = default;
+	virtual ~ICacheFactory() = default;
 
 	/**
 	 * Create a cache bucket from a name.
@@ -115,6 +109,21 @@ public:
 	 * Create a cache record builder from a cache key.
 	 */
 	virtual FCacheRecordBuilder CreateRecord(const FCacheKey& Key) = 0;
+};
+
+///////////////////////////////////////////////////////////////////////////////////////////////////////////////////////
+
+/**
+ * Interface to a store of cache records.
+ *
+ * Functions on this interface may be called from any thread. When a callback is provided, it may
+ * be invoked from a different thread than the request was made on, and may be invoked before the
+ * request function returns, both of which must be supported by the caller and the callback.
+ */
+class ICacheStore
+{
+public:
+	virtual ~ICacheStore() = default;
 
 	/**
 	 * Asynchronous request to put cache records according to the policy.
@@ -126,14 +135,14 @@ public:
 	 * @param Context A description of the request. An object path is typically sufficient.
 	 * @param Policy Flags to control the behavior of the request. See ECachePolicy.
 	 * @param Priority A priority to consider when scheduling the request. See EPriority.
-	 * @param Callback A callback invoked for every key in the batch as it completes or is canceled.
+	 * @param OnComplete A callback invoked for every key in the batch as it completes or is canceled.
 	 */
 	virtual FRequest Put(
 		TArrayView<FCacheRecord> Records,
 		FStringView Context,
 		ECachePolicy Policy = ECachePolicy::Default,
 		EPriority Priority = EPriority::Normal,
-		FOnCachePutComplete&& Callback = FOnCachePutComplete()) = 0;
+		FOnCachePutComplete&& OnComplete = FOnCachePutComplete()) = 0;
 
 	/**
 	 * Asynchronous request to get cache records according to the policy.
@@ -145,14 +154,14 @@ public:
 	 * @param Context A description of the request. An object path is typically sufficient.
 	 * @param Policy Flags to control the behavior of the request. See ECachePolicy.
 	 * @param Priority A priority to consider when scheduling the request. See EPriority.
-	 * @param Callback A callback invoked for every key in the batch as it completes or is canceled.
+	 * @param OnComplete A callback invoked for every key in the batch as it completes or is canceled.
 	 */
 	virtual FRequest Get(
 		TConstArrayView<FCacheKey> Keys,
 		FStringView Context,
 		ECachePolicy Policy,
 		EPriority Priority,
-		FOnCacheGetComplete&& Callback) = 0;
+		FOnCacheGetComplete&& OnComplete) = 0;
 
 	/**
 	 * Asynchronous request to get cache record payloads according to the policy.
@@ -164,14 +173,14 @@ public:
 	 * @param Context A description of the request. An object path is typically sufficient.
 	 * @param Policy Flags to control the behavior of the request. See ECachePolicy.
 	 * @param Priority A priority to consider when scheduling the request. See EPriority.
-	 * @param Callback A callback invoked for every key in the batch as it completes or is canceled.
+	 * @param OnComplete A callback invoked for every key in the batch as it completes or is canceled.
 	 */
 	virtual FRequest GetPayload(
 		TConstArrayView<FCachePayloadKey> Keys,
 		FStringView Context,
 		ECachePolicy Policy,
 		EPriority Priority,
-		FOnCacheGetPayloadComplete&& Callback) = 0;
+		FOnCacheGetPayloadComplete&& OnComplete) = 0;
 
 	/**
 	 * Cancel all queued and active cache requests and invoke their callbacks.
@@ -186,6 +195,20 @@ public:
 
 ///////////////////////////////////////////////////////////////////////////////////////////////////////////////////////
 
+/**
+ * Interface to the cache.
+ *
+ * @see ICacheStore for cache record storage.
+ * @see ICacheFactory for cache type constructors.
+ */
+class ICache : public ICacheStore, public ICacheFactory
+{
+public:
+	virtual ~ICache() = default;
+};
+
+///////////////////////////////////////////////////////////////////////////////////////////////////////////////////////
+
 /** Parameters for the completion callback for cache put requests. */
 struct FCachePutCompleteParams
 {
@@ -232,5 +255,4 @@ struct FCacheGetPayloadCompleteParams
 
 ///////////////////////////////////////////////////////////////////////////////////////////////////////////////////////
 
-} // DerivedData
-} // UE
+} // UE::DerivedData
diff --git a/Engine/Source/Developer/DerivedDataCache/Public/DerivedDataRequest.h b/Engine/Source/Developer/DerivedDataCache/Public/DerivedDataRequest.h
--- a/Engine/Source/Developer/DerivedDataCache/Public/DerivedDataRequest.h
+++ b/Engine/Source/Developer/DerivedDataCache/Public/DerivedDataRequest.h
@@ -4,10 +4,9 @@
 
 #include "CoreTypes.h"
 #include "Templates/RefCounting.h"
+#include "Templates/UnrealTemplate.h"
 
-namespace UE
-{
-namespace DerivedData
+namespace UE::DerivedData
 {
 
 ///////////////////////////////////////////////////////////////////////////////////////////////////////////////////////
@@ -63,8 +62,8 @@ enum class EStatus : uint8
 
 ///////////////////////////////////////////////////////////////////////////////////////////////////////////////////////
 
-/** Interface to a request. Use via FRequest. */
-class IRequest : public IRefCountedObject
+/** Interface to a request. Use via FRequest to automate reference counting. */
+class IRequest
 {
 public:
 	virtual ~IRequest() = default;
@@ -80,22 +79,38 @@ public:
 
 	/** Poll the request and return true if complete and false otherwise. */
 	virtual bool Poll() = 0;
+
+	/** Add a reference to the request. */
+	virtual void AddRef() const = 0;
+
+	/** Release a reference. The request is deleted when the last reference is released. */
+	virtual void Release() const = 0;
 };
 
 ///////////////////////////////////////////////////////////////////////////////////////////////////////////////////////
 
 /** Handle to a request that manages its lifetime with reference counting. */
-class FRequest
+template <typename RequestType>
+class TRequest
 {
 public:
 	/** Construct a null request. */
-	FRequest() = default;
+	TRequest() = default;
 
 	/** Construct a request from a raw pointer. */
-	inline explicit FRequest(IRequest* InRequest)
-		: Request(InRequest)
-	{
-	}
+	inline explicit TRequest(RequestType* InRequest) : Request(InRequest) {}
+
+	/** Construct a request from a request of a compatible type. */
+	template <typename OtherRequestType, decltype(ImplicitConv<RequestType*>(DeclVal<OtherRequestType*>()))* = nullptr>
+	inline TRequest(const TRequest<OtherRequestType>& InRequest) : Request(InRequest.Request) {}
+	template <typename OtherRequestType, decltype(ImplicitConv<RequestType*>(DeclVal<OtherRequestType*>()))* = nullptr>
+	inline TRequest(TRequest<OtherRequestType>&& InRequest) : Request(MoveTemp(InRequest.Request)) {}
+
+	/** Assign from a request of a compatible type. */
+	template <typename OtherRequestType, decltype(ImplicitConv<RequestType*>(DeclVal<OtherRequestType*>()))* = nullptr>
+	inline TRequest& operator=(const TRequest<OtherRequestType>& InRequest) { Request = InRequest.Request; return *this; }
+	template <typename OtherRequestType, decltype(ImplicitConv<RequestType*>(DeclVal<OtherRequestType*>()))* = nullptr>
+	inline TRequest& operator=(TRequest<OtherRequestType>&& InRequest) { Request = MoveTemp(InRequest.Request); return *this; }
 
 	/** Reset this to null. */
 	inline void Reset() { Request.SafeRelease(); }
@@ -107,6 +122,11 @@ public:
 	/** Whether the request is not null. */
 	inline explicit operator bool() const { return IsValid(); }
 
+	/** Access the referenced request. */
+	inline RequestType* Get() const { return Request; }
+	inline RequestType* operator->() const { return Request; }
+	inline RequestType& operator*() const { return *Request; }
+
 	/**
 	 * Set the priority of the request.
 	 *
@@ -158,11 +178,25 @@ public:
 		return !Request || Request->Poll();
 	}
 
+	template <typename OtherRequestType, decltype(ImplicitConv<IRequest*>(DeclVal<OtherRequestType*>()))* = nullptr>
+	inline bool operator==(const TRequest<OtherRequestType>& Other) const { return Request == Other.Request; }
+	template <typename OtherRequestType, decltype(ImplicitConv<IRequest*>(DeclVal<OtherRequestType*>()))* = nullptr>
+	inline bool operator!=(const TRequest<OtherRequestType>& Other) const { return Request != Other.Request; }
+	template <typename OtherRequestType, decltype(ImplicitConv<IRequest*>(DeclVal<OtherRequestType*>()))* = nullptr>
+	inline bool operator<(const TRequest<OtherRequestType>& Other) const { return Request < Other.Request; }
+
+	friend inline uint32 GetTypeHash(const TRequest& Other) { return ::GetTypeHash(Other.Request); }
+
+	template <typename OtherRequestType>
+	friend class TRequest;
+
 private:
-	TRefCountPtr<IRequest> Request;
+	TRefCountPtr<RequestType> Request;
 };
 
+/** Handle to a request that manages its lifetime with reference counting. */
+using FRequest = TRequest<IRequest>;
+
 ///////////////////////////////////////////////////////////////////////////////////////////////////////////////////////
 
-} // DerivedData
-} // UE
+} // UE::DerivedData
