DDC: Added an assert to FCacheRecordBuilder::BuildAsync to require a callback

#rnx

diff --git a/Engine/Source/Developer/DerivedDataCache/Private/DerivedDataCacheRecord.cpp b/Engine/Source/Developer/DerivedDataCache/Private/DerivedDataCacheRecord.cpp
--- a/Engine/Source/Developer/DerivedDataCache/Private/DerivedDataCacheRecord.cpp
+++ b/Engine/Source/Developer/DerivedDataCache/Private/DerivedDataCacheRecord.cpp
@@ -214,6 +214,8 @@ FCacheRecord FCacheRecordBuilderInternal::Build()
 
 FRequest FCacheRecordBuilderInternal::BuildAsync(FOnCacheRecordComplete&& OnComplete, EPriority Priority)
 {
+	checkf(OnComplete, TEXT("Failed to build cache record for %s because the completion callback is null."),
+		*WriteToString<96>(Key));
 	OnComplete(Build());
 	return FRequest();
 }
