DDC: Implement ICacheStore on FDerivedDataBackendInterface

This allows the new cache to be implemented more efficiently using the old cache backends because functionality like the corruption wrapper and key length limiter can be bypassed and the individual backends can store cache records in whatever way is most efficient for them.

The hierarchical backend may request payloads when they are not required due to incomplete tracking of backend state, and GetPayload will never fill other backends due to the inefficiency of the existing backend framework.

The filesystem backend does not cache any state in memory, which makes requests for individual payloads less efficient than if it cached a mapping of cache payload key to raw hash after the first request for a cache key.

The HTTP, S3, and pak file backends are not implemented for the new interface.

The backends do not implement ICacheStore::CancelAll() because the existing backend framework provides WaitForQuiescence to wait for completion of async requests, and the implementation of ICacheStore by those backends works with that mechanism.

The non-leaf backends (hierarchical, async put, etc.) do not update stats from the ICacheStore functions.

#rnx

diff --git a/Engine/Source/Developer/DerivedDataCache/Private/DerivedDataBackendAsyncPutWrapper.cpp b/Engine/Source/Developer/DerivedDataCache/Private/DerivedDataBackendAsyncPutWrapper.cpp
--- a/Engine/Source/Developer/DerivedDataCache/Private/DerivedDataBackendAsyncPutWrapper.cpp
+++ b/Engine/Source/Developer/DerivedDataCache/Private/DerivedDataBackendAsyncPutWrapper.cpp
@@ -2,6 +2,11 @@
 
 #include "DerivedDataBackendAsyncPutWrapper.h"
 #include "MemoryDerivedDataBackend.h"
+#include "Tasks/Task.h"
+#include <atomic>
+
+namespace UE::DerivedData::Backends
+{
 
 /** 
  * Async task to handle the fire and forget async put
@@ -144,9 +149,10 @@ public:
 	}
 };
 
-FDerivedDataBackendAsyncPutWrapper::FDerivedDataBackendAsyncPutWrapper(FDerivedDataBackendInterface* InInnerBackend, bool bCacheInFlightPuts)
-	: InnerBackend(InInnerBackend)
-	, InflightCache(bCacheInFlightPuts ? (new FMemoryDerivedDataBackend(TEXT("InflightMemoryCache"))) : NULL)
+FDerivedDataBackendAsyncPutWrapper::FDerivedDataBackendAsyncPutWrapper(ICacheFactory& InFactory, FDerivedDataBackendInterface* InInnerBackend, bool bCacheInFlightPuts)
+	: Factory(InFactory)
+	, InnerBackend(InInnerBackend)
+	, InflightCache(bCacheInFlightPuts ? (new FMemoryDerivedDataBackend(Factory, TEXT("InflightMemoryCache"))) : NULL)
 {
 	check(InnerBackend);
 }
@@ -317,3 +323,248 @@ TSharedRef<FDerivedDataCacheStatsNode> FDerivedDataBackendAsyncPutWrapper::Gathe
 
 	return Usage;
 }
+
+class FDerivedDataAsyncWrapperRequest final : public IRequest, private IQueuedWork
+{
+public:
+	inline explicit FDerivedDataAsyncWrapperRequest(TUniqueFunction<void (bool bCancel)>&& InFunction)
+		: Function(MoveTemp(InFunction))
+		, DoneEvent(FPlatformProcess::GetSynchEventFromPool(true))
+	{
+	}
+
+	inline ~FDerivedDataAsyncWrapperRequest() final
+	{
+		FPlatformProcess::ReturnSynchEventToPool(DoneEvent);
+	}
+
+	inline void Start(EPriority Priority)
+	{
+		FDerivedDataBackend::Get().AddToAsyncCompletionCounter(1);
+		AddRef(); // Keep the request alive until the queued work executes.
+
+		DoneEvent->Reset();
+		WorkNotFinishedCounter.fetch_add(1, std::memory_order_relaxed);
+		GDDCIOThreadPool->AddQueuedWork(this, GetPriority(Priority));
+	}
+
+	inline void Execute(bool bCancel)
+	{
+		FScopeCycleCounter Scope(GetStatId(), /*bAlways*/ true);
+		Function(/*bCancel*/ false);
+		WorkNotFinishedCounter.fetch_sub(1, std::memory_order_release);
+		DoneEvent->Trigger();
+
+		Release(); // DO NOT ACCESS ANY MEMBERS PAST THIS POINT!
+		FDerivedDataBackend::Get().AddToAsyncCompletionCounter(-1);
+	}
+
+	// IRequest Interface
+
+	inline void SetPriority(EPriority Priority) final
+	{
+		if (GDDCIOThreadPool->RetractQueuedWork(this))
+		{
+			GDDCIOThreadPool->AddQueuedWork(this, GetPriority(Priority));
+		}
+	}
+
+	inline void Cancel() final
+	{
+		if (!Poll())
+		{
+			if (GDDCIOThreadPool->RetractQueuedWork(this))
+			{
+				Abandon();
+			}
+			else
+			{
+				Wait();
+			}
+		}
+	}
+
+	inline void Wait() final
+	{
+		if (!Poll())
+		{
+			if (GDDCIOThreadPool->RetractQueuedWork(this))
+			{
+				DoThreadedWork();
+			}
+			else
+			{
+				FScopeCycleCounter Scope(GetStatId());
+				DoneEvent->Wait();
+			}
+		}
+	}
+
+	inline bool Poll() final
+	{
+		return WorkNotFinishedCounter.load(std::memory_order_acquire) == 0;
+	}
+
+	inline void AddRef() const final
+	{
+		ReferenceCount.fetch_add(1, std::memory_order_relaxed);
+	}
+
+	inline void Release() const final
+	{
+		if (ReferenceCount.fetch_sub(1, std::memory_order_acq_rel) == 1)
+		{
+			delete this;
+		}
+	}
+
+private:
+	static EQueuedWorkPriority GetPriority(EPriority Priority)
+	{
+		switch (Priority)
+		{
+		case EPriority::Blocking: return EQueuedWorkPriority::Highest;
+		case EPriority::Highest:  return EQueuedWorkPriority::Highest;
+		case EPriority::High:     return EQueuedWorkPriority::High;
+		case EPriority::Normal:   return EQueuedWorkPriority::Normal;
+		case EPriority::Low:      return EQueuedWorkPriority::Low;
+		case EPriority::Lowest:   return EQueuedWorkPriority::Lowest;
+		default: checkNoEntry();  return EQueuedWorkPriority::Normal;
+		}
+	}
+
+	// IQueuedWork Interface
+
+	inline void DoThreadedWork() final
+	{
+		Execute(/*bCancel*/ false);
+	}
+
+	inline void Abandon() final
+	{
+		Execute(/*bCancel*/ true);
+	}
+
+	inline TStatId GetStatId() const
+	{
+		RETURN_QUICK_DECLARE_CYCLE_STAT(FDerivedDataAsyncWrapperRequest, STATGROUP_ThreadPoolAsyncTasks);
+	}
+
+private:
+	mutable std::atomic<uint32> ReferenceCount{0};
+	std::atomic<uint32> WorkNotFinishedCounter{0};
+	TUniqueFunction<void (bool bCancel)> Function;
+	FEvent* DoneEvent;
+};
+
+FRequest FDerivedDataBackendAsyncPutWrapper::Put(
+	TArrayView<FCacheRecord> Records,
+	FStringView Context,
+	ECachePolicy Policy,
+	EPriority Priority,
+	FOnCachePutComplete&& OnComplete)
+{
+	if (Priority == EPriority::Blocking)
+	{
+		return InnerBackend->Put(Records, Context, Policy, Priority, MoveTemp(OnComplete));
+	}
+	else
+	{
+		TArray<FCacheRecord> RecordsArray;
+		RecordsArray.Reserve(Records.Num());
+		for (FCacheRecord& Record : Records)
+		{
+			RecordsArray.Add(MoveTemp(Record));
+		}
+		TRequest<FDerivedDataAsyncWrapperRequest> Request(new FDerivedDataAsyncWrapperRequest(
+			[Backend = InnerBackend, Records = MoveTemp(RecordsArray), Context = FString(Context), Policy, Priority, OnComplete = MoveTemp(OnComplete)](bool bCancel) mutable
+			{
+				if (!bCancel)
+				{
+					Backend->Put(Records, Context, Policy, Priority, MoveTemp(OnComplete)).Wait();
+				}
+				else if (OnComplete)
+				{
+					for (const FCacheRecord& Record : Records)
+					{
+						OnComplete({Record.GetKey(), EStatus::Canceled});
+					}
+				}
+			}));
+		Request->Start(Priority == EPriority::Normal ? EPriority::Low : Priority);
+		return Request;
+	}
+}
+
+FRequest FDerivedDataBackendAsyncPutWrapper::Get(
+	TConstArrayView<FCacheKey> Keys,
+	FStringView Context,
+	ECachePolicy Policy,
+	EPriority Priority,
+	FOnCacheGetComplete&& OnComplete)
+{
+	if (Priority == EPriority::Blocking)
+	{
+		return InnerBackend->Get(Keys, Context, Policy, Priority, MoveTemp(OnComplete));
+	}
+	else
+	{
+		TRequest<FDerivedDataAsyncWrapperRequest> Request(new FDerivedDataAsyncWrapperRequest(
+			[&Factory = Factory, Backend = InnerBackend, Keys = TArray<FCacheKey>(Keys), Context = FString(Context), Policy, Priority, OnComplete = MoveTemp(OnComplete)](bool bCancel) mutable
+			{
+				if (!bCancel)
+				{
+					Backend->Get(Keys, Context, Policy, Priority, MoveTemp(OnComplete)).Wait();
+				}
+				else if (OnComplete)
+				{
+					for (const FCacheKey& Key : Keys)
+					{
+						OnComplete({Factory.CreateRecord(Key).Build(), EStatus::Canceled});
+					}
+				}
+			}));
+		Request->Start(Priority);
+		return Request;
+	}
+}
+
+FRequest FDerivedDataBackendAsyncPutWrapper::GetPayload(
+	TConstArrayView<FCachePayloadKey> Keys,
+	FStringView Context,
+	ECachePolicy Policy,
+	EPriority Priority,
+	FOnCacheGetPayloadComplete&& OnComplete)
+{
+	if (Priority == EPriority::Blocking)
+	{
+		return InnerBackend->GetPayload(Keys, Context, Policy, Priority, MoveTemp(OnComplete));
+	}
+	else
+	{
+		TRequest<FDerivedDataAsyncWrapperRequest> Request(new FDerivedDataAsyncWrapperRequest(
+			[Backend = InnerBackend, Keys = TArray<FCachePayloadKey>(Keys), Context = FString(Context), Policy, Priority, OnComplete = MoveTemp(OnComplete)](bool bCancel) mutable
+			{
+				if (!bCancel)
+				{
+					Backend->GetPayload(Keys, Context, Policy, Priority, MoveTemp(OnComplete)).Wait();
+				}
+				else if (OnComplete)
+				{
+					for (const FCachePayloadKey& Key : Keys)
+					{
+						OnComplete({Key.CacheKey, FPayload(Key.Id), EStatus::Canceled});
+					}
+				}
+			}));
+		Request->Start(Priority);
+		return Request;
+	}
+}
+
+void FDerivedDataBackendAsyncPutWrapper::CancelAll()
+{
+	InnerBackend->CancelAll();
+}
+
+} // UE::DerivedData::Backends
diff --git a/Engine/Source/Developer/DerivedDataCache/Private/DerivedDataBackendAsyncPutWrapper.h b/Engine/Source/Developer/DerivedDataCache/Private/DerivedDataBackendAsyncPutWrapper.h
--- a/Engine/Source/Developer/DerivedDataCache/Private/DerivedDataBackendAsyncPutWrapper.h
+++ b/Engine/Source/Developer/DerivedDataCache/Private/DerivedDataBackendAsyncPutWrapper.h
@@ -11,6 +11,9 @@
 #include "MemoryDerivedDataBackend.h"
 #include "Async/AsyncWork.h"
 
+namespace UE::DerivedData::Backends
+{
+
 /** 
  * Thread safe set helper
 **/
@@ -61,7 +64,7 @@ public:
 	 * @param	InInnerBackend		Backend to use for storage, my responsibilities are about async puts
 	 * @param	bCacheInFlightPuts	if true, cache in-flight puts in a memory cache so that they hit immediately
 	 */
-	FDerivedDataBackendAsyncPutWrapper(FDerivedDataBackendInterface* InInnerBackend, bool bCacheInFlightPuts);
+	FDerivedDataBackendAsyncPutWrapper(ICacheFactory& InFactory, FDerivedDataBackendInterface* InInnerBackend, bool bCacheInFlightPuts);
 
 	/** Return a name for this interface */
 	virtual FString GetName() const override
@@ -129,10 +132,34 @@ public:
 
 	virtual bool ApplyDebugOptions(FBackendDebugOptions& InOptions) override;
 
+	virtual FRequest Put(
+		TArrayView<FCacheRecord> Records,
+		FStringView Context,
+		ECachePolicy Policy,
+		EPriority Priority,
+		FOnCachePutComplete&& OnComplete) override;
+
+	virtual FRequest Get(
+		TConstArrayView<FCacheKey> Keys,
+		FStringView Context,
+		ECachePolicy Policy,
+		EPriority Priority,
+		FOnCacheGetComplete&& OnComplete) override;
+
+	virtual FRequest GetPayload(
+		TConstArrayView<FCachePayloadKey> Keys,
+		FStringView Context,
+		ECachePolicy Policy,
+		EPriority Priority,
+		FOnCacheGetPayloadComplete&& OnComplete) override;
+
+	virtual void CancelAll() override;
+
 private:
 	FDerivedDataCacheUsageStats UsageStats;
 	FDerivedDataCacheUsageStats PutSyncUsageStats;
 
+	ICacheFactory& Factory;
 	/** Backend to use for storage, my responsibilities are about async puts **/
 	FDerivedDataBackendInterface*					InnerBackend;
 	/** Memory based cache to deal with gets that happen while an async put is still in flight **/
@@ -141,5 +168,4 @@ private:
 	FThreadSet										FilesInFlight;
 };
 
-
-
+} // UE::DerivedData::Backends
diff --git a/Engine/Source/Developer/DerivedDataCache/Private/DerivedDataBackendCorruptionWrapper.h b/Engine/Source/Developer/DerivedDataCache/Private/DerivedDataBackendCorruptionWrapper.h
--- a/Engine/Source/Developer/DerivedDataCache/Private/DerivedDataBackendCorruptionWrapper.h
+++ b/Engine/Source/Developer/DerivedDataCache/Private/DerivedDataBackendCorruptionWrapper.h
@@ -7,6 +7,9 @@
 #include "ProfilingDebugging/CookStats.h"
 #include "DerivedDataCacheUsageStats.h"
 
+namespace UE::DerivedData::Backends
+{
+
 /** 
  * Helper class for placing a footer at the end of of a cache file.
  * No effort is made to byte-swap this as we assume local format.
@@ -248,6 +251,41 @@ public:
 		return InnerBackend->ApplyDebugOptions(InOptions);
 	}
 
+	virtual FRequest Put(
+		TArrayView<FCacheRecord> Records,
+		FStringView Context,
+		ECachePolicy Policy,
+		EPriority Priority,
+		FOnCachePutComplete&& OnComplete) override
+	{
+		return InnerBackend->Put(Records, Context, Policy, Priority, MoveTemp(OnComplete));
+	}
+
+	virtual FRequest Get(
+		TConstArrayView<FCacheKey> Keys,
+		FStringView Context,
+		ECachePolicy Policy,
+		EPriority Priority,
+		FOnCacheGetComplete&& OnComplete) override
+	{
+		return InnerBackend->Get(Keys, Context, Policy, Priority, MoveTemp(OnComplete));
+	}
+
+	virtual FRequest GetPayload(
+		TConstArrayView<FCachePayloadKey> Keys,
+		FStringView Context,
+		ECachePolicy Policy,
+		EPriority Priority,
+		FOnCacheGetPayloadComplete&& OnComplete) override
+	{
+		return InnerBackend->GetPayload(Keys, Context, Policy, Priority, MoveTemp(OnComplete));
+	}
+
+	virtual void CancelAll() override
+	{
+		InnerBackend->CancelAll();
+	}
+
 private:
 	FDerivedDataCacheUsageStats UsageStats;
 
@@ -255,3 +293,4 @@ private:
 	FDerivedDataBackendInterface* InnerBackend;
 };
 
+} // UE::DerivedData::Backends
diff --git a/Engine/Source/Developer/DerivedDataCache/Private/DerivedDataBackendVerifyWrapper.h b/Engine/Source/Developer/DerivedDataCache/Private/DerivedDataBackendVerifyWrapper.h
--- a/Engine/Source/Developer/DerivedDataCache/Private/DerivedDataBackendVerifyWrapper.h
+++ b/Engine/Source/Developer/DerivedDataCache/Private/DerivedDataBackendVerifyWrapper.h
@@ -12,6 +12,9 @@
 
 class Error;
 
+namespace UE::DerivedData::Backends
+{
+
 /**
  * FDerivedDataBackendVerifyWrapper, a wrapper for derived data that verifies the cache is bit-wise identical by failing all gets and verifying the puts
 **/
@@ -159,6 +162,41 @@ public:
 		return InnerBackend->ApplyDebugOptions(InOptions);
 	}
 
+	virtual FRequest Put(
+		TArrayView<FCacheRecord> Records,
+		FStringView Context,
+		ECachePolicy Policy,
+		EPriority Priority,
+		FOnCachePutComplete&& OnComplete) override
+	{
+		return InnerBackend->Put(Records, Context, Policy, Priority, MoveTemp(OnComplete));
+	}
+
+	virtual FRequest Get(
+		TConstArrayView<FCacheKey> Keys,
+		FStringView Context,
+		ECachePolicy Policy,
+		EPriority Priority,
+		FOnCacheGetComplete&& OnComplete) override
+	{
+		return InnerBackend->Get(Keys, Context, Policy, Priority, MoveTemp(OnComplete));
+	}
+
+	virtual FRequest GetPayload(
+		TConstArrayView<FCachePayloadKey> Keys,
+		FStringView Context,
+		ECachePolicy Policy,
+		EPriority Priority,
+		FOnCacheGetPayloadComplete&& OnComplete) override
+	{
+		return InnerBackend->GetPayload(Keys, Context, Policy, Priority, MoveTemp(OnComplete));
+	}
+
+	virtual void CancelAll() override
+	{
+		InnerBackend->CancelAll();
+	}
+
 private:
 	FDerivedDataCacheUsageStats UsageStats;
 
@@ -171,3 +209,5 @@ private:
 	/** Backend to service the actual requests */
 	FDerivedDataBackendInterface*					InnerBackend;
 };
+
+} // UE::DerivedData::Backends
diff --git a/Engine/Source/Developer/DerivedDataCache/Private/DerivedDataBackends.cpp b/Engine/Source/Developer/DerivedDataCache/Private/DerivedDataBackends.cpp
--- a/Engine/Source/Developer/DerivedDataCache/Private/DerivedDataBackends.cpp
+++ b/Engine/Source/Developer/DerivedDataCache/Private/DerivedDataBackends.cpp
@@ -10,6 +10,7 @@
 #include "HAL/IConsoleManager.h"
 #include "Misc/App.h"
 #include "DerivedDataBackendInterface.h"
+#include "DerivedDataCachePrivate.h"
 #include "DerivedDataCacheUsageStats.h"
 #include "MemoryDerivedDataBackend.h"
 #include "HttpDerivedDataBackend.h"
@@ -37,13 +38,16 @@ DEFINE_LOG_CATEGORY(LogDerivedDataCache);
 #define MAX_BACKEND_KEY_LENGTH (120)
 #define LOCTEXT_NAMESPACE "DerivedDataBackendGraph"
 
-FDerivedDataBackendInterface* CreateFileSystemDerivedDataBackend(const TCHAR* CacheDirectory, const TCHAR* InParams, const TCHAR* InAccessLogFileName = nullptr);
+namespace UE::DerivedData::Backends
+{
+
+FDerivedDataBackendInterface* CreateFileSystemDerivedDataBackend(ICacheFactory& Factory, const TCHAR* CacheDirectory, const TCHAR* InParams, const TCHAR* InAccessLogFileName = nullptr);
 
 /**
   * This class is used to create a singleton that represents the derived data cache hierarchy and all of the wrappers necessary
   * ideally this would be data driven and the backends would be plugins...
 **/
-class FDerivedDataBackendGraph : public FDerivedDataBackend
+class FDerivedDataBackendGraph : public FDerivedDataBackend, public ICacheFactory
 {
 public:
 	/**
@@ -109,7 +113,7 @@ public:
 		// Make sure AsyncPutWrapper and KeyLengthWrapper are created
 		if( !AsyncPutWrapper )
 		{
-			AsyncPutWrapper = new FDerivedDataBackendAsyncPutWrapper( RootCache, true );
+			AsyncPutWrapper = new FDerivedDataBackendAsyncPutWrapper( *this, RootCache, true );
 			check(AsyncPutWrapper);
 			CreatedBackends.Add( AsyncPutWrapper );
 			RootCache = AsyncPutWrapper;
@@ -277,7 +281,7 @@ public:
 				FGuid Temp = FGuid::NewGuid();
 				ReadPakFilename = PakFilename;
 				WritePakFilename = PakFilename + TEXT(".") + Temp.ToString();
-				WritePakCache = bCompressed? new FCompressedPakFileDerivedDataBackend( *WritePakFilename, true ) : new FPakFileDerivedDataBackend( *WritePakFilename, true );
+				WritePakCache = bCompressed? new FCompressedPakFileDerivedDataBackend( *this, *WritePakFilename, true ) : new FPakFileDerivedDataBackend( *this, *WritePakFilename, true );
 				PakNode = WritePakCache;
 			}
 			else
@@ -285,7 +289,7 @@ public:
 				bool bReadPak = FPlatformFileManager::Get().GetPlatformFile().FileExists( *PakFilename );
 				if( bReadPak )
 				{
-					FPakFileDerivedDataBackend* ReadPak = bCompressed? new FCompressedPakFileDerivedDataBackend( *PakFilename, false ) : new FPakFileDerivedDataBackend( *PakFilename, false );
+					FPakFileDerivedDataBackend* ReadPak = bCompressed? new FCompressedPakFileDerivedDataBackend( *this, *PakFilename, false ) : new FPakFileDerivedDataBackend( *this, *PakFilename, false );
 					ReadPakFilename = PakFilename;
 					PakNode = ReadPak;
 					ReadPakCache.Add(ReadPak);
@@ -374,7 +378,7 @@ public:
 		FDerivedDataBackendInterface* AsyncNode = NULL;
 		if( InnerNode )
 		{
-			AsyncNode = new FDerivedDataBackendAsyncPutWrapper( InnerNode, true );
+			AsyncNode = new FDerivedDataBackendAsyncPutWrapper( *this, InnerNode, true );
 		}
 		else
 		{
@@ -475,7 +479,7 @@ public:
 		FDerivedDataBackendInterface* Hierarchy = NULL;
 		if( InnerNodes.Num() > 1 )
 		{
-			FHierarchicalDerivedDataBackend* HierarchyBackend = new FHierarchicalDerivedDataBackend( InnerNodes );
+			FHierarchicalDerivedDataBackend* HierarchyBackend = new FHierarchicalDerivedDataBackend( *this, InnerNodes );
 			Hierarchy = HierarchyBackend;
 			if (HierarchicalWrapper == NULL)
 			{
@@ -613,7 +617,7 @@ public:
 
 				if (!bShared || IFileManager::Get().DirectoryExists(*Path))
 				{
-					InnerFileSystem = CreateFileSystemDerivedDataBackend( *Path, Entry, *WriteAccessLog);
+					InnerFileSystem = CreateFileSystemDerivedDataBackend(*this, *Path, Entry, *WriteAccessLog);
 				}
 
 				if (InnerFileSystem)
@@ -707,7 +711,7 @@ public:
 		}
 
 		// Insert the backend corruption wrapper. Since the filesystem already uses this, and we're recycling the data with the trailer intact, we need to use it for the S3 cache too.
-		FS3DerivedDataBackend* Backend = new FS3DerivedDataBackend(*ManifestPath, *BaseUrl, *Region, *CanaryObjectKey, *CachePath);
+		FS3DerivedDataBackend* Backend = new FS3DerivedDataBackend(*this, *ManifestPath, *BaseUrl, *Region, *CanaryObjectKey, *CachePath);
 		return new FDerivedDataBackendCorruptionWrapper(Backend);
 #else
 		UE_LOG(LogDerivedDataCache, Log, TEXT("S3 backend is not supported on the current platform."));
@@ -807,7 +811,7 @@ public:
 
 		const bool bReadOnly = GetParsedBool(Entry, TEXT("ReadOnly="));
 
-		FHttpDerivedDataBackend* backend = new FHttpDerivedDataBackend(*ServiceUrl, *Namespace, *OAuthProvider, *OAuthClientId, *OAuthSecret, bReadOnly);
+		FHttpDerivedDataBackend* backend = new FHttpDerivedDataBackend(*this, *ServiceUrl, *Namespace, *OAuthProvider, *OAuthClientId, *OAuthSecret, bReadOnly);
 
 		if (ForceSpeedClass != FDerivedDataBackendInterface::ESpeedClass::Unknown)
 		{
@@ -857,7 +861,7 @@ public:
 			MaxCacheSize = FMath::Min(MaxCacheSize, MaxSupportedCacheSize);
 
 			UE_LOG( LogDerivedDataCache, Display, TEXT("Max Cache Size: %d MB"), MaxCacheSize);
-			Cache = new FMemoryDerivedDataBackend(TEXT("Boot"), MaxCacheSize * 1024 * 1024, true /* bCanBeDisabled */);
+			Cache = new FMemoryDerivedDataBackend(*this, TEXT("Boot"), MaxCacheSize * 1024 * 1024, true /* bCanBeDisabled */);
 
 			if( Cache && Filename.Len() )
 			{
@@ -899,7 +903,7 @@ public:
 		FString Filename;
 
 		FParse::Value( Entry, TEXT("Filename="), Filename );
-		Cache = new FMemoryDerivedDataBackend(NodeName);
+		Cache = new FMemoryDerivedDataBackend(*this, NodeName);
 		if( Cache && Filename.Len() )
 		{
 			if( Cache->LoadCache( *Filename ) )
@@ -970,7 +974,7 @@ public:
 
 				for(const FString& MergePakName : MergePakList)
 				{
-					FPakFileDerivedDataBackend ReadPak(*FPaths::Combine(*FPaths::GetPath(WritePakFilename), *MergePakName), false);
+					FPakFileDerivedDataBackend ReadPak(*this, *FPaths::Combine(*FPaths::GetPath(WritePakFilename), *MergePakName), false);
 					WritePakCache->MergeCache(&ReadPak);
 				}
 			}
@@ -995,7 +999,7 @@ public:
 							UE_LOG(LogDerivedDataCache, Error, TEXT("Could not delete the pak file %s to overwrite it with a new one."), *ReadPakFilename);
 						}
 					}
-					if (!FPakFileDerivedDataBackend::SortAndCopy(WritePakFilename, ReadPakFilename))
+					if (!FPakFileDerivedDataBackend::SortAndCopy(*this, WritePakFilename, ReadPakFilename))
 					{
 						UE_LOG(LogDerivedDataCache, Error, TEXT("Couldn't sort pak file (%s)"), *WritePakFilename);
 					}
@@ -1062,7 +1066,7 @@ public:
 		FPakFileDerivedDataBackend* ReadPak = NULL;
 		if (HierarchicalWrapper && FPlatformFileManager::Get().GetPlatformFile().FileExists(PakFilename))
 		{
-			ReadPak = new FPakFileDerivedDataBackend(PakFilename, false);
+			ReadPak = new FPakFileDerivedDataBackend(*this, PakFilename, false);
 
 			HierarchicalWrapper->AddInnerBackend(ReadPak);
 			CreatedBackends.Add(ReadPak);
@@ -1109,6 +1113,16 @@ public:
 		return MakeShared<FDerivedDataCacheStatsNode>(nullptr, TEXT(""));
 	}
 
+	virtual FCacheBucket CreateBucket(FStringView Name) override
+	{
+		return Private::CreateCacheBucket(Name);
+	}
+
+	virtual FCacheRecordBuilder CreateRecord(const FCacheKey& Key) override
+	{
+		return Private::CreateCacheRecordBuilder(Key);
+	}
+
 private:
 
 	/** Delete all created backends in the reversed order they were created. */
@@ -1179,12 +1193,13 @@ private:
 	FAutoConsoleCommand UnountPakCommand;
 };
 
+} // UE::DerivedData::Backends
+
 FDerivedDataBackend& FDerivedDataBackend::Get()
 {
-	return FDerivedDataBackendGraph::Get();
+	return UE::DerivedData::Backends::FDerivedDataBackendGraph::Get();
 }
 
-
 /**
  * Parse debug options for the provided node name. Returns true if any options were specified
  */
diff --git a/Engine/Source/Developer/DerivedDataCache/Private/DerivedDataCache.cpp b/Engine/Source/Developer/DerivedDataCache/Private/DerivedDataCache.cpp
--- a/Engine/Source/Developer/DerivedDataCache/Private/DerivedDataCache.cpp
+++ b/Engine/Source/Developer/DerivedDataCache/Private/DerivedDataCache.cpp
@@ -744,496 +744,39 @@ public:
 		FStringView Context,
 		ECachePolicy Policy,
 		EPriority Priority,
-		FOnCachePutComplete&& Callback) final;
+		FOnCachePutComplete&& OnComplete) final
+	{
+		return FDerivedDataBackend::Get().GetRoot().Put(Records, Context, Policy, Priority, MoveTemp(OnComplete));
+	}
 
 	FRequest Get(
 		TConstArrayView<FCacheKey> Keys,
 		FStringView Context,
 		ECachePolicy Policy,
 		EPriority Priority,
-		FOnCacheGetComplete&& Callback) final;
+		FOnCacheGetComplete&& OnComplete) final
+	{
+		return FDerivedDataBackend::Get().GetRoot().Get(Keys, Context, Policy, Priority, MoveTemp(OnComplete));
+	}
 
 	FRequest GetPayload(
 		TConstArrayView<FCachePayloadKey> Keys,
 		FStringView Context,
 		ECachePolicy Policy,
 		EPriority Priority,
-		FOnCacheGetPayloadComplete&& Callback) final;
-
-	void CancelAll() final;
-
-private:
-	void Put(
-		FCacheRecord Record,
-		FStringView Context,
-		ECachePolicy Policy,
-		EPriority Priority,
-		FOnCachePutComplete& Callback);
-
-	void Get(
-		const FCacheKey& Key,
-		FStringView Context,
-		ECachePolicy Policy,
-		EPriority Priority,
-		FOnCacheGetComplete& Callback);
-
-	void GetPayload(
-		const FCachePayloadKey& Key,
-		FStringView Context,
-		ECachePolicy Policy,
-		EPriority Priority,
-		FOnCacheGetPayloadComplete& Callback);
-
-	static FString MakeRecordKey(const FCacheKey& Key)
+		FOnCacheGetPayloadComplete&& OnComplete) final
 	{
-		check(Key != FCacheKey::Empty);
-		TStringBuilder<96> Out;
-		Out << Key.Bucket << TEXT("_") << Key.Hash;
-		return FDerivedDataCacheInterface::SanitizeCacheKey(Out.ToString());
+		return FDerivedDataBackend::Get().GetRoot().GetPayload(Keys, Context, Policy, Priority, MoveTemp(OnComplete));
 	}
 
-	static FString MakePayloadKey(const FCacheKey& CacheKey, const FIoHash& CompressedHash)
+	void CancelAll() final
 	{
-		check(CacheKey != FCacheKey::Empty && !CompressedHash.IsZero());
-		TStringBuilder<96> Out;
-		Out << CacheKey.Bucket << TEXT("_CAS_") << CompressedHash;
-		return FDerivedDataCacheInterface::SanitizeCacheKey(Out.ToString());
+		return FDerivedDataBackend::Get().GetRoot().CancelAll();
 	}
 };
 
 ///////////////////////////////////////////////////////////////////////////////////////////////////////////////////////
 
-void FCache::Put(
-	FCacheRecord Record,
-	FStringView Context,
-	ECachePolicy Policy,
-	EPriority Priority,
-	FOnCachePutComplete& Callback)
-{
-	const FCacheKey& Key = Record.GetKey();
-	EStatus Status = EStatus::Error;
-
-	ON_SCOPE_EXIT
-	{
-		if (Callback)
-		{
-			Callback(FCachePutCompleteParams{Key, Status});
-		}
-	};
-
-	if (!EnumHasAnyFlags(Policy, ECachePolicy::Store))
-	{
-		UE_LOG(LogDerivedDataCache, VeryVerbose, TEXT("Cache: Put skipped for %s from '%.*s'"),
-			*WriteToString<96>(Key), Context.Len(), Context.GetData());
-		return;
-	}
-
-	FCbWriter Writer;
-	Writer.BeginObject();
-
-	const auto PutPayload = [&Writer, &Key, Context](const FPayload& Payload, bool bInline = false)
-	{
-		const FCompositeBuffer CompressedBuffer = Payload.GetData().GetCompressed();
-		const FIoHash CompressedHash = FIoHash::HashBuffer(CompressedBuffer);
-		if (!bInline)
-		{
-			const FCachePayloadKey PayloadKey{Key, Payload.GetId()};
-			const FSharedBuffer Buffer = CompressedBuffer.Flatten();
-			checkf(Buffer.GetSize() <= MAX_int32,
-				TEXT("Cache: Put failed for %s from '%.*s' because payload is 2 GiB or greater."),
-				*WriteToString<96>(PayloadKey), Context.Len(), Context.GetData());
-			const TConstArrayView<uint8> BufferView = MakeArrayView(
-				static_cast<const uint8*>(Buffer.GetData()),
-				static_cast<int32>(Buffer.GetSize()));
-			GetDerivedDataCacheRef().Put(*MakePayloadKey(PayloadKey.CacheKey, CompressedHash), BufferView, Context);
-		}
-
-		Writer.BeginObject();
-		Writer.AddObjectId("Id"_ASV, FCbObjectId(Payload.GetId().GetView()));
-		Writer.AddBinaryAttachment("RawHash"_ASV, Payload.GetRawHash());
-		Writer.AddHash("CompressedHash"_ASV, CompressedHash);
-		if (bInline)
-		{
-			Writer.AddBinary("CompressedData"_ASV, CompressedBuffer);
-		}
-		Writer.EndObject();
-	};
-
-	if (const FPayload& Value = Record.GetValuePayload())
-	{
-		const bool bInline = Record.GetAttachmentPayloads().IsEmpty();
-		Writer.SetName("Value"_ASV);
-		PutPayload(Value, bInline);
-	}
-	if (!Record.GetAttachmentPayloads().IsEmpty())
-	{
-		Writer.BeginArray("Attachments"_ASV);
-		for (const FPayload& Attachment : Record.GetAttachmentPayloads())
-		{
-			PutPayload(Attachment);
-		}
-		Writer.EndArray();
-	}
-
-	const FCbObject& Meta = Record.GetMeta();
-	if (Meta.CreateViewIterator())
-	{
-		Writer.AddObject("Meta"_ASV, Meta);
-	}
-
-	Writer.EndObject();
-
-	TArray<uint8> Data;
-	const uint64 SaveSize = Writer.GetSaveSize();
-	checkf(SaveSize <= MAX_int32,
-		TEXT("Cache: Put failed for %s from '%.*s' because record is 2 GiB or greater."),
-		*WriteToString<96>(Key), Context.Len(), Context.GetData());
-	Data.SetNumUninitialized(static_cast<int32>(SaveSize));
-	Writer.Save(MakeMemoryView(Data));
-
-	GetDerivedDataCacheRef().Put(*MakeRecordKey(Key), Data, Context);
-
-	UE_LOG(LogDerivedDataCache, VeryVerbose, TEXT("Cache: Put for %s from '%.*s'"),
-		*WriteToString<96>(Key), Context.Len(), Context.GetData());
-	Status = EStatus::Ok;
-}
-
-FRequest FCache::Put(
-	TArrayView<FCacheRecord> Records,
-	FStringView Context,
-	ECachePolicy Policy,
-	EPriority Priority,
-	FOnCachePutComplete&& Callback)
-{
-	for (FCacheRecord& Record : Records)
-	{
-		Put(MoveTemp(Record), Context, Policy, Priority, Callback);
-	}
-	return FRequest();
-}
-
-void FCache::Get(
-	const FCacheKey& Key,
-	FStringView Context,
-	ECachePolicy Policy,
-	EPriority Priority,
-	FOnCacheGetComplete& Callback)
-{
-	FPayload Value;
-	TArray<FPayload> Attachments;
-	TArray<FIoHash, TInlineAllocator<1>> SkippedPayloads;
-	FCacheRecordBuilder RecordBuilder = CreateRecord(Key);
-	EStatus Status = EStatus::Error;
-
-	ON_SCOPE_EXIT
-	{
-		if (Callback)
-		{
-			Callback(FCacheGetCompleteParams{RecordBuilder.Build(), Status});
-		}
-	};
-
-	// Skip the request if querying the cache is disabled.
-	if (!EnumHasAnyFlags(Policy, ECachePolicy::Query))
-	{
-		UE_LOG(LogDerivedDataCache, VeryVerbose, TEXT("Cache: Get skipped for %s from '%.*s'"),
-			*WriteToString<96>(Key), Context.Len(), Context.GetData());
-		return;
-	}
-
-	// Request the record from storage.
-	FSharedBuffer RecordBuffer;
-	{
-		TArray<uint8> RecordData;
-		if (!GetDerivedDataCacheRef().GetSynchronous(*MakeRecordKey(Key), RecordData, Context))
-		{
-			UE_LOG(LogDerivedDataCache, Verbose, TEXT("Cache: Get cache miss for %s from '%.*s'"),
-				*WriteToString<96>(Key), Context.Len(), Context.GetData());
-			return;
-		}
-		RecordBuffer = MakeSharedBufferFromArray(MoveTemp(RecordData));
-	}
-
-	// Validate that the record can be read as compact binary without crashing.
-	if (ValidateCompactBinaryRange(RecordBuffer, ECbValidateMode::Default) != ECbValidateError::None)
-	{
-		UE_LOG(LogDerivedDataCache, Verbose, TEXT("Cache: Get cache miss with corrupted record for %s from '%.*s'"),
-			*WriteToString<96>(Key), Context.Len(), Context.GetData());
-		return;
-	}
-
-	// Read the record from its compact binary fields.
-	const FCbObject RecordObject(MoveTemp(RecordBuffer));
-	FCbObject MetaObject = RecordObject["Meta"_ASV].AsObject();
-	const FCbObject ValueObject = RecordObject["Value"_ASV].AsObject();
-	const FCbArray AttachmentsArray = RecordObject["Attachments"_ASV].AsArray();
-
-	// Read the value and attachments if they have been requested.
-	const auto GetPayload = [&SkippedPayloads, &Key, Policy, Context](FCbObject PayloadObject, ECachePolicy SkipFlag) -> FPayload
-	{
-		TArray<uint8> PayloadData;
-		const FPayloadId PayloadId = FPayloadId(PayloadObject.FindView("Id"_ASV).AsObjectId().GetView());
-		const FIoHash PayloadRawHash = PayloadObject.FindView("RawHash"_ASV).AsBinaryAttachment();
-		const FIoHash PayloadCompressedHash = PayloadObject.FindView("CompressedHash"_ASV).AsHash();
-		const FSharedBuffer PayloadCompressedData = PayloadObject["CompressedData"_ASV].AsBinary();
-		const FCachePayloadKey PayloadKey{Key, PayloadId};
-		if (EnumHasAnyFlags(Policy, SkipFlag))
-		{
-			if (PayloadCompressedData.IsNull())
-			{
-				SkippedPayloads.Add(PayloadCompressedHash);
-			}
-			return FPayload(PayloadId, PayloadRawHash);
-		}
-
-		FCompressedBuffer CompressedBuffer;
-		if (PayloadCompressedData)
-		{
-			CompressedBuffer = FCompressedBuffer::FromCompressed(PayloadCompressedData);
-		}
-		else if (GetDerivedDataCacheRef().GetSynchronous(*MakePayloadKey(PayloadKey.CacheKey, PayloadCompressedHash), PayloadData, Context))
-		{
-			CompressedBuffer = FCompressedBuffer::FromCompressed(MakeSharedBufferFromArray(MoveTemp(PayloadData)));
-		}
-		else
-		{
-			UE_LOG(LogDerivedDataCache, Verbose, TEXT("Cache: Get cache miss with missing content for %s from '%.*s'"),
-				*WriteToString<96>(PayloadKey), Context.Len(), Context.GetData());
-			return FPayload();
-		}
-
-		if (CompressedBuffer.GetRawHash() == PayloadRawHash &&
-			FIoHash::HashBuffer(CompressedBuffer.GetCompressed()) == PayloadCompressedHash)
-		{
-			return FPayload(PayloadId, MoveTemp(CompressedBuffer));
-		}
-		else
-		{
-			UE_LOG(LogDerivedDataCache, Verbose, TEXT("Cache: Get cache miss with corrupted content for %s from '%.*s'"),
-				*WriteToString<96>(PayloadKey), Context.Len(), Context.GetData());
-			return FPayload();
-		}
-	};
-
-	if (ValueObject.CreateViewIterator() && !(Value = GetPayload(ValueObject, ECachePolicy::SkipValue)))
-	{
-		return;
-	}
-
-	Attachments.Reserve(AttachmentsArray.Num());
-	for (FCbField AttachmentField : AttachmentsArray)
-	{
-		Attachments.Add(GetPayload(AttachmentField.AsObject(), ECachePolicy::SkipAttachments));
-		if (!Attachments.Last())
-		{
-			return;
-		}
-	}
-
-	// Check for existence of the value and attachments if they were skipped.
-	if (!SkippedPayloads.IsEmpty())
-	{
-		TArray<FString, TInlineAllocator<1>> SkippedPayloadKeys;
-		SkippedPayloadKeys.Reserve(SkippedPayloads.Num());
-		Algo::Transform(SkippedPayloads, SkippedPayloadKeys,
-			[&Key](const FIoHash& Hash) -> FString { return MakePayloadKey(Key, Hash); });
-		const bool bPayloadsExist = EnumHasAnyFlags(Policy, ECachePolicy::StoreLocal)
-			? GetDerivedDataCacheRef().TryToPrefetch(SkippedPayloadKeys, Context)
-			: GetDerivedDataCacheRef().AllCachedDataProbablyExists(SkippedPayloadKeys);
-		if (!bPayloadsExist)
-		{
-			UE_LOG(LogDerivedDataCache, Verbose, TEXT("Cache: Get cache miss with missing content for %s from '%.*s'"),
-				*WriteToString<96>(Key), Context.Len(), Context.GetData());
-			return;
-		}
-	}
-
-	// Populate the builder now that any error conditions have been handled.
-	if (!EnumHasAnyFlags(Policy, ECachePolicy::SkipMeta) && MetaObject.CreateViewIterator())
-	{
-		RecordBuilder.SetMeta(MoveTemp(MetaObject));
-	}
-	if (Value)
-	{
-		RecordBuilder.SetValue(MoveTemp(Value));
-	}
-	for (FPayload& Attachment : Attachments)
-	{
-		RecordBuilder.AddAttachment(MoveTemp(Attachment));
-	}
-
-	UE_LOG(LogDerivedDataCache, Verbose, TEXT("Cache: Get cache hit for %s from '%.*s'"),
-		*WriteToString<96>(Key), Context.Len(), Context.GetData());
-	Status = EStatus::Ok;
-}
-
-FRequest FCache::Get(
-	TConstArrayView<FCacheKey> Keys,
-	FStringView Context,
-	ECachePolicy Policy,
-	EPriority Priority,
-	FOnCacheGetComplete&& Callback)
-{
-	for (const FCacheKey& Key : Keys)
-	{
-		Get(Key, Context, Policy, Priority, Callback);
-	}
-	return FRequest();
-}
-
-void FCache::GetPayload(
-	const FCachePayloadKey& Key,
-	FStringView Context,
-	ECachePolicy Policy,
-	EPriority Priority,
-	FOnCacheGetPayloadComplete& Callback)
-{
-	FPayload Payload(Key.Id, FIoHash());
-	EStatus Status = EStatus::Error;
-
-	ON_SCOPE_EXIT
-	{
-		if (Callback)
-		{
-			Callback(FCacheGetPayloadCompleteParams{Key.CacheKey, MoveTemp(Payload), Status});
-		}
-	};
-
-	// Skip the request if querying the cache is disabled.
-	if (!EnumHasAnyFlags(Policy, ECachePolicy::Query))
-	{
-		UE_LOG(LogDerivedDataCache, VeryVerbose, TEXT("Cache: GetPayload skipped %s from '%.*s'"),
-			*WriteToString<96>(Key), Context.Len(), Context.GetData());
-		return;
-	}
-
-	// Request the record from storage.
-	FSharedBuffer RecordBuffer;
-	{
-		TArray<uint8> RecordData;
-		if (!GetDerivedDataCacheRef().GetSynchronous(*MakeRecordKey(Key.CacheKey), RecordData, Context))
-		{
-			UE_LOG(LogDerivedDataCache, Verbose, TEXT("Cache: Get cache miss with missing record for %s from '%.*s'"),
-				*WriteToString<96>(Key), Context.Len(), Context.GetData());
-			return;
-		}
-		RecordBuffer = MakeSharedBufferFromArray(MoveTemp(RecordData));
-	}
-
-	// Validate that the record can be read as compact binary without crashing.
-	if (ValidateCompactBinaryRange(RecordBuffer, ECbValidateMode::Default) != ECbValidateError::None)
-	{
-		UE_LOG(LogDerivedDataCache, Verbose, TEXT("Cache: Get cache miss with corrupted record for %s from '%.*s'"),
-			*WriteToString<96>(Key), Context.Len(), Context.GetData());
-		return;
-	}
-
-	// Read the record from its compact binary fields.
-	const FCbObject RecordObject(MoveTemp(RecordBuffer));
-	const FCbObject ValueObject = RecordObject["Value"_ASV].AsObject();
-	const FCbArrayView AttachmentsArray = RecordObject.FindView("Attachments"_ASV).AsArrayView();
-
-	// Find the hash of the payload in the record.
-	FIoHash PayloadCompressedHash;
-	FSharedBuffer PayloadCompressedData;
-	const FCbObjectId KeyObjectId(Key.Id.GetView());
-	if (ValueObject["Id"_ASV].AsObjectId() == KeyObjectId)
-	{
-		Payload = FPayload(Key.Id, ValueObject.FindView("RawHash"_ASV).AsBinaryAttachment());
-		PayloadCompressedHash = ValueObject.FindView("CompressedHash"_ASV).AsHash();
-		PayloadCompressedData = ValueObject["CompressedData"_ASV].AsBinary();
-	}
-	else
-	{
-		for (FCbFieldView AttachmentField : AttachmentsArray)
-		{
-			if (AttachmentField.AsObjectView()["Id"_ASV].AsObjectId() == KeyObjectId)
-			{
-				Payload = FPayload(Key.Id, AttachmentField.AsObjectView()["RawHash"_ASV].AsBinaryAttachment());
-				PayloadCompressedHash = AttachmentField.AsObjectView()["CompressedHash"_ASV].AsHash();
-				break;
-			}
-		}
-	}
-
-	// Verify that the payload was found in the record.
-	if (PayloadCompressedHash.IsZero())
-	{
-		UE_LOG(LogDerivedDataCache, Verbose,
-			TEXT("Cache: GetPayload cache miss with missing payload in record for %s from '%.*s'"),
-			*WriteToString<96>(Key), Context.Len(), Context.GetData());
-		return;
-	}
-
-	// Check for existence of the payload if it is being skipped.
-	if (EnumHasAnyFlags(Policy, ECachePolicy::SkipValue))
-	{
-		const bool bPayloadExists = PayloadCompressedData || EnumHasAnyFlags(Policy, ECachePolicy::StoreLocal)
-			? GetDerivedDataCacheRef().TryToPrefetch({MakePayloadKey(Key.CacheKey, PayloadCompressedHash)}, Context)
-			: GetDerivedDataCacheRef().CachedDataProbablyExists(*MakePayloadKey(Key.CacheKey, PayloadCompressedHash));
-		if (bPayloadExists)
-		{
-			UE_LOG(LogDerivedDataCache, Verbose, TEXT("Cache: GetPayload (exists) cache hit for %s from '%.*s'"),
-				*WriteToString<96>(Key), Context.Len(), Context.GetData());
-			Status = EStatus::Ok;
-		}
-		else
-		{
-			UE_LOG(LogDerivedDataCache, Verbose, TEXT("Cache: GetPayload (exists) cache miss for %s from '%.*s'"),
-				*WriteToString<96>(Key), Context.Len(), Context.GetData());
-		}
-		return;
-	}
-
-	// Request the payload from storage.
-	if (PayloadCompressedData.IsNull())
-	{
-		TArray<uint8> Data;
-		if (!GetDerivedDataCacheRef().GetSynchronous(*MakePayloadKey(Key.CacheKey, PayloadCompressedHash), Data, Context))
-		{
-			UE_LOG(LogDerivedDataCache, Verbose, TEXT("Cache: GetPayload cache miss for %s from '%.*s'"),
-				*WriteToString<96>(Key), Context.Len(), Context.GetData());
-			return;
-		}
-		PayloadCompressedData = MakeSharedBufferFromArray(MoveTemp(Data));
-	}
-
-	FCompressedBuffer CompressedBuffer = FCompressedBuffer::FromCompressed(PayloadCompressedData);
-	if (CompressedBuffer.GetRawHash() == Payload.GetRawHash() &&
-		FIoHash::HashBuffer(CompressedBuffer.GetCompressed()) == PayloadCompressedHash)
-	{
-		UE_LOG(LogDerivedDataCache, Verbose, TEXT("Cache: GetPayload cache hit for %s from '%.*s'"),
-			*WriteToString<96>(Key), Context.Len(), Context.GetData());
-		Payload = FPayload(Key.Id, MoveTemp(CompressedBuffer));
-		Status = EStatus::Ok;
-	}
-	else
-	{
-		UE_LOG(LogDerivedDataCache, Verbose, TEXT("Cache: GetPayload cache miss with corrupted content for %s from '%.*s'"),
-			*WriteToString<96>(Key), Context.Len(), Context.GetData());
-	}
-}
-
-FRequest FCache::GetPayload(
-	TConstArrayView<FCachePayloadKey> Keys,
-	FStringView Context,
-	ECachePolicy Policy,
-	EPriority Priority,
-	FOnCacheGetPayloadComplete&& Callback)
-{
-	for (const FCachePayloadKey& Key : Keys)
-	{
-		GetPayload(Key, Context, Policy, Priority, Callback);
-	}
-	return FRequest();
-}
-
-void FCache::CancelAll()
-{
-}
-
-///////////////////////////////////////////////////////////////////////////////////////////////////////////////////////
-
 ICache* CreateCache()
 {
 	return new FCache();
diff --git a/Engine/Source/Developer/DerivedDataCache/Private/DerivedDataLimitKeyLengthWrapper.h b/Engine/Source/Developer/DerivedDataCache/Private/DerivedDataLimitKeyLengthWrapper.h
--- a/Engine/Source/Developer/DerivedDataCache/Private/DerivedDataLimitKeyLengthWrapper.h
+++ b/Engine/Source/Developer/DerivedDataCache/Private/DerivedDataLimitKeyLengthWrapper.h
@@ -8,6 +8,9 @@
 #include "DerivedDataCacheUsageStats.h"
 #include "Misc/SecureHash.h"
 
+namespace UE::DerivedData::Backends
+{
+
 /** 
  * A backend wrapper that limits the key size and uses hashing...in this case it wraps the payload and the payload contains the full key to verify the integrity of the hash
 **/
@@ -239,6 +242,41 @@ public:
 		return Usage;
 	}
 
+	virtual FRequest Put(
+		TArrayView<FCacheRecord> Records,
+		FStringView Context,
+		ECachePolicy Policy,
+		EPriority Priority,
+		FOnCachePutComplete&& OnComplete) override
+	{
+		return InnerBackend->Put(Records, Context, Policy, Priority, MoveTemp(OnComplete));
+	}
+
+	virtual FRequest Get(
+		TConstArrayView<FCacheKey> Keys,
+		FStringView Context,
+		ECachePolicy Policy,
+		EPriority Priority,
+		FOnCacheGetComplete&& OnComplete) override
+	{
+		return InnerBackend->Get(Keys, Context, Policy, Priority, MoveTemp(OnComplete));
+	}
+
+	virtual FRequest GetPayload(
+		TConstArrayView<FCachePayloadKey> Keys,
+		FStringView Context,
+		ECachePolicy Policy,
+		EPriority Priority,
+		FOnCacheGetPayloadComplete&& OnComplete) override
+	{
+		return InnerBackend->GetPayload(Keys, Context, Policy, Priority, MoveTemp(OnComplete));
+	}
+
+	virtual void CancelAll() override
+	{
+		InnerBackend->CancelAll();
+	}
+
 private:
 	FDerivedDataCacheUsageStats UsageStats;
 
@@ -286,3 +324,5 @@ private:
 	/** Maximum size of the backend key length **/
 	int32 MaxKeyLength;
 };
+
+} // UE::DerivedData::Backends
diff --git a/Engine/Source/Developer/DerivedDataCache/Private/FileSystemDerivedDataBackend.cpp b/Engine/Source/Developer/DerivedDataCache/Private/FileSystemDerivedDataBackend.cpp
--- a/Engine/Source/Developer/DerivedDataCache/Private/FileSystemDerivedDataBackend.cpp
+++ b/Engine/Source/Developer/DerivedDataCache/Private/FileSystemDerivedDataBackend.cpp
@@ -3,18 +3,26 @@
 
 #include "CoreMinimal.h"
 #include "Algo/AllOf.h"
+#include "Algo/Accumulate.h"
 #include "Misc/MessageDialog.h"
 #include "HAL/FileManager.h"
 #include "Misc/CoreMisc.h"
 #include "Misc/CommandLine.h"
 #include "Misc/FileHelper.h"
 #include "Misc/Paths.h"
+#include "Misc/PathViews.h"
 #include "Misc/Guid.h"
+#include "Misc/ScopeExit.h"
 #include "Misc/ScopeLock.h"
+#include "Misc/StringBuilder.h"
+#include "Serialization/CompactBinary.h"
+#include "Serialization/CompactBinaryValidation.h"
+#include "Serialization/CompactBinaryWriter.h"
 #include "Async/TaskGraphInterfaces.h"
 #include "Async/Async.h"
 
 #include "DerivedDataCacheInterface.h"
+#include "DerivedDataCacheRecord.h"
 #include "DerivedDataBackendInterface.h"
 #include "DDCCleanup.h"
 
@@ -31,6 +39,9 @@
 #endif // PLATFORM_LINUX
 #define MAX_CACHE_EXTENTION_LEN (4)
 
+namespace UE::DerivedData::Backends
+{
+
 FString BuildPathForCacheKey(const TCHAR* CacheKey)
 {
 	FString Key = FString(CacheKey).ToUpper();
@@ -42,6 +53,22 @@ FString BuildPathForCacheKey(const TCHAR* CacheKey)
 	return HashPath / Key + TEXT(".udd");
 }
 
+void BuildPathForCacheObject(const FCacheKey& CacheKey, FStringBuilderBase& Path)
+{
+	uint32 Hash;
+	FMemory::Memcpy(&Hash, &CacheKey.Hash, sizeof(uint32));
+	Path.Appendf(TEXT("%1u/%1u/%1u/"), (Hash / 100) % 10, (Hash / 10) % 10, Hash % 10);
+	Path << CacheKey.Bucket << TEXT('_') << CacheKey.Hash << TEXT(".udd"_SV);
+}
+
+void BuildPathForCachePayload(const FCacheKey& CacheKey, const FIoHash& RawHash, FStringBuilderBase& Path)
+{
+	uint32 Hash;
+	FMemory::Memcpy(&Hash, &RawHash, sizeof(uint32));
+	Path.Appendf(TEXT("%1u/%1u/%1u/"), (Hash / 100) % 10, (Hash / 10) % 10, Hash % 10);
+	Path << CacheKey.Bucket << TEXT("_CAS_"_SV) << RawHash << TEXT(".udd"_SV);
+}
+
 /**
  * Helper function to get the value of parsed bool as the return value
  **/
@@ -70,8 +97,9 @@ public:
 	 * @param InCacheDirectory	directory to store the cache in
 	 * @param bForceReadOnly	if true, do not attempt to write to this cache
 	*/
-	FFileSystemDerivedDataBackend(const TCHAR* InCacheDirectory, const TCHAR* InParams, const TCHAR* InAccessLogFileName)
-		: CachePath(InCacheDirectory)
+	FFileSystemDerivedDataBackend(ICacheFactory& InFactory, const TCHAR* InCacheDirectory, const TCHAR* InParams, const TCHAR* InAccessLogFileName)
+		: Factory(InFactory)
+		, CachePath(InCacheDirectory)
 		, SpeedClass(ESpeedClass::Unknown)
 		, bReadOnly(false)
 		, bTouch(false)
@@ -212,7 +240,7 @@ public:
 			
 			if (IsUsable() && InAccessLogFileName != nullptr && *InAccessLogFileName != 0)
 			{
-				AccessLogWriter.Reset(new FAccessLogWriter(InAccessLogFileName));
+				AccessLogWriter.Reset(new FAccessLogWriter(InAccessLogFileName, CachePath));
 			}
 		}
 	}
@@ -446,30 +474,65 @@ public:
 	class FAccessLogWriter
 	{
 	public:
-		FAccessLogWriter(const TCHAR* FileName)
+		FAccessLogWriter(const TCHAR* FileName, const FString& CachePath)
 			: Archive(IFileManager::Get().CreateFileWriter(FileName, FILEWRITE_AllowRead))
+			, BasePath(CachePath / TEXT(""))
+		{
+		}
+
+		void Append(const TCHAR* CacheKey, FStringView Path)
 		{
+			FScopeLock Lock(&CriticalSection);
+
+			bool bIsAlreadyInSet = false;
+			CacheKeys.FindOrAdd(FString(CacheKey), &bIsAlreadyInSet);
+			if (!bIsAlreadyInSet)
+			{
+				AppendPath(Path);
+			}
 		}
 
-		void Append(const TCHAR* CacheKey)
+		void Append(const FCacheKey& CacheKey, FStringView Path)
 		{
 			FScopeLock Lock(&CriticalSection);
 
-			FString CacheKeyStr(CacheKey);
-			if (!CacheKeys.Contains(CacheKeyStr))
+			bool bIsAlreadyInSet = false;
+			ObjectKeys.FindOrAdd(CacheKey, &bIsAlreadyInSet);
+			if (!bIsAlreadyInSet)
 			{
-				CacheKeys.Add(MoveTemp(CacheKeyStr));
+				AppendPath(Path);
+			}
+		}
 
-				auto FileName = StringCast<ANSICHAR>(*BuildPathForCacheKey(CacheKey));
-				Archive->Serialize(const_cast<ANSICHAR*>(FileName.Get()), FileName.Length());
-				Archive->Serialize(const_cast<ANSICHAR*>(LINE_TERMINATOR_ANSI), sizeof(LINE_TERMINATOR_ANSI) - 1);
+		void Append(const FCacheKey& CacheKey, const FPayloadId& Id, FStringView Path)
+		{
+			FScopeLock Lock(&CriticalSection);
+
+			bool bIsAlreadyInSet = false;
+			PayloadKeys.FindOrAdd(FCachePayloadKey{CacheKey, Id}, &bIsAlreadyInSet);
+			if (!bIsAlreadyInSet)
+			{
+				AppendPath(Path);
 			}
 		}
 
 	private:
+		void AppendPath(FStringView Path)
+		{
+			if (Path.StartsWith(BasePath))
+			{
+				const FTCHARToUTF8 PathUtf8(Path);
+				Archive->Serialize(const_cast<ANSICHAR*>(PathUtf8.Get()), PathUtf8.Length());
+				Archive->Serialize(const_cast<ANSICHAR*>(LINE_TERMINATOR_ANSI), sizeof(LINE_TERMINATOR_ANSI) - 1);
+			}
+		}
+
 		TUniquePtr<FArchive> Archive;
+		FString BasePath;
 		FCriticalSection CriticalSection;
 		TSet<FString> CacheKeys;
+		TSet<FCacheKey> ObjectKeys;
+		TSet<FCachePayloadKey> PayloadKeys;
 	};
 	
 	
@@ -518,7 +581,7 @@ public:
 
 			if (AccessLogWriter.IsValid())
 			{
-				AccessLogWriter->Append(CacheKey);
+				AccessLogWriter->Append(CacheKey, Filename);
 			}
 
 			COOK_STAT(Timer.AddHit(0));
@@ -590,7 +653,7 @@ public:
 
 			if (AccessLogWriter.IsValid())
 			{
-				AccessLogWriter->Append(CacheKey);
+				AccessLogWriter->Append(CacheKey, Filename);
 			}
 
 			UE_LOG(LogDerivedDataCache, Verbose, TEXT("%s: Cache hit on %s (%d bytes, %.02f secs, %.2fMB/s)"), *GetName(), CacheKey, Data.Num(), ReadDuration, ReadSpeed);
@@ -637,9 +700,11 @@ public:
 
 		if (IsWritable())
 		{
+			FString Filename = BuildFilename(CacheKey);
+
 			if (AccessLogWriter.IsValid())
 			{
-				AccessLogWriter->Append(CacheKey);
+				AccessLogWriter->Append(CacheKey, Filename);
 			}
 
 			// don't put anything we pretended didn't exist
@@ -654,7 +719,6 @@ public:
 			{
 				COOK_STAT(Timer.AddHit(Data.Num()));
 				check(Data.Num());
-				FString Filename = BuildFilename(CacheKey);
 				FString TempFilename(TEXT("temp.")); 
 				TempFilename += FGuid::NewGuid().ToString();
 				TempFilename = FPaths::GetPath(Filename) / TempFilename;
@@ -794,6 +858,557 @@ public:
 		return true;
 	}
 
+	virtual FRequest Put(
+		TArrayView<FCacheRecord> Records,
+		FStringView Context,
+		ECachePolicy Policy,
+		EPriority Priority,
+		FOnCachePutComplete&& OnComplete) override
+	{
+		for (const FCacheRecord& Record : Records)
+		{
+			COOK_STAT(auto Timer = UsageStats.TimePut());
+			if (PutCacheRecord(Record, Context, Policy))
+			{
+				UE_LOG(LogDerivedDataCache, Verbose, TEXT("%s: Cache put complete for %s from '%.*s'"),
+					*CachePath, *WriteToString<96>(Record.GetKey()), Context.Len(), Context.GetData());
+				COOK_STAT(Timer.AddHit(MeasureCacheRecord(Record)));
+				if (OnComplete)
+				{
+					OnComplete({Record.GetKey(), EStatus::Ok});
+				}
+			}
+			else
+			{
+				COOK_STAT(Timer.AddMiss(MeasureCacheRecord(Record)));
+				if (OnComplete)
+				{
+					OnComplete({Record.GetKey(), EStatus::Error});
+				}
+			}
+		}
+		return FRequest();
+	}
+
+	virtual FRequest Get(
+		TConstArrayView<FCacheKey> Keys,
+		FStringView Context,
+		ECachePolicy Policy,
+		EPriority Priority,
+		FOnCacheGetComplete&& OnComplete) override
+	{
+		for (const FCacheKey& Key : Keys)
+		{
+			COOK_STAT(auto Timer = UsageStats.TimeGet());
+			if (FCacheRecord Record = GetCacheRecord(Key, Context, Policy))
+			{
+				UE_LOG(LogDerivedDataCache, Verbose, TEXT("%s: Cache hit for %s from '%.*s'"),
+					*CachePath, *WriteToString<96>(Key), Context.Len(), Context.GetData());
+				COOK_STAT(Timer.AddHit(MeasureCacheRecord(Record)));
+				if (OnComplete)
+				{
+					OnComplete({MoveTemp(Record), EStatus::Ok});
+				}
+			}
+			else
+			{
+				if (OnComplete)
+				{
+					OnComplete({Factory.CreateRecord(Key).Build(), EStatus::Error});
+				}
+			}
+		}
+		return FRequest();
+	}
+
+	virtual FRequest GetPayload(
+		TConstArrayView<FCachePayloadKey> Keys,
+		FStringView Context,
+		ECachePolicy Policy,
+		EPriority Priority,
+		FOnCacheGetPayloadComplete&& OnComplete) override
+	{
+		TArray<FCachePayloadKey, TInlineAllocator<16>> SortedKeys(Keys);
+		SortedKeys.StableSort();
+
+		FCacheRecord Record;
+		for (const FCachePayloadKey& Key : SortedKeys)
+		{
+			COOK_STAT(auto Timer = UsageStats.TimeGet());
+			if (!Record || Record.GetKey() != Key.CacheKey)
+			{
+				Record = GetCacheRecord(Key.CacheKey, Context, Policy | ECachePolicy::SkipData, /*bAlwaysLoadInlineData*/ true);
+			}
+			if (FPayload Payload = GetCachePayload(Key.CacheKey, Context, Policy, Record.GetPayload(Key.Id)))
+			{
+				UE_LOG(LogDerivedDataCache, Verbose, TEXT("%s: Cache hit for %s from '%.*s'"),
+					*CachePath, *WriteToString<96>(Key), Context.Len(), Context.GetData());
+				COOK_STAT(Timer.AddHit(Payload.GetData().GetRawSize()));
+				if (OnComplete)
+				{
+					OnComplete({Key.CacheKey, MoveTemp(Payload), EStatus::Ok});
+				}
+			}
+			else
+			{
+				if (OnComplete)
+				{
+					OnComplete({Key.CacheKey, FPayload(Key.Id), EStatus::Error});
+				}
+			}
+		}
+		return FRequest();
+	}
+
+	virtual void CancelAll() override
+	{
+	}
+
+private:
+	uint64 MeasureCacheRecord(const FCacheRecord& Record) const
+	{
+		return Record.GetMeta().GetSize() +
+			Record.GetValuePayload().GetData().GetRawSize() +
+			Algo::TransformAccumulate(Record.GetAttachmentPayloads(),
+				[](const FPayload& Payload) { return Payload.GetData().GetRawSize(); }, uint64(0));
+	}
+
+	bool PutCacheRecord(const FCacheRecord& Record, FStringView Context, ECachePolicy Policy) const
+	{
+		if (!IsWritable())
+		{
+			UE_LOG(LogDerivedDataCache, VeryVerbose, TEXT("%s: Skipped put of %s from '%.*s' because this is not writable"),
+				*CachePath, *WriteToString<96>(Record.GetKey()), Context.Len(), Context.GetData());
+			return false;
+		}
+
+		const FCacheKey& Key = Record.GetKey();
+
+		// Skip the request if querying the cache is disabled.
+		if (!EnumHasAnyFlags(Policy, SpeedClass == ESpeedClass::Local ? ECachePolicy::StoreLocal : ECachePolicy::StoreRemote))
+		{
+			UE_LOG(LogDerivedDataCache, VeryVerbose, TEXT("%s: Skipped put of %s from '%.*s' due to cache policy"),
+				*CachePath, *WriteToString<96>(Key), Context.Len(), Context.GetData());
+			return false;
+		}
+
+		// Save the payloads and build the cache object.
+		const TConstArrayView<FPayload> Attachments = Record.GetAttachmentPayloads();
+		TCbWriter<512> Writer;
+		Writer.BeginObject();
+		if (const FCbObject& Meta = Record.GetMeta())
+		{
+			Writer.AddObject("Meta"_ASV, Meta);
+			Writer.AddHash("MetaHash"_ASV, Meta.GetHash());
+		}
+		if (const FPayload& Value = Record.GetValuePayload())
+		{
+			Writer.SetName("Value"_ASV);
+			if (!PutCachePayload(Key, Context, Value, Writer, /*bStoreInline*/ Attachments.IsEmpty()))
+			{
+				return false;
+			}
+		}
+		if (!Attachments.IsEmpty())
+		{
+			Writer.BeginArray("Attachments"_ASV);
+			for (const FPayload& Attachment : Attachments)
+			{
+				if (!PutCachePayload(Key, Context, Attachment, Writer, /*bStoreInline*/ false))
+				{
+					return false;
+				}
+			}
+			Writer.EndArray();
+		}
+		Writer.EndObject();
+
+		// Save the record to storage.
+		TStringBuilder<256> Path;
+		BuildCacheObjectPath(Key, Path);
+		if (!FileExists(Path))
+		{
+			FUniqueBuffer Buffer = FUniqueBuffer::Alloc(Writer.GetSaveSize());
+			Writer.Save(Buffer);
+			if (!SaveFile(Path, FCompositeBuffer(Buffer.MoveToShared()), Context))
+			{
+				return false;
+			}
+		}
+
+		if (AccessLogWriter)
+		{
+			AccessLogWriter->Append(Key, Path);
+		}
+
+		return true;
+	}
+
+	FCacheRecord GetCacheRecord(const FCacheKey& Key, FStringView Context, ECachePolicy Policy, bool bAlwaysLoadInlineData = false) const
+	{
+		if (!IsUsable())
+		{
+			UE_LOG(LogDerivedDataCache, VeryVerbose, TEXT("%s: Skipped get of %s from '%.*s' because this cache store is not available"),
+				*CachePath, *WriteToString<96>(Key), Context.Len(), Context.GetData());
+			return FCacheRecord();
+		}
+
+		// Skip the request if querying the cache is disabled.
+		if (!EnumHasAnyFlags(Policy, SpeedClass == ESpeedClass::Local ? ECachePolicy::QueryLocal : ECachePolicy::QueryRemote))
+		{
+			UE_LOG(LogDerivedDataCache, VeryVerbose, TEXT("%s: Skipped get of %s from '%.*s' due to cache policy"),
+				*CachePath, *WriteToString<96>(Key), Context.Len(), Context.GetData());
+			return FCacheRecord();
+		}
+
+		// Request the record from storage.
+		TStringBuilder<256> Path;
+		BuildCacheObjectPath(Key, Path);
+		FSharedBuffer Buffer = LoadFile(Path, Context);
+		if (Buffer.IsNull())
+		{
+			UE_LOG(LogDerivedDataCache, Verbose, TEXT("%s: Cache miss with missing record for %s from '%.*s'"),
+				*CachePath, *WriteToString<96>(Key), Context.Len(), Context.GetData());
+			return FCacheRecord();
+		}
+
+		// Delete the record from storage if is corrupted or has missing payloads.
+		bool bDeleteCacheObject = true;
+		ON_SCOPE_EXIT
+		{
+			if (bDeleteCacheObject && !bReadOnly)
+			{
+				IFileManager::Get().Delete(*Path, /*bRequireExists*/ false, /*bEvenReadOnly*/ false, /*bQuiet*/ true);
+			}
+		};
+
+		// Validate that the record can be read as compact binary without crashing.
+		if (ValidateCompactBinaryRange(Buffer, ECbValidateMode::Default) != ECbValidateError::None)
+		{
+			UE_LOG(LogDerivedDataCache, Display, TEXT("%s: Cache miss with corrupted record for %s from '%.*s'"),
+				*CachePath, *WriteToString<96>(Key), Context.Len(), Context.GetData());
+			return FCacheRecord();
+		}
+
+		const FCbObject RecordObject(MoveTemp(Buffer));
+		FCacheRecordBuilder RecordBuilder = Factory.CreateRecord(Key);
+
+		if (!EnumHasAnyFlags(Policy, ECachePolicy::SkipMeta))
+		{
+			if (FCbFieldView MetaHash = RecordObject.FindView("MetaHash"_ASV))
+			{
+				if (FCbObject MetaObject = RecordObject["Meta"_ASV].AsObject(); MetaObject.GetHash() == MetaHash.AsHash())
+				{
+					RecordBuilder.SetMeta(MoveTemp(MetaObject));
+				}
+				else
+				{
+					UE_LOG(LogDerivedDataCache, Display, TEXT("%s: Cache miss with corrupted metadata for %s from '%.*s'"),
+						*CachePath, *WriteToString<96>(Key), Context.Len(), Context.GetData());
+					return FCacheRecord();
+				}
+			}
+		}
+
+		if (FCbObject ValueObject = RecordObject["Value"_ASV].AsObject())
+		{
+			const ECachePolicy ValuePolicy = Policy | (ECachePolicy::SkipData & ~ECachePolicy::SkipValue);
+			FPayload Payload = GetCachePayload(Key, Context, ValuePolicy, ValueObject, bAlwaysLoadInlineData);
+			if (Payload.IsNull())
+			{
+				return FCacheRecord();
+			}
+			RecordBuilder.SetValue(MoveTemp(Payload));
+		}
+
+		for (FCbField AttachmentField : RecordObject["Attachments"_ASV])
+		{
+			const ECachePolicy AttachmentsPolicy = Policy | (ECachePolicy::SkipData & ~ECachePolicy::SkipAttachments);
+			FPayload Payload = GetCachePayload(Key, Context, AttachmentsPolicy, AttachmentField.AsObject(), bAlwaysLoadInlineData);
+			if (Payload.IsNull())
+			{
+				return FCacheRecord();
+			}
+			RecordBuilder.AddAttachment(MoveTemp(Payload));
+		}
+
+		bDeleteCacheObject = false;
+
+		if (AccessLogWriter)
+		{
+			AccessLogWriter->Append(Key, Path);
+		}
+
+		return RecordBuilder.Build();
+	}
+
+	bool PutCachePayload(const FCacheKey& Key, FStringView Context, const FPayload& Payload, FCbWriter& Writer, bool bStoreInline) const
+	{
+		const FCompressedBuffer& CompressedBuffer = Payload.GetData();
+		const FIoHash RawHash = CompressedBuffer.GetRawHash();
+
+		if (!bStoreInline)
+		{
+			TStringBuilder<256> Path;
+			BuildCachePayloadPath(Key, RawHash, Path);
+			if (!FileExists(Path))
+			{
+				// Save the compressed buffer with its hash as a header.
+				const FIoHash CompressedHash = FIoHash::HashBuffer(CompressedBuffer.GetCompressed());
+				const FCompositeBuffer BufferWithHash(
+					FSharedBuffer::Clone(MakeMemoryView(CompressedHash.GetBytes())),
+					CompressedBuffer.GetCompressed());
+				if (!SaveFile(Path, BufferWithHash, Context))
+				{
+					return false;
+				}
+			}
+
+			if (AccessLogWriter)
+			{
+				AccessLogWriter->Append(Key, Payload.GetId(), Path);
+			}
+		}
+
+		Writer.BeginObject();
+		Writer.AddObjectId("Id"_ASV, FCbObjectId(Payload.GetId().GetView()));
+		if (bStoreInline)
+		{
+			Writer.AddHash("RawHash"_ASV, RawHash);
+			Writer.AddHash("CompressedHash"_ASV, FIoHash::HashBuffer(CompressedBuffer.GetCompressed()));
+			Writer.AddBinary("CompressedData"_ASV, CompressedBuffer.GetCompressed());
+		}
+		else
+		{
+			Writer.AddBinaryAttachment("RawHash"_ASV, RawHash);
+		}
+		Writer.EndObject();
+		return true;
+	}
+
+	FPayload GetCachePayload(const FCacheKey& Key, FStringView Context, ECachePolicy Policy, const FCbObject& Object, bool bAlwaysLoadInlineData = false) const
+	{
+		const FPayloadId Id(Object.FindView("Id"_ASV).AsObjectId().GetView());
+		const FIoHash RawHash = Object.FindView("RawHash"_ASV).AsHash();
+		FIoHash CompressedHash = Object.FindView("CompressedHash"_ASV).AsHash();
+		FSharedBuffer CompressedData = Object["CompressedData"_ASV].AsBinary();
+
+		if (Id.IsNull() || RawHash.IsZero() || !(CompressedHash.IsZero() == CompressedData.IsNull()))
+		{
+			UE_LOG(LogDerivedDataCache, Display, TEXT("%s: Cache miss with invalid record format for %s from '%.*s'"),
+				*CachePath, *WriteToString<96>(Key), Context.Len(), Context.GetData());
+			return FPayload();
+		}
+
+		FPayload Payload(Id, RawHash);
+
+		if (CompressedData)
+		{
+			if (EnumHasAllFlags(Policy, ECachePolicy::SkipData) && !bAlwaysLoadInlineData)
+			{
+				return Payload;
+			}
+			else
+			{
+				return ValidateCachePayload(Key, Context, Payload, CompressedHash, MoveTemp(CompressedData));
+			}
+		}
+
+		return GetCachePayload(Key, Context, Policy, Payload);
+	}
+
+	FPayload GetCachePayload(const FCacheKey& Key, FStringView Context, ECachePolicy Policy, const FPayload& Payload) const
+	{
+		if (Payload.GetData())
+		{
+			return Payload;
+		}
+
+		TStringBuilder<256> Path;
+		BuildCachePayloadPath(Key, Payload.GetRawHash(), Path);
+		if (EnumHasAllFlags(Policy, ECachePolicy::SkipData))
+		{
+			if (FileExists(Path))
+			{
+				if (AccessLogWriter)
+				{
+					AccessLogWriter->Append(Key, Payload.GetId(), Path);
+				}
+				return Payload;
+			}
+		}
+		else
+		{
+			if (FSharedBuffer CompressedData = LoadFile(Path, Context); CompressedData.GetSize() >= sizeof(FIoHash))
+			{
+				const FMemoryView CompressedDataView = CompressedData;
+				const FIoHash CompressedHash(CompressedDataView.Left(sizeof(FIoHash)));
+				CompressedData = FSharedBuffer::MakeView(CompressedDataView + sizeof(FIoHash), MoveTemp(CompressedData));
+				const FPayload FullPayload = ValidateCachePayload(Key, Context, Payload, CompressedHash, MoveTemp(CompressedData));
+				if (FullPayload && AccessLogWriter)
+				{
+					AccessLogWriter->Append(Key, FullPayload.GetId(), Path);
+				}
+				return FullPayload;
+			}
+		}
+
+		UE_LOG(LogDerivedDataCache, Verbose,
+			TEXT("%s: Cache miss with missing payload %s with hash %s for %s from '%.*s'"),
+			*CachePath, *WriteToString<16>(Payload.GetId()), *WriteToString<48>(Payload.GetRawHash()), *WriteToString<96>(Key),
+			Context.Len(), Context.GetData());
+		return FPayload();
+	}
+
+	FPayload ValidateCachePayload(
+		const FCacheKey& Key,
+		FStringView Context,
+		const FPayload& Payload,
+		const FIoHash& CompressedHash,
+		FSharedBuffer&& CompressedData) const
+	{
+		if (FCompressedBuffer CompressedBuffer = FCompressedBuffer::FromCompressed(MoveTemp(CompressedData));
+			CompressedBuffer &&
+			CompressedBuffer.GetRawHash() == Payload.GetRawHash() &&
+			FIoHash::HashBuffer(CompressedBuffer.GetCompressed()) == CompressedHash)
+		{
+			return FPayload(Payload.GetId(), MoveTemp(CompressedBuffer));
+		}
+		UE_LOG(LogDerivedDataCache, Display,
+			TEXT("%s: Cache miss with corrupted payload %s with hash %s for %s from '%.*s'"),
+			*CachePath, *WriteToString<16>(Payload.GetId()), *WriteToString<48>(Payload.GetRawHash()), *WriteToString<96>(Key),
+			Context.Len(), Context.GetData());
+		return FPayload();
+	}
+
+	void BuildCacheObjectPath(const FCacheKey& CacheKey, FStringBuilderBase& Path) const
+	{
+		Path << CachePath << TEXT('/');
+		BuildPathForCacheObject(CacheKey, Path);
+	}
+
+	void BuildCachePayloadPath(const FCacheKey& CacheKey, const FIoHash& RawHash, FStringBuilderBase& Path) const
+	{
+		Path << CachePath << TEXT('/');
+		BuildPathForCachePayload(CacheKey, RawHash, Path);
+	}
+
+	bool SaveFile(FStringBuilderBase& Path, const FCompositeBuffer& Buffer, FStringView Context) const
+	{
+		TStringBuilder<256> TempPath;
+		TempPath << FPathViews::GetPath(Path) << TEXT("/Temp.") << FGuid::NewGuid().ToString();
+
+		ON_SCOPE_EXIT
+		{
+			IFileManager::Get().Delete(*TempPath, /*bRequireExists*/ false, /*bEvenReadOnly*/ false, /*bQuiet*/ true);
+		};
+
+		if (SaveBufferToFile(TempPath, Buffer))
+		{
+			if (IFileManager::Get().FileSize(*TempPath) == int64(Buffer.GetSize()))
+			{
+				if (IFileManager::Get().Move(*Path, *TempPath, /*bReplace*/ false, /*bEvenIfReadOnly*/ false, /*bAttributes*/ false, /*bDoNotRetryOrError*/ true))
+				{
+					return true;
+				}
+				else
+				{
+					UE_LOG(LogDerivedDataCache, Log,
+						TEXT("%s: Move collision when writing file %s from '%.*s'"),
+						*CachePath, *Path, Context.Len(), Context.GetData());
+					return true;
+				}
+			}
+			else
+			{
+				UE_LOG(LogDerivedDataCache, Warning,
+					TEXT("%s: Failed to write to temp file %s when saving %s from '%.*s'. ")
+					TEXT("File is %" INT64_FMT " bytes when %" UINT64_FMT " bytes are expected."),
+					*CachePath, *TempPath, *Path, Context.Len(), Context.GetData(),
+					IFileManager::Get().FileSize(*TempPath), Buffer.GetSize());
+				return false;
+			}
+		}
+		else
+		{
+			UE_LOG(LogDerivedDataCache, Warning,
+				TEXT("%s: Failed to write to temp file %s when saving %s from '%.*s' (error 0x%08x)"),
+				*CachePath, *TempPath, *Path, Context.Len(), Context.GetData(), FPlatformMisc::GetLastError());
+			return false;
+		}
+	}
+
+	bool SaveBufferToFile(FStringBuilderBase& Path, const FCompositeBuffer& Buffer) const
+	{
+		if (TUniquePtr<FArchive> Ar{IFileManager::Get().CreateFileWriter(*Path, FILEWRITE_Silent)})
+		{
+			for (const FSharedBuffer& Segment : Buffer.GetSegments())
+			{
+				Ar->Serialize(const_cast<void*>(Segment.GetData()), Segment.GetSize());
+			}
+			return Ar->Close() && !Ar->IsError();
+		}
+		return false;
+	}
+
+	FSharedBuffer LoadFile(FStringBuilderBase& Path, FStringView Context) const
+	{
+		check(IsUsable());
+		const double StartTime = FPlatformTime::Seconds();
+
+		FSharedBuffer Buffer;
+
+		// Check for existence before reading because it may update the modification time and avoid the
+		// file being deleted by a cache cleanup thread or process.
+		if (!FileExists(Path))
+		{
+			return Buffer;
+		}
+
+		if (TUniquePtr<FArchive> Ar{IFileManager::Get().CreateFileReader(*Path, FILEREAD_Silent)})
+		{
+			const int64 TotalSize = Ar->TotalSize();
+			FUniqueBuffer MutableBuffer = FUniqueBuffer::Alloc(uint64(TotalSize));
+			Ar->Serialize(MutableBuffer.GetData(), TotalSize);
+			if (Ar->Close())
+			{
+				Buffer = MutableBuffer.MoveToShared();
+			}
+		}
+
+		const double ReadDuration = FPlatformTime::Seconds() - StartTime;
+		const double ReadSpeed = ReadDuration > 0.001 ? (Buffer.GetSize() / ReadDuration) / (1024.0 * 1024.0) : 0.0;
+
+		if (!GIsBuildMachine && ReadDuration > 5.0)
+		{
+			// Slower than 0.5 MiB/s?
+			UE_CLOG(ReadSpeed < 0.5, LogDerivedDataCache, Warning, TEXT("%s: Loading %s from '%.*s' is very slow (%.2f MiB/s); ")
+				TEXT("consider disabling this cache backend"), *CachePath, *Path, Context.Len(), Context.GetData(), ReadSpeed);
+		}
+
+		if (Buffer)
+		{
+			UE_LOG(LogDerivedDataCache, VeryVerbose, TEXT("%s: Loaded %s from '%.*s' (%" UINT64_FMT " bytes, %.02f secs, %.2f MiB/s)"),
+				*CachePath, *Path, Context.Len(), Context.GetData(), Buffer.GetSize(), ReadDuration, ReadSpeed);
+		}
+
+		return Buffer;
+	}
+
+	bool FileExists(FStringBuilderBase& Path) const
+	{
+		const FDateTime TimeStamp = IFileManager::Get().GetTimeStamp(*Path);
+		if (TimeStamp == FDateTime::MinValue())
+		{
+			return false;
+		}
+		if (bTouch || (!bReadOnly && (FDateTime::Now() - TimeStamp).GetDays() > (DaysToDeleteUnusedFiles / 4)))
+		{
+			IFileManager::Get().SetTimeStamp(*Path, FDateTime::Now());
+		}
+		return true;
+	}
+
 private:
 	FDerivedDataCacheUsageStats UsageStats;
 
@@ -803,11 +1418,12 @@ private:
 	 * @param	CacheKey	Alphanumeric+underscore key of this cache item
 	 * @return				filename built from the cache key
 	 */
-	FString BuildFilename(const TCHAR* CacheKey)
+	FString BuildFilename(const TCHAR* CacheKey) const
 	{
 		return CachePath / BuildPathForCacheKey(CacheKey);
 	}
 
+	ICacheFactory& Factory;
 	/** Base path we are storing the cache files in. **/
 	FString	CachePath;
 	/** Class of this cache */
@@ -877,9 +1493,9 @@ private:
 
 };
 
-FDerivedDataBackendInterface* CreateFileSystemDerivedDataBackend(const TCHAR* CacheDirectory, const TCHAR* InParams, const TCHAR* InAccessLogFileName /*= nullptr*/)
+FDerivedDataBackendInterface* CreateFileSystemDerivedDataBackend(ICacheFactory& Factory, const TCHAR* CacheDirectory, const TCHAR* InParams, const TCHAR* InAccessLogFileName /*= nullptr*/)
 {
-	FFileSystemDerivedDataBackend* FileDDB = new FFileSystemDerivedDataBackend( CacheDirectory, InParams, InAccessLogFileName);
+	FFileSystemDerivedDataBackend* FileDDB = new FFileSystemDerivedDataBackend(Factory, CacheDirectory, InParams, InAccessLogFileName);
 
 	if (!FileDDB->IsUsable())
 	{
@@ -889,3 +1505,5 @@ FDerivedDataBackendInterface* CreateFileSystemDerivedDataBackend(const TCHAR* Ca
 
 	return FileDDB;
 }
+
+} // UE::DerivedData::Backends
diff --git a/Engine/Source/Developer/DerivedDataCache/Private/HierarchicalDerivedDataBackend.h b/Engine/Source/Developer/DerivedDataCache/Private/HierarchicalDerivedDataBackend.h
--- a/Engine/Source/Developer/DerivedDataCache/Private/HierarchicalDerivedDataBackend.h
+++ b/Engine/Source/Developer/DerivedDataCache/Private/HierarchicalDerivedDataBackend.h
@@ -5,11 +5,17 @@
 #include "CoreMinimal.h"
 #include "DerivedDataBackendInterface.h"
 #include "ProfilingDebugging/CookStats.h"
+#include "DerivedDataCacheRecord.h"
 #include "DerivedDataCacheUsageStats.h"
 #include "DerivedDataBackendAsyncPutWrapper.h"
 #include "Templates/UniquePtr.h"
 #include "Containers/ArrayView.h"
 
+extern bool GVerifyDDC;
+
+namespace UE::DerivedData::Backends
+{
+
 /** 
  * A backend wrapper that implements a cache hierarchy of backends. 
 **/
@@ -21,9 +27,16 @@ public:
 	 * Constructor
 	 * @param	InInnerBackends Backends to call into for actual storage of the cache, first item is the "fastest cache"
 	 */
-	FHierarchicalDerivedDataBackend(const TArray<FDerivedDataBackendInterface*>& InInnerBackends)
-		: InnerBackends(InInnerBackends)
+	FHierarchicalDerivedDataBackend(ICacheFactory& InFactory, const TArray<FDerivedDataBackendInterface*>& InInnerBackends)
+		: Factory(InFactory)
+		, InnerBackends(InInnerBackends)
 		, bIsWritable(false)
+		, bHasLocalBackends(false)
+		, bHasRemoteBackends(false)
+		, bHasMultipleLocalBackends(false)
+		, bHasMultipleRemoteBackends(false)
+		, bHasWritableLocalBackends(false)
+		, bHasWritableRemoteBackends(false)
 	{
 		check(InnerBackends.Num() > 1); // if it is just one, then you don't need this wrapper
 		UpdateAsyncInnerBackends();
@@ -44,11 +57,27 @@ public:
 	void UpdateAsyncInnerBackends()
 	{
 		bIsWritable = false;
+		bHasLocalBackends = false;
+		bHasRemoteBackends = false;
+		bHasMultipleLocalBackends = false;
+		bHasMultipleRemoteBackends = false;
+		bHasWritableLocalBackends = false;
+		bHasWritableRemoteBackends = false;
 		for (int32 CacheIndex = 0; CacheIndex < InnerBackends.Num(); CacheIndex++)
 		{
-			if (InnerBackends[CacheIndex]->IsWritable())
+			const bool bIsWritableBackend = InnerBackends[CacheIndex]->IsWritable();
+			bIsWritable |= bIsWritableBackend;
+			if (InnerBackends[CacheIndex]->GetSpeedClass() == ESpeedClass::Local)
+			{
+				bHasWritableLocalBackends |= bIsWritableBackend;
+				bHasMultipleLocalBackends = bHasLocalBackends;
+				bHasLocalBackends = true;
+			}
+			else
 			{
-				bIsWritable = true;
+				bHasWritableRemoteBackends |= bIsWritableBackend;
+				bHasMultipleRemoteBackends = bHasRemoteBackends;
+				bHasRemoteBackends = true;
 			}
 		}
 		if (bIsWritable)
@@ -56,7 +85,7 @@ public:
 			for (int32 CacheIndex = 0; CacheIndex < InnerBackends.Num(); CacheIndex++)
 			{
 				// async puts to allow us to fill all levels without holding up the engine
-				AsyncPutInnerBackends.Emplace(new FDerivedDataBackendAsyncPutWrapper(InnerBackends[CacheIndex], false));
+				AsyncPutInnerBackends.Emplace(new FDerivedDataBackendAsyncPutWrapper(Factory, InnerBackends[CacheIndex], false));
 			}
 		}
 	}
@@ -112,8 +141,6 @@ public:
 			}
 			else
 			{
-				extern bool GVerifyDDC;
-
 				if (GVerifyDDC)
 				{
 					ensureMsgf(!AsyncPutInnerBackends[CacheIndex]->CachedDataProbablyExists(CacheKey), TEXT("%s did not exist in sync interface for CachedDataProbablyExists but was found in async wrapper"), CacheKey);
@@ -378,8 +405,6 @@ public:
 			}
 			else
 			{
-				extern bool GVerifyDDC;
-
 				if (GVerifyDDC)
 				{
 					TArray<uint8> TempData;
@@ -463,16 +488,186 @@ public:
 		return Usage;
 	}
 
+	virtual FRequest Put(
+		TArrayView<FCacheRecord> Records,
+		FStringView Context,
+		ECachePolicy Policy,
+		EPriority Priority,
+		FOnCachePutComplete&& OnComplete) override
+	{
+		TSet<FCacheKey> RecordsOk;
+		for (int32 PutCacheIndex = 0; PutCacheIndex < InnerBackends.Num(); PutCacheIndex++)
+		{
+			if (InnerBackends[PutCacheIndex]->IsWritable())
+			{
+				// Every record must put synchronously before switching to async calls.
+				if (RecordsOk.Num() < Records.Num())
+				{
+					InnerBackends[PutCacheIndex]->Put(Records, Context, Policy, Priority,
+						[&OnComplete, &RecordsOk](FCachePutCompleteParams&& Params)
+						{
+							if (Params.Status == EStatus::Ok)
+							{
+								bool bIsAlreadyInSet = false;
+								RecordsOk.FindOrAdd(Params.Key, &bIsAlreadyInSet);
+								if (OnComplete && !bIsAlreadyInSet)
+								{
+									OnComplete(MoveTemp(Params));
+								}
+							}
+						}).Wait();
+				}
+				else
+				{
+					AsyncPutInnerBackends[PutCacheIndex]->Put(Records, Context, Policy, Priority);
+				}
+			}
+		}
+
+		if (OnComplete && RecordsOk.Num() < Records.Num())
+		{
+			for (const FCacheRecord& Record : Records)
+			{
+				if (!RecordsOk.Contains(Record.GetKey()))
+				{
+					OnComplete({Record.GetKey(), EStatus::Error});
+				}
+			}
+		}
+
+		return FRequest();
+	}
+
+	virtual FRequest Get(
+		TConstArrayView<FCacheKey> Keys,
+		FStringView Context,
+		ECachePolicy Policy,
+		EPriority Priority,
+		FOnCacheGetComplete&& OnComplete) override
+	{
+		const bool bQueryLocal = bHasLocalBackends && EnumHasAnyFlags(Policy, ECachePolicy::QueryLocal);
+		const bool bStoreLocal = bHasWritableLocalBackends && EnumHasAnyFlags(Policy, ECachePolicy::StoreLocal);
+		const bool bQueryRemote = bHasRemoteBackends && EnumHasAnyFlags(Policy, ECachePolicy::QueryRemote);
+		const bool bStoreRemote = bHasWritableRemoteBackends && EnumHasAnyFlags(Policy, ECachePolicy::StoreRemote);
+		const bool bStoreLocalCopy = bStoreLocal && bHasMultipleLocalBackends && !EnumHasAnyFlags(Policy, ECachePolicy::SkipLocalCopy);
+
+		TArray<FCacheKey, TInlineAllocator<16>> RemainingKeys(Keys);
+
+		TSet<FCacheKey> KeysOk;
+		for (int32 GetCacheIndex = 0; GetCacheIndex < InnerBackends.Num() && !RemainingKeys.IsEmpty(); ++GetCacheIndex)
+		{
+			// Remove SkipData flags when possibly filling other backends because only complete records can be written.
+			const bool bIsLocalGet = InnerBackends[GetCacheIndex]->GetSpeedClass() == ESpeedClass::Local;
+			const bool bFill =
+				(bIsLocalGet && bStoreLocalCopy) ||
+				(bIsLocalGet && bStoreRemote) ||
+				(!bIsLocalGet && bStoreLocal) ||
+				(!bIsLocalGet && bStoreRemote && bHasMultipleRemoteBackends);
+			const ECachePolicy FillPolicy = bFill ? (Policy & ~ECachePolicy::SkipData) : Policy;
+
+			// Block on this because backends in this hierarchy are not expected to be asynchronous.
+			InnerBackends[GetCacheIndex]->Get(RemainingKeys, Context, Policy, Priority,
+				[this, Context, GetCacheIndex, bStoreLocal, bStoreRemote, bStoreLocalCopy, bIsLocalGet, bFill, &OnComplete, &KeysOk](FCacheGetCompleteParams&& Params)
+				{
+					if (Params.Status == EStatus::Ok)
+					{
+						if (bFill)
+						{
+							for (int32 FillCacheIndex = 0; FillCacheIndex < InnerBackends.Num(); ++FillCacheIndex)
+							{
+								if (GetCacheIndex != FillCacheIndex)
+								{
+									const bool bIsLocalFill = InnerBackends[FillCacheIndex]->GetSpeedClass() == ESpeedClass::Local;
+									if ((bIsLocalFill && bStoreLocal && bIsLocalGet <= bStoreLocalCopy) || (!bIsLocalFill && bStoreRemote))
+									{
+										FCacheRecord Record = Params.Record.Clone();
+										AsyncPutInnerBackends[FillCacheIndex]->Put(MakeArrayView(&Record, 1), Context);
+									}
+								}
+							}
+						}
+
+						KeysOk.Add(Params.Record.GetKey());
+						if (OnComplete)
+						{
+							OnComplete(MoveTemp(Params));
+						}
+					}
+				}).Wait();
+
+			RemainingKeys.RemoveAll([&KeysOk](const FCacheKey& Key) { return KeysOk.Contains(Key); });
+		}
+
+		if (OnComplete)
+		{
+			for (const FCacheKey& Key : RemainingKeys)
+			{
+				OnComplete({Factory.CreateRecord(Key).Build(), EStatus::Error});
+			}
+		}
+
+		return FRequest();
+	}
+
+	virtual FRequest GetPayload(
+		TConstArrayView<FCachePayloadKey> Keys,
+		FStringView Context,
+		ECachePolicy Policy,
+		EPriority Priority,
+		FOnCacheGetPayloadComplete&& OnComplete) override
+	{
+		TArray<FCachePayloadKey, TInlineAllocator<16>> RemainingKeys(Keys);
+
+		TSet<FCachePayloadKey> KeysOk;
+		for (int32 GetCacheIndex = 0; GetCacheIndex < InnerBackends.Num() && !RemainingKeys.IsEmpty(); ++GetCacheIndex)
+		{
+			InnerBackends[GetCacheIndex]->GetPayload(RemainingKeys, Context, Policy, Priority,
+				[&OnComplete, &KeysOk](FCacheGetPayloadCompleteParams&& Params)
+				{
+					if (Params.Status == EStatus::Ok)
+					{
+						KeysOk.Add({Params.Key, Params.Payload.GetId()});
+						if (OnComplete)
+						{
+							OnComplete(MoveTemp(Params));
+						}
+					}
+				}).Wait();
+
+			RemainingKeys.RemoveAll([&KeysOk](const FCachePayloadKey& Key) { return KeysOk.Contains(Key); });
+		}
 
+		if (OnComplete)
+		{
+			for (const FCachePayloadKey& Key : RemainingKeys)
+			{
+				OnComplete({Key.CacheKey, FPayload(Key.Id), EStatus::Error});
+			}
+		}
+
+		return FRequest();
+	}
+
+	virtual void CancelAll() override
+	{
+	}
 
 private:
 	FDerivedDataCacheUsageStats UsageStats;
 
+	ICacheFactory& Factory;
 	/** Array of backends forming the hierarchical cache...the first element is the fastest cache. **/
 	TArray<FDerivedDataBackendInterface*> InnerBackends;
 	/** Each of the backends wrapped with an async put **/
 	TArray<TUniquePtr<FDerivedDataBackendInterface> > AsyncPutInnerBackends;
 	/** As an optimization, we check our writable status at contruction **/
 	bool bIsWritable;
+	bool bHasLocalBackends;
+	bool bHasRemoteBackends;
+	bool bHasMultipleLocalBackends;
+	bool bHasMultipleRemoteBackends;
+	bool bHasWritableLocalBackends;
+	bool bHasWritableRemoteBackends;
 };
 
+} // UE::DerivedData::Backends
diff --git a/Engine/Source/Developer/DerivedDataCache/Private/HttpDerivedDataBackend.cpp b/Engine/Source/Developer/DerivedDataCache/Private/HttpDerivedDataBackend.cpp
--- a/Engine/Source/Developer/DerivedDataCache/Private/HttpDerivedDataBackend.cpp
+++ b/Engine/Source/Developer/DerivedDataCache/Private/HttpDerivedDataBackend.cpp
@@ -15,6 +15,7 @@
 #include "Algo/Find.h"
 #include "Containers/StaticArray.h"
 #include "Containers/Ticker.h"
+#include "DerivedDataCacheRecord.h"
 #include "Dom/JsonObject.h"
 #include "GenericPlatform/GenericPlatformFile.h"
 #include "HAL/PlatformFileManager.h"
@@ -31,6 +32,11 @@
 #include "Serialization/JsonSerializer.h"
 #include "Policies/CondensedJsonPrintPolicy.h"
 
+#if WITH_SSL
+#include "Ssl.h"
+#include <openssl/ssl.h>
+#endif
+
 // Enables data request helpers that internally
 // batch requests to reduce the number of concurrent
 // connections.
@@ -58,6 +64,9 @@
 #define UE_HTTPDDC_BATCH_HEAD_WEIGHT 1
 #define UE_HTTPDDC_BATCH_WEIGHT_HINT 12
 
+namespace UE::DerivedData::Backends
+{
+
 TRACE_DECLARE_INT_COUNTER(HttpDDC_Exist, TEXT("HttpDDC Exist"));
 TRACE_DECLARE_INT_COUNTER(HttpDDC_ExistHit, TEXT("HttpDDC Exist Hit"));
 TRACE_DECLARE_INT_COUNTER(HttpDDC_Get, TEXT("HttpDDC Get"));
@@ -92,7 +101,7 @@ private:
  * CURL has a global library initialization (curl_global_init). We rely on this happening in 
  * the Online/HTTP library which is a dependency on this module.
  */
-class FRequest
+class FHttpRequest
 {
 public:
 	/**
@@ -118,7 +127,7 @@ public:
 		FailedTimeout
 	};
 
-	FRequest(const TCHAR* InDomain, FHttpAccessToken* InAuthorizationToken, bool bInLogErrors)
+	FHttpRequest(const TCHAR* InDomain, FHttpAccessToken* InAuthorizationToken, bool bInLogErrors)
 		: bLogErrors(bInLogErrors)
 		, Domain(InDomain)
 		, AuthorizationToken(InAuthorizationToken)
@@ -127,7 +136,7 @@ public:
 		Reset();
 	}
 
-	~FRequest()
+	~FHttpRequest()
 	{
 		curl_easy_cleanup(Curl);
 	}
@@ -163,7 +172,7 @@ public:
 		curl_easy_setopt(Curl, CURLOPT_SSLCERTTYPE, "PEM");
 		// Response functions
 		curl_easy_setopt(Curl, CURLOPT_HEADERDATA, this);
-		curl_easy_setopt(Curl, CURLOPT_HEADERFUNCTION, &FRequest::StaticWriteHeaderFn);
+		curl_easy_setopt(Curl, CURLOPT_HEADERFUNCTION, &FHttpRequest::StaticWriteHeaderFn);
 		curl_easy_setopt(Curl, CURLOPT_WRITEDATA, this);
 		curl_easy_setopt(Curl, CURLOPT_WRITEFUNCTION, StaticWriteBodyFn);
 		// SSL certification verification
@@ -572,7 +581,7 @@ private:
 
 	static size_t StaticDebugCallback(CURL * Handle, curl_infotype DebugInfoType, char * DebugInfo, size_t DebugInfoSize, void* UserData)
 	{
-		FRequest* Request = static_cast<FRequest*>(UserData);
+		FHttpRequest* Request = static_cast<FHttpRequest*>(UserData);
 
 		switch (DebugInfoType)
 		{
@@ -620,7 +629,7 @@ private:
 
 	static size_t StaticReadFn(void* Ptr, size_t SizeInBlocks, size_t BlockSizeInBytes, void* UserData)
 	{
-		FRequest* Request = static_cast<FRequest*>(UserData);
+		FHttpRequest* Request = static_cast<FHttpRequest*>(UserData);
 		TArrayView<const uint8>& ReadDataView = Request->ReadDataView;
 
 		const size_t Offset = Request->BytesSent;
@@ -636,7 +645,7 @@ private:
 
 	static size_t StaticWriteHeaderFn(void* Ptr, size_t SizeInBlocks, size_t BlockSizeInBytes, void* UserData)
 	{
-		FRequest* Request = static_cast<FRequest*>(UserData);
+		FHttpRequest* Request = static_cast<FHttpRequest*>(UserData);
 		const size_t WriteSize = SizeInBlocks * BlockSizeInBytes;
 		TArray<uint8>* WriteHeaderBufferPtr = Request->WriteHeaderBufferPtr;
 		if (WriteHeaderBufferPtr && WriteSize > 0)
@@ -658,7 +667,7 @@ private:
 
 	static size_t StaticWriteBodyFn(void* Ptr, size_t SizeInBlocks, size_t BlockSizeInBytes, void* UserData)
 	{
-		FRequest* Request = static_cast<FRequest*>(UserData);
+		FHttpRequest* Request = static_cast<FHttpRequest*>(UserData);
 		const size_t WriteSize = SizeInBlocks * BlockSizeInBytes;
 		TArray<uint8>* WriteDataBufferPtr = Request->WriteDataBufferPtr;
 
@@ -692,7 +701,7 @@ private:
 
 	static size_t StaticSeekFn(void* UserData, curl_off_t Offset, int Origin)
 	{
-		FRequest* Request = static_cast<FRequest*>(UserData);
+		FHttpRequest* Request = static_cast<FHttpRequest*>(UserData);
 		size_t NewPosition = 0;
 
 		switch (Origin)
@@ -721,8 +730,8 @@ private:
 //----------------------------------------------------------------------------------------------------------
 bool VerifyPayload(const FSHAHash& Hash, const TArray<uint8>& Payload);
 bool VerifyPayload(const FIoHash& Hash, const TArray<uint8>& Payload);
-bool VerifyRequest(const class FRequest* Request, const TArray<uint8>& Payload);
-bool HashPayload(class FRequest* Request, const TArrayView<const uint8> Payload);
+bool VerifyRequest(const class FHttpRequest* Request, const TArray<uint8>& Payload);
+bool HashPayload(class FHttpRequest* Request, const TArrayView<const uint8> Payload);
 bool ShouldAbortForShutdown();
 
 //----------------------------------------------------------------------------------------------------------
@@ -742,7 +751,7 @@ struct FRequestPool
 		for (uint8 i = 0; i < Pool.Num(); ++i)
 		{
 			Pool[i].Usage = 0u;
-			Pool[i].Request = new FRequest(InServiceUrl, InAuthorizationToken, true);
+			Pool[i].Request = new FHttpRequest(InServiceUrl, InAuthorizationToken, true);
 		}
 		
 	}
@@ -762,7 +771,7 @@ struct FRequestPool
 	 * "owned by the caller and need to release it to the pool when work has been completed.
 	 * @return Usable request instance if one is available, otherwise null.
 	 */
-	FRequest* GetFreeRequest()
+	FHttpRequest* GetFreeRequest()
 	{
 		for (uint8 i = 0; i < Pool.Num(); ++i)
 		{
@@ -784,10 +793,10 @@ struct FRequestPool
 	 * "owned by the caller and need to release it to the pool when work has been completed.
 	 * @return Usable request instance.
 	 */
-	FRequest* WaitForFreeRequest()
+	FHttpRequest* WaitForFreeRequest()
 	{
 		TRACE_CPUPROFILER_EVENT_SCOPE(HttpDDC_WaitForConnPool);
-		FRequest* Request = nullptr;
+		FHttpRequest* Request = nullptr;
 		while (true)
 		{
 			Request = GetFreeRequest();
@@ -802,7 +811,7 @@ struct FRequestPool
 	 * Release request to the pool.
 	 * @param Request Request that should be freed. Note that any buffer owened by the request can now be reset.
 	 */
-	void ReleaseRequestToPool(FRequest* Request)
+	void ReleaseRequestToPool(FHttpRequest* Request)
 	{
 		for (uint8 i = 0; i < Pool.Num(); ++i)
 		{
@@ -819,7 +828,7 @@ struct FRequestPool
 	/**
 	 * While holding a request, make it shared across many users.
 	 */
-	void MakeRequestShared(FRequest* Request, uint8 Users)
+	void MakeRequestShared(FHttpRequest* Request, uint8 Users)
 	{
 		check(Users != 0);
 		for (uint8 i = 0; i < Pool.Num(); ++i)
@@ -839,7 +848,7 @@ private:
 	struct FEntry
 	{
 		std::atomic<uint8> Usage;
-		FRequest* Request;
+		FHttpRequest* Request;
 	};
 
 	TArray<FEntry> Pool;
@@ -872,20 +881,20 @@ public:
 		return Request != nullptr;
 	}
 
-	FRequest* Get() const
+	FHttpRequest* Get() const
 	{
 		check(IsValid());
 		return Request;
 	}
 
-	FRequest* operator->()
+	FHttpRequest* operator->()
 	{
 		check(IsValid());
 		return Request;
 	}
 
 private:
-	FRequest* Request;
+	FHttpRequest* Request;
 	FRequestPool* Pool;
 };
 
@@ -910,8 +919,8 @@ struct FDataRequestHelper
 		{
 			// We are below the threshold, make the connection immediately. OutData is set so this is a get.
 			FString Uri = FString::Printf(TEXT("api/v1/c/ddc/%s/%s/%s.raw"), InNamespace, InBucket, InCacheKey);
-			const FRequest::Result Result = Request->PerformBlockingDownload(*Uri, OutData);
-			if (FRequest::IsSuccessResponse(Request->GetResponseCode()))
+			const FHttpRequest::Result Result = Request->PerformBlockingDownload(*Uri, OutData);
+			if (FHttpRequest::IsSuccessResponse(Request->GetResponseCode()))
 			{
 				if (VerifyRequest(Request, *OutData))
 				{
@@ -925,8 +934,8 @@ struct FDataRequestHelper
 		{
 			// We are below the threshold, make the connection immediately. OutData is missing so this is a head.
 			FString Uri = FString::Printf(TEXT("api/v1/c/ddc/%s/%s/%s"), InNamespace, InBucket, InCacheKey);
-			const FRequest::Result Result = Request->PerformBlockingQuery<FRequest::Head>(*Uri);
-			if (FRequest::IsSuccessResponse(Request->GetResponseCode()))
+			const FHttpRequest::Result Result = Request->PerformBlockingQuery<FHttpRequest::Head>(*Uri);
+			if (FHttpRequest::IsSuccessResponse(Request->GetResponseCode()))
 			{
 				TRACE_COUNTER_ADD(HttpDDC_ExistHit, int64(1));
 				bVerified[0] = true;
@@ -971,7 +980,7 @@ struct FDataRequestHelper
 				InBucket,
 				CacheKeys,
 				TConstArrayView<TArray<uint8>*>(),
-				FRequest::RequestVerb::Head,
+				FHttpRequest::RequestVerb::Head,
 				&bVerified
 			};
 
@@ -1034,7 +1043,7 @@ private:
 		const TCHAR* Bucket;
 		TConstArrayView<const TCHAR*> CacheKeys;
 		TConstArrayView<TArray<uint8>*> OutDatas;
-		FRequest::RequestVerb Verb;
+		FHttpRequest::RequestVerb Verb;
 		TBitArray<>* bSuccess;
 	};
 
@@ -1044,11 +1053,11 @@ private:
 		std::atomic<uint32> Reserved;
 		std::atomic<uint32> Ready;
 		std::atomic<uint32> WeightHint;
-		FRequest* Request;
+		FHttpRequest* Request;
 		FEvent* Complete;
 	};
 
-	FRequest* Request;
+	FHttpRequest* Request;
 	FRequestPool* Pool;
 	TBitArray<> bVerified;
 	static std::atomic<uint32> FirstAvailableBatch;
@@ -1057,7 +1066,7 @@ private:
 	/**
 	 * Queues up a request to be batched. Blocks until the query is made.
 	 */
-	static FRequest* QueueBatchRequest(FRequestPool* InPool, 
+	static FHttpRequest* QueueBatchRequest(FRequestPool* InPool, 
 		const TCHAR* InNamespace, 
 		const TCHAR* InBucket, 
 		TConstArrayView<const TCHAR*> InCacheKeys,
@@ -1105,14 +1114,14 @@ private:
 				InBucket,
 				InCacheKeys,
 				OutDatas,
-				OutDatas.Num() ? FRequest::RequestVerb::Get : FRequest::RequestVerb::Head,
+				OutDatas.Num() ? FHttpRequest::RequestVerb::Get : FHttpRequest::RequestVerb::Head,
 				&bOutVerified
 			};
 
 			// Signal we are ready for batch to be submitted
 			Batch.Ready.fetch_add(1u, std::memory_order_release);
 
-			FRequest* Request = nullptr;
+			FHttpRequest* Request = nullptr;
 
 			// The first to reserve a slot is the "driver" of the batch
 			if (Reserve == 0)
@@ -1164,7 +1173,7 @@ private:
 	/**
 	 * Creates request uri and headers and submits the request
 	 */
-	static void PerformBatchQuery(FRequest* Request, TArrayView<FQueuedBatchEntry> Entries)
+	static void PerformBatchQuery(FHttpRequest* Request, TArrayView<FQueuedBatchEntry> Entries)
 	{
 		TRACE_CPUPROFILER_EVENT_SCOPE(HttpDDC_BatchGet);
 		const TCHAR* Uri(TEXT("api/v1/c/ddc-rpc/batchget"));
@@ -1179,7 +1188,7 @@ private:
 				TSharedPtr<FJsonObject> Object = MakeShared<FJsonObject>();
 				Object->SetField(TEXT("bucket"), MakeShared<FJsonValueString>(Entry.Bucket));
 				Object->SetField(TEXT("key"), MakeShared<FJsonValueString>(Entry.CacheKeys[KeyIdx]));
-				if (Entry.Verb == FRequest::RequestVerb::Head)
+				if (Entry.Verb == FHttpRequest::RequestVerb::Head)
 				{
 					Object->SetField(TEXT("verb"), MakeShared<FJsonValueString>(TEXT("HEAD")));
 				}
@@ -1194,7 +1203,7 @@ private:
 		FBufferArchive RequestData;
 		if (FJsonSerializer::Serialize(RequestObject.ToSharedRef(), TJsonWriterFactory<ANSICHAR, TCondensedJsonPrintPolicy<ANSICHAR>>::Create(&RequestData)))
 		{
-			Request->PerformBlockingUpload<FRequest::PostJson>(Uri, MakeArrayView(RequestData));
+			Request->PerformBlockingUpload<FHttpRequest::PostJson>(Uri, MakeArrayView(RequestData));
 			ResponseCode = Request->GetResponseCode();
 
 			if (ResponseCode == 200)
@@ -1447,11 +1456,11 @@ struct FDataUploadHelper
 		, bSuccess(false)
 		, bQueued(false)
 	{
-		FRequest* Request = InPool->GetFreeRequest();
+		FHttpRequest* Request = InPool->GetFreeRequest();
 		if (Request)
 		{
 			ResponseCode = PerformPut(Request, InNamespace, InBucket, InCacheKey, InData, InUsageStats);
-			bSuccess = FRequest::IsSuccessResponse(Request->GetResponseCode());
+			bSuccess = FHttpRequest::IsSuccessResponse(Request->GetResponseCode());
 
 			ProcessQueuedPutsAndReleaseRequest(InPool, Request, InUsageStats);
 		}
@@ -1509,7 +1518,7 @@ private:
 	bool bSuccess;
 	bool bQueued;
 
-	static void ProcessQueuedPutsAndReleaseRequest(FRequestPool* Pool, FRequest* Request, FDerivedDataCacheUsageStats& UsageStats)
+	static void ProcessQueuedPutsAndReleaseRequest(FRequestPool* Pool, FHttpRequest* Request, FDerivedDataCacheUsageStats& UsageStats)
 	{
 		while (Request)
 		{
@@ -1550,7 +1559,7 @@ private:
 		}
 	}
 
-	static int64 PerformPut(FRequest* Request, const TCHAR* Namespace, const TCHAR* Bucket, const TCHAR* CacheKey, const TArrayView<const uint8> Data, FDerivedDataCacheUsageStats& UsageStats)
+	static int64 PerformPut(FHttpRequest* Request, const TCHAR* Namespace, const TCHAR* Bucket, const TCHAR* CacheKey, const TArrayView<const uint8> Data, FDerivedDataCacheUsageStats& UsageStats)
 	{
 		COOK_STAT(auto Timer = UsageStats.TimePut());
 
@@ -1559,10 +1568,10 @@ private:
 		TStringBuilder<256> Uri;
 		Uri.Appendf(TEXT("api/v1/c/ddc/%s/%s/%s"), Namespace, Bucket, CacheKey);
 
-		Request->PerformBlockingUpload<FRequest::Put>(*Uri, Data);
+		Request->PerformBlockingUpload<FHttpRequest::Put>(*Uri, Data);
 
 		const int64 ResponseCode = Request->GetResponseCode();
-		if (FRequest::IsSuccessResponse(ResponseCode))
+		if (FHttpRequest::IsSuccessResponse(ResponseCode))
 		{
 			TRACE_COUNTER_ADD(HttpDDC_BytesSent, int64(Request->GetBytesSent()));
 			COOK_STAT(Timer.AddHit(Request->GetBytesSent()));
@@ -1581,9 +1590,6 @@ TLockFreePointerListUnordered<FDataUploadHelper::FQueuedEntry, PLATFORM_CACHE_LI
 //----------------------------------------------------------------------------------------------------------
 
 #if WITH_SSL
-#include "Ssl.h"
-
-#include <openssl/ssl.h>
 
 static int SslCertVerify(int PreverifyOk, X509_STORE_CTX* Context)
 {
@@ -1595,7 +1601,7 @@ static int SslCertVerify(int PreverifyOk, X509_STORE_CTX* Context)
 		SSL_CTX* SslContext = SSL_get_SSL_CTX(Handle);
 		check(SslContext);
 
-		FRequest* Request = static_cast<FRequest*>(SSL_CTX_get_app_data(SslContext));
+		FHttpRequest* Request = static_cast<FHttpRequest*>(SSL_CTX_get_app_data(SslContext));
 		check(Request);
 
 		const FString& Domain = Request->GetDomain();
@@ -1684,7 +1690,7 @@ bool VerifyPayload(const FIoHash& Hash, const TArray<uint8>& Payload)
  * @param Payload Payload received.
  * @return True if the data is correct, false if checksums doesn't match.
  */
-bool VerifyRequest(const FRequest* Request, const TArray<uint8>& Payload)
+bool VerifyRequest(const FHttpRequest* Request, const TArray<uint8>& Payload)
 {
 	FString ReceivedHashStr;
 	if (Request->GetHeader("X-Jupiter-Sha1", ReceivedHashStr))
@@ -1709,7 +1715,7 @@ bool VerifyRequest(const FRequest* Request, const TArray<uint8>& Payload)
  * @param Payload Payload that will be sent.
  * @return True on success, false on failure.
  */
-bool HashPayload(FRequest* Request, const TArrayView<const uint8> Payload)
+bool HashPayload(FHttpRequest* Request, const TArrayView<const uint8> Payload)
 {
 	FIoHash PayloadHash = FIoHash::HashBuffer(Payload.GetData(), Payload.Num());
 	Request->SetHeader(TEXT("X-Jupiter-IoHash"), *WriteToString<48>(PayloadHash));
@@ -1751,13 +1757,15 @@ uint32 FHttpAccessToken::GetSerial() const
 //----------------------------------------------------------------------------------------------------------
 
 FHttpDerivedDataBackend::FHttpDerivedDataBackend(
+	ICacheFactory& InFactory,
 	const TCHAR* InServiceUrl, 
 	const TCHAR* InNamespace, 
 	const TCHAR* InOAuthProvider,
 	const TCHAR* InOAuthClientId,
 	const TCHAR* InOAuthSecret,
 	const bool bInReadOnly)
-	: Domain(InServiceUrl)
+	: Factory(InFactory)
+	, Domain(InServiceUrl)
 	, Namespace(InNamespace)
 	, DefaultBucket(TEXT("default"))
 	, OAuthProvider(InOAuthProvider)
@@ -1814,10 +1822,10 @@ bool FHttpDerivedDataBackend::ApplyDebugOptions(FBackendDebugOptions& InOptions)
 
 bool FHttpDerivedDataBackend::IsServiceReady()
 {
-	FRequest Request(*Domain, nullptr, false);
-	FRequest::Result Result = Request.PerformBlockingDownload(TEXT("health/ready"), nullptr);
+	FHttpRequest Request(*Domain, nullptr, false);
+	FHttpRequest::Result Result = Request.PerformBlockingDownload(TEXT("health/ready"), nullptr);
 	
-	if (Result == FRequest::Success && Request.GetResponseCode() == 200)
+	if (Result == FHttpRequest::Success && Request.GetResponseCode() == 200)
 	{
 		UE_LOG(LogDerivedDataCache, Display, TEXT("HTTP DDC service status: %s."), *Request.GetResponseAsString());
 		return true;
@@ -1862,7 +1870,7 @@ bool FHttpDerivedDataBackend::AcquireAccessToken()
 		FString AuthDomain(DomainEnd, *OAuthProvider);
 		FString Uri(*OAuthProvider + DomainEnd + 1);
 
-		FRequest Request(*AuthDomain, nullptr, false);
+		FHttpRequest Request(*AuthDomain, nullptr, false);
 
 		// If contents of the secret string is a file path, resolve and read form data.
 		if (OAuthSecret.StartsWith(TEXT("\\\\")))
@@ -1890,9 +1898,9 @@ bool FHttpDerivedDataBackend::AcquireAccessToken()
 		auto OAuthFormDataUTF8 = FTCHARToUTF8(*OAuthFormData);
 		FormData.Append((uint8*)OAuthFormDataUTF8.Get(), OAuthFormDataUTF8.Length());
 
-		FRequest::Result Result = Request.PerformBlockingUpload<FRequest::Post>(*Uri, MakeArrayView(FormData));
+		FHttpRequest::Result Result = Request.PerformBlockingUpload<FHttpRequest::Post>(*Uri, MakeArrayView(FormData));
 
-		if (Result == FRequest::Success && Request.GetResponseCode() == 200)
+		if (Result == FHttpRequest::Success && Request.GetResponseCode() == 200)
 		{
 			TSharedPtr<FJsonObject> ResponseObject = Request.GetResponseAsJsonObject();
 			if (ResponseObject)
@@ -1967,7 +1975,7 @@ bool FHttpDerivedDataBackend::CachedDataProbablyExists(const TCHAR* CacheKey)
 		FDataRequestHelper RequestHelper(GetRequestPools[IsInGameThread()].Get(), *Namespace, *DefaultBucket, CacheKey, nullptr);
 		const int64 ResponseCode = RequestHelper.GetResponseCode();
 
-		if (FRequest::IsSuccessResponse(ResponseCode) && RequestHelper.IsSuccess())
+		if (FHttpRequest::IsSuccessResponse(ResponseCode) && RequestHelper.IsSuccess())
 		{
 			COOK_STAT(Timer.AddHit(0));
 			return true;
@@ -1985,12 +1993,12 @@ bool FHttpDerivedDataBackend::CachedDataProbablyExists(const TCHAR* CacheKey)
 	for (int32 Attempts = 0; Attempts < UE_HTTPDDC_MAX_ATTEMPTS; ++Attempts)
 	{
 		FScopedRequestPtr Request(GetRequestPools[IsInGameThread()].Get());
-		const FRequest::Result Result = Request->PerformBlockingQuery<FRequest::Head>(*Uri);
+		const FHttpRequest::Result Result = Request->PerformBlockingQuery<FHttpRequest::Head>(*Uri);
 		const int64 ResponseCode = Request->GetResponseCode();
 
-		if (FRequest::IsSuccessResponse(ResponseCode) || ResponseCode == 400)
+		if (FHttpRequest::IsSuccessResponse(ResponseCode) || ResponseCode == 400)
 		{
-			const bool bIsHit = (Result == FRequest::Success && FRequest::IsSuccessResponse(ResponseCode));
+			const bool bIsHit = (Result == FHttpRequest::Success && FHttpRequest::IsSuccessResponse(ResponseCode));
 			if (bIsHit)
 			{
 				TRACE_COUNTER_ADD(HttpDDC_ExistHit, int64(1));
@@ -2020,7 +2028,7 @@ TBitArray<> FHttpDerivedDataBackend::CachedDataProbablyExistsBatch(TConstArrayVi
 		FDataRequestHelper RequestHelper(GetRequestPools[IsInGameThread()].Get(), *Namespace, *DefaultBucket, CacheKeys);
 		const int64 ResponseCode = RequestHelper.GetResponseCode();
 
-		if (FRequest::IsSuccessResponse(ResponseCode) && RequestHelper.IsSuccess())
+		if (FHttpRequest::IsSuccessResponse(ResponseCode) && RequestHelper.IsSuccess())
 		{
 			COOK_STAT(Timer.AddHit(0));
 			return RequestHelper.IsBatchSuccess();
@@ -2052,10 +2060,10 @@ TBitArray<> FHttpDerivedDataBackend::CachedDataProbablyExistsBatch(TConstArrayVi
 	for (int32 Attempts = 0; Attempts < UE_HTTPDDC_MAX_ATTEMPTS; ++Attempts)
 	{
 		FScopedRequestPtr Request(GetRequestPools[IsInGameThread()].Get());
-		const FRequest::Result Result = Request->PerformBlockingUpload<FRequest::PostJson>(Uri, BodyView);
+		const FHttpRequest::Result Result = Request->PerformBlockingUpload<FHttpRequest::PostJson>(Uri, BodyView);
 		const int64 ResponseCode = Request->GetResponseCode();
 
-		if (Result == FRequest::Success && ResponseCode == 200)
+		if (Result == FHttpRequest::Success && ResponseCode == 200)
 		{
 			TArray<TSharedPtr<FJsonValue>> ResponseArray = Request->GetResponseAsJsonArray();
 
@@ -2103,7 +2111,7 @@ bool FHttpDerivedDataBackend::GetCachedData(const TCHAR* CacheKey, TArray<uint8>
 		FDataRequestHelper RequestHelper(GetRequestPools[IsInGameThread()].Get(), *Namespace, *DefaultBucket, CacheKey, &OutData);
 		const int64 ResponseCode = RequestHelper.GetResponseCode();
 
-		if (FRequest::IsSuccessResponse(ResponseCode) && RequestHelper.IsSuccess())
+		if (FHttpRequest::IsSuccessResponse(ResponseCode) && RequestHelper.IsSuccess())
 		{
 			COOK_STAT(Timer.AddHit(OutData.Num()));
 			check(OutData.Num() > 0);
@@ -2124,11 +2132,11 @@ bool FHttpDerivedDataBackend::GetCachedData(const TCHAR* CacheKey, TArray<uint8>
 		FScopedRequestPtr Request(GetRequestPools[IsInGameThread()].Get());
 		if (Request.IsValid())
 		{
-			FRequest::Result Result = Request->PerformBlockingDownload(*Uri, &OutData);
+			FHttpRequest::Result Result = Request->PerformBlockingDownload(*Uri, &OutData);
 			const uint64 ResponseCode = Request->GetResponseCode();
 
 			// Request was successful, make sure we got all the expected data.
-			if (FRequest::IsSuccessResponse(ResponseCode) && VerifyRequest(Request.Get(), OutData))
+			if (FHttpRequest::IsSuccessResponse(ResponseCode) && VerifyRequest(Request.Get(), OutData))
 			{
 				TRACE_COUNTER_ADD(HttpDDC_GetHit, int64(1));
 				TRACE_COUNTER_ADD(HttpDDC_BytesReceived, int64(Request->GetBytesReceived()));
@@ -2169,7 +2177,7 @@ FDerivedDataBackendInterface::EPutStatus FHttpDerivedDataBackend::PutCachedData(
 
 		const int64 ResponseCode = Request.GetResponseCode();
 
-		if (Request.IsSuccess() && (Request.IsQueued() || FRequest::IsSuccessResponse(ResponseCode)))
+		if (Request.IsSuccess() && (Request.IsQueued() || FHttpRequest::IsSuccessResponse(ResponseCode)))
 		{
 			return Request.IsQueued() ? EPutStatus::Executing : EPutStatus::Cached;
 		}
@@ -2199,7 +2207,7 @@ FDerivedDataBackendInterface::EPutStatus FHttpDerivedDataBackend::PutCachedData(
 			// Append the content hash to the header
 			HashPayload(Request.Get(), InData);
 
-			FRequest::Result Result = Request->PerformBlockingUpload<FRequest::Put>(*Uri, InData);
+			FHttpRequest::Result Result = Request->PerformBlockingUpload<FHttpRequest::Put>(*Uri, InData);
 			ResponseCode = Request->GetResponseCode();
 
 			if (ResponseCode == 200)
@@ -2238,7 +2246,7 @@ void FHttpDerivedDataBackend::RemoveCachedData(const TCHAR* CacheKey, bool bTran
 		FScopedRequestPtr Request(PutRequestPools[IsInGameThread()].Get());
 		if (Request.IsValid())
 		{
-			FRequest::Result Result = Request->PerformBlockingQuery<FRequest::Delete>(*Uri);
+			FHttpRequest::Result Result = Request->PerformBlockingQuery<FHttpRequest::Delete>(*Uri);
 			ResponseCode = Request->GetResponseCode();
 
 			if (ResponseCode == 200)
@@ -2264,4 +2272,57 @@ TSharedRef<FDerivedDataCacheStatsNode> FHttpDerivedDataBackend::GatherUsageStats
 	return Usage;
 }
 
+FRequest FHttpDerivedDataBackend::Put(
+	TArrayView<FCacheRecord> Records,
+	FStringView Context,
+	ECachePolicy Policy,
+	EPriority Priority,
+	FOnCachePutComplete&& OnComplete)
+{
+	if (OnComplete)
+	{
+		for (const FCacheRecord& Record : Records)
+		{
+			OnComplete({Record.GetKey(), EStatus::Error});
+		}
+	}
+	return FRequest();
+}
+
+FRequest FHttpDerivedDataBackend::Get(
+	TConstArrayView<FCacheKey> Keys,
+	FStringView Context,
+	ECachePolicy Policy,
+	EPriority Priority,
+	FOnCacheGetComplete&& OnComplete)
+{
+	if (OnComplete)
+	{
+		for (const FCacheKey& Key : Keys)
+		{
+			OnComplete({Factory.CreateRecord(Key).Build(), EStatus::Error});
+		}
+	}
+	return FRequest();
+}
+
+FRequest FHttpDerivedDataBackend::GetPayload(
+	TConstArrayView<FCachePayloadKey> Keys,
+	FStringView Context,
+	ECachePolicy Policy,
+	EPriority Priority,
+	FOnCacheGetPayloadComplete&& OnComplete)
+{
+	if (OnComplete)
+	{
+		for (const FCachePayloadKey& Key : Keys)
+		{
+			OnComplete({Key.CacheKey, FPayload(Key.Id), EStatus::Error});
+		}
+	}
+	return FRequest();
+}
+
+} // UE::DerivedData::Backends
+
 #endif //WITH_HTTP_DDC_BACKEND
diff --git a/Engine/Source/Developer/DerivedDataCache/Private/HttpDerivedDataBackend.h b/Engine/Source/Developer/DerivedDataCache/Private/HttpDerivedDataBackend.h
--- a/Engine/Source/Developer/DerivedDataCache/Private/HttpDerivedDataBackend.h
+++ b/Engine/Source/Developer/DerivedDataCache/Private/HttpDerivedDataBackend.h
@@ -15,6 +15,9 @@
 
 #if WITH_HTTP_DDC_BACKEND
 
+namespace UE::DerivedData::Backends
+{
+
 typedef TSharedPtr<class IHttpRequest> FHttpRequestPtr;
 typedef TSharedPtr<class IHttpResponse, ESPMode::ThreadSafe> FHttpResponsePtr;
 
@@ -35,7 +38,7 @@ public:
 	 * @param OAuthClientId	OAuth client identifier.
 	 * @param OAuthData		OAuth form data to send to login service. Can either be the raw form data or a Windows network file address (starting with "\\").
 	 */
-	FHttpDerivedDataBackend(const TCHAR* ServiceUrl, 
+	FHttpDerivedDataBackend(ICacheFactory& Factory, const TCHAR* ServiceUrl, 
 		const TCHAR* Namespace, 
 		const TCHAR* OAuthProvider, 
 		const TCHAR* OAuthClientId, 
@@ -71,8 +74,33 @@ public:
 
 	void SetSpeedClass(ESpeedClass InSpeedClass) { SpeedClass = InSpeedClass; }
 
+	virtual FRequest Put(
+		TArrayView<FCacheRecord> Records,
+		FStringView Context,
+		ECachePolicy Policy,
+		EPriority Priority,
+		FOnCachePutComplete&& OnComplete) override;
+
+	virtual FRequest Get(
+		TConstArrayView<FCacheKey> Keys,
+		FStringView Context,
+		ECachePolicy Policy,
+		EPriority Priority,
+		FOnCacheGetComplete&& OnComplete) override;
+
+	virtual FRequest GetPayload(
+		TConstArrayView<FCachePayloadKey> Keys,
+		FStringView Context,
+		ECachePolicy Policy,
+		EPriority Priority,
+		FOnCacheGetPayloadComplete&& OnComplete) override;
+
+	virtual void CancelAll() override
+	{
+	}
+
 private:
-	
+	ICacheFactory& Factory;
 	FString Domain;
 	FString Namespace;
 	FString DefaultBucket;
@@ -94,4 +122,6 @@ private:
 	bool ShouldRetryOnError(int64 ResponseCode);
 };
 
+} // UE::DerivedData::Backends
+
 #endif //WITH_HTTP_DDC_BACKEND
diff --git a/Engine/Source/Developer/DerivedDataCache/Private/MemoryDerivedDataBackend.cpp b/Engine/Source/Developer/DerivedDataCache/Private/MemoryDerivedDataBackend.cpp
--- a/Engine/Source/Developer/DerivedDataCache/Private/MemoryDerivedDataBackend.cpp
+++ b/Engine/Source/Developer/DerivedDataCache/Private/MemoryDerivedDataBackend.cpp
@@ -1,10 +1,18 @@
 // Copyright Epic Games, Inc. All Rights Reserved.
 
 #include "MemoryDerivedDataBackend.h"
+
+#include "Algo/Accumulate.h"
+#include "Misc/ScopeExit.h"
+#include "Serialization/CompactBinary.h"
 #include "Templates/UniquePtr.h"
 
-FMemoryDerivedDataBackend::FMemoryDerivedDataBackend(const TCHAR* InName, int64 InMaxCacheSize, bool bInCanBeDisabled)
-	: Name(InName)
+namespace UE::DerivedData::Backends
+{
+
+FMemoryDerivedDataBackend::FMemoryDerivedDataBackend(ICacheFactory& InFactory, const TCHAR* InName, int64 InMaxCacheSize, bool bInCanBeDisabled)
+	: Factory(InFactory)
+	, Name(InName)
 	, MaxCacheSize(InMaxCacheSize)
 	, bDisabled( false )
 	, CurrentCacheSize( SerializationSpecificDataSize )
@@ -373,4 +381,156 @@ bool FMemoryDerivedDataBackend::ShouldSimulateMiss(const TCHAR* InKey)
 	}
 
 	return false;
-}
\ No newline at end of file
+}
+
+int64 FMemoryDerivedDataBackend::CalcRawCacheRecordSize(const FCacheRecord& Record) const
+{
+	const auto CalcCachePayloadSize = [](const FPayload& Payload) { return Payload.GetData().GetRawSize(); };
+	const uint64 ValueSize = CalcCachePayloadSize(Record.GetValuePayload());
+	return int64(Algo::TransformAccumulate(Record.GetAttachmentPayloads(), CalcCachePayloadSize, ValueSize));
+}
+
+int64 FMemoryDerivedDataBackend::CalcSerializedCacheRecordSize(const FCacheRecord& Record) const
+{
+	// Estimate the serialized size of the cache record.
+	uint64 TotalSize = 20;
+	TotalSize += Record.GetKey().Bucket.ToString<TCHAR>().Len();
+	TotalSize += Record.GetMeta().GetSize();
+	const auto CalcCachePayloadSize = [](const FPayload& Payload)
+	{
+		return Payload ? Payload.GetData().GetCompressedSize() + 32 : 0;
+	};
+	TotalSize += CalcCachePayloadSize(Record.GetValuePayload());
+	TotalSize += Algo::TransformAccumulate(Record.GetAttachmentPayloads(), CalcCachePayloadSize, uint64(0));
+	return int64(TotalSize);
+}
+
+FRequest FMemoryDerivedDataBackend::Put(
+	TArrayView<FCacheRecord> Records,
+	FStringView Context,
+	ECachePolicy Policy,
+	EPriority Priority,
+	FOnCachePutComplete&& OnComplete)
+{
+	for (FCacheRecord& Record : Records)
+	{
+		const FCacheKey& Key = Record.GetKey();
+		EStatus Status = EStatus::Error;
+		ON_SCOPE_EXIT
+		{
+			if (OnComplete)
+			{
+				OnComplete({Key, Status});
+			}
+		};
+
+		if (!Record.GetValuePayload() && Record.GetAttachmentPayloads().IsEmpty())
+		{
+			FScopeLock ScopeLock(&SynchronizationObject);
+			if (bDisabled)
+			{
+				continue;
+			}
+			if (const FCacheRecord* Existing = CacheRecords.Find(Key))
+			{
+				CurrentCacheSize -= CalcSerializedCacheRecordSize(*Existing);
+				bMaxSizeExceeded = false;
+			}
+			Status = EStatus::Ok;
+		}
+		else
+		{
+			COOK_STAT(auto Timer = UsageStats.TimePut());
+			const int64 RecordSize = CalcSerializedCacheRecordSize(Record);
+
+			FScopeLock ScopeLock(&SynchronizationObject);
+			Status = CacheRecords.Contains(Key) ? EStatus::Ok : EStatus::Error;
+			if (bDisabled || Status == EStatus::Ok)
+			{
+				continue;
+			}
+			if (MaxCacheSize > 0 && (CurrentCacheSize + RecordSize) > MaxCacheSize)
+			{
+				UE_LOG(LogDerivedDataCache, Display, TEXT("Failed to cache data. Maximum cache size reached. CurrentSize %" INT64_FMT " KiB / MaxSize: %" INT64_FMT " KiB"), CurrentCacheSize / 1024, MaxCacheSize / 1024);
+				bMaxSizeExceeded = true;
+				continue;
+			}
+
+			CurrentCacheSize += RecordSize;
+			CacheRecords.Add(MoveTemp(Record));
+			COOK_STAT(Timer.AddHit(RecordSize));
+			Status = EStatus::Ok;
+		}
+	}
+	return FRequest();
+}
+
+FRequest FMemoryDerivedDataBackend::Get(
+	TConstArrayView<FCacheKey> Keys,
+	FStringView Context,
+	ECachePolicy Policy,
+	EPriority Priority,
+	FOnCacheGetComplete&& OnComplete)
+{
+	for (const FCacheKey& Key : Keys)
+	{
+		COOK_STAT(auto Timer = UsageStats.TimeGet());
+		FCacheRecord Record;
+		if (FScopeLock ScopeLock(&SynchronizationObject); const FCacheRecord* CacheRecord = CacheRecords.Find(Key))
+		{
+			Record = CacheRecord->Clone();
+		}
+		if (Record)
+		{
+			COOK_STAT(Timer.AddHit(CalcRawCacheRecordSize(Record)));
+			if (OnComplete)
+			{
+				OnComplete({MoveTemp(Record), EStatus::Ok});
+			}
+		}
+		else
+		{
+			if (OnComplete)
+			{
+				OnComplete({Factory.CreateRecord(Key).Build(), EStatus::Error});
+			}
+		}
+	}
+	return FRequest();
+}
+
+FRequest FMemoryDerivedDataBackend::GetPayload(
+	TConstArrayView<FCachePayloadKey> Keys,
+	FStringView Context,
+	ECachePolicy Policy,
+	EPriority Priority,
+	FOnCacheGetPayloadComplete&& OnComplete)
+{
+	for (const FCachePayloadKey& PayloadKey : Keys)
+	{
+		COOK_STAT(auto Timer = UsageStats.TimeGet());
+		FPayload Payload;
+		if (FScopeLock ScopeLock(&SynchronizationObject); const FCacheRecord* Record = CacheRecords.Find(PayloadKey.CacheKey))
+		{
+			Payload = Record->GetAttachmentPayload(PayloadKey.Id);
+		}
+		if (Payload)
+		{
+			COOK_STAT(Timer.AddHit(Payload.GetData().GetRawSize()));
+			if (OnComplete)
+			{
+				OnComplete({PayloadKey.CacheKey, MoveTemp(Payload), EStatus::Ok});
+			}
+		}
+		else
+		{
+			if (OnComplete)
+			{
+				OnComplete({PayloadKey.CacheKey, FPayload(PayloadKey.Id), EStatus::Error});
+			}
+		}
+	}
+	return FRequest();
+}
+
+} // UE::DerivedData::Backends
diff --git a/Engine/Source/Developer/DerivedDataCache/Private/MemoryDerivedDataBackend.h b/Engine/Source/Developer/DerivedDataCache/Private/MemoryDerivedDataBackend.h
--- a/Engine/Source/Developer/DerivedDataCache/Private/MemoryDerivedDataBackend.h
+++ b/Engine/Source/Developer/DerivedDataCache/Private/MemoryDerivedDataBackend.h
@@ -3,6 +3,7 @@
 #pragma once
 
 #include "CoreMinimal.h"
+#include "DerivedDataCacheRecord.h"
 #include "HAL/FileManager.h"
 #include "FileBackedDerivedDataBackend.h"
 #include "ProfilingDebugging/CookStats.h"
@@ -11,13 +12,16 @@
 
 class Error;
 
+namespace UE::DerivedData::Backends
+{
+
 /** 
  * A simple thread safe, memory based backend. This is used for Async puts and the boot cache.
 **/
 class FMemoryDerivedDataBackend : public FFileBackedDerivedDataBackend
 {
 public:
-	explicit FMemoryDerivedDataBackend(const TCHAR* InName, int64 InMaxCacheSize = -1, bool bCanBeDisabled = false);
+	explicit FMemoryDerivedDataBackend(ICacheFactory& InFactory, const TCHAR* InName, int64 InMaxCacheSize = -1, bool bCanBeDisabled = false);
 	~FMemoryDerivedDataBackend();
 
 	/** Return a name for this interface */
@@ -89,7 +93,31 @@ public:
 	 * Apply debug options
 	 */
 	bool ApplyDebugOptions(FBackendDebugOptions& InOptions) override;
-	
+
+	virtual FRequest Put(
+		TArrayView<FCacheRecord> Records,
+		FStringView Context,
+		ECachePolicy Policy,
+		EPriority Priority,
+		FOnCachePutComplete&& OnComplete) override;
+
+	virtual FRequest Get(
+		TConstArrayView<FCacheKey> Keys,
+		FStringView Context,
+		ECachePolicy Policy,
+		EPriority Priority,
+		FOnCacheGetComplete&& OnComplete) override;
+
+	virtual FRequest GetPayload(
+		TConstArrayView<FCachePayloadKey> Keys,
+		FStringView Context,
+		ECachePolicy Policy,
+		EPriority Priority,
+		FOnCacheGetPayloadComplete&& OnComplete) override;
+
+	virtual void CancelAll() override
+	{
+	}
 
 private:
 	/** Name of the cache file loaded (if any). */
@@ -112,11 +140,19 @@ private:
 		return (Key.Len() + 1) * sizeof(TCHAR) + sizeof(Val.Age) + Val.Data.Num();
 	}
 
+	int64 CalcRawCacheRecordSize(const FCacheRecord& Record) const;
+	int64 CalcSerializedCacheRecordSize(const FCacheRecord& Record) const;
+
+	/** Factory used to construct cache records. */
+	ICacheFactory& Factory;
+
 	/** Name of this cache (used for debugging) */
 	FString Name;
 
 	/** Set of files that are being written to disk asynchronously. */
 	TMap<FString, FCacheValue*> CacheItems;
+	/** Set of records in this cache. */
+	TSet<FCacheRecord, FCacheRecordKeyFuncs> CacheRecords;
 	/** Maximum size the cached items can grow up to ( in bytes ) */
 	int64 MaxCacheSize;
 	/** When set to true, this cache is disabled...ignore all requests. */
@@ -176,3 +212,4 @@ protected:
 	TSet<FName> DebugMissedKeys;
 };
 
+} // UE::DerivedData::Backends
diff --git a/Engine/Source/Developer/DerivedDataCache/Private/PakFileDerivedDataBackend.cpp b/Engine/Source/Developer/DerivedDataCache/Private/PakFileDerivedDataBackend.cpp
--- a/Engine/Source/Developer/DerivedDataCache/Private/PakFileDerivedDataBackend.cpp
+++ b/Engine/Source/Developer/DerivedDataCache/Private/PakFileDerivedDataBackend.cpp
@@ -2,10 +2,15 @@
 
 #include "PakFileDerivedDataBackend.h"
 #include "Misc/Compression.h"
+#include "DerivedDataCacheRecord.h"
 #include "DerivedDataCacheUsageStats.h"
 
-FPakFileDerivedDataBackend::FPakFileDerivedDataBackend(const TCHAR* InFilename, bool bInWriting)
-	: bWriting(bInWriting)
+namespace UE::DerivedData::Backends
+{
+
+FPakFileDerivedDataBackend::FPakFileDerivedDataBackend(ICacheFactory& InFactory, const TCHAR* InFilename, bool bInWriting)
+	: Factory(InFactory)
+	, bWriting(bInWriting)
 	, bClosed(false)
 	, Filename(InFilename)
 {
@@ -391,13 +396,13 @@ void FPakFileDerivedDataBackend::MergeCache(FPakFileDerivedDataBackend* OtherPak
 	}
 }
 
-bool FPakFileDerivedDataBackend::SortAndCopy(const FString &InputFilename, const FString &OutputFilename)
+bool FPakFileDerivedDataBackend::SortAndCopy(ICacheFactory& InFactory, const FString &InputFilename, const FString &OutputFilename)
 {
 	// Open the input and output files
-	FPakFileDerivedDataBackend InputPak(*InputFilename, false);
+	FPakFileDerivedDataBackend InputPak(InFactory, *InputFilename, false);
 	if (InputPak.bClosed) return false;
 
-	FPakFileDerivedDataBackend OutputPak(*OutputFilename, true);
+	FPakFileDerivedDataBackend OutputPak(InFactory, *OutputFilename, true);
 	if (OutputPak.bClosed) return false;
 
 	// Sort the key names
@@ -435,8 +440,59 @@ TSharedRef<FDerivedDataCacheStatsNode> FPakFileDerivedDataBackend::GatherUsageSt
 	return Usage;
 }
 
-FCompressedPakFileDerivedDataBackend::FCompressedPakFileDerivedDataBackend(const TCHAR* InFilename, bool bInWriting)
-	: FPakFileDerivedDataBackend(InFilename, bInWriting)
+FRequest FPakFileDerivedDataBackend::Put(
+	TArrayView<FCacheRecord> Records,
+	FStringView Context,
+	ECachePolicy Policy,
+	EPriority Priority,
+	FOnCachePutComplete&& OnComplete)
+{
+	if (OnComplete)
+	{
+		for (const FCacheRecord& Record : Records)
+		{
+			OnComplete({Record.GetKey(), EStatus::Error});
+		}
+	}
+	return FRequest();
+}
+
+FRequest FPakFileDerivedDataBackend::Get(
+	TConstArrayView<FCacheKey> Keys,
+	FStringView Context,
+	ECachePolicy Policy,
+	EPriority Priority,
+	FOnCacheGetComplete&& OnComplete)
+{
+	if (OnComplete)
+	{
+		for (const FCacheKey& Key : Keys)
+		{
+			OnComplete({Factory.CreateRecord(Key).Build(), EStatus::Error});
+		}
+	}
+	return FRequest();
+}
+
+FRequest FPakFileDerivedDataBackend::GetPayload(
+	TConstArrayView<FCachePayloadKey> Keys,
+	FStringView Context,
+	ECachePolicy Policy,
+	EPriority Priority,
+	FOnCacheGetPayloadComplete&& OnComplete)
+{
+	if (OnComplete)
+	{
+		for (const FCachePayloadKey& Key : Keys)
+		{
+			OnComplete({Key.CacheKey, FPayload(Key.Id), EStatus::Error});
+		}
+	}
+	return FRequest();
+}
+
+FCompressedPakFileDerivedDataBackend::FCompressedPakFileDerivedDataBackend(ICacheFactory& InFactory, const TCHAR* InFilename, bool bInWriting)
+	: FPakFileDerivedDataBackend(InFactory, InFilename, bInWriting)
 {
 }
 
@@ -470,3 +526,5 @@ bool FCompressedPakFileDerivedDataBackend::GetCachedData(const TCHAR* CacheKey,
 
 	return true;
 }
+
+} // UE::DerivedData::Backends
diff --git a/Engine/Source/Developer/DerivedDataCache/Private/PakFileDerivedDataBackend.h b/Engine/Source/Developer/DerivedDataCache/Private/PakFileDerivedDataBackend.h
--- a/Engine/Source/Developer/DerivedDataCache/Private/PakFileDerivedDataBackend.h
+++ b/Engine/Source/Developer/DerivedDataCache/Private/PakFileDerivedDataBackend.h
@@ -16,13 +16,16 @@
 
 class Error;
 
+namespace UE::DerivedData::Backends
+{
+
 /** 
  * A simple thread safe, pak file based backend. 
 **/
 class FPakFileDerivedDataBackend : public FDerivedDataBackendInterface
 {
 public:
-	FPakFileDerivedDataBackend(const TCHAR* InFilename, bool bInWriting);
+	FPakFileDerivedDataBackend(ICacheFactory& InFactory, const TCHAR* InFilename, bool bInWriting);
 	~FPakFileDerivedDataBackend();
 
 	void Close();
@@ -88,7 +91,7 @@ public:
 		return Filename;
 	}
 
-	static bool SortAndCopy(const FString &InputFilename, const FString &OutputFilename);
+	static bool SortAndCopy(ICacheFactory& InFactory, const FString &InputFilename, const FString &OutputFilename);
 
 	virtual TSharedRef<FDerivedDataCacheStatsNode> GatherUsageStats() const override;
 
@@ -98,6 +101,31 @@ public:
 
 	virtual bool ApplyDebugOptions(FBackendDebugOptions& InOptions) override { return false; }
 
+	virtual FRequest Put(
+		TArrayView<FCacheRecord> Records,
+		FStringView Context,
+		ECachePolicy Policy,
+		EPriority Priority,
+		FOnCachePutComplete&& OnComplete) override;
+
+	virtual FRequest Get(
+		TConstArrayView<FCacheKey> Keys,
+		FStringView Context,
+		ECachePolicy Policy,
+		EPriority Priority,
+		FOnCacheGetComplete&& OnComplete) override;
+
+	virtual FRequest GetPayload(
+		TConstArrayView<FCachePayloadKey> Keys,
+		FStringView Context,
+		ECachePolicy Policy,
+		EPriority Priority,
+		FOnCacheGetPayloadComplete&& OnComplete) override;
+
+	virtual void CancelAll() override
+	{
+	}
+
 private:
 	FDerivedDataCacheUsageStats UsageStats;
 
@@ -114,6 +142,7 @@ private:
 		}
 	};
 
+	ICacheFactory& Factory;
 	/** When set to true, we are a pak writer (we don't do reads). */
 	bool bWriting;
 	/** When set to true, we are a pak writer and we saved, so we shouldn't be used anymore. Also, a read cache that failed to open. */
@@ -136,7 +165,7 @@ private:
 class FCompressedPakFileDerivedDataBackend : public FPakFileDerivedDataBackend
 {
 public:
-	FCompressedPakFileDerivedDataBackend(const TCHAR* InFilename, bool bInWriting);
+	FCompressedPakFileDerivedDataBackend(ICacheFactory& InFactory, const TCHAR* InFilename, bool bInWriting);
 
 	virtual EPutStatus PutCachedData(const TCHAR* CacheKey, TArrayView<const uint8> InData, bool bPutEvenIfExists) override;
 	virtual bool GetCachedData(const TCHAR* CacheKey, TArray<uint8>& OutData) override;
@@ -151,3 +180,5 @@ private:
 	static const EName CompressionFormat = NAME_Zlib;
 	static const ECompressionFlags CompressionFlags = COMPRESS_BiasMemory;
 };
+
+} // UE::DerivedData::Backends
diff --git a/Engine/Source/Developer/DerivedDataCache/Private/S3DerivedDataBackend.cpp b/Engine/Source/Developer/DerivedDataCache/Private/S3DerivedDataBackend.cpp
--- a/Engine/Source/Developer/DerivedDataCache/Private/S3DerivedDataBackend.cpp
+++ b/Engine/Source/Developer/DerivedDataCache/Private/S3DerivedDataBackend.cpp
@@ -24,6 +24,7 @@
 #include "Misc/OutputDeviceRedirector.h"
 #include "Async/ParallelFor.h"
 #include "Serialization/MemoryReader.h"
+#include "DerivedDataCacheRecord.h"
 #include "Dom/JsonObject.h"
 #include "Serialization/JsonReader.h"
 #include "Serialization/JsonSerializer.h"
@@ -35,6 +36,11 @@
 #include "DesktopPlatformModule.h"
 #include "Misc/ConfigCacheIni.h"
 
+#if WITH_SSL
+#include "Ssl.h"
+#include <openssl/ssl.h>
+#endif
+
 #define S3DDC_BACKEND_WAIT_INTERVAL 0.01f
 #define S3DDC_HTTP_REQUEST_TIMEOUT_SECONDS 30L
 #define S3DDC_HTTP_REQUEST_TIMOUT_ENABLED 1
@@ -43,6 +49,9 @@
 #define S3DDC_MAX_ATTEMPTS 4
 #define S3DDC_MAX_BUFFER_RESERVE 104857600u
 
+namespace UE::DerivedData::Backends
+{
+
 TRACE_DECLARE_INT_COUNTER(S3DDC_Exist, TEXT("S3DDC Exist"));
 TRACE_DECLARE_INT_COUNTER(S3DDC_ExistHit, TEXT("S3DDC Exist Hit"));
 TRACE_DECLARE_INT_COUNTER(S3DDC_Get, TEXT("S3DDC Get"));
@@ -173,10 +182,10 @@ struct IRequestCallback
  * CURL has a global library initialization (curl_global_init). We rely on this happening in 
  * the Online/HTTP library which is a dependency on this module.
  */
-class FS3DerivedDataBackend::FRequest
+class FS3DerivedDataBackend::FHttpRequest
 {
 public:
-	FRequest(const ANSICHAR* InRegion, const ANSICHAR* InAccessKey, const ANSICHAR* InSecretKey)
+	FHttpRequest(const ANSICHAR* InRegion, const ANSICHAR* InAccessKey, const ANSICHAR* InSecretKey)
 		: Region(InRegion)
 		, AccessKey(InAccessKey)
 		, SecretKey(InSecretKey)
@@ -184,7 +193,7 @@ public:
 		Curl = curl_easy_init();
 	}
 
-	~FRequest()
+	~FHttpRequest()
 	{
 		curl_easy_cleanup(Curl);
 	}
@@ -249,11 +258,11 @@ public:
 		// Progress
 		curl_easy_setopt(Curl, CURLOPT_NOPROGRESS, 0);
 		curl_easy_setopt(Curl, CURLOPT_XFERINFODATA, Callback);
-		curl_easy_setopt(Curl, CURLOPT_XFERINFOFUNCTION, &FRequest::StaticStatusFn);
+		curl_easy_setopt(Curl, CURLOPT_XFERINFOFUNCTION, &FHttpRequest::StaticStatusFn);
 
 		// Response
 		curl_easy_setopt(Curl, CURLOPT_HEADERDATA, &CallbackData);
-		curl_easy_setopt(Curl, CURLOPT_HEADERFUNCTION, &FRequest::StaticWriteHeaderFn);
+		curl_easy_setopt(Curl, CURLOPT_HEADERFUNCTION, &FHttpRequest::StaticWriteHeaderFn);
 		curl_easy_setopt(Curl, CURLOPT_WRITEDATA, &CallbackData);
 		curl_easy_setopt(Curl, CURLOPT_WRITEFUNCTION, StaticWriteBodyFn);
 
@@ -526,7 +535,7 @@ public:
 		for (uint8 i = 0; i < Pool.Num(); ++i)
 		{
 			Pool[i].Usage = 0u;
-			Pool[i].Request = new FRequest(Region, AccessKey, SecretKey);
+			Pool[i].Request = new FHttpRequest(Region, AccessKey, SecretKey);
 		}
 	}
 
@@ -542,7 +551,7 @@ public:
 
 	long Download(const TCHAR* Url, IRequestCallback* Callback, TArray<uint8>& OutData, FOutputDevice* Log)
 	{
-		FRequest* Request = WaitForFreeRequest();
+		FHttpRequest* Request = WaitForFreeRequest();
 		long ResponseCode = Request->PerformBlocking(TCHAR_TO_ANSI(Url), Callback, OutData, Log);
 		ReleaseRequestToPool(Request);
 		return ResponseCode;
@@ -552,12 +561,12 @@ private:
 	struct FEntry
 	{
 		TAtomic<uint8> Usage;
-		FRequest* Request;
+		FHttpRequest* Request;
 	};
 
 	TArray<FEntry> Pool;
 
-	FRequest* WaitForFreeRequest()
+	FHttpRequest* WaitForFreeRequest()
 	{
 		TRACE_CPUPROFILER_EVENT_SCOPE(S3DDC_WaitForConnPool);
 		while (true)
@@ -577,7 +586,7 @@ private:
 		}
 	}
 
-	void ReleaseRequestToPool(FRequest* Request)
+	void ReleaseRequestToPool(FHttpRequest* Request)
 	{
 		for (uint8 i = 0; i < Pool.Num(); ++i)
 		{
@@ -753,8 +762,9 @@ struct FS3DerivedDataBackend::FRootManifest
 // FS3DerivedDataBackend
 //----------------------------------------------------------------------------------------------------------
 
-FS3DerivedDataBackend::FS3DerivedDataBackend(const TCHAR* InRootManifestPath, const TCHAR* InBaseUrl, const TCHAR* InRegion, const TCHAR* InCanaryObjectKey, const TCHAR* InCachePath)
-	: RootManifestPath(InRootManifestPath)
+FS3DerivedDataBackend::FS3DerivedDataBackend(ICacheFactory& InFactory, const TCHAR* InRootManifestPath, const TCHAR* InBaseUrl, const TCHAR* InRegion, const TCHAR* InCanaryObjectKey, const TCHAR* InCachePath)
+	: Factory(InFactory)
+	, RootManifestPath(InRootManifestPath)
 	, BaseUrl(InBaseUrl)
 	, Region(InRegion)
 	, CanaryObjectKey(InCanaryObjectKey)
@@ -1167,4 +1177,57 @@ bool FS3DerivedDataBackend::ShouldSimulateMiss(const TCHAR* InKey)
 	return false;
 }
 
+FRequest FS3DerivedDataBackend::Put(
+	TArrayView<FCacheRecord> Records,
+	FStringView Context,
+	ECachePolicy Policy,
+	EPriority Priority,
+	FOnCachePutComplete&& OnComplete)
+{
+	if (OnComplete)
+	{
+		for (const FCacheRecord& Record : Records)
+		{
+			OnComplete({Record.GetKey(), EStatus::Error});
+		}
+	}
+	return FRequest();
+}
+
+FRequest FS3DerivedDataBackend::Get(
+	TConstArrayView<FCacheKey> Keys,
+	FStringView Context,
+	ECachePolicy Policy,
+	EPriority Priority,
+	FOnCacheGetComplete&& OnComplete)
+{
+	if (OnComplete)
+	{
+		for (const FCacheKey& Key : Keys)
+		{
+			OnComplete({Factory.CreateRecord(Key).Build(), EStatus::Error});
+		}
+	}
+	return FRequest();
+}
+
+FRequest FS3DerivedDataBackend::GetPayload(
+	TConstArrayView<FCachePayloadKey> Keys,
+	FStringView Context,
+	ECachePolicy Policy,
+	EPriority Priority,
+	FOnCacheGetPayloadComplete&& OnComplete)
+{
+	if (OnComplete)
+	{
+		for (const FCachePayloadKey& Key : Keys)
+		{
+			OnComplete({Key.CacheKey, FPayload(Key.Id), EStatus::Error});
+		}
+	}
+	return FRequest();
+}
+
+} // UE::DerivedData::Backends
+
 #endif
diff --git a/Engine/Source/Developer/DerivedDataCache/Private/S3DerivedDataBackend.h b/Engine/Source/Developer/DerivedDataCache/Private/S3DerivedDataBackend.h
--- a/Engine/Source/Developer/DerivedDataCache/Private/S3DerivedDataBackend.h
+++ b/Engine/Source/Developer/DerivedDataCache/Private/S3DerivedDataBackend.h
@@ -15,6 +15,9 @@
 
 #if WITH_S3_DDC_BACKEND
 
+namespace UE::DerivedData::Backends
+{
+
 /**
  * Backend for a read-only AWS S3 based caching service.
  **/
@@ -30,7 +33,7 @@ public:
 	 * @param  InCanaryObjectKey    Key for a canary object used to test whether this backend is usable
 	 * @param  InCachePath          Path to cache the DDC files
 	 */
-	FS3DerivedDataBackend(const TCHAR* InRootManifestPath, const TCHAR* InBaseUrl, const TCHAR* InRegion, const TCHAR* InCanaryObjectKey, const TCHAR* InCachePath);
+	FS3DerivedDataBackend(ICacheFactory& Factory, const TCHAR* InRootManifestPath, const TCHAR* InBaseUrl, const TCHAR* InRegion, const TCHAR* InCanaryObjectKey, const TCHAR* InCachePath);
 	~FS3DerivedDataBackend();
 
 	/**
@@ -58,6 +61,31 @@ public:
 
 	bool ApplyDebugOptions(FBackendDebugOptions& InOptions) override;
 
+	virtual FRequest Put(
+		TArrayView<FCacheRecord> Records,
+		FStringView Context,
+		ECachePolicy Policy,
+		EPriority Priority,
+		FOnCachePutComplete&& OnComplete) override;
+
+	virtual FRequest Get(
+		TConstArrayView<FCacheKey> Keys,
+		FStringView Context,
+		ECachePolicy Policy,
+		EPriority Priority,
+		FOnCacheGetComplete&& OnComplete) override;
+
+	virtual FRequest GetPayload(
+		TConstArrayView<FCachePayloadKey> Keys,
+		FStringView Context,
+		ECachePolicy Policy,
+		EPriority Priority,
+		FOnCacheGetPayloadComplete&& OnComplete) override;
+
+	virtual void CancelAll() override
+	{
+	}
+
 private:
 	struct FBundle;
 	struct FBundleEntry;
@@ -65,9 +93,10 @@ private:
 
 	struct FRootManifest;
 
-	class FRequest;
+	class FHttpRequest;
 	class FRequestPool;
 
+	ICacheFactory& Factory;
 	FString RootManifestPath;
 	FString BaseUrl;
 	FString Region;
@@ -95,4 +124,6 @@ private:
 	TSet<FName> DebugMissedKeys;
 };
 
+} // UE::DerivedData::Backends
+
 #endif
diff --git a/Engine/Source/Developer/DerivedDataCache/Public/DerivedDataBackendInterface.h b/Engine/Source/Developer/DerivedDataCache/Public/DerivedDataBackendInterface.h
--- a/Engine/Source/Developer/DerivedDataCache/Public/DerivedDataBackendInterface.h
+++ b/Engine/Source/Developer/DerivedDataCache/Public/DerivedDataBackendInterface.h
@@ -5,6 +5,7 @@
 
 #include "CoreMinimal.h"
 #include "Containers/BitArray.h"
+#include "DerivedDataCache.h"
 #include "Stats/Stats.h"
 
 class FDerivedDataCacheUsageStats;
@@ -28,7 +29,7 @@ DECLARE_FLOAT_ACCUMULATOR_STAT_EXTERN(TEXT("Exists Time"),STAT_DDC_ExistTime,STA
  * Interface for cache server backends. 
  * The entire API should be callable from any thread (except the singleton can be assumed to be called at least once before concurrent access).
 **/
-class FDerivedDataBackendInterface
+class FDerivedDataBackendInterface : public UE::DerivedData::ICacheStore
 {
 public:
 
