DDC: Implemented ICache on FDerivedDataCacheInterface

Moved IDerivedDataCacheModule to its own header and renamed CreateOrGetDDC to CreateOrGetCache.

#rnx

diff --git a/Engine/Source/Developer/DerivedDataCache/Private/DerivedDataCache.cpp b/Engine/Source/Developer/DerivedDataCache/Private/DerivedDataCache.cpp
--- a/Engine/Source/Developer/DerivedDataCache/Private/DerivedDataCache.cpp
+++ b/Engine/Source/Developer/DerivedDataCache/Private/DerivedDataCache.cpp
@@ -160,21 +160,17 @@ namespace DerivedDataCacheCookStats
 }
 #endif
 
-namespace UE::DerivedData
-{
-
-ICache* CreateCache();
-
-} // UE::DerivedData
-
 /** Whether we want to verify the DDC (pass in -VerifyDDC on the command line)*/
 bool GVerifyDDC = false;
 
+namespace UE::DerivedData
+{
+
 /**
  * Implementation of the derived data cache
  * This API is fully threadsafe
 **/
-class FDerivedDataCache : public FDerivedDataCacheInterface
+class FDerivedDataCache final : public FDerivedDataCacheInterface
 {
 
 	/** 
@@ -354,7 +350,6 @@ public:
 	/** Constructor, called once to cereate a singleton **/
 	FDerivedDataCache()
 		: CurrentHandle(19248) // we will skip some potential handles to catch errors
-		, Cache(UE::DerivedData::CreateCache())
 	{
 		FDerivedDataBackend::Get(); // we need to make sure this starts before we all us to start
 
@@ -376,11 +371,6 @@ public:
 		PendingTasks.Empty();
 	}
 
-	virtual UE::DerivedData::ICache& GetCache()
-	{
-		return *Cache;
-	}
-
 	virtual bool GetSynchronous(FDerivedDataPluginInterface* DataDeriver, TArray<uint8>& OutData, bool* bDataWasBuilt = nullptr) override
 	{
 		DDC_SCOPE_CYCLE_COUNTER(DDC_GetSynchronous);
@@ -678,63 +668,7 @@ private:
 	/** Cache notification delegate */
 	FOnDDCNotification DDCNotificationEvent;
 
-	TUniquePtr<UE::DerivedData::ICache> Cache;
-};
-
-static FDerivedDataCacheInterface* GDerivedDataCacheInstance;
-
-PRAGMA_DISABLE_DEPRECATION_WARNINGS
-
-/**
- * Module for the DDC
- */
-class FDerivedDataCacheModule final : public IDerivedDataCacheModule
-{
-public:
-	FDerivedDataCacheInterface& GetDDC() final
-	{
-		return **CreateOrGetDDC();
-	}
-
-	FDerivedDataCacheInterface* const* CreateOrGetDDC() final
-	{
-		FScopeLock Lock(&CreateLock);
-		if (!GDerivedDataCacheInstance)
-		{
-			GDerivedDataCacheInstance = new FDerivedDataCache();
-			check(GDerivedDataCacheInstance);
-		}
-		return &GDerivedDataCacheInstance;
-	}
-
-	void ShutdownModule() final
-	{
-		FDDCCleanup::Shutdown();
-
-		delete GDerivedDataCacheInstance;
-		GDerivedDataCacheInstance = nullptr;
-	}
-
-private:
-	FCriticalSection CreateLock;
-};
-
-PRAGMA_ENABLE_DEPRECATION_WARNINGS
-
-IMPLEMENT_MODULE(FDerivedDataCacheModule, DerivedDataCache);
-
-///////////////////////////////////////////////////////////////////////////////////////////////////////////////////////
-
-namespace UE::DerivedData
-{
-
-///////////////////////////////////////////////////////////////////////////////////////////////////////////////////////
-
-class FCache final : public ICache
-{
 public:
-	~FCache() final = default;
-
 	FCacheBucket CreateBucket(FStringView Name) final { return Private::CreateCacheBucket(Name); }
 
 	FCacheRecordBuilder CreateRecord(const FCacheKey& Key) final { return Private::CreateCacheRecordBuilder(Key); }
@@ -775,13 +709,46 @@ public:
 	}
 };
 
-///////////////////////////////////////////////////////////////////////////////////////////////////////////////////////
+} // UE::DerivedData
 
-ICache* CreateCache()
+static FDerivedDataCacheInterface* GDerivedDataCacheInstance;
+
+PRAGMA_DISABLE_DEPRECATION_WARNINGS
+
+/**
+ * Module for the DDC
+ */
+class FDerivedDataCacheModule final : public IDerivedDataCacheModule
 {
-	return new FCache();
-}
+public:
+	FDerivedDataCacheInterface& GetDDC() final
+	{
+		return **CreateOrGetCache();
+	}
 
-///////////////////////////////////////////////////////////////////////////////////////////////////////////////////////
+	FDerivedDataCacheInterface* const* CreateOrGetCache() final
+	{
+		FScopeLock Lock(&CreateLock);
+		if (!GDerivedDataCacheInstance)
+		{
+			GDerivedDataCacheInstance = new UE::DerivedData::FDerivedDataCache();
+			check(GDerivedDataCacheInstance);
+		}
+		return &GDerivedDataCacheInstance;
+	}
 
-} // UE::DerivedData
+	void ShutdownModule() final
+	{
+		FDDCCleanup::Shutdown();
+
+		delete GDerivedDataCacheInstance;
+		GDerivedDataCacheInstance = nullptr;
+	}
+
+private:
+	FCriticalSection CreateLock;
+};
+
+PRAGMA_ENABLE_DEPRECATION_WARNINGS
+
+IMPLEMENT_MODULE(FDerivedDataCacheModule, DerivedDataCache);
diff --git a/Engine/Source/Developer/DerivedDataCache/Public/DerivedDataCacheInterface.h b/Engine/Source/Developer/DerivedDataCache/Public/DerivedDataCacheInterface.h
--- a/Engine/Source/Developer/DerivedDataCache/Public/DerivedDataCacheInterface.h
+++ b/Engine/Source/Developer/DerivedDataCache/Public/DerivedDataCacheInterface.h
@@ -6,27 +6,23 @@
 #include "CoreMinimal.h"
 #include "Containers/BitArray.h"
 #include "Containers/StringView.h"
-#include "Modules/ModuleInterface.h"
+#include "DerivedDataCache.h"
+#include "DerivedDataCacheModule.h"
 
 class FDerivedDataCacheUsageStats;
 class FDerivedDataCacheStatsNode;
 class IDDCCleanup;
 
-namespace UE { namespace DerivedData { class ICache; } }
-
 /** 
  * Interface for the derived data cache
  * This API is fully threadsafe (with the possible exception of the system interface: NotfiyBootComplete, etc).
-**/
-class FDerivedDataCacheInterface
+ */
+class FDerivedDataCacheInterface : public UE::DerivedData::ICache
 {
 public:
-	virtual ~FDerivedDataCacheInterface()
-	{
-	}
+	virtual ~FDerivedDataCacheInterface() = default;
 
-	/** Returns the interface to the cache. */
-	virtual UE::DerivedData::ICache& GetCache() = 0;
+	using UE::DerivedData::ICache::Put;
 
 	//--------------------
 	// High Level Interface
@@ -237,7 +233,7 @@ public:
 	virtual void NotifyBootComplete() = 0;
 
 	/**
-	 * Adds or subtracts a number from the thread safe counter which tracks outstand async requests. This is used to ensure everything is complete prior to shutdown.
+	 * Adds or subtracts a number from the thread safe counter which tracks outstanding async requests. This is used to ensure everything is complete prior to shutdown.
 	 */
 	virtual void AddToAsyncCompletionCounter(int32 Addend) = 0;
 
@@ -309,23 +305,3 @@ public:
 	virtual FOnDDCNotification& GetDDCNotificationEvent() = 0;
 
 };
-
-/**
- * Module for the DDC
- */
-class IDerivedDataCacheModule : public IModuleInterface
-{
-public:
-	/** Return the DDC interface **/
-	UE_DEPRECATED(4.27, "GetDDC has been replaced by CreateOrGetDDC.")
-	virtual FDerivedDataCacheInterface& GetDDC() = 0;
-
-	/**
-	 * Returns the DDC interface, and creates it on the first call.
-	 *
-	 * This will never return a null interface, but an extra level of indirection is used to support
-	 * callers that store the return value. The returned interface pointer will become null when the
-	 * module shuts down, which destroys DDC.
-	 */
-	virtual FDerivedDataCacheInterface* const* CreateOrGetDDC() = 0;
-};
diff --git a/Engine/Source/Developer/DerivedDataCache/Public/DerivedDataCacheModule.h b/Engine/Source/Developer/DerivedDataCache/Public/DerivedDataCacheModule.h
new file mode 100644
--- /dev/null
+++ b/Engine/Source/Developer/DerivedDataCache/Public/DerivedDataCacheModule.h
@@ -0,0 +1,27 @@
+// Copyright Epic Games, Inc. All Rights Reserved.
+
+#pragma once
+
+#include "Modules/ModuleInterface.h"
+
+class FDerivedDataCacheInterface;
+
+/**
+ * Module for the Derived Data Cache and Derived Data Build.
+ */
+class IDerivedDataCacheModule : public IModuleInterface
+{
+public:
+	/** Return the DDC interface **/
+	UE_DEPRECATED(4.27, "GetDDC has been replaced by CreateOrGetCache.")
+	virtual FDerivedDataCacheInterface& GetDDC() = 0;
+
+	/**
+	 * Returns the cache, which is created by the first call to this function.
+	 *
+	 * This always returns a pointer to a valid cache, but that pointer becomes null when the module
+	 * shuts down and destroys the cache. This extra level of indirection allows a caller to observe
+	 * the destruction of the cache without polling this function or monitoring the module lifetime.
+	 */
+	virtual FDerivedDataCacheInterface* const* CreateOrGetCache() = 0;
+};
diff --git a/Engine/Source/Editor/UnrealEd/Private/EditorDomain/EditorDomain.cpp b/Engine/Source/Editor/UnrealEd/Private/EditorDomain/EditorDomain.cpp
--- a/Engine/Source/Editor/UnrealEd/Private/EditorDomain/EditorDomain.cpp
+++ b/Engine/Source/Editor/UnrealEd/Private/EditorDomain/EditorDomain.cpp
@@ -220,7 +220,7 @@ int64 FEditorDomain::FileSize(const FPackagePath& PackagePath, EPackageSegment P
 			return Workspace->FileSize(PackagePath, PackageSegment, OutUpdatedPath);
 		}
 
-		UE::DerivedData::ICache& Cache = GetDerivedDataCacheRef().GetCache();
+		UE::DerivedData::ICache& Cache = GetDerivedDataCacheRef();
 		auto MetaDataGetComplete =
 			[&FileSize, &PackageSource, &PackagePath, PackageSegment, Locks=this->Locks, OutUpdatedPath]
 			(UE::DerivedData::FCacheGetCompleteParams&& Params)
diff --git a/Engine/Source/Editor/UnrealEd/Private/EditorDomain/EditorDomainUtils.cpp b/Engine/Source/Editor/UnrealEd/Private/EditorDomain/EditorDomainUtils.cpp
--- a/Engine/Source/Editor/UnrealEd/Private/EditorDomain/EditorDomainUtils.cpp
+++ b/Engine/Source/Editor/UnrealEd/Private/EditorDomain/EditorDomainUtils.cpp
@@ -76,7 +76,7 @@ EPackageDigestResult GetPackageDigest(IAssetRegistry& AssetRegistry, FName Packa
 UE::DerivedData::FCacheKey GetEditorDomainPackageKey(const FPackageDigest& PackageDigest)
 {
 	static UE::DerivedData::FCacheBucket EditorDomainPackageCacheBucket =
-		GetDerivedDataCacheRef().GetCache().CreateBucket(EditorDomainPackageBucketName);
+		GetDerivedDataCacheRef().CreateBucket(EditorDomainPackageBucketName);
 	return UE::DerivedData::FCacheKey{EditorDomainPackageCacheBucket, PackageDigest};
 }
 
@@ -84,7 +84,7 @@ UE::DerivedData::FRequest RequestEditorDomainPackage(const FPackagePath& Package
 	const FPackageDigest& PackageDigest, UE::DerivedData::EPriority CachePriority,
 	UE::DerivedData::FOnCacheGetComplete&& Callback)
 {
-	UE::DerivedData::ICache& Cache = GetDerivedDataCacheRef().GetCache();
+	UE::DerivedData::ICache& Cache = GetDerivedDataCacheRef();
 	return Cache.Get({ GetEditorDomainPackageKey(PackageDigest) },
 		PackagePath.GetDebugName(),
 		UE::DerivedData::ECachePolicy::QueryLocal, CachePriority, MoveTemp(Callback));
@@ -157,7 +157,7 @@ bool TrySavePackage(UPackage* Package)
 		PackageBuffer = MakeSharedBufferFromArray(MoveTemp(TempBytes));
 	}
 
-	UE::DerivedData::ICache& Cache = GetDerivedDataCacheRef().GetCache();
+	UE::DerivedData::ICache& Cache = GetDerivedDataCacheRef();
 	UE::DerivedData::FCacheRecordBuilder RecordBuilder = Cache.CreateRecord(GetEditorDomainPackageKey(PackageDigest));
 	TCbWriter<256> MetaData;
 	MetaData.BeginObject();
diff --git a/Engine/Source/Runtime/Core/Private/Misc/CoreMisc.cpp b/Engine/Source/Runtime/Core/Private/Misc/CoreMisc.cpp
--- a/Engine/Source/Runtime/Core/Private/Misc/CoreMisc.cpp
+++ b/Engine/Source/Runtime/Core/Private/Misc/CoreMisc.cpp
@@ -18,7 +18,7 @@
 #include "Misc/ConfigCacheIni.h"
 
 #include "Modules/ModuleManager.h"
-#include "DerivedDataCacheInterface.h"
+#include "DerivedDataCacheModule.h"
 #include "Interfaces/ITargetPlatformManagerModule.h"
 
 DEFINE_LOG_CATEGORY(LogSHA);
@@ -101,7 +101,7 @@ FDerivedDataCacheInterface* GetDerivedDataCache()
 			bInitialized = true;
 			if (IDerivedDataCacheModule* Module = FModuleManager::LoadModulePtr<IDerivedDataCacheModule>("DerivedDataCache"))
 			{
-				DDC = Module->CreateOrGetDDC();
+				DDC = Module->CreateOrGetCache();
 			}
 		}
 	}
