DDC: Display pak file size when opening a pak cache

#rnx

diff --git a/Engine/Source/Developer/DerivedDataCache/Private/PakFileDerivedDataBackend.cpp b/Engine/Source/Developer/DerivedDataCache/Private/PakFileDerivedDataBackend.cpp
--- a/Engine/Source/Developer/DerivedDataCache/Private/PakFileDerivedDataBackend.cpp
+++ b/Engine/Source/Developer/DerivedDataCache/Private/PakFileDerivedDataBackend.cpp
@@ -42,7 +42,8 @@ FPakFileDerivedDataBackend::FPakFileDerivedDataBackend(ICacheFactory& InFactory,
 		}
 		else
 		{
-			UE_LOG(LogDerivedDataCache, Display, TEXT("Pak cache opened for reading %s."), InFilename);
+			UE_LOG(LogDerivedDataCache, Display, TEXT("Pak cache opened for reading %s (%" INT64_FMT " MiB)."),
+				InFilename, FileHandle->TotalSize() / 1024 / 1024);
 		}
 	}
 }
