DDC: Added the start of the new build interface UE::DerivedData::IBuild

This only includes the functionality required to create a build definition, and does not include the functionalty required to execute a build from the build definition.

#rnx

diff --git a/Engine/Source/Developer/DerivedDataCache/Private/DerivedDataBuildDefinition.cpp b/Engine/Source/Developer/DerivedDataCache/Private/DerivedDataBuildDefinition.cpp
new file mode 100644
--- /dev/null
+++ b/Engine/Source/Developer/DerivedDataCache/Private/DerivedDataBuildDefinition.cpp
@@ -0,0 +1,326 @@
+// Copyright Epic Games, Inc. All Rights Reserved.
+
+#include "DerivedDataBuildDefinition.h"
+
+#include "Algo/AllOf.h"
+#include "Containers/Map.h"
+#include "Containers/StringConv.h"
+#include "Containers/StringView.h"
+#include "Containers/UnrealString.h"
+#include "DerivedDataBuildKey.h"
+#include "DerivedDataBuildPrivate.h"
+#include "Misc/Guid.h"
+#include "Misc/TVariant.h"
+#include "Serialization/CompactBinary.h"
+#include "Serialization/CompactBinaryWriter.h"
+#include "Templates/Function.h"
+#include <atomic>
+
+namespace UE::DerivedData::Private
+{
+
+///////////////////////////////////////////////////////////////////////////////////////////////////////////////////////
+
+class FBuildDefinitionBuilderInternal final : public IBuildDefinitionBuilderInternal
+{
+public:
+	inline FBuildDefinitionBuilderInternal(FStringView InName, FStringView InFunction)
+		: Name(InName)
+		, Function(InFunction)
+	{
+		checkf(!Name.IsEmpty(), TEXT("A build definition requires a non-empty name."));
+		checkf(!Function.IsEmpty(), TEXT("A build definition requires a non-empty function name."));
+		checkf(Algo::AllOf(Function, FChar::IsAlnum), TEXT("A build definition requires an alphanumeric function name. ")
+			TEXT("Definition: %s; Function: %s"), *Name, *Function);
+	}
+
+	~FBuildDefinitionBuilderInternal() final = default;
+
+	void AddConstant(FStringView Key, const FCbObject& Value) final
+	{
+		Add<FCbObject>(Key, Value);
+	}
+
+	void AddInputBuild(FStringView Key, const FBuildPayloadKey& PayloadKey) final
+	{
+		Add<FBuildPayloadKey>(Key, PayloadKey);
+	}
+
+	void AddInputBulkData(FStringView Key, const FGuid& BulkDataId) final
+	{
+		Add<FGuid>(Key, BulkDataId);
+	}
+
+	void AddInputFile(FStringView Key, FStringView Path) final
+	{
+		Add<FString>(Key, Path);
+	}
+
+	void AddInputHash(FStringView Key, const FIoHash& RawHash) final
+	{
+		Add<FIoHash>(Key, RawHash);
+	}
+
+	IBuildDefinitionInternal* Build() final;
+
+	using InputType = TVariant<FCbObject, FBuildPayloadKey, FGuid, FString, FIoHash>;
+
+	FString Name;
+	FString Function;
+	TMap<FString, InputType> Inputs;
+
+private:
+	template <typename ValueType, typename ArgType>
+	inline void Add(FStringView Key, ArgType&& Value)
+	{
+		const uint32 KeyHash = GetTypeHash(Key);
+		checkf(!Key.IsEmpty(), TEXT("Empty key used in the %s build input for %s."), *Function, *Name);
+		checkf(!Inputs.ContainsByHash(KeyHash, Key), TEXT("Duplicate key %.*s used in the %s build input for %s."),
+			Key.Len(), Key.GetData(), *Function, *Name);
+		Inputs.EmplaceByHash(KeyHash, Key, InputType(TInPlaceType<ValueType>(), Forward<ArgType>(Value)));
+	}
+};
+
+///////////////////////////////////////////////////////////////////////////////////////////////////////////////////////
+
+class FBuildDefinitionInternal final : public IBuildDefinitionInternal
+{
+public:
+	explicit FBuildDefinitionInternal(FBuildDefinitionBuilderInternal&& DefinitionBuilder);
+	explicit FBuildDefinitionInternal(FStringView Name, FCbObject&& Definition);
+
+	~FBuildDefinitionInternal() final = default;
+
+	const FBuildKey& GetKey() const final { return Key; }
+
+	FStringView GetName() const final { return Name; }
+	FStringView GetFunction() const final { return Function; }
+
+	void IterateConstants(TFunctionRef<void (FStringView Key, FCbObject&& Value)> Visitor) const final;
+	void IterateInputBuilds(TFunctionRef<void (FStringView Key, const FBuildPayloadKey& PayloadKey)> Visitor) const final;
+	void IterateInputBulkData(TFunctionRef<void (FStringView Key, const FGuid& BulkDataId)> Visitor) const final;
+	void IterateInputFiles(TFunctionRef<void (FStringView Key, FStringView Path)> Visitor) const final;
+	void IterateInputHashes(TFunctionRef<void (FStringView Key, const FIoHash& RawHash)> Visitor) const final;
+
+	void Save(FCbWriter& Writer) const final;
+
+	inline void AddRef() const final
+	{
+		ReferenceCount.fetch_add(1, std::memory_order_relaxed);
+	}
+
+	inline void Release() const final
+	{
+		if (ReferenceCount.fetch_sub(1, std::memory_order_acq_rel) == 1)
+		{
+			delete this;
+		}
+	}
+
+private:
+	mutable std::atomic<uint32> ReferenceCount{0};
+	FString Name;
+	FString Function;
+	FCbObject Definition;
+	FBuildKey Key;
+};
+
+///////////////////////////////////////////////////////////////////////////////////////////////////////////////////////
+
+FBuildDefinitionInternal::FBuildDefinitionInternal(FBuildDefinitionBuilderInternal&& DefinitionBuilder)
+	: Name(MoveTemp(DefinitionBuilder.Name))
+	, Function(MoveTemp(DefinitionBuilder.Function))
+{
+	DefinitionBuilder.Inputs.KeySort(TLess<>());
+
+	bool bHasConstants = false;
+	bool bHasBuilds = false;
+	bool bHasBulkData = false;
+	bool bHasFiles = false;
+	bool bHasHashes = false;
+
+	for (const TPair<FString, FBuildDefinitionBuilderInternal::InputType>& Pair : DefinitionBuilder.Inputs)
+	{
+		if (Pair.Value.IsType<FCbObject>())
+		{
+			bHasConstants = true;
+		}
+		else if (Pair.Value.IsType<FBuildPayloadKey>())
+		{
+			bHasBuilds = true;
+		}
+		else if (Pair.Value.IsType<FGuid>())
+		{
+			bHasBulkData = true;
+		}
+		else if (Pair.Value.IsType<FString>())
+		{
+			bHasFiles = true;
+		}
+		else if (Pair.Value.IsType<FIoHash>())
+		{
+			bHasHashes = true;
+		}
+	}
+
+	TCbWriter<2048> Writer;
+	Writer.BeginObject();
+	Writer.AddString("Function"_ASV, Function);
+
+	if (bHasConstants)
+	{
+		Writer.BeginObject("Constants"_ASV);
+		for (const TPair<FString, FBuildDefinitionBuilderInternal::InputType>& Pair : DefinitionBuilder.Inputs)
+		{
+			if (Pair.Value.IsType<FCbObject>())
+			{
+				Writer.AddObject(FTCHARToUTF8(Pair.Key), Pair.Value.Get<FCbObject>());
+			}
+		}
+		Writer.EndObject();
+	}
+
+	if (bHasBuilds)
+	{
+		Writer.BeginObject("Builds"_ASV);
+		for (const TPair<FString, FBuildDefinitionBuilderInternal::InputType>& Pair : DefinitionBuilder.Inputs)
+		{
+			if (Pair.Value.IsType<FBuildPayloadKey>())
+			{
+				const FBuildPayloadKey& PayloadKey = Pair.Value.Get<FBuildPayloadKey>();
+				Writer.BeginObject(FTCHARToUTF8(Pair.Key));
+				Writer.AddHash("Build"_ASV, PayloadKey.BuildKey.Hash);
+				Writer.AddObjectId("Payload"_ASV, FCbObjectId(PayloadKey.Id.GetView()));
+				Writer.EndObject();
+			}
+		}
+		Writer.EndObject();
+	}
+
+	if (bHasBulkData)
+	{
+		Writer.BeginObject("BulkData"_ASV);
+		for (const TPair<FString, FBuildDefinitionBuilderInternal::InputType>& Pair : DefinitionBuilder.Inputs)
+		{
+			if (Pair.Value.IsType<FGuid>())
+			{
+				Writer.AddUuid(FTCHARToUTF8(Pair.Key), Pair.Value.Get<FGuid>());
+			}
+		}
+		Writer.EndObject();
+	}
+
+	if (bHasFiles)
+	{
+		Writer.BeginObject("Files"_ASV);
+		for (const TPair<FString, FBuildDefinitionBuilderInternal::InputType>& Pair : DefinitionBuilder.Inputs)
+		{
+			if (Pair.Value.IsType<FString>())
+			{
+				Writer.AddString(FTCHARToUTF8(Pair.Key), Pair.Value.Get<FString>());
+			}
+		}
+		Writer.EndObject();
+	}
+
+	if (bHasHashes)
+	{
+		Writer.BeginObject("Hashes"_ASV);
+		for (const TPair<FString, FBuildDefinitionBuilderInternal::InputType>& Pair : DefinitionBuilder.Inputs)
+		{
+			if (Pair.Value.IsType<FIoHash>())
+			{
+				Writer.AddBinaryAttachment(FTCHARToUTF8(Pair.Key), Pair.Value.Get<FIoHash>());
+			}
+		}
+		Writer.EndObject();
+	}
+
+	Writer.EndObject();
+	Definition = Writer.Save().AsObject();
+	Key = FBuildKey{Definition.GetHash()};
+}
+
+FBuildDefinitionInternal::FBuildDefinitionInternal(FStringView InName, FCbObject&& InDefinition)
+	: Name(InName)
+	, Function(InDefinition.FindView("Function"_ASV).AsString())
+	, Definition(MoveTemp(InDefinition))
+{
+	checkf(!InName.IsEmpty(), TEXT("A build definition requires a non-empty name."));
+	Definition.MakeOwned();
+	if (const bool bIsValid = Definition && Algo::AllOf(Function, FChar::IsAlnum))
+	{
+		Key = FBuildKey{Definition.GetHash()};
+	}
+}
+
+void FBuildDefinitionInternal::IterateConstants(TFunctionRef<void (FStringView Key, FCbObject&& Value)> Visitor) const
+{
+	for (FCbField Field : Definition["Constants"_ASV])
+	{
+		Visitor(FUTF8ToTCHAR(Field.GetName()), Field.AsObject());
+	}
+}
+
+void FBuildDefinitionInternal::IterateInputBuilds(TFunctionRef<void (FStringView Key, const FBuildPayloadKey& PayloadKey)> Visitor) const
+{
+	for (FCbFieldView Field : Definition.FindView("Builds"_ASV))
+	{
+		FCbObjectView Build = Field.AsObjectView();
+		const FBuildKey BuildKey{Build["Build"_ASV].AsHash()};
+		const FPayloadId Id = FPayloadId(Build["Payload"_ASV].AsObjectId().GetView());
+		Visitor(FUTF8ToTCHAR(Field.GetName()), FBuildPayloadKey{BuildKey, Id});
+	}
+}
+
+void FBuildDefinitionInternal::IterateInputBulkData(TFunctionRef<void (FStringView Key, const FGuid& BulkDataId)> Visitor) const
+{
+	for (FCbFieldView Field : Definition.FindView("BulkData"_ASV))
+	{
+		Visitor(FUTF8ToTCHAR(Field.GetName()), Field.AsUuid());
+	}
+}
+
+void FBuildDefinitionInternal::IterateInputFiles(TFunctionRef<void (FStringView Key, FStringView Path)> Visitor) const
+{
+	for (FCbFieldView Field : Definition.FindView("Files"_ASV))
+	{
+		Visitor(FUTF8ToTCHAR(Field.GetName()), FUTF8ToTCHAR(Field.AsString()));
+	}
+}
+
+void FBuildDefinitionInternal::IterateInputHashes(TFunctionRef<void (FStringView Key, const FIoHash& RawHash)> Visitor) const
+{
+	for (FCbFieldView Field : Definition.FindView("Hashes"_ASV))
+	{
+		Visitor(FUTF8ToTCHAR(Field.GetName()), Field.AsBinaryAttachment());
+	}
+}
+
+void FBuildDefinitionInternal::Save(FCbWriter& Writer) const
+{
+	Writer.AddObject(Definition);
+}
+
+///////////////////////////////////////////////////////////////////////////////////////////////////////////////////////
+
+IBuildDefinitionInternal* FBuildDefinitionBuilderInternal::Build()
+{
+	return new FBuildDefinitionInternal(MoveTemp(*this));
+}
+
+///////////////////////////////////////////////////////////////////////////////////////////////////////////////////////
+
+FBuildDefinitionBuilder CreateBuildDefinition(FStringView Name, FStringView Function)
+{
+	return CreateBuildDefinitionBuilder(new FBuildDefinitionBuilderInternal(Name, Function));
+}
+
+FBuildDefinition LoadBuildDefinition(FStringView Name, FCbObject&& Definition)
+{
+	return CreateBuildDefinition(new FBuildDefinitionInternal(Name, MoveTemp(Definition)));
+}
+
+///////////////////////////////////////////////////////////////////////////////////////////////////////////////////////
+
+} // UE::DerivedData::Private
diff --git a/Engine/Source/Developer/DerivedDataCache/Private/DerivedDataBuildPrivate.h b/Engine/Source/Developer/DerivedDataCache/Private/DerivedDataBuildPrivate.h
new file mode 100644
--- /dev/null
+++ b/Engine/Source/Developer/DerivedDataCache/Private/DerivedDataBuildPrivate.h
@@ -0,0 +1,21 @@
+// Copyright Epic Games, Inc. All Rights Reserved.
+
+#pragma once
+
+#include "Containers/StringFwd.h"
+
+class FCbObject;
+
+namespace UE::DerivedData { class FBuildDefinition; }
+namespace UE::DerivedData { class FBuildDefinitionBuilder; }
+
+namespace UE::DerivedData::Private
+{
+
+// Implemented in DerivedDataBuildDefinition.cpp
+FBuildDefinitionBuilder CreateBuildDefinition(FStringView Name, FStringView Function);
+
+// Implemented in DerivedDataBuildDefinition.cpp
+FBuildDefinition LoadBuildDefinition(FStringView Name, FCbObject&& Definition);
+
+} // UE::DerivedData::Private
diff --git a/Engine/Source/Developer/DerivedDataCache/Private/DerivedDataCache.cpp b/Engine/Source/Developer/DerivedDataCache/Private/DerivedDataCache.cpp
--- a/Engine/Source/Developer/DerivedDataCache/Private/DerivedDataCache.cpp
+++ b/Engine/Source/Developer/DerivedDataCache/Private/DerivedDataCache.cpp
@@ -17,6 +17,9 @@
 #include "Modules/ModuleManager.h"
 
 #include "DerivedDataBackendInterface.h"
+#include "DerivedDataBuild.h"
+#include "DerivedDataBuildDefinition.h"
+#include "DerivedDataBuildPrivate.h"
 #include "DerivedDataCachePrivate.h"
 #include "DerivedDataCacheRecord.h"
 #include "DerivedDataPluginInterface.h"
@@ -170,7 +173,7 @@ namespace UE::DerivedData
  * Implementation of the derived data cache
  * This API is fully threadsafe
 **/
-class FDerivedDataCache final : public FDerivedDataCacheInterface
+class FDerivedDataCache final : public FDerivedDataCacheInterface, public IBuild
 {
 
 	/** 
@@ -669,6 +672,8 @@ private:
 	FOnDDCNotification DDCNotificationEvent;
 
 public:
+	// ICache Interface
+
 	FCacheBucket CreateBucket(FStringView Name) final { return Private::CreateCacheBucket(Name); }
 
 	FCacheRecordBuilder CreateRecord(const FCacheKey& Key) final { return Private::CreateCacheRecordBuilder(Key); }
@@ -707,11 +712,25 @@ public:
 	{
 		return FDerivedDataBackend::Get().GetRoot().CancelAll();
 	}
+
+public:
+	// IBuild Interface
+
+	FBuildDefinitionBuilder CreateDefinition(FStringView Name, FStringView Function) final
+	{
+		return Private::CreateBuildDefinition(Name, Function);
+	}
+
+	FBuildDefinition LoadDefinition(FStringView Name, FCbObject&& Definition) final
+	{
+		return Private::LoadBuildDefinition(Name, MoveTemp(Definition));
+	}
 };
 
 } // UE::DerivedData
 
 static FDerivedDataCacheInterface* GDerivedDataCacheInstance;
+static UE::DerivedData::IBuild* GDerivedDataBuildInstance;
 
 PRAGMA_DISABLE_DEPRECATION_WARNINGS
 
@@ -731,12 +750,27 @@ public:
 		FScopeLock Lock(&CreateLock);
 		if (!GDerivedDataCacheInstance)
 		{
-			GDerivedDataCacheInstance = new UE::DerivedData::FDerivedDataCache();
-			check(GDerivedDataCacheInstance);
+			UE::DerivedData::FDerivedDataCache* Instance = new UE::DerivedData::FDerivedDataCache();
+			GDerivedDataCacheInstance = Instance;
+			GDerivedDataBuildInstance = Instance;
+			check(Instance);
 		}
 		return &GDerivedDataCacheInstance;
 	}
 
+	UE::DerivedData::IBuild* const* CreateOrGetBuild() final
+	{
+		FScopeLock Lock(&CreateLock);
+		if (!GDerivedDataBuildInstance)
+		{
+			UE::DerivedData::FDerivedDataCache* Instance = new UE::DerivedData::FDerivedDataCache();
+			GDerivedDataCacheInstance = Instance;
+			GDerivedDataBuildInstance = Instance;
+			check(Instance);
+		}
+		return &GDerivedDataBuildInstance;
+	}
+
 	void ShutdownModule() final
 	{
 		FDDCCleanup::Shutdown();
diff --git a/Engine/Source/Developer/DerivedDataCache/Public/DerivedDataBuild.h b/Engine/Source/Developer/DerivedDataCache/Public/DerivedDataBuild.h
new file mode 100644
--- /dev/null
+++ b/Engine/Source/Developer/DerivedDataCache/Public/DerivedDataBuild.h
@@ -0,0 +1,43 @@
+// Copyright Epic Games, Inc. All Rights Reserved.
+
+#pragma once
+
+#include "Containers/StringFwd.h"
+
+class FCbObject;
+
+namespace UE::DerivedData { class FBuildDefinition; }
+namespace UE::DerivedData { class FBuildDefinitionBuilder; }
+
+namespace UE::DerivedData
+{
+
+/**
+ * Interface to the build system.
+ *
+ * This is only a preview of a portion of the interface and does not support build execution.
+ */
+class IBuild
+{
+public:
+	virtual ~IBuild() = default;
+
+	/**
+	 * Create a build definition builder.
+	 *
+	 * @param Name       The name by which to identify this definition for logging and profiling.
+	 * @param Function   The name of the build function with which to build this definition.
+	 */
+	virtual FBuildDefinitionBuilder CreateDefinition(FStringView Name, FStringView Function) = 0;
+
+	/**
+	 * Load a build definition from compact binary.
+	 *
+	 * @param Name         The name by which to identify this definition for logging and profiling.
+	 * @param Definition   An object saved from a build definition. Holds a reference and is cloned if not owned.
+	 * @return A valid build definition, or a build definition with an empty key on error.
+	 */
+	virtual FBuildDefinition LoadDefinition(FStringView Name, FCbObject&& Definition) = 0;
+};
+
+} // UE::DerivedData
diff --git a/Engine/Source/Developer/DerivedDataCache/Public/DerivedDataBuildDefinition.h b/Engine/Source/Developer/DerivedDataCache/Public/DerivedDataBuildDefinition.h
new file mode 100644
--- /dev/null
+++ b/Engine/Source/Developer/DerivedDataCache/Public/DerivedDataBuildDefinition.h
@@ -0,0 +1,256 @@
+// Copyright Epic Games, Inc. All Rights Reserved.
+
+#pragma once
+
+#include "CoreTypes.h"
+#include "Containers/StringView.h"
+#include "Templates/RefCounting.h"
+#include "Templates/UniquePtr.h"
+
+class FCbObject;
+class FCbWriter;
+struct FGuid;
+struct FIoHash;
+template <typename FuncType> class TFunctionRef;
+
+namespace UE::DerivedData { class FBuildDefinition; }
+namespace UE::DerivedData { class FBuildDefinitionBuilder; }
+namespace UE::DerivedData { struct FBuildKey; }
+namespace UE::DerivedData { struct FBuildPayloadKey; }
+
+///////////////////////////////////////////////////////////////////////////////////////////////////////////////////////
+
+namespace UE::DerivedData::Private
+{
+
+class IBuildDefinitionInternal
+{
+public:
+	virtual ~IBuildDefinitionInternal() = default;
+	virtual const FBuildKey& GetKey() const = 0;
+	virtual FStringView GetName() const = 0;
+	virtual FStringView GetFunction() const = 0;
+	virtual void IterateConstants(TFunctionRef<void (FStringView Key, FCbObject&& Value)> Visitor) const = 0;
+	virtual void IterateInputBuilds(TFunctionRef<void (FStringView Key, const FBuildPayloadKey& PayloadKey)> Visitor) const = 0;
+	virtual void IterateInputBulkData(TFunctionRef<void (FStringView Key, const FGuid& BulkDataId)> Visitor) const = 0;
+	virtual void IterateInputFiles(TFunctionRef<void (FStringView Key, FStringView Path)> Visitor) const = 0;
+	virtual void IterateInputHashes(TFunctionRef<void (FStringView Key, const FIoHash& RawHash)> Visitor) const = 0;
+	virtual void Save(FCbWriter& Writer) const = 0;
+	virtual void AddRef() const = 0;
+	virtual void Release() const = 0;
+};
+
+FBuildDefinition CreateBuildDefinition(IBuildDefinitionInternal* Definition);
+
+class IBuildDefinitionBuilderInternal
+{
+public:
+	virtual ~IBuildDefinitionBuilderInternal() = default;
+	virtual void AddConstant(FStringView Key, const FCbObject& Value) = 0;
+	virtual void AddInputBuild(FStringView Key, const FBuildPayloadKey& PayloadKey) = 0;
+	virtual void AddInputBulkData(FStringView Key, const FGuid& BulkDataId) = 0;
+	virtual void AddInputFile(FStringView Key, FStringView Path) = 0;
+	virtual void AddInputHash(FStringView Key, const FIoHash& RawHash) = 0;
+	virtual IBuildDefinitionInternal* Build() = 0;
+};
+
+FBuildDefinitionBuilder CreateBuildDefinitionBuilder(IBuildDefinitionBuilderInternal* DefinitionBuilder);
+
+} // UE::DerivedData::Private
+
+///////////////////////////////////////////////////////////////////////////////////////////////////////////////////////
+
+namespace UE::DerivedData
+{
+
+/**
+ * A build definition references a build function and the inputs to that build function.
+ *
+ * The purpose of a definition is to capture everything required to execute a derived data build.
+ * Inputs to the build are partly static (constants and input hashes) and partly variable (files,
+ * bulk data, payloads from other builds); and the build function is referenced by name, but does
+ * not correspond to a fixed version of the build function.
+ *
+ * The key for the definition uniquely identifies the definition and is derived from its function
+ * and inputs, such that any definition with the same function and inputs will have the same key.
+ *
+ * The keys for constants and inputs are names that are unique within the build definition across
+ * every type of constant and input.
+ *
+ * To build a definition against a specific version of the function and inputs, queue it to build
+ * on FBuildSession, which uses IBuildInputProvider to fetch the inputs at the expected versions.
+ *
+ * Build definitions are immutable, and are created by a builder from IBuild::CreateDefinition().
+ *
+ * @see FBuildDefinitionBuilder
+ * @see FBuildSession
+ */
+class FBuildDefinition
+{
+public:
+	/** Returns the key that uniquely identifies this build definition. */
+	inline const FBuildKey& GetKey() const
+	{
+		return Definition->GetKey();
+	}
+
+	/** Returns the name by which to identify this definition for logging and profiling. */
+	inline FStringView GetName() const
+	{
+		return Definition->GetName();
+	}
+
+	/** Returns the name of the build function with which to build this definition. */
+	inline FStringView GetFunction() const
+	{
+		return Definition->GetFunction();
+	}
+
+	/** Visits every constant in order by key. */
+	inline void IterateConstants(TFunctionRef<void (FStringView Key, FCbObject&& Value)> Visitor) const
+	{
+		return Definition->IterateConstants(Visitor);
+	}
+
+	/** Visits every input build payload in order by key. */
+	inline void IterateInputBuilds(TFunctionRef<void (FStringView Key, const FBuildPayloadKey& PayloadKey)> Visitor) const
+	{
+		return Definition->IterateInputBuilds(Visitor);
+	}
+
+	/** Visits every input bulk data in order by key. */
+	inline void IterateInputBulkData(TFunctionRef<void (FStringView Key, const FGuid& BulkDataId)> Visitor) const
+	{
+		return Definition->IterateInputBulkData(Visitor);
+	}
+
+	/** Visits every input file in order by key. */
+	inline void IterateInputFiles(TFunctionRef<void (FStringView Key, FStringView Path)> Visitor) const
+	{
+		return Definition->IterateInputFiles(Visitor);
+	}
+
+	/** Visits every input hash in order by key. */
+	inline void IterateInputHashes(TFunctionRef<void (FStringView Key, const FIoHash& RawHash)> Visitor) const
+	{
+		return Definition->IterateInputHashes(Visitor);
+	}
+
+	/** Saves the build definition to a compact binary object. Calls BeginObject and EndObject. */
+	inline void Save(FCbWriter& Writer) const
+	{
+		return Definition->Save(Writer);
+	}
+
+private:
+	friend FBuildDefinition Private::CreateBuildDefinition(Private::IBuildDefinitionInternal* Definition);
+
+	/** Construct a build definition. Use Build() on a builder from IBuild::CreateDefinition(). */
+	inline explicit FBuildDefinition(Private::IBuildDefinitionInternal* InDefinition)
+		: Definition(InDefinition)
+	{
+	}
+
+	TRefCountPtr<Private::IBuildDefinitionInternal> Definition;
+};
+
+/**
+ * A build definition builder is used to construct a build definition.
+ *
+ * Create using IBuild::CreateDefinition() which must be given a build function name.
+ *
+ * @see FBuildDefinition
+ */
+class FBuildDefinitionBuilder
+{
+public:
+	/** Add a constant object with a key that is unique within this definition. */
+	inline void AddConstant(FStringView Key, const FCbObject& Value)
+	{
+		DefinitionBuilder->AddConstant(Key, Value);
+	}
+
+	/** Add a payload from another build with a key that is unique within this definition. */
+	inline void AddInputBuild(FStringView Key, const FBuildPayloadKey& PayloadKey)
+	{
+		DefinitionBuilder->AddInputBuild(Key, PayloadKey);
+	}
+
+	/**
+	 * Add a bulk data input with a key that is unique within this definition.
+	 *
+	 * @param BulkDataId   Identifier that uniquely identifies this data in the IBuildInputProvider.
+	 */
+	inline void AddInputBulkData(FStringView Key, const FGuid& BulkDataId)
+	{
+		DefinitionBuilder->AddInputBulkData(Key, BulkDataId);
+	}
+
+	/**
+	 * Add a file input with a key that is unique within this definition.
+	 *
+	 * @param Path   Path to the file relative to a mounted content root.
+	 */
+	inline void AddInputFile(FStringView Key, FStringView Path)
+	{
+		DefinitionBuilder->AddInputFile(Key, Path);
+	}
+
+	/**
+	 * Add a hash input with a key that is unique within this definition.
+	 *
+	 * @param RawHash   Hash of the raw data that will resolve it in the IBuildInputProvider.
+	 */
+	inline void AddInputHash(FStringView Key, const FIoHash& RawHash)
+	{
+		DefinitionBuilder->AddInputHash(Key, RawHash);
+	}
+
+	/**
+	 * Build a build definition, which makes this builder subsequently unusable.
+	 */
+	inline FBuildDefinition Build()
+	{
+		return CreateBuildDefinition(DefinitionBuilder->Build());
+	}
+
+private:
+	friend FBuildDefinitionBuilder Private::CreateBuildDefinitionBuilder(Private::IBuildDefinitionBuilderInternal* DefinitionBuilder);
+
+	/** Construct a build definition builder. Use IBuild::CreateDefinition(). */
+	inline explicit FBuildDefinitionBuilder(Private::IBuildDefinitionBuilderInternal* InDefinitionBuilder)
+		: DefinitionBuilder(InDefinitionBuilder)
+	{
+	}
+
+	TUniquePtr<Private::IBuildDefinitionBuilderInternal> DefinitionBuilder;
+};
+
+/**
+ * A build definition provider finds a definition from its key.
+ */
+class IBuildDefinitionProvider
+{
+public:
+	/** Returns a definition matching the key, or a definition with an empty key on error. */
+	virtual FBuildDefinition GetDefinition(const FBuildKey& Key) = 0;
+};
+
+} // UE::DerivedData
+
+///////////////////////////////////////////////////////////////////////////////////////////////////////////////////////
+
+namespace UE::DerivedData::Private
+{
+
+inline FBuildDefinition CreateBuildDefinition(IBuildDefinitionInternal* Definition)
+{
+	return FBuildDefinition(Definition);
+}
+
+inline FBuildDefinitionBuilder CreateBuildDefinitionBuilder(IBuildDefinitionBuilderInternal* DefinitionBuilder)
+{
+	return FBuildDefinitionBuilder(DefinitionBuilder);
+}
+
+} // UE::DerivedData::Private
diff --git a/Engine/Source/Developer/DerivedDataCache/Public/DerivedDataBuildKey.h b/Engine/Source/Developer/DerivedDataCache/Public/DerivedDataBuildKey.h
new file mode 100644
--- /dev/null
+++ b/Engine/Source/Developer/DerivedDataCache/Public/DerivedDataBuildKey.h
@@ -0,0 +1,99 @@
+// Copyright Epic Games, Inc. All Rights Reserved.
+
+#pragma once
+
+#include "CoreTypes.h"
+#include "Containers/StringFwd.h"
+#include "DerivedDataPayload.h"
+#include "IO/IoHash.h"
+#include "Templates/TypeHash.h"
+
+namespace UE::DerivedData
+{
+
+///////////////////////////////////////////////////////////////////////////////////////////////////////////////////////
+
+/** A key that uniquely identifies a build definition. */
+struct FBuildKey
+{
+	FIoHash Hash;
+
+	/** A key with a zero hash. */
+	static const FBuildKey Empty;
+};
+
+inline const FBuildKey FBuildKey::Empty;
+
+/** A key that uniquely identifies a payload within a build output. */
+struct FBuildPayloadKey
+{
+	FBuildKey BuildKey;
+	FPayloadId Id;
+
+	/** A payload key with an empty build key and a null payload identifier. */
+	static const FBuildPayloadKey Empty;
+};
+
+inline const FBuildPayloadKey FBuildPayloadKey::Empty;
+
+///////////////////////////////////////////////////////////////////////////////////////////////////////////////////////
+
+inline bool operator==(const FBuildKey& A, const FBuildKey& B)
+{
+	return A.Hash == B.Hash;
+}
+
+inline bool operator!=(const FBuildKey& A, const FBuildKey& B)
+{
+	return A.Hash != B.Hash;
+}
+
+inline bool operator<(const FBuildKey& A, const FBuildKey& B)
+{
+	return A.Hash < B.Hash;
+}
+
+inline uint32 GetTypeHash(const FBuildKey& Key)
+{
+	return GetTypeHash(Key.Hash);
+}
+
+template <typename CharType>
+inline TStringBuilderBase<CharType>& operator<<(TStringBuilderBase<CharType>& Builder, const FBuildKey& Key)
+{
+	return Builder << "Build/"_ASV << Key.Hash;
+}
+
+///////////////////////////////////////////////////////////////////////////////////////////////////////////////////////
+
+inline bool operator==(const FBuildPayloadKey& A, const FBuildPayloadKey& B)
+{
+	return A.BuildKey == B.BuildKey && A.Id == B.Id;
+}
+
+inline bool operator!=(const FBuildPayloadKey& A, const FBuildPayloadKey& B)
+{
+	return A.BuildKey != B.BuildKey || A.Id != B.Id;
+}
+
+inline bool operator<(const FBuildPayloadKey& A, const FBuildPayloadKey& B)
+{
+	const FBuildKey& KeyA = A.BuildKey;
+	const FBuildKey& KeyB = B.BuildKey;
+	return KeyA == KeyB ? A.Id < B.Id : KeyA < KeyB;
+}
+
+inline uint32 GetTypeHash(const FBuildPayloadKey& Key)
+{
+	return HashCombine(GetTypeHash(Key.BuildKey), GetTypeHash(Key.Id));
+}
+
+template <typename CharType>
+inline TStringBuilderBase<CharType>& operator<<(TStringBuilderBase<CharType>& Builder, const FBuildPayloadKey& Key)
+{
+	return Builder << Key.BuildKey << CharType('/') << Key.Id;
+}
+
+///////////////////////////////////////////////////////////////////////////////////////////////////////////////////////
+
+} // UE::DerivedData
diff --git a/Engine/Source/Developer/DerivedDataCache/Public/DerivedDataCacheModule.h b/Engine/Source/Developer/DerivedDataCache/Public/DerivedDataCacheModule.h
--- a/Engine/Source/Developer/DerivedDataCache/Public/DerivedDataCacheModule.h
+++ b/Engine/Source/Developer/DerivedDataCache/Public/DerivedDataCacheModule.h
@@ -6,6 +6,8 @@
 
 class FDerivedDataCacheInterface;
 
+namespace UE::DerivedData { class IBuild; }
+
 /**
  * Module for the Derived Data Cache and Derived Data Build.
  */
@@ -24,4 +26,13 @@ public:
 	 * the destruction of the cache without polling this function or monitoring the module lifetime.
 	 */
 	virtual FDerivedDataCacheInterface* const* CreateOrGetCache() = 0;
+
+	/**
+	 * Returns the build system, which is created by the first call to this function.
+	 *
+	 * This always returns a pointer to a valid build, but that pointer becomes null when the module
+	 * shuts down and destroys the build. This extra level of indirection allows a caller to observe
+	 * the destruction of the build without polling this function or monitoring the module lifetime.
+	 */
+	virtual UE::DerivedData::IBuild* const* CreateOrGetBuild() = 0;
 };
diff --git a/Engine/Source/Runtime/Core/Private/Misc/CoreMisc.cpp b/Engine/Source/Runtime/Core/Private/Misc/CoreMisc.cpp
--- a/Engine/Source/Runtime/Core/Private/Misc/CoreMisc.cpp
+++ b/Engine/Source/Runtime/Core/Private/Misc/CoreMisc.cpp
@@ -91,7 +91,7 @@ bool FStaticSelfRegisteringExec::Exec( UWorld* InWorld, const TCHAR* Cmd, FOutpu
 
 FDerivedDataCacheInterface* GetDerivedDataCache()
 {
-	static FDerivedDataCacheInterface* const* DDC;
+	static FDerivedDataCacheInterface* const* Cache;
 	static bool bInitialized = false;
 	if (!bInitialized)
 	{
@@ -101,22 +101,52 @@ FDerivedDataCacheInterface* GetDerivedDataCache()
 			bInitialized = true;
 			if (IDerivedDataCacheModule* Module = FModuleManager::LoadModulePtr<IDerivedDataCacheModule>("DerivedDataCache"))
 			{
-				DDC = Module->CreateOrGetCache();
+				Cache = Module->CreateOrGetCache();
 			}
 		}
 	}
-	return DDC ? *DDC : nullptr;
+	return Cache ? *Cache : nullptr;
 }
 
 FDerivedDataCacheInterface& GetDerivedDataCacheRef()
 {
-	FDerivedDataCacheInterface* DDC = GetDerivedDataCache();
-	if (!DDC)
+	FDerivedDataCacheInterface* Cache = GetDerivedDataCache();
+	if (!Cache)
 	{
 		UE_LOG(LogInit, Fatal, TEXT("Derived Data Cache was requested, but not available."));
-		CA_ASSUME(DDC); // Suppress static analysis warning in unreachable code (fatal error)
+		CA_ASSUME(Cache); // Suppress static analysis warning in unreachable code (fatal error)
 	}
-	return *DDC;
+	return *Cache;
+}
+
+UE::DerivedData::IBuild* GetDerivedDataBuild()
+{
+	static UE::DerivedData::IBuild* const* Build;
+	static bool bInitialized = false;
+	if (!bInitialized)
+	{
+		if (!FPlatformProperties::RequiresCookedData())
+		{
+			check(IsInGameThread());
+			bInitialized = true;
+			if (IDerivedDataCacheModule* Module = FModuleManager::LoadModulePtr<IDerivedDataCacheModule>("DerivedDataCache"))
+			{
+				Build = Module->CreateOrGetBuild();
+			}
+		}
+	}
+	return Build ? *Build : nullptr;
+}
+
+UE::DerivedData::IBuild& GetDerivedDataBuildRef()
+{
+	UE::DerivedData::IBuild* Build = GetDerivedDataBuild();
+	if (!Build)
+	{
+		UE_LOG(LogInit, Fatal, TEXT("Derived Data Build was requested, but not available."));
+		CA_ASSUME(Build); // Suppress static analysis warning in unreachable code (fatal error)
+	}
+	return *Build;
 }
 
 class ITargetPlatformManagerModule* GetTargetPlatformManager(bool bFailOnInitErrors)
diff --git a/Engine/Source/Runtime/Core/Public/Misc/CoreMisc.h b/Engine/Source/Runtime/Core/Public/Misc/CoreMisc.h
--- a/Engine/Source/Runtime/Core/Public/Misc/CoreMisc.h
+++ b/Engine/Source/Runtime/Core/Public/Misc/CoreMisc.h
@@ -13,6 +13,8 @@
 #include "CoreGlobals.h"
 #include "HAL/ThreadSingleton.h"
 
+namespace UE::DerivedData { class IBuild; }
+
 /**
  * Exec handler that registers itself and is being routed via StaticExec.
  * Note: Not intended for use with UObjects!
@@ -65,12 +67,18 @@ struct CORE_API FMaintenance
 	Module singletons.
 -----------------------------------------------------------------------------*/
 
-/** Return the DDC interface, if it is available, otherwise return NULL **/
+/** Returns the derived data cache interface if it is available, otherwise null. */
 CORE_API class FDerivedDataCacheInterface* GetDerivedDataCache();
 
-/** Return the DDC interface, fatal error if it is not available. **/
+/** Returns the derived data cache interface, or fatal error if it is not available. */
 CORE_API class FDerivedDataCacheInterface& GetDerivedDataCacheRef();
 
+/** Return the derived data build interface if it is available, otherwise null. */
+CORE_API UE::DerivedData::IBuild* GetDerivedDataBuild();
+
+/** Return the derived data build interface, or fatal error if it is not available. */
+CORE_API UE::DerivedData::IBuild& GetDerivedDataBuildRef();
+
 /**
  * Return the Target Platform Manager interface, if it is available, otherwise return nullptr.
  *
