DDC: Moved Private::CreateBuildDefinition[Builder] out of the header to avoid them being used externally

#rnx

diff --git a/Engine/Source/Developer/DerivedDataCache/Private/DerivedDataBuildDefinition.cpp b/Engine/Source/Developer/DerivedDataCache/Private/DerivedDataBuildDefinition.cpp
--- a/Engine/Source/Developer/DerivedDataCache/Private/DerivedDataBuildDefinition.cpp
+++ b/Engine/Source/Developer/DerivedDataCache/Private/DerivedDataBuildDefinition.cpp
@@ -61,7 +61,7 @@ public:
 		Add<FIoHash>(Key, RawHash);
 	}
 
-	IBuildDefinitionInternal* Build() final;
+	FBuildDefinition Build() final;
 
 	using InputType = TVariant<FCbObject, FBuildPayloadKey, FGuid, FString, FIoHash>;
 
@@ -304,9 +304,21 @@ void FBuildDefinitionInternal::Save(FCbWriter& Writer) const
 
 ///////////////////////////////////////////////////////////////////////////////////////////////////////////////////////
 
-IBuildDefinitionInternal* FBuildDefinitionBuilderInternal::Build()
+FBuildDefinition FBuildDefinitionBuilderInternal::Build()
 {
-	return new FBuildDefinitionInternal(MoveTemp(*this));
+	return CreateBuildDefinition(new FBuildDefinitionInternal(MoveTemp(*this)));
+}
+
+///////////////////////////////////////////////////////////////////////////////////////////////////////////////////////
+
+FBuildDefinition CreateBuildDefinition(IBuildDefinitionInternal* Definition)
+{
+	return FBuildDefinition(Definition);
+}
+
+FBuildDefinitionBuilder CreateBuildDefinitionBuilder(IBuildDefinitionBuilderInternal* DefinitionBuilder)
+{
+	return FBuildDefinitionBuilder(DefinitionBuilder);
 }
 
 ///////////////////////////////////////////////////////////////////////////////////////////////////////////////////////
diff --git a/Engine/Source/Developer/DerivedDataCache/Public/DerivedDataBuildDefinition.h b/Engine/Source/Developer/DerivedDataCache/Public/DerivedDataBuildDefinition.h
--- a/Engine/Source/Developer/DerivedDataCache/Public/DerivedDataBuildDefinition.h
+++ b/Engine/Source/Developer/DerivedDataCache/Public/DerivedDataBuildDefinition.h
@@ -51,7 +51,7 @@ public:
 	virtual void AddInputBulkData(FStringView Key, const FGuid& BulkDataId) = 0;
 	virtual void AddInputFile(FStringView Key, FStringView Path) = 0;
 	virtual void AddInputHash(FStringView Key, const FIoHash& RawHash) = 0;
-	virtual IBuildDefinitionInternal* Build() = 0;
+	virtual FBuildDefinition Build() = 0;
 };
 
 FBuildDefinitionBuilder CreateBuildDefinitionBuilder(IBuildDefinitionBuilderInternal* DefinitionBuilder);
@@ -211,7 +211,7 @@ public:
 	 */
 	inline FBuildDefinition Build()
 	{
-		return CreateBuildDefinition(DefinitionBuilder->Build());
+		return DefinitionBuilder->Build();
 	}
 
 private:
@@ -237,20 +237,3 @@ public:
 };
 
 } // UE::DerivedData
-
-///////////////////////////////////////////////////////////////////////////////////////////////////////////////////////
-
-namespace UE::DerivedData::Private
-{
-
-inline FBuildDefinition CreateBuildDefinition(IBuildDefinitionInternal* Definition)
-{
-	return FBuildDefinition(Definition);
-}
-
-inline FBuildDefinitionBuilder CreateBuildDefinitionBuilder(IBuildDefinitionBuilderInternal* DefinitionBuilder)
-{
-	return FBuildDefinitionBuilder(DefinitionBuilder);
-}
-
-} // UE::DerivedData::Private
