DDC: Added FPayload::Null and FPayloadId::Null

#rnx

diff --git a/Engine/Source/Developer/DerivedDataCache/Public/DerivedDataPayload.h b/Engine/Source/Developer/DerivedDataCache/Public/DerivedDataPayload.h
--- a/Engine/Source/Developer/DerivedDataCache/Public/DerivedDataPayload.h
+++ b/Engine/Source/Developer/DerivedDataCache/Public/DerivedDataPayload.h
@@ -49,10 +49,15 @@ public:
 	/** Reset this to null. */
 	inline void Reset() { *this = FPayloadId(); }
 
+	/** A null ID. */
+	static const FPayloadId Null;
+
 private:
 	alignas(uint32) ByteArray Bytes{};
 };
 
+inline const FPayloadId FPayloadId::Null;
+
 ///////////////////////////////////////////////////////////////////////////////////////////////////////////////////////
 
 /**
@@ -96,12 +101,17 @@ public:
 	/** Reset this to null. */
 	inline void Reset() { *this = FPayload(); }
 
+	/** A null payload. */
+	static const FPayload Null;
+
 private:
 	FPayloadId Id;
 	FIoHash RawHash;
 	FCompressedBuffer Data;
 };
 
+inline const FPayload FPayload::Null;
+
 ///////////////////////////////////////////////////////////////////////////////////////////////////////////////////////
 
 inline FPayloadId::FPayloadId(const FMemoryView InId)
