DDC: Changed FCacheRecord to be reference-counted instead of uniquely owned

#rnx

diff --git a/Engine/Source/Developer/DerivedDataCache/Private/DerivedDataCacheRecord.cpp b/Engine/Source/Developer/DerivedDataCache/Private/DerivedDataCacheRecord.cpp
--- a/Engine/Source/Developer/DerivedDataCache/Private/DerivedDataCacheRecord.cpp
+++ b/Engine/Source/Developer/DerivedDataCache/Private/DerivedDataCacheRecord.cpp
@@ -4,9 +4,12 @@
 
 #include "Containers/UnrealString.h"
 #include "DerivedDataCacheKey.h"
+#include "HAL/CriticalSection.h"
+#include "Misc/ScopeRWLock.h"
 #include "Misc/StringBuilder.h"
 #include "Serialization/CompactBinary.h"
 #include "UObject/NameTypes.h"
+#include <atomic>
 
 namespace UE::DerivedData::Private
 {
@@ -46,8 +49,6 @@ public:
 
 	~FCacheRecordInternal() final = default;
 
-	FCacheRecord Clone() const final;
-
 	const FCacheKey& GetKey() const final;
 	const FCbObject& GetMeta() const final;
 
@@ -60,22 +61,32 @@ public:
 
 	const FPayload& GetPayload(const FPayloadId& Id) const final;
 
+	inline void AddRef() const final
+	{
+		ReferenceCount.fetch_add(1, std::memory_order_relaxed);
+	}
+
+	inline void Release() const final
+	{
+		if (ReferenceCount.fetch_sub(1, std::memory_order_acq_rel) == 1)
+		{
+			delete this;
+		}
+	}
+
+private:
+	mutable std::atomic<uint32> ReferenceCount{0};
 	FCacheKey Key;
 	FCbObject Meta;
 	FPayload Value;
 	TArray<FPayload> Attachments;
+	mutable FRWLock CacheLock;
 	mutable FSharedBuffer ValueCache;
 	mutable TArray<FSharedBuffer> AttachmentsCache;
 };
 
 ///////////////////////////////////////////////////////////////////////////////////////////////////////////////////////
 
-static const FPayload& GetEmptyCachePayload()
-{
-	static const FPayload Empty;
-	return Empty;
-}
-
 static FPayloadId GetOrCreatePayloadId(const FPayloadId& Id, const FIoHash& RawHash)
 {
 	return Id.IsValid() ? Id : FPayloadId::FromHash(RawHash);
@@ -91,11 +102,6 @@ FCacheRecordInternal::FCacheRecordInternal(FCacheRecordBuilderInternal&& RecordB
 {
 }
 
-FCacheRecord FCacheRecordInternal::Clone() const
-{
-	return CreateCacheRecord(new FCacheRecordInternal(*this));
-}
-
 const FCacheKey& FCacheRecordInternal::GetKey() const
 {
 	return Key;
@@ -108,10 +114,20 @@ const FCbObject& FCacheRecordInternal::GetMeta() const
 
 FSharedBuffer FCacheRecordInternal::GetValue() const
 {
-	if (ValueCache.IsNull() && Value.IsValid())
+	if (Value.IsNull())
+	{
+		return FSharedBuffer();
+	}
+	if (FReadScopeLock Lock(CacheLock); ValueCache)
+	{
+		return ValueCache;
+	}
+	FWriteScopeLock Lock(CacheLock);
+	if (ValueCache)
 	{
-		ValueCache = Value.GetData().Decompress();
+		return ValueCache;
 	}
+	ValueCache = Value.GetData().Decompress();
 	return ValueCache;
 }
 
@@ -125,16 +141,21 @@ FSharedBuffer FCacheRecordInternal::GetAttachment(const FPayloadId& Id) const
 	const int32 Index = Algo::LowerBound(Attachments, Id, FPayloadLessById());
 	if (Attachments.IsValidIndex(Index) && FPayloadEqualById()(Attachments[Index], Id))
 	{
+		if (FReadScopeLock Lock(CacheLock); AttachmentsCache.IsValidIndex(Index) && AttachmentsCache[Index])
+		{
+			return AttachmentsCache[Index];
+		}
+		FWriteScopeLock Lock(CacheLock);
 		if (AttachmentsCache.IsEmpty())
 		{
 			AttachmentsCache.SetNum(Attachments.Num());
 		}
-		FSharedBuffer& DataCache = AttachmentsCache[Index];
-		if (!DataCache)
+		FSharedBuffer& Cache = AttachmentsCache[Index];
+		if (Cache.IsNull())
 		{
-			DataCache = Attachments[Index].GetData().Decompress();
+			Cache = Attachments[Index].GetData().Decompress();
 		}
-		return DataCache;
+		return Cache;
 	}
 	return FSharedBuffer();
 }
@@ -146,7 +167,7 @@ const FPayload& FCacheRecordInternal::GetAttachmentPayload(const FPayloadId& Id)
 	{
 		return Attachments[Index];
 	}
-	return GetEmptyCachePayload();
+	return FPayload::Null;
 }
 
 TConstArrayView<FPayload> FCacheRecordInternal::GetAttachmentPayloads() const
diff --git a/Engine/Source/Developer/DerivedDataCache/Public/DerivedDataCacheRecord.h b/Engine/Source/Developer/DerivedDataCache/Public/DerivedDataCacheRecord.h
--- a/Engine/Source/Developer/DerivedDataCache/Public/DerivedDataCacheRecord.h
+++ b/Engine/Source/Developer/DerivedDataCache/Public/DerivedDataCacheRecord.h
@@ -5,7 +5,8 @@
 #include "DerivedDataCacheKey.h"
 #include "DerivedDataPayload.h"
 #include "DerivedDataRequest.h"
-#include "Memory/MemoryFwd.h"
+#include "Memory/SharedBuffer.h"
+#include "Templates/RefCounting.h"
 #include "Templates/UniquePtr.h"
 
 class FCbObject;
@@ -24,7 +25,6 @@ class ICacheRecordInternal
 {
 public:
 	virtual ~ICacheRecordInternal() = default;
-	virtual FCacheRecord Clone() const = 0;
 	virtual const FCacheKey& GetKey() const = 0;
 	virtual const FCbObject& GetMeta() const = 0;
 	virtual FSharedBuffer GetValue() const = 0;
@@ -33,6 +33,8 @@ public:
 	virtual const FPayload& GetAttachmentPayload(const FPayloadId& Id) const = 0;
 	virtual TConstArrayView<FPayload> GetAttachmentPayloads() const = 0;
 	virtual const FPayload& GetPayload(const FPayloadId& Id) const = 0;
+	virtual void AddRef() const = 0;
+	virtual void Release() const = 0;
 };
 
 FCacheRecord CreateCacheRecord(ICacheRecordInternal* Record);
@@ -81,7 +83,7 @@ public:
 	explicit FCacheRecord() = default;
 
 	/** Clone this cache record. Prefer to move records, and clone only when a copy is necessary. */
-	inline FCacheRecord Clone() const { return Record->Clone(); }
+	inline FCacheRecord Clone() const { return *this; }
 
 	/** Returns the key that identifies this record in the cache. Always available in a non-null cache record. */
 	inline const FCacheKey& GetKey() const { return Record->GetKey(); }
@@ -120,20 +122,20 @@ public:
 private:
 	friend FCacheRecord Private::CreateCacheRecord(Private::ICacheRecordInternal* Record);
 
-	/** Construct a cache record. Use Record.Clone() or Builder.Build[Async](). */
+	/** Construct a cache record. Use Build() or BuildAsync() on a builder from ICache::CreateRecord(). */
 	inline explicit FCacheRecord(Private::ICacheRecordInternal* InRecord)
 		: Record(InRecord)
 	{
 	}
 
-	TUniquePtr<Private::ICacheRecordInternal> Record;
+	TRefCountPtr<Private::ICacheRecordInternal> Record;
 };
 
 /**
  * A cache record builder is used to construct a cache record.
  *
- * Create using ICache::CreateRecord, which must be given a key uniquely corresponds to the value
- * and attachments for the cache record. Optional metadata may vary and may be non-deterministic.
+ * Create using ICache::CreateRecord() which must be given a key that uniquely corresponds to the
+ * value and attachments for the cache record. Metadata may vary between records of the same key.
  *
  * The value and attachments can be provided as buffers, which will be compressed, or as payloads
  * which were previously compressed and have an identifier assigned.
