DDC: Changed Put to take a const array of cache records now that records are copyable

#rnx

diff --git a/Engine/Source/Developer/DerivedDataCache/Private/DerivedDataBackendAsyncPutWrapper.cpp b/Engine/Source/Developer/DerivedDataCache/Private/DerivedDataBackendAsyncPutWrapper.cpp
--- a/Engine/Source/Developer/DerivedDataCache/Private/DerivedDataBackendAsyncPutWrapper.cpp
+++ b/Engine/Source/Developer/DerivedDataCache/Private/DerivedDataBackendAsyncPutWrapper.cpp
@@ -458,7 +458,7 @@ private:
 };
 
 FRequest FDerivedDataBackendAsyncPutWrapper::Put(
-	TArrayView<FCacheRecord> Records,
+	TConstArrayView<FCacheRecord> Records,
 	FStringView Context,
 	ECachePolicy Policy,
 	EPriority Priority,
@@ -470,14 +470,8 @@ FRequest FDerivedDataBackendAsyncPutWrapper::Put(
 	}
 	else
 	{
-		TArray<FCacheRecord> RecordsArray;
-		RecordsArray.Reserve(Records.Num());
-		for (FCacheRecord& Record : Records)
-		{
-			RecordsArray.Add(MoveTemp(Record));
-		}
 		TRequest<FDerivedDataAsyncWrapperRequest> Request(new FDerivedDataAsyncWrapperRequest(
-			[Backend = InnerBackend, Records = MoveTemp(RecordsArray), Context = FString(Context), Policy, Priority, OnComplete = MoveTemp(OnComplete)](bool bCancel) mutable
+			[Backend = InnerBackend, Records = TArray<FCacheRecord>(Records), Context = FString(Context), Policy, Priority, OnComplete = MoveTemp(OnComplete)](bool bCancel) mutable
 			{
 				if (!bCancel)
 				{
diff --git a/Engine/Source/Developer/DerivedDataCache/Private/DerivedDataBackendAsyncPutWrapper.h b/Engine/Source/Developer/DerivedDataCache/Private/DerivedDataBackendAsyncPutWrapper.h
--- a/Engine/Source/Developer/DerivedDataCache/Private/DerivedDataBackendAsyncPutWrapper.h
+++ b/Engine/Source/Developer/DerivedDataCache/Private/DerivedDataBackendAsyncPutWrapper.h
@@ -133,7 +133,7 @@ public:
 	virtual bool ApplyDebugOptions(FBackendDebugOptions& InOptions) override;
 
 	virtual FRequest Put(
-		TArrayView<FCacheRecord> Records,
+		TConstArrayView<FCacheRecord> Records,
 		FStringView Context,
 		ECachePolicy Policy,
 		EPriority Priority,
diff --git a/Engine/Source/Developer/DerivedDataCache/Private/DerivedDataBackendCorruptionWrapper.h b/Engine/Source/Developer/DerivedDataCache/Private/DerivedDataBackendCorruptionWrapper.h
--- a/Engine/Source/Developer/DerivedDataCache/Private/DerivedDataBackendCorruptionWrapper.h
+++ b/Engine/Source/Developer/DerivedDataCache/Private/DerivedDataBackendCorruptionWrapper.h
@@ -252,7 +252,7 @@ public:
 	}
 
 	virtual FRequest Put(
-		TArrayView<FCacheRecord> Records,
+		TConstArrayView<FCacheRecord> Records,
 		FStringView Context,
 		ECachePolicy Policy,
 		EPriority Priority,
diff --git a/Engine/Source/Developer/DerivedDataCache/Private/DerivedDataBackendVerifyWrapper.h b/Engine/Source/Developer/DerivedDataCache/Private/DerivedDataBackendVerifyWrapper.h
--- a/Engine/Source/Developer/DerivedDataCache/Private/DerivedDataBackendVerifyWrapper.h
+++ b/Engine/Source/Developer/DerivedDataCache/Private/DerivedDataBackendVerifyWrapper.h
@@ -163,7 +163,7 @@ public:
 	}
 
 	virtual FRequest Put(
-		TArrayView<FCacheRecord> Records,
+		TConstArrayView<FCacheRecord> Records,
 		FStringView Context,
 		ECachePolicy Policy,
 		EPriority Priority,
diff --git a/Engine/Source/Developer/DerivedDataCache/Private/DerivedDataCache.cpp b/Engine/Source/Developer/DerivedDataCache/Private/DerivedDataCache.cpp
--- a/Engine/Source/Developer/DerivedDataCache/Private/DerivedDataCache.cpp
+++ b/Engine/Source/Developer/DerivedDataCache/Private/DerivedDataCache.cpp
@@ -679,7 +679,7 @@ public:
 	FCacheRecordBuilder CreateRecord(const FCacheKey& Key) final { return Private::CreateCacheRecordBuilder(Key); }
 
 	FRequest Put(
-		TArrayView<FCacheRecord> Records,
+		TConstArrayView<FCacheRecord> Records,
 		FStringView Context,
 		ECachePolicy Policy,
 		EPriority Priority,
diff --git a/Engine/Source/Developer/DerivedDataCache/Private/DerivedDataLimitKeyLengthWrapper.h b/Engine/Source/Developer/DerivedDataCache/Private/DerivedDataLimitKeyLengthWrapper.h
--- a/Engine/Source/Developer/DerivedDataCache/Private/DerivedDataLimitKeyLengthWrapper.h
+++ b/Engine/Source/Developer/DerivedDataCache/Private/DerivedDataLimitKeyLengthWrapper.h
@@ -243,7 +243,7 @@ public:
 	}
 
 	virtual FRequest Put(
-		TArrayView<FCacheRecord> Records,
+		TConstArrayView<FCacheRecord> Records,
 		FStringView Context,
 		ECachePolicy Policy,
 		EPriority Priority,
diff --git a/Engine/Source/Developer/DerivedDataCache/Private/FileSystemDerivedDataBackend.cpp b/Engine/Source/Developer/DerivedDataCache/Private/FileSystemDerivedDataBackend.cpp
--- a/Engine/Source/Developer/DerivedDataCache/Private/FileSystemDerivedDataBackend.cpp
+++ b/Engine/Source/Developer/DerivedDataCache/Private/FileSystemDerivedDataBackend.cpp
@@ -859,7 +859,7 @@ public:
 	}
 
 	virtual FRequest Put(
-		TArrayView<FCacheRecord> Records,
+		TConstArrayView<FCacheRecord> Records,
 		FStringView Context,
 		ECachePolicy Policy,
 		EPriority Priority,
diff --git a/Engine/Source/Developer/DerivedDataCache/Private/HierarchicalDerivedDataBackend.h b/Engine/Source/Developer/DerivedDataCache/Private/HierarchicalDerivedDataBackend.h
--- a/Engine/Source/Developer/DerivedDataCache/Private/HierarchicalDerivedDataBackend.h
+++ b/Engine/Source/Developer/DerivedDataCache/Private/HierarchicalDerivedDataBackend.h
@@ -489,7 +489,7 @@ public:
 	}
 
 	virtual FRequest Put(
-		TArrayView<FCacheRecord> Records,
+		TConstArrayView<FCacheRecord> Records,
 		FStringView Context,
 		ECachePolicy Policy,
 		EPriority Priority,
@@ -580,8 +580,7 @@ public:
 									const bool bIsLocalFill = InnerBackends[FillCacheIndex]->GetSpeedClass() == ESpeedClass::Local;
 									if ((bIsLocalFill && bStoreLocal && bIsLocalGet <= bStoreLocalCopy) || (!bIsLocalFill && bStoreRemote))
 									{
-										FCacheRecord Record = Params.Record.Clone();
-										AsyncPutInnerBackends[FillCacheIndex]->Put(MakeArrayView(&Record, 1), Context);
+										AsyncPutInnerBackends[FillCacheIndex]->Put({Params.Record}, Context);
 									}
 								}
 							}
diff --git a/Engine/Source/Developer/DerivedDataCache/Private/HttpDerivedDataBackend.cpp b/Engine/Source/Developer/DerivedDataCache/Private/HttpDerivedDataBackend.cpp
--- a/Engine/Source/Developer/DerivedDataCache/Private/HttpDerivedDataBackend.cpp
+++ b/Engine/Source/Developer/DerivedDataCache/Private/HttpDerivedDataBackend.cpp
@@ -2272,7 +2272,7 @@ TSharedRef<FDerivedDataCacheStatsNode> FHttpDerivedDataBackend::GatherUsageStats
 }
 
 FRequest FHttpDerivedDataBackend::Put(
-	TArrayView<FCacheRecord> Records,
+	TConstArrayView<FCacheRecord> Records,
 	FStringView Context,
 	ECachePolicy Policy,
 	EPriority Priority,
diff --git a/Engine/Source/Developer/DerivedDataCache/Private/HttpDerivedDataBackend.h b/Engine/Source/Developer/DerivedDataCache/Private/HttpDerivedDataBackend.h
--- a/Engine/Source/Developer/DerivedDataCache/Private/HttpDerivedDataBackend.h
+++ b/Engine/Source/Developer/DerivedDataCache/Private/HttpDerivedDataBackend.h
@@ -75,7 +75,7 @@ public:
 	void SetSpeedClass(ESpeedClass InSpeedClass) { SpeedClass = InSpeedClass; }
 
 	virtual FRequest Put(
-		TArrayView<FCacheRecord> Records,
+		TConstArrayView<FCacheRecord> Records,
 		FStringView Context,
 		ECachePolicy Policy,
 		EPriority Priority,
diff --git a/Engine/Source/Developer/DerivedDataCache/Private/MemoryDerivedDataBackend.cpp b/Engine/Source/Developer/DerivedDataCache/Private/MemoryDerivedDataBackend.cpp
--- a/Engine/Source/Developer/DerivedDataCache/Private/MemoryDerivedDataBackend.cpp
+++ b/Engine/Source/Developer/DerivedDataCache/Private/MemoryDerivedDataBackend.cpp
@@ -406,13 +406,13 @@ int64 FMemoryDerivedDataBackend::CalcSerializedCacheRecordSize(const FCacheRecor
 }
 
 FRequest FMemoryDerivedDataBackend::Put(
-	TArrayView<FCacheRecord> Records,
+	TConstArrayView<FCacheRecord> Records,
 	FStringView Context,
 	ECachePolicy Policy,
 	EPriority Priority,
 	FOnCachePutComplete&& OnComplete)
 {
-	for (FCacheRecord& Record : Records)
+	for (const FCacheRecord& Record : Records)
 	{
 		const FCacheKey& Key = Record.GetKey();
 		EStatus Status = EStatus::Error;
@@ -457,7 +457,7 @@ FRequest FMemoryDerivedDataBackend::Put(
 			}
 
 			CurrentCacheSize += RecordSize;
-			CacheRecords.Add(MoveTemp(Record));
+			CacheRecords.Add(Record);
 			COOK_STAT(Timer.AddHit(RecordSize));
 			Status = EStatus::Ok;
 		}
@@ -478,7 +478,7 @@ FRequest FMemoryDerivedDataBackend::Get(
 		FCacheRecord Record;
 		if (FScopeLock ScopeLock(&SynchronizationObject); const FCacheRecord* CacheRecord = CacheRecords.Find(Key))
 		{
-			Record = CacheRecord->Clone();
+			Record = *CacheRecord;
 		}
 		if (Record)
 		{
diff --git a/Engine/Source/Developer/DerivedDataCache/Private/MemoryDerivedDataBackend.h b/Engine/Source/Developer/DerivedDataCache/Private/MemoryDerivedDataBackend.h
--- a/Engine/Source/Developer/DerivedDataCache/Private/MemoryDerivedDataBackend.h
+++ b/Engine/Source/Developer/DerivedDataCache/Private/MemoryDerivedDataBackend.h
@@ -95,7 +95,7 @@ public:
 	bool ApplyDebugOptions(FBackendDebugOptions& InOptions) override;
 
 	virtual FRequest Put(
-		TArrayView<FCacheRecord> Records,
+		TConstArrayView<FCacheRecord> Records,
 		FStringView Context,
 		ECachePolicy Policy,
 		EPriority Priority,
diff --git a/Engine/Source/Developer/DerivedDataCache/Private/PakFileDerivedDataBackend.cpp b/Engine/Source/Developer/DerivedDataCache/Private/PakFileDerivedDataBackend.cpp
--- a/Engine/Source/Developer/DerivedDataCache/Private/PakFileDerivedDataBackend.cpp
+++ b/Engine/Source/Developer/DerivedDataCache/Private/PakFileDerivedDataBackend.cpp
@@ -442,7 +442,7 @@ TSharedRef<FDerivedDataCacheStatsNode> FPakFileDerivedDataBackend::GatherUsageSt
 }
 
 FRequest FPakFileDerivedDataBackend::Put(
-	TArrayView<FCacheRecord> Records,
+	TConstArrayView<FCacheRecord> Records,
 	FStringView Context,
 	ECachePolicy Policy,
 	EPriority Priority,
diff --git a/Engine/Source/Developer/DerivedDataCache/Private/PakFileDerivedDataBackend.h b/Engine/Source/Developer/DerivedDataCache/Private/PakFileDerivedDataBackend.h
--- a/Engine/Source/Developer/DerivedDataCache/Private/PakFileDerivedDataBackend.h
+++ b/Engine/Source/Developer/DerivedDataCache/Private/PakFileDerivedDataBackend.h
@@ -102,7 +102,7 @@ public:
 	virtual bool ApplyDebugOptions(FBackendDebugOptions& InOptions) override { return false; }
 
 	virtual FRequest Put(
-		TArrayView<FCacheRecord> Records,
+		TConstArrayView<FCacheRecord> Records,
 		FStringView Context,
 		ECachePolicy Policy,
 		EPriority Priority,
diff --git a/Engine/Source/Developer/DerivedDataCache/Private/S3DerivedDataBackend.cpp b/Engine/Source/Developer/DerivedDataCache/Private/S3DerivedDataBackend.cpp
--- a/Engine/Source/Developer/DerivedDataCache/Private/S3DerivedDataBackend.cpp
+++ b/Engine/Source/Developer/DerivedDataCache/Private/S3DerivedDataBackend.cpp
@@ -1178,7 +1178,7 @@ bool FS3DerivedDataBackend::ShouldSimulateMiss(const TCHAR* InKey)
 }
 
 FRequest FS3DerivedDataBackend::Put(
-	TArrayView<FCacheRecord> Records,
+	TConstArrayView<FCacheRecord> Records,
 	FStringView Context,
 	ECachePolicy Policy,
 	EPriority Priority,
diff --git a/Engine/Source/Developer/DerivedDataCache/Private/S3DerivedDataBackend.h b/Engine/Source/Developer/DerivedDataCache/Private/S3DerivedDataBackend.h
--- a/Engine/Source/Developer/DerivedDataCache/Private/S3DerivedDataBackend.h
+++ b/Engine/Source/Developer/DerivedDataCache/Private/S3DerivedDataBackend.h
@@ -62,7 +62,7 @@ public:
 	bool ApplyDebugOptions(FBackendDebugOptions& InOptions) override;
 
 	virtual FRequest Put(
-		TArrayView<FCacheRecord> Records,
+		TConstArrayView<FCacheRecord> Records,
 		FStringView Context,
 		ECachePolicy Policy,
 		EPriority Priority,
diff --git a/Engine/Source/Developer/DerivedDataCache/Public/DerivedDataCache.h b/Engine/Source/Developer/DerivedDataCache/Public/DerivedDataCache.h
--- a/Engine/Source/Developer/DerivedDataCache/Public/DerivedDataCache.h
+++ b/Engine/Source/Developer/DerivedDataCache/Public/DerivedDataCache.h
@@ -131,14 +131,14 @@ public:
 	 * The callback will always be called for every key, and may be called from an arbitrary thread.
 	 * Records may finish storing in any order, and from multiple threads concurrently.
 	 *
-	 * @param Records The cache records to store. Must have a key. Records are reset to null.
+	 * @param Records The cache records to store. Must have a key.
 	 * @param Context A description of the request. An object path is typically sufficient.
 	 * @param Policy Flags to control the behavior of the request. See ECachePolicy.
 	 * @param Priority A priority to consider when scheduling the request. See EPriority.
 	 * @param OnComplete A callback invoked for every key in the batch as it completes or is canceled.
 	 */
 	virtual FRequest Put(
-		TArrayView<FCacheRecord> Records,
+		TConstArrayView<FCacheRecord> Records,
 		FStringView Context,
 		ECachePolicy Policy = ECachePolicy::Default,
 		EPriority Priority = EPriority::Normal,
@@ -225,7 +225,7 @@ struct FCacheGetCompleteParams
 	/**
 	 * Record for the part of the get request that completed or was canceled.
 	 *
-	 * The key is always populated. The remainder of the record is populated when Status is Cached.
+	 * The key is always populated. The remainder of the record is populated when Status is Ok.
 	 *
 	 * The value, attachments, and metadata may be skipped based on cache policy flags. When a value
 	 * or attachment has been skipped, it will have a payload but its buffers will be null.
@@ -245,7 +245,7 @@ struct FCacheGetPayloadCompleteParams
 	/**
 	 * Payload for the part of the payload request that completed or was canceled.
 	 *
-	 * The ID is always populated. The hash and buffer are populated when Status is Cached.
+	 * The ID is always populated. The hash and buffer are populated when Status is Ok.
 	 */
 	FPayload&& Payload;
 
diff --git a/Engine/Source/Developer/DerivedDataCache/Public/DerivedDataCacheRecord.h b/Engine/Source/Developer/DerivedDataCache/Public/DerivedDataCacheRecord.h
--- a/Engine/Source/Developer/DerivedDataCache/Public/DerivedDataCacheRecord.h
+++ b/Engine/Source/Developer/DerivedDataCache/Public/DerivedDataCacheRecord.h
@@ -82,9 +82,6 @@ public:
 	/** Construct a null cache record. */
 	explicit FCacheRecord() = default;
 
-	/** Clone this cache record. Prefer to move records, and clone only when a copy is necessary. */
-	inline FCacheRecord Clone() const { return *this; }
-
 	/** Returns the key that identifies this record in the cache. Always available in a non-null cache record. */
 	inline const FCacheKey& GetKey() const { return Record->GetKey(); }
 
