DDC: Added raw size to the payload

#rnx

diff --git a/Engine/Source/Developer/DerivedDataCache/Private/FileSystemDerivedDataBackend.cpp b/Engine/Source/Developer/DerivedDataCache/Private/FileSystemDerivedDataBackend.cpp
--- a/Engine/Source/Developer/DerivedDataCache/Private/FileSystemDerivedDataBackend.cpp
+++ b/Engine/Source/Developer/DerivedDataCache/Private/FileSystemDerivedDataBackend.cpp
@@ -1172,6 +1172,7 @@ private:
 
 		Writer.BeginObject();
 		Writer.AddObjectId("Id"_ASV, FCbObjectId(Payload.GetId().GetView()));
+		Writer.AddInteger("RawSize"_ASV, CompressedBuffer.GetRawSize());
 		if (bStoreInline)
 		{
 			Writer.AddHash("RawHash"_ASV, RawHash);
@@ -1189,18 +1190,19 @@ private:
 	FPayload GetCachePayload(const FCacheKey& Key, FStringView Context, ECachePolicy Policy, const FCbObject& Object, bool bAlwaysLoadInlineData = false) const
 	{
 		const FPayloadId Id(Object.FindView("Id"_ASV).AsObjectId().GetView());
+		const uint64 RawSize = Object.FindView("RawSize"_ASV).AsUInt64(MAX_uint64);
 		const FIoHash RawHash = Object.FindView("RawHash"_ASV).AsHash();
 		FIoHash CompressedHash = Object.FindView("CompressedHash"_ASV).AsHash();
 		FSharedBuffer CompressedData = Object["CompressedData"_ASV].AsBinary();
 
-		if (Id.IsNull() || RawHash.IsZero() || !(CompressedHash.IsZero() == CompressedData.IsNull()))
+		if (Id.IsNull() || RawSize == MAX_uint64 || RawHash.IsZero() || !(CompressedHash.IsZero() == CompressedData.IsNull()))
 		{
 			UE_LOG(LogDerivedDataCache, Display, TEXT("%s: Cache miss with invalid record format for %s from '%.*s'"),
 				*CachePath, *WriteToString<96>(Key), Context.Len(), Context.GetData());
 			return FPayload();
 		}
 
-		FPayload Payload(Id, RawHash);
+		FPayload Payload(Id, RawHash, RawSize);
 
 		if (CompressedData)
 		{
diff --git a/Engine/Source/Developer/DerivedDataCache/Public/DerivedDataPayload.h b/Engine/Source/Developer/DerivedDataCache/Public/DerivedDataPayload.h
--- a/Engine/Source/Developer/DerivedDataCache/Public/DerivedDataPayload.h
+++ b/Engine/Source/Developer/DerivedDataCache/Public/DerivedDataPayload.h
@@ -61,7 +61,7 @@ inline const FPayloadId FPayloadId::Null;
 ///////////////////////////////////////////////////////////////////////////////////////////////////////////////////////
 
 /**
- * A payload is described by a PayloadId and an IoHash of its raw buffer (uncompressed).
+ * A payload is described by an ID and by the hash and size of its raw buffer (uncompressed).
  *
  * Payloads may be constructed with or without a compressed buffer. A payload without a buffer is
  * used to uniquely describe a payload without the overhead of its buffer.
@@ -72,22 +72,25 @@ public:
 	/** Construct a null payload. */
 	FPayload() = default;
 
-	/** Construct a payload with no hash or buffer. */
+	/** Construct a payload with no hash, size, or buffer. */
 	inline explicit FPayload(const FPayloadId& Id);
 
-	/** Construct a payload from a raw buffer hash. */
-	inline FPayload(const FPayloadId& Id, const FIoHash& RawHash);
+	/** Construct a payload from a the hash and size of the raw buffer. */
+	inline FPayload(const FPayloadId& Id, const FIoHash& RawHash, uint64 RawSize);
 
-	/** Construct a payload from a compressed buffer. */
+	/** Construct a payload from a compressed buffer, which is cloned if not owned. */
 	inline FPayload(const FPayloadId& Id, const FCompressedBuffer& Data);
 	inline FPayload(const FPayloadId& Id, FCompressedBuffer&& Data);
 
-	/** Returns the PayloadId for the payload. */
+	/** Returns the ID for the payload. */
 	inline const FPayloadId& GetId() const { return Id; }
 
-	/** Returns the IoHash of the raw buffer (uncompressed) for the payload. */
+	/** Returns the hash of the raw buffer (uncompressed) for the payload. */
 	inline const FIoHash& GetRawHash() const { return RawHash; }
 
+	/** Returns the size of the raw buffer (uncompressed) for the payload. */
+	inline uint64 GetRawSize() const { return RawSize; }
+
 	/** Returns the compressed buffer for the payload. May be null. */
 	inline const FCompressedBuffer& GetData() const { return Data; }
 
@@ -107,6 +110,7 @@ public:
 private:
 	FPayloadId Id;
 	FIoHash RawHash;
+	uint64 RawSize = 0;
 	FCompressedBuffer Data;
 };
 
@@ -146,9 +150,10 @@ inline FPayload::FPayload(const FPayloadId& InId)
 	checkf(Id.IsValid(), TEXT("A valid ID is required to construct a payload."));
 }
 
-inline FPayload::FPayload(const FPayloadId& InId, const FIoHash& InRawHash)
+inline FPayload::FPayload(const FPayloadId& InId, const FIoHash& InRawHash, const uint64 InRawSize)
 	: Id(InId)
 	, RawHash(InRawHash)
+	, RawSize(InRawSize)
 {
 	checkf(Id.IsValid(), TEXT("A valid ID is required to construct a payload."));
 }
@@ -156,7 +161,8 @@ inline FPayload::FPayload(const FPayloadId& InId, const FIoHash& InRawHash)
 inline FPayload::FPayload(const FPayloadId& InId, const FCompressedBuffer& InData)
 	: Id(InId)
 	, RawHash(InData.GetRawHash())
-	, Data(InData)
+	, RawSize(InData.GetRawSize())
+	, Data(InData.MakeOwned())
 {
 	checkf(Id.IsValid(), TEXT("A valid ID is required to construct a payload."));
 }
@@ -164,7 +170,8 @@ inline FPayload::FPayload(const FPayloadId& InId, const FCompressedBuffer& InDat
 inline FPayload::FPayload(const FPayloadId& InId, FCompressedBuffer&& InData)
 	: Id(InId)
 	, RawHash(InData.GetRawHash())
-	, Data(MoveTemp(InData))
+	, RawSize(InData.GetRawSize())
+	, Data(MoveTemp(InData).MakeOwned())
 {
 	checkf(Id.IsValid(), TEXT("A valid ID is required to construct a payload."));
 }
@@ -206,16 +213,20 @@ inline bool FPayloadId::IsNull() const
 
 ///////////////////////////////////////////////////////////////////////////////////////////////////////////////////////
 
-/** Compare payloads by their ID and raw buffer hash. */
+/** Compare payloads by their ID and the hash and size of their raw buffer. */
 inline bool operator==(const FPayload& A, const FPayload& B)
 {
-	return A.GetId() == B.GetId() && A.GetRawHash() == B.GetRawHash();
+	return A.GetId() == B.GetId() && A.GetRawHash() == B.GetRawHash() && A.GetRawSize() == B.GetRawSize();
 }
 
-/** Compare payloads by their ID and raw buffer hash. */
+/** Compare payloads by their ID and the hash and size of their raw buffer. */
 inline bool operator<(const FPayload& A, const FPayload& B)
 {
-	return A.GetId() == B.GetId() && A.GetRawHash() == B.GetRawHash();
+	const FPayloadId& IdA = A.GetId();
+	const FPayloadId& IdB = B.GetId();
+	const FIoHash& HashA = A.GetRawHash();
+	const FIoHash& HashB = B.GetRawHash();
+	return !(IdA == IdB) ? IdA < IdB : !(HashA == HashB) ? HashA < HashB : A.GetRawSize() < B.GetRawSize();
 }
 
 /** Compare payloads by their ID. */
