DDC: Updated documentation and removed unnecessary explicit hash construction

#rnx

diff --git a/Engine/Source/Developer/DerivedDataCache/Private/DerivedDataCacheRecord.cpp b/Engine/Source/Developer/DerivedDataCache/Private/DerivedDataCacheRecord.cpp
--- a/Engine/Source/Developer/DerivedDataCache/Private/DerivedDataCacheRecord.cpp
+++ b/Engine/Source/Developer/DerivedDataCache/Private/DerivedDataCacheRecord.cpp
@@ -196,7 +196,7 @@ void FCacheRecordBuilderInternal::SetMeta(FCbObject&& InMeta)
 FPayloadId FCacheRecordBuilderInternal::SetValue(const FSharedBuffer& Buffer, const FPayloadId& Id)
 {
 	FCompressedBuffer CompressedBuffer = FCompressedBuffer::Compress(NAME_Default, Buffer);
-	const FPayloadId ValueId = GetOrCreatePayloadId(Id, FIoHash(CompressedBuffer.GetRawHash()));
+	const FPayloadId ValueId = GetOrCreatePayloadId(Id, CompressedBuffer.GetRawHash());
 	return SetValue(FPayload(ValueId, MoveTemp(CompressedBuffer)));
 }
 
@@ -213,7 +213,7 @@ FPayloadId FCacheRecordBuilderInternal::SetValue(FPayload&& Payload)
 FPayloadId FCacheRecordBuilderInternal::AddAttachment(const FSharedBuffer& Buffer, const FPayloadId& Id)
 {
 	FCompressedBuffer CompressedBuffer = FCompressedBuffer::Compress(NAME_Default, Buffer);
-	const FPayloadId AttachmentId = GetOrCreatePayloadId(Id, FIoHash(CompressedBuffer.GetRawHash()));
+	const FPayloadId AttachmentId = GetOrCreatePayloadId(Id, CompressedBuffer.GetRawHash());
 	return AddAttachment(FPayload(AttachmentId, MoveTemp(CompressedBuffer)));
 }
 
diff --git a/Engine/Source/Developer/DerivedDataCache/Public/DerivedDataBuildDefinition.h b/Engine/Source/Developer/DerivedDataCache/Public/DerivedDataBuildDefinition.h
--- a/Engine/Source/Developer/DerivedDataCache/Public/DerivedDataBuildDefinition.h
+++ b/Engine/Source/Developer/DerivedDataCache/Public/DerivedDataBuildDefinition.h
@@ -67,18 +67,17 @@ namespace UE::DerivedData
  * A build definition references a build function and the inputs to that build function.
  *
  * The purpose of a definition is to capture everything required to execute a derived data build.
- * Inputs to the build are partly static (constants and input hashes) and partly variable (files,
- * bulk data, payloads from other builds); and the build function is referenced by name, but does
- * not correspond to a fixed version of the build function.
+ * The definition is partly fixed (function name, constants, input hashes) and is partly variable
+ * (function version, build dependencies, bulk data, files), and the variable components mean the
+ * output from building the definition can vary depending on the versions of those components.
  *
- * The key for the definition uniquely identifies the definition and is derived from its function
- * and inputs, such that any definition with the same function and inputs will have the same key.
+ * The key for the definition uniquely identifies the definition using the hash of the serialized
+ * compact binary representation of the definition.
  *
- * The keys for constants and inputs are names that are unique within the build definition across
- * every type of constant and input.
+ * The keys for constants and inputs are names that are unique within the build definition.
  *
  * To build a definition against a specific version of the function and inputs, queue it to build
- * on FBuildSession, which uses IBuildInputProvider to fetch the inputs at the expected versions.
+ * on a build session, which uses its build input provider to convert it to a build action.
  *
  * Build definitions are immutable, and are created by a builder from IBuild::CreateDefinition().
  *
diff --git a/Engine/Source/Developer/DerivedDataCache/Public/DerivedDataCache.h b/Engine/Source/Developer/DerivedDataCache/Public/DerivedDataCache.h
--- a/Engine/Source/Developer/DerivedDataCache/Public/DerivedDataCache.h
+++ b/Engine/Source/Developer/DerivedDataCache/Public/DerivedDataCache.h
@@ -131,11 +131,11 @@ public:
 	 * The callback will always be called for every key, and may be called from an arbitrary thread.
 	 * Records may finish storing in any order, and from multiple threads concurrently.
 	 *
-	 * @param Records The cache records to store. Must have a key.
-	 * @param Context A description of the request. An object path is typically sufficient.
-	 * @param Policy Flags to control the behavior of the request. See ECachePolicy.
-	 * @param Priority A priority to consider when scheduling the request. See EPriority.
-	 * @param OnComplete A callback invoked for every key in the batch as it completes or is canceled.
+	 * @param Records      The cache records to store. Must have a key.
+	 * @param Context      A description of the request. An object path is typically sufficient.
+	 * @param Policy       Flags to control the behavior of the request. See ECachePolicy.
+	 * @param Priority     A priority to consider when scheduling the request. See EPriority.
+	 * @param OnComplete   A callback invoked for every record as it completes or is canceled.
 	 */
 	virtual FRequest Put(
 		TConstArrayView<FCacheRecord> Records,
@@ -150,11 +150,11 @@ public:
 	 * The callback will always be called for every key, and may be called from an arbitrary thread.
 	 * Records may become available in any order, and from multiple threads concurrently.
 	 *
-	 * @param Keys The keys identifying the cache records to query.
-	 * @param Context A description of the request. An object path is typically sufficient.
-	 * @param Policy Flags to control the behavior of the request. See ECachePolicy.
-	 * @param Priority A priority to consider when scheduling the request. See EPriority.
-	 * @param OnComplete A callback invoked for every key in the batch as it completes or is canceled.
+	 * @param Keys         The keys identifying the cache records to query.
+	 * @param Context      A description of the request. An object path is typically sufficient.
+	 * @param Policy       Flags to control the behavior of the request. See ECachePolicy.
+	 * @param Priority     A priority to consider when scheduling the request. See EPriority.
+	 * @param OnComplete   A callback invoked for every key as it completes or is canceled.
 	 */
 	virtual FRequest Get(
 		TConstArrayView<FCacheKey> Keys,
@@ -169,11 +169,11 @@ public:
 	 * The callback will always be called for every key, and may be called from an arbitrary thread.
 	 * Payloads may become available in any order, and from multiple threads concurrently.
 	 *
-	 * @param Keys The keys identifying the cache record payloads to query.
-	 * @param Context A description of the request. An object path is typically sufficient.
-	 * @param Policy Flags to control the behavior of the request. See ECachePolicy.
-	 * @param Priority A priority to consider when scheduling the request. See EPriority.
-	 * @param OnComplete A callback invoked for every key in the batch as it completes or is canceled.
+	 * @param Keys         The keys identifying the cache record payloads to query.
+	 * @param Context      A description of the request. An object path is typically sufficient.
+	 * @param Policy       Flags to control the behavior of the request. See ECachePolicy.
+	 * @param Priority     A priority to consider when scheduling the request. See EPriority.
+	 * @param OnComplete   A callback invoked for every key as it completes or is canceled.
 	 */
 	virtual FRequest GetPayload(
 		TConstArrayView<FCachePayloadKey> Keys,
@@ -228,7 +228,7 @@ struct FCacheGetCompleteParams
 	 * The key is always populated. The remainder of the record is populated when Status is Ok.
 	 *
 	 * The value, attachments, and metadata may be skipped based on cache policy flags. When a value
-	 * or attachment has been skipped, it will have a payload but its buffers will be null.
+	 * or attachment has been skipped, it will have a payload but its buffer will be null.
 	 */
 	FCacheRecord&& Record;
 
@@ -245,7 +245,9 @@ struct FCacheGetPayloadCompleteParams
 	/**
 	 * Payload for the part of the payload request that completed or was canceled.
 	 *
-	 * The ID is always populated. The hash and buffer are populated when Status is Ok.
+	 * The ID is always populated.
+	 * The hash and size are populated when Status is Ok.
+	 * The buffer is populated when Status is Ok and the data was not skipped by the cache policy.
 	 */
 	FPayload&& Payload;
 
diff --git a/Engine/Source/Developer/DerivedDataCache/Public/DerivedDataCacheRecord.h b/Engine/Source/Developer/DerivedDataCache/Public/DerivedDataCacheRecord.h
--- a/Engine/Source/Developer/DerivedDataCache/Public/DerivedDataCacheRecord.h
+++ b/Engine/Source/Developer/DerivedDataCache/Public/DerivedDataCacheRecord.h
@@ -73,8 +73,8 @@ namespace UE::DerivedData
  * of the payloads are exposed by the cache record.
  *
  * The value, attachments, and metadata are optional and can be skipped when requesting a record.
- * When the value or attachments have been skipped, the record will contain a payload with the ID
- * and hash, but no buffer containing the associated data.
+ * When the value or attachments have been skipped, the record will contain a payload with a null
+ * buffer but will otherwise be populated.
  */
 class FCacheRecord
 {
@@ -145,7 +145,7 @@ public:
 	/**
 	 * Set the metadata for the cache record.
 	 *
-	 * The metadata is cloned if not owned.
+	 * @param Meta   The metadata, which is cloned if not owned.
 	 */
 	inline void SetMeta(FCbObject&& Meta)
 	{
@@ -155,9 +155,9 @@ public:
 	/**
 	 * Set the value for the cache record.
 	 *
-	 * @param Buffer The value, which is cloned if not owned, and is compressed by the builder.
-	 * @param Id An ID for the value that is unique within the scope of the cache record. If omitted
-	 *           the builder will create an ID by hashing the buffer.
+	 * @param Buffer   The value, which is compressed by the builder, and cloned if not owned.
+	 * @param Id       An ID for the value that is unique within this cache record. When omitted,
+	 *                 the buffer will be hashed to create an ID.
 	 * @return The ID that was provided or created.
 	 */
 	inline FPayloadId SetValue(const FSharedBuffer& Buffer, const FPayloadId& Id = FPayloadId())
@@ -168,7 +168,7 @@ public:
 	/**
 	 * Set the value for the cache record.
 	 *
-	 * @param Payload The value payload, which is reset to null.
+	 * @param Payload   The value payload, which must have a buffer.
 	 * @return The ID that was provided. Unique within the scope of the cache record.
 	 */
 	inline FPayloadId SetValue(FPayload&& Payload)
@@ -179,9 +179,9 @@ public:
 	/**
 	 * Add an attachment to the cache record.
 	 *
-	 * @param Buffer The attachment, which is cloned if not owned, and is compressed by the builder.
-	 * @param Id An ID for the attachment that is unique within the scope of the cache record. If it
-	 *           is omitted the builder will create an ID by hashing the buffer.
+	 * @param Buffer   The attachment, which is compressed by the builder, and cloned if not owned.
+	 * @param Id       An ID for the attachment that is unique within this cache record. When omitted,
+	 *                 the buffer will be hashed to create an ID.
 	 * @return The ID that was provided or created.
 	 */
 	inline FPayloadId AddAttachment(const FSharedBuffer& Buffer, const FPayloadId& Id = FPayloadId())
@@ -192,7 +192,7 @@ public:
 	/**
 	 * Add an attachment to the cache record.
 	 *
-	 * @param Payload The attachment payload, which is reset to null.
+	 * @param Payload   The attachment payload, which must have a buffer.
 	 * @return The ID that was provided. Unique within the scope of the cache record.
 	 */
 	inline FPayloadId AddAttachment(FPayload&& Payload)
diff --git a/Engine/Source/Developer/DerivedDataCache/Public/DerivedDataRequest.h b/Engine/Source/Developer/DerivedDataCache/Public/DerivedDataRequest.h
--- a/Engine/Source/Developer/DerivedDataCache/Public/DerivedDataRequest.h
+++ b/Engine/Source/Developer/DerivedDataCache/Public/DerivedDataRequest.h
@@ -130,7 +130,7 @@ public:
 	/**
 	 * Set the priority of the request.
 	 *
-	 * @param Priority The new priority for the request. May be higher or lower than the original.
+	 * @param Priority   The priority, which may be higher or lower than the existing priority.
 	 */
 	inline void SetPriority(EPriority Priority)
 	{
