DDC: Minor updates to formatting of code and asserts

#rnx

diff --git a/Engine/Source/Developer/DerivedDataCache/Private/DerivedDataBuildDefinition.cpp b/Engine/Source/Developer/DerivedDataCache/Private/DerivedDataBuildDefinition.cpp
--- a/Engine/Source/Developer/DerivedDataCache/Private/DerivedDataBuildDefinition.cpp
+++ b/Engine/Source/Developer/DerivedDataCache/Private/DerivedDataBuildDefinition.cpp
@@ -19,8 +19,6 @@
 namespace UE::DerivedData::Private
 {
 
-///////////////////////////////////////////////////////////////////////////////////////////////////////////////////////
-
 class FBuildDefinitionBuilderInternal final : public IBuildDefinitionBuilderInternal
 {
 public:
@@ -29,9 +27,9 @@ public:
 		, Function(InFunction)
 	{
 		checkf(!Name.IsEmpty(), TEXT("A build definition requires a non-empty name."));
-		checkf(!Function.IsEmpty(), TEXT("A build definition requires a non-empty function name."));
-		checkf(Algo::AllOf(Function, FChar::IsAlnum), TEXT("A build definition requires an alphanumeric function name. ")
-			TEXT("Definition: %s; Function: %s"), *Name, *Function);
+		checkf(!Function.IsEmpty(), TEXT("A build definition requires a non-empty function name for build of '%s'."), *Name);
+		checkf(Algo::AllOf(Function, FChar::IsAlnum), TEXT("Definition requires an alphanumeric function name for ")
+			TEXT("build of '%s' by %s."), *Name, *Function);
 	}
 
 	~FBuildDefinitionBuilderInternal() final = default;
@@ -74,9 +72,9 @@ private:
 	inline void Add(FStringView Key, ArgType&& Value)
 	{
 		const uint32 KeyHash = GetTypeHash(Key);
-		checkf(!Key.IsEmpty(), TEXT("Empty key used in the %s build input for %s."), *Function, *Name);
-		checkf(!Inputs.ContainsByHash(KeyHash, Key), TEXT("Duplicate key %.*s used in the %s build input for %s."),
-			Key.Len(), Key.GetData(), *Function, *Name);
+		checkf(!Key.IsEmpty(), TEXT("Empty key used in definition for build of '%s' by %s."), *Name, *Function);
+		checkf(!Inputs.ContainsByHash(KeyHash, Key), TEXT("Duplicate key '%.*s' used in definition for "),
+			TEXT("build of '%s' by %s."), Key.Len(), Key.GetData(), *Name, *Function);
 		Inputs.EmplaceByHash(KeyHash, Key, InputType(TInPlaceType<ValueType>(), Forward<ArgType>(Value)));
 	}
 };
@@ -321,8 +319,6 @@ FBuildDefinitionBuilder CreateBuildDefinitionBuilder(IBuildDefinitionBuilderInte
 	return FBuildDefinitionBuilder(DefinitionBuilder);
 }
 
-///////////////////////////////////////////////////////////////////////////////////////////////////////////////////////
-
 FBuildDefinitionBuilder CreateBuildDefinition(FStringView Name, FStringView Function)
 {
 	return CreateBuildDefinitionBuilder(new FBuildDefinitionBuilderInternal(Name, Function));
@@ -333,6 +329,4 @@ FBuildDefinition LoadBuildDefinition(FStringView Name, FCbObject&& Definition)
 	return CreateBuildDefinition(new FBuildDefinitionInternal(Name, MoveTemp(Definition)));
 }
 
-///////////////////////////////////////////////////////////////////////////////////////////////////////////////////////
-
 } // UE::DerivedData::Private
diff --git a/Engine/Source/Developer/DerivedDataCache/Public/DerivedDataBuildDefinition.h b/Engine/Source/Developer/DerivedDataCache/Public/DerivedDataBuildDefinition.h
--- a/Engine/Source/Developer/DerivedDataCache/Public/DerivedDataBuildDefinition.h
+++ b/Engine/Source/Developer/DerivedDataCache/Public/DerivedDataBuildDefinition.h
@@ -18,8 +18,6 @@ namespace UE::DerivedData { class FBuildDefinitionBuilder; }
 namespace UE::DerivedData { struct FBuildKey; }
 namespace UE::DerivedData { struct FBuildPayloadKey; }
 
-///////////////////////////////////////////////////////////////////////////////////////////////////////////////////////
-
 namespace UE::DerivedData::Private
 {
 
@@ -58,8 +56,6 @@ FBuildDefinitionBuilder CreateBuildDefinitionBuilder(IBuildDefinitionBuilderInte
 
 } // UE::DerivedData::Private
 
-///////////////////////////////////////////////////////////////////////////////////////////////////////////////////////
-
 namespace UE::DerivedData
 {
 
@@ -178,7 +174,7 @@ public:
 	/**
 	 * Add a bulk data input with a key that is unique within this definition.
 	 *
-	 * @param BulkDataId   Identifier that uniquely identifies this data in the IBuildInputProvider.
+	 * @param BulkDataId   Identifier that uniquely identifies this data in the IBuildInputResolver.
 	 */
 	inline void AddInputBulkData(FStringView Key, const FGuid& BulkDataId)
 	{
@@ -198,7 +194,7 @@ public:
 	/**
 	 * Add a hash input with a key that is unique within this definition.
 	 *
-	 * @param RawHash   Hash of the raw data that will resolve it in the IBuildInputProvider.
+	 * @param RawHash   Hash of the raw data that will resolve it in the IBuildInputResolver.
 	 */
 	inline void AddInputHash(FStringView Key, const FIoHash& RawHash)
 	{
@@ -225,14 +221,4 @@ private:
 	TUniquePtr<Private::IBuildDefinitionBuilderInternal> DefinitionBuilder;
 };
 
-/**
- * A build definition provider finds a definition from its key.
- */
-class IBuildDefinitionProvider
-{
-public:
-	/** Returns a definition matching the key, or a definition with an empty key on error. */
-	virtual FBuildDefinition GetDefinition(const FBuildKey& Key) = 0;
-};
-
 } // UE::DerivedData
diff --git a/Engine/Source/Developer/DerivedDataCache/Public/DerivedDataBuildKey.h b/Engine/Source/Developer/DerivedDataCache/Public/DerivedDataBuildKey.h
--- a/Engine/Source/Developer/DerivedDataCache/Public/DerivedDataBuildKey.h
+++ b/Engine/Source/Developer/DerivedDataCache/Public/DerivedDataBuildKey.h
@@ -11,8 +11,6 @@
 namespace UE::DerivedData
 {
 
-///////////////////////////////////////////////////////////////////////////////////////////////////////////////////////
-
 /** A key that uniquely identifies a build definition. */
 struct FBuildKey
 {
@@ -94,6 +92,4 @@ inline TStringBuilderBase<CharType>& operator<<(TStringBuilderBase<CharType>& Bu
 	return Builder << Key.BuildKey << CharType('/') << Key.Id;
 }
 
-///////////////////////////////////////////////////////////////////////////////////////////////////////////////////////
-
 } // UE::DerivedData
