DDC: Added UE::DerivedData::IBuildFunction and related types

IBuildFunction and the associated build context types define the interface for new derived data build functions. The inputs to a build function are referenced by FBuildDefinition. This only includes the functionality required to implement and register a build function, and does not include the functionality required to execute a build starting from a build definition.

#rnx

diff --git a/Engine/Source/Developer/DerivedDataCache/Public/DerivedDataBuildFunction.h b/Engine/Source/Developer/DerivedDataCache/Public/DerivedDataBuildFunction.h
new file mode 100644
--- /dev/null
+++ b/Engine/Source/Developer/DerivedDataCache/Public/DerivedDataBuildFunction.h
@@ -0,0 +1,132 @@
+// Copyright Epic Games, Inc. All Rights Reserved.
+
+#pragma once
+
+#include "CoreTypes.h"
+#include "Containers/StringView.h"
+#include "Memory/MemoryFwd.h"
+
+class FCbObject;
+class FCompressedBuffer;
+
+namespace UE::DerivedData { class FBuildCacheContext; }
+namespace UE::DerivedData { class FBuildContext; }
+namespace UE::DerivedData { class FCacheBucket; }
+namespace UE::DerivedData { class FPayload; }
+namespace UE::DerivedData { struct FPayloadId; }
+namespace UE::DerivedData { enum class ECachePolicy : uint8; }
+
+namespace UE::DerivedData
+{
+
+/**
+ * A build function is referenced by a build definition as the way to build its inputs.
+ *
+ * Functions are expected to be pure and maintain no state. Inputs are provided through the build
+ * context, and outputs are saved through the build context.
+ *
+ * Functions have a version which is used as a proxy for their code. Any code changes that affect
+ * the behavior of a function must have a corresponding change in the function version.
+ *
+ * Functions must be registered by a build function factory to be available to the build system.
+ *
+ * @see TBuildFunctionFactory
+ */
+class IBuildFunction
+{
+public:
+	virtual ~IBuildFunction() = default;
+
+	/** Returns the name of the build function, which identifies it in a build definition. */
+	virtual FStringView GetName() const = 0;
+
+	/** Returns the version of the build function, which must change when the function changes. */
+	virtual FGuid GetVersion() const = 0;
+
+	/**
+	 * Determines the cache bucket and cache policy to use for the build.
+	 *
+	 * Default behavior is to use the function name as the cache bucket and ECachePolicy::Default.
+	 */
+	virtual void Cache(FBuildCacheContext& Context) const {}
+
+	/**
+	 * Build the output for the input in the build context.
+	 *
+	 * Asynchronous build completion is possible by using BeginAsyncBuild and EndAsyncBuild provided
+	 * on the build context. Otherwise, the build is complete when this function returns.
+	 */
+	virtual void Build(FBuildContext& Context) const = 0;
+
+	/** Whether the build may be executed on any thread and concurrently with any other build. */
+	virtual bool AllowsParallelBuild() const { return true; }
+
+	/** Whether the build may be executed externally, in another process or on a remote agent. */
+	virtual bool AllowsExternalBuild() const { return false; }
+
+	/**
+	 * Whether the build produces deterministic output.
+	 *
+	 * Builds must strive to always produce the same output for a given input. Depending on external
+	 * code can make this infeasible for some build functions. To opt out of automatic validation of
+	 * build determinism, override this function to return false.
+	 */
+	virtual bool IsDeterministic() const { return true; }
+};
+
+/** A build context provides the inputs for a build function and saves its outputs. */
+class FBuildContext
+{
+public:
+	/** Returns the constant with the matching key, or an object with no fields if not found. */
+	virtual FCbObject GetConstant(FStringView Key) const = 0;
+
+	/** Returns the input with the matching key, or a null buffer if not found. */
+	virtual FSharedBuffer GetInput(FStringView Key) const = 0;
+
+	/** Adds a payload to the build output. Must have a non-null buffer and a unique ID. */
+	virtual void AddPayload(const FPayload& Payload) = 0;
+	virtual void AddPayload(const FPayloadId& Id, const FCompressedBuffer& Buffer) = 0;
+	virtual void AddPayload(const FPayloadId& Id, const FSharedBuffer& Buffer) = 0;
+	virtual void AddPayload(const FPayloadId& Id, const FCbObject& Object) = 0;
+
+	/** Overrides the default cache policy used when writing this build in the cache. */
+	virtual void SetCachePolicy(ECachePolicy Policy) = 0;
+
+	/**
+	 * Make this an asynchronous build by giving the caller responsibility for completing the build.
+	 *
+	 * This is an advanced feature that bypasses many of the safety checks that are available within
+	 * the scope of the synchronous build function. Take extra care to only consume inputs available
+	 * in the build context and avoid reading any global state.
+	 *
+	 * This may be called at most once on a build context. The caller becomes responsible for ending
+	 * the build by calling EndAsyncBuild on this context once the build has finished. Once an async
+	 * build has begun, it may end from any thread, but the context may only be used from one thread
+	 * at a time.
+	 */
+	virtual void BeginAsyncBuild() = 0;
+
+	/**
+	 * Mark the end of an asynchronous build.
+	 *
+	 * It is invalid to call any other function on the build context after calling this.
+	 */
+	virtual void EndAsyncBuild() = 0;
+};
+
+/** A build cache context allows cache behavior to be modified based on constant inputs. */
+class FBuildCacheContext
+{
+public:
+	/** Returns the constant with the matching key, or an object with no fields if not found. */
+	virtual FCbObject GetConstant(FStringView Key) const = 0;
+
+	/** Overrides the default cache bucket used when reading or writing this build in the cache. */
+	virtual void SetCacheBucket(FCacheBucket Bucket) = 0;
+
+	/** Overrides the default cache policy used when reading or writing this build in the cache. */
+	virtual void SetCachePolicy(ECachePolicy Policy) = 0;
+};
+
+} // UE::DerivedData
diff --git a/Engine/Source/Developer/DerivedDataCache/Public/DerivedDataBuildFunctionFactory.h b/Engine/Source/Developer/DerivedDataCache/Public/DerivedDataBuildFunctionFactory.h
new file mode 100644
--- /dev/null
+++ b/Engine/Source/Developer/DerivedDataCache/Public/DerivedDataBuildFunctionFactory.h
@@ -0,0 +1,59 @@
+// Copyright Epic Games, Inc. All Rights Reserved.
+
+#pragma once
+
+#include "DerivedDataBuildFunction.h"
+#include "Features/IModularFeatures.h"
+#include "UObject/NameTypes.h"
+
+namespace UE::DerivedData
+{
+
+/** Base for the build function factory. DO NOT USE DIRECTLY. Use TBuildFunctionFactory. */
+class IBuildFunctionFactory : public IModularFeature
+{
+public:
+	/** Returns the build function associated with this factory. */
+	virtual const IBuildFunction& GetFunction() const = 0;
+
+	/** Returns the name of the build function factory modular feature. */
+	static FName GetFeatureName()
+	{
+		return FName("BuildFunctionFactory");
+	}
+};
+
+/**
+ * Factory that creates and registers a build function.
+ *
+ * A build function must be registered by a build function factory before it can execute a build.
+ * Register a function in the source file that implements it or in the corresponding module.
+ *
+ * Example: static const TBuildFunctionFactory<FExampleFunction> ExampleFunctionFactory;
+ */
+template <typename FunctionType>
+class TBuildFunctionFactory final : private IBuildFunctionFactory
+{
+	static_assert(sizeof(FunctionType) == sizeof(IBuildFunction), "IBuildFunction must be pure and maintain no state.");
+
+public:
+	TBuildFunctionFactory()
+	{
+		IModularFeatures::Get().RegisterModularFeature(GetFeatureName(), this);
+	}
+
+	~TBuildFunctionFactory()
+	{
+		IModularFeatures::Get().UnregisterModularFeature(GetFeatureName(), this);
+	}
+
+private:
+	const IBuildFunction& GetFunction() const final
+	{
+		return Function;
+	}
+
+	const FunctionType Function;
+};
+
+} // UE::DerivedData
