DDC: Added support for cache records with null data

#rnx

diff --git a/Engine/Source/Developer/DerivedDataCache/Private/DerivedDataCacheRecord.cpp b/Engine/Source/Developer/DerivedDataCache/Private/DerivedDataCacheRecord.cpp
--- a/Engine/Source/Developer/DerivedDataCache/Private/DerivedDataCacheRecord.cpp
+++ b/Engine/Source/Developer/DerivedDataCache/Private/DerivedDataCacheRecord.cpp
@@ -25,10 +25,10 @@ public:
 	void SetMeta(FCbObject&& Meta) final;
 
 	FPayloadId SetValue(const FSharedBuffer& Buffer, const FPayloadId& Id) final;
-	FPayloadId SetValue(FPayload&& Payload) final;
+	FPayloadId SetValue(const FPayload& Payload) final;
 
 	FPayloadId AddAttachment(const FSharedBuffer& Buffer, const FPayloadId& Id) final;
-	FPayloadId AddAttachment(FPayload&& Payload) final;
+	FPayloadId AddAttachment(const FPayload& Payload) final;
 
 	FCacheRecord Build() final;
 	FRequest BuildAsync(FOnCacheRecordComplete&& OnComplete, EPriority Priority) final;
@@ -138,8 +138,8 @@ const FPayload& FCacheRecordInternal::GetValuePayload() const
 
 FSharedBuffer FCacheRecordInternal::GetAttachment(const FPayloadId& Id) const
 {
-	const int32 Index = Algo::LowerBound(Attachments, Id, FPayloadLessById());
-	if (Attachments.IsValidIndex(Index) && FPayloadEqualById()(Attachments[Index], Id))
+	const int32 Index = Algo::BinarySearchBy(Attachments, Id, &FPayload::GetId);
+	if (Attachments.IsValidIndex(Index))
 	{
 		if (FReadScopeLock Lock(CacheLock); AttachmentsCache.IsValidIndex(Index) && AttachmentsCache[Index])
 		{
@@ -162,12 +162,8 @@ FSharedBuffer FCacheRecordInternal::GetAttachment(const FPayloadId& Id) const
 
 const FPayload& FCacheRecordInternal::GetAttachmentPayload(const FPayloadId& Id) const
 {
-	const int32 Index = Algo::LowerBound(Attachments, Id, FPayloadLessById());
-	if (Attachments.IsValidIndex(Index) && FPayloadEqualById()(Attachments[Index], Id))
-	{
-		return Attachments[Index];
-	}
-	return FPayload::Null;
+	const int32 Index = Algo::BinarySearchBy(Attachments, Id, &FPayload::GetId);
+	return Attachments.IsValidIndex(Index) ? Attachments[Index] : FPayload::Null;
 }
 
 TConstArrayView<FPayload> FCacheRecordInternal::GetAttachmentPayloads() const
@@ -200,14 +196,18 @@ FPayloadId FCacheRecordBuilderInternal::SetValue(const FSharedBuffer& Buffer, co
 	return SetValue(FPayload(ValueId, MoveTemp(CompressedBuffer)));
 }
 
-FPayloadId FCacheRecordBuilderInternal::SetValue(FPayload&& Payload)
+FPayloadId FCacheRecordBuilderInternal::SetValue(const FPayload& Payload)
 {
 	checkf(Payload, TEXT("Failed to set value on %s because the payload is null."), *WriteToString<96>(Key));
+	const FPayloadId& Id = Payload.GetId();
 	checkf(Value.IsNull(),
 		TEXT("Cache: Failed to set value on %s with ID %s because it has an existing value with ID %s."),
-		*WriteToString<96>(Key), *WriteToString<32>(Payload.GetId()), *WriteToString<32>(Value.GetId()));
-	Value = MoveTemp(Payload);
-	return Value.GetId();
+		*WriteToString<96>(Key), *WriteToString<32>(Id), *WriteToString<32>(Value.GetId()));
+	checkf(Algo::BinarySearchBy(Attachments, Id, &FPayload::GetId) == INDEX_NONE,
+		TEXT("Failed to set on %s with ID %s because it has an existing attachment with that ID."),
+		*WriteToString<96>(Key), *WriteToString<32>(Id));
+	Value = Payload;
+	return Id;
 }
 
 FPayloadId FCacheRecordBuilderInternal::AddAttachment(const FSharedBuffer& Buffer, const FPayloadId& Id)
@@ -217,15 +217,16 @@ FPayloadId FCacheRecordBuilderInternal::AddAttachment(const FSharedBuffer& Buffe
 	return AddAttachment(FPayload(AttachmentId, MoveTemp(CompressedBuffer)));
 }
 
-FPayloadId FCacheRecordBuilderInternal::AddAttachment(FPayload&& Payload)
+FPayloadId FCacheRecordBuilderInternal::AddAttachment(const FPayload& Payload)
 {
 	checkf(Payload, TEXT("Failed to add attachment on %s because the payload is null."), *WriteToString<96>(Key));
-	const int32 Index = Algo::LowerBound(Attachments, Payload, FPayloadLessById());
-	checkf(!Attachments.IsValidIndex(Index) || !FPayloadEqualById()(Attachments[Index], Payload),
-		TEXT("Failed to add attachment on %s with ID %s because it has an existing attachment with that ID."),
-		*WriteToString<96>(Key), *WriteToString<32>(Payload.GetId()));
-	Attachments.Insert(MoveTemp(Payload), Index);
-	return Attachments[Index].GetId();
+	const FPayloadId& Id = Payload.GetId();
+	const int32 Index = Algo::LowerBoundBy(Attachments, Id, &FPayload::GetId);
+	checkf(!(Attachments.IsValidIndex(Index) && Attachments[Index].GetId() == Id) && Value.GetId() != Id,
+		TEXT("Failed to add attachment on %s with ID %s because it has an existing attachment or value with that ID."),
+		*WriteToString<96>(Key), *WriteToString<32>(Id));
+	Attachments.Insert(Payload, Index);
+	return Id;
 }
 
 FCacheRecord FCacheRecordBuilderInternal::Build()
diff --git a/Engine/Source/Developer/DerivedDataCache/Private/FileSystemDerivedDataBackend.cpp b/Engine/Source/Developer/DerivedDataCache/Private/FileSystemDerivedDataBackend.cpp
--- a/Engine/Source/Developer/DerivedDataCache/Private/FileSystemDerivedDataBackend.cpp
+++ b/Engine/Source/Developer/DerivedDataCache/Private/FileSystemDerivedDataBackend.cpp
@@ -943,7 +943,7 @@ public:
 			{
 				UE_LOG(LogDerivedDataCache, Verbose, TEXT("%s: Cache hit for %s from '%.*s'"),
 					*CachePath, *WriteToString<96>(Key), Context.Len(), Context.GetData());
-				COOK_STAT(Timer.AddHit(Payload.GetData().GetRawSize()));
+				COOK_STAT(Timer.AddHit(Payload.GetRawSize()));
 				if (OnComplete)
 				{
 					OnComplete({Key.CacheKey, MoveTemp(Payload), EStatus::Ok});
@@ -968,9 +968,8 @@ private:
 	uint64 MeasureCacheRecord(const FCacheRecord& Record) const
 	{
 		return Record.GetMeta().GetSize() +
-			Record.GetValuePayload().GetData().GetRawSize() +
-			Algo::TransformAccumulate(Record.GetAttachmentPayloads(),
-				[](const FPayload& Payload) { return Payload.GetData().GetRawSize(); }, uint64(0));
+			Record.GetValuePayload().GetRawSize() +
+			Algo::TransformAccumulate(Record.GetAttachmentPayloads(), &FPayload::GetRawSize, uint64(0));
 	}
 
 	bool PutCacheRecord(const FCacheRecord& Record, FStringView Context, ECachePolicy Policy) const
@@ -1144,8 +1143,9 @@ private:
 
 	bool PutCachePayload(const FCacheKey& Key, FStringView Context, const FPayload& Payload, FCbWriter& Writer, bool bStoreInline) const
 	{
-		const FCompressedBuffer& CompressedBuffer = Payload.GetData();
-		const FIoHash RawHash = CompressedBuffer.GetRawHash();
+		const FIoHash& RawHash = Payload.GetRawHash();
+		const bool bHasData = Payload.HasData();
+		bStoreInline &= bHasData;
 
 		if (!bStoreInline)
 		{
@@ -1153,7 +1153,12 @@ private:
 			BuildCachePayloadPath(Key, RawHash, Path);
 			if (!FileExists(Path))
 			{
+				if (!bHasData)
+				{
+					return false;
+				}
 				// Save the compressed buffer with its hash as a header.
+				const FCompressedBuffer& CompressedBuffer = Payload.GetData();
 				const FIoHash CompressedHash = FIoHash::HashBuffer(CompressedBuffer.GetCompressed());
 				const FCompositeBuffer BufferWithHash(
 					FSharedBuffer::Clone(MakeMemoryView(CompressedHash.GetBytes())),
@@ -1172,9 +1177,10 @@ private:
 
 		Writer.BeginObject();
 		Writer.AddObjectId("Id"_ASV, FCbObjectId(Payload.GetId().GetView()));
-		Writer.AddInteger("RawSize"_ASV, CompressedBuffer.GetRawSize());
+		Writer.AddInteger("RawSize"_ASV, Payload.GetRawSize());
 		if (bStoreInline)
 		{
+			const FCompressedBuffer& CompressedBuffer = Payload.GetData();
 			Writer.AddHash("RawHash"_ASV, RawHash);
 			Writer.AddHash("CompressedHash"_ASV, FIoHash::HashBuffer(CompressedBuffer.GetCompressed()));
 			Writer.AddBinary("CompressedData"_ASV, CompressedBuffer.GetCompressed());
@@ -1221,7 +1227,7 @@ private:
 
 	FPayload GetCachePayload(const FCacheKey& Key, FStringView Context, ECachePolicy Policy, const FPayload& Payload) const
 	{
-		if (Payload.GetData())
+		if (Payload.HasData())
 		{
 			return Payload;
 		}
diff --git a/Engine/Source/Developer/DerivedDataCache/Private/MemoryDerivedDataBackend.cpp b/Engine/Source/Developer/DerivedDataCache/Private/MemoryDerivedDataBackend.cpp
--- a/Engine/Source/Developer/DerivedDataCache/Private/MemoryDerivedDataBackend.cpp
+++ b/Engine/Source/Developer/DerivedDataCache/Private/MemoryDerivedDataBackend.cpp
@@ -3,6 +3,7 @@
 #include "MemoryDerivedDataBackend.h"
 
 #include "Algo/Accumulate.h"
+#include "Algo/AllOf.h"
 #include "Misc/ScopeExit.h"
 #include "Serialization/CompactBinary.h"
 #include "Templates/UniquePtr.h"
@@ -385,9 +386,8 @@ bool FMemoryDerivedDataBackend::ShouldSimulateMiss(const TCHAR* InKey)
 
 int64 FMemoryDerivedDataBackend::CalcRawCacheRecordSize(const FCacheRecord& Record) const
 {
-	const auto CalcCachePayloadSize = [](const FPayload& Payload) { return Payload.GetData().GetRawSize(); };
-	const uint64 ValueSize = CalcCachePayloadSize(Record.GetValuePayload());
-	return int64(Algo::TransformAccumulate(Record.GetAttachmentPayloads(), CalcCachePayloadSize, ValueSize));
+	const uint64 ValueSize = Record.GetValuePayload().GetRawSize();
+	return int64(Algo::TransformAccumulate(Record.GetAttachmentPayloads(), &FPayload::GetRawSize, ValueSize));
 }
 
 int64 FMemoryDerivedDataBackend::CalcSerializedCacheRecordSize(const FCacheRecord& Record) const
@@ -424,7 +424,15 @@ FRequest FMemoryDerivedDataBackend::Put(
 			}
 		};
 
-		if (!Record.GetValuePayload() && Record.GetAttachmentPayloads().IsEmpty())
+		const FPayload& Value = Record.GetValuePayload();
+		const TConstArrayView<FPayload> Attachments = Record.GetAttachmentPayloads();
+
+		if ((Value && !Value.HasData()) || !Algo::AllOf(Attachments, &FPayload::HasData))
+		{
+			continue;
+		}
+
+		if (!Value && Attachments.IsEmpty())
 		{
 			FScopeLock ScopeLock(&SynchronizationObject);
 			if (bDisabled)
diff --git a/Engine/Source/Developer/DerivedDataCache/Public/DerivedDataCache.h b/Engine/Source/Developer/DerivedDataCache/Public/DerivedDataCache.h
--- a/Engine/Source/Developer/DerivedDataCache/Public/DerivedDataCache.h
+++ b/Engine/Source/Developer/DerivedDataCache/Public/DerivedDataCache.h
@@ -131,6 +131,10 @@ public:
 	 * The callback will always be called for every key, and may be called from an arbitrary thread.
 	 * Records may finish storing in any order, and from multiple threads concurrently.
 	 *
+	 * A cache store is free to interpret a record containing only a key as a request to delete that
+	 * record from the store. Records containing payloads without data must be supported, and should
+	 * reference matching payloads that it already has stored, or otherwise return an error.
+	 *
 	 * @param Records      The cache records to store. Must have a key.
 	 * @param Context      A description of the request. An object path is typically sufficient.
 	 * @param Policy       Flags to control the behavior of the request. See ECachePolicy.
@@ -228,7 +232,7 @@ struct FCacheGetCompleteParams
 	 * The key is always populated. The remainder of the record is populated when Status is Ok.
 	 *
 	 * The value, attachments, and metadata may be skipped based on cache policy flags. When a value
-	 * or attachment has been skipped, it will have a payload but its buffer will be null.
+	 * or attachment has been skipped, it will have a payload but its data will be null.
 	 */
 	FCacheRecord&& Record;
 
@@ -247,7 +251,7 @@ struct FCacheGetPayloadCompleteParams
 	 *
 	 * The ID is always populated.
 	 * The hash and size are populated when Status is Ok.
-	 * The buffer is populated when Status is Ok and the data was not skipped by the cache policy.
+	 * The data is populated when Status is Ok and the data was not skipped by the cache policy.
 	 */
 	FPayload&& Payload;
 
diff --git a/Engine/Source/Developer/DerivedDataCache/Public/DerivedDataCacheRecord.h b/Engine/Source/Developer/DerivedDataCache/Public/DerivedDataCacheRecord.h
--- a/Engine/Source/Developer/DerivedDataCache/Public/DerivedDataCacheRecord.h
+++ b/Engine/Source/Developer/DerivedDataCache/Public/DerivedDataCacheRecord.h
@@ -45,9 +45,9 @@ public:
 	virtual ~ICacheRecordBuilderInternal() = default;
 	virtual void SetMeta(FCbObject&& Meta) = 0;
 	virtual FPayloadId SetValue(const FSharedBuffer& Buffer, const FPayloadId& Id) = 0;
-	virtual FPayloadId SetValue(FPayload&& Payload) = 0;
+	virtual FPayloadId SetValue(const FPayload& Payload) = 0;
 	virtual FPayloadId AddAttachment(const FSharedBuffer& Buffer, const FPayloadId& Id) = 0;
-	virtual FPayloadId AddAttachment(FPayload&& Payload) = 0;
+	virtual FPayloadId AddAttachment(const FPayload& Payload) = 0;
 	virtual FCacheRecord Build() = 0;
 	virtual FRequest BuildAsync(FOnCacheRecordComplete&& OnComplete, EPriority Priority) = 0;
 };
@@ -74,7 +74,7 @@ namespace UE::DerivedData
  *
  * The value, attachments, and metadata are optional and can be skipped when requesting a record.
  * When the value or attachments have been skipped, the record will contain a payload with a null
- * buffer but will otherwise be populated.
+ * data but will otherwise be populated.
  */
 class FCacheRecord
 {
@@ -91,19 +91,19 @@ public:
 	/** Returns the value. Null if no value or requested with ECachePolicy::SkipValue. */
 	inline FSharedBuffer GetValue() const { return Record->GetValue(); }
 
-	/** Returns the value payload. Null if no value. Buffer is null if value was skipped. */
+	/** Returns the value payload. Null if no value. Data is null if value was skipped. */
 	inline const FPayload& GetValuePayload() const { return Record->GetValuePayload(); }
 
 	/** Returns the attachment matching the ID. Null if no match or requested with ECachePolicy::SkipAttachments. */
 	inline FSharedBuffer GetAttachment(const FPayloadId& Id) const { return Record->GetAttachment(Id); }
 
-	/** Returns the attachment payload matching the ID. Null if no match. Buffer is null if attachments were skipped. */
+	/** Returns the attachment payload matching the ID. Null if no match. Data is null if attachments were skipped. */
 	inline const FPayload& GetAttachmentPayload(const FPayloadId& Id) const { return Record->GetAttachmentPayload(Id); }
 
-	/** Returns a view of the attachments. Always available in a non-null cache record, but buffer may be skipped. */
+	/** Returns a view of the attachments. Always available in a non-null cache record, but data may be skipped. */
 	inline TConstArrayView<FPayload> GetAttachmentPayloads() const { return Record->GetAttachmentPayloads(); }
 
-	/** Returns the payload matching the ID, whether value or attachment. Null if no match. Buffer is null if skipped. */
+	/** Returns the payload matching the ID, whether value or attachment. Null if no match. Data is null if skipped. */
 	inline const FPayload& GetPayload(const FPayloadId& Id) const { return Record->GetPayload(Id); }
 
 	/** Whether this is null. */
@@ -168,12 +168,12 @@ public:
 	/**
 	 * Set the value for the cache record.
 	 *
-	 * @param Payload   The value payload, which must have a buffer.
+	 * @param Payload   The payload, which must have data unless it is known to be in the cache.
 	 * @return The ID that was provided. Unique within the scope of the cache record.
 	 */
-	inline FPayloadId SetValue(FPayload&& Payload)
+	inline FPayloadId SetValue(const FPayload& Payload)
 	{
-		return RecordBuilder->SetValue(MoveTemp(Payload));
+		return RecordBuilder->SetValue(Payload);
 	}
 
 	/**
@@ -192,12 +192,12 @@ public:
 	/**
 	 * Add an attachment to the cache record.
 	 *
-	 * @param Payload   The attachment payload, which must have a buffer.
+	 * @param Payload   The payload, which must have data unless it is known to be in the cache.
 	 * @return The ID that was provided. Unique within the scope of the cache record.
 	 */
-	inline FPayloadId AddAttachment(FPayload&& Payload)
+	inline FPayloadId AddAttachment(const FPayload& Payload)
 	{
-		return RecordBuilder->AddAttachment(MoveTemp(Payload));
+		return RecordBuilder->AddAttachment(Payload);
 	}
 
 	/**
diff --git a/Engine/Source/Developer/DerivedDataCache/Public/DerivedDataPayload.h b/Engine/Source/Developer/DerivedDataCache/Public/DerivedDataPayload.h
--- a/Engine/Source/Developer/DerivedDataCache/Public/DerivedDataPayload.h
+++ b/Engine/Source/Developer/DerivedDataCache/Public/DerivedDataPayload.h
@@ -63,8 +63,8 @@ inline const FPayloadId FPayloadId::Null;
 /**
  * A payload is described by an ID and by the hash and size of its raw buffer (uncompressed).
  *
- * Payloads may be constructed with or without a compressed buffer. A payload without a buffer is
- * used to uniquely describe a payload without the overhead of its buffer.
+ * Payloads may be constructed with or without data in the form of a compressed buffer. A payload
+ * without data can be used as a reference to the payload.
  */
 class FPayload
 {
@@ -72,7 +72,7 @@ public:
 	/** Construct a null payload. */
 	FPayload() = default;
 
-	/** Construct a payload with no hash, size, or buffer. */
+	/** Construct a payload with no hash, size, or data. */
 	inline explicit FPayload(const FPayloadId& Id);
 
 	/** Construct a payload from a the hash and size of the raw buffer. */
@@ -94,6 +94,9 @@ public:
 	/** Returns the compressed buffer for the payload. May be null. */
 	inline const FCompressedBuffer& GetData() const { return Data; }
 
+	/** Whether the compressed buffer for the payload is available. */
+	inline bool HasData() const { return !!Data; }
+
 	/** Whether this is null. */
 	inline bool IsNull() const { return Id.IsNull(); }
 	/** Whether this is not null. */
@@ -229,22 +232,6 @@ inline bool operator<(const FPayload& A, const FPayload& B)
 	return !(IdA == IdB) ? IdA < IdB : !(HashA == HashB) ? HashA < HashB : A.GetRawSize() < B.GetRawSize();
 }
 
-/** Compare payloads by their ID. */
-struct FPayloadEqualById
-{
-	inline bool operator()(const FPayload& A, const FPayload& B) const { return A.GetId() == B.GetId(); }
-	inline bool operator()(const FPayload& A, const FPayloadId& B) const { return A.GetId() == B; }
-	inline bool operator()(const FPayloadId& A, const FPayload& B) const { return A == B.GetId(); }
-};
-
-/** Compare payloads by their ID. */
-struct FPayloadLessById
-{
-	inline bool operator()(const FPayload& A, const FPayload& B) const { return A.GetId() < B.GetId(); }
-	inline bool operator()(const FPayload& A, const FPayloadId& B) const { return A.GetId() < B; }
-	inline bool operator()(const FPayloadId& A, const FPayload& B) const { return A < B.GetId(); }
-};
-
 ///////////////////////////////////////////////////////////////////////////////////////////////////////////////////////
 
 } // UE::DerivedData
