DDC: Added UE::DerivedData::FBuildOutput

FBuildOutput and the associated builder type define the output of a derived data build as a collection of payloads and diagnostic messages.

#rnx

diff --git a/Engine/Source/Developer/DerivedDataCache/Private/DerivedDataBuildOutput.cpp b/Engine/Source/Developer/DerivedDataCache/Private/DerivedDataBuildOutput.cpp
new file mode 100644
--- /dev/null
+++ b/Engine/Source/Developer/DerivedDataCache/Private/DerivedDataBuildOutput.cpp
@@ -0,0 +1,301 @@
+// Copyright Epic Games, Inc. All Rights Reserved.
+
+#include "DerivedDataBuildOutput.h"
+
+#include "Algo/BinarySearch.h"
+#include "Containers/StringConv.h"
+#include "Containers/StringView.h"
+#include "Containers/UnrealString.h"
+#include "DerivedDataBuildDefinition.h"
+#include "DerivedDataBuildKey.h"
+#include "DerivedDataBuildPrivate.h"
+#include "DerivedDataCacheRecord.h"
+#include "Misc/StringBuilder.h"
+#include "Serialization/CompactBinary.h"
+#include "Serialization/CompactBinaryWriter.h"
+
+namespace UE::DerivedData::Private
+{
+
+static FStringView LexToString(EBuildDiagnosticLevel Level)
+{
+	switch (Level)
+	{
+	default: return TEXT("Unknown"_SV);
+	case EBuildDiagnosticLevel::Error: return TEXT("Error"_SV);
+	case EBuildDiagnosticLevel::Warning: return TEXT("Warning"_SV);
+	}
+}
+
+static void LexFromString(EBuildDiagnosticLevel& OutLevel, FAnsiStringView String)
+{
+	if (String.Equals("Warning"_ASV, ESearchCase::CaseSensitive))
+	{
+		OutLevel = EBuildDiagnosticLevel::Warning;
+	}
+	else
+	{
+		OutLevel = EBuildDiagnosticLevel::Error;
+	}
+}
+
+///////////////////////////////////////////////////////////////////////////////////////////////////////////////////////
+
+class FBuildOutputBuilderInternal final : public IBuildOutputBuilderInternal
+{
+public:
+	explicit FBuildOutputBuilderInternal(FStringView Name, FStringView Function)
+		: Name(Name)
+		, Function(Function)
+		, bHasError(false)
+		, bHasDiagnostics(false)
+	{
+		DiagnosticsWriter.BeginArray();
+	}
+
+	~FBuildOutputBuilderInternal() final = default;
+
+	bool HasError() const final { return bHasError; }
+	void SetMeta(FCbObject&& InMeta) final { Meta = MoveTemp(InMeta); }
+	void AddDiagnostic(const FBuildDiagnostic& Diagnostic) final;
+	void AddPayload(const FPayload& Payload) final;
+	FBuildOutput Build() final;
+
+	FString Name;
+	FString Function;
+	FCbObject Meta;
+	FCbField Diagnostics;
+	FCbWriter DiagnosticsWriter;
+	TArray<FPayload> Payloads;
+	bool bHasError;
+	bool bHasDiagnostics;
+};
+
+class FBuildOutputInternal final : public IBuildOutputInternal
+{
+public:
+	explicit FBuildOutputInternal(FBuildOutputBuilderInternal&& InOutput);
+	explicit FBuildOutputInternal(FStringView Name, FStringView Function, const FCbObject& InOutput);
+	explicit FBuildOutputInternal(FStringView Name, FStringView Function, const FCacheRecord& InOutput);
+
+	~FBuildOutputInternal() final = default;
+
+	FStringView GetName() const final { return Name; }
+	FStringView GetFunction() const final { return Function; }
+
+	const FCbObject& GetMeta() const final { return Meta; }
+
+	bool HasError() const final;
+	void IterateDiagnostics(TFunctionRef<void (const FBuildDiagnostic& Diagnostic)> Visitor) const final;
+	const FPayload& GetPayload(const FPayloadId& Id) const final;
+	TConstArrayView<FPayload> GetPayloads() const final { return Payloads; }
+
+	void Save(FCbWriter& Writer) const;
+	void Save(FCacheRecordBuilder& RecordBuilder) const;
+
+	inline void AddRef() const final
+	{
+		ReferenceCount.fetch_add(1, std::memory_order_relaxed);
+	}
+
+	inline void Release() const final
+	{
+		if (ReferenceCount.fetch_sub(1, std::memory_order_acq_rel) == 1)
+		{
+			delete this;
+		}
+	}
+
+private:
+	mutable std::atomic<uint32> ReferenceCount{0};
+	FString Name;
+	FString Function;
+	TArray<FPayload> Payloads;
+	FCbField Diagnostics;
+	FCbObject Meta;
+};
+
+///////////////////////////////////////////////////////////////////////////////////////////////////////////////////////
+
+FBuildOutputInternal::FBuildOutputInternal(FBuildOutputBuilderInternal&& InOutput)
+	: Name(MoveTemp(InOutput.Name))
+	, Function(MoveTemp(InOutput.Function))
+	, Payloads(MoveTemp(InOutput.Payloads))
+	, Diagnostics(MoveTemp(InOutput.Diagnostics))
+	, Meta(MoveTemp(InOutput.Meta))
+{
+}
+
+FBuildOutputInternal::FBuildOutputInternal(FStringView Name, FStringView Function, const FCbObject& Output)
+	: Name(Name)
+	, Function(Function)
+{
+	for (FCbFieldView PayloadField : Output["Payloads"_ASV])
+	{
+		const FCbObjectView Payload = PayloadField.AsObjectView();
+		Payloads.Emplace(
+			FPayloadId(Payload["Id"_ASV].AsObjectId().GetView()),
+			Payload["RawHash"_ASV].AsAttachment(),
+			Payload["RawSize"_ASV].AsUInt64());
+	}
+	Diagnostics = Output["Diagnostics"_ASV];
+	Meta = Output["Meta"_ASV].AsObject();
+}
+
+FBuildOutputInternal::FBuildOutputInternal(FStringView Name, FStringView Function, const FCacheRecord& Output)
+	: Name(Name)
+	, Function(Function)
+	, Payloads(Output.GetAttachmentPayloads())
+	, Meta(Output.GetMeta())
+{
+	if (FSharedBuffer Buffer = Output.GetValue())
+	{
+		Diagnostics = FCbObject(MoveTemp(Buffer))["Diagnostics"_ASV];
+	}
+}
+
+bool FBuildOutputInternal::HasError() const
+{
+	for (FCbFieldView Field : Diagnostics.CreateViewIterator())
+	{
+		const FCbObjectView Diagnostic = Field.AsObjectView();
+		EBuildDiagnosticLevel Level;
+		LexFromString(Level, Diagnostic["Level"_ASV].AsString());
+		if (Level == EBuildDiagnosticLevel::Error)
+		{
+			return true;
+		}
+	}
+	return false;
+}
+
+void FBuildOutputInternal::IterateDiagnostics(TFunctionRef<void (const FBuildDiagnostic& Diagnostic)> Visitor) const
+{
+	for (FCbFieldView Field : Diagnostics.CreateViewIterator())
+	{
+		EBuildDiagnosticLevel Level;
+		const FCbObjectView Diagnostic = Field.AsObjectView();
+		LexFromString(Level, Diagnostic["Level"_ASV].AsString());
+		const FAnsiStringView Category = Diagnostic["Category"_ASV].AsString();
+		const FAnsiStringView Message = Diagnostic["Message"_ASV].AsString();
+		Visitor(FBuildDiagnostic{FUTF8ToTCHAR(Category), FUTF8ToTCHAR(Message), Level});
+	}
+}
+
+const FPayload& FBuildOutputInternal::GetPayload(const FPayloadId& Id) const
+{
+	const int32 Index = Algo::BinarySearchBy(Payloads, Id, &FPayload::GetId);
+	return Payloads.IsValidIndex(Index) ? Payloads[Index] : FPayload::Null;
+}
+
+void FBuildOutputInternal::Save(FCbWriter& Writer) const
+{
+	Writer.BeginObject();
+	if (!Payloads.IsEmpty())
+	{
+		Writer.BeginArray("Payloads"_ASV);
+		for (const FPayload& Payload : Payloads)
+		{
+			Writer.BeginObject();
+			Writer.AddObjectId("Id"_ASV, FCbObjectId(Payload.GetId().GetView()));
+			Writer.AddInteger("RawSize"_ASV, Payload.GetRawSize());
+			Writer.AddBinaryAttachment("RawHash"_ASV, Payload.GetRawHash());
+			Writer.EndObject();
+		}
+		Writer.EndArray();
+	}
+	if (Diagnostics)
+	{
+		Writer.AddField("Diagnostics"_ASV, Diagnostics);
+	}
+	if (Meta)
+	{
+		Writer.AddObject("Meta"_ASV, Meta);
+	}
+	Writer.EndObject();
+}
+
+void FBuildOutputInternal::Save(FCacheRecordBuilder& RecordBuilder) const
+{
+	RecordBuilder.SetMeta(FCbObject(Meta));
+	if (Diagnostics)
+	{
+		TCbWriter<128> Value;
+		Value.BeginObject();
+		Value.AddField("Diagnostics"_ASV, Diagnostics);
+		Value.EndObject();
+		RecordBuilder.SetValue(Value.Save().GetBuffer());
+	}
+	for (const FPayload& Payload : Payloads)
+	{
+		RecordBuilder.AddAttachment(Payload);
+	}
+}
+
+///////////////////////////////////////////////////////////////////////////////////////////////////////////////////////
+
+void FBuildOutputBuilderInternal::AddPayload(const FPayload& Payload)
+{
+	checkf(Payload, TEXT("Null payload added in output for build of '%s' by %s."), *Name, *Function);
+	const FPayloadId& Id = Payload.GetId();
+	const int32 Index = Algo::LowerBoundBy(Payloads, Id, &FPayload::GetId);
+	checkf(!(Payloads.IsValidIndex(Index) && Payloads[Index].GetId() == Id),
+		TEXT("Duplicate ID %s used by payload for build of '%s' by %s."), *WriteToString<32>(Id), *Name, *Function);
+	Payloads.Insert(Payload, Index);
+}
+
+void FBuildOutputBuilderInternal::AddDiagnostic(const FBuildDiagnostic& Diagnostic)
+{
+	bHasError |= Diagnostic.Level == EBuildDiagnosticLevel::Error;
+	bHasDiagnostics = true;
+	DiagnosticsWriter.BeginObject();
+	DiagnosticsWriter.AddString("Level"_ASV, LexToString(Diagnostic.Level));
+	DiagnosticsWriter.AddString("Category"_ASV, Diagnostic.Category);
+	DiagnosticsWriter.AddString("Message"_ASV, Diagnostic.Message);
+	DiagnosticsWriter.EndObject();
+}
+
+FBuildOutput FBuildOutputBuilderInternal::Build()
+{
+	DiagnosticsWriter.EndArray();
+	if (bHasDiagnostics)
+	{
+		Diagnostics = DiagnosticsWriter.Save();
+	}
+	DiagnosticsWriter.Reset();
+	return CreateBuildOutput(new FBuildOutputInternal(MoveTemp(*this)));
+}
+
+///////////////////////////////////////////////////////////////////////////////////////////////////////////////////////
+
+FBuildOutput CreateBuildOutput(IBuildOutputInternal* Output)
+{
+	return FBuildOutput(Output);
+}
+
+FBuildOutputBuilder CreateBuildOutputBuilder(IBuildOutputBuilderInternal* OutputBuilder)
+{
+	return FBuildOutputBuilder(OutputBuilder);
+}
+
+const IBuildOutputInternal* GetBuildOutput(const FBuildOutput& Output)
+{
+	return Output.Output;
+}
+
+FBuildOutputBuilder CreateBuildOutput(FStringView Name, FStringView Function)
+{
+	return CreateBuildOutputBuilder(new FBuildOutputBuilderInternal(Name, Function));
+}
+
+FBuildOutput LoadBuildOutput(FStringView Name, FStringView Function, const FCbObject& Output)
+{
+	return CreateBuildOutput(new FBuildOutputInternal(Name, Function, Output));
+}
+
+FBuildOutput LoadBuildOutput(FStringView Name, FStringView Function, const FCacheRecord& Output)
+{
+	return CreateBuildOutput(new FBuildOutputInternal(Name, Function, Output));
+}
+
+} // UE::DerivedData::Private
diff --git a/Engine/Source/Developer/DerivedDataCache/Private/DerivedDataBuildPrivate.h b/Engine/Source/Developer/DerivedDataCache/Private/DerivedDataBuildPrivate.h
--- a/Engine/Source/Developer/DerivedDataCache/Private/DerivedDataBuildPrivate.h
+++ b/Engine/Source/Developer/DerivedDataCache/Private/DerivedDataBuildPrivate.h
@@ -5,17 +5,25 @@
 #include "Containers/StringFwd.h"
 
 class FCbObject;
+class FCbWriter;
 
 namespace UE::DerivedData { class FBuildDefinition; }
 namespace UE::DerivedData { class FBuildDefinitionBuilder; }
+namespace UE::DerivedData { class FBuildOutput; }
+namespace UE::DerivedData { class FBuildOutputBuilder; }
+namespace UE::DerivedData { class FCacheRecord; }
+namespace UE::DerivedData { class FCacheRecordBuilder; }
 
 namespace UE::DerivedData::Private
 {
 
 // Implemented in DerivedDataBuildDefinition.cpp
 FBuildDefinitionBuilder CreateBuildDefinition(FStringView Name, FStringView Function);
-
-// Implemented in DerivedDataBuildDefinition.cpp
 FBuildDefinition LoadBuildDefinition(FStringView Name, FCbObject&& Definition);
 
+// Implemented in DerivedDataBuildOutput.cpp
+FBuildOutputBuilder CreateBuildOutput(FStringView Name, FStringView Function);
+FBuildOutput LoadBuildOutput(FStringView Name, FStringView Function, const FCbObject& Output);
+FBuildOutput LoadBuildOutput(FStringView Name, FStringView Function, const FCacheRecord& Output);
+
 } // UE::DerivedData::Private
diff --git a/Engine/Source/Developer/DerivedDataCache/Private/DerivedDataCache.cpp b/Engine/Source/Developer/DerivedDataCache/Private/DerivedDataCache.cpp
--- a/Engine/Source/Developer/DerivedDataCache/Private/DerivedDataCache.cpp
+++ b/Engine/Source/Developer/DerivedDataCache/Private/DerivedDataCache.cpp
@@ -19,6 +19,7 @@
 #include "DerivedDataBackendInterface.h"
 #include "DerivedDataBuild.h"
 #include "DerivedDataBuildDefinition.h"
+#include "DerivedDataBuildOutput.h"
 #include "DerivedDataBuildPrivate.h"
 #include "DerivedDataCachePrivate.h"
 #include "DerivedDataCacheRecord.h"
@@ -725,6 +726,21 @@ public:
 	{
 		return Private::LoadBuildDefinition(Name, MoveTemp(Definition));
 	}
+
+	FBuildOutputBuilder CreateOutput(FStringView Name, FStringView Function) final
+	{
+		return Private::CreateBuildOutput(Name, Function);
+	}
+
+	FBuildOutput LoadOutput(FStringView Name, FStringView Function, const FCbObject& Output) final
+	{
+		return Private::LoadBuildOutput(Name, Function, Output);
+	}
+
+	FBuildOutput LoadOutput(FStringView Name, FStringView Function, const FCacheRecord& Output) final
+	{
+		return Private::LoadBuildOutput(Name, Function, Output);
+	}
 };
 
 } // UE::DerivedData
diff --git a/Engine/Source/Developer/DerivedDataCache/Public/DerivedDataBuild.h b/Engine/Source/Developer/DerivedDataCache/Public/DerivedDataBuild.h
--- a/Engine/Source/Developer/DerivedDataCache/Public/DerivedDataBuild.h
+++ b/Engine/Source/Developer/DerivedDataCache/Public/DerivedDataBuild.h
@@ -8,6 +8,9 @@ class FCbObject;
 
 namespace UE::DerivedData { class FBuildDefinition; }
 namespace UE::DerivedData { class FBuildDefinitionBuilder; }
+namespace UE::DerivedData { class FBuildOutput; }
+namespace UE::DerivedData { class FBuildOutputBuilder; }
+namespace UE::DerivedData { class FCacheRecord; }
 
 namespace UE::DerivedData
 {
@@ -38,6 +41,24 @@ public:
 	 * @return A valid build definition, or a build definition with an empty key on error.
 	 */
 	virtual FBuildDefinition LoadDefinition(FStringView Name, FCbObject&& Definition) = 0;
+
+	/**
+	 * Create a build output builder.
+	 *
+	 * @param Name       The name by which to identify this output for logging and profiling.
+	 * @param Function   The name of the build function that produced this output.
+	 */
+	virtual FBuildOutputBuilder CreateOutput(FStringView Name, FStringView Function) = 0;
+
+	/**
+	 * Load a build output.
+	 *
+	 * @param Name       The name by which to identify this output for logging and profiling.
+	 * @param Function   The name of the build function that produced this output.
+	 * @param Output     The saved output to load.
+	 */
+	virtual FBuildOutput LoadOutput(FStringView Name, FStringView Function, const FCbObject& Output) = 0;
+	virtual FBuildOutput LoadOutput(FStringView Name, FStringView Function, const FCacheRecord& Output) = 0;
 };
 
 } // UE::DerivedData
diff --git a/Engine/Source/Developer/DerivedDataCache/Public/DerivedDataBuildOutput.h b/Engine/Source/Developer/DerivedDataCache/Public/DerivedDataBuildOutput.h
new file mode 100644
--- /dev/null
+++ b/Engine/Source/Developer/DerivedDataCache/Public/DerivedDataBuildOutput.h
@@ -0,0 +1,200 @@
+// Copyright Epic Games, Inc. All Rights Reserved.
+
+#pragma once
+
+#include "CoreTypes.h"
+#include "Containers/ArrayView.h"
+#include "Containers/StringView.h"
+#include "Templates/RefCounting.h"
+#include "Templates/UniquePtr.h"
+
+class FCbObject;
+class FCbWriter;
+
+template <typename FuncType> class TFunctionRef;
+
+namespace UE::DerivedData { class FBuildOutput; }
+namespace UE::DerivedData { class FBuildOutputBuilder; }
+namespace UE::DerivedData { class FCacheRecordBuilder; }
+namespace UE::DerivedData { class FPayload; }
+namespace UE::DerivedData { struct FBuildDiagnostic; }
+namespace UE::DerivedData { struct FPayloadId; }
+
+namespace UE::DerivedData::Private
+{
+
+class IBuildOutputInternal
+{
+public:
+	virtual ~IBuildOutputInternal() = default;
+	virtual FStringView GetName() const = 0;
+	virtual FStringView GetFunction() const = 0;
+	virtual const FCbObject& GetMeta() const = 0;
+	virtual bool HasError() const = 0;
+	virtual void IterateDiagnostics(TFunctionRef<void (const FBuildDiagnostic& Diagnostic)> Visitor) const = 0;
+	virtual const FPayload& GetPayload(const FPayloadId& Id) const = 0;
+	virtual TConstArrayView<FPayload> GetPayloads() const = 0;
+	virtual void Save(FCbWriter& Writer) const = 0;
+	virtual void Save(FCacheRecordBuilder& RecordBuilder) const = 0;
+	virtual void AddRef() const = 0;
+	virtual void Release() const = 0;
+};
+
+FBuildOutput CreateBuildOutput(IBuildOutputInternal* Output);
+const IBuildOutputInternal* GetBuildOutput(const FBuildOutput& Output);
+
+class IBuildOutputBuilderInternal
+{
+public:
+	virtual ~IBuildOutputBuilderInternal() = default;
+	virtual void SetMeta(FCbObject&& Meta) = 0;
+	virtual bool HasError() const = 0;
+	virtual void AddDiagnostic(const FBuildDiagnostic& Diagnostic) = 0;
+	virtual void AddPayload(const FPayload& Payload) = 0;
+	virtual FBuildOutput Build() = 0;
+};
+
+FBuildOutputBuilder CreateBuildOutputBuilder(IBuildOutputBuilderInternal* OutputBuilder);
+
+} // UE::DerivedData::Private
+
+namespace UE::DerivedData
+{
+
+/** Level of severity for build diagnostics. */
+enum class EBuildDiagnosticLevel : uint8
+{
+	/** Errors always indicates a failure of the corresponding build. */
+	Error,
+	/** Warnings are expected to be actionable issues found while executing a build. */
+	Warning,
+};
+
+/** A build diagnostic is a message logged by a build. */
+struct FBuildDiagnostic
+{
+	/** The name of the category of the diagnostic. */
+	FStringView Category;
+	/** The message of the diagnostic. */
+	FStringView Message;
+	/** The level of severity of the diagnostic. */
+	EBuildDiagnosticLevel Level;
+};
+
+/**
+ * A build output contains the payloads and diagnostics produced by a build.
+ *
+ * The output will not contain any payloads if it has any errors.
+ *
+ * The output can be requested without data, which means the payloads will have null data.
+ */
+class FBuildOutput
+{
+public:
+	/** Returns the name of the build definition that produced this output. */
+	inline FStringView GetName() const { return Output->GetName(); }
+
+	/** Returns the name of the build function that produced this output. */
+	inline FStringView GetFunction() const { return Output->GetFunction(); }
+
+	/** Returns the optional metadata. */
+	inline const FCbObject& GetMeta() const { return Output->GetMeta(); }
+
+	/** Returns whether the output has any error diagnostics. */
+	inline bool HasError() const { return Output->HasError(); }
+
+	/** Returns the payloads in the output in order by ID. */
+	inline TConstArrayView<FPayload> GetPayloads() const { return Output->GetPayloads(); }
+
+	/** Returns the payload matching the ID. Null if no match. Buffer is null if skipped. */
+	inline const FPayload& GetPayload(const FPayloadId& Id) const { return Output->GetPayload(Id); }
+
+	/** Visits every diagnostic in the order it was recorded. */
+	inline void IterateDiagnostics(TFunctionRef<void (const FBuildDiagnostic& Diagnostic)> Visitor) const
+	{
+		Output->IterateDiagnostics(Visitor);
+	}
+
+	/** Saves the build output to a compact binary object with payloads as attachments. */
+	void Save(FCbWriter& Writer) const
+	{
+		Output->Save(Writer);
+	}
+
+	/** Saves the build output to a cache record. */
+	void Save(FCacheRecordBuilder& RecordBuilder) const
+	{
+		Output->Save(RecordBuilder);
+	}
+
+private:
+	friend FBuildOutput Private::CreateBuildOutput(Private::IBuildOutputInternal* Output);
+	friend const Private::IBuildOutputInternal* Private::GetBuildOutput(const FBuildOutput& Output);
+
+	inline explicit FBuildOutput(Private::IBuildOutputInternal* InOutput)
+		: Output(InOutput)
+	{
+	}
+
+	TRefCountPtr<Private::IBuildOutputInternal> Output;
+};
+
+/**
+ * A build output builder is used to construct a build output.
+ *
+ * Create using IBuild::CreateOutput() which must be given a build definition.
+ *
+ * @see FBuildOutput
+ */
+class FBuildOutputBuilder
+{
+public:
+	/** Returns whether the output has any error diagnostics. */
+	inline bool HasError() const
+	{
+		return OutputBuilder->HasError();
+	}
+
+	/** Set the metadata for the build output. Holds a reference and is cloned if not owned. */
+	inline void SetMeta(FCbObject&& Meta)
+	{
+		return OutputBuilder->SetMeta(MoveTemp(Meta));
+	}
+
+	/** Add a payload to the output. The ID must be unique in this output. */
+	inline void AddPayload(const FPayload& Payload)
+	{
+		OutputBuilder->AddPayload(Payload);
+	}
+
+	/** Add an error diagnostic to the output. */
+	inline void AddError(FStringView Category, FStringView Message)
+	{
+		OutputBuilder->AddDiagnostic({Category, Message, EBuildDiagnosticLevel::Error});
+	}
+
+	/** Add a warning diagnostic to the output. */
+	inline void AddWarning(FStringView Category, FStringView Message)
+	{
+		OutputBuilder->AddDiagnostic({Category, Message, EBuildDiagnosticLevel::Warning});
+	}
+
+	/** Build a build output, which makes this builder subsequently unusable. */
+	inline FBuildOutput Build()
+	{
+		return OutputBuilder->Build();
+	}
+
+private:
+	friend FBuildOutputBuilder Private::CreateBuildOutputBuilder(Private::IBuildOutputBuilderInternal* OutputBuilder);
+
+	/** Construct a build output builder. Use IBuild::CreateOutput(). */
+	inline explicit FBuildOutputBuilder(Private::IBuildOutputBuilderInternal* InOutputBuilder)
+		: OutputBuilder(InOutputBuilder)
+	{
+	}
+
+	TUniquePtr<Private::IBuildOutputBuilderInternal> OutputBuilder;
+};
+
+} // UE::DerivedData
