DDC: Forward constructor arguments through TBuildFunctionFactory

#rnx

diff --git a/Engine/Source/Developer/DerivedDataCache/Public/DerivedDataBuildFunctionFactory.h b/Engine/Source/Developer/DerivedDataCache/Public/DerivedDataBuildFunctionFactory.h
--- a/Engine/Source/Developer/DerivedDataCache/Public/DerivedDataBuildFunctionFactory.h
+++ b/Engine/Source/Developer/DerivedDataCache/Public/DerivedDataBuildFunctionFactory.h
@@ -29,7 +29,9 @@ public:
  * A build function must be registered by a build function factory before it can execute a build.
  * Register a function in the source file that implements it or in the corresponding module.
  *
- * Example: static const TBuildFunctionFactory<FExampleFunction> ExampleFunctionFactory;
+ * Examples:
+ * static const TBuildFunctionFactory<FExampleFunction> ExampleFunctionFactory;
+ * static const TBuildFunctionFactory<TExampleFunction<FType>> ExampleFunctionFactory(Name, Version);
  */
 template <typename FunctionType>
 class TBuildFunctionFactory final : private IBuildFunctionFactory
@@ -37,7 +39,9 @@ class TBuildFunctionFactory final : private IBuildFunctionFactory
 	static_assert(sizeof(FunctionType) == sizeof(IBuildFunction), "IBuildFunction must be pure and maintain no state.");
 
 public:
-	TBuildFunctionFactory()
+	template <typename... ArgTypes>
+	explicit TBuildFunctionFactory(ArgTypes&&... Args)
+		: Function(Forward<ArgTypes>(Args)...)
 	{
 		IModularFeatures::Get().RegisterModularFeature(GetFeatureName(), this);
 	}
