DDC: Organized functions for consistency and fixed includes of Templates/Function.h

#rnx

diff --git a/Engine/Source/Developer/DerivedDataCache/Private/DerivedDataBuildDefinition.cpp b/Engine/Source/Developer/DerivedDataCache/Private/DerivedDataBuildDefinition.cpp
--- a/Engine/Source/Developer/DerivedDataCache/Private/DerivedDataBuildDefinition.cpp
+++ b/Engine/Source/Developer/DerivedDataCache/Private/DerivedDataBuildDefinition.cpp
@@ -13,7 +13,6 @@
 #include "Misc/TVariant.h"
 #include "Serialization/CompactBinary.h"
 #include "Serialization/CompactBinaryWriter.h"
-#include "Templates/Function.h"
 #include <atomic>
 
 namespace UE::DerivedData::Private
diff --git a/Engine/Source/Developer/DerivedDataCache/Private/DerivedDataBuildOutput.cpp b/Engine/Source/Developer/DerivedDataCache/Private/DerivedDataBuildOutput.cpp
--- a/Engine/Source/Developer/DerivedDataCache/Private/DerivedDataBuildOutput.cpp
+++ b/Engine/Source/Developer/DerivedDataCache/Private/DerivedDataBuildOutput.cpp
@@ -55,10 +55,10 @@ public:
 
 	~FBuildOutputBuilderInternal() final = default;
 
-	bool HasError() const final { return bHasError; }
 	void SetMeta(FCbObject&& InMeta) final { Meta = MoveTemp(InMeta); }
-	void AddDiagnostic(const FBuildDiagnostic& Diagnostic) final;
 	void AddPayload(const FPayload& Payload) final;
+	void AddDiagnostic(const FBuildDiagnostic& Diagnostic) final;
+	bool HasError() const final { return bHasError; }
 	FBuildOutput Build() final;
 
 	FString Name;
@@ -85,10 +85,10 @@ public:
 
 	const FCbObject& GetMeta() const final { return Meta; }
 
-	bool HasError() const final;
-	void IterateDiagnostics(TFunctionRef<void (const FBuildDiagnostic& Diagnostic)> Visitor) const final;
 	const FPayload& GetPayload(const FPayloadId& Id) const final;
 	TConstArrayView<FPayload> GetPayloads() const final { return Payloads; }
+	void IterateDiagnostics(TFunctionRef<void (const FBuildDiagnostic& Diagnostic)> Visitor) const final;
+	bool HasError() const final;
 
 	void Save(FCbWriter& Writer) const;
 	void Save(FCacheRecordBuilder& RecordBuilder) const;
@@ -154,19 +154,10 @@ FBuildOutputInternal::FBuildOutputInternal(FStringView Name, FStringView Functio
 	}
 }
 
-bool FBuildOutputInternal::HasError() const
+const FPayload& FBuildOutputInternal::GetPayload(const FPayloadId& Id) const
 {
-	for (FCbFieldView Field : Diagnostics.CreateViewIterator())
-	{
-		const FCbObjectView Diagnostic = Field.AsObjectView();
-		EBuildDiagnosticLevel Level;
-		LexFromString(Level, Diagnostic["Level"_ASV].AsString());
-		if (Level == EBuildDiagnosticLevel::Error)
-		{
-			return true;
-		}
-	}
-	return false;
+	const int32 Index = Algo::BinarySearchBy(Payloads, Id, &FPayload::GetId);
+	return Payloads.IsValidIndex(Index) ? Payloads[Index] : FPayload::Null;
 }
 
 void FBuildOutputInternal::IterateDiagnostics(TFunctionRef<void (const FBuildDiagnostic& Diagnostic)> Visitor) const
@@ -182,10 +173,19 @@ void FBuildOutputInternal::IterateDiagnostics(TFunctionRef<void (const FBuildDia
 	}
 }
 
-const FPayload& FBuildOutputInternal::GetPayload(const FPayloadId& Id) const
+bool FBuildOutputInternal::HasError() const
 {
-	const int32 Index = Algo::BinarySearchBy(Payloads, Id, &FPayload::GetId);
-	return Payloads.IsValidIndex(Index) ? Payloads[Index] : FPayload::Null;
+	for (FCbFieldView Field : Diagnostics.CreateViewIterator())
+	{
+		const FCbObjectView Diagnostic = Field.AsObjectView();
+		EBuildDiagnosticLevel Level;
+		LexFromString(Level, Diagnostic["Level"_ASV].AsString());
+		if (Level == EBuildDiagnosticLevel::Error)
+		{
+			return true;
+		}
+	}
+	return false;
 }
 
 void FBuildOutputInternal::Save(FCbWriter& Writer) const
diff --git a/Engine/Source/Developer/DerivedDataCache/Public/DerivedDataBuildDefinition.h b/Engine/Source/Developer/DerivedDataCache/Public/DerivedDataBuildDefinition.h
--- a/Engine/Source/Developer/DerivedDataCache/Public/DerivedDataBuildDefinition.h
+++ b/Engine/Source/Developer/DerivedDataCache/Public/DerivedDataBuildDefinition.h
@@ -4,6 +4,7 @@
 
 #include "CoreTypes.h"
 #include "Containers/StringView.h"
+#include "Templates/Function.h"
 #include "Templates/RefCounting.h"
 #include "Templates/UniquePtr.h"
 
@@ -11,7 +12,6 @@ class FCbObject;
 class FCbWriter;
 struct FGuid;
 struct FIoHash;
-template <typename FuncType> class TFunctionRef;
 
 namespace UE::DerivedData { class FBuildDefinition; }
 namespace UE::DerivedData { class FBuildDefinitionBuilder; }
diff --git a/Engine/Source/Developer/DerivedDataCache/Public/DerivedDataBuildFunctionFactory.h b/Engine/Source/Developer/DerivedDataCache/Public/DerivedDataBuildFunctionFactory.h
--- a/Engine/Source/Developer/DerivedDataCache/Public/DerivedDataBuildFunctionFactory.h
+++ b/Engine/Source/Developer/DerivedDataCache/Public/DerivedDataBuildFunctionFactory.h
@@ -9,7 +9,7 @@
 namespace UE::DerivedData
 {
 
-/** Base for the build function factory. DO NOT USE DIRECTLY. Use TBuildFunctionFactory. */
+/** DO NOT USE DIRECTLY. Base for the build function factory. Use TBuildFunctionFactory. */
 class IBuildFunctionFactory : public IModularFeature
 {
 public:
diff --git a/Engine/Source/Developer/DerivedDataCache/Public/DerivedDataBuildOutput.h b/Engine/Source/Developer/DerivedDataCache/Public/DerivedDataBuildOutput.h
--- a/Engine/Source/Developer/DerivedDataCache/Public/DerivedDataBuildOutput.h
+++ b/Engine/Source/Developer/DerivedDataCache/Public/DerivedDataBuildOutput.h
@@ -5,14 +5,13 @@
 #include "CoreTypes.h"
 #include "Containers/ArrayView.h"
 #include "Containers/StringView.h"
+#include "Templates/Function.h"
 #include "Templates/RefCounting.h"
 #include "Templates/UniquePtr.h"
 
 class FCbObject;
 class FCbWriter;
 
-template <typename FuncType> class TFunctionRef;
-
 namespace UE::DerivedData { class FBuildOutput; }
 namespace UE::DerivedData { class FBuildOutputBuilder; }
 namespace UE::DerivedData { class FCacheRecordBuilder; }
@@ -30,10 +29,10 @@ public:
 	virtual FStringView GetName() const = 0;
 	virtual FStringView GetFunction() const = 0;
 	virtual const FCbObject& GetMeta() const = 0;
-	virtual bool HasError() const = 0;
-	virtual void IterateDiagnostics(TFunctionRef<void (const FBuildDiagnostic& Diagnostic)> Visitor) const = 0;
 	virtual const FPayload& GetPayload(const FPayloadId& Id) const = 0;
 	virtual TConstArrayView<FPayload> GetPayloads() const = 0;
+	virtual void IterateDiagnostics(TFunctionRef<void (const FBuildDiagnostic& Diagnostic)> Visitor) const = 0;
+	virtual bool HasError() const = 0;
 	virtual void Save(FCbWriter& Writer) const = 0;
 	virtual void Save(FCacheRecordBuilder& RecordBuilder) const = 0;
 	virtual void AddRef() const = 0;
@@ -48,9 +47,9 @@ class IBuildOutputBuilderInternal
 public:
 	virtual ~IBuildOutputBuilderInternal() = default;
 	virtual void SetMeta(FCbObject&& Meta) = 0;
-	virtual bool HasError() const = 0;
-	virtual void AddDiagnostic(const FBuildDiagnostic& Diagnostic) = 0;
 	virtual void AddPayload(const FPayload& Payload) = 0;
+	virtual void AddDiagnostic(const FBuildDiagnostic& Diagnostic) = 0;
+	virtual bool HasError() const = 0;
 	virtual FBuildOutput Build() = 0;
 };
 
@@ -100,21 +99,21 @@ public:
 	/** Returns the optional metadata. */
 	inline const FCbObject& GetMeta() const { return Output->GetMeta(); }
 
-	/** Returns whether the output has any error diagnostics. */
-	inline bool HasError() const { return Output->HasError(); }
+	/** Returns the payload matching the ID. Null if no match. Buffer is null if skipped. */
+	inline const FPayload& GetPayload(const FPayloadId& Id) const { return Output->GetPayload(Id); }
 
 	/** Returns the payloads in the output in order by ID. */
 	inline TConstArrayView<FPayload> GetPayloads() const { return Output->GetPayloads(); }
 
-	/** Returns the payload matching the ID. Null if no match. Buffer is null if skipped. */
-	inline const FPayload& GetPayload(const FPayloadId& Id) const { return Output->GetPayload(Id); }
-
 	/** Visits every diagnostic in the order it was recorded. */
 	inline void IterateDiagnostics(TFunctionRef<void (const FBuildDiagnostic& Diagnostic)> Visitor) const
 	{
 		Output->IterateDiagnostics(Visitor);
 	}
 
+	/** Returns whether the output has any error diagnostics. */
+	inline bool HasError() const { return Output->HasError(); }
+
 	/** Saves the build output to a compact binary object with payloads as attachments. */
 	void Save(FCbWriter& Writer) const
 	{
@@ -149,12 +148,6 @@ private:
 class FBuildOutputBuilder
 {
 public:
-	/** Returns whether the output has any error diagnostics. */
-	inline bool HasError() const
-	{
-		return OutputBuilder->HasError();
-	}
-
 	/** Set the metadata for the build output. Holds a reference and is cloned if not owned. */
 	inline void SetMeta(FCbObject&& Meta)
 	{
@@ -179,6 +172,12 @@ public:
 		OutputBuilder->AddDiagnostic({Category, Message, EBuildDiagnosticLevel::Warning});
 	}
 
+	/** Returns whether the output has any error diagnostics. */
+	inline bool HasError() const
+	{
+		return OutputBuilder->HasError();
+	}
+
 	/** Build a build output, which makes this builder subsequently unusable. */
 	inline FBuildOutput Build()
 	{
diff --git a/Engine/Source/Developer/DerivedDataCache/Public/DerivedDataCache.h b/Engine/Source/Developer/DerivedDataCache/Public/DerivedDataCache.h
--- a/Engine/Source/Developer/DerivedDataCache/Public/DerivedDataCache.h
+++ b/Engine/Source/Developer/DerivedDataCache/Public/DerivedDataCache.h
@@ -6,6 +6,7 @@
 #include "DerivedDataCacheKey.h"
 #include "DerivedDataRequest.h"
 #include "Misc/EnumClassFlags.h"
+#include "Templates/Function.h"
 
 namespace UE::DerivedData
 {
diff --git a/Engine/Source/Developer/DerivedDataCache/Public/DerivedDataCacheRecord.h b/Engine/Source/Developer/DerivedDataCache/Public/DerivedDataCacheRecord.h
--- a/Engine/Source/Developer/DerivedDataCache/Public/DerivedDataCacheRecord.h
+++ b/Engine/Source/Developer/DerivedDataCache/Public/DerivedDataCacheRecord.h
@@ -6,11 +6,11 @@
 #include "DerivedDataPayload.h"
 #include "DerivedDataRequest.h"
 #include "Memory/SharedBuffer.h"
+#include "Templates/Function.h"
 #include "Templates/RefCounting.h"
 #include "Templates/UniquePtr.h"
 
 class FCbObject;
-template <typename FuncType> class TUniqueFunction;
 
 namespace UE::DerivedData { class FCacheRecord; }
 namespace UE::DerivedData { class FCacheRecordBuilder; }
