DDC: Added UE::DerivedData::FBuildAction

FBuildAction and the associated builder type define the exact action to be executed, including the function and its version, and the inputs referenced by hash.

#rnx

diff --git a/Engine/Source/Developer/DerivedDataCache/Private/DerivedDataBuildAction.cpp b/Engine/Source/Developer/DerivedDataCache/Private/DerivedDataBuildAction.cpp
new file mode 100644
--- /dev/null
+++ b/Engine/Source/Developer/DerivedDataCache/Private/DerivedDataBuildAction.cpp
@@ -0,0 +1,251 @@
+// Copyright Epic Games, Inc. All Rights Reserved.
+
+#include "DerivedDataBuildAction.h"
+
+#include "Algo/AllOf.h"
+#include "Containers/Map.h"
+#include "Containers/StringConv.h"
+#include "Containers/StringView.h"
+#include "Containers/UnrealString.h"
+#include "DerivedDataBuildPrivate.h"
+#include "Misc/Guid.h"
+#include "Misc/TVariant.h"
+#include "Serialization/CompactBinary.h"
+#include "Serialization/CompactBinaryWriter.h"
+#include "Templates/Tuple.h"
+#include <atomic>
+
+namespace UE::DerivedData::Private
+{
+
+class FBuildActionBuilderInternal final : public IBuildActionBuilderInternal
+{
+public:
+	inline FBuildActionBuilderInternal(
+		FStringView InName,
+		FStringView InFunction,
+		const FGuid& InFunctionVersion,
+		const FGuid& InBuildSystemVersion)
+		: Name(InName)
+		, Function(InFunction)
+		, FunctionVersion(InFunctionVersion)
+		, BuildSystemVersion(InBuildSystemVersion)
+	{
+		checkf(!Name.IsEmpty(), TEXT("A build action requires a non-empty name."));
+		checkf(!Function.IsEmpty(), TEXT("A build action requires a non-empty function name for build of '%s'."), *Name);
+		checkf(Algo::AllOf(Function, FChar::IsAlnum), TEXT("Action requires an alphanumeric function name for ")
+			TEXT("build of '%s' by %s."), *Name, *Function);
+	}
+
+	~FBuildActionBuilderInternal() final = default;
+
+	void AddConstant(FStringView Key, const FCbObject& Value) final
+	{
+		Add(Key, Value);
+	}
+
+	void AddInput(FStringView Key, const FIoHash& RawHash, uint64 RawSize) final
+	{
+		Add(Key, MakeTuple(RawHash, RawSize));
+	}
+
+	FBuildAction Build() final;
+
+	using InputType = TVariant<FCbObject, TTuple<FIoHash, uint64>>;
+
+	FString Name;
+	FString Function;
+	FGuid FunctionVersion;
+	FGuid BuildSystemVersion;
+	TMap<FString, InputType> Inputs;
+
+private:
+	template <typename ValueType>
+	inline void Add(FStringView Key, const ValueType& Value)
+	{
+		const uint32 KeyHash = GetTypeHash(Key);
+		checkf(!Key.IsEmpty(), TEXT("Empty key used in action for build of '%s' by %s."), *Name, *Function);
+		checkf(!Inputs.ContainsByHash(KeyHash, Key), TEXT("Duplicate key '%.*s' used in action for "),
+			TEXT("build of '%s' by %s."), Key.Len(), Key.GetData(), *Name, *Function);
+		Inputs.EmplaceByHash(KeyHash, Key, InputType(TInPlaceType<ValueType>(), Value));
+	}
+};
+
+///////////////////////////////////////////////////////////////////////////////////////////////////////////////////////
+
+class FBuildActionInternal final : public IBuildActionInternal
+{
+public:
+	explicit FBuildActionInternal(FBuildActionBuilderInternal&& ActionBuilder);
+	explicit FBuildActionInternal(FStringView Name, FCbObject&& Action);
+
+	~FBuildActionInternal() final = default;
+
+	const FBuildActionKey& GetKey() const final { return Key; }
+
+	FStringView GetName() const final { return Name; }
+	FStringView GetFunction() const final { return Function; }
+	const FGuid& GetFunctionVersion() const final { return FunctionVersion; }
+	const FGuid& GetBuildSystemVersion() const final { return BuildSystemVersion; }
+
+	void IterateConstants(TFunctionRef<void (FStringView Key, FCbObject&& Value)> Visitor) const final;
+	void IterateInputs(TFunctionRef<void (FStringView Key, const FIoHash& RawHash, uint64 RawSize)> Visitor) const final;
+
+	void Save(FCbWriter& Writer) const final;
+
+	inline void AddRef() const final
+	{
+		ReferenceCount.fetch_add(1, std::memory_order_relaxed);
+	}
+
+	inline void Release() const final
+	{
+		if (ReferenceCount.fetch_sub(1, std::memory_order_acq_rel) == 1)
+		{
+			delete this;
+		}
+	}
+
+private:
+	mutable std::atomic<uint32> ReferenceCount{0};
+	FString Name;
+	FString Function;
+	FGuid FunctionVersion;
+	FGuid BuildSystemVersion;
+	FBuildActionKey Key;
+	FCbObject Action;
+};
+
+///////////////////////////////////////////////////////////////////////////////////////////////////////////////////////
+
+FBuildActionInternal::FBuildActionInternal(FBuildActionBuilderInternal&& ActionBuilder)
+	: Name(MoveTemp(ActionBuilder.Name))
+	, Function(MoveTemp(ActionBuilder.Function))
+	, FunctionVersion(MoveTemp(ActionBuilder.FunctionVersion))
+	, BuildSystemVersion(MoveTemp(ActionBuilder.BuildSystemVersion))
+{
+	ActionBuilder.Inputs.KeySort(TLess<>());
+
+	bool bHasConstants = false;
+	bool bHasInputs = false;
+
+	for (const TPair<FString, FBuildActionBuilderInternal::InputType>& Pair : ActionBuilder.Inputs)
+	{
+		if (Pair.Value.IsType<FCbObject>())
+		{
+			bHasConstants = true;
+		}
+		else if (Pair.Value.IsType<TTuple<FIoHash, uint64>>())
+		{
+			bHasInputs = true;
+		}
+	}
+
+	TCbWriter<2048> Writer;
+	Writer.BeginObject();
+	Writer.AddString("Function"_ASV, Function);
+	Writer.AddUuid("FunctionVersion"_ASV, FunctionVersion);
+	Writer.AddUuid("BuildSystemVersion"_ASV, BuildSystemVersion);
+
+	if (bHasConstants)
+	{
+		Writer.BeginObject("Constants"_ASV);
+		for (const TPair<FString, FBuildActionBuilderInternal::InputType>& Pair : ActionBuilder.Inputs)
+		{
+			if (Pair.Value.IsType<FCbObject>())
+			{
+				Writer.AddObject(FTCHARToUTF8(Pair.Key), Pair.Value.Get<FCbObject>());
+			}
+		}
+		Writer.EndObject();
+	}
+
+	if (bHasInputs)
+	{
+		Writer.BeginObject("Inputs"_ASV);
+		for (const TPair<FString, FBuildActionBuilderInternal::InputType>& Pair : ActionBuilder.Inputs)
+		{
+			if (Pair.Value.IsType<TTuple<FIoHash, uint64>>())
+			{
+				const TTuple<FIoHash, uint64>& Input = Pair.Value.Get<TTuple<FIoHash, uint64>>();
+				Writer.BeginObject(FTCHARToUTF8(Pair.Key));
+				Writer.AddBinaryAttachment("RawHash"_ASV, Input.Get<FIoHash>());
+				Writer.AddInteger("RawSize"_ASV, Input.Get<uint64>());
+				Writer.EndObject();
+			}
+		}
+		Writer.EndObject();
+	}
+
+	Writer.EndObject();
+	Action = Writer.Save().AsObject();
+	Key = FBuildActionKey{Action.GetHash()};
+}
+
+FBuildActionInternal::FBuildActionInternal(FStringView InName, FCbObject&& InAction)
+	: Name(InName)
+	, Function(InAction.FindView("Function"_ASV).AsString())
+	, FunctionVersion(InAction.FindView("FunctionVersion"_ASV).AsUuid())
+	, BuildSystemVersion(InAction.FindView("BuildSystemVersion"_ASV).AsUuid())
+	, Action(MoveTemp(InAction))
+{
+	checkf(!InName.IsEmpty(), TEXT("A build action requires a non-empty name."));
+	Action.MakeOwned();
+	if (const bool bIsValid = Action && Algo::AllOf(Function, FChar::IsAlnum))
+	{
+		Key = FBuildActionKey{Action.GetHash()};
+	}
+}
+
+void FBuildActionInternal::IterateConstants(TFunctionRef<void (FStringView Key, FCbObject&& Value)> Visitor) const
+{
+	for (FCbField Field : Action["Constants"_ASV])
+	{
+		Visitor(FUTF8ToTCHAR(Field.GetName()), Field.AsObject());
+	}
+}
+
+void FBuildActionInternal::IterateInputs(TFunctionRef<void (FStringView Key, const FIoHash& RawHash, uint64 RawSize)> Visitor) const
+{
+	for (FCbFieldView Field : Action.FindView("Inputs"_ASV))
+	{
+		FCbObjectView Input = Field.AsObjectView();
+		Visitor(FUTF8ToTCHAR(Field.GetName()), Input["RawHash"_ASV].AsHash(), Input["RawSize"_ASV].AsUInt64());
+	}
+}
+
+void FBuildActionInternal::Save(FCbWriter& Writer) const
+{
+	Writer.AddObject(Action);
+}
+
+///////////////////////////////////////////////////////////////////////////////////////////////////////////////////////
+
+FBuildAction FBuildActionBuilderInternal::Build()
+{
+	return CreateBuildAction(new FBuildActionInternal(MoveTemp(*this)));
+}
+
+///////////////////////////////////////////////////////////////////////////////////////////////////////////////////////
+
+FBuildAction CreateBuildAction(IBuildActionInternal* Action)
+{
+	return FBuildAction(Action);
+}
+
+FBuildActionBuilder CreateBuildActionBuilder(IBuildActionBuilderInternal* ActionBuilder)
+{
+	return FBuildActionBuilder(ActionBuilder);
+}
+
+FBuildActionBuilder CreateBuildAction(FStringView Name, FStringView Function, const FGuid& FunctionVersion, const FGuid& BuildSystemVersion)
+{
+	return CreateBuildActionBuilder(new FBuildActionBuilderInternal(Name, Function, FunctionVersion, BuildSystemVersion));
+}
+
+FBuildAction LoadBuildAction(FStringView Name, FCbObject&& Action)
+{
+	return CreateBuildAction(new FBuildActionInternal(Name, MoveTemp(Action)));
+}
+
+} // UE::DerivedData::Private
diff --git a/Engine/Source/Developer/DerivedDataCache/Private/DerivedDataBuildDefinition.cpp b/Engine/Source/Developer/DerivedDataCache/Private/DerivedDataBuildDefinition.cpp
--- a/Engine/Source/Developer/DerivedDataCache/Private/DerivedDataBuildDefinition.cpp
+++ b/Engine/Source/Developer/DerivedDataCache/Private/DerivedDataBuildDefinition.cpp
@@ -93,6 +93,9 @@ public:
 	FStringView GetName() const final { return Name; }
 	FStringView GetFunction() const final { return Function; }
 
+	bool HasConstants() const final;
+	bool HasInputs() const final;
+
 	void IterateConstants(TFunctionRef<void (FStringView Key, FCbObject&& Value)> Visitor) const final;
 	void IterateInputBuilds(TFunctionRef<void (FStringView Key, const FBuildPayloadKey& PayloadKey)> Visitor) const final;
 	void IterateInputBulkData(TFunctionRef<void (FStringView Key, const FGuid& BulkDataId)> Visitor) const final;
@@ -160,6 +163,8 @@ FBuildDefinitionInternal::FBuildDefinitionInternal(FBuildDefinitionBuilderIntern
 		}
 	}
 
+	const bool bHasInputs = bHasBuilds | bHasBulkData | bHasFiles | bHasHashes;
+
 	TCbWriter<2048> Writer;
 	Writer.BeginObject();
 	Writer.AddString("Function"_ASV, Function);
@@ -177,6 +182,11 @@ FBuildDefinitionInternal::FBuildDefinitionInternal(FBuildDefinitionBuilderIntern
 		Writer.EndObject();
 	}
 
+	if (bHasInputs)
+	{
+		Writer.BeginObject("Inputs"_ASV);
+	}
+
 	if (bHasBuilds)
 	{
 		Writer.BeginObject("Builds"_ASV);
@@ -233,6 +243,11 @@ FBuildDefinitionInternal::FBuildDefinitionInternal(FBuildDefinitionBuilderIntern
 		Writer.EndObject();
 	}
 
+	if (bHasInputs)
+	{
+		Writer.EndObject();
+	}
+
 	Writer.EndObject();
 	Definition = Writer.Save().AsObject();
 	Key = FBuildKey{Definition.GetHash()};
@@ -251,6 +266,16 @@ FBuildDefinitionInternal::FBuildDefinitionInternal(FStringView InName, FCbObject
 	}
 }
 
+bool FBuildDefinitionInternal::HasConstants() const
+{
+	return Definition["Constants"_ASV].HasValue();
+}
+
+bool FBuildDefinitionInternal::HasInputs() const
+{
+	return Definition["Inputs"_ASV].HasValue();
+}
+
 void FBuildDefinitionInternal::IterateConstants(TFunctionRef<void (FStringView Key, FCbObject&& Value)> Visitor) const
 {
 	for (FCbField Field : Definition["Constants"_ASV])
@@ -261,7 +286,7 @@ void FBuildDefinitionInternal::IterateConstants(TFunctionRef<void (FStringView K
 
 void FBuildDefinitionInternal::IterateInputBuilds(TFunctionRef<void (FStringView Key, const FBuildPayloadKey& PayloadKey)> Visitor) const
 {
-	for (FCbFieldView Field : Definition.FindView("Builds"_ASV))
+	for (FCbFieldView Field : Definition.FindView("Inputs"_ASV).AsObjectView().FindView("Builds"_ASV))
 	{
 		FCbObjectView Build = Field.AsObjectView();
 		const FBuildKey BuildKey{Build["Build"_ASV].AsHash()};
@@ -272,7 +297,7 @@ void FBuildDefinitionInternal::IterateInputBuilds(TFunctionRef<void (FStringView
 
 void FBuildDefinitionInternal::IterateInputBulkData(TFunctionRef<void (FStringView Key, const FGuid& BulkDataId)> Visitor) const
 {
-	for (FCbFieldView Field : Definition.FindView("BulkData"_ASV))
+	for (FCbFieldView Field : Definition.FindView("Inputs"_ASV).AsObjectView().FindView("BulkData"_ASV))
 	{
 		Visitor(FUTF8ToTCHAR(Field.GetName()), Field.AsUuid());
 	}
@@ -280,7 +305,7 @@ void FBuildDefinitionInternal::IterateInputBulkData(TFunctionRef<void (FStringVi
 
 void FBuildDefinitionInternal::IterateInputFiles(TFunctionRef<void (FStringView Key, FStringView Path)> Visitor) const
 {
-	for (FCbFieldView Field : Definition.FindView("Files"_ASV))
+	for (FCbFieldView Field : Definition.FindView("Inputs"_ASV).AsObjectView().FindView("Files"_ASV))
 	{
 		Visitor(FUTF8ToTCHAR(Field.GetName()), FUTF8ToTCHAR(Field.AsString()));
 	}
@@ -288,7 +313,7 @@ void FBuildDefinitionInternal::IterateInputFiles(TFunctionRef<void (FStringView
 
 void FBuildDefinitionInternal::IterateInputHashes(TFunctionRef<void (FStringView Key, const FIoHash& RawHash)> Visitor) const
 {
-	for (FCbFieldView Field : Definition.FindView("Hashes"_ASV))
+	for (FCbFieldView Field : Definition.FindView("Inputs"_ASV).AsObjectView().FindView("Hashes"_ASV))
 	{
 		Visitor(FUTF8ToTCHAR(Field.GetName()), Field.AsBinaryAttachment());
 	}
diff --git a/Engine/Source/Developer/DerivedDataCache/Private/DerivedDataBuildPrivate.h b/Engine/Source/Developer/DerivedDataCache/Private/DerivedDataBuildPrivate.h
--- a/Engine/Source/Developer/DerivedDataCache/Private/DerivedDataBuildPrivate.h
+++ b/Engine/Source/Developer/DerivedDataCache/Private/DerivedDataBuildPrivate.h
@@ -6,7 +6,10 @@
 
 class FCbObject;
 class FCbWriter;
+struct FGuid;
 
+namespace UE::DerivedData { class FBuildAction; }
+namespace UE::DerivedData { class FBuildActionBuilder; }
 namespace UE::DerivedData { class FBuildDefinition; }
 namespace UE::DerivedData { class FBuildDefinitionBuilder; }
 namespace UE::DerivedData { class FBuildOutput; }
@@ -21,6 +24,10 @@ namespace UE::DerivedData::Private
 FBuildDefinitionBuilder CreateBuildDefinition(FStringView Name, FStringView Function);
 FBuildDefinition LoadBuildDefinition(FStringView Name, FCbObject&& Definition);
 
+// Implemented in DerivedDataBuildAction.cpp
+FBuildActionBuilder CreateBuildAction(FStringView Name, FStringView Function, const FGuid& FunctionVersion, const FGuid& BuildSystemVersion);
+FBuildAction LoadBuildAction(FStringView Name, FCbObject&& Action);
+
 // Implemented in DerivedDataBuildOutput.cpp
 FBuildOutputBuilder CreateBuildOutput(FStringView Name, FStringView Function);
 FBuildOutput LoadBuildOutput(FStringView Name, FStringView Function, const FCbObject& Output);
diff --git a/Engine/Source/Developer/DerivedDataCache/Private/DerivedDataCache.cpp b/Engine/Source/Developer/DerivedDataCache/Private/DerivedDataCache.cpp
--- a/Engine/Source/Developer/DerivedDataCache/Private/DerivedDataCache.cpp
+++ b/Engine/Source/Developer/DerivedDataCache/Private/DerivedDataCache.cpp
@@ -18,6 +18,7 @@
 
 #include "DerivedDataBackendInterface.h"
 #include "DerivedDataBuild.h"
+#include "DerivedDataBuildAction.h"
 #include "DerivedDataBuildDefinition.h"
 #include "DerivedDataBuildOutput.h"
 #include "DerivedDataBuildPrivate.h"
@@ -727,6 +728,17 @@ public:
 		return Private::LoadBuildDefinition(Name, MoveTemp(Definition));
 	}
 
+	FBuildActionBuilder CreateAction(FStringView Name, FStringView Function) final
+	{
+		// DDC-TODO: Find the function version from the function registry on the scheduler.
+		return Private::CreateBuildAction(Name, Function, FGuid::NewGuid(), BuildSystemVersion);
+	}
+
+	FBuildAction LoadAction(FStringView Name, FCbObject&& Action) final
+	{
+		return Private::LoadBuildAction(Name, MoveTemp(Action));
+	}
+
 	FBuildOutputBuilder CreateOutput(FStringView Name, FStringView Function) final
 	{
 		return Private::CreateBuildOutput(Name, Function);
@@ -741,6 +753,13 @@ public:
 	{
 		return Private::LoadBuildOutput(Name, Function, Output);
 	}
+
+	const FGuid& GetVersion() const final
+	{
+		return BuildSystemVersion;
+	}
+
+	const FGuid BuildSystemVersion{TEXT("ac0574e5-62bd-4c2e-84ec-f2efe48c0fef")};
 };
 
 } // UE::DerivedData
diff --git a/Engine/Source/Developer/DerivedDataCache/Public/DerivedDataBuild.h b/Engine/Source/Developer/DerivedDataCache/Public/DerivedDataBuild.h
--- a/Engine/Source/Developer/DerivedDataCache/Public/DerivedDataBuild.h
+++ b/Engine/Source/Developer/DerivedDataCache/Public/DerivedDataBuild.h
@@ -5,7 +5,10 @@
 #include "Containers/StringFwd.h"
 
 class FCbObject;
+struct FGuid;
 
+namespace UE::DerivedData { class FBuildAction; }
+namespace UE::DerivedData { class FBuildActionBuilder; }
 namespace UE::DerivedData { class FBuildDefinition; }
 namespace UE::DerivedData { class FBuildDefinitionBuilder; }
 namespace UE::DerivedData { class FBuildOutput; }
@@ -42,6 +45,23 @@ public:
 	 */
 	virtual FBuildDefinition LoadDefinition(FStringView Name, FCbObject&& Definition) = 0;
 
+	/**
+	 * Create a build action builder.
+	 *
+	 * @param Name       The name by which to identify this action for logging and profiling.
+	 * @param Function   The name of the build function that produced this action.
+	 */
+	virtual FBuildActionBuilder CreateAction(FStringView Name, FStringView Function) = 0;
+
+	/**
+	 * Load a build action from compact binary.
+	 *
+	 * @param Name       The name by which to identify this action for logging and profiling.
+	 * @param Action     The saved action to load.
+	 * @return A valid build action, or a build action with an empty key on error.
+	 */
+	virtual FBuildAction LoadAction(FStringView Name, FCbObject&& Action) = 0;
+
 	/**
 	 * Create a build output builder.
 	 *
@@ -59,6 +79,15 @@ public:
 	 */
 	virtual FBuildOutput LoadOutput(FStringView Name, FStringView Function, const FCbObject& Output) = 0;
 	virtual FBuildOutput LoadOutput(FStringView Name, FStringView Function, const FCacheRecord& Output) = 0;
+
+	/**
+	 * Returns the version of the build system.
+	 *
+	 * This version is expected to change very infrequently, only when formats and protocols used by
+	 * the build system are changed in a way that breaks compatibility. This version is incorporated
+	 * into build actions to keep the build output separate for different build versions.
+	 */
+	virtual const FGuid& GetVersion() const = 0;
 };
 
 } // UE::DerivedData
diff --git a/Engine/Source/Developer/DerivedDataCache/Public/DerivedDataBuildAction.h b/Engine/Source/Developer/DerivedDataCache/Public/DerivedDataBuildAction.h
new file mode 100644
--- /dev/null
+++ b/Engine/Source/Developer/DerivedDataCache/Public/DerivedDataBuildAction.h
@@ -0,0 +1,183 @@
+// Copyright Epic Games, Inc. All Rights Reserved.
+
+#pragma once
+
+#include "CoreTypes.h"
+#include "Containers/StringFwd.h"
+#include "IO/IoHash.h"
+#include "Templates/Function.h"
+#include "Templates/RefCounting.h"
+#include "Templates/UniquePtr.h"
+
+class FCbObject;
+class FCbWriter;
+struct FGuid;
+
+namespace UE::DerivedData { class FBuildAction; }
+namespace UE::DerivedData { class FBuildActionBuilder; }
+
+namespace UE::DerivedData
+{
+
+/** A key that uniquely identifies a build action. */
+struct FBuildActionKey
+{
+	FIoHash Hash;
+};
+
+} // UE::DerivedData
+
+namespace UE::DerivedData::Private
+{
+
+class IBuildActionInternal
+{
+public:
+	virtual ~IBuildActionInternal() = default;
+	virtual const FBuildActionKey& GetKey() const = 0;
+	virtual FStringView GetName() const = 0;
+	virtual FStringView GetFunction() const = 0;
+	virtual const FGuid& GetFunctionVersion() const = 0;
+	virtual const FGuid& GetBuildSystemVersion() const = 0;
+	virtual void IterateConstants(TFunctionRef<void (FStringView Key, FCbObject&& Value)> Visitor) const = 0;
+	virtual void IterateInputs(TFunctionRef<void (FStringView Key, const FIoHash& RawHash, uint64 RawSize)> Visitor) const = 0;
+	virtual void Save(FCbWriter& Writer) const = 0;
+	virtual void AddRef() const = 0;
+	virtual void Release() const = 0;
+};
+
+FBuildAction CreateBuildAction(IBuildActionInternal* Action);
+
+class IBuildActionBuilderInternal
+{
+public:
+	virtual ~IBuildActionBuilderInternal() = default;
+	virtual void AddConstant(FStringView Key, const FCbObject& Value) = 0;
+	virtual void AddInput(FStringView Key, const FIoHash& RawHash, uint64 RawSize) = 0;
+	virtual FBuildAction Build() = 0;
+};
+
+FBuildActionBuilder CreateBuildActionBuilder(IBuildActionBuilderInternal* ActionBuilder);
+
+} // UE::DerivedData::Private
+
+namespace UE::DerivedData
+{
+
+/**
+ * A build action references a build function and the inputs to that build function.
+ *
+ * The purpose of an action is to capture everything required to execute a derived data build for
+ * a fixed version of the build function and its constants and inputs.
+ *
+ * The key for the action uniquely identifies the action and is derived by hashing the serialized
+ * compact binary representation of the action.
+ *
+ * The keys for constants and inputs are names that are unique within the build action.
+ *
+ * The build action is immutable, and is created by a build session from a build definition.
+ *
+ * @see FBuildDefinition
+ * @see FBuildSession
+ */
+class FBuildAction
+{
+public:
+	/** Returns the key that uniquely identifies this build action. */
+	inline const FBuildActionKey& GetKey() const { return Action->GetKey(); }
+
+	/** Returns the name by which to identify this action for logging and profiling. */
+	inline FStringView GetName() const
+	{
+		return Action->GetName();
+	}
+
+	/** Returns the name of the build function with which to build this action. */
+	inline FStringView GetFunction() const
+	{
+		return Action->GetFunction();
+	}
+
+	/** Returns the version of the build function with which to build this action. */
+	inline const FGuid& GetFunctionVersion() const
+	{
+		return Action->GetFunctionVersion();
+	}
+
+	/** Returns the version of the build system required to build this action. */
+	inline const FGuid& GetBuildSystemVersion() const
+	{
+		return Action->GetBuildSystemVersion();
+	}
+
+	/** Visits every constant in order by key. */
+	inline void IterateConstants(TFunctionRef<void (FStringView Key, FCbObject&& Value)> Visitor) const
+	{
+		Action->IterateConstants(MoveTemp(Visitor));
+	}
+
+	/** Visits every input in order by key. */
+	inline void IterateInputs(TFunctionRef<void (FStringView Key, const FIoHash& RawHash, uint64 RawSize)> Visitor) const
+	{
+		Action->IterateInputs(MoveTemp(Visitor));
+	}
+
+	/** Saves the build action to a compact binary object. Calls BeginObject and EndObject. */
+	inline void Save(FCbWriter& Writer) const
+	{
+		return Action->Save(Writer);
+	}
+
+private:
+	friend FBuildAction Private::CreateBuildAction(Private::IBuildActionInternal* Action);
+
+	/** Construct a build action. Use Build() on a builder from IBuild::CreateAction(). */
+	inline explicit FBuildAction(Private::IBuildActionInternal* InAction)
+		: Action(InAction)
+	{
+	}
+
+	TRefCountPtr<Private::IBuildActionInternal> Action;
+};
+
+/**
+ * A build action builder is used to construct a build action.
+ *
+ * Create using IBuild::CreateAction() which must be given a build function name.
+ *
+ * @see FBuildAction
+ */
+class FBuildActionBuilder
+{
+public:
+	/** Add a constant object with a key that is unique within this action. */
+	inline void AddConstant(FStringView Key, const FCbObject& Value)
+	{
+		ActionBuilder->AddConstant(Key, Value);
+	}
+
+	/** Add an input with a key that is unique within this action. */
+	inline void AddInput(FStringView Key, const FIoHash& RawHash, uint64 RawSize)
+	{
+		ActionBuilder->AddInput(Key, RawHash, RawSize);
+	}
+
+	/** Build a build action, which makes this builder subsequently unusable. */
+	inline FBuildAction Build()
+	{
+		return ActionBuilder->Build();
+	}
+
+private:
+	friend FBuildActionBuilder Private::CreateBuildActionBuilder(Private::IBuildActionBuilderInternal* ActionBuilder);
+
+	/** Construct a build action builder. Use IBuild::CreateAction(). */
+	inline explicit FBuildActionBuilder(Private::IBuildActionBuilderInternal* InActionBuilder)
+		: ActionBuilder(InActionBuilder)
+	{
+	}
+
+	TUniquePtr<Private::IBuildActionBuilderInternal> ActionBuilder;
+};
+
+} // UE::DerivedData
diff --git a/Engine/Source/Developer/DerivedDataCache/Public/DerivedDataBuildDefinition.h b/Engine/Source/Developer/DerivedDataCache/Public/DerivedDataBuildDefinition.h
--- a/Engine/Source/Developer/DerivedDataCache/Public/DerivedDataBuildDefinition.h
+++ b/Engine/Source/Developer/DerivedDataCache/Public/DerivedDataBuildDefinition.h
@@ -28,6 +28,8 @@ public:
 	virtual const FBuildKey& GetKey() const = 0;
 	virtual FStringView GetName() const = 0;
 	virtual FStringView GetFunction() const = 0;
+	virtual bool HasConstants() const = 0;
+	virtual bool HasInputs() const = 0;
 	virtual void IterateConstants(TFunctionRef<void (FStringView Key, FCbObject&& Value)> Visitor) const = 0;
 	virtual void IterateInputBuilds(TFunctionRef<void (FStringView Key, const FBuildPayloadKey& PayloadKey)> Visitor) const = 0;
 	virtual void IterateInputBulkData(TFunctionRef<void (FStringView Key, const FGuid& BulkDataId)> Visitor) const = 0;
@@ -84,22 +86,19 @@ class FBuildDefinition
 {
 public:
 	/** Returns the key that uniquely identifies this build definition. */
-	inline const FBuildKey& GetKey() const
-	{
-		return Definition->GetKey();
-	}
+	inline const FBuildKey& GetKey() const { return Definition->GetKey(); }
 
 	/** Returns the name by which to identify this definition for logging and profiling. */
-	inline FStringView GetName() const
-	{
-		return Definition->GetName();
-	}
+	inline FStringView GetName() const { return Definition->GetName(); }
 
 	/** Returns the name of the build function with which to build this definition. */
-	inline FStringView GetFunction() const
-	{
-		return Definition->GetFunction();
-	}
+	inline FStringView GetFunction() const { return Definition->GetFunction(); }
+
+	/** Returns whether the definition has any constants. */
+	inline bool HasConstants() const { return Definition->HasConstants(); }
+
+	/** Returns whether the definition has any inputs. */
+	inline bool HasInputs() const { return Definition->HasInputs(); }
 
 	/** Visits every constant in order by key. */
 	inline void IterateConstants(TFunctionRef<void (FStringView Key, FCbObject&& Value)> Visitor) const
