DDC: Changed the builder types to delete themselves on Build()

#rnx

diff --git a/Engine/Source/Developer/DerivedDataCache/Private/DerivedDataCacheRecord.cpp b/Engine/Source/Developer/DerivedDataCache/Private/DerivedDataCacheRecord.cpp
--- a/Engine/Source/Developer/DerivedDataCache/Private/DerivedDataCacheRecord.cpp
+++ b/Engine/Source/Developer/DerivedDataCache/Private/DerivedDataCacheRecord.cpp
@@ -236,6 +236,7 @@ FCacheRecord FCacheRecordBuilderInternal::Build()
 
 FRequest FCacheRecordBuilderInternal::BuildAsync(FOnCacheRecordComplete&& OnComplete, EPriority Priority)
 {
+	ON_SCOPE_EXIT { delete this; };
 	checkf(OnComplete, TEXT("Failed to build cache record for %s because the completion callback is null."),
 		*WriteToString<96>(Key));
 	OnComplete(Build());
diff --git a/Engine/Source/Developer/DerivedDataCache/Public/DerivedDataBuildAction.h b/Engine/Source/Developer/DerivedDataCache/Public/DerivedDataBuildAction.h
--- a/Engine/Source/Developer/DerivedDataCache/Public/DerivedDataBuildAction.h
+++ b/Engine/Source/Developer/DerivedDataCache/Public/DerivedDataBuildAction.h
@@ -5,6 +5,7 @@
 #include "CoreTypes.h"
 #include "Containers/StringFwd.h"
 #include "IO/IoHash.h"
+#include "Misc/ScopeExit.h"
 #include "Templates/Function.h"
 #include "Templates/RefCounting.h"
 #include "Templates/UniquePtr.h"
@@ -165,6 +166,7 @@ public:
 	/** Build a build action, which makes this builder subsequently unusable. */
 	inline FBuildAction Build()
 	{
+		ON_SCOPE_EXIT { ActionBuilder = nullptr; };
 		return ActionBuilder->Build();
 	}
 
diff --git a/Engine/Source/Developer/DerivedDataCache/Public/DerivedDataBuildDefinition.h b/Engine/Source/Developer/DerivedDataCache/Public/DerivedDataBuildDefinition.h
--- a/Engine/Source/Developer/DerivedDataCache/Public/DerivedDataBuildDefinition.h
+++ b/Engine/Source/Developer/DerivedDataCache/Public/DerivedDataBuildDefinition.h
@@ -4,6 +4,7 @@
 
 #include "CoreTypes.h"
 #include "Containers/StringView.h"
+#include "Misc/ScopeExit.h"
 #include "Templates/Function.h"
 #include "Templates/RefCounting.h"
 #include "Templates/UniquePtr.h"
@@ -205,6 +206,7 @@ public:
 	 */
 	inline FBuildDefinition Build()
 	{
+		ON_SCOPE_EXIT { DefinitionBuilder = nullptr; };
 		return DefinitionBuilder->Build();
 	}
 
diff --git a/Engine/Source/Developer/DerivedDataCache/Public/DerivedDataBuildOutput.h b/Engine/Source/Developer/DerivedDataCache/Public/DerivedDataBuildOutput.h
--- a/Engine/Source/Developer/DerivedDataCache/Public/DerivedDataBuildOutput.h
+++ b/Engine/Source/Developer/DerivedDataCache/Public/DerivedDataBuildOutput.h
@@ -5,6 +5,7 @@
 #include "CoreTypes.h"
 #include "Containers/ArrayView.h"
 #include "Containers/StringView.h"
+#include "Misc/ScopeExit.h"
 #include "Templates/Function.h"
 #include "Templates/RefCounting.h"
 #include "Templates/UniquePtr.h"
@@ -181,6 +182,7 @@ public:
 	/** Build a build output, which makes this builder subsequently unusable. */
 	inline FBuildOutput Build()
 	{
+		ON_SCOPE_EXIT { OutputBuilder = nullptr; };
 		return OutputBuilder->Build();
 	}
 
diff --git a/Engine/Source/Developer/DerivedDataCache/Public/DerivedDataCacheRecord.h b/Engine/Source/Developer/DerivedDataCache/Public/DerivedDataCacheRecord.h
--- a/Engine/Source/Developer/DerivedDataCache/Public/DerivedDataCacheRecord.h
+++ b/Engine/Source/Developer/DerivedDataCache/Public/DerivedDataCacheRecord.h
@@ -219,7 +219,7 @@ public:
 	 */
 	inline FRequest BuildAsync(FOnCacheRecordComplete&& OnComplete, EPriority Priority)
 	{
-		return RecordBuilder->BuildAsync(MoveTemp(OnComplete), Priority);
+		return RecordBuilder.Release()->BuildAsync(MoveTemp(OnComplete), Priority);
 	}
 
 private:
