DDC: Temporary workaround for static analysis issue with Get<T> on 2-tuples

#rnx

diff --git a/Engine/Source/Developer/DerivedDataCache/Private/DerivedDataBuildAction.cpp b/Engine/Source/Developer/DerivedDataCache/Private/DerivedDataBuildAction.cpp
--- a/Engine/Source/Developer/DerivedDataCache/Private/DerivedDataBuildAction.cpp
+++ b/Engine/Source/Developer/DerivedDataCache/Private/DerivedDataBuildAction.cpp
@@ -169,8 +169,8 @@ FBuildActionInternal::FBuildActionInternal(FBuildActionBuilderInternal&& ActionB
 			{
 				const TTuple<FIoHash, uint64>& Input = Pair.Value.Get<TTuple<FIoHash, uint64>>();
 				Writer.BeginObject(FTCHARToUTF8(Pair.Key));
-				Writer.AddBinaryAttachment("RawHash"_ASV, Input.Get<FIoHash>());
-				Writer.AddInteger("RawSize"_ASV, Input.Get<uint64>());
+				Writer.AddBinaryAttachment("RawHash"_ASV, Input.Get<0>());
+				Writer.AddInteger("RawSize"_ASV, Input.Get<1>());
 				Writer.EndObject();
 			}
 		}
