DDC: Added optional versions of BuildAction, BuildDefinition, BuildOutput, CacheRecord

- Optional versions are now returned from operations that may fail.
- Fixed a bug in the cache where a null record could be dereferenced.
- Improved validation when loading build types.

#rnx

diff --git a/Engine/Source/Developer/DerivedDataCache/Private/DerivedDataBuildAction.cpp b/Engine/Source/Developer/DerivedDataCache/Private/DerivedDataBuildAction.cpp
--- a/Engine/Source/Developer/DerivedDataCache/Private/DerivedDataBuildAction.cpp
+++ b/Engine/Source/Developer/DerivedDataCache/Private/DerivedDataBuildAction.cpp
@@ -77,7 +77,7 @@ class FBuildActionInternal final : public IBuildActionInternal
 {
 public:
 	explicit FBuildActionInternal(FBuildActionBuilderInternal&& ActionBuilder);
-	explicit FBuildActionInternal(FStringView Name, FCbObject&& Action);
+	explicit FBuildActionInternal(FStringView Name, FCbObject&& Action, bool& bOutIsValid);
 
 	~FBuildActionInternal() final = default;
 
@@ -112,8 +112,8 @@ private:
 	FString Function;
 	FGuid FunctionVersion;
 	FGuid BuildSystemVersion;
-	FBuildActionKey Key;
 	FCbObject Action;
+	FBuildActionKey Key;
 };
 
 ///////////////////////////////////////////////////////////////////////////////////////////////////////////////////////
@@ -179,22 +179,31 @@ FBuildActionInternal::FBuildActionInternal(FBuildActionBuilderInternal&& ActionB
 
 	Writer.EndObject();
 	Action = Writer.Save().AsObject();
-	Key = FBuildActionKey{Action.GetHash()};
+	Key.Hash = Action.GetHash();
 }
 
-FBuildActionInternal::FBuildActionInternal(FStringView InName, FCbObject&& InAction)
+FBuildActionInternal::FBuildActionInternal(FStringView InName, FCbObject&& InAction, bool& bOutIsValid)
 	: Name(InName)
 	, Function(InAction.FindView("Function"_ASV).AsString())
 	, FunctionVersion(InAction.FindView("FunctionVersion"_ASV).AsUuid())
 	, BuildSystemVersion(InAction.FindView("BuildSystemVersion"_ASV).AsUuid())
 	, Action(MoveTemp(InAction))
+	, Key{Action.GetHash()}
 {
 	checkf(!InName.IsEmpty(), TEXT("A build action requires a non-empty name."));
 	Action.MakeOwned();
-	if (const bool bIsValid = Action && Algo::AllOf(Function, FChar::IsAlnum))
-	{
-		Key = FBuildActionKey{Action.GetHash()};
-	}
+	bOutIsValid = Action
+		&& !Function.IsEmpty() && Algo::AllOf(Function, FChar::IsAlnum)
+		&& FunctionVersion.IsValid()
+		&& BuildSystemVersion.IsValid()
+		&& Algo::AllOf(Action.FindView("Constants"_ASV),
+			[](FCbFieldView Field) { return Field.GetName().Len() > 0 && Field.IsObject(); })
+		&& Algo::AllOf(Action.FindView("Inputs"_ASV), [](FCbFieldView Field)
+			{
+				return Field.GetName().Len() > 0 && Field.IsObject()
+					&& Field.AsObjectView()["RawHash"_ASV].IsBinaryAttachment()
+					&& Field.AsObjectView()["RawSize"_ASV].IsInteger();
+			});
 }
 
 void FBuildActionInternal::IterateConstants(TFunctionRef<void (FStringView Key, FCbObject&& Value)> Visitor) const
@@ -243,9 +252,15 @@ FBuildActionBuilder CreateBuildAction(FStringView Name, FStringView Function, co
 	return CreateBuildActionBuilder(new FBuildActionBuilderInternal(Name, Function, FunctionVersion, BuildSystemVersion));
 }
 
-FBuildAction LoadBuildAction(FStringView Name, FCbObject&& Action)
+FOptionalBuildAction LoadBuildAction(FStringView Name, FCbObject&& Action)
 {
-	return CreateBuildAction(new FBuildActionInternal(Name, MoveTemp(Action)));
+	bool bIsValid = false;
+	FOptionalBuildAction Out = CreateBuildAction(new FBuildActionInternal(Name, MoveTemp(Action), bIsValid));
+	if (!bIsValid)
+	{
+		Out.Reset();
+	}
+	return Out;
 }
 
 } // UE::DerivedData::Private
diff --git a/Engine/Source/Developer/DerivedDataCache/Private/DerivedDataBuildDefinition.cpp b/Engine/Source/Developer/DerivedDataCache/Private/DerivedDataBuildDefinition.cpp
--- a/Engine/Source/Developer/DerivedDataCache/Private/DerivedDataBuildDefinition.cpp
+++ b/Engine/Source/Developer/DerivedDataCache/Private/DerivedDataBuildDefinition.cpp
@@ -84,7 +84,7 @@ class FBuildDefinitionInternal final : public IBuildDefinitionInternal
 {
 public:
 	explicit FBuildDefinitionInternal(FBuildDefinitionBuilderInternal&& DefinitionBuilder);
-	explicit FBuildDefinitionInternal(FStringView Name, FCbObject&& Definition);
+	explicit FBuildDefinitionInternal(FStringView Name, FCbObject&& Definition, bool& bOutIsValid);
 
 	~FBuildDefinitionInternal() final = default;
 
@@ -250,20 +250,33 @@ FBuildDefinitionInternal::FBuildDefinitionInternal(FBuildDefinitionBuilderIntern
 
 	Writer.EndObject();
 	Definition = Writer.Save().AsObject();
-	Key = FBuildKey{Definition.GetHash()};
+	Key.Hash = Definition.GetHash();
 }
 
-FBuildDefinitionInternal::FBuildDefinitionInternal(FStringView InName, FCbObject&& InDefinition)
+FBuildDefinitionInternal::FBuildDefinitionInternal(FStringView InName, FCbObject&& InDefinition, bool& bOutIsValid)
 	: Name(InName)
 	, Function(InDefinition.FindView("Function"_ASV).AsString())
 	, Definition(MoveTemp(InDefinition))
+	, Key{Definition.GetHash()}
 {
 	checkf(!InName.IsEmpty(), TEXT("A build definition requires a non-empty name."));
 	Definition.MakeOwned();
-	if (const bool bIsValid = Definition && Algo::AllOf(Function, FChar::IsAlnum))
-	{
-		Key = FBuildKey{Definition.GetHash()};
-	}
+	bOutIsValid = Definition
+		&& !Function.IsEmpty() && Algo::AllOf(Function, FChar::IsAlnum)
+		&& Algo::AllOf(Definition.FindView("Constants"_ASV),
+			[](FCbFieldView Field) { return Field.GetName().Len() > 0 && Field.IsObject(); })
+		&& Algo::AllOf(Definition.FindView("Inputs"_ASV).AsObjectView().FindView("Builds"_ASV), [](FCbFieldView Field)
+			{
+				return Field.GetName().Len() > 0 && Field.IsObject()
+					&& Field.AsObjectView()["Build"_ASV].IsHash()
+					&& Field.AsObjectView()["Payload"_ASV].IsObjectId();
+			})
+		&& Algo::AllOf(Definition.FindView("Inputs"_ASV).AsObjectView().FindView("BulkData"_ASV),
+			[](FCbFieldView Field) { return Field.GetName().Len() > 0 && Field.IsUuid(); })
+		&& Algo::AllOf(Definition.FindView("Inputs"_ASV).AsObjectView().FindView("Files"_ASV),
+			[](FCbFieldView Field) { return Field.GetName().Len() > 0 && Field.AsString().Len() > 0; })
+		&& Algo::AllOf(Definition.FindView("Inputs"_ASV).AsObjectView().FindView("Hashes"_ASV),
+			[](FCbFieldView Field) { return Field.GetName().Len() > 0 && Field.IsBinaryAttachment(); });
 }
 
 bool FBuildDefinitionInternal::HasConstants() const
@@ -348,9 +361,15 @@ FBuildDefinitionBuilder CreateBuildDefinition(FStringView Name, FStringView Func
 	return CreateBuildDefinitionBuilder(new FBuildDefinitionBuilderInternal(Name, Function));
 }
 
-FBuildDefinition LoadBuildDefinition(FStringView Name, FCbObject&& Definition)
+FOptionalBuildDefinition LoadBuildDefinition(FStringView Name, FCbObject&& Definition)
 {
-	return CreateBuildDefinition(new FBuildDefinitionInternal(Name, MoveTemp(Definition)));
+	bool bIsValid = false;
+	FOptionalBuildDefinition Out = CreateBuildDefinition(new FBuildDefinitionInternal(Name, MoveTemp(Definition), bIsValid));
+	if (!bIsValid)
+	{
+		Out.Reset();
+	}
+	return Out;
 }
 
 } // UE::DerivedData::Private
diff --git a/Engine/Source/Developer/DerivedDataCache/Private/DerivedDataBuildOutput.cpp b/Engine/Source/Developer/DerivedDataCache/Private/DerivedDataBuildOutput.cpp
--- a/Engine/Source/Developer/DerivedDataCache/Private/DerivedDataBuildOutput.cpp
+++ b/Engine/Source/Developer/DerivedDataCache/Private/DerivedDataBuildOutput.cpp
@@ -2,6 +2,7 @@
 
 #include "DerivedDataBuildOutput.h"
 
+#include "Algo/AllOf.h"
 #include "Algo/BinarySearch.h"
 #include "Containers/StringConv.h"
 #include "Containers/StringView.h"
@@ -75,8 +76,8 @@ class FBuildOutputInternal final : public IBuildOutputInternal
 {
 public:
 	explicit FBuildOutputInternal(FBuildOutputBuilderInternal&& InOutput);
-	explicit FBuildOutputInternal(FStringView Name, FStringView Function, const FCbObject& InOutput);
-	explicit FBuildOutputInternal(FStringView Name, FStringView Function, const FCacheRecord& InOutput);
+	explicit FBuildOutputInternal(FStringView Name, FStringView Function, const FCbObject& InOutput, bool& bOutIsValid);
+	explicit FBuildOutputInternal(FStringView Name, FStringView Function, const FCacheRecord& InOutput, bool& bOutIsValid);
 
 	~FBuildOutputInternal() final = default;
 
@@ -90,8 +91,10 @@ public:
 	void IterateDiagnostics(TFunctionRef<void (const FBuildDiagnostic& Diagnostic)> Visitor) const final;
 	bool HasError() const final;
 
-	void Save(FCbWriter& Writer) const;
-	void Save(FCacheRecordBuilder& RecordBuilder) const;
+	void Save(FCbWriter& Writer) const final;
+	void Save(FCacheRecordBuilder& RecordBuilder) const final;
+
+	bool IsValid() const;
 
 	inline void AddRef() const final
 	{
@@ -126,32 +129,42 @@ FBuildOutputInternal::FBuildOutputInternal(FBuildOutputBuilderInternal&& InOutpu
 {
 }
 
-FBuildOutputInternal::FBuildOutputInternal(FStringView Name, FStringView Function, const FCbObject& Output)
+FBuildOutputInternal::FBuildOutputInternal(FStringView InName, FStringView InFunction, const FCbObject& InOutput, bool& bOutIsValid)
 	: Name(Name)
 	, Function(Function)
 {
-	for (FCbFieldView PayloadField : Output["Payloads"_ASV])
+	bOutIsValid = false;
+	checkf(!InName.IsEmpty(), TEXT("A build output requires a non-empty name."));
+	for (FCbFieldView PayloadField : InOutput["Payloads"_ASV])
 	{
 		const FCbObjectView Payload = PayloadField.AsObjectView();
-		Payloads.Emplace(
-			FPayloadId(Payload["Id"_ASV].AsObjectId().GetView()),
-			Payload["RawHash"_ASV].AsAttachment(),
-			Payload["RawSize"_ASV].AsUInt64());
+		const FPayloadId Id(Payload["Id"_ASV].AsObjectId().GetView());
+		const FIoHash& RawHash = Payload["RawHash"_ASV].AsAttachment();
+		if (Id.IsNull() || RawHash.IsZero() || !Payload["RawSize"_ASV].IsInteger())
+		{
+			return;
+		}
+		Payloads.Emplace(Id, RawHash, Payload["RawSize"_ASV].AsUInt64());
 	}
-	Diagnostics = Output["Diagnostics"_ASV];
-	Meta = Output["Meta"_ASV].AsObject();
+	Diagnostics = InOutput["Diagnostics"_ASV];
+	Diagnostics.MakeOwned();
+	Meta = InOutput["Meta"_ASV].AsObject();
+	Meta.MakeOwned();
+	bOutIsValid = IsValid() && (!InOutput.FindView("Meta"_ASV) || InOutput.FindView("Meta"_ASV).IsObject());
 }
 
-FBuildOutputInternal::FBuildOutputInternal(FStringView Name, FStringView Function, const FCacheRecord& Output)
-	: Name(Name)
-	, Function(Function)
-	, Payloads(Output.GetAttachmentPayloads())
-	, Meta(Output.GetMeta())
+FBuildOutputInternal::FBuildOutputInternal(FStringView InName, FStringView InFunction, const FCacheRecord& InOutput, bool& bOutIsValid)
+	: Name(InName)
+	, Function(InFunction)
+	, Payloads(InOutput.GetAttachmentPayloads())
+	, Meta(InOutput.GetMeta())
 {
-	if (FSharedBuffer Buffer = Output.GetValue())
+	checkf(!InName.IsEmpty(), TEXT("A build output requires a non-empty name."));
+	if (FSharedBuffer Buffer = InOutput.GetValue())
 	{
 		Diagnostics = FCbObject(MoveTemp(Buffer))["Diagnostics"_ASV];
 	}
+	bOutIsValid = IsValid();
 }
 
 const FPayload& FBuildOutputInternal::GetPayload(const FPayloadId& Id) const
@@ -232,6 +245,20 @@ void FBuildOutputInternal::Save(FCacheRecordBuilder& RecordBuilder) const
 	}
 }
 
+bool FBuildOutputInternal::IsValid() const
+{
+	return !Function.IsEmpty() && Algo::AllOf(Function, FChar::IsAlnum)
+		&& Algo::AllOf(Payloads, [](const FPayload& Payload) { return !Payload.GetRawHash().IsZero(); })
+		&& (!Diagnostics || Diagnostics.IsArray())
+		&& Algo::AllOf(Diagnostics.CreateViewIterator(), [](FCbFieldView Field)
+			{
+				return Field.IsObject()
+					&& Field.AsObjectView()["Level"_ASV].AsString().Len() > 0
+					&& Field.AsObjectView()["Category"_ASV].AsString().Len() > 0
+					&& Field.AsObjectView()["Message"_ASV].AsString().Len() > 0;
+			});
+}
+
 ///////////////////////////////////////////////////////////////////////////////////////////////////////////////////////
 
 void FBuildOutputBuilderInternal::AddPayload(const FPayload& Payload)
@@ -278,24 +305,31 @@ FBuildOutputBuilder CreateBuildOutputBuilder(IBuildOutputBuilderInternal* Output
 	return FBuildOutputBuilder(OutputBuilder);
 }
 
-const IBuildOutputInternal* GetBuildOutput(const FBuildOutput& Output)
-{
-	return Output.Output;
-}
-
 FBuildOutputBuilder CreateBuildOutput(FStringView Name, FStringView Function)
 {
 	return CreateBuildOutputBuilder(new FBuildOutputBuilderInternal(Name, Function));
 }
 
-FBuildOutput LoadBuildOutput(FStringView Name, FStringView Function, const FCbObject& Output)
+FOptionalBuildOutput LoadBuildOutput(FStringView Name, FStringView Function, const FCbObject& Output)
 {
-	return CreateBuildOutput(new FBuildOutputInternal(Name, Function, Output));
+	bool bIsValid = false;
+	FOptionalBuildOutput Out = CreateBuildOutput(new FBuildOutputInternal(Name, Function, Output, bIsValid));
+	if (!bIsValid)
+	{
+		Out.Reset();
+	}
+	return Out;
 }
 
-FBuildOutput LoadBuildOutput(FStringView Name, FStringView Function, const FCacheRecord& Output)
+FOptionalBuildOutput LoadBuildOutput(FStringView Name, FStringView Function, const FCacheRecord& Output)
 {
-	return CreateBuildOutput(new FBuildOutputInternal(Name, Function, Output));
+	bool bIsValid = false;
+	FOptionalBuildOutput Out = CreateBuildOutput(new FBuildOutputInternal(Name, Function, Output, bIsValid));
+	if (!bIsValid)
+	{
+		Out.Reset();
+	}
+	return Out;
 }
 
 } // UE::DerivedData::Private
diff --git a/Engine/Source/Developer/DerivedDataCache/Private/DerivedDataBuildPrivate.h b/Engine/Source/Developer/DerivedDataCache/Private/DerivedDataBuildPrivate.h
--- a/Engine/Source/Developer/DerivedDataCache/Private/DerivedDataBuildPrivate.h
+++ b/Engine/Source/Developer/DerivedDataCache/Private/DerivedDataBuildPrivate.h
@@ -8,29 +8,29 @@ class FCbObject;
 class FCbWriter;
 struct FGuid;
 
-namespace UE::DerivedData { class FBuildAction; }
 namespace UE::DerivedData { class FBuildActionBuilder; }
-namespace UE::DerivedData { class FBuildDefinition; }
 namespace UE::DerivedData { class FBuildDefinitionBuilder; }
-namespace UE::DerivedData { class FBuildOutput; }
 namespace UE::DerivedData { class FBuildOutputBuilder; }
 namespace UE::DerivedData { class FCacheRecord; }
 namespace UE::DerivedData { class FCacheRecordBuilder; }
+namespace UE::DerivedData { class FOptionalBuildAction; }
+namespace UE::DerivedData { class FOptionalBuildDefinition; }
+namespace UE::DerivedData { class FOptionalBuildOutput; }
 
 namespace UE::DerivedData::Private
 {
 
 // Implemented in DerivedDataBuildDefinition.cpp
 FBuildDefinitionBuilder CreateBuildDefinition(FStringView Name, FStringView Function);
-FBuildDefinition LoadBuildDefinition(FStringView Name, FCbObject&& Definition);
+FOptionalBuildDefinition LoadBuildDefinition(FStringView Name, FCbObject&& Definition);
 
 // Implemented in DerivedDataBuildAction.cpp
 FBuildActionBuilder CreateBuildAction(FStringView Name, FStringView Function, const FGuid& FunctionVersion, const FGuid& BuildSystemVersion);
-FBuildAction LoadBuildAction(FStringView Name, FCbObject&& Action);
+FOptionalBuildAction LoadBuildAction(FStringView Name, FCbObject&& Action);
 
 // Implemented in DerivedDataBuildOutput.cpp
 FBuildOutputBuilder CreateBuildOutput(FStringView Name, FStringView Function);
-FBuildOutput LoadBuildOutput(FStringView Name, FStringView Function, const FCbObject& Output);
-FBuildOutput LoadBuildOutput(FStringView Name, FStringView Function, const FCacheRecord& Output);
+FOptionalBuildOutput LoadBuildOutput(FStringView Name, FStringView Function, const FCbObject& Output);
+FOptionalBuildOutput LoadBuildOutput(FStringView Name, FStringView Function, const FCacheRecord& Output);
 
 } // UE::DerivedData::Private
diff --git a/Engine/Source/Developer/DerivedDataCache/Private/DerivedDataCache.cpp b/Engine/Source/Developer/DerivedDataCache/Private/DerivedDataCache.cpp
--- a/Engine/Source/Developer/DerivedDataCache/Private/DerivedDataCache.cpp
+++ b/Engine/Source/Developer/DerivedDataCache/Private/DerivedDataCache.cpp
@@ -723,7 +723,7 @@ public:
 		return Private::CreateBuildDefinition(Name, Function);
 	}
 
-	FBuildDefinition LoadDefinition(FStringView Name, FCbObject&& Definition) final
+	FOptionalBuildDefinition LoadDefinition(FStringView Name, FCbObject&& Definition) final
 	{
 		return Private::LoadBuildDefinition(Name, MoveTemp(Definition));
 	}
@@ -734,7 +734,7 @@ public:
 		return Private::CreateBuildAction(Name, Function, FGuid::NewGuid(), BuildSystemVersion);
 	}
 
-	FBuildAction LoadAction(FStringView Name, FCbObject&& Action) final
+	FOptionalBuildAction LoadAction(FStringView Name, FCbObject&& Action) final
 	{
 		return Private::LoadBuildAction(Name, MoveTemp(Action));
 	}
@@ -744,12 +744,12 @@ public:
 		return Private::CreateBuildOutput(Name, Function);
 	}
 
-	FBuildOutput LoadOutput(FStringView Name, FStringView Function, const FCbObject& Output) final
+	FOptionalBuildOutput LoadOutput(FStringView Name, FStringView Function, const FCbObject& Output) final
 	{
 		return Private::LoadBuildOutput(Name, Function, Output);
 	}
 
-	FBuildOutput LoadOutput(FStringView Name, FStringView Function, const FCacheRecord& Output) final
+	FOptionalBuildOutput LoadOutput(FStringView Name, FStringView Function, const FCacheRecord& Output) final
 	{
 		return Private::LoadBuildOutput(Name, Function, Output);
 	}
diff --git a/Engine/Source/Developer/DerivedDataCache/Private/FileSystemDerivedDataBackend.cpp b/Engine/Source/Developer/DerivedDataCache/Private/FileSystemDerivedDataBackend.cpp
--- a/Engine/Source/Developer/DerivedDataCache/Private/FileSystemDerivedDataBackend.cpp
+++ b/Engine/Source/Developer/DerivedDataCache/Private/FileSystemDerivedDataBackend.cpp
@@ -900,14 +900,14 @@ public:
 		for (const FCacheKey& Key : Keys)
 		{
 			COOK_STAT(auto Timer = UsageStats.TimeGet());
-			if (FCacheRecord Record = GetCacheRecord(Key, Context, Policy))
+			if (FOptionalCacheRecord Record = GetCacheRecord(Key, Context, Policy))
 			{
 				UE_LOG(LogDerivedDataCache, Verbose, TEXT("%s: Cache hit for %s from '%.*s'"),
 					*CachePath, *WriteToString<96>(Key), Context.Len(), Context.GetData());
-				COOK_STAT(Timer.AddHit(MeasureCacheRecord(Record)));
+				COOK_STAT(Timer.AddHit(MeasureCacheRecord(Record.Get())));
 				if (OnComplete)
 				{
-					OnComplete({MoveTemp(Record), EStatus::Ok});
+					OnComplete({MoveTemp(Record).Get(), EStatus::Ok});
 				}
 			}
 			else
@@ -931,15 +931,15 @@ public:
 		TArray<FCachePayloadKey, TInlineAllocator<16>> SortedKeys(Keys);
 		SortedKeys.StableSort();
 
-		FCacheRecord Record;
+		FOptionalCacheRecord Record;
 		for (const FCachePayloadKey& Key : SortedKeys)
 		{
 			COOK_STAT(auto Timer = UsageStats.TimeGet());
-			if (!Record || Record.GetKey() != Key.CacheKey)
+			if (!Record || Record.Get().GetKey() != Key.CacheKey)
 			{
 				Record = GetCacheRecord(Key.CacheKey, Context, Policy | ECachePolicy::SkipData, /*bAlwaysLoadInlineData*/ true);
 			}
-			if (FPayload Payload = GetCachePayload(Key.CacheKey, Context, Policy, Record.GetPayload(Key.Id)))
+			if (FPayload Payload = Record ? GetCachePayload(Key.CacheKey, Context, Policy, Record.Get().GetPayload(Key.Id)) : FPayload::Null)
 			{
 				UE_LOG(LogDerivedDataCache, Verbose, TEXT("%s: Cache hit for %s from '%.*s'"),
 					*CachePath, *WriteToString<96>(Key), Context.Len(), Context.GetData());
@@ -1043,13 +1043,13 @@ private:
 		return true;
 	}
 
-	FCacheRecord GetCacheRecord(const FCacheKey& Key, FStringView Context, ECachePolicy Policy, bool bAlwaysLoadInlineData = false) const
+	FOptionalCacheRecord GetCacheRecord(const FCacheKey& Key, FStringView Context, ECachePolicy Policy, bool bAlwaysLoadInlineData = false) const
 	{
 		if (!IsUsable())
 		{
 			UE_LOG(LogDerivedDataCache, VeryVerbose, TEXT("%s: Skipped get of %s from '%.*s' because this cache store is not available"),
 				*CachePath, *WriteToString<96>(Key), Context.Len(), Context.GetData());
-			return FCacheRecord();
+			return FOptionalCacheRecord();
 		}
 
 		// Skip the request if querying the cache is disabled.
@@ -1057,7 +1057,7 @@ private:
 		{
 			UE_LOG(LogDerivedDataCache, VeryVerbose, TEXT("%s: Skipped get of %s from '%.*s' due to cache policy"),
 				*CachePath, *WriteToString<96>(Key), Context.Len(), Context.GetData());
-			return FCacheRecord();
+			return FOptionalCacheRecord();
 		}
 
 		// Request the record from storage.
@@ -1068,7 +1068,7 @@ private:
 		{
 			UE_LOG(LogDerivedDataCache, Verbose, TEXT("%s: Cache miss with missing record for %s from '%.*s'"),
 				*CachePath, *WriteToString<96>(Key), Context.Len(), Context.GetData());
-			return FCacheRecord();
+			return FOptionalCacheRecord();
 		}
 
 		// Delete the record from storage if is corrupted or has missing payloads.
@@ -1086,7 +1086,7 @@ private:
 		{
 			UE_LOG(LogDerivedDataCache, Display, TEXT("%s: Cache miss with corrupted record for %s from '%.*s'"),
 				*CachePath, *WriteToString<96>(Key), Context.Len(), Context.GetData());
-			return FCacheRecord();
+			return FOptionalCacheRecord();
 		}
 
 		const FCbObject RecordObject(MoveTemp(Buffer));
@@ -1104,7 +1104,7 @@ private:
 				{
 					UE_LOG(LogDerivedDataCache, Display, TEXT("%s: Cache miss with corrupted metadata for %s from '%.*s'"),
 						*CachePath, *WriteToString<96>(Key), Context.Len(), Context.GetData());
-					return FCacheRecord();
+					return FOptionalCacheRecord();
 				}
 			}
 		}
@@ -1115,7 +1115,7 @@ private:
 			FPayload Payload = GetCachePayload(Key, Context, ValuePolicy, ValueObject, bAlwaysLoadInlineData);
 			if (Payload.IsNull())
 			{
-				return FCacheRecord();
+				return FOptionalCacheRecord();
 			}
 			RecordBuilder.SetValue(MoveTemp(Payload));
 		}
@@ -1126,7 +1126,7 @@ private:
 			FPayload Payload = GetCachePayload(Key, Context, AttachmentsPolicy, AttachmentField.AsObject(), bAlwaysLoadInlineData);
 			if (Payload.IsNull())
 			{
-				return FCacheRecord();
+				return FOptionalCacheRecord();
 			}
 			RecordBuilder.AddAttachment(MoveTemp(Payload));
 		}
diff --git a/Engine/Source/Developer/DerivedDataCache/Private/MemoryDerivedDataBackend.cpp b/Engine/Source/Developer/DerivedDataCache/Private/MemoryDerivedDataBackend.cpp
--- a/Engine/Source/Developer/DerivedDataCache/Private/MemoryDerivedDataBackend.cpp
+++ b/Engine/Source/Developer/DerivedDataCache/Private/MemoryDerivedDataBackend.cpp
@@ -483,17 +483,17 @@ FRequest FMemoryDerivedDataBackend::Get(
 	for (const FCacheKey& Key : Keys)
 	{
 		COOK_STAT(auto Timer = UsageStats.TimeGet());
-		FCacheRecord Record;
+		FOptionalCacheRecord Record;
 		if (FScopeLock ScopeLock(&SynchronizationObject); const FCacheRecord* CacheRecord = CacheRecords.Find(Key))
 		{
 			Record = *CacheRecord;
 		}
 		if (Record)
 		{
-			COOK_STAT(Timer.AddHit(CalcRawCacheRecordSize(Record)));
+			COOK_STAT(Timer.AddHit(CalcRawCacheRecordSize(Record.Get())));
 			if (OnComplete)
 			{
-				OnComplete({MoveTemp(Record), EStatus::Ok});
+				OnComplete({MoveTemp(Record).Get(), EStatus::Ok});
 			}
 		}
 		else
diff --git a/Engine/Source/Developer/DerivedDataCache/Public/DerivedDataBuild.h b/Engine/Source/Developer/DerivedDataCache/Public/DerivedDataBuild.h
--- a/Engine/Source/Developer/DerivedDataCache/Public/DerivedDataBuild.h
+++ b/Engine/Source/Developer/DerivedDataCache/Public/DerivedDataBuild.h
@@ -7,13 +7,13 @@
 class FCbObject;
 struct FGuid;
 
-namespace UE::DerivedData { class FBuildAction; }
 namespace UE::DerivedData { class FBuildActionBuilder; }
-namespace UE::DerivedData { class FBuildDefinition; }
 namespace UE::DerivedData { class FBuildDefinitionBuilder; }
-namespace UE::DerivedData { class FBuildOutput; }
 namespace UE::DerivedData { class FBuildOutputBuilder; }
 namespace UE::DerivedData { class FCacheRecord; }
+namespace UE::DerivedData { class FOptionalBuildAction; }
+namespace UE::DerivedData { class FOptionalBuildDefinition; }
+namespace UE::DerivedData { class FOptionalBuildOutput; }
 
 namespace UE::DerivedData
 {
@@ -41,9 +41,9 @@ public:
 	 *
 	 * @param Name         The name by which to identify this definition for logging and profiling.
 	 * @param Definition   An object saved from a build definition. Holds a reference and is cloned if not owned.
-	 * @return A valid build definition, or a build definition with an empty key on error.
+	 * @return A valid build definition, or null on error.
 	 */
-	virtual FBuildDefinition LoadDefinition(FStringView Name, FCbObject&& Definition) = 0;
+	virtual FOptionalBuildDefinition LoadDefinition(FStringView Name, FCbObject&& Definition) = 0;
 
 	/**
 	 * Create a build action builder.
@@ -58,9 +58,9 @@ public:
 	 *
 	 * @param Name       The name by which to identify this action for logging and profiling.
 	 * @param Action     The saved action to load.
-	 * @return A valid build action, or a build action with an empty key on error.
+	 * @return A valid build action, or null on error.
 	 */
-	virtual FBuildAction LoadAction(FStringView Name, FCbObject&& Action) = 0;
+	virtual FOptionalBuildAction LoadAction(FStringView Name, FCbObject&& Action) = 0;
 
 	/**
 	 * Create a build output builder.
@@ -76,9 +76,10 @@ public:
 	 * @param Name       The name by which to identify this output for logging and profiling.
 	 * @param Function   The name of the build function that produced this output.
 	 * @param Output     The saved output to load.
+	 * @return A valid build output, or null on error.
 	 */
-	virtual FBuildOutput LoadOutput(FStringView Name, FStringView Function, const FCbObject& Output) = 0;
-	virtual FBuildOutput LoadOutput(FStringView Name, FStringView Function, const FCacheRecord& Output) = 0;
+	virtual FOptionalBuildOutput LoadOutput(FStringView Name, FStringView Function, const FCbObject& Output) = 0;
+	virtual FOptionalBuildOutput LoadOutput(FStringView Name, FStringView Function, const FCacheRecord& Output) = 0;
 
 	/**
 	 * Returns the version of the build system.
diff --git a/Engine/Source/Developer/DerivedDataCache/Public/DerivedDataBuildAction.h b/Engine/Source/Developer/DerivedDataCache/Public/DerivedDataBuildAction.h
--- a/Engine/Source/Developer/DerivedDataCache/Public/DerivedDataBuildAction.h
+++ b/Engine/Source/Developer/DerivedDataCache/Public/DerivedDataBuildAction.h
@@ -130,6 +130,7 @@ public:
 	}
 
 private:
+	friend class FOptionalBuildAction;
 	friend FBuildAction Private::CreateBuildAction(Private::IBuildActionInternal* Action);
 
 	/** Construct a build action. Use Build() on a builder from IBuild::CreateAction(). */
@@ -182,4 +183,30 @@ private:
 	TUniquePtr<Private::IBuildActionBuilderInternal> ActionBuilder;
 };
 
+/**
+ * A build action that can be null.
+ *
+ * @see FBuildAction
+ */
+class FOptionalBuildAction : private FBuildAction
+{
+public:
+	inline FOptionalBuildAction() : FBuildAction(nullptr) {}
+
+	inline FOptionalBuildAction(FBuildAction&& InAction) : FBuildAction(MoveTemp(InAction)) {}
+	inline FOptionalBuildAction(const FBuildAction& InAction) : FBuildAction(InAction) {}
+	inline FOptionalBuildAction& operator=(FBuildAction&& InAction) { FBuildAction::operator=(MoveTemp(InAction)); return *this; }
+	inline FOptionalBuildAction& operator=(const FBuildAction& InAction) { FBuildAction::operator=(InAction); return *this; }
+
+	/** Returns the build action. The caller must check for null before using this accessor. */
+	inline const FBuildAction& Get() const & { return *this; }
+	inline FBuildAction&& Get() && { return MoveTemp(*this); }
+
+	inline bool IsNull() const { return !IsValid(); }
+	inline bool IsValid() const { return Action.IsValid(); }
+	inline explicit operator bool() const { return IsValid(); }
+
+	inline void Reset() { *this = FOptionalBuildAction(); }
+};
+
 } // UE::DerivedData
diff --git a/Engine/Source/Developer/DerivedDataCache/Public/DerivedDataBuildDefinition.h b/Engine/Source/Developer/DerivedDataCache/Public/DerivedDataBuildDefinition.h
--- a/Engine/Source/Developer/DerivedDataCache/Public/DerivedDataBuildDefinition.h
+++ b/Engine/Source/Developer/DerivedDataCache/Public/DerivedDataBuildDefinition.h
@@ -138,6 +138,7 @@ public:
 	}
 
 private:
+	friend class FOptionalBuildDefinition;
 	friend FBuildDefinition Private::CreateBuildDefinition(Private::IBuildDefinitionInternal* Definition);
 
 	/** Construct a build definition. Use Build() on a builder from IBuild::CreateDefinition(). */
@@ -222,4 +223,30 @@ private:
 	TUniquePtr<Private::IBuildDefinitionBuilderInternal> DefinitionBuilder;
 };
 
+/**
+ * A build definition that can be null.
+ *
+ * @see FBuildDefinition
+ */
+class FOptionalBuildDefinition : private FBuildDefinition
+{
+public:
+	inline FOptionalBuildDefinition() : FBuildDefinition(nullptr) {}
+
+	inline FOptionalBuildDefinition(FBuildDefinition&& InDefinition) : FBuildDefinition(MoveTemp(InDefinition)) {}
+	inline FOptionalBuildDefinition(const FBuildDefinition& InDefinition) : FBuildDefinition(InDefinition) {}
+	inline FOptionalBuildDefinition& operator=(FBuildDefinition&& InDefinition) { FBuildDefinition::operator=(MoveTemp(InDefinition)); return *this; }
+	inline FOptionalBuildDefinition& operator=(const FBuildDefinition& InDefinition) { FBuildDefinition::operator=(InDefinition); return *this; }
+
+	/** Returns the build definition. The caller must check for null before using this accessor. */
+	inline const FBuildDefinition& Get() const & { return *this; }
+	inline FBuildDefinition&& Get() && { return MoveTemp(*this); }
+
+	inline bool IsNull() const { return !IsValid(); }
+	inline bool IsValid() const { return Definition.IsValid(); }
+	inline explicit operator bool() const { return IsValid(); }
+
+	inline void Reset() { *this = FOptionalBuildDefinition(); }
+};
+
 } // UE::DerivedData
diff --git a/Engine/Source/Developer/DerivedDataCache/Public/DerivedDataBuildOutput.h b/Engine/Source/Developer/DerivedDataCache/Public/DerivedDataBuildOutput.h
--- a/Engine/Source/Developer/DerivedDataCache/Public/DerivedDataBuildOutput.h
+++ b/Engine/Source/Developer/DerivedDataCache/Public/DerivedDataBuildOutput.h
@@ -41,7 +41,6 @@ public:
 };
 
 FBuildOutput CreateBuildOutput(IBuildOutputInternal* Output);
-const IBuildOutputInternal* GetBuildOutput(const FBuildOutput& Output);
 
 class IBuildOutputBuilderInternal
 {
@@ -91,7 +90,7 @@ struct FBuildDiagnostic
 class FBuildOutput
 {
 public:
-	/** Returns the name of the build definition that produced this output. */
+	/** Returns the name by which to identify this output for logging and profiling. */
 	inline FStringView GetName() const { return Output->GetName(); }
 
 	/** Returns the name of the build function that produced this output. */
@@ -128,8 +127,8 @@ public:
 	}
 
 private:
+	friend class FOptionalBuildOutput;
 	friend FBuildOutput Private::CreateBuildOutput(Private::IBuildOutputInternal* Output);
-	friend const Private::IBuildOutputInternal* Private::GetBuildOutput(const FBuildOutput& Output);
 
 	inline explicit FBuildOutput(Private::IBuildOutputInternal* InOutput)
 		: Output(InOutput)
@@ -142,7 +141,7 @@ private:
 /**
  * A build output builder is used to construct a build output.
  *
- * Create using IBuild::CreateOutput() which must be given a build definition.
+ * Create using IBuild::CreateOutput().
  *
  * @see FBuildOutput
  */
@@ -198,4 +197,30 @@ private:
 	TUniquePtr<Private::IBuildOutputBuilderInternal> OutputBuilder;
 };
 
+/**
+ * A build output that can be null.
+ *
+ * @see FBuildOutput
+ */
+class FOptionalBuildOutput : private FBuildOutput
+{
+public:
+	inline FOptionalBuildOutput() : FBuildOutput(nullptr) {}
+
+	inline FOptionalBuildOutput(FBuildOutput&& InOutput) : FBuildOutput(MoveTemp(InOutput)) {}
+	inline FOptionalBuildOutput(const FBuildOutput& InOutput) : FBuildOutput(InOutput) {}
+	inline FOptionalBuildOutput& operator=(FBuildOutput&& InOutput) { FBuildOutput::operator=(MoveTemp(InOutput)); return *this; }
+	inline FOptionalBuildOutput& operator=(const FBuildOutput& InOutput) { FBuildOutput::operator=(InOutput); return *this; }
+
+	/** Returns the build output. The caller must check for null before using this accessor. */
+	inline const FBuildOutput& Get() const & { return *this; }
+	inline FBuildOutput&& Get() && { return MoveTemp(*this); }
+
+	inline bool IsNull() const { return !IsValid(); }
+	inline bool IsValid() const { return Output.IsValid(); }
+	inline explicit operator bool() const { return IsValid(); }
+
+	inline void Reset() { *this = FOptionalBuildOutput(); }
+};
+
 } // UE::DerivedData
diff --git a/Engine/Source/Developer/DerivedDataCache/Public/DerivedDataCacheRecord.h b/Engine/Source/Developer/DerivedDataCache/Public/DerivedDataCacheRecord.h
--- a/Engine/Source/Developer/DerivedDataCache/Public/DerivedDataCacheRecord.h
+++ b/Engine/Source/Developer/DerivedDataCache/Public/DerivedDataCacheRecord.h
@@ -79,10 +79,7 @@ namespace UE::DerivedData
 class FCacheRecord
 {
 public:
-	/** Construct a null cache record. */
-	explicit FCacheRecord() = default;
-
-	/** Returns the key that identifies this record in the cache. Always available in a non-null cache record. */
+	/** Returns the key that identifies this record in the cache. */
 	inline const FCacheKey& GetKey() const { return Record->GetKey(); }
 
 	/** Returns the metadata. Null when requested with ECachePolicy::SkipMeta. */
@@ -100,23 +97,14 @@ public:
 	/** Returns the attachment payload matching the ID. Null if no match. Data is null if attachments were skipped. */
 	inline const FPayload& GetAttachmentPayload(const FPayloadId& Id) const { return Record->GetAttachmentPayload(Id); }
 
-	/** Returns a view of the attachments. Always available in a non-null cache record, but data may be skipped. */
+	/** Returns a view of the attachments. Always available, but data may be skipped. */
 	inline TConstArrayView<FPayload> GetAttachmentPayloads() const { return Record->GetAttachmentPayloads(); }
 
 	/** Returns the payload matching the ID, whether value or attachment. Null if no match. Data is null if skipped. */
 	inline const FPayload& GetPayload(const FPayloadId& Id) const { return Record->GetPayload(Id); }
 
-	/** Whether this is null. */
-	inline bool IsNull() const { return !Record; }
-	/** Whether this is not null. */
-	inline bool IsValid() const { return !IsNull(); }
-	/** Whether this is not null. */
-	inline explicit operator bool() const { return IsValid(); }
-
-	/** Reset this to null. */
-	inline void Reset() { *this = FCacheRecord(); }
-
 private:
+	friend class FOptionalCacheRecord;
 	friend FCacheRecord Private::CreateCacheRecord(Private::ICacheRecordInternal* Record);
 
 	/** Construct a cache record. Use Build() or BuildAsync() on a builder from ICache::CreateRecord(). */
@@ -234,6 +222,32 @@ private:
 	TUniquePtr<Private::ICacheRecordBuilderInternal> RecordBuilder;
 };
 
+/**
+ * A cache record that can be null.
+ *
+ * @see FCacheRecord
+ */
+class FOptionalCacheRecord : private FCacheRecord
+{
+public:
+	inline FOptionalCacheRecord() : FCacheRecord(nullptr) {}
+
+	inline FOptionalCacheRecord(FCacheRecord&& InRecord) : FCacheRecord(MoveTemp(InRecord)) {}
+	inline FOptionalCacheRecord(const FCacheRecord& InRecord) : FCacheRecord(InRecord) {}
+	inline FOptionalCacheRecord& operator=(FCacheRecord&& InRecord) { FCacheRecord::operator=(MoveTemp(InRecord)); return *this; }
+	inline FOptionalCacheRecord& operator=(const FCacheRecord& InRecord) { FCacheRecord::operator=(InRecord); return *this; }
+
+	/** Returns the cache record. The caller must check for null before using this accessor. */
+	inline const FCacheRecord& Get() const & { return *this; }
+	inline FCacheRecord&& Get() && { return MoveTemp(*this); }
+
+	inline bool IsNull() const { return !IsValid(); }
+	inline bool IsValid() const { return Record.IsValid(); }
+	inline explicit operator bool() const { return IsValid(); }
+
+	inline void Reset() { *this = FOptionalCacheRecord(); }
+};
+
 } // UE::DerivedData
 
 ///////////////////////////////////////////////////////////////////////////////////////////////////////////////////////
