DDC: Fixed a build error caught by Clang

#rnx

diff --git a/Engine/Source/Developer/DerivedDataCache/Private/DerivedDataBuildOutput.cpp b/Engine/Source/Developer/DerivedDataCache/Private/DerivedDataBuildOutput.cpp
--- a/Engine/Source/Developer/DerivedDataCache/Private/DerivedDataBuildOutput.cpp
+++ b/Engine/Source/Developer/DerivedDataCache/Private/DerivedDataBuildOutput.cpp
@@ -130,8 +130,8 @@ FBuildOutputInternal::FBuildOutputInternal(FBuildOutputBuilderInternal&& InOutpu
 }
 
 FBuildOutputInternal::FBuildOutputInternal(FStringView InName, FStringView InFunction, const FCbObject& InOutput, bool& bOutIsValid)
-	: Name(Name)
-	, Function(Function)
+	: Name(InName)
+	, Function(InFunction)
 {
 	bOutIsValid = false;
 	checkf(!InName.IsEmpty(), TEXT("A build output requires a non-empty name."));
