compile fix for Windows-based platform extensions.

#rnx

diff --git a/Engine/Source/Developer/DerivedDataCache/DerivedDataCache.Build.cs b/Engine/Source/Developer/DerivedDataCache/DerivedDataCache.Build.cs
--- a/Engine/Source/Developer/DerivedDataCache/DerivedDataCache.Build.cs
+++ b/Engine/Source/Developer/DerivedDataCache/DerivedDataCache.Build.cs
@@ -21,5 +21,12 @@ public class DerivedDataCache : ModuleRules
 		{
 			DynamicallyLoadedModuleNames.Add("DDCUtils");
 		}
+
+		// Platform-specific opt-in
+		if (Target.Platform == UnrealTargetPlatform.Win64)
+		{
+			PrivateDefinitions.Add("WITH_HTTP_DDC_BACKEND=1");
+			PrivateDefinitions.Add("WITH_S3_DDC_BACKEND=1");
+		}
 	}
 }
diff --git a/Engine/Source/Developer/DerivedDataCache/Private/HttpDerivedDataBackend.h b/Engine/Source/Developer/DerivedDataCache/Private/HttpDerivedDataBackend.h
--- a/Engine/Source/Developer/DerivedDataCache/Private/HttpDerivedDataBackend.h
+++ b/Engine/Source/Developer/DerivedDataCache/Private/HttpDerivedDataBackend.h
@@ -7,10 +7,8 @@
 #include "DerivedDataCacheUsageStats.h"
 
 // Macro for whether to enable the S3 backend. libcurl is not currently available on Mac.
-#if PLATFORM_WINDOWS
-#define WITH_HTTP_DDC_BACKEND 1
-#else
-#define WITH_HTTP_DDC_BACKEND 0
+#if !defined(WITH_HTTP_DDC_BACKEND)
+	#define WITH_HTTP_DDC_BACKEND 0
 #endif
 
 #if WITH_HTTP_DDC_BACKEND
diff --git a/Engine/Source/Developer/DerivedDataCache/Private/S3DerivedDataBackend.h b/Engine/Source/Developer/DerivedDataCache/Private/S3DerivedDataBackend.h
--- a/Engine/Source/Developer/DerivedDataCache/Private/S3DerivedDataBackend.h
+++ b/Engine/Source/Developer/DerivedDataCache/Private/S3DerivedDataBackend.h
@@ -7,9 +7,7 @@
 #include "DerivedDataCacheUsageStats.h"
 
 // Macro for whether to enable the S3 backend. libcurl is not currently available on Mac.
-#if PLATFORM_WINDOWS
-	#define WITH_S3_DDC_BACKEND 1
-#else
+#if !defined(WITH_S3_DDC_BACKEND)
 	#define WITH_S3_DDC_BACKEND 0
 #endif
 
