Add Trace Cpu Events and Counters to FileSystem DDC backend.

diff --git a/Engine/Source/Developer/DerivedDataCache/Private/FileSystemDerivedDataBackend.cpp b/Engine/Source/Developer/DerivedDataCache/Private/FileSystemDerivedDataBackend.cpp
--- a/Engine/Source/Developer/DerivedDataCache/Private/FileSystemDerivedDataBackend.cpp
+++ b/Engine/Source/Developer/DerivedDataCache/Private/FileSystemDerivedDataBackend.cpp
@@ -15,6 +15,8 @@
 #include "Misc/ScopeExit.h"
 #include "Misc/ScopeLock.h"
 #include "Misc/StringBuilder.h"
+#include "ProfilingDebugging/CpuProfilerTrace.h"
+#include "ProfilingDebugging/CountersTrace.h"
 #include "Serialization/CompactBinary.h"
 #include "Serialization/CompactBinaryValidation.h"
 #include "Serialization/CompactBinaryWriter.h"
@@ -42,6 +44,15 @@
 namespace UE::DerivedData::Backends
 {
 
+TRACE_DECLARE_INT_COUNTER(FileSystemDDC_Exist, TEXT("FileSystemDDC Exist"));
+TRACE_DECLARE_INT_COUNTER(FileSystemDDC_ExistHit, TEXT("FileSystemDDC Exist Hit"));
+TRACE_DECLARE_INT_COUNTER(FileSystemDDC_Get, TEXT("FileSystemDDC Get"));
+TRACE_DECLARE_INT_COUNTER(FileSystemDDC_GetHit, TEXT("FileSystemDDC Get Hit"));
+TRACE_DECLARE_INT_COUNTER(FileSystemDDC_Put, TEXT("FileSystemDDC Put"));
+TRACE_DECLARE_INT_COUNTER(FileSystemDDC_PutHit, TEXT("FileSystemDDC Put Hit"));
+TRACE_DECLARE_INT_COUNTER(FileSystemDDC_BytesRead, TEXT("FileSystemDDC Bytes Read"));
+TRACE_DECLARE_INT_COUNTER(FileSystemDDC_BytesWritten, TEXT("FileSystemDDC Bytes Written"));
+
 FString BuildPathForCacheKey(const TCHAR* CacheKey)
 {
 	FString Key = FString(CacheKey).ToUpper();
@@ -544,6 +555,8 @@ public:
 	 */
 	virtual bool CachedDataProbablyExists(const TCHAR* CacheKey) override
 	{
+		TRACE_CPUPROFILER_EVENT_SCOPE(FileSystemDDC_Exist);
+		TRACE_COUNTER_INCREMENT(FileSystemDDC_Exist);
 		COOK_STAT(auto Timer = UsageStats.TimeProbablyExists());
 		check(IsUsable());
 
@@ -584,6 +597,7 @@ public:
 				AccessLogWriter->Append(CacheKey, Filename);
 			}
 
+			TRACE_COUNTER_INCREMENT(FileSystemDDC_ExistHit);
 			COOK_STAT(Timer.AddHit(0));
 		}
 
@@ -612,6 +626,8 @@ public:
 	 */
 	virtual bool GetCachedData(const TCHAR* CacheKey, TArray<uint8>& Data) override
 	{
+		TRACE_CPUPROFILER_EVENT_SCOPE(FileSystemDDC_Get);
+		TRACE_COUNTER_INCREMENT(FileSystemDDC_Get);
 		COOK_STAT(auto Timer = UsageStats.TimeGet());
 		check(IsUsable());
 		FString Filename = BuildFilename(CacheKey);
@@ -657,6 +673,8 @@ public:
 			}
 
 			UE_LOG(LogDerivedDataCache, Verbose, TEXT("%s: Cache hit on %s (%d bytes, %.02f secs, %.2fMB/s)"), *GetName(), CacheKey, Data.Num(), ReadDuration, ReadSpeed);
+			TRACE_COUNTER_INCREMENT(FileSystemDDC_GetHit);
+			TRACE_COUNTER_ADD(FileSystemDDC_BytesRead, int64(Data.Num()));
 			COOK_STAT(Timer.AddHit(Data.Num()));
 			return true;
 		}
@@ -695,6 +713,8 @@ public:
 	 */
 	virtual EPutStatus PutCachedData(const TCHAR* CacheKey, TArrayView<const uint8> Data, bool bPutEvenIfExists) override
 	{
+		TRACE_CPUPROFILER_EVENT_SCOPE(FileSystemDDC_Put);
+		TRACE_COUNTER_INCREMENT(FileSystemDDC_Put);
 		COOK_STAT(auto Timer = UsageStats.TimePut());
 		check(IsUsable());
 
@@ -717,6 +737,8 @@ public:
 
 			if (bPutEvenIfExists || !CachedDataProbablyExists(CacheKey))
 			{
+				TRACE_COUNTER_INCREMENT(FileSystemDDC_PutHit);
+				TRACE_COUNTER_ADD(FileSystemDDC_BytesWritten, int64(Data.Num()));
 				COOK_STAT(Timer.AddHit(Data.Num()));
 				check(Data.Num());
 				FString TempFilename(TEXT("temp.")); 
@@ -867,11 +889,14 @@ public:
 	{
 		for (const FCacheRecord& Record : Records)
 		{
+			TRACE_CPUPROFILER_EVENT_SCOPE(FileSystemDDC_Put);
+			TRACE_COUNTER_INCREMENT(FileSystemDDC_Put);
 			COOK_STAT(auto Timer = UsageStats.TimePut());
 			if (PutCacheRecord(Record, Context, Policy))
 			{
 				UE_LOG(LogDerivedDataCache, Verbose, TEXT("%s: Cache put complete for %s from '%.*s'"),
 					*CachePath, *WriteToString<96>(Record.GetKey()), Context.Len(), Context.GetData());
+				TRACE_COUNTER_ADD(FileSystemDDC_BytesWritten, MeasureCacheRecord(Record));
 				COOK_STAT(Timer.AddHit(MeasureCacheRecord(Record)));
 				if (OnComplete)
 				{
@@ -899,11 +924,14 @@ public:
 	{
 		for (const FCacheKey& Key : Keys)
 		{
+			TRACE_CPUPROFILER_EVENT_SCOPE(FileSystemDDC_Get);
+			TRACE_COUNTER_INCREMENT(FileSystemDDC_Get);
 			COOK_STAT(auto Timer = UsageStats.TimeGet());
 			if (FOptionalCacheRecord Record = GetCacheRecord(Key, Context, Policy))
 			{
 				UE_LOG(LogDerivedDataCache, Verbose, TEXT("%s: Cache hit for %s from '%.*s'"),
 					*CachePath, *WriteToString<96>(Key), Context.Len(), Context.GetData());
+				TRACE_COUNTER_ADD(FileSystemDDC_BytesRead, MeasureCacheRecord(Record.Get()));
 				COOK_STAT(Timer.AddHit(MeasureCacheRecord(Record.Get())));
 				if (OnComplete)
 				{
@@ -934,6 +962,8 @@ public:
 		FOptionalCacheRecord Record;
 		for (const FCachePayloadKey& Key : SortedKeys)
 		{
+			TRACE_CPUPROFILER_EVENT_SCOPE(FileSystemDDC_Get);
+			TRACE_COUNTER_INCREMENT(FileSystemDDC_Get);
 			COOK_STAT(auto Timer = UsageStats.TimeGet());
 			if (!Record || Record.Get().GetKey() != Key.CacheKey)
 			{
@@ -943,6 +973,7 @@ public:
 			{
 				UE_LOG(LogDerivedDataCache, Verbose, TEXT("%s: Cache hit for %s from '%.*s'"),
 					*CachePath, *WriteToString<96>(Key), Context.Len(), Context.GetData());
+				TRACE_COUNTER_ADD(FileSystemDDC_BytesRead, Payload.GetRawSize());
 				COOK_STAT(Timer.AddHit(Payload.GetRawSize()));
 				if (OnComplete)
 				{
