DDC: Added HasConstants() and HasInputs() to FBuildAction to mirror FBuildDefinition

#rnx

diff --git a/Engine/Source/Developer/DerivedDataCache/Private/DerivedDataBuildAction.cpp b/Engine/Source/Developer/DerivedDataCache/Private/DerivedDataBuildAction.cpp
--- a/Engine/Source/Developer/DerivedDataCache/Private/DerivedDataBuildAction.cpp
+++ b/Engine/Source/Developer/DerivedDataCache/Private/DerivedDataBuildAction.cpp
@@ -88,6 +88,9 @@ public:
 	const FGuid& GetFunctionVersion() const final { return FunctionVersion; }
 	const FGuid& GetBuildSystemVersion() const final { return BuildSystemVersion; }
 
+	bool HasConstants() const final;
+	bool HasInputs() const final;
+
 	void IterateConstants(TFunctionRef<void (FStringView Key, FCbObject&& Value)> Visitor) const final;
 	void IterateInputs(TFunctionRef<void (FStringView Key, const FIoHash& RawHash, uint64 RawSize)> Visitor) const final;
 
@@ -206,6 +209,16 @@ FBuildActionInternal::FBuildActionInternal(FStringView InName, FCbObject&& InAct
 			});
 }
 
+bool FBuildActionInternal::HasConstants() const
+{
+	return Action["Constants"_ASV].HasValue();
+}
+
+bool FBuildActionInternal::HasInputs() const
+{
+	return Action["Inputs"_ASV].HasValue();
+}
+
 void FBuildActionInternal::IterateConstants(TFunctionRef<void (FStringView Key, FCbObject&& Value)> Visitor) const
 {
 	for (FCbField Field : Action["Constants"_ASV])
diff --git a/Engine/Source/Developer/DerivedDataCache/Public/DerivedDataBuildAction.h b/Engine/Source/Developer/DerivedDataCache/Public/DerivedDataBuildAction.h
--- a/Engine/Source/Developer/DerivedDataCache/Public/DerivedDataBuildAction.h
+++ b/Engine/Source/Developer/DerivedDataCache/Public/DerivedDataBuildAction.h
@@ -3,7 +3,7 @@
 #pragma once
 
 #include "CoreTypes.h"
-#include "Containers/StringFwd.h"
+#include "Containers/StringView.h"
 #include "IO/IoHash.h"
 #include "Misc/ScopeExit.h"
 #include "Templates/Function.h"
@@ -40,6 +40,8 @@ public:
 	virtual FStringView GetFunction() const = 0;
 	virtual const FGuid& GetFunctionVersion() const = 0;
 	virtual const FGuid& GetBuildSystemVersion() const = 0;
+	virtual bool HasConstants() const = 0;
+	virtual bool HasInputs() const = 0;
 	virtual void IterateConstants(TFunctionRef<void (FStringView Key, FCbObject&& Value)> Visitor) const = 0;
 	virtual void IterateInputs(TFunctionRef<void (FStringView Key, const FIoHash& RawHash, uint64 RawSize)> Visitor) const = 0;
 	virtual void Save(FCbWriter& Writer) const = 0;
@@ -88,28 +90,22 @@ public:
 	inline const FBuildActionKey& GetKey() const { return Action->GetKey(); }
 
 	/** Returns the name by which to identify this action for logging and profiling. */
-	inline FStringView GetName() const
-	{
-		return Action->GetName();
-	}
+	inline FStringView GetName() const { return Action->GetName(); }
 
 	/** Returns the name of the build function with which to build this action. */
-	inline FStringView GetFunction() const
-	{
-		return Action->GetFunction();
-	}
+	inline FStringView GetFunction() const { return Action->GetFunction(); }
+
+	/** Returns whether the action has any constants. */
+	inline bool HasConstants() const { return Action->HasConstants(); }
+
+	/** Returns whether the action has any inputs. */
+	inline bool HasInputs() const { return Action->HasInputs(); }
 
 	/** Returns the version of the build function with which to build this action. */
-	inline const FGuid& GetFunctionVersion() const
-	{
-		return Action->GetFunctionVersion();
-	}
+	inline const FGuid& GetFunctionVersion() const { return Action->GetFunctionVersion(); }
 
 	/** Returns the version of the build system required to build this action. */
-	inline const FGuid& GetBuildSystemVersion() const
-	{
-		return Action->GetBuildSystemVersion();
-	}
+	inline const FGuid& GetBuildSystemVersion() const { return Action->GetBuildSystemVersion(); }
 
 	/** Visits every constant in order by key. */
 	inline void IterateConstants(TFunctionRef<void (FStringView Key, FCbObject&& Value)> Visitor) const
