DDC: Reordered EPriority to make operations like Min and Max make sense with it

#rnx

diff --git a/Engine/Source/Developer/DerivedDataCache/Public/DerivedDataRequest.h b/Engine/Source/Developer/DerivedDataCache/Public/DerivedDataRequest.h
--- a/Engine/Source/Developer/DerivedDataCache/Public/DerivedDataRequest.h
+++ b/Engine/Source/Developer/DerivedDataCache/Public/DerivedDataRequest.h
@@ -15,36 +15,36 @@ namespace UE::DerivedData
 enum class EPriority : uint8
 {
 	/**
-	 * Blocking is to be used only when the thread making the request will wait on completion of the
-	 * request before doing any other work. Requests at this priority level will be processed before
-	 * any request at a lower priority level. This priority permits a request to be processed on the
-	 * thread making the request. Waiting on a request may increase its priority to this level.
+	 * Lowest is the minimum priority for asynchronous requests, and primarily for requests that are
+	 * speculative in nature, while minimizing impact on other requests.
 	 */
-	Blocking,
+	Lowest,
 	/**
-	 * Highest is the maximum priority for asynchronous requests, and intended for requests that are
-	 * required to maintain interactivity of the program.
+	 * Low is intended for requests that are below the default priority, but are used for operations
+	 * that the program will execute now rather than at an unknown time in the future.
 	 */
-	Highest,
+	Low,
+	/**
+	 * Normal is intended as the default request priority.
+	 */
+	Normal,
 	/**
 	 * High is intended for requests that are on the critical path, but are not required to maintain
 	 * interactivity of the program.
 	 */
 	High,
 	/**
-	 * Normal is intended as the default request priority.
-	 */
-	Normal,
-	/**
-	 * Low is intended for requests that are below the default priority, but are used for operations
-	 * that the program will execute now rather than at an unknown time in the future.
+	 * Highest is the maximum priority for asynchronous requests, and intended for requests that are
+	 * required to maintain interactivity of the program.
 	 */
-	Low,
+	Highest,
 	/**
-	 * Lowest is the minimum priority for asynchronous requests, and primarily for requests that are
-	 * speculative in nature, while minimizing impact on other requests.
+	 * Blocking is to be used only when the thread making the request will wait on completion of the
+	 * request before doing any other work. Requests at this priority level will be processed before
+	 * any request at a lower priority level. This priority permits a request to be processed on the
+	 * thread making the request. Waiting on a request may increase its priority to this level.
 	 */
-	Lowest,
+	Blocking,
 };
 
 ///////////////////////////////////////////////////////////////////////////////////////////////////////////////////////
