DDC: Moved IBuildFunction configuration into a new Configure() function instead of several virtual functions

Also added CancelAsyncBuild and clarified the comments on how async builds work.

#rnx

diff --git a/Engine/Source/Developer/DerivedDataCache/Public/DerivedDataBuildFunction.h b/Engine/Source/Developer/DerivedDataCache/Public/DerivedDataBuildFunction.h
--- a/Engine/Source/Developer/DerivedDataCache/Public/DerivedDataBuildFunction.h
+++ b/Engine/Source/Developer/DerivedDataCache/Public/DerivedDataBuildFunction.h
@@ -9,11 +9,12 @@
 class FCbObject;
 class FCompressedBuffer;
 
-namespace UE::DerivedData { class FBuildCacheContext; }
+namespace UE::DerivedData { class FBuildConfigContext; }
 namespace UE::DerivedData { class FBuildContext; }
 namespace UE::DerivedData { class FCacheBucket; }
 namespace UE::DerivedData { class FPayload; }
 namespace UE::DerivedData { struct FPayloadId; }
+namespace UE::DerivedData { enum class EBuildPolicy : uint8; }
 namespace UE::DerivedData { enum class ECachePolicy : uint8; }
 
 namespace UE::DerivedData
@@ -30,6 +31,8 @@ namespace UE::DerivedData
  *
  * Functions must be registered by a build function factory to be available to the build system.
  *
+ * @see FBuildContext
+ * @see FBuildConfigContext
  * @see TBuildFunctionFactory
  */
 class IBuildFunction
@@ -44,40 +47,40 @@ public:
 	virtual FGuid GetVersion() const = 0;
 
 	/**
-	 * Determines the cache bucket and cache policy to use for the build.
+	 * Build the output for the input in the build context.
+	 *
+	 * The build is complete when the function returns, unless BeginAsyncBuild has been called on
+	 * the context, in which case a call to EndAsyncBuild from any thread marks the end.
 	 *
-	 * Default behavior is to use the function name as the cache bucket and ECachePolicy::Default.
+	 * The build is considered to be successful unless an error was logged during its execution.
 	 */
-	virtual void Cache(FBuildCacheContext& Context) const {}
+	virtual void Build(FBuildContext& Context) const = 0;
 
 	/**
-	 * Build the output for the input in the build context.
+	 * Configure the build based on its constants.
 	 *
-	 * Asynchronous build completion is possible by using BeginAsyncBuild and EndAsyncBuild provided
-	 * on the build context. Otherwise, the build is complete when this function returns.
+	 * The default configuration is to use the function name as the cache bucket, use the default
+	 * cache policy, use the default build policy, and enable non-deterministic output checks.
 	 */
-	virtual void Build(FBuildContext& Context) const = 0;
-
-	/** Whether the build may be executed on any thread and concurrently with any other build. */
-	virtual bool AllowsParallelBuild() const { return true; }
-
-	/** Whether the build may be executed externally, in another process or on a remote agent. */
-	virtual bool AllowsExternalBuild() const { return false; }
+	virtual void Configure(FBuildConfigContext& Context) const {}
 
 	/**
-	 * Whether the build produces deterministic output.
+	 * Cancel an asynchronous build.
 	 *
-	 * Builds must strive to always produce the same output for a given input. Depending on external
-	 * code can make this infeasible for some build functions. To opt out of automatic validation of
-	 * build determinism, override this function to return false.
+	 * This function is only called when canceling a build that has called BeginAsyncBuild, which
+	 * requires this to be implemented to return only after EndAsyncBuild has been called exactly
+	 * once on the given context. Due to timing, this can be called after EndAsyncBuild, and must
+	 * handle that case, but will never be called unless BeginAsyncBuild has been called.
 	 */
-	virtual bool IsDeterministic() const { return true; }
+	virtual void CancelAsyncBuild(FBuildContext& Context) const {}
 };
 
 /** A build context provides the inputs for a build function and saves its outputs. */
 class FBuildContext
 {
 public:
+	virtual ~FBuildContext() = default;
+
 	/** Returns the constant with the matching key, or an object with no fields if not found. */
 	virtual FCbObject GetConstant(FStringView Key) const = 0;
 
@@ -94,16 +97,16 @@ public:
 	virtual void SetCachePolicy(ECachePolicy Policy) = 0;
 
 	/**
-	 * Make this an asynchronous build by giving the caller responsibility for completing the build.
+	 * Make this an asynchronous build by making the caller responsible for completing the build.
 	 *
-	 * This is an advanced feature that bypasses many of the safety checks that are available within
-	 * the scope of the synchronous build function. Take extra care to only consume inputs available
-	 * in the build context and avoid reading any global state.
+	 * This is an advanced feature that bypasses many of the safety checks for synchronous builds
+	 * that validate that the build function only accesses inputs through the context. Take extra
+	 * care to only consume inputs available in the build context and avoid reading global state.
 	 *
-	 * This may be called at most once on a build context. The caller becomes responsible for ending
-	 * the build by calling EndAsyncBuild on this context once the build has finished. Once an async
-	 * build has begun, it may end from any thread, but the context may only be used from one thread
-	 * at a time.
+	 * This may be called at most once on a build context. A function that uses this must support
+	 * async cancellation by implementing CancelAsyncBuild. Once this has been called, the caller
+	 * is responsible for calling EndAsyncBuild to finish this build. An async build may end from
+	 * any thread, but the context is only safe to use from one thread at a time.
 	 */
 	virtual void BeginAsyncBuild() = 0;
 
@@ -115,18 +118,32 @@ public:
 	virtual void EndAsyncBuild() = 0;
 };
 
-/** A build cache context allows cache behavior to be modified based on constant inputs. */
-class FBuildCacheContext
+/** A build config context allows cache and build behavior to be modified based on constant inputs. */
+class FBuildConfigContext
 {
 public:
+	virtual ~FBuildConfigContext() = default;
+
 	/** Returns the constant with the matching key, or an object with no fields if not found. */
 	virtual FCbObject GetConstant(FStringView Key) const = 0;
 
-	/** Overrides the default cache bucket used when reading or writing this build in the cache. */
+	/** Overrides the cache bucket used when reading or writing this build in the cache. */
 	virtual void SetCacheBucket(FCacheBucket Bucket) = 0;
 
-	/** Overrides the default cache policy used when reading or writing this build in the cache. */
+	/** Returns the cache policy used when reading or writing this build in the cache. */
+	virtual ECachePolicy GetCachePolicy() const = 0;
+
+	/** Overrides the cache policy used when reading or writing this build in the cache. */
 	virtual void SetCachePolicy(ECachePolicy Policy) = 0;
+
+	/** Returns the build policy used when executing this build. */
+	virtual EBuildPolicy GetBuildPolicy() const = 0;
+
+	/** Overrides the build policy used when executing this build. */
+	virtual void SetBuildPolicy(EBuildPolicy Policy) = 0;
+
+	/** Skips verification that this function has deterministic output. */
+	virtual void SkipDeterministicOutputCheck() = 0;
 };
 
 } // UE::DerivedData
