DDC: Added UE::DerivedData::FBuildInput

FBuildInput and the associated builder type collect the inputs for a build.

#rnx

diff --git a/Engine/Source/Developer/DerivedDataCache/Private/DerivedDataBuildInput.cpp b/Engine/Source/Developer/DerivedDataCache/Private/DerivedDataBuildInput.cpp
new file mode 100644
--- /dev/null
+++ b/Engine/Source/Developer/DerivedDataCache/Private/DerivedDataBuildInput.cpp
@@ -0,0 +1,115 @@
+// Copyright Epic Games, Inc. All Rights Reserved.
+
+#include "DerivedDataBuildInput.h"
+
+#include "Compression/CompressedBuffer.h"
+#include "Containers/Map.h"
+#include "Containers/StringView.h"
+#include "Containers/UnrealString.h"
+#include "DerivedDataBuildPrivate.h"
+#include <atomic>
+
+namespace UE::DerivedData::Private
+{
+
+class FBuildInputBuilderInternal final : public IBuildInputBuilderInternal
+{
+public:
+	inline explicit FBuildInputBuilderInternal(FStringView InName)
+		: Name(InName)
+	{
+		checkf(!Name.IsEmpty(), TEXT("A build input requires a non-empty name."));
+	}
+
+	~FBuildInputBuilderInternal() final = default;
+
+	void AddInput(FStringView Key, const FCompressedBuffer& Buffer) final
+	{
+		const uint32 KeyHash = GetTypeHash(Key);
+		checkf(!Key.IsEmpty(), TEXT("Empty key used in input for build of '%s'."), *Name);
+		checkf(!Inputs.ContainsByHash(KeyHash, Key), TEXT("Duplicate key '%.*s' used in input for "),
+			TEXT("build of '%s'."), Key.Len(), Key.GetData(), *Name);
+		Inputs.EmplaceByHash(KeyHash, Key, Buffer);
+	}
+
+	FBuildInput Build() final;
+
+	FString Name;
+	TMap<FString, FCompressedBuffer> Inputs;
+};
+
+///////////////////////////////////////////////////////////////////////////////////////////////////////////////////////
+
+class FBuildInputInternal final : public IBuildInputInternal
+{
+public:
+	explicit FBuildInputInternal(FBuildInputBuilderInternal&& InputBuilder);
+
+	~FBuildInputInternal() final = default;
+
+	FStringView GetName() const final { return Name; }
+
+	const FCompressedBuffer& GetInput(FStringView Key) const final;
+
+	inline void AddRef() const final
+	{
+		ReferenceCount.fetch_add(1, std::memory_order_relaxed);
+	}
+
+	inline void Release() const final
+	{
+		if (ReferenceCount.fetch_sub(1, std::memory_order_acq_rel) == 1)
+		{
+			delete this;
+		}
+	}
+
+private:
+	FString Name;
+	TMap<FString, FCompressedBuffer> Inputs;
+	mutable std::atomic<uint32> ReferenceCount{0};
+};
+
+///////////////////////////////////////////////////////////////////////////////////////////////////////////////////////
+
+FBuildInputInternal::FBuildInputInternal(FBuildInputBuilderInternal&& InputBuilder)
+	: Name(MoveTemp(InputBuilder.Name))
+	, Inputs(MoveTemp(InputBuilder.Inputs))
+{
+	Inputs.KeySort(TLess<>());
+}
+
+const FCompressedBuffer& FBuildInputInternal::GetInput(FStringView Key) const
+{
+	if (const FCompressedBuffer* Buffer = Inputs.FindByHash(GetTypeHash(Key), Key))
+	{
+		return *Buffer;
+	}
+	return FCompressedBuffer::Null;
+}
+
+///////////////////////////////////////////////////////////////////////////////////////////////////////////////////////
+
+FBuildInput FBuildInputBuilderInternal::Build()
+{
+	return CreateBuildInput(new FBuildInputInternal(MoveTemp(*this)));
+}
+
+///////////////////////////////////////////////////////////////////////////////////////////////////////////////////////
+
+FBuildInput CreateBuildInput(IBuildInputInternal* Input)
+{
+	return FBuildInput(Input);
+}
+
+FBuildInputBuilder CreateBuildInputBuilder(IBuildInputBuilderInternal* InputBuilder)
+{
+	return FBuildInputBuilder(InputBuilder);
+}
+
+FBuildInputBuilder CreateBuildInput(FStringView Name)
+{
+	return CreateBuildInputBuilder(new FBuildInputBuilderInternal(Name));
+}
+
+} // UE::DerivedData::Private
diff --git a/Engine/Source/Developer/DerivedDataCache/Private/DerivedDataBuildPrivate.h b/Engine/Source/Developer/DerivedDataCache/Private/DerivedDataBuildPrivate.h
--- a/Engine/Source/Developer/DerivedDataCache/Private/DerivedDataBuildPrivate.h
+++ b/Engine/Source/Developer/DerivedDataCache/Private/DerivedDataBuildPrivate.h
@@ -10,6 +10,7 @@ struct FGuid;
 
 namespace UE::DerivedData { class FBuildActionBuilder; }
 namespace UE::DerivedData { class FBuildDefinitionBuilder; }
+namespace UE::DerivedData { class FBuildInputBuilder; }
 namespace UE::DerivedData { class FBuildOutputBuilder; }
 namespace UE::DerivedData { class FCacheRecord; }
 namespace UE::DerivedData { class FCacheRecordBuilder; }
@@ -28,6 +29,9 @@ FOptionalBuildDefinition LoadBuildDefinition(FStringView Name, FCbObject&& Defin
 FBuildActionBuilder CreateBuildAction(FStringView Name, FStringView Function, const FGuid& FunctionVersion, const FGuid& BuildSystemVersion);
 FOptionalBuildAction LoadBuildAction(FStringView Name, FCbObject&& Action);
 
+// Implemented in DerivedDataBuildInput.cpp
+FBuildInputBuilder CreateBuildInput(FStringView Name);
+
 // Implemented in DerivedDataBuildOutput.cpp
 FBuildOutputBuilder CreateBuildOutput(FStringView Name, FStringView Function);
 FOptionalBuildOutput LoadBuildOutput(FStringView Name, FStringView Function, const FCbObject& Output);
diff --git a/Engine/Source/Developer/DerivedDataCache/Private/DerivedDataCache.cpp b/Engine/Source/Developer/DerivedDataCache/Private/DerivedDataCache.cpp
--- a/Engine/Source/Developer/DerivedDataCache/Private/DerivedDataCache.cpp
+++ b/Engine/Source/Developer/DerivedDataCache/Private/DerivedDataCache.cpp
@@ -20,6 +20,7 @@
 #include "DerivedDataBuild.h"
 #include "DerivedDataBuildAction.h"
 #include "DerivedDataBuildDefinition.h"
+#include "DerivedDataBuildInput.h"
 #include "DerivedDataBuildOutput.h"
 #include "DerivedDataBuildPrivate.h"
 #include "DerivedDataCachePrivate.h"
@@ -739,6 +740,11 @@ public:
 		return Private::LoadBuildAction(Name, MoveTemp(Action));
 	}
 
+	FBuildInputBuilder CreateInput(FStringView Name) final
+	{
+		return Private::CreateBuildInput(Name);
+	}
+
 	FBuildOutputBuilder CreateOutput(FStringView Name, FStringView Function) final
 	{
 		return Private::CreateBuildOutput(Name, Function);
diff --git a/Engine/Source/Developer/DerivedDataCache/Public/DerivedDataBuild.h b/Engine/Source/Developer/DerivedDataCache/Public/DerivedDataBuild.h
--- a/Engine/Source/Developer/DerivedDataCache/Public/DerivedDataBuild.h
+++ b/Engine/Source/Developer/DerivedDataCache/Public/DerivedDataBuild.h
@@ -9,6 +9,7 @@ struct FGuid;
 
 namespace UE::DerivedData { class FBuildActionBuilder; }
 namespace UE::DerivedData { class FBuildDefinitionBuilder; }
+namespace UE::DerivedData { class FBuildInputBuilder; }
 namespace UE::DerivedData { class FBuildOutputBuilder; }
 namespace UE::DerivedData { class FCacheRecord; }
 namespace UE::DerivedData { class FOptionalBuildAction; }
@@ -56,12 +57,19 @@ public:
 	/**
 	 * Load a build action from compact binary.
 	 *
-	 * @param Name       The name by which to identify this action for logging and profiling.
-	 * @param Action     The saved action to load.
+	 * @param Name     The name by which to identify this action for logging and profiling.
+	 * @param Action   The saved action to load.
 	 * @return A valid build action, or null on error.
 	 */
 	virtual FOptionalBuildAction LoadAction(FStringView Name, FCbObject&& Action) = 0;
 
+	/**
+	 * Create a build input builder.
+	 *
+	 * @param Name   The name by which to identify this input for logging and profiling.
+	 */
+	virtual FBuildInputBuilder CreateInput(FStringView Name) = 0;
+
 	/**
 	 * Create a build output builder.
 	 *
diff --git a/Engine/Source/Developer/DerivedDataCache/Public/DerivedDataBuildInput.h b/Engine/Source/Developer/DerivedDataCache/Public/DerivedDataBuildInput.h
new file mode 100644
--- /dev/null
+++ b/Engine/Source/Developer/DerivedDataCache/Public/DerivedDataBuildInput.h
@@ -0,0 +1,136 @@
+// Copyright Epic Games, Inc. All Rights Reserved.
+
+#pragma once
+
+#include "CoreTypes.h"
+#include "Containers/StringView.h"
+#include "Misc/ScopeExit.h"
+#include "Templates/RefCounting.h"
+#include "Templates/UniquePtr.h"
+
+class FCompressedBuffer;
+
+namespace UE::DerivedData { class FBuildInput; }
+namespace UE::DerivedData { class FBuildInputBuilder; }
+
+namespace UE::DerivedData::Private
+{
+
+class IBuildInputInternal
+{
+public:
+	virtual ~IBuildInputInternal() = default;
+	virtual FStringView GetName() const = 0;
+	virtual const FCompressedBuffer& GetInput(FStringView Key) const = 0;
+	virtual void AddRef() const = 0;
+	virtual void Release() const = 0;
+};
+
+FBuildInput CreateBuildInput(IBuildInputInternal* Input);
+
+class IBuildInputBuilderInternal
+{
+public:
+	virtual ~IBuildInputBuilderInternal() = default;
+	virtual void AddInput(FStringView Key, const FCompressedBuffer& Buffer) = 0;
+	virtual FBuildInput Build() = 0;
+};
+
+FBuildInputBuilder CreateBuildInputBuilder(IBuildInputBuilderInternal* InputBuilder);
+
+} // UE::DerivedData::Private
+
+namespace UE::DerivedData
+{
+
+/**
+ * A build input is an immutable container of key/value pairs for the inputs to a build function.
+ *
+ * The keys for inputs are names that are unique within the build input.
+ *
+ * @see FBuildAction
+ */
+class FBuildInput
+{
+public:
+	/** Returns the name by which to identify this input for logging and profiling. */
+	inline FStringView GetName() const { return Input->GetName(); }
+
+	/** Visits every input in order by key. */
+	inline const FCompressedBuffer& GetInput(FStringView Key) const { return Input->GetInput(Key); }
+
+private:
+	friend class FOptionalBuildInput;
+	friend FBuildInput Private::CreateBuildInput(Private::IBuildInputInternal* Input);
+
+	/** Construct a build input. Use Build() on a builder from IBuild::CreateInput(). */
+	inline explicit FBuildInput(Private::IBuildInputInternal* InInput)
+		: Input(InInput)
+	{
+	}
+
+	TRefCountPtr<Private::IBuildInputInternal> Input;
+};
+
+/**
+ * A build input builder is used to construct a build input.
+ *
+ * Create using IBuild::CreateInput().
+ *
+ * @see FBuildInput
+ */
+class FBuildInputBuilder
+{
+public:
+	/** Add an input with a key that is unique within this input. */
+	inline void AddInput(FStringView Key, const FCompressedBuffer& Buffer)
+	{
+		InputBuilder->AddInput(Key, Buffer);
+	}
+
+	/** Build a build input, which makes this builder subsequently unusable. */
+	inline FBuildInput Build()
+	{
+		ON_SCOPE_EXIT { InputBuilder = nullptr; };
+		return InputBuilder->Build();
+	}
+
+private:
+	friend FBuildInputBuilder Private::CreateBuildInputBuilder(Private::IBuildInputBuilderInternal* InputBuilder);
+
+	/** Construct a build input builder. Use IBuild::CreateInput(). */
+	inline explicit FBuildInputBuilder(Private::IBuildInputBuilderInternal* InInputBuilder)
+		: InputBuilder(InInputBuilder)
+	{
+	}
+
+	TUniquePtr<Private::IBuildInputBuilderInternal> InputBuilder;
+};
+
+/**
+ * A build input that can be null.
+ *
+ * @see FBuildInput
+ */
+class FOptionalBuildInput : private FBuildInput
+{
+public:
+	inline FOptionalBuildInput() : FBuildInput(nullptr) {}
+
+	inline FOptionalBuildInput(FBuildInput&& InInput) : FBuildInput(MoveTemp(InInput)) {}
+	inline FOptionalBuildInput(const FBuildInput& InInput) : FBuildInput(InInput) {}
+	inline FOptionalBuildInput& operator=(FBuildInput&& InInput) { FBuildInput::operator=(MoveTemp(InInput)); return *this; }
+	inline FOptionalBuildInput& operator=(const FBuildInput& InInput) { FBuildInput::operator=(InInput); return *this; }
+
+	/** Returns the build input. The caller must check for null before using this accessor. */
+	inline const FBuildInput& Get() const & { return *this; }
+	inline FBuildInput&& Get() && { return MoveTemp(*this); }
+
+	inline bool IsNull() const { return !IsValid(); }
+	inline bool IsValid() const { return Input.IsValid(); }
+	inline explicit operator bool() const { return IsValid(); }
+
+	inline void Reset() { *this = FOptionalBuildInput(); }
+};
+
+} // UE::DerivedData
