DDC: Made name and function asserts consistent and reordered members for debuggability

#rnx

diff --git a/Engine/Source/Developer/DerivedDataCache/Private/DerivedDataBuildAction.cpp b/Engine/Source/Developer/DerivedDataCache/Private/DerivedDataBuildAction.cpp
--- a/Engine/Source/Developer/DerivedDataCache/Private/DerivedDataBuildAction.cpp
+++ b/Engine/Source/Developer/DerivedDataCache/Private/DerivedDataBuildAction.cpp
@@ -32,9 +32,8 @@ public:
 		, BuildSystemVersion(InBuildSystemVersion)
 	{
 		checkf(!Name.IsEmpty(), TEXT("A build action requires a non-empty name."));
-		checkf(!Function.IsEmpty(), TEXT("A build action requires a non-empty function name for build of '%s'."), *Name);
-		checkf(Algo::AllOf(Function, FChar::IsAlnum), TEXT("Action requires an alphanumeric function name for ")
-			TEXT("build of '%s' by %s."), *Name, *Function);
+		checkf(!Function.IsEmpty() && Algo::AllOf(Function, FChar::IsAlnum),
+			TEXT("A build function name must be alphanumeric and non-empty for build of '%s' by %s."), *Name, *Function);
 	}
 
 	~FBuildActionBuilderInternal() final = default;
@@ -110,13 +109,13 @@ public:
 	}
 
 private:
-	mutable std::atomic<uint32> ReferenceCount{0};
 	FString Name;
 	FString Function;
 	FGuid FunctionVersion;
 	FGuid BuildSystemVersion;
 	FCbObject Action;
 	FBuildActionKey Key;
+	mutable std::atomic<uint32> ReferenceCount{0};
 };
 
 ///////////////////////////////////////////////////////////////////////////////////////////////////////////////////////
@@ -193,7 +192,7 @@ FBuildActionInternal::FBuildActionInternal(FStringView InName, FCbObject&& InAct
 	, Action(MoveTemp(InAction))
 	, Key{Action.GetHash()}
 {
-	checkf(!InName.IsEmpty(), TEXT("A build action requires a non-empty name."));
+	checkf(!Name.IsEmpty(), TEXT("A build action requires a non-empty name."));
 	Action.MakeOwned();
 	bOutIsValid = Action
 		&& !Function.IsEmpty() && Algo::AllOf(Function, FChar::IsAlnum)
diff --git a/Engine/Source/Developer/DerivedDataCache/Private/DerivedDataBuildDefinition.cpp b/Engine/Source/Developer/DerivedDataCache/Private/DerivedDataBuildDefinition.cpp
--- a/Engine/Source/Developer/DerivedDataCache/Private/DerivedDataBuildDefinition.cpp
+++ b/Engine/Source/Developer/DerivedDataCache/Private/DerivedDataBuildDefinition.cpp
@@ -26,9 +26,8 @@ public:
 		, Function(InFunction)
 	{
 		checkf(!Name.IsEmpty(), TEXT("A build definition requires a non-empty name."));
-		checkf(!Function.IsEmpty(), TEXT("A build definition requires a non-empty function name for build of '%s'."), *Name);
-		checkf(Algo::AllOf(Function, FChar::IsAlnum), TEXT("Definition requires an alphanumeric function name for ")
-			TEXT("build of '%s' by %s."), *Name, *Function);
+		checkf(!Function.IsEmpty() && Algo::AllOf(Function, FChar::IsAlnum),
+			TEXT("A build function name must be alphanumeric and non-empty for build of '%s' by %s."), *Name, *Function);
 	}
 
 	~FBuildDefinitionBuilderInternal() final = default;
@@ -118,11 +117,11 @@ public:
 	}
 
 private:
-	mutable std::atomic<uint32> ReferenceCount{0};
 	FString Name;
 	FString Function;
 	FCbObject Definition;
 	FBuildKey Key;
+	mutable std::atomic<uint32> ReferenceCount{0};
 };
 
 ///////////////////////////////////////////////////////////////////////////////////////////////////////////////////////
@@ -259,7 +258,7 @@ FBuildDefinitionInternal::FBuildDefinitionInternal(FStringView InName, FCbObject
 	, Definition(MoveTemp(InDefinition))
 	, Key{Definition.GetHash()}
 {
-	checkf(!InName.IsEmpty(), TEXT("A build definition requires a non-empty name."));
+	checkf(!Name.IsEmpty(), TEXT("A build definition requires a non-empty name."));
 	Definition.MakeOwned();
 	bOutIsValid = Definition
 		&& !Function.IsEmpty() && Algo::AllOf(Function, FChar::IsAlnum)
diff --git a/Engine/Source/Developer/DerivedDataCache/Private/DerivedDataBuildOutput.cpp b/Engine/Source/Developer/DerivedDataCache/Private/DerivedDataBuildOutput.cpp
--- a/Engine/Source/Developer/DerivedDataCache/Private/DerivedDataBuildOutput.cpp
+++ b/Engine/Source/Developer/DerivedDataCache/Private/DerivedDataBuildOutput.cpp
@@ -45,12 +45,15 @@ static void LexFromString(EBuildDiagnosticLevel& OutLevel, FAnsiStringView Strin
 class FBuildOutputBuilderInternal final : public IBuildOutputBuilderInternal
 {
 public:
-	explicit FBuildOutputBuilderInternal(FStringView Name, FStringView Function)
-		: Name(Name)
-		, Function(Function)
+	explicit FBuildOutputBuilderInternal(FStringView InName, FStringView InFunction)
+		: Name(InName)
+		, Function(InFunction)
 		, bHasError(false)
 		, bHasDiagnostics(false)
 	{
+		checkf(!Name.IsEmpty(), TEXT("A build output requires a non-empty name."));
+		checkf(!Function.IsEmpty() && Algo::AllOf(Function, FChar::IsAlnum),
+			TEXT("A build function name must be alphanumeric and non-empty for build of '%s' by %s."), *Name, *Function);
 		DiagnosticsWriter.BeginArray();
 	}
 
@@ -65,9 +68,9 @@ public:
 	FString Name;
 	FString Function;
 	FCbObject Meta;
+	TArray<FPayload> Payloads;
 	FCbField Diagnostics;
 	FCbWriter DiagnosticsWriter;
-	TArray<FPayload> Payloads;
 	bool bHasError;
 	bool bHasDiagnostics;
 };
@@ -75,9 +78,9 @@ public:
 class FBuildOutputInternal final : public IBuildOutputInternal
 {
 public:
-	explicit FBuildOutputInternal(FBuildOutputBuilderInternal&& InOutput);
-	explicit FBuildOutputInternal(FStringView Name, FStringView Function, const FCbObject& InOutput, bool& bOutIsValid);
-	explicit FBuildOutputInternal(FStringView Name, FStringView Function, const FCacheRecord& InOutput, bool& bOutIsValid);
+	explicit FBuildOutputInternal(FBuildOutputBuilderInternal&& Output);
+	explicit FBuildOutputInternal(FStringView Name, FStringView Function, const FCbObject& Output, bool& bOutIsValid);
+	explicit FBuildOutputInternal(FStringView Name, FStringView Function, const FCacheRecord& Output, bool& bOutIsValid);
 
 	~FBuildOutputInternal() final = default;
 
@@ -110,12 +113,12 @@ public:
 	}
 
 private:
-	mutable std::atomic<uint32> ReferenceCount{0};
 	FString Name;
 	FString Function;
+	FCbObject Meta;
 	TArray<FPayload> Payloads;
 	FCbField Diagnostics;
-	FCbObject Meta;
+	mutable std::atomic<uint32> ReferenceCount{0};
 };
 
 ///////////////////////////////////////////////////////////////////////////////////////////////////////////////////////
@@ -123,9 +126,9 @@ private:
 FBuildOutputInternal::FBuildOutputInternal(FBuildOutputBuilderInternal&& InOutput)
 	: Name(MoveTemp(InOutput.Name))
 	, Function(MoveTemp(InOutput.Function))
+	, Meta(MoveTemp(InOutput.Meta))
 	, Payloads(MoveTemp(InOutput.Payloads))
 	, Diagnostics(MoveTemp(InOutput.Diagnostics))
-	, Meta(MoveTemp(InOutput.Meta))
 {
 }
 
@@ -134,7 +137,11 @@ FBuildOutputInternal::FBuildOutputInternal(FStringView InName, FStringView InFun
 	, Function(InFunction)
 {
 	bOutIsValid = false;
-	checkf(!InName.IsEmpty(), TEXT("A build output requires a non-empty name."));
+	checkf(!Name.IsEmpty(), TEXT("A build output requires a non-empty name."));
+	checkf(!Function.IsEmpty() && Algo::AllOf(Function, FChar::IsAlnum),
+		TEXT("A build function name must be alphanumeric and non-empty for build of '%s' by %s."), *Name, *Function);
+	Meta = InOutput["Meta"_ASV].AsObject();
+	Meta.MakeOwned();
 	for (FCbFieldView PayloadField : InOutput["Payloads"_ASV])
 	{
 		const FCbObjectView Payload = PayloadField.AsObjectView();
@@ -148,18 +155,18 @@ FBuildOutputInternal::FBuildOutputInternal(FStringView InName, FStringView InFun
 	}
 	Diagnostics = InOutput["Diagnostics"_ASV];
 	Diagnostics.MakeOwned();
-	Meta = InOutput["Meta"_ASV].AsObject();
-	Meta.MakeOwned();
 	bOutIsValid = IsValid() && (!InOutput.FindView("Meta"_ASV) || InOutput.FindView("Meta"_ASV).IsObject());
 }
 
 FBuildOutputInternal::FBuildOutputInternal(FStringView InName, FStringView InFunction, const FCacheRecord& InOutput, bool& bOutIsValid)
 	: Name(InName)
 	, Function(InFunction)
-	, Payloads(InOutput.GetAttachmentPayloads())
 	, Meta(InOutput.GetMeta())
+	, Payloads(InOutput.GetAttachmentPayloads())
 {
-	checkf(!InName.IsEmpty(), TEXT("A build output requires a non-empty name."));
+	checkf(!Name.IsEmpty(), TEXT("A build output requires a non-empty name."));
+	checkf(!Function.IsEmpty() && Algo::AllOf(Function, FChar::IsAlnum),
+		TEXT("A build function name must be alphanumeric and non-empty for build of '%s' by %s."), *Name, *Function);
 	if (FSharedBuffer Buffer = InOutput.GetValue())
 	{
 		Diagnostics = FCbObject(MoveTemp(Buffer))["Diagnostics"_ASV];
diff --git a/Engine/Source/Developer/DerivedDataCache/Private/DerivedDataCacheRecord.cpp b/Engine/Source/Developer/DerivedDataCache/Private/DerivedDataCacheRecord.cpp
--- a/Engine/Source/Developer/DerivedDataCache/Private/DerivedDataCacheRecord.cpp
+++ b/Engine/Source/Developer/DerivedDataCache/Private/DerivedDataCacheRecord.cpp
@@ -76,7 +76,6 @@ public:
 	}
 
 private:
-	mutable std::atomic<uint32> ReferenceCount{0};
 	FCacheKey Key;
 	FCbObject Meta;
 	FPayload Value;
@@ -84,6 +83,7 @@ private:
 	mutable FRWLock CacheLock;
 	mutable FSharedBuffer ValueCache;
 	mutable TArray<FSharedBuffer> AttachmentsCache;
+	mutable std::atomic<uint32> ReferenceCount{0};
 };
 
 ///////////////////////////////////////////////////////////////////////////////////////////////////////////////////////
@@ -202,10 +202,10 @@ FPayloadId FCacheRecordBuilderInternal::SetValue(const FPayload& Payload)
 	checkf(Payload, TEXT("Failed to set value on %s because the payload is null."), *WriteToString<96>(Key));
 	const FPayloadId& Id = Payload.GetId();
 	checkf(Value.IsNull(),
-		TEXT("Cache: Failed to set value on %s with ID %s because it has an existing value with ID %s."),
+		TEXT("Failed to set value on %s with ID %s because it has an existing value with ID %s."),
 		*WriteToString<96>(Key), *WriteToString<32>(Id), *WriteToString<32>(Value.GetId()));
 	checkf(Algo::BinarySearchBy(Attachments, Id, &FPayload::GetId) == INDEX_NONE,
-		TEXT("Failed to set on %s with ID %s because it has an existing attachment with that ID."),
+		TEXT("Failed to set value on %s with ID %s because it has an existing attachment with that ID."),
 		*WriteToString<96>(Key), *WriteToString<32>(Id));
 	Value = Payload;
 	return Id;
diff --git a/Engine/Source/Developer/DerivedDataCache/Public/DerivedDataBuildAction.h b/Engine/Source/Developer/DerivedDataCache/Public/DerivedDataBuildAction.h
--- a/Engine/Source/Developer/DerivedDataCache/Public/DerivedDataBuildAction.h
+++ b/Engine/Source/Developer/DerivedDataCache/Public/DerivedDataBuildAction.h
@@ -68,7 +68,7 @@ namespace UE::DerivedData
 {
 
 /**
- * A build action references a build function and the inputs to that build function.
+ * A build action is an immutable reference to a build function and its inputs.
  *
  * The purpose of an action is to capture everything required to execute a derived data build for
  * a fixed version of the build function and its constants and inputs.
@@ -78,8 +78,6 @@ namespace UE::DerivedData
  *
  * The keys for constants and inputs are names that are unique within the build action.
  *
- * The build action is immutable, and is created by a build session from a build definition.
- *
  * @see FBuildDefinition
  * @see FBuildSession
  */
diff --git a/Engine/Source/Developer/DerivedDataCache/Public/DerivedDataBuildDefinition.h b/Engine/Source/Developer/DerivedDataCache/Public/DerivedDataBuildDefinition.h
--- a/Engine/Source/Developer/DerivedDataCache/Public/DerivedDataBuildDefinition.h
+++ b/Engine/Source/Developer/DerivedDataCache/Public/DerivedDataBuildDefinition.h
@@ -63,7 +63,7 @@ namespace UE::DerivedData
 {
 
 /**
- * A build definition references a build function and the inputs to that build function.
+ * A build definition is an immutable reference to a build function and its inputs.
  *
  * The purpose of a definition is to capture everything required to execute a derived data build.
  * The definition is partly fixed (function name, constants, input hashes) and is partly variable
@@ -78,8 +78,6 @@ namespace UE::DerivedData
  * To build a definition against a specific version of the function and inputs, queue it to build
  * on a build session, which uses its build input provider to convert it to a build action.
  *
- * Build definitions are immutable, and are created by a builder from IBuild::CreateDefinition().
- *
  * @see FBuildDefinitionBuilder
  * @see FBuildSession
  */
diff --git a/Engine/Source/Developer/DerivedDataCache/Public/DerivedDataBuildOutput.h b/Engine/Source/Developer/DerivedDataCache/Public/DerivedDataBuildOutput.h
--- a/Engine/Source/Developer/DerivedDataCache/Public/DerivedDataBuildOutput.h
+++ b/Engine/Source/Developer/DerivedDataCache/Public/DerivedDataBuildOutput.h
@@ -81,11 +81,11 @@ struct FBuildDiagnostic
 };
 
 /**
- * A build output contains the payloads and diagnostics produced by a build.
+ * A build output is an immutable container of payloads and diagnostics produced by a build.
  *
  * The output will not contain any payloads if it has any errors.
  *
- * The output can be requested without data, which means the payloads will have null data.
+ * The output can be requested without data, which means that the payloads will have null data.
  */
 class FBuildOutput
 {
