DDC: Excluded payloads from build outputs with errors

The payloads are not used for failed builds and will cause unnecessary overhead to copy over the network for remote builds.

#rnx

diff --git a/Engine/Source/Developer/DerivedDataCache/Private/DerivedDataBuildOutput.cpp b/Engine/Source/Developer/DerivedDataCache/Private/DerivedDataBuildOutput.cpp
--- a/Engine/Source/Developer/DerivedDataCache/Private/DerivedDataBuildOutput.cpp
+++ b/Engine/Source/Developer/DerivedDataCache/Private/DerivedDataBuildOutput.cpp
@@ -130,6 +130,10 @@ FBuildOutputInternal::FBuildOutputInternal(FBuildOutputBuilderInternal&& InOutpu
 	, Payloads(MoveTemp(InOutput.Payloads))
 	, Diagnostics(MoveTemp(InOutput.Diagnostics))
 {
+	if (InOutput.bHasError)
+	{
+		Payloads.Empty();
+	}
 }
 
 FBuildOutputInternal::FBuildOutputInternal(FStringView InName, FStringView InFunction, const FCbObject& InOutput, bool& bOutIsValid)
