DDC: Moved the module and IBuild implementation to their own files

#rnx

diff --git a/Engine/Source/Developer/DerivedDataCache/Private/DerivedDataBuild.cpp b/Engine/Source/Developer/DerivedDataCache/Private/DerivedDataBuild.cpp
new file mode 100644
--- /dev/null
+++ b/Engine/Source/Developer/DerivedDataCache/Private/DerivedDataBuild.cpp
@@ -0,0 +1,82 @@
+// Copyright Epic Games, Inc. All Rights Reserved.
+
+#include "DerivedDataBuild.h"
+
+#include "DerivedDataBuildAction.h"
+#include "DerivedDataBuildDefinition.h"
+#include "DerivedDataBuildInput.h"
+#include "DerivedDataBuildOutput.h"
+#include "DerivedDataBuildPrivate.h"
+#include "DerivedDataCache.h"
+#include "Misc/Guid.h"
+
+namespace UE::DerivedData::Private
+{
+
+class FBuild final : public IBuild
+{
+public:
+	explicit FBuild(ICache& InCache)
+		: Cache(InCache)
+	{
+	}
+
+	FBuildDefinitionBuilder CreateDefinition(FStringView Name, FStringView Function) final
+	{
+		return CreateBuildDefinition(Name, Function);
+	}
+
+	FOptionalBuildDefinition LoadDefinition(FStringView Name, FCbObject&& Definition) final
+	{
+		return LoadBuildDefinition(Name, MoveTemp(Definition));
+	}
+
+	FBuildActionBuilder CreateAction(FStringView Name, FStringView Function) final
+	{
+		// DDC-TODO: Find the function version from the function registry.
+		return CreateBuildAction(Name, Function, FGuid::NewGuid(), Version);
+	}
+
+	FOptionalBuildAction LoadAction(FStringView Name, FCbObject&& Action) final
+	{
+		return LoadBuildAction(Name, MoveTemp(Action));
+	}
+
+	FBuildInputBuilder CreateInput(FStringView Name) final
+	{
+		return CreateBuildInput(Name);
+	}
+
+	FBuildOutputBuilder CreateOutput(FStringView Name, FStringView Function) final
+	{
+		return CreateBuildOutput(Name, Function);
+	}
+
+	FOptionalBuildOutput LoadOutput(FStringView Name, FStringView Function, const FCbObject& Output) final
+	{
+		return LoadBuildOutput(Name, Function, Output);
+	}
+
+	FOptionalBuildOutput LoadOutput(FStringView Name, FStringView Function, const FCacheRecord& Output) final
+	{
+		return LoadBuildOutput(Name, Function, Output);
+	}
+
+	const FGuid& GetVersion() const final
+	{
+		return Version;
+	}
+
+private:
+	ICache& Cache;
+	const FGuid Version{TEXT("ac0574e5-62bd-4c2e-84ec-f2efe48c0fef")};
+};
+
+///////////////////////////////////////////////////////////////////////////////////////////////////////////////////////
+
+IBuild* CreateBuild(ICache& Cache)
+{
+	return new FBuild(Cache);
+}
+
+} // UE::DerivedData::Private
diff --git a/Engine/Source/Developer/DerivedDataCache/Private/DerivedDataBuildPrivate.h b/Engine/Source/Developer/DerivedDataCache/Private/DerivedDataBuildPrivate.h
--- a/Engine/Source/Developer/DerivedDataCache/Private/DerivedDataBuildPrivate.h
+++ b/Engine/Source/Developer/DerivedDataCache/Private/DerivedDataBuildPrivate.h
@@ -17,10 +17,15 @@ namespace UE::DerivedData { class FCacheRecordBuilder; }
 namespace UE::DerivedData { class FOptionalBuildAction; }
 namespace UE::DerivedData { class FOptionalBuildDefinition; }
 namespace UE::DerivedData { class FOptionalBuildOutput; }
+namespace UE::DerivedData { class IBuild; }
+namespace UE::DerivedData { class ICache; }
 
 namespace UE::DerivedData::Private
 {
 
+// Implemented in DerivedDataBuild.cpp
+IBuild* CreateBuild(ICache& Cache);
+
 // Implemented in DerivedDataBuildDefinition.cpp
 FBuildDefinitionBuilder CreateBuildDefinition(FStringView Name, FStringView Function);
 FOptionalBuildDefinition LoadBuildDefinition(FStringView Name, FCbObject&& Definition);
diff --git a/Engine/Source/Developer/DerivedDataCache/Private/DerivedDataCache.cpp b/Engine/Source/Developer/DerivedDataCache/Private/DerivedDataCache.cpp
--- a/Engine/Source/Developer/DerivedDataCache/Private/DerivedDataCache.cpp
+++ b/Engine/Source/Developer/DerivedDataCache/Private/DerivedDataCache.cpp
@@ -5,40 +5,20 @@
 
 #include "CoreMinimal.h"
 #include "Algo/AllOf.h"
-#include "Misc/CommandLine.h"
-#include "HAL/ThreadSafeCounter.h"
-#include "Misc/ScopeLock.h"
-#include "Stats/StatsMisc.h"
-#include "Stats/Stats.h"
 #include "Async/AsyncWork.h"
 #include "Async/TaskGraphInterfaces.h"
-#include "Serialization/MemoryReader.h"
-#include "Serialization/MemoryWriter.h"
-#include "Modules/ModuleManager.h"
-
+#include "DDCCleanup.h"
 #include "DerivedDataBackendInterface.h"
-#include "DerivedDataBuild.h"
-#include "DerivedDataBuildAction.h"
-#include "DerivedDataBuildDefinition.h"
-#include "DerivedDataBuildInput.h"
-#include "DerivedDataBuildOutput.h"
-#include "DerivedDataBuildPrivate.h"
 #include "DerivedDataCachePrivate.h"
 #include "DerivedDataCacheRecord.h"
 #include "DerivedDataPluginInterface.h"
-#include "DDCCleanup.h"
-#include "ProfilingDebugging/CookStats.h"
-
-#include "Algo/AllOf.h"
-#include "Algo/Transform.h"
+#include "HAL/ThreadSafeCounter.h"
 #include "Misc/CoreMisc.h"
-#include "Misc/ScopeExit.h"
-#include "Misc/StringBuilder.h"
 #include "Misc/CommandLine.h"
-#include "Serialization/CompactBinary.h"
-#include "Serialization/CompactBinaryValidation.h"
-#include "Serialization/CompactBinaryWriter.h"
-
+#include "Misc/ScopeLock.h"
+#include "ProfilingDebugging/CookStats.h"
+#include "Stats/Stats.h"
+#include "Stats/StatsMisc.h"
 #include <atomic>
 
 DEFINE_STAT(STAT_DDC_NumGets);
@@ -54,7 +34,6 @@ DEFINE_STAT(STAT_DDC_ExistTime);
 //#define DDC_SCOPE_CYCLE_COUNTER(x) QUICK_SCOPE_CYCLE_COUNTER(STAT_ ## x)
 #define DDC_SCOPE_CYCLE_COUNTER(x) TRACE_CPUPROFILER_EVENT_SCOPE(x);
 
-
 #if ENABLE_COOK_STATS
 #include "DerivedDataCacheUsageStats.h"
 namespace DerivedDataCacheCookStats
@@ -169,14 +148,14 @@ namespace DerivedDataCacheCookStats
 /** Whether we want to verify the DDC (pass in -VerifyDDC on the command line)*/
 bool GVerifyDDC = false;
 
-namespace UE::DerivedData
+namespace UE::DerivedData::Private
 {
 
 /**
  * Implementation of the derived data cache
  * This API is fully threadsafe
 **/
-class FDerivedDataCache final : public FDerivedDataCacheInterface, public IBuild
+class FDerivedDataCache final : public FDerivedDataCacheInterface
 {
 
 	/** 
@@ -367,6 +346,8 @@ public:
 	/** Destructor, flushes all sync tasks **/
 	~FDerivedDataCache()
 	{
+		FDDCCleanup::Shutdown();
+
 		WaitForQuiescence(true);
 		FScopeLock ScopeLock(&SynchronizationObject);
 		for (TMap<uint32,FAsyncTask<FBuildAsyncWorker>*>::TIterator It(PendingTasks); It; ++It)
@@ -677,9 +658,9 @@ private:
 public:
 	// ICache Interface
 
-	FCacheBucket CreateBucket(FStringView Name) final { return Private::CreateCacheBucket(Name); }
+	FCacheBucket CreateBucket(FStringView Name) final { return CreateCacheBucket(Name); }
 
-	FCacheRecordBuilder CreateRecord(const FCacheKey& Key) final { return Private::CreateCacheRecordBuilder(Key); }
+	FCacheRecordBuilder CreateRecord(const FCacheKey& Key) final { return CreateCacheRecordBuilder(Key); }
 
 	FRequest Put(
 		TConstArrayView<FCacheRecord> Records,
@@ -715,115 +696,11 @@ public:
 	{
 		return FDerivedDataBackend::Get().GetRoot().CancelAll();
 	}
-
-public:
-	// IBuild Interface
-
-	FBuildDefinitionBuilder CreateDefinition(FStringView Name, FStringView Function) final
-	{
-		return Private::CreateBuildDefinition(Name, Function);
-	}
-
-	FOptionalBuildDefinition LoadDefinition(FStringView Name, FCbObject&& Definition) final
-	{
-		return Private::LoadBuildDefinition(Name, MoveTemp(Definition));
-	}
-
-	FBuildActionBuilder CreateAction(FStringView Name, FStringView Function) final
-	{
-		// DDC-TODO: Find the function version from the function registry on the scheduler.
-		return Private::CreateBuildAction(Name, Function, FGuid::NewGuid(), BuildSystemVersion);
-	}
-
-	FOptionalBuildAction LoadAction(FStringView Name, FCbObject&& Action) final
-	{
-		return Private::LoadBuildAction(Name, MoveTemp(Action));
-	}
-
-	FBuildInputBuilder CreateInput(FStringView Name) final
-	{
-		return Private::CreateBuildInput(Name);
-	}
-
-	FBuildOutputBuilder CreateOutput(FStringView Name, FStringView Function) final
-	{
-		return Private::CreateBuildOutput(Name, Function);
-	}
-
-	FOptionalBuildOutput LoadOutput(FStringView Name, FStringView Function, const FCbObject& Output) final
-	{
-		return Private::LoadBuildOutput(Name, Function, Output);
-	}
-
-	FOptionalBuildOutput LoadOutput(FStringView Name, FStringView Function, const FCacheRecord& Output) final
-	{
-		return Private::LoadBuildOutput(Name, Function, Output);
-	}
-
-	const FGuid& GetVersion() const final
-	{
-		return BuildSystemVersion;
-	}
-
-	const FGuid BuildSystemVersion{TEXT("ac0574e5-62bd-4c2e-84ec-f2efe48c0fef")};
 };
 
-} // UE::DerivedData
-
-static FDerivedDataCacheInterface* GDerivedDataCacheInstance;
-static UE::DerivedData::IBuild* GDerivedDataBuildInstance;
-
-PRAGMA_DISABLE_DEPRECATION_WARNINGS
-
-/**
- * Module for the DDC
- */
-class FDerivedDataCacheModule final : public IDerivedDataCacheModule
+FDerivedDataCacheInterface* CreateCache()
 {
-public:
-	FDerivedDataCacheInterface& GetDDC() final
-	{
-		return **CreateOrGetCache();
-	}
-
-	FDerivedDataCacheInterface* const* CreateOrGetCache() final
-	{
-		FScopeLock Lock(&CreateLock);
-		if (!GDerivedDataCacheInstance)
-		{
-			UE::DerivedData::FDerivedDataCache* Instance = new UE::DerivedData::FDerivedDataCache();
-			GDerivedDataCacheInstance = Instance;
-			GDerivedDataBuildInstance = Instance;
-			check(Instance);
-		}
-		return &GDerivedDataCacheInstance;
-	}
-
-	UE::DerivedData::IBuild* const* CreateOrGetBuild() final
-	{
-		FScopeLock Lock(&CreateLock);
-		if (!GDerivedDataBuildInstance)
-		{
-			UE::DerivedData::FDerivedDataCache* Instance = new UE::DerivedData::FDerivedDataCache();
-			GDerivedDataCacheInstance = Instance;
-			GDerivedDataBuildInstance = Instance;
-			check(Instance);
-		}
-		return &GDerivedDataBuildInstance;
-	}
-
-	void ShutdownModule() final
-	{
-		FDDCCleanup::Shutdown();
-
-		delete GDerivedDataCacheInstance;
-		GDerivedDataCacheInstance = nullptr;
-	}
-
-private:
-	FCriticalSection CreateLock;
-};
-
-PRAGMA_ENABLE_DEPRECATION_WARNINGS
+	return new FDerivedDataCache();
+}
 
-IMPLEMENT_MODULE(FDerivedDataCacheModule, DerivedDataCache);
+} // UE::DerivedData::Private
diff --git a/Engine/Source/Developer/DerivedDataCache/Private/DerivedDataCacheModule.cpp b/Engine/Source/Developer/DerivedDataCache/Private/DerivedDataCacheModule.cpp
new file mode 100644
--- /dev/null
+++ b/Engine/Source/Developer/DerivedDataCache/Private/DerivedDataCacheModule.cpp
@@ -0,0 +1,67 @@
+// Copyright Epic Games, Inc. All Rights Reserved.
+
+#include "DerivedDataCacheModule.h"
+
+#include "DerivedDataBuild.h"
+#include "DerivedDataBuildPrivate.h"
+#include "DerivedDataCache.h"
+#include "DerivedDataCacheInterface.h"
+#include "DerivedDataCachePrivate.h"
+#include "Misc/ScopeLock.h"
+#include "Modules/ModuleManager.h"
+
+namespace UE::DerivedData::Private
+{
+
+static FDerivedDataCacheInterface* GDerivedDataCacheInstance;
+static IBuild* GDerivedDataBuildInstance;
+
+PRAGMA_DISABLE_DEPRECATION_WARNINGS
+
+class FDerivedDataCacheModule final : public IDerivedDataCacheModule
+{
+public:
+	FDerivedDataCacheInterface& GetDDC() final
+	{
+		return **CreateOrGetCache();
+	}
+
+	FDerivedDataCacheInterface* const* CreateOrGetCache() final
+	{
+		FScopeLock Lock(&CreateLock);
+		if (!GDerivedDataCacheInstance)
+		{
+			GDerivedDataCacheInstance = CreateCache();
+			check(GDerivedDataCacheInstance);
+		}
+		return &GDerivedDataCacheInstance;
+	}
+
+	IBuild* const* CreateOrGetBuild() final
+	{
+		FScopeLock Lock(&CreateLock);
+		if (!GDerivedDataBuildInstance)
+		{
+			GDerivedDataBuildInstance = CreateBuild(**CreateOrGetCache());
+			check(GDerivedDataBuildInstance);
+		}
+		return &GDerivedDataBuildInstance;
+	}
+
+	void ShutdownModule() final
+	{
+		delete GDerivedDataBuildInstance;
+		GDerivedDataBuildInstance = nullptr;
+		delete GDerivedDataCacheInstance;
+		GDerivedDataCacheInstance = nullptr;
+	}
+
+private:
+	FCriticalSection CreateLock;
+};
+
+PRAGMA_ENABLE_DEPRECATION_WARNINGS
+
+} // UE::DerivedData::Private
+
+IMPLEMENT_MODULE(UE::DerivedData::Private::FDerivedDataCacheModule, DerivedDataCache);
diff --git a/Engine/Source/Developer/DerivedDataCache/Private/DerivedDataCachePrivate.h b/Engine/Source/Developer/DerivedDataCache/Private/DerivedDataCachePrivate.h
--- a/Engine/Source/Developer/DerivedDataCache/Private/DerivedDataCachePrivate.h
+++ b/Engine/Source/Developer/DerivedDataCache/Private/DerivedDataCachePrivate.h
@@ -4,6 +4,8 @@
 
 #include "Containers/StringFwd.h"
 
+class FDerivedDataCacheInterface;
+
 namespace UE::DerivedData { class FCacheBucket; }
 namespace UE::DerivedData { class FCacheRecordBuilder; }
 namespace UE::DerivedData { struct FCacheKey; }
@@ -11,6 +13,9 @@ namespace UE::DerivedData { struct FCacheKey; }
 namespace UE::DerivedData::Private
 {
 
+// Implemented in DerivedDataCache.cpp
+FDerivedDataCacheInterface* CreateCache();
+
 // Implemented in DerivedDataCacheKey.cpp
 FCacheBucket CreateCacheBucket(FStringView Name);
 
