DDC: Added FBuildActionKey::Empty

#rnx

diff --git a/Engine/Source/Developer/DerivedDataCache/Public/DerivedDataBuildAction.h b/Engine/Source/Developer/DerivedDataCache/Public/DerivedDataBuildAction.h
--- a/Engine/Source/Developer/DerivedDataCache/Public/DerivedDataBuildAction.h
+++ b/Engine/Source/Developer/DerivedDataCache/Public/DerivedDataBuildAction.h
@@ -24,8 +24,13 @@ namespace UE::DerivedData
 struct FBuildActionKey
 {
 	FIoHash Hash;
+
+	/** A key with a zero hash. */
+	static const FBuildActionKey Empty;
 };
 
+inline const FBuildActionKey FBuildActionKey::Empty;
+
 } // UE::DerivedData
 
 namespace UE::DerivedData::Private
