DDC: Renamed FBuildInput to FBuildInputs

#rnx

diff --git a/Engine/Source/Developer/DerivedDataCache/Private/DerivedDataBuild.cpp b/Engine/Source/Developer/DerivedDataCache/Private/DerivedDataBuild.cpp
--- a/Engine/Source/Developer/DerivedDataCache/Private/DerivedDataBuild.cpp
+++ b/Engine/Source/Developer/DerivedDataCache/Private/DerivedDataBuild.cpp
@@ -4,7 +4,7 @@
 
 #include "DerivedDataBuildAction.h"
 #include "DerivedDataBuildDefinition.h"
-#include "DerivedDataBuildInput.h"
+#include "DerivedDataBuildInputs.h"
 #include "DerivedDataBuildOutput.h"
 #include "DerivedDataBuildPrivate.h"
 #include "DerivedDataCache.h"
@@ -42,9 +42,9 @@ public:
 		return LoadBuildAction(Name, MoveTemp(Action));
 	}
 
-	FBuildInputBuilder CreateInput(FStringView Name) final
+	FBuildInputsBuilder CreateInputs(FStringView Name) final
 	{
-		return CreateBuildInput(Name);
+		return CreateBuildInputs(Name);
 	}
 
 	FBuildOutputBuilder CreateOutput(FStringView Name, FStringView Function) final
diff --git a/Engine/Source/Developer/DerivedDataCache/Private/DerivedDataBuildInput.cpp b/Engine/Source/Developer/DerivedDataCache/Private/DerivedDataBuildInput.cpp
deleted file mode 100644
diff --git a/Engine/Source/Developer/DerivedDataCache/Private/DerivedDataBuildInputs.cpp b/Engine/Source/Developer/DerivedDataCache/Private/DerivedDataBuildInputs.cpp
new file mode 100644
--- /dev/null
+++ b/Engine/Source/Developer/DerivedDataCache/Private/DerivedDataBuildInputs.cpp
@@ -0,0 +1,115 @@
+// Copyright Epic Games, Inc. All Rights Reserved.
+
+#include "DerivedDataBuildInputs.h"
+
+#include "Compression/CompressedBuffer.h"
+#include "Containers/Map.h"
+#include "Containers/StringView.h"
+#include "Containers/UnrealString.h"
+#include "DerivedDataBuildPrivate.h"
+#include <atomic>
+
+namespace UE::DerivedData::Private
+{
+
+class FBuildInputsBuilderInternal final : public IBuildInputsBuilderInternal
+{
+public:
+	inline explicit FBuildInputsBuilderInternal(FStringView InName)
+		: Name(InName)
+	{
+		checkf(!Name.IsEmpty(), TEXT("Build inputs require a non-empty name."));
+	}
+
+	~FBuildInputsBuilderInternal() final = default;
+
+	void AddInput(FStringView Key, const FCompressedBuffer& Buffer) final
+	{
+		const uint32 KeyHash = GetTypeHash(Key);
+		checkf(!Key.IsEmpty(), TEXT("Empty key used in input for build of '%s'."), *Name);
+		checkf(!Inputs.ContainsByHash(KeyHash, Key), TEXT("Duplicate key '%.*s' used in input for "),
+			TEXT("build of '%s'."), Key.Len(), Key.GetData(), *Name);
+		Inputs.EmplaceByHash(KeyHash, Key, Buffer);
+	}
+
+	FBuildInputs Build() final;
+
+	FString Name;
+	TMap<FString, FCompressedBuffer> Inputs;
+};
+
+///////////////////////////////////////////////////////////////////////////////////////////////////////////////////////
+
+class FBuildInputsInternal final : public IBuildInputsInternal
+{
+public:
+	explicit FBuildInputsInternal(FBuildInputsBuilderInternal&& InputsBuilder);
+
+	~FBuildInputsInternal() final = default;
+
+	FStringView GetName() const final { return Name; }
+
+	const FCompressedBuffer& GetInput(FStringView Key) const final;
+
+	inline void AddRef() const final
+	{
+		ReferenceCount.fetch_add(1, std::memory_order_relaxed);
+	}
+
+	inline void Release() const final
+	{
+		if (ReferenceCount.fetch_sub(1, std::memory_order_acq_rel) == 1)
+		{
+			delete this;
+		}
+	}
+
+private:
+	FString Name;
+	TMap<FString, FCompressedBuffer> Inputs;
+	mutable std::atomic<uint32> ReferenceCount{0};
+};
+
+///////////////////////////////////////////////////////////////////////////////////////////////////////////////////////
+
+FBuildInputsInternal::FBuildInputsInternal(FBuildInputsBuilderInternal&& InputsBuilder)
+	: Name(MoveTemp(InputsBuilder.Name))
+	, Inputs(MoveTemp(InputsBuilder.Inputs))
+{
+	Inputs.KeySort(TLess<>());
+}
+
+const FCompressedBuffer& FBuildInputsInternal::GetInput(FStringView Key) const
+{
+	if (const FCompressedBuffer* Buffer = Inputs.FindByHash(GetTypeHash(Key), Key))
+	{
+		return *Buffer;
+	}
+	return FCompressedBuffer::Null;
+}
+
+///////////////////////////////////////////////////////////////////////////////////////////////////////////////////////
+
+FBuildInputs FBuildInputsBuilderInternal::Build()
+{
+	return CreateBuildInputs(new FBuildInputsInternal(MoveTemp(*this)));
+}
+
+///////////////////////////////////////////////////////////////////////////////////////////////////////////////////////
+
+FBuildInputs CreateBuildInputs(IBuildInputsInternal* Inputs)
+{
+	return FBuildInputs(Inputs);
+}
+
+FBuildInputsBuilder CreateBuildInputsBuilder(IBuildInputsBuilderInternal* InputsBuilder)
+{
+	return FBuildInputsBuilder(InputsBuilder);
+}
+
+FBuildInputsBuilder CreateBuildInputs(FStringView Name)
+{
+	return CreateBuildInputsBuilder(new FBuildInputsBuilderInternal(Name));
+}
+
+} // UE::DerivedData::Private
diff --git a/Engine/Source/Developer/DerivedDataCache/Private/DerivedDataBuildPrivate.h b/Engine/Source/Developer/DerivedDataCache/Private/DerivedDataBuildPrivate.h
--- a/Engine/Source/Developer/DerivedDataCache/Private/DerivedDataBuildPrivate.h
+++ b/Engine/Source/Developer/DerivedDataCache/Private/DerivedDataBuildPrivate.h
@@ -10,7 +10,7 @@ struct FGuid;
 
 namespace UE::DerivedData { class FBuildActionBuilder; }
 namespace UE::DerivedData { class FBuildDefinitionBuilder; }
-namespace UE::DerivedData { class FBuildInputBuilder; }
+namespace UE::DerivedData { class FBuildInputsBuilder; }
 namespace UE::DerivedData { class FBuildOutputBuilder; }
 namespace UE::DerivedData { class FCacheRecord; }
 namespace UE::DerivedData { class FCacheRecordBuilder; }
@@ -35,7 +35,7 @@ FBuildActionBuilder CreateBuildAction(FStringView Name, FStringView Function, co
 FOptionalBuildAction LoadBuildAction(FStringView Name, FCbObject&& Action);
 
 // Implemented in DerivedDataBuildInput.cpp
-FBuildInputBuilder CreateBuildInput(FStringView Name);
+FBuildInputsBuilder CreateBuildInputs(FStringView Name);
 
 // Implemented in DerivedDataBuildOutput.cpp
 FBuildOutputBuilder CreateBuildOutput(FStringView Name, FStringView Function);
diff --git a/Engine/Source/Developer/DerivedDataCache/Public/DerivedDataBuild.h b/Engine/Source/Developer/DerivedDataCache/Public/DerivedDataBuild.h
--- a/Engine/Source/Developer/DerivedDataCache/Public/DerivedDataBuild.h
+++ b/Engine/Source/Developer/DerivedDataCache/Public/DerivedDataBuild.h
@@ -9,7 +9,7 @@ struct FGuid;
 
 namespace UE::DerivedData { class FBuildActionBuilder; }
 namespace UE::DerivedData { class FBuildDefinitionBuilder; }
-namespace UE::DerivedData { class FBuildInputBuilder; }
+namespace UE::DerivedData { class FBuildInputsBuilder; }
 namespace UE::DerivedData { class FBuildOutputBuilder; }
 namespace UE::DerivedData { class FCacheRecord; }
 namespace UE::DerivedData { class FOptionalBuildAction; }
@@ -64,11 +64,11 @@ public:
 	virtual FOptionalBuildAction LoadAction(FStringView Name, FCbObject&& Action) = 0;
 
 	/**
-	 * Create a build input builder.
+	 * Create a build inputs builder.
 	 *
-	 * @param Name   The name by which to identify this input for logging and profiling.
+	 * @param Name   The name by which to identify the inputs for logging and profiling.
 	 */
-	virtual FBuildInputBuilder CreateInput(FStringView Name) = 0;
+	virtual FBuildInputsBuilder CreateInputs(FStringView Name) = 0;
 
 	/**
 	 * Create a build output builder.
diff --git a/Engine/Source/Developer/DerivedDataCache/Public/DerivedDataBuildInput.h b/Engine/Source/Developer/DerivedDataCache/Public/DerivedDataBuildInput.h
deleted file mode 100644
diff --git a/Engine/Source/Developer/DerivedDataCache/Public/DerivedDataBuildInputs.h b/Engine/Source/Developer/DerivedDataCache/Public/DerivedDataBuildInputs.h
new file mode 100644
--- /dev/null
+++ b/Engine/Source/Developer/DerivedDataCache/Public/DerivedDataBuildInputs.h
@@ -0,0 +1,136 @@
+// Copyright Epic Games, Inc. All Rights Reserved.
+
+#pragma once
+
+#include "CoreTypes.h"
+#include "Containers/StringView.h"
+#include "Misc/ScopeExit.h"
+#include "Templates/RefCounting.h"
+#include "Templates/UniquePtr.h"
+
+class FCompressedBuffer;
+
+namespace UE::DerivedData { class FBuildInputs; }
+namespace UE::DerivedData { class FBuildInputsBuilder; }
+
+namespace UE::DerivedData::Private
+{
+
+class IBuildInputsInternal
+{
+public:
+	virtual ~IBuildInputsInternal() = default;
+	virtual FStringView GetName() const = 0;
+	virtual const FCompressedBuffer& GetInput(FStringView Key) const = 0;
+	virtual void AddRef() const = 0;
+	virtual void Release() const = 0;
+};
+
+FBuildInputs CreateBuildInputs(IBuildInputsInternal* Input);
+
+class IBuildInputsBuilderInternal
+{
+public:
+	virtual ~IBuildInputsBuilderInternal() = default;
+	virtual void AddInput(FStringView Key, const FCompressedBuffer& Buffer) = 0;
+	virtual FBuildInputs Build() = 0;
+};
+
+FBuildInputsBuilder CreateBuildInputsBuilder(IBuildInputsBuilderInternal* InputsBuilder);
+
+} // UE::DerivedData::Private
+
+namespace UE::DerivedData
+{
+
+/**
+ * Build inputs are an immutable container of key/value pairs for the inputs to a build function.
+ *
+ * The keys for inputs are names that are unique within the build inputs.
+ *
+ * @see FBuildAction
+ */
+class FBuildInputs
+{
+public:
+	/** Returns the name by which to identify the inputs for logging and profiling. */
+	inline FStringView GetName() const { return Inputs->GetName(); }
+
+	/** Visits every input in order by key. */
+	inline const FCompressedBuffer& GetInput(FStringView Key) const { return Inputs->GetInput(Key); }
+
+private:
+	friend class FOptionalBuildInputs;
+	friend FBuildInputs Private::CreateBuildInputs(Private::IBuildInputsInternal* Inputs);
+
+	/** Construct build inputs. Use Build() on a builder from IBuild::CreateInputs(). */
+	inline explicit FBuildInputs(Private::IBuildInputsInternal* InInputs)
+		: Inputs(InInputs)
+	{
+	}
+
+	TRefCountPtr<Private::IBuildInputsInternal> Inputs;
+};
+
+/**
+ * A build inputs builder is used to construct build inputs.
+ *
+ * Create using IBuild::CreateInputs().
+ *
+ * @see FBuildInputs
+ */
+class FBuildInputsBuilder
+{
+public:
+	/** Add an input with a key that is unique within this input. */
+	inline void AddInput(FStringView Key, const FCompressedBuffer& Buffer)
+	{
+		InputsBuilder->AddInput(Key, Buffer);
+	}
+
+	/** Build build inputs, which makes this builder subsequently unusable. */
+	inline FBuildInputs Build()
+	{
+		ON_SCOPE_EXIT { InputsBuilder = nullptr; };
+		return InputsBuilder->Build();
+	}
+
+private:
+	friend FBuildInputsBuilder Private::CreateBuildInputsBuilder(Private::IBuildInputsBuilderInternal* InputsBuilder);
+
+	/** Construct a build inputs builder. Use IBuild::CreateInputs(). */
+	inline explicit FBuildInputsBuilder(Private::IBuildInputsBuilderInternal* InInputsBuilder)
+		: InputsBuilder(InInputsBuilder)
+	{
+	}
+
+	TUniquePtr<Private::IBuildInputsBuilderInternal> InputsBuilder;
+};
+
+/**
+ * Build inputs that can be null.
+ *
+ * @see FBuildInputs
+ */
+class FOptionalBuildInputs : private FBuildInputs
+{
+public:
+	inline FOptionalBuildInputs() : FBuildInputs(nullptr) {}
+
+	inline FOptionalBuildInputs(FBuildInputs&& InInputs) : FBuildInputs(MoveTemp(InInputs)) {}
+	inline FOptionalBuildInputs(const FBuildInputs& InInputs) : FBuildInputs(InInputs) {}
+	inline FOptionalBuildInputs& operator=(FBuildInputs&& InInputs) { FBuildInputs::operator=(MoveTemp(InInputs)); return *this; }
+	inline FOptionalBuildInputs& operator=(const FBuildInputs& InInputs) { FBuildInputs::operator=(InInputs); return *this; }
+
+	/** Returns the build inputs. The caller must check for null before using this accessor. */
+	inline const FBuildInputs& Get() const & { return *this; }
+	inline FBuildInputs&& Get() && { return MoveTemp(*this); }
+
+	inline bool IsNull() const { return !IsValid(); }
+	inline bool IsValid() const { return Inputs.IsValid(); }
+	inline explicit operator bool() const { return IsValid(); }
+
+	inline void Reset() { *this = FOptionalBuildInputs(); }
+};
+
+} // UE::DerivedData
