DDC: Moved FBuildActionKey to DerivedDataBuildKey.h

#rnx

diff --git a/Engine/Source/Developer/DerivedDataCache/Private/DerivedDataBuildAction.cpp b/Engine/Source/Developer/DerivedDataCache/Private/DerivedDataBuildAction.cpp
--- a/Engine/Source/Developer/DerivedDataCache/Private/DerivedDataBuildAction.cpp
+++ b/Engine/Source/Developer/DerivedDataCache/Private/DerivedDataBuildAction.cpp
@@ -7,7 +7,9 @@
 #include "Containers/StringConv.h"
 #include "Containers/StringView.h"
 #include "Containers/UnrealString.h"
+#include "DerivedDataBuildKey.h"
 #include "DerivedDataBuildPrivate.h"
+#include "IO/IoHash.h"
 #include "Misc/Guid.h"
 #include "Misc/TVariant.h"
 #include "Serialization/CompactBinary.h"
diff --git a/Engine/Source/Developer/DerivedDataCache/Public/DerivedDataBuildAction.h b/Engine/Source/Developer/DerivedDataCache/Public/DerivedDataBuildAction.h
--- a/Engine/Source/Developer/DerivedDataCache/Public/DerivedDataBuildAction.h
+++ b/Engine/Source/Developer/DerivedDataCache/Public/DerivedDataBuildAction.h
@@ -4,7 +4,6 @@
 
 #include "CoreTypes.h"
 #include "Containers/StringView.h"
-#include "IO/IoHash.h"
 #include "Misc/ScopeExit.h"
 #include "Templates/Function.h"
 #include "Templates/RefCounting.h"
@@ -13,25 +12,11 @@
 class FCbObject;
 class FCbWriter;
 struct FGuid;
+struct FIoHash;
 
 namespace UE::DerivedData { class FBuildAction; }
 namespace UE::DerivedData { class FBuildActionBuilder; }
-
-namespace UE::DerivedData
-{
-
-/** A key that uniquely identifies a build action. */
-struct FBuildActionKey
-{
-	FIoHash Hash;
-
-	/** A key with a zero hash. */
-	static const FBuildActionKey Empty;
-};
-
-inline const FBuildActionKey FBuildActionKey::Empty;
-
-} // UE::DerivedData
+namespace UE::DerivedData { struct FBuildActionKey; }
 
 namespace UE::DerivedData::Private
 {
diff --git a/Engine/Source/Developer/DerivedDataCache/Public/DerivedDataBuildKey.h b/Engine/Source/Developer/DerivedDataCache/Public/DerivedDataBuildKey.h
--- a/Engine/Source/Developer/DerivedDataCache/Public/DerivedDataBuildKey.h
+++ b/Engine/Source/Developer/DerivedDataCache/Public/DerivedDataBuildKey.h
@@ -22,6 +22,17 @@ struct FBuildKey
 
 inline const FBuildKey FBuildKey::Empty;
 
+/** A key that uniquely identifies a build action. */
+struct FBuildActionKey
+{
+	FIoHash Hash;
+
+	/** A key with a zero hash. */
+	static const FBuildActionKey Empty;
+};
+
+inline const FBuildActionKey FBuildActionKey::Empty;
+
 /** A key that uniquely identifies a payload within a build output. */
 struct FBuildPayloadKey
 {
@@ -64,6 +75,34 @@ inline TStringBuilderBase<CharType>& operator<<(TStringBuilderBase<CharType>& Bu
 
 ///////////////////////////////////////////////////////////////////////////////////////////////////////////////////////
 
+inline bool operator==(const FBuildActionKey& A, const FBuildActionKey& B)
+{
+	return A.Hash == B.Hash;
+}
+
+inline bool operator!=(const FBuildActionKey& A, const FBuildActionKey& B)
+{
+	return A.Hash != B.Hash;
+}
+
+inline bool operator<(const FBuildActionKey& A, const FBuildActionKey& B)
+{
+	return A.Hash < B.Hash;
+}
+
+inline uint32 GetTypeHash(const FBuildActionKey& Key)
+{
+	return GetTypeHash(Key.Hash);
+}
+
+template <typename CharType>
+inline TStringBuilderBase<CharType>& operator<<(TStringBuilderBase<CharType>& Builder, const FBuildActionKey& Key)
+{
+	return Builder << "BuildAction/"_ASV << Key.Hash;
+}
+
+///////////////////////////////////////////////////////////////////////////////////////////////////////////////////////
+
 inline bool operator==(const FBuildPayloadKey& A, const FBuildPayloadKey& B)
 {
 	return A.BuildKey == B.BuildKey && A.Id == B.Id;
