DDC: Fixed checkSlow by adding otherwise-unused GetTypeHash for FCacheBucketOwner

#rnx

diff --git a/Engine/Source/Developer/DerivedDataCache/Private/DerivedDataCacheKey.cpp b/Engine/Source/Developer/DerivedDataCache/Private/DerivedDataCacheKey.cpp
--- a/Engine/Source/Developer/DerivedDataCache/Private/DerivedDataCacheKey.cpp
+++ b/Engine/Source/Developer/DerivedDataCache/Private/DerivedDataCacheKey.cpp
@@ -25,6 +25,11 @@ public:
 	using FCacheBucket::operator==;
 
 	inline bool operator==(FStringView Bucket) const { return ToString<TCHAR>().Equals(Bucket, ESearchCase::IgnoreCase); }
+
+	friend inline uint32 GetTypeHash(const FCacheBucketOwner& Bucket)
+	{
+		return ::GetTypeHash(Bucket.ToString<TCHAR>());
+	}
 };
 
 inline const TCHAR* FCacheBucketOwner::AllocName(FStringView Bucket)
