DDC: Added UE::DerivedData::IBuildFunctionRegistry

IBuildFunctionRegistry is an interface to look up build functions and their versions by name.

#rnx

diff --git a/Engine/Source/Developer/DerivedDataCache/Private/DerivedDataBuild.cpp b/Engine/Source/Developer/DerivedDataCache/Private/DerivedDataBuild.cpp
--- a/Engine/Source/Developer/DerivedDataCache/Private/DerivedDataBuild.cpp
+++ b/Engine/Source/Developer/DerivedDataCache/Private/DerivedDataBuild.cpp
@@ -4,6 +4,7 @@
 
 #include "DerivedDataBuildAction.h"
 #include "DerivedDataBuildDefinition.h"
+#include "DerivedDataBuildFunctionRegistry.h"
 #include "DerivedDataBuildInputs.h"
 #include "DerivedDataBuildOutput.h"
 #include "DerivedDataBuildPrivate.h"
@@ -13,6 +14,8 @@
 namespace UE::DerivedData::Private
 {
 
+DEFINE_LOG_CATEGORY(LogDerivedDataBuild);
+
 class FBuild final : public IBuild
 {
 public:
@@ -33,8 +36,7 @@ public:
 
 	FBuildActionBuilder CreateAction(FStringView Name, FStringView Function) final
 	{
-		// DDC-TODO: Find the function version from the function registry.
-		return CreateBuildAction(Name, Function, FGuid::NewGuid(), Version);
+		return CreateBuildAction(Name, Function, FunctionRegistry->FindFunctionVersion(Function), Version);
 	}
 
 	FOptionalBuildAction LoadAction(FStringView Name, FCbObject&& Action) final
@@ -67,8 +69,14 @@ public:
 		return Version;
 	}
 
+	IBuildFunctionRegistry& GetFunctionRegistry() const final
+	{
+		return *FunctionRegistry;
+	}
+
 private:
 	ICache& Cache;
+	TUniquePtr<IBuildFunctionRegistry> FunctionRegistry{CreateBuildFunctionRegistry()};
 	const FGuid Version{TEXT("ac0574e5-62bd-4c2e-84ec-f2efe48c0fef")};
 };
 
diff --git a/Engine/Source/Developer/DerivedDataCache/Private/DerivedDataBuildFunctionRegistry.cpp b/Engine/Source/Developer/DerivedDataCache/Private/DerivedDataBuildFunctionRegistry.cpp
new file mode 100644
--- /dev/null
+++ b/Engine/Source/Developer/DerivedDataCache/Private/DerivedDataBuildFunctionRegistry.cpp
@@ -0,0 +1,121 @@
+// Copyright Epic Games, Inc. All Rights Reserved.
+
+#include "DerivedDataBuildFunctionRegistry.h"
+
+#include "Containers/Map.h"
+#include "Containers/UnrealString.h"
+#include "DerivedDataBuildFunction.h"
+#include "DerivedDataBuildFunctionFactory.h"
+#include "DerivedDataBuildPrivate.h"
+#include "Features/IModularFeatures.h"
+#include "HAL/CriticalSection.h"
+#include "Misc/Guid.h"
+#include "Misc/ScopeRWLock.h"
+
+namespace UE::DerivedData::Private
+{
+
+class FBuildFunctionRegistry : public IBuildFunctionRegistry
+{
+public:
+	FBuildFunctionRegistry();
+	~FBuildFunctionRegistry();
+
+	const IBuildFunction* FindFunction(FStringView Function) const;
+	FGuid FindFunctionVersion(FStringView Function) const;
+
+private:
+	void OnModularFeatureRegistered(const FName& Type, IModularFeature* ModularFeature);
+	void OnModularFeatureUnregistered(const FName& Type, IModularFeature* ModularFeature);
+
+	void AddFunctionNoLock(IBuildFunctionFactory* Factory);
+	void RemoveFunction(IBuildFunctionFactory* Factory);
+
+private:
+	mutable FRWLock Lock;
+	TMap<FString, IBuildFunctionFactory*> Functions;
+};
+
+FBuildFunctionRegistry::FBuildFunctionRegistry()
+{
+	IModularFeatures& ModularFeatures = IModularFeatures::Get();
+	for (IBuildFunctionFactory* Factory : ModularFeatures.GetModularFeatureImplementations<IBuildFunctionFactory>(IBuildFunctionFactory::GetFeatureName()))
+	{
+		AddFunctionNoLock(Factory);
+	}
+	ModularFeatures.OnModularFeatureRegistered().AddRaw(this, &FBuildFunctionRegistry::OnModularFeatureRegistered);
+	ModularFeatures.OnModularFeatureUnregistered().AddRaw(this, &FBuildFunctionRegistry::OnModularFeatureUnregistered);
+}
+
+FBuildFunctionRegistry::~FBuildFunctionRegistry()
+{
+	IModularFeatures& ModularFeatures = IModularFeatures::Get();
+	ModularFeatures.OnModularFeatureUnregistered().RemoveAll(this);
+	ModularFeatures.OnModularFeatureRegistered().RemoveAll(this);
+}
+
+void FBuildFunctionRegistry::OnModularFeatureRegistered(const FName& Type, IModularFeature* ModularFeature)
+{
+	if (Type == IBuildFunctionFactory::GetFeatureName())
+	{
+		FWriteScopeLock WriteLock(Lock);
+		AddFunctionNoLock(static_cast<IBuildFunctionFactory*>(ModularFeature));
+	}
+}
+
+void FBuildFunctionRegistry::OnModularFeatureUnregistered(const FName& Type, IModularFeature* ModularFeature)
+{
+	if (Type == IBuildFunctionFactory::GetFeatureName())
+	{
+		RemoveFunction(static_cast<IBuildFunctionFactory*>(ModularFeature));
+	}
+}
+
+void FBuildFunctionRegistry::AddFunctionNoLock(IBuildFunctionFactory* Factory)
+{
+	const IBuildFunction& Function = Factory->GetFunction();
+	const FStringView FunctionName = Function.GetName();
+	const uint32 FunctionHash = GetTypeHash(FunctionName);
+	UE_CLOG(!Function.GetVersion().IsValid(), LogDerivedDataBuild, Error,
+		TEXT("Version of zero is not allowed in build function with the name %.*s."),
+		FunctionName.Len(), FunctionName.GetData());
+	UE_CLOG(Functions.FindByHash(FunctionHash, FunctionName), LogDerivedDataBuild, Error,
+		TEXT("More than one build function has been registered with the name %.*s."),
+		FunctionName.Len(), FunctionName.GetData());
+	Functions.EmplaceByHash(FunctionHash, FunctionName, Factory);
+}
+
+void FBuildFunctionRegistry::RemoveFunction(IBuildFunctionFactory* Factory)
+{
+	const FStringView Function = Factory->GetFunction().GetName();
+	const uint32 FunctionHash = GetTypeHash(Function);
+	FWriteScopeLock WriteLock(Lock);
+	Functions.RemoveByHash(FunctionHash, Function);
+}
+
+const IBuildFunction* FBuildFunctionRegistry::FindFunction(FStringView Function) const
+{
+	FReadScopeLock ReadLock(Lock);
+	if (IBuildFunctionFactory* const* Factory = Functions.FindByHash(GetTypeHash(Function), Function))
+	{
+		return &(**Factory).GetFunction();
+	}
+	return nullptr;
+}
+
+FGuid FBuildFunctionRegistry::FindFunctionVersion(FStringView Function) const
+{
+	FReadScopeLock ReadLock(Lock);
+	if (IBuildFunctionFactory* const* Factory = Functions.FindByHash(GetTypeHash(Function), Function))
+	{
+		return (**Factory).GetFunction().GetVersion();
+	}
+	return FGuid();
+}
+
+IBuildFunctionRegistry* CreateBuildFunctionRegistry()
+{
+	return new FBuildFunctionRegistry();
+}
+
+} // UE::DerivedData::Private
diff --git a/Engine/Source/Developer/DerivedDataCache/Private/DerivedDataBuildPrivate.h b/Engine/Source/Developer/DerivedDataCache/Private/DerivedDataBuildPrivate.h
--- a/Engine/Source/Developer/DerivedDataCache/Private/DerivedDataBuildPrivate.h
+++ b/Engine/Source/Developer/DerivedDataCache/Private/DerivedDataBuildPrivate.h
@@ -3,6 +3,7 @@
 #pragma once
 
 #include "Containers/StringFwd.h"
+#include "Logging/LogMacros.h"
 
 class FCbObject;
 class FCbWriter;
@@ -18,14 +19,20 @@ namespace UE::DerivedData { class FOptionalBuildAction; }
 namespace UE::DerivedData { class FOptionalBuildDefinition; }
 namespace UE::DerivedData { class FOptionalBuildOutput; }
 namespace UE::DerivedData { class IBuild; }
+namespace UE::DerivedData { class IBuildFunctionRegistry; }
 namespace UE::DerivedData { class ICache; }
 
 namespace UE::DerivedData::Private
 {
 
+DECLARE_LOG_CATEGORY_EXTERN(LogDerivedDataBuild, Log, All);
+
 // Implemented in DerivedDataBuild.cpp
 IBuild* CreateBuild(ICache& Cache);
 
+// Implemented in DerivedDataBuildFunctionRegistry.cpp
+IBuildFunctionRegistry* CreateBuildFunctionRegistry();
+
 // Implemented in DerivedDataBuildDefinition.cpp
 FBuildDefinitionBuilder CreateBuildDefinition(FStringView Name, FStringView Function);
 FOptionalBuildDefinition LoadBuildDefinition(FStringView Name, FCbObject&& Definition);
diff --git a/Engine/Source/Developer/DerivedDataCache/Public/DerivedDataBuild.h b/Engine/Source/Developer/DerivedDataCache/Public/DerivedDataBuild.h
--- a/Engine/Source/Developer/DerivedDataCache/Public/DerivedDataBuild.h
+++ b/Engine/Source/Developer/DerivedDataCache/Public/DerivedDataBuild.h
@@ -15,6 +15,7 @@ namespace UE::DerivedData { class FCacheRecord; }
 namespace UE::DerivedData { class FOptionalBuildAction; }
 namespace UE::DerivedData { class FOptionalBuildDefinition; }
 namespace UE::DerivedData { class FOptionalBuildOutput; }
+namespace UE::DerivedData { class IBuildFunctionRegistry; }
 
 namespace UE::DerivedData
 {
@@ -97,6 +98,11 @@ public:
 	 * into build actions to keep the build output separate for different build versions.
 	 */
 	virtual const FGuid& GetVersion() const = 0;
+
+	/**
+	 * Returns the build function registry used by the build system.
+	 */
+	virtual IBuildFunctionRegistry& GetFunctionRegistry() const = 0;
 };
 
 } // UE::DerivedData
diff --git a/Engine/Source/Developer/DerivedDataCache/Public/DerivedDataBuildFunctionRegistry.h b/Engine/Source/Developer/DerivedDataCache/Public/DerivedDataBuildFunctionRegistry.h
new file mode 100644
--- /dev/null
+++ b/Engine/Source/Developer/DerivedDataCache/Public/DerivedDataBuildFunctionRegistry.h
@@ -0,0 +1,28 @@
+// Copyright Epic Games, Inc. All Rights Reserved.
+
+#pragma once
+
+#include "Containers/StringFwd.h"
+
+struct FGuid;
+
+namespace UE::DerivedData { class IBuildFunction; }
+
+namespace UE::DerivedData
+{
+
+/**
+ * A build function registry maintains a collection of build functions.
+ */
+class IBuildFunctionRegistry
+{
+public:
+	virtual ~IBuildFunctionRegistry() = default;
+
+	/** Find a function by name. Returns null if not found. Safe to call from a scheduled job or the main thread. */
+	virtual const IBuildFunction* FindFunction(FStringView Function) const = 0;
+	/** Find a function version by name. Returns zero if not found. Safe to call from any thread. */
+	virtual FGuid FindFunctionVersion(FStringView Function) const = 0;
+};
+
+} // UE::DerivedData
