Allow for the use of -DDC=None commandline switch to permit DDC to be used without having a backend graph from the standalone build workers.

diff --git a/Engine/Source/Developer/DerivedDataCache/Private/DerivedDataBackends.cpp b/Engine/Source/Developer/DerivedDataCache/Private/DerivedDataBackends.cpp
--- a/Engine/Source/Developer/DerivedDataCache/Private/DerivedDataBackends.cpp
+++ b/Engine/Source/Developer/DerivedDataCache/Private/DerivedDataBackends.cpp
@@ -81,6 +81,12 @@ public:
 		{
 			if( GraphName.Len() > 0 )
 			{
+				if ( GraphName == TEXT("None"))
+				{
+					UE_LOG(LogDerivedDataCache, Display, TEXT("DDC backend graph of 'None' specified, graph will be invalid for use."));
+					return;
+				}
+
 				RootCache = ParseNode( TEXT("Root"), GEngineIni, *GraphName, ParsedNodes );				
 			}
 
