DDC: Added UE::DerivedData::IBuildScheduler and UE:::DerivedData::IBuildJob

IBuildScheduler schedules IBuildJob, which are created by the upcoming FBuildSession.

#rnx

diff --git a/Engine/Source/Developer/DerivedDataCache/Private/DerivedDataBuildPrivate.h b/Engine/Source/Developer/DerivedDataCache/Private/DerivedDataBuildPrivate.h
--- a/Engine/Source/Developer/DerivedDataCache/Private/DerivedDataBuildPrivate.h
+++ b/Engine/Source/Developer/DerivedDataCache/Private/DerivedDataBuildPrivate.h
@@ -6,7 +6,6 @@
 #include "Logging/LogMacros.h"
 
 class FCbObject;
-class FCbWriter;
 struct FGuid;
 
 namespace UE::DerivedData { class FBuildActionBuilder; }
@@ -14,12 +13,12 @@ namespace UE::DerivedData { class FBuildDefinitionBuilder; }
 namespace UE::DerivedData { class FBuildInputsBuilder; }
 namespace UE::DerivedData { class FBuildOutputBuilder; }
 namespace UE::DerivedData { class FCacheRecord; }
-namespace UE::DerivedData { class FCacheRecordBuilder; }
 namespace UE::DerivedData { class FOptionalBuildAction; }
 namespace UE::DerivedData { class FOptionalBuildDefinition; }
 namespace UE::DerivedData { class FOptionalBuildOutput; }
 namespace UE::DerivedData { class IBuild; }
 namespace UE::DerivedData { class IBuildFunctionRegistry; }
+namespace UE::DerivedData { class IBuildScheduler; }
 namespace UE::DerivedData { class ICache; }
 
 namespace UE::DerivedData::Private
@@ -33,6 +32,9 @@ IBuild* CreateBuild(ICache& Cache);
 // Implemented in DerivedDataBuildFunctionRegistry.cpp
 IBuildFunctionRegistry* CreateBuildFunctionRegistry();
 
+// Implemented in DerivedDataBuildScheduler.cpp
+IBuildScheduler* CreateBuildScheduler();
+
 // Implemented in DerivedDataBuildDefinition.cpp
 FBuildDefinitionBuilder CreateBuildDefinition(FStringView Name, FStringView Function);
 FOptionalBuildDefinition LoadBuildDefinition(FStringView Name, FCbObject&& Definition);
diff --git a/Engine/Source/Developer/DerivedDataCache/Private/DerivedDataBuildScheduler.cpp b/Engine/Source/Developer/DerivedDataCache/Private/DerivedDataBuildScheduler.cpp
new file mode 100644
--- /dev/null
+++ b/Engine/Source/Developer/DerivedDataCache/Private/DerivedDataBuildScheduler.cpp
@@ -0,0 +1,22 @@
+// Copyright Epic Games, Inc. All Rights Reserved.
+
+#include "DerivedDataBuildScheduler.h"
+
+#include "DerivedDataBuildFunction.h"
+#include "DerivedDataBuildFunctionFactory.h"
+#include "DerivedDataBuildFunctionRegistry.h"
+#include "Misc/Guid.h"
+
+namespace UE::DerivedData::Private
+{
+
+class FBuildScheduler final : public IBuildScheduler
+{
+};
+
+IBuildScheduler* CreateBuildScheduler()
+{
+	return new FBuildScheduler();
+}
+
+} // UE::DerivedData::Private
diff --git a/Engine/Source/Developer/DerivedDataCache/Public/DerivedDataBuildJob.h b/Engine/Source/Developer/DerivedDataCache/Public/DerivedDataBuildJob.h
new file mode 100644
--- /dev/null
+++ b/Engine/Source/Developer/DerivedDataCache/Public/DerivedDataBuildJob.h
@@ -0,0 +1,78 @@
+// Copyright Epic Games, Inc. All Rights Reserved.
+
+#pragma once
+
+#include "CoreTypes.h"
+#include "Containers/StringFwd.h"
+#include "DerivedDataRequest.h"
+
+class FEvent;
+
+template <typename FuncType> class TUniqueFunction;
+
+namespace UE::DerivedData { class FBuildOutput; }
+namespace UE::DerivedData { class IBuildScheduler; }
+namespace UE::DerivedData { struct FBuildJobCompleteParams; }
+namespace UE::DerivedData { struct FBuildActionKey; }
+namespace UE::DerivedData { struct FBuildKey; }
+namespace UE::DerivedData { enum class EBuildPolicy : uint8; }
+
+namespace UE::DerivedData
+{
+
+using FOnBuildJobComplete = TUniqueFunction<void (FBuildJobCompleteParams&& Params)>;
+
+/**
+ * A build job is responsible for the execution of one build.
+ *
+ * Jobs typically proceed through each one of a sequence of states, though a state may be skipped
+ * if the action was found in the cache or if the scheduler finds duplicate jobs for a definition
+ * or an action.
+ *
+ * The job depends on the build scheduler to move it through its states. That relationship allows
+ * the scheduler more control over resources such as: memory, compute, storage, network.
+ */
+class IBuildJob : public IRequest
+{
+public:
+	/** Returns the name by which to identify this job for logging and profiling. */
+	virtual FStringView GetName() const = 0;
+	/** Returns the name of the function to build with, or "Unknown" if not resolved yet. */
+	virtual FStringView GetFunction() const = 0;
+	/** Returns the build policy of this job. */
+	virtual EBuildPolicy GetPolicy() const = 0;
+	/** Returns the priority of this job. */
+	virtual EPriority GetPriority() const = 0;
+
+	/** Returns the key, or null if the job was created directly from an action. */
+	virtual const FBuildKey& GetKey() const = 0;
+	/** Returns the action key, or null if the action has not been resolved yet. */
+	virtual const FBuildActionKey& GetActionKey() const = 0;
+
+	/** Start the build by dispatching it to the scheduler. */
+	virtual void Schedule(IBuildScheduler& Scheduler, EBuildPolicy Policy, EPriority Priority, FOnBuildJobComplete&& OnComplete) = 0;
+
+	/** Called by the scheduler to begin execution for the corresponding state. */
+	virtual void BeginCacheQuery() = 0;
+	virtual void BeginCacheStore() = 0;
+	virtual void BeginResolveKey() = 0;
+	virtual void BeginResolveAction() = 0;
+	virtual void BeginResolveInputs() = 0;
+	virtual void BeginExecuteBuild() = 0;
+
+	/** Called by the scheduler if it has cached output compatible with the build policy. */
+	virtual void SetOutput(const FBuildOutput& Output) = 0;
+};
+
+/** Parameters for the completion callback for build jobs. */
+struct FBuildJobCompleteParams
+{
+	/** Job that is complete. */
+	const IBuildJob& Job;
+	/** Output for the job that completed or was canceled. */
+	FBuildOutput&& Output;
+	/** Status of the job. */
+	EStatus Status = EStatus::Error;
+};
+
+} // UE::DerivedData
diff --git a/Engine/Source/Developer/DerivedDataCache/Public/DerivedDataBuildScheduler.h b/Engine/Source/Developer/DerivedDataCache/Public/DerivedDataBuildScheduler.h
new file mode 100644
--- /dev/null
+++ b/Engine/Source/Developer/DerivedDataCache/Public/DerivedDataBuildScheduler.h
@@ -0,0 +1,47 @@
+// Copyright Epic Games, Inc. All Rights Reserved.
+
+#pragma once
+
+#include "DerivedDataBuildJob.h"
+
+namespace UE::DerivedData
+{
+
+/**
+ * A build scheduler is responsible for deciding when and where a job executes in certain states.
+ *
+ * Jobs dispatch themselves to their scheduler when they are prepared to access limited resources
+ * such as: memory, compute, storage, network. A scheduler may allow a job to execute immediately
+ * or may queue it to execute later. A scheduler that uses a job queue is expected to execute the
+ * jobs in priority order, respecting updates to priority.
+ */
+class IBuildScheduler
+{
+public:
+	virtual ~IBuildScheduler() = default;
+
+	/** Begin processing of the job by this scheduler. Always paired with EndJob. */
+	virtual void BeginJob(IBuildJob* Job) {}
+
+	/** End processing of the job by this scheduler. Always paired with BeginJob. */
+	virtual void EndJob(IBuildJob* Job) {}
+
+	/** Dispatch the job immediately if it is queued. May be called multiple times and/or concurrently. */
+	virtual void CancelJob(IBuildJob* Job) {}
+
+	/** Update the priority of the job if it is queued. May be called multiple times and/or concurrently. */
+	virtual void UpdateJobPriority(IBuildJob* Job) {}
+
+	/** Dispatch the job by executing the requested state now or queuing it for later. */
+	virtual void DispatchCacheQuery(IBuildJob* Job) { Job->BeginCacheQuery(); }
+	virtual void DispatchCacheStore(IBuildJob* Job) { Job->BeginCacheStore(); }
+	virtual void DispatchResolveKey(IBuildJob* Job) { Job->BeginResolveKey(); }
+	virtual void DispatchResolveAction(IBuildJob* Job) { Job->BeginResolveAction(); }
+	virtual void DispatchResolveInputs(IBuildJob* Job) { Job->BeginResolveInputs(); }
+	virtual void DispatchExecuteBuild(IBuildJob* Job) { Job->BeginExecuteBuild(); }
+
+	/** Called when the output of the job is available. Always called once, and always between BeginJob and EndJob. */
+	virtual void CompleteJob(const FBuildJobCompleteParams& Params) {}
+};
+
+} // UE::DerivedData
