Avoid crashing in DDC if the Jupiter server sends a malformed response with a payload size that sends us off the end of the response buffer.  Emit an error message instead.

diff --git a/Engine/Source/Developer/DerivedDataCache/Private/HttpDerivedDataBackend.cpp b/Engine/Source/Developer/DerivedDataCache/Private/HttpDerivedDataBackend.cpp
--- a/Engine/Source/Developer/DerivedDataCache/Private/HttpDerivedDataBackend.cpp
+++ b/Engine/Source/Developer/DerivedDataCache/Private/HttpDerivedDataBackend.cpp
@@ -1366,6 +1366,12 @@ private:
 
 						if (PayloadSize > 0)
 						{
+							if (Response + PayloadSize > ResponseEnd)
+							{
+								UE_LOG(LogDerivedDataCache, Display, ResponseErrorMessage);
+								return false;
+							}
+
 							if (bSuccess[KeyIdx])
 							{
 								Response += PayloadSize;
