DDC: Added UE::DerivedData::IBuildSession

#rnx

diff --git a/Engine/Source/Developer/DerivedDataCache/Public/DerivedDataBuildSession.h b/Engine/Source/Developer/DerivedDataCache/Public/DerivedDataBuildSession.h
new file mode 100644
--- /dev/null
+++ b/Engine/Source/Developer/DerivedDataCache/Public/DerivedDataBuildSession.h
@@ -0,0 +1,247 @@
+// Copyright Epic Games, Inc. All Rights Reserved.
+
+#pragma once
+
+#include "CoreTypes.h"
+#include "Containers/StringView.h"
+#include "DerivedDataBuildKey.h"
+#include "DerivedDataRequest.h"
+#include "Misc/EnumClassFlags.h"
+#include "Templates/Function.h"
+#include "Templates/UniquePtr.h"
+
+namespace UE::DerivedData { class FBuildAction; }
+namespace UE::DerivedData { class FBuildDefinition; }
+namespace UE::DerivedData { class FBuildOutput; }
+namespace UE::DerivedData { class FBuildSession; }
+namespace UE::DerivedData { class FOptionalBuildInputs; }
+namespace UE::DerivedData { class FPayload; }
+namespace UE::DerivedData { struct FBuildCompleteParams; }
+namespace UE::DerivedData { struct FBuildActionCompleteParams; }
+namespace UE::DerivedData { struct FBuildPayloadCompleteParams; }
+
+namespace UE::DerivedData
+{
+
+using FOnBuildComplete = TUniqueFunction<void (FBuildCompleteParams&& Params)>;
+using FOnBuildActionComplete = TUniqueFunction<void (FBuildActionCompleteParams&& Params)>;
+using FOnBuildPayloadComplete = TUniqueFunction<void (FBuildPayloadCompleteParams&& Params)>;
+
+/**
+ * Flags to control the behavior of build requests.
+ *
+ * The build policy flags can be combined to support a variety of usage patterns. Examples:
+ *
+ * Build(Default): allow cache get; build if missing; return every payload.
+ * Build(SkipCacheGet): never get from the cache; always build; return every payload.
+ * Build(SkipBuild | SkipData): allow cache get; never build; skip every payload.
+ *
+ * BuildPayload(Default): allow cache get; build if missing; return one payload.
+ * BuildPayload(SkipCacheGet): never get from the cache; always build; return one payload.
+ * BuildPayload(SkipBuild | SkipData): allow cache get; never build; skip every payload.
+ */
+enum class EBuildPolicy : uint8
+{
+	/** A value without any flags set. */
+	None            = 0,
+
+	/** Allow local execution. */
+	Local           = 1 << 0,
+	/** Allow remote execution if the function has a registered build worker. */
+	Remote          = 1 << 1,
+
+	/** Skip attempting a cache get before building. */
+	SkipCacheGet    = 1 << 2,
+	/** Skip attempting a cache get before building. */
+	SkipCachePut    = 1 << 3,
+	/** Skip attempting a build when the cache get is missed or skipped. */
+	SkipBuild       = 1 << 4,
+	/** Skip fetching the payload data from the cache. */
+	SkipData        = 1 << 5,
+
+	/** Allow cache get+put, local+remote build when missed or skipped, and fetch the payload(s). */
+	Default         = Local | Remote,
+};
+
+ENUM_CLASS_FLAGS(EBuildPolicy);
+
+} // UE::DerivedData
+
+namespace UE::DerivedData::Private
+{
+
+class IBuildSessionInternal
+{
+public:
+	virtual ~IBuildSessionInternal() = default;
+	virtual FStringView GetName() const = 0;
+	virtual FRequest Build(
+		const FBuildDefinition& Definition,
+		EBuildPolicy Policy,
+		EPriority Priority,
+		FOnBuildComplete&& OnComplete) = 0;
+	virtual FRequest BuildAction(
+		const FBuildAction& Action,
+		const FOptionalBuildInputs& Inputs,
+		EBuildPolicy Policy,
+		EPriority Priority,
+		FOnBuildActionComplete&& OnComplete) = 0;
+	virtual FRequest BuildPayload(
+		const FBuildPayloadKey& Payload,
+		EBuildPolicy Policy,
+		EPriority Priority,
+		FOnBuildPayloadComplete&& OnComplete) = 0;
+};
+
+FBuildSession CreateBuildSession(IBuildSessionInternal* Session);
+
+} // UE::DerivedData::Private
+
+namespace UE::DerivedData
+{
+
+/**
+ * A build session is the main point to the build scheduler.
+ *
+ * The purpose of a session is to group together related builds that use the same input resolver,
+ * such as grouping builds by target platform. A request to build one definition can lead to more
+ * builds being scheduled if the definition references payloads from other builds as inputs.
+ */
+class FBuildSession
+{
+public:
+	/** Returns the name by which to identify this session for logging and profiling. */
+	inline FStringView GetName() const
+	{
+		return Session->GetName();
+	}
+
+	/**
+	 * Asynchronous request to execute a build according to the policy.
+	 *
+	 * The callback will always be called, and may be called from an arbitrary thread.
+	 *
+	 * @param Definition   The build function to execute and references to its inputs.
+	 * @param Policy       Flags to control the behavior of the request. See EBuildPolicy.
+	 * @param Priority     A priority to consider when scheduling the request. See EPriority.
+	 * @param OnComplete   A callback invoked when the build completes or is canceled.
+	 */
+	inline FRequest Build(
+		const FBuildDefinition& Definition,
+		EBuildPolicy Policy,
+		EPriority Priority,
+		FOnBuildComplete&& OnComplete)
+	{
+		return Session->Build(Definition, Policy, Priority, MoveTemp(OnComplete));
+	}
+
+	/**
+	 * Asynchronous request to execute a build according to the policy.
+	 *
+	 * The callback will always be called, and may be called from an arbitrary thread.
+	 *
+	 * @param Action       The build function to execute and references to its inputs.
+	 * @param Inputs       The build inputs referenced by the action, if it has any.
+	 * @param Policy       Flags to control the behavior of the request. See EBuildPolicy.
+	 * @param Priority     A priority to consider when scheduling the request. See EPriority.
+	 * @param OnComplete   A callback invoked when the build completes or is canceled.
+	 */
+	inline FRequest BuildAction(
+		const FBuildAction& Action,
+		const FOptionalBuildInputs& Inputs,
+		EBuildPolicy Policy,
+		EPriority Priority,
+		FOnBuildActionComplete&& OnComplete)
+	{
+		return Session->BuildAction(Action, Inputs, Policy, Priority, MoveTemp(OnComplete));
+	}
+
+	/**
+	 * Asynchronous request to execute a build according to the policy and return one payload.
+	 *
+	 * The callback will always be called, and may be called from an arbitrary thread.
+	 *
+	 * @param Payload      The key identifying the build definition and the payload to return.
+	 * @param Policy       Flags to control the behavior of the request. See EBuildPolicy.
+	 * @param Priority     A priority to consider when scheduling the request. See EPriority.
+	 * @param OnComplete   A callback invoked when the build completes or is canceled.
+	 */
+	inline FRequest BuildPayload(
+		const FBuildPayloadKey& Payload,
+		EBuildPolicy Policy,
+		EPriority Priority,
+		FOnBuildPayloadComplete&& OnComplete)
+	{
+		return Session->BuildPayload(Payload, Policy, Priority, MoveTemp(OnComplete));
+	}
+
+private:
+	friend FBuildSession Private::CreateBuildSession(Private::IBuildSessionInternal* Session);
+
+	/** Construct a build session. Use IBuild::CreateSession(). */
+	inline explicit FBuildSession(Private::IBuildSessionInternal* InSession)
+		: Session(InSession)
+	{
+	}
+
+	TUniquePtr<Private::IBuildSessionInternal> Session;
+};
+
+/** Parameters for the completion callback for build requests. */
+struct FBuildCompleteParams
+{
+	/** Key for the build request that completed or was canceled. */
+	FBuildKey Key;
+
+	/**
+	 * Output for the build request that completed or was canceled.
+	 *
+	 * The name, function, and diagnostics are always populated.
+	 *
+	 * The payloads are populated when Status is Ok, but with null data if skipped by the policy.
+	 */
+	FBuildOutput&& Output;
+
+	/** Status of the build request. */
+	EStatus Status = EStatus::Error;
+};
+
+/** Parameters for the completion callback for build action requests. */
+struct FBuildActionCompleteParams
+{
+	/** Key for the build action request that completed or was canceled. */
+	FBuildActionKey Key;
+
+	/**
+	 * Output for the build action request that completed or was canceled.
+	 *
+	 * The name, function, and diagnostics are always populated.
+	 *
+	 * The payloads are populated when Status is Ok, but with null data if skipped by the policy.
+	 */
+	FBuildOutput&& Output;
+
+	/** Status of the build request. */
+	EStatus Status = EStatus::Error;
+};
+
+/** Parameters for the completion callback for build payload requests. */
+struct FBuildPayloadCompleteParams
+{
+	/** Key for the build request that completed or was canceled. See Payload for ID. */
+	FBuildKey Key;
+
+	/**
+	 * Payload from the build payload request that completed or was canceled.
+	 *
+	 * The ID is always populated.
+	 * The hash and size are populated when Status is Ok.
+	 * The data is populated when Status is Ok and the data was not skipped by the policy.
+	 */
+	FPayload&& Payload;
+
+	/** Status of the build request. */
+	EStatus Status = EStatus::Error;
+};
+
+} // UE::DerivedData
