DDC: Added UE::DerivedData::Private::FBuildJobContext

FBuildJobContext is the IBuildContext and IBuildConfigContext that is used when executing build jobs.

#rnx

diff --git a/Engine/Source/Developer/DerivedDataCache/Private/DerivedDataBuildJobContext.cpp b/Engine/Source/Developer/DerivedDataCache/Private/DerivedDataBuildJobContext.cpp
new file mode 100644
--- /dev/null
+++ b/Engine/Source/Developer/DerivedDataCache/Private/DerivedDataBuildJobContext.cpp
@@ -0,0 +1,141 @@
+// Copyright Epic Games, Inc. All Rights Reserved.
+
+#include "DerivedDataBuildJobContext.h"
+
+#include "DerivedDataBuildJob.h"
+#include "DerivedDataBuildOutput.h"
+#include "DerivedDataBuildSession.h"
+#include "DerivedDataCache.h"
+#include "Misc/StringBuilder.h"
+#include "UObject/NameTypes.h"
+
+namespace UE::DerivedData::Private
+{
+
+FBuildJobContext::FBuildJobContext(
+	IBuildJob& InJob,
+	const FCacheKey& InCacheKey,
+	const IBuildFunction& InFunction,
+	FBuildOutputBuilder& InOutputBuilder,
+	EBuildPolicy InBuildPolicy,
+	TUniqueFunction<void ()>&& InOnEndAsyncBuild)
+	: Job(InJob)
+	, CacheKey(InCacheKey)
+	, Function(InFunction)
+	, OutputBuilder(InOutputBuilder)
+	, OnEndAsyncBuild(MoveTemp(InOnEndAsyncBuild))
+	, CachePolicy(ECachePolicy::Default)
+	, BuildPolicy(InBuildPolicy)
+{
+}
+
+void FBuildJobContext::AddConstant(FStringView Key, FCbObject&& Value)
+{
+	Constants.EmplaceByHash(GetTypeHash(Key), Key, MoveTemp(Value));
+}
+
+void FBuildJobContext::AddInput(FStringView Key, const FCompressedBuffer& Value)
+{
+	Inputs.EmplaceByHash(GetTypeHash(Key), Key, Value);
+}
+
+void FBuildJobContext::ResetInputs()
+{
+	Constants.Empty();
+	Inputs.Empty();
+}
+
+FCbObject FBuildJobContext::GetConstant(FStringView Key) const
+{
+	if (const FCbObject* Object = Constants.FindByHash(GetTypeHash(Key), Key))
+	{
+		return *Object;
+	}
+	return FCbObject();
+}
+
+FSharedBuffer FBuildJobContext::GetInput(FStringView Key) const
+{
+	if (const FCompressedBuffer* Input = Inputs.FindByHash(GetTypeHash(Key), Key))
+	{
+		return Input->Decompress();
+	}
+	return FSharedBuffer();
+}
+
+void FBuildJobContext::AddPayload(const FPayload& Payload)
+{
+	OutputBuilder.AddPayload(Payload);
+}
+
+void FBuildJobContext::AddPayload(const FPayloadId& Id, const FCompressedBuffer& Buffer)
+{
+	AddPayload(FPayload(Id, Buffer));
+}
+
+void FBuildJobContext::AddPayload(const FPayloadId& Id, const FSharedBuffer& Buffer)
+{
+	AddPayload(FPayload(Id, FCompressedBuffer::Compress(NAME_Default, Buffer)));
+}
+
+void FBuildJobContext::AddPayload(const FPayloadId& Id, const FCbObject& Object)
+{
+	FMemoryView SerializedView;
+	if (Object.TryGetSerializedView(SerializedView))
+	{
+		AddPayload(FPayload(Id, FCompressedBuffer::Compress(NAME_Default, Object.GetBuffer())));
+	}
+	else
+	{
+		AddPayload(FPayload(Id, FCompressedBuffer::Compress(NAME_Default, FCbObject::Clone(Object).GetBuffer())));
+	}
+}
+
+void FBuildJobContext::BeginAsyncBuild()
+{
+	checkf(!bIsAsyncBuild, TEXT("BeginAsyncBuild may only be called once for build of '%s' by %s."),
+		*WriteToString<128>(Job.GetName()), *WriteToString<32>(Job.GetFunction()));
+	bIsAsyncBuild = true;
+}
+
+void FBuildJobContext::EndAsyncBuild()
+{
+	checkf(bIsAsyncBuild, TEXT("EndAsyncBuild may only be called after BeginAsyncBuild for build of '%s' by %s."),
+		*WriteToString<128>(Job.GetName()), *WriteToString<32>(Job.GetFunction()));
+	checkf(!bIsAsyncBuildComplete, TEXT("EndAsyncBuild may only be called once for build of '%s' by %s."),
+		*WriteToString<128>(Job.GetName()), *WriteToString<32>(Job.GetFunction()));
+	bIsAsyncBuildComplete = true;
+	OnEndAsyncBuild();
+}
+
+void FBuildJobContext::SetCacheBucket(FCacheBucket Bucket)
+{
+	checkf(!Bucket.IsNull(), TEXT("Null cache bucket not allowed for build of '%s' by %s. ")
+		TEXT("The cache can be disabled by calling SetCachePolicy(ECachePolicy::Disable)."),
+		*WriteToString<128>(Job.GetName()), *WriteToString<32>(Job.GetFunction()));
+	CacheKey.Bucket = Bucket;
+}
+
+void FBuildJobContext::SetCachePolicy(ECachePolicy Policy)
+{
+	checkf(!EnumHasAnyFlags(Policy, ECachePolicy::SkipData),
+		TEXT("SkipData flags not allowed on cache policy for build of '%s' by %s. ")
+		TEXT("Flags for skipping data may be set indirectly through EBuildPolicy."),
+		*WriteToString<128>(Job.GetName()), *WriteToString<32>(Job.GetFunction()));
+	CachePolicy = Policy;
+}
+
+void FBuildJobContext::SetBuildPolicy(EBuildPolicy Policy)
+{
+	checkf(!EnumHasAnyFlags(BuildPolicy ^ Policy, EBuildPolicy::SkipCacheGet | EBuildPolicy::SkipCachePut),
+		TEXT("SkipCache flags may not be modified on build policy for build of '%s' by %s. ")
+		TEXT("Flags for skipping cache operations may be set through ECachePolicy."),
+		*WriteToString<128>(Job.GetName()), *WriteToString<32>(Job.GetFunction()));
+	checkf(!EnumHasAnyFlags(BuildPolicy ^ Policy, EBuildPolicy::SkipBuild | EBuildPolicy::SkipData),
+		TEXT("Skip flags may not be modified on build policy for build of '%s' by %s. ")
+		TEXT("Flags for skipping the build or the data may only be set through the session."),
+		*WriteToString<128>(Job.GetName()), *WriteToString<32>(Job.GetFunction()));
+	BuildPolicy = Policy;
+}
+
+} // UE::DerivedData::Private
diff --git a/Engine/Source/Developer/DerivedDataCache/Private/DerivedDataBuildJobContext.h b/Engine/Source/Developer/DerivedDataCache/Private/DerivedDataBuildJobContext.h
new file mode 100644
--- /dev/null
+++ b/Engine/Source/Developer/DerivedDataCache/Private/DerivedDataBuildJobContext.h
@@ -0,0 +1,96 @@
+// Copyright Epic Games, Inc. All Rights Reserved.
+
+#pragma once
+
+#include "Compression/CompressedBuffer.h"
+#include "Containers/Map.h"
+#include "Containers/StringView.h"
+#include "Containers/UnrealString.h"
+#include "DerivedDataBuildFunction.h"
+#include "DerivedDataCacheKey.h"
+#include "Memory/MemoryFwd.h"
+#include "Serialization/CompactBinary.h"
+#include "Templates/Function.h"
+#include <atomic>
+
+namespace UE::DerivedData { class FBuildOutputBuilder; }
+namespace UE::DerivedData { class IBuildJob; }
+namespace UE::DerivedData { enum class EBuildPolicy : uint8; }
+namespace UE::DerivedData { enum class ECachePolicy : uint8; }
+
+namespace UE::DerivedData::Private
+{
+
+class FBuildJobContext final : public FBuildContext, public FBuildConfigContext
+{
+public:
+	FBuildJobContext(
+		IBuildJob& InJob,
+		const FCacheKey& InCacheKey,
+		const IBuildFunction& InFunction,
+		FBuildOutputBuilder& InOutputBuilder,
+		EBuildPolicy InBuildPolicy,
+		TUniqueFunction<void ()>&& InOnEndAsyncBuild);
+
+	inline const FCacheKey& GetCacheKey() const { return CacheKey; }
+	inline const IBuildFunction& GetFunction() const { return Function; }
+
+	inline ECachePolicy GetCachePolicy() const final { return CachePolicy; }
+	inline EBuildPolicy GetBuildPolicy() const final { return BuildPolicy; }
+
+	inline bool IsAsyncBuild() const { return bIsAsyncBuild; }
+	inline bool ShouldCheckDeterministicOutput() const { return bDeterministicOutputCheck; }
+
+	void AddConstant(FStringView Key, FCbObject&& Value);
+	void AddInput(FStringView Key, const FCompressedBuffer& Value);
+
+	void ResetInputs();
+
+private:
+	FCbObject GetConstant(FStringView Key) const final;
+	FSharedBuffer GetInput(FStringView Key) const final;
+
+	void AddPayload(const FPayload& Payload) final;
+	void AddPayload(const FPayloadId& Id, const FCompressedBuffer& Buffer) final;
+	void AddPayload(const FPayloadId& Id, const FSharedBuffer& Buffer) final;
+	void AddPayload(const FPayloadId& Id, const FCbObject& Object) final;
+
+	void BeginAsyncBuild() final;
+	void EndAsyncBuild() final;
+
+	void SetCacheBucket(FCacheBucket Bucket) final;
+	void SetCachePolicy(ECachePolicy Policy) final;
+	void SetBuildPolicy(EBuildPolicy Policy) final;
+	void SkipDeterministicOutputCheck() { bDeterministicOutputCheck = false; }
+
+public:
+	void AddRef() const
+	{
+		ReferenceCount.fetch_add(1, std::memory_order_relaxed);
+	}
+
+	void Release() const
+	{
+		if (ReferenceCount.fetch_sub(1, std::memory_order_acq_rel) == 1)
+		{
+			delete this;
+		}
+	}
+
+private:
+	IBuildJob& Job;
+	FCacheKey CacheKey;
+	const IBuildFunction& Function;
+	FBuildOutputBuilder& OutputBuilder;
+	TMap<FString, FCbObject> Constants;
+	TMap<FString, FCompressedBuffer> Inputs;
+	TUniqueFunction<void ()> OnEndAsyncBuild;
+	ECachePolicy CachePolicy;
+	EBuildPolicy BuildPolicy;
+	bool bIsAsyncBuild{false};
+	bool bIsAsyncBuildComplete{false};
+	bool bDeterministicOutputCheck{true};
+	mutable std::atomic<uint32> ReferenceCount{0};
+};
+
+} // UE::DerivedData::Private
