DDC: Added UE::DerivedData::IBuildInputResolver

Builds require IBuildInputResolver to resolve keys to definitions and definitions to inputs.

#rnx

diff --git a/Engine/Source/Developer/DerivedDataCache/Public/DerivedDataBuildInputResolver.h b/Engine/Source/Developer/DerivedDataCache/Public/DerivedDataBuildInputResolver.h
new file mode 100644
--- /dev/null
+++ b/Engine/Source/Developer/DerivedDataCache/Public/DerivedDataBuildInputResolver.h
@@ -0,0 +1,116 @@
+// Copyright Epic Games, Inc. All Rights Reserved.
+
+#pragma once
+
+#include "Compression/CompressedBuffer.h"
+#include "Containers/ArrayView.h"
+#include "Containers/StringView.h"
+#include "DerivedDataBuildKey.h"
+#include "DerivedDataRequest.h"
+#include "IO/IoHash.h"
+#include "Templates/Function.h"
+
+struct FGuid;
+
+namespace UE::DerivedData { class FBuildDefinition; }
+namespace UE::DerivedData { class FOptionalBuildDefinition; }
+namespace UE::DerivedData { struct FBuildDefinitionResolvedParams; }
+namespace UE::DerivedData { struct FBuildInputResolvedParams; }
+
+namespace UE::DerivedData
+{
+
+using FBuildInputFilter = TUniqueFunction<bool (FStringView Key)>;
+using FOnBuildDefinitionResolved = TUniqueFunction<void (FBuildDefinitionResolvedParams&& Params)>;
+using FOnBuildInputResolved = TUniqueFunction<void (FBuildInputResolvedParams&& Params)>;
+
+/** Interface to resolve input references from a build definition into the referenced data. */
+class IBuildInputResolver
+{
+public:
+	/**
+	 * Asynchronous request to resolve a definition from a key.
+	 *
+	 * @param Key          The key of the definition to resolve.
+	 * @param OnResolved   A callback invoked when resolving completes or is canceled.
+	 */
+	virtual FRequest GetDefinition(
+		const FBuildKey& Key,
+		FOnBuildDefinitionResolved&& OnResolved) = 0;
+
+	/**
+	 * Asynchronous request to resolve metadata for the inputs from the definition.
+	 *
+	 * @param Definition   The definition to resolve input metadata for.
+	 * @param Priority     A priority to consider when scheduling the request. See EPriority.
+	 * @param OnResolver   A callback invoked when resolving completes or is canceled.
+	 */
+	virtual FRequest GetMeta(
+		const FBuildDefinition& Definition,
+		EPriority Priority,
+		FOnBuildInputResolved&& OnResolved) = 0;
+
+	/**
+	 * Asynchronous request to resolve data for the inputs from the definition.
+	 *
+	 * @param Definition   The definition to resolve input data for.
+	 * @param Priority     A priority to consider when scheduling the request. See EPriority.
+	 * @param OnResolver   A callback invoked when resolving completes or is canceled.
+	 * @param Filter       An optional predicate to filter which input keys have data resolved.
+	 */
+	virtual FRequest GetData(
+		const FBuildDefinition& Definition,
+		EPriority Priority,
+		FOnBuildInputResolved&& OnResolved,
+		FBuildInputFilter&& Filter = FBuildInputFilter()) = 0;
+};
+
+/** Metadata for build inputs. */
+struct FBuildInputMeta
+{
+	/** Size of the raw (uncompressed) input in bytes. */
+	uint64 RawSize = 0;
+	/** Hash of the raw (uncompressed) input. */
+	FIoHash RawHash;
+};
+
+/** Parameters for the resolved callback for build definition requests. */
+struct FBuildDefinitionResolvedParams
+{
+	/** Key for the build definition request that resolved or was canceled. */
+	FBuildKey Key;
+
+	/** The resolved build definition. Only available when Status is Ok. */
+	FOptionalBuildDefinition&& Definition;
+
+	/** Status of the input request. */
+	EStatus Status = EStatus::Error;
+};
+
+/** Data and metadata for build inputs with the input key. */
+struct FBuildInputByKey
+{
+	/** Key used to identify this input in the definition. */
+	FStringView Key;
+
+	/** Metadata for the input. */
+	FBuildInputMeta Meta;
+
+	/** Data for the input. Null if only metadata was requested. */
+	FCompressedBuffer Data;
+};
+
+/** Parameters for the resolved callback for build input requests. */
+struct FBuildInputResolvedParams
+{
+	/** Key for the build input request that resolved or was canceled. */
+	FBuildKey Key;
+
+	/** All of the requested inputs sorted by key. Only available when Status is Ok. */
+	TConstArrayView<FBuildInputByKey> Inputs;
+
+	/** Status of the input request. */
+	EStatus Status = EStatus::Error;
+};
+
+} // UE::DerivedData
