DDC: Added UE::DerivedData::Private::FBuildJob

FBuildJob is the primary implementation of IBuildJob that is responsible for executing builds.

#rnx

diff --git a/Engine/Source/Developer/DerivedDataCache/Private/DerivedDataBuildJob.cpp b/Engine/Source/Developer/DerivedDataCache/Private/DerivedDataBuildJob.cpp
new file mode 100644
--- /dev/null
+++ b/Engine/Source/Developer/DerivedDataCache/Private/DerivedDataBuildJob.cpp
@@ -0,0 +1,996 @@
+// Copyright Epic Games, Inc. All Rights Reserved.
+
+#include "DerivedDataBuildJob.h"
+
+#include "Compression/CompressedBuffer.h"
+#include "DerivedDataBuild.h"
+#include "DerivedDataBuildAction.h"
+#include "DerivedDataBuildDefinition.h"
+#include "DerivedDataBuildFunction.h"
+#include "DerivedDataBuildFunctionRegistry.h"
+#include "DerivedDataBuildInputResolver.h"
+#include "DerivedDataBuildInputs.h"
+#include "DerivedDataBuildJobContext.h"
+#include "DerivedDataBuildKey.h"
+#include "DerivedDataBuildOutput.h"
+#include "DerivedDataBuildPrivate.h"
+#include "DerivedDataBuildScheduler.h"
+#include "DerivedDataBuildSession.h"
+#include "DerivedDataCache.h"
+#include "DerivedDataCacheRecord.h"
+#include "DerivedDataPayload.h"
+#include "HAL/CriticalSection.h"
+#include "HAL/Event.h"
+#include "HAL/PlatformProcess.h"
+#include "HAL/PlatformMath.h"
+#include "Misc/Guid.h"
+#include "Misc/ScopeRWLock.h"
+#include "Misc/StringBuilder.h"
+#include "Serialization/CompactBinary.h"
+#include "Templates/Function.h"
+#include "Templates/RefCounting.h"
+#include <atomic>
+
+namespace UE::DerivedData::Private
+{
+
+///////////////////////////////////////////////////////////////////////////////////////////////////////////////////////
+
+/** A build job proceeds through these states in order and may skip states. */
+enum class EBuildJobState : uint32
+{
+	/** A new job that has not been scheduled. */
+	NotStarted,
+	/** Resolve the key to a definition. */
+	ResolveKey,
+	/** Wait for the input resolver to resolve the key. */
+	ResolveKeyWait,
+	/** Resolve the definition to an action. */
+	ResolveAction,
+	/** Wait for the input resolver to resolve the metadata. */
+	ResolveActionWait,
+	/** Query the cache for the action. */
+	CacheQuery,
+	/** Wait for the cache query to finish. */
+	CacheQueryWait,
+	/** Resolve the definition input data. */
+	ResolveInputs,
+	/** Wait for the input resolver to resolve the data. */
+	ResolveInputsWait,
+	/** Execute the function to build the output. */
+	ExecuteBuild,
+	/** Wait for an async function to execute. */
+	ExecuteBuildWait,
+	/** Store the output in the cache. */
+	CacheStore,
+	/** Wait for the cache store to finish. */
+	CacheStoreWait,
+	/** Complete. */
+	Complete,
+};
+
+class FBuildJob final : public IBuildJob
+{
+public:
+	/** Resolve the key to a definition, then build like the definition constructor. */
+	FBuildJob(ICache& Cache, IBuild& BuildSystem, const FBuildKey& Key, IBuildInputResolver* InputResolver);
+	/** Resolve the definition to an action, query the cache, resolve inputs, then build like the action constructor. */
+	FBuildJob(ICache& Cache, IBuild& BuildSystem, const FBuildDefinition& Definition, IBuildInputResolver* InputResolver);
+	/** Execute the build function to build the output, store to the cache, then complete the build. */
+	FBuildJob(ICache& Cache, IBuild& BuildSystem, const FBuildAction& Action, const FOptionalBuildInputs& Inputs);
+
+	/** Destroy the job, which must be complete or not started. */
+	~FBuildJob();
+
+	// IBuildJob Interface
+
+	inline FStringView GetName() const final { return Name; }
+	inline FStringView GetFunction() const final { return FunctionName; }
+	inline EBuildPolicy GetPolicy() const final { return Context ? Context->GetBuildPolicy() : DefaultBuildPolicy; }
+	inline EPriority GetPriority() const final { return Priority; }
+
+	inline const FBuildKey& GetKey() const final { return DefinitionKey; }
+	inline const FBuildActionKey& GetActionKey() const final { return ActionKey; }
+
+	void Schedule(IBuildScheduler& Scheduler, EBuildPolicy Policy, EPriority Priority, FOnBuildJobComplete&& OnComplete) final;
+
+	// IRequest Interface
+
+	void SetPriority(EPriority Priority) final;
+	void Cancel() final;
+	void Wait() final;
+	bool Poll() final;
+	void AddRef() const final;
+	void Release() const final;
+
+private:
+	void BeginJob();
+	void Configure();
+	void EndJob();
+
+	void BeginResolveKey() final;
+	void BeginResolveAction() final;
+	void BeginCacheQuery() final;
+	void BeginResolveInputs() final;
+	void BeginExecuteBuild() final;
+	void BeginCacheStore() final;
+
+	void EndResolveKey(FBuildDefinitionResolvedParams&& Params);
+	void EndResolveAction(FBuildInputResolvedParams&& Params);
+	void EndCacheQuery(FCacheGetCompleteParams&& Params);
+	void EndResolveInputs(FBuildInputResolvedParams&& Params);
+	void EndExecuteBuild();
+	void EndCacheStore(FCachePutCompleteParams&& Params);
+
+	void EnterResolveKey();
+	void EnterResolveAction();
+	void EnterCacheQuery();
+	void EnterResolveInputs();
+	void EnterExecuteBuild();
+	void EnterCacheStore();
+
+	void SetDefinition(FBuildDefinition&& Definition);
+	void SetAction(FBuildAction&& Action);
+	void SetInputs(FBuildInputs&& Inputs);
+	void SetOutput(const FBuildOutput& Output) final;
+	void SetOutputNoCheck(FBuildOutput&& Output);
+
+	void CreateAction(TConstArrayView<FBuildInputByKey> InputMeta);
+
+	/** Terminate the job and send the error to the output complete callback. */
+	void CompleteWithError(FStringView Error);
+
+	/** Process cancellation requests and return whether the requested state is valid to execute. */
+	bool CanExecuteState(EBuildJobState RequestedState);
+
+	/** Advance to the new state, dispatching to the scheduler and invoking callbacks as appropriate. */
+	void AdvanceToState(EBuildJobState NewState, FRequest NewRequest = FRequest());
+
+	/** Execute a transition from the old state to the new state. */
+	void ExecuteTransition(EBuildJobState OldState, EBuildJobState NewState);
+
+private:
+	FString Name;
+	FString FunctionName{TEXT("Unknown")};
+
+	/** Active state for the job. Always moves through states in order. */
+	EBuildJobState State{EBuildJobState::NotStarted};
+	/** Next state for the job. Used to handle re-entrant calls to AdvanceToState. */
+	EBuildJobState NextState{EBuildJobState::NotStarted};
+
+	EBuildPolicy DefaultBuildPolicy{};
+	EPriority Priority{};
+
+	/** True if the build was canceled before it was complete. */
+	bool bIsCanceled{};
+	/** True if the build was retrieved from or stored to the cache. */
+	bool bIsCached{};
+	/** True if AdvanceToState is executing. */
+	bool bInAdvanceToState{};
+
+	mutable std::atomic<uint32> ReferenceCount{0};
+
+	/** Available in [ResolveKey, Complete] for jobs created from a key or definition. */
+	FBuildKey DefinitionKey;
+	/** Available in [CacheQuery, Complete]. */
+	FBuildActionKey ActionKey;
+	/** Available in [ResolveAction, ExecuteBuild) for jobs created from a key or definition. */
+	FOptionalBuildDefinition Definition;
+	/** Available in [CacheQuery, CacheStore). */
+	FOptionalBuildAction Action;
+	/** Available in [ExecuteBuild, CacheStore) for jobs with non-constant inputs. */
+	FOptionalBuildInputs Inputs;
+	/** Available in [CacheStore, Complete). */
+	FOptionalBuildOutput Output;
+
+	/** Available in [NotStarted, CacheStore). */
+	FBuildOutputBuilder OutputBuilder;
+
+	/** Available in [CacheQuery, CacheStoreWait). Context used by IBuildFunction. */
+	TRefCountPtr<FBuildJobContext> Context;
+
+	/** Lock to synchronize writes to State, NextState, bIsCanceled, Event, WaitRequest, Scheduler. */
+	FRWLock Lock;
+	/** Event used by and created by Wait(). Null until Wait() has been called on an incomplete job. */
+	FEvent* Event{};
+	/** Request that the job is waiting on in its wait states. Null outside of states that use it. */
+	FRequest WaitRequest;
+	/** Scheduler that the job is scheduled on. Null before Schedule() and after completion. */
+	IBuildScheduler* Scheduler{};
+	/** Resolver for definitions and inputs. May be null for a job with nothing to resolve. */
+	IBuildInputResolver* InputResolver{};
+	/** Invoked exactly once when the output is complete or when the job fails. */
+	FOnBuildJobComplete OnComplete;
+
+	ICache& Cache;
+	IBuild& BuildSystem;
+};
+
+///////////////////////////////////////////////////////////////////////////////////////////////////////////////////////
+
+const TCHAR* LexToString(EBuildJobState State)
+{
+	switch (State)
+	{
+	case EBuildJobState::ResolveKey:        return TEXT("ResolveKey");
+	case EBuildJobState::ResolveKeyWait:    return TEXT("ResolveKeyWait");
+	case EBuildJobState::ResolveAction:     return TEXT("ResolveAction");
+	case EBuildJobState::ResolveActionWait: return TEXT("ResolveActionWait");
+	case EBuildJobState::CacheQuery:        return TEXT("CacheQuery");
+	case EBuildJobState::CacheQueryWait:    return TEXT("CacheQueryWait");
+	case EBuildJobState::ResolveInputs:     return TEXT("ResolveInputs");
+	case EBuildJobState::ResolveInputsWait: return TEXT("ResolveInputsWait");
+	case EBuildJobState::ExecuteBuild:      return TEXT("ExecuteBuild");
+	case EBuildJobState::ExecuteBuildWait:  return TEXT("ExecuteBuildWait");
+	case EBuildJobState::CacheStore:        return TEXT("CacheStore");
+	case EBuildJobState::CacheStoreWait:    return TEXT("CacheStoreWait");
+	case EBuildJobState::Complete:          return TEXT("Complete");
+	default: checkNoEntry();                return TEXT("Unknown");
+	}
+}
+
+///////////////////////////////////////////////////////////////////////////////////////////////////////////////////////
+
+FBuildJob::FBuildJob(
+	ICache& InCache,
+	IBuild& InBuildSystem,
+	const FBuildKey& InKey,
+	IBuildInputResolver* InInputResolver)
+	: Name(WriteToString<64>(TEXT("Resolve: "_SV), InKey))
+	, DefinitionKey(InKey)
+	, OutputBuilder(InBuildSystem.CreateOutput(Name, FunctionName))
+	, InputResolver(InInputResolver)
+	, Cache(InCache)
+	, BuildSystem(InBuildSystem)
+{
+}
+
+FBuildJob::FBuildJob(
+	ICache& InCache,
+	IBuild& InBuildSystem,
+	const FBuildDefinition& InDefinition,
+	IBuildInputResolver* InInputResolver)
+	: Name(InDefinition.GetName())
+	, FunctionName(InDefinition.GetFunction())
+	, DefinitionKey(InDefinition.GetKey())
+	, Definition(InDefinition)
+	, OutputBuilder(InBuildSystem.CreateOutput(Name, FunctionName))
+	, InputResolver(InInputResolver)
+	, Cache(InCache)
+	, BuildSystem(InBuildSystem)
+{
+}
+
+FBuildJob::FBuildJob(
+	ICache& InCache,
+	IBuild& InBuildSystem,
+	const FBuildAction& InAction,
+	const FOptionalBuildInputs& InInputs)
+	: Name(InAction.GetName())
+	, FunctionName(InAction.GetFunction())
+	, ActionKey(InAction.GetKey())
+	, Action(InAction)
+	, Inputs(InInputs)
+	, OutputBuilder(InBuildSystem.CreateOutput(Name, FunctionName))
+	, Cache(InCache)
+	, BuildSystem(InBuildSystem)
+{
+}
+
+FBuildJob::~FBuildJob()
+{
+	checkf(State == EBuildJobState::NotStarted || State == EBuildJobState::Complete,
+		TEXT("Job in state %s must complete before being destroyed for build of '%s' by %s."),
+		LexToString(State), *Name, *FunctionName);
+	checkf(!OnComplete,
+		TEXT("Job in state %s must invoke its completion callback before being destroyed for build of '%s' by %s."),
+		LexToString(State), *Name, *FunctionName);
+	FPlatformProcess::ReturnSynchEventToPool(Event);
+}
+
+void FBuildJob::Schedule(
+	IBuildScheduler& InScheduler,
+	EBuildPolicy InPolicy,
+	EPriority InPriority,
+	FOnBuildJobComplete&& InOnComplete)
+{
+	checkf(State == EBuildJobState::NotStarted,
+		TEXT("Job in state %s was previously scheduled for build of '%s' by %s."),
+		LexToString(State), *Name, *FunctionName);
+	Scheduler = &InScheduler;
+	DefaultBuildPolicy = InPolicy;
+	Priority = InPriority;
+	OnComplete = MoveTemp(InOnComplete);
+	return AdvanceToState(EBuildJobState::ResolveKey);
+}
+
+///////////////////////////////////////////////////////////////////////////////////////////////////////////////////////
+
+void FBuildJob::SetPriority(EPriority InPriority)
+{
+	IBuildScheduler* LocalScheduler;
+	{
+		FWriteScopeLock WriteLock(Lock);
+		Priority = InPriority;
+		LocalScheduler = Scheduler;
+		WaitRequest.SetPriority(InPriority);
+	}
+	if (LocalScheduler)
+	{
+		LocalScheduler->UpdateJobPriority(this);
+	}
+}
+
+void FBuildJob::Cancel()
+{
+	FRequest LocalRequest;
+	IBuildScheduler* LocalScheduler;
+	TRefCountPtr<FBuildJobContext> LocalContext;
+	if (FReadScopeLock ReadLock(Lock); State == EBuildJobState::Complete)
+	{
+		return;
+	}
+	else
+	{
+		bIsCanceled = true;
+		LocalRequest = WaitRequest;
+		LocalScheduler = Scheduler;
+		if (State == EBuildJobState::ExecuteBuildWait)
+		{
+			LocalContext = Context;
+		}
+	}
+
+	// Cancel the job on the scheduler, which invokes Begin[State] if the job was queued.
+	if (LocalScheduler)
+	{
+		LocalScheduler->CancelJob(this);
+	}
+
+	// Cancel the request, which invokes End[State] if the request is not complete.
+	if (LocalRequest)
+	{
+		LocalRequest.Cancel();
+	}
+
+	// Cancel the async build, which invokes EndExecuteBuild if the build is not complete.
+	if (LocalContext)
+	{
+		LocalContext->GetFunction().CancelAsyncBuild(*LocalContext);
+	}
+
+	// Most jobs will be complete at this point, but it is possible for a job to reach this point
+	// because of a race with the scheduler. A job can be removed from every scheduler queue, and
+	// missed by CancelJob, but not yet have called into Begin[State] or Set[Type].
+	Wait();
+}
+
+void FBuildJob::Wait()
+{
+	FEvent* LocalEvent;
+
+	if (FReadScopeLock ReadLock(Lock); State == EBuildJobState::Complete)
+	{
+		return;
+	}
+	else
+	{
+		LocalEvent = Event;
+	}
+
+	SetPriority(EPriority::Blocking);
+
+	if (!LocalEvent)
+	{
+		if (FWriteScopeLock WriteLock(Lock); State == EBuildJobState::Complete)
+		{
+			return;
+		}
+		else if (Event)
+		{
+			LocalEvent = Event;
+		}
+		else
+		{
+			LocalEvent = Event = FPlatformProcess::GetSynchEventFromPool(/*bIsManualReset*/ true);
+		}
+	}
+
+	LocalEvent->Wait();
+}
+
+bool FBuildJob::Poll()
+{
+	FReadScopeLock ReadLock(Lock);
+	return State == EBuildJobState::Complete;
+}
+
+void FBuildJob::AddRef() const
+{
+	ReferenceCount.fetch_add(1, std::memory_order_relaxed);
+}
+
+void FBuildJob::Release() const
+{
+	if (ReferenceCount.fetch_sub(1, std::memory_order_acq_rel) == 1)
+	{
+		delete this;
+	}
+}
+
+///////////////////////////////////////////////////////////////////////////////////////////////////////////////////////
+
+void FBuildJob::BeginJob()
+{
+	AddRef();
+	Scheduler->BeginJob(this);
+}
+
+///////////////////////////////////////////////////////////////////////////////////////////////////////////////////////
+
+void FBuildJob::EnterResolveKey()
+{
+	if (Action)
+	{
+		return AdvanceToState(EBuildJobState::CacheQuery);
+	}
+	if (Definition)
+	{
+		return AdvanceToState(EBuildJobState::ResolveAction);
+	}
+	if (DefinitionKey != FBuildKey::Empty)
+	{
+		return CompleteWithError(TEXT("Failed to resolve null key."_SV));
+	}
+	if (!InputResolver)
+	{
+		return CompleteWithError(TEXT("Failed to resolve key due to null input resolver."_SV));
+	}
+	Scheduler->DispatchResolveKey(this);
+}
+
+void FBuildJob::BeginResolveKey()
+{
+	if (CanExecuteState(EBuildJobState::ResolveKey))
+	{
+		return AdvanceToState(EBuildJobState::ResolveKeyWait, InputResolver->GetDefinition(DefinitionKey,
+			[this](FBuildDefinitionResolvedParams&& Params) { EndResolveKey(MoveTemp(Params)); }));
+	}
+}
+
+void FBuildJob::EndResolveKey(FBuildDefinitionResolvedParams&& Params)
+{
+	if (CanExecuteState(EBuildJobState::ResolveKeyWait))
+	{
+		if (Params.Status == EStatus::Ok && Params.Definition)
+		{
+			SetDefinition(MoveTemp(Params.Definition).Get());
+		}
+		else
+		{
+			CompleteWithError(WriteToString<128>(TEXT("Failed to resolve key "_SV), Params.Key, TEXT("."_SV)));
+		}
+	}
+}
+
+void FBuildJob::SetDefinition(FBuildDefinition&& InDefinition)
+{
+	Name = InDefinition.GetName();
+	FunctionName = InDefinition.GetFunction();
+	checkf(Definition.IsNull(), TEXT("Job already has a definition for build of '%s' by %s."), *Name, *FunctionName);
+	checkf(State == EBuildJobState::ResolveKey || State == EBuildJobState::ResolveKeyWait,
+		TEXT("Job is not expecting a definition in state %s for build of '%s' by %s"),
+		LexToString(State), *Name, *FunctionName);
+	OutputBuilder = BuildSystem.CreateOutput(Name, FunctionName);
+	Definition = MoveTemp(InDefinition);
+	return AdvanceToState(EBuildJobState::ResolveAction);
+}
+
+///////////////////////////////////////////////////////////////////////////////////////////////////////////////////////
+
+void FBuildJob::EnterResolveAction()
+{
+	if (Definition.Get().HasInputs() && !InputResolver)
+	{
+		return CompleteWithError(TEXT("Failed to resolve input metadata due to null input resolver."_SV));
+	}
+	if (!Definition.Get().HasInputs())
+	{
+		return CreateAction({});
+	}
+	Scheduler->DispatchResolveAction(this);
+}
+
+void FBuildJob::BeginResolveAction()
+{
+	if (CanExecuteState(EBuildJobState::ResolveAction))
+	{
+		return AdvanceToState(EBuildJobState::ResolveActionWait, InputResolver->GetMeta(Definition.Get(), Priority,
+			[this](FBuildInputResolvedParams&& Params) { EndResolveAction(MoveTemp(Params)); }));
+	}
+}
+
+void FBuildJob::EndResolveAction(FBuildInputResolvedParams&& Params)
+{
+	if (CanExecuteState(EBuildJobState::ResolveActionWait))
+	{
+		if (Params.Status == EStatus::Ok)
+		{
+			CreateAction(Params.Inputs);
+		}
+		else
+		{
+			CompleteWithError(TEXT("Failed to resolve input metadata."_SV));
+		}
+	}
+}
+
+void FBuildJob::CreateAction(TConstArrayView<FBuildInputByKey> InputMeta)
+{
+	const FBuildDefinition& LocalDefinition = Definition.Get();
+	FBuildActionBuilder Builder = BuildSystem.CreateAction(Name, FunctionName);
+	LocalDefinition.IterateConstants([this, &Builder](FStringView Key, FCbObject&& Value)
+	{
+		Builder.AddConstant(Key, Value);
+	});
+	for (const FBuildInputByKey& InputByKey : InputMeta)
+	{
+		Builder.AddInput(InputByKey.Key, InputByKey.Meta.RawHash, InputByKey.Meta.RawSize);
+	}
+	SetAction(Builder.Build());
+}
+
+void FBuildJob::SetAction(FBuildAction&& InAction)
+{
+	checkf(Action.IsNull(), TEXT("Job already has an action for build of '%s' by %s."), *Name, *FunctionName);
+	checkf(State == EBuildJobState::ResolveKey ||
+		State == EBuildJobState::ResolveAction || State == EBuildJobState::ResolveActionWait,
+		TEXT("Job is not expecting an action in state %s for build of '%s' by %s"),
+		LexToString(State), *Name, *FunctionName);
+	ActionKey = InAction.GetKey();
+	Action = MoveTemp(InAction);
+	return AdvanceToState(EBuildJobState::CacheQuery);
+}
+
+///////////////////////////////////////////////////////////////////////////////////////////////////////////////////////
+
+void FBuildJob::Configure()
+{
+	const IBuildFunction* Function = BuildSystem.GetFunctionRegistry().FindFunction(FunctionName);
+	if (BuildSystem.GetVersion() != Action.Get().GetBuildSystemVersion())
+	{
+		CompleteWithError(WriteToString<192>(TEXT("Failed because the build system is version "_SV),
+			BuildSystem.GetVersion().ToString(), TEXT(" when version "_SV),
+			Action.Get().GetBuildSystemVersion().ToString(), TEXT(" is expected."_SV)));
+	}
+	else if (!Function)
+	{
+		CompleteWithError(WriteToString<128>(TEXT("Failed because the function "_SV), FunctionName,
+			TEXT(" was not found."_SV)));
+	}
+	else if (!Function->GetVersion().IsValid())
+	{
+		CompleteWithError(WriteToString<128>(TEXT("Failed because the function "_SV), FunctionName,
+			TEXT(" has a version of zero."_SV)));
+	}
+	else if (Function->GetVersion() != Action.Get().GetFunctionVersion())
+	{
+		CompleteWithError(WriteToString<192>(TEXT("Failed because the function "_SV), FunctionName,
+			TEXT(" is version "_SV), Function->GetVersion().ToString(), TEXT(" when version "_SV),
+			Action.Get().GetFunctionVersion().ToString(), TEXT(" is expected."_SV)));
+	}
+	else
+	{
+		const FCacheKey CacheKey{Cache.CreateBucket(FunctionName), Action.Get().GetKey().Hash};
+		Context = new FBuildJobContext(*this, CacheKey, *Function, OutputBuilder, DefaultBuildPolicy,
+			[this] { EndExecuteBuild(); });
+		Function->Configure(*Context);
+	}
+}
+
+///////////////////////////////////////////////////////////////////////////////////////////////////////////////////////
+
+void FBuildJob::EnterCacheQuery()
+{
+	const EBuildPolicy BuildPolicy = Context->GetBuildPolicy();
+	ECachePolicy CachePolicy = Context->GetCachePolicy();
+	if (!EnumHasAnyFlags(CachePolicy, ECachePolicy::Query) ||
+		EnumHasAnyFlags(BuildPolicy, EBuildPolicy::SkipCacheGet))
+	{
+		return AdvanceToState(EBuildJobState::ResolveInputs);
+	}
+	Scheduler->DispatchCacheQuery(this);
+}
+
+void FBuildJob::BeginCacheQuery()
+{
+	if (CanExecuteState(EBuildJobState::CacheQuery))
+	{
+		const EBuildPolicy BuildPolicy = Context->GetBuildPolicy();
+		ECachePolicy CachePolicy = Context->GetCachePolicy();
+		if (EnumHasAnyFlags(BuildPolicy, EBuildPolicy::SkipData))
+		{
+			CachePolicy |= ECachePolicy::SkipAttachments;
+		}
+		return AdvanceToState(EBuildJobState::CacheQueryWait,
+			Cache.Get({Context->GetCacheKey()}, Name, CachePolicy, Priority,
+				[this](FCacheGetCompleteParams&& Params) { EndCacheQuery(MoveTemp(Params)); }));
+	}
+}
+
+void FBuildJob::EndCacheQuery(FCacheGetCompleteParams&& Params)
+{
+	if (CanExecuteState(EBuildJobState::CacheQueryWait))
+	{
+		if (Params.Status == EStatus::Ok)
+		{
+			if (FOptionalBuildOutput CacheOutput = BuildSystem.LoadOutput(Name, FunctionName, Params.Record))
+			{
+				bIsCached = true;
+				SetOutputNoCheck(MoveTemp(CacheOutput).Get());
+			}
+		}
+		return AdvanceToState(EBuildJobState::ResolveInputs);
+	}
+}
+
+///////////////////////////////////////////////////////////////////////////////////////////////////////////////////////
+
+void FBuildJob::EnterResolveInputs()
+{
+	if (Inputs || !Action.Get().HasInputs())
+	{
+		return AdvanceToState(EBuildJobState::ExecuteBuild);
+	}
+	if (!InputResolver)
+	{
+		return CompleteWithError(TEXT("Failed to resolve input data due to null input resolver."_SV));
+	}
+	if (!Definition)
+	{
+		return CompleteWithError(TEXT("Failed to resolve input data due to null definition."_SV));
+	}
+	Scheduler->DispatchResolveInputs(this);
+}
+
+void FBuildJob::BeginResolveInputs()
+{
+	if (CanExecuteState(EBuildJobState::ResolveInputs))
+	{
+		const FBuildDefinition LocalDefinition = MoveTemp(Definition).Get();
+		return AdvanceToState(EBuildJobState::ResolveInputsWait, InputResolver->GetData(LocalDefinition, Priority,
+			[this](FBuildInputResolvedParams&& Params) { EndResolveInputs(MoveTemp(Params)); }));
+	}
+}
+
+void FBuildJob::EndResolveInputs(FBuildInputResolvedParams&& Params)
+{
+	if (CanExecuteState(EBuildJobState::ResolveInputsWait))
+	{
+		if (Params.Status == EStatus::Ok)
+		{
+			FBuildInputsBuilder Builder = BuildSystem.CreateInputs(Name);
+			for (const FBuildInputByKey& InputByKey : Params.Inputs)
+			{
+				Builder.AddInput(InputByKey.Key, InputByKey.Data);
+			}
+			SetInputs(Builder.Build());
+		}
+		else
+		{
+			CompleteWithError(TEXT("Failed to resolve input data."_SV));
+		}
+	}
+}
+
+void FBuildJob::SetInputs(FBuildInputs&& InInputs)
+{
+	checkf(Inputs.IsNull(), TEXT("Job already has inputs for build of '%s' by %s."), *Name, *FunctionName);
+	checkf(State == EBuildJobState::ResolveInputs || State == EBuildJobState::ResolveInputsWait,
+		TEXT("Job is not expecting inputs in state %s for build of '%s' by %s"),
+		LexToString(State), *Name, *FunctionName);
+	Inputs = MoveTemp(InInputs);
+	return AdvanceToState(EBuildJobState::ExecuteBuild);
+}
+
+///////////////////////////////////////////////////////////////////////////////////////////////////////////////////////
+
+void FBuildJob::EnterExecuteBuild()
+{
+	const EBuildPolicy BuildPolicy = Context->GetBuildPolicy();
+	if (EnumHasAnyFlags(BuildPolicy, EBuildPolicy::SkipBuild))
+	{
+		return CompleteWithError(TEXT("Failed because build policy skipped building."_SV));
+	}
+	if (!EnumHasAnyFlags(BuildPolicy, EBuildPolicy::Local))
+	{
+		return CompleteWithError(TEXT("Failed because build policy does not allow local execution."_SV));
+	}
+	Action.Get().IterateConstants([this](FStringView Key, FCbObject&& Value)
+	{
+		Context->AddConstant(Key, MoveTemp(Value));
+	});
+	bool bMatchedInputs = true;
+	Action.Get().IterateInputs([this, &bMatchedInputs](FStringView Key, const FIoHash& RawHash, uint64 RawSize)
+	{
+		const FCompressedBuffer& Buffer = Inputs.Get().GetInput(Key);
+		if (Buffer.GetRawHash() == RawHash && Buffer.GetRawSize() == RawSize)
+		{
+			Context->AddInput(Key, Buffer);
+		}
+		else
+		{
+			bMatchedInputs = false;
+			CompleteWithError(WriteToString<256>(TEXT("Failed because input '"_SV), Key, TEXT("' with hash "_SV),
+				RawHash, TEXT(" ("_SV), RawSize, TEXT(" bytes) was resolved with hash "_SV), Buffer.GetRawHash(),
+				TEXT(" ("_SV), Buffer.GetRawSize(), TEXT(" bytes)."_SV)));
+		}
+	});
+	Action.Reset();
+	Inputs.Reset();
+	if (bMatchedInputs)
+	{
+		Scheduler->DispatchExecuteBuild(this);
+	}
+}
+
+void FBuildJob::BeginExecuteBuild()
+{
+	if (CanExecuteState(EBuildJobState::ExecuteBuild))
+	{
+		Context->GetFunction().Build(*Context);
+		if (Context->IsAsyncBuild())
+		{
+			return AdvanceToState(EBuildJobState::ExecuteBuildWait);
+		}
+		else
+		{
+			return SetOutputNoCheck(OutputBuilder.Build());
+		}
+	}
+}
+
+void FBuildJob::EndExecuteBuild()
+{
+	if (CanExecuteState(EBuildJobState::ExecuteBuildWait))
+	{
+		SetOutputNoCheck(OutputBuilder.Build());
+	}
+}
+
+void FBuildJob::SetOutput(const FBuildOutput& InOutput)
+{
+	checkf(State == EBuildJobState::ResolveKey || State == EBuildJobState::ResolveAction ||
+		State == EBuildJobState::CacheQuery || State == EBuildJobState::ResolveInputs ||
+		State == EBuildJobState::ExecuteBuild || State == EBuildJobState::ExecuteBuildWait,
+		TEXT("Job is not expecting an output in state %s for build of '%s' by %s"),
+		LexToString(State), *Name, *FunctionName);
+	SetOutputNoCheck(FBuildOutput(InOutput));
+}
+
+void FBuildJob::SetOutputNoCheck(FBuildOutput&& InOutput)
+{
+	checkf(Output.IsNull(), TEXT("Job already has an output for build of '%s' by %s."), *Name, *FunctionName);
+	Output = MoveTemp(InOutput);
+
+	const EStatus Status = bIsCanceled ? EStatus::Canceled : Output.Get().HasError() ? EStatus::Error : EStatus::Ok;
+	Scheduler->CompleteJob({*this, FBuildOutput(Output.Get()), Status});
+	if (OnComplete)
+	{
+		OnComplete({*this, FBuildOutput(Output.Get()), Status});
+		OnComplete = nullptr;
+	}
+
+	return AdvanceToState(EBuildJobState::CacheStore);
+}
+
+///////////////////////////////////////////////////////////////////////////////////////////////////////////////////////
+
+void FBuildJob::EnterCacheStore()
+{
+	const EBuildPolicy BuildPolicy = Context->GetBuildPolicy();
+	ECachePolicy CachePolicy = Context->GetCachePolicy();
+	if (bIsCached ||
+		!EnumHasAnyFlags(CachePolicy, ECachePolicy::Store) ||
+		EnumHasAnyFlags(BuildPolicy, EBuildPolicy::SkipCachePut) ||
+		Output.Get().HasError())
+	{
+		return AdvanceToState(EBuildJobState::Complete);
+	}
+	Scheduler->DispatchCacheStore(this);
+}
+
+void FBuildJob::BeginCacheStore()
+{
+	if (CanExecuteState(EBuildJobState::CacheStore))
+	{
+		FCacheRecordBuilder RecordBuilder = Cache.CreateRecord(Context->GetCacheKey());
+		Output.Get().Save(RecordBuilder);
+		return AdvanceToState(EBuildJobState::CacheStoreWait,
+			Cache.Put({RecordBuilder.Build()}, Name, Context->GetCachePolicy(), FMath::Min(Priority, EPriority::Highest),
+				[this](FCachePutCompleteParams&& Params) { EndCacheStore(MoveTemp(Params)); }));
+	}
+}
+
+void FBuildJob::EndCacheStore(FCachePutCompleteParams&& Params)
+{
+	if (CanExecuteState(EBuildJobState::CacheStoreWait))
+	{
+		bIsCached = Params.Status == EStatus::Ok;
+		return AdvanceToState(EBuildJobState::Complete);
+	}
+}
+
+///////////////////////////////////////////////////////////////////////////////////////////////////////////////////////
+
+void FBuildJob::EndJob()
+{
+	TRequest Self(this);
+
+	if (Event)
+	{
+		Event->Trigger();
+	}
+	Scheduler->EndJob(this);
+	Release();
+
+	FWriteScopeLock WriteLock(Lock);
+	Scheduler = nullptr;
+	InputResolver = nullptr;
+}
+
+///////////////////////////////////////////////////////////////////////////////////////////////////////////////////////
+
+void FBuildJob::CompleteWithError(FStringView Error)
+{
+	if (FWriteScopeLock WriteLock(Lock); Output || NextState == EBuildJobState::Complete)
+	{
+		return;
+	}
+	TStringBuilder<32> Category;
+	Category << ImplicitConv<FName>(LogDerivedDataBuild.GetCategoryName());
+	OutputBuilder.AddError(Category, Error);
+	SetOutputNoCheck(OutputBuilder.Build());
+	return AdvanceToState(EBuildJobState::Complete);
+}
+
+bool FBuildJob::CanExecuteState(EBuildJobState RequestedState)
+{
+	if (bIsCanceled)
+	{
+		AdvanceToState(EBuildJobState::Complete);
+		return false;
+	}
+	const EBuildJobState LocalState = State;
+	if (LocalState == RequestedState)
+	{
+		return true;
+	}
+	switch (RequestedState)
+	{
+	case EBuildJobState::ResolveKeyWait:    return LocalState == EBuildJobState::ResolveKey;
+	case EBuildJobState::ResolveActionWait: return LocalState == EBuildJobState::ResolveAction;
+	case EBuildJobState::CacheQueryWait:    return LocalState == EBuildJobState::CacheQuery;
+	case EBuildJobState::ResolveInputsWait: return LocalState == EBuildJobState::ResolveInputs;
+	case EBuildJobState::ExecuteBuildWait:  return LocalState == EBuildJobState::ExecuteBuild;
+	case EBuildJobState::CacheStoreWait:    return LocalState == EBuildJobState::CacheStore;
+	}
+	return false;
+}
+
+void FBuildJob::AdvanceToState(EBuildJobState NewState, FRequest NewRequest)
+{
+	bool bLocalIsCanceled = false;
+	EBuildJobState OldState = EBuildJobState::Complete;
+
+	if (FWriteScopeLock WriteLock(Lock); NextState < NewState)
+	{
+		// The structure of the job requires that the any previous WaitRequest is complete by
+		// the time that AdvanceToState is run, because nothing can happen in waiting states,
+		// and the callback from the request should be triggering this state transition. This
+		// is not possible to validate by polling the request because the request contract is
+		// that it is not considered complete until the callback finishes executing.
+		WaitRequest = MoveTemp(NewRequest);
+
+		bLocalIsCanceled = bIsCanceled;
+		if (bLocalIsCanceled)
+		{
+			NewState = EBuildJobState::Complete;
+		}
+
+		NextState = NewState;
+
+		if (!bInAdvanceToState)
+		{
+			bInAdvanceToState = true;
+			OldState = State;
+			State = NewState;
+		}
+	}
+
+	while (OldState < NewState)
+	{
+		ExecuteTransition(OldState, NewState);
+
+		if (NewState == EBuildJobState::Complete)
+		{
+			return;
+		}
+
+		if (FWriteScopeLock WriteLock(Lock); NewState < NextState)
+		{
+			OldState = NewState;
+			NewState = NextState;
+			State = NewState;
+		}
+		else
+		{
+			bInAdvanceToState = false;
+			break;
+		}
+	}
+}
+
+void FBuildJob::ExecuteTransition(EBuildJobState OldState, EBuildJobState NewState)
+{
+	if (OldState == EBuildJobState::NotStarted)
+	{
+		BeginJob();
+	}
+	if (OldState < EBuildJobState::CacheQuery && EBuildJobState::CacheQuery <= NewState)
+	{
+		Configure();
+	}
+	if (OldState < EBuildJobState::ResolveInputsWait && EBuildJobState::ResolveInputsWait <= NewState)
+	{
+		Definition.Reset();
+	}
+	if (OldState < EBuildJobState::ExecuteBuildWait && EBuildJobState::ExecuteBuildWait <= NewState)
+	{
+		Action.Reset();
+		Inputs.Reset();
+	}
+	if (OldState <= EBuildJobState::ExecuteBuildWait && EBuildJobState::ExecuteBuildWait < NewState && !Output)
+	{
+		SetOutputNoCheck(OutputBuilder.Build());
+	}
+	if (OldState < EBuildJobState::CacheStore && EBuildJobState::CacheStore <= NewState)
+	{
+		Context->ResetInputs();
+	}
+	if (OldState <= EBuildJobState::CacheStore && EBuildJobState::CacheStore < NewState)
+	{
+		Output.Reset();
+		Context = nullptr;
+	}
+
+	switch (NewState)
+	{
+	case EBuildJobState::ResolveKey:    return EnterResolveKey();
+	case EBuildJobState::ResolveAction: return EnterResolveAction();
+	case EBuildJobState::CacheQuery:    return EnterCacheQuery();
+	case EBuildJobState::ResolveInputs: return EnterResolveInputs();
+	case EBuildJobState::ExecuteBuild:  return EnterExecuteBuild();
+	case EBuildJobState::CacheStore:    return EnterCacheStore();
+	case EBuildJobState::Complete:      return EndJob();
+	}
+}
+
+///////////////////////////////////////////////////////////////////////////////////////////////////////////////////////
+
+TRequest<IBuildJob> CreateBuildJob(ICache& Cache, IBuild& BuildSystem, const FBuildKey& Key, IBuildInputResolver* InputResolver)
+{
+	return TRequest(new FBuildJob(Cache, BuildSystem, Key, InputResolver));
+}
+
+TRequest<IBuildJob> CreateBuildJob(ICache& Cache, IBuild& BuildSystem, const FBuildDefinition& Definition, IBuildInputResolver* InputResolver)
+{
+	return TRequest(new FBuildJob(Cache, BuildSystem, Definition, InputResolver));
+}
+
+TRequest<IBuildJob> CreateBuildJob(ICache& Cache, IBuild& BuildSystem, const FBuildAction& Action, const FOptionalBuildInputs& Inputs)
+{
+	return TRequest(new FBuildJob(Cache, BuildSystem, Action, Inputs));
+}
+
+} // UE::DerivedData::Private
diff --git a/Engine/Source/Developer/DerivedDataCache/Private/DerivedDataBuildPrivate.h b/Engine/Source/Developer/DerivedDataCache/Private/DerivedDataBuildPrivate.h
--- a/Engine/Source/Developer/DerivedDataCache/Private/DerivedDataBuildPrivate.h
+++ b/Engine/Source/Developer/DerivedDataCache/Private/DerivedDataBuildPrivate.h
@@ -8,18 +8,25 @@
 class FCbObject;
 struct FGuid;
 
+namespace UE::DerivedData { class FBuildAction; }
 namespace UE::DerivedData { class FBuildActionBuilder; }
+namespace UE::DerivedData { class FBuildDefinition; }
 namespace UE::DerivedData { class FBuildDefinitionBuilder; }
 namespace UE::DerivedData { class FBuildInputsBuilder; }
 namespace UE::DerivedData { class FBuildOutputBuilder; }
 namespace UE::DerivedData { class FCacheRecord; }
 namespace UE::DerivedData { class FOptionalBuildAction; }
 namespace UE::DerivedData { class FOptionalBuildDefinition; }
+namespace UE::DerivedData { class FOptionalBuildInputs; }
 namespace UE::DerivedData { class FOptionalBuildOutput; }
 namespace UE::DerivedData { class IBuild; }
 namespace UE::DerivedData { class IBuildFunctionRegistry; }
+namespace UE::DerivedData { class IBuildInputResolver; }
+namespace UE::DerivedData { class IBuildJob; }
 namespace UE::DerivedData { class IBuildScheduler; }
 namespace UE::DerivedData { class ICache; }
+namespace UE::DerivedData { struct FBuildKey; }
+namespace UE::DerivedData { template <typename RequestType> class TRequest; }
 
 namespace UE::DerivedData::Private
 {
@@ -51,4 +58,9 @@ FBuildOutputBuilder CreateBuildOutput(FStringView Name, FStringView Function);
 FOptionalBuildOutput LoadBuildOutput(FStringView Name, FStringView Function, const FCbObject& Output);
 FOptionalBuildOutput LoadBuildOutput(FStringView Name, FStringView Function, const FCacheRecord& Output);
 
+// Implemented in DerivedDataBuildJob.cpp
+TRequest<IBuildJob> CreateBuildJob(ICache& Cache, IBuild& BuildSystem, const FBuildKey& Key, IBuildInputResolver* InputResolver);
+TRequest<IBuildJob> CreateBuildJob(ICache& Cache, IBuild& BuildSystem, const FBuildDefinition& Definition, IBuildInputResolver* InputResolver);
+TRequest<IBuildJob> CreateBuildJob(ICache& Cache, IBuild& BuildSystem, const FBuildAction& Action, const FOptionalBuildInputs& Inputs);
+
 } // UE::DerivedData::Private
