DDC: Added the implementation of UE::DerivedData::IBuildSession

#rnx

diff --git a/Engine/Source/Developer/DerivedDataCache/Private/DerivedDataBuild.cpp b/Engine/Source/Developer/DerivedDataCache/Private/DerivedDataBuild.cpp
--- a/Engine/Source/Developer/DerivedDataCache/Private/DerivedDataBuild.cpp
+++ b/Engine/Source/Developer/DerivedDataCache/Private/DerivedDataBuild.cpp
@@ -8,6 +8,8 @@
 #include "DerivedDataBuildInputs.h"
 #include "DerivedDataBuildOutput.h"
 #include "DerivedDataBuildPrivate.h"
+#include "DerivedDataBuildScheduler.h"
+#include "DerivedDataBuildSession.h"
 #include "DerivedDataCache.h"
 #include "Misc/Guid.h"
 
@@ -16,6 +18,51 @@ namespace UE::DerivedData::Private
 
 DEFINE_LOG_CATEGORY(LogDerivedDataBuild);
 
+/**
+ * Derived Data Build System
+ *
+ * Public Data Types:
+ *
+ * FBuildDefinition:
+ * - Function, Constants, Key->InputId
+ * - From FBuildDefinitionBuilder via IBuild::CreateDefinition()
+ * - Serializes to/from FCbObject
+ * FBuildAction:
+ * - Function+Version, BuildSystemVersion, Constants, Key->InputHash+InputSize
+ * - From FBuildActionBuilder via IBuild::CreateAction()
+ * - Serializes to/from FCbObject
+ * FBuildInputs:
+ * - Key->InputBuffer
+ * - From FBuildInputsBuilder via IBuild::CreateInputs()
+ * FBuildOutput:
+ * - Metadata, Payloads[], Diagnostics[] (Level, Category, Message)
+ * - From FBuildOutputBuilder via IBuild::CreateOutput()
+ * - Serializes to/from FCbObject and FCacheRecord
+ * FBuildKey:
+ * - Unique ID for FBuildDefinition
+ * FBuildActionKey:
+ * - Unique ID for FBuildAction
+ * - Combined with FCacheBucket to create a FCacheKey
+ *
+ * Public Interface Types:
+ *
+ * IBuildFunction:
+ * - Name (Unique), Version
+ * - Represents the steps necessary to convert inputs to the output
+ * IBuildFunctionRegistry:
+ * - Registry of IBuildFunction used by IBuild
+ * IBuildInputResolver:
+ * - Resolves FBuildKey to FBuildDefinition
+ * - Resolves FBuildDefinition to Key->(RawHash, RawSize, OptionalBuffer)
+ * IBuildJob:
+ * - From FBuildKey or FBuildDefinition or FBuildAction+FBuildInputs via FBuildSession
+ * - Represents an executing build request and implements IRequest
+ * IBuildScheduler:
+ * - Schedules execution of the operations on IBuildJob
+ * FBuildSession:
+ * - From ICache+IBuild+IBuildScheduler+IBuildInputResolver via IBuild::CreateSession()
+ * - Interface for scheduling builds from a key, definition, or action+inputs
+ */
 class FBuild final : public IBuild
 {
 public:
@@ -64,6 +111,11 @@ public:
 		return LoadBuildOutput(Name, Function, Output);
 	}
 
+	FBuildSession CreateSession(FStringView Name, IBuildInputResolver* InputResolver, IBuildScheduler* Scheduler) final
+	{
+		return CreateBuildSession(Name, Cache, *this, Scheduler ? *Scheduler : *DefaultScheduler, InputResolver);
+	}
+
 	const FGuid& GetVersion() const final
 	{
 		return Version;
@@ -77,6 +129,7 @@ public:
 private:
 	ICache& Cache;
 	TUniquePtr<IBuildFunctionRegistry> FunctionRegistry{CreateBuildFunctionRegistry()};
+	TUniquePtr<IBuildScheduler> DefaultScheduler{CreateBuildScheduler()};
 	const FGuid Version{TEXT("ac0574e5-62bd-4c2e-84ec-f2efe48c0fef")};
 };
 
diff --git a/Engine/Source/Developer/DerivedDataCache/Private/DerivedDataBuildPrivate.h b/Engine/Source/Developer/DerivedDataCache/Private/DerivedDataBuildPrivate.h
--- a/Engine/Source/Developer/DerivedDataCache/Private/DerivedDataBuildPrivate.h
+++ b/Engine/Source/Developer/DerivedDataCache/Private/DerivedDataBuildPrivate.h
@@ -14,6 +14,7 @@ namespace UE::DerivedData { class FBuildDefinition; }
 namespace UE::DerivedData { class FBuildDefinitionBuilder; }
 namespace UE::DerivedData { class FBuildInputsBuilder; }
 namespace UE::DerivedData { class FBuildOutputBuilder; }
+namespace UE::DerivedData { class FBuildSession; }
 namespace UE::DerivedData { class FCacheRecord; }
 namespace UE::DerivedData { class FOptionalBuildAction; }
 namespace UE::DerivedData { class FOptionalBuildDefinition; }
@@ -58,6 +59,14 @@ FBuildOutputBuilder CreateBuildOutput(FStringView Name, FStringView Function);
 FOptionalBuildOutput LoadBuildOutput(FStringView Name, FStringView Function, const FCbObject& Output);
 FOptionalBuildOutput LoadBuildOutput(FStringView Name, FStringView Function, const FCacheRecord& Output);
 
+// Implemented in DerivedDataBuildSession.cpp
+FBuildSession CreateBuildSession(
+	FStringView Name,
+	ICache& Cache,
+	IBuild& BuildSystem,
+	IBuildScheduler& Scheduler,
+	IBuildInputResolver* InputResolver);
+
 // Implemented in DerivedDataBuildJob.cpp
 TRequest<IBuildJob> CreateBuildJob(ICache& Cache, IBuild& BuildSystem, const FBuildKey& Key, IBuildInputResolver* InputResolver);
 TRequest<IBuildJob> CreateBuildJob(ICache& Cache, IBuild& BuildSystem, const FBuildDefinition& Definition, IBuildInputResolver* InputResolver);
diff --git a/Engine/Source/Developer/DerivedDataCache/Private/DerivedDataBuildSession.cpp b/Engine/Source/Developer/DerivedDataCache/Private/DerivedDataBuildSession.cpp
new file mode 100644
--- /dev/null
+++ b/Engine/Source/Developer/DerivedDataCache/Private/DerivedDataBuildSession.cpp
@@ -0,0 +1,149 @@
+// Copyright Epic Games, Inc. All Rights Reserved.
+
+#include "DerivedDataBuildSession.h"
+
+#include "Containers/UnrealString.h"
+#include "DerivedDataBuildAction.h"
+#include "DerivedDataBuildDefinition.h"
+#include "DerivedDataBuildJob.h"
+#include "DerivedDataBuildOutput.h"
+#include "DerivedDataBuildPrivate.h"
+
+namespace UE::DerivedData::Private
+{
+
+class FBuildSessionInternal final : public IBuildSessionInternal
+{
+public:
+	FBuildSessionInternal(
+		FStringView InName,
+		ICache& InCache,
+		IBuild& InBuildSystem,
+		IBuildScheduler& InScheduler,
+		IBuildInputResolver* InInputResolver)
+		: Name(InName)
+		, Cache(InCache)
+		, BuildSystem(InBuildSystem)
+		, Scheduler(InScheduler)
+		, InputResolver(InInputResolver)
+	{
+	}
+
+	FStringView GetName() const final { return Name; }
+
+	FRequest Build(
+		const FBuildDefinition& Definition,
+		EBuildPolicy Policy,
+		EPriority Priority,
+		FOnBuildComplete&& OnComplete) final;
+
+	FRequest BuildAction(
+		const FBuildAction& Action,
+		const FOptionalBuildInputs& Inputs,
+		EBuildPolicy Policy,
+		EPriority Priority,
+		FOnBuildActionComplete&& OnComplete) final;
+
+	FRequest BuildPayload(
+		const FBuildPayloadKey& Payload,
+		EBuildPolicy Policy,
+		EPriority Priority,
+		FOnBuildPayloadComplete&& OnComplete) final;
+
+	FString Name;
+	ICache& Cache;
+	IBuild& BuildSystem;
+	IBuildScheduler& Scheduler;
+	IBuildInputResolver* InputResolver;
+};
+
+///////////////////////////////////////////////////////////////////////////////////////////////////////////////////////
+
+FRequest FBuildSessionInternal::Build(
+	const FBuildDefinition& Definition,
+	EBuildPolicy Policy,
+	EPriority Priority,
+	FOnBuildComplete&& OnComplete)
+{
+	TRequest Job(CreateBuildJob(Cache, BuildSystem, Definition, InputResolver));
+	FOnBuildJobComplete OnJobComplete;
+	if (OnComplete)
+	{
+		OnJobComplete = [Definition, OnComplete = MoveTemp(OnComplete)](FBuildJobCompleteParams&& Params)
+		{
+			OnComplete({Definition.GetKey(), MoveTemp(Params.Output), Params.Status});
+		};
+	}
+	Job->Schedule(Scheduler, Policy, Priority, MoveTemp(OnJobComplete));
+	return Job;
+}
+
+FRequest FBuildSessionInternal::BuildAction(
+	const FBuildAction& Action,
+	const FOptionalBuildInputs& Inputs,
+	EBuildPolicy Policy,
+	EPriority Priority,
+	FOnBuildActionComplete&& OnComplete)
+{
+	TRequest Job(CreateBuildJob(Cache, BuildSystem, Action, Inputs));
+	FOnBuildJobComplete OnJobComplete;
+	if (OnComplete)
+	{
+		OnJobComplete = [Action, OnComplete = MoveTemp(OnComplete)](FBuildJobCompleteParams&& Params)
+		{
+			OnComplete({Action.GetKey(), MoveTemp(Params.Output), Params.Status});
+		};
+	}
+	Job->Schedule(Scheduler, Policy, Priority, MoveTemp(OnJobComplete));
+	return Job;
+}
+
+FRequest FBuildSessionInternal::BuildPayload(
+	const FBuildPayloadKey& PayloadKey,
+	EBuildPolicy Policy,
+	EPriority Priority,
+	FOnBuildPayloadComplete&& OnComplete)
+{
+	// This requests the entire output to get one payload. It will be optimized later to request only one payload.
+	TRequest Job(CreateBuildJob(Cache, BuildSystem, PayloadKey.BuildKey, InputResolver));
+	FOnBuildJobComplete OnJobComplete;
+	if (OnComplete)
+	{
+		OnJobComplete = [PayloadKey, OnComplete = MoveTemp(OnComplete)](FBuildJobCompleteParams&& Params)
+		{
+			FPayload Payload;
+			EStatus Status = Params.Status;
+			if (Status == EStatus::Ok)
+			{
+				Payload = Params.Output.GetPayload(PayloadKey.Id);
+				Status = Payload ? EStatus::Ok : EStatus::Error;
+			}
+			if (!Payload)
+			{
+				Payload = FPayload(PayloadKey.Id);
+			}
+			OnComplete({PayloadKey.BuildKey, MoveTemp(Payload), Status});
+		};
+	}
+	Job->Schedule(Scheduler, Policy, Priority, MoveTemp(OnJobComplete));
+	return Job;
+}
+
+///////////////////////////////////////////////////////////////////////////////////////////////////////////////////////
+
+FBuildSession CreateBuildSession(IBuildSessionInternal* Session)
+{
+	return FBuildSession(Session);
+}
+
+FBuildSession CreateBuildSession(
+	FStringView Name,
+	ICache& Cache,
+	IBuild& BuildSystem,
+	IBuildScheduler& Scheduler,
+	IBuildInputResolver* InputResolver)
+{
+	return CreateBuildSession(new FBuildSessionInternal(Name, Cache, BuildSystem, Scheduler, InputResolver));
+}
+
+} // UE::DerivedData::Private
diff --git a/Engine/Source/Developer/DerivedDataCache/Public/DerivedDataBuild.h b/Engine/Source/Developer/DerivedDataCache/Public/DerivedDataBuild.h
--- a/Engine/Source/Developer/DerivedDataCache/Public/DerivedDataBuild.h
+++ b/Engine/Source/Developer/DerivedDataCache/Public/DerivedDataBuild.h
@@ -11,11 +11,14 @@ namespace UE::DerivedData { class FBuildActionBuilder; }
 namespace UE::DerivedData { class FBuildDefinitionBuilder; }
 namespace UE::DerivedData { class FBuildInputsBuilder; }
 namespace UE::DerivedData { class FBuildOutputBuilder; }
+namespace UE::DerivedData { class FBuildSession; }
 namespace UE::DerivedData { class FCacheRecord; }
 namespace UE::DerivedData { class FOptionalBuildAction; }
 namespace UE::DerivedData { class FOptionalBuildDefinition; }
 namespace UE::DerivedData { class FOptionalBuildOutput; }
 namespace UE::DerivedData { class IBuildFunctionRegistry; }
+namespace UE::DerivedData { class IBuildInputResolver; }
+namespace UE::DerivedData { class IBuildScheduler; }
 
 namespace UE::DerivedData
 {
@@ -23,7 +26,19 @@ namespace UE::DerivedData
 /**
  * Interface to the build system.
  *
- * This is only a preview of a portion of the interface and does not support build execution.
+ * Executing a build typically requires a definition, input resolver, session, and function.
+ *
+ * Use IBuild::CreateDefinition() to make a new build definition, or use IBuild::LoadDefinition()
+ * to load a build definition that was previously saved. This references the function to execute,
+ * and the inputs needed by the function.
+ *
+ * Use IBuild::CreateSession() to make a new build session with a build input resolver to resolve
+ * input references into the referenced data. Use FBuildSession::Build() to schedule a definition
+ * to build, along with any of its transitive build dependencies.
+ *
+ * Implement IBuildFunction, with a unique name and version, to add payloads to the build context
+ * based on constants and inputs in the context. Use TBuildFunctionFactory to add the function to
+ * the registry at IBuild::GetFunctionRegistry() to allow the build job to find it.
  */
 class IBuild
 {
@@ -90,6 +105,15 @@ public:
 	virtual FOptionalBuildOutput LoadOutput(FStringView Name, FStringView Function, const FCbObject& Output) = 0;
 	virtual FOptionalBuildOutput LoadOutput(FStringView Name, FStringView Function, const FCacheRecord& Output) = 0;
 
+	/**
+	 * Create a build session.
+	 *
+	 * @param Name            The name by which to identify this session for logging and profiling.
+	 * @param InputResolver   The input resolver to resolve definitions and inputs for requested builds.
+	 * @param Scheduler       The scheduler for builds created through the session. Optional.
+	 */
+	virtual FBuildSession CreateSession(FStringView Name, IBuildInputResolver* InputResolver, IBuildScheduler* Scheduler = nullptr) = 0;
+
 	/**
 	 * Returns the version of the build system.
 	 *
