DDC: Added interfaces for external build workers

- IBuildWorkerFactory creates a FBuildWorker that describes the functions, environment, files, and platform of the worker.
- IBuildWorkerExecutor asynchronously executes a build action and inputs against a build worker.
- IBuildWorkerRegistry discovers and registers build workers, and dispatches builds from build jobs.

#rnx

diff --git a/Engine/Source/Developer/DerivedDataCache/Private/DerivedDataBuildWorker.h b/Engine/Source/Developer/DerivedDataCache/Private/DerivedDataBuildWorker.h
new file mode 100644
--- /dev/null
+++ b/Engine/Source/Developer/DerivedDataCache/Private/DerivedDataBuildWorker.h
@@ -0,0 +1,111 @@
+// Copyright Epic Games, Inc. All Rights Reserved.
+
+#pragma once
+
+#include "CoreTypes.h"
+#include "Containers/ArrayView.h"
+#include "Containers/StringFwd.h"
+#include "DerivedDataBuildKey.h"
+#include "DerivedDataRequest.h"
+#include "Features/IModularFeature.h"
+#include "UObject/NameTypes.h"
+
+class FCompressedBuffer;
+struct FGuid;
+struct FIoHash;
+
+template <typename FuncType> class TFunctionRef;
+template <typename FuncType> class TUniqueFunction;
+
+namespace UE::DerivedData { class FBuildAction; }
+namespace UE::DerivedData { class FOptionalBuildInputs; }
+namespace UE::DerivedData { class FOptionalBuildOutput; }
+namespace UE::DerivedData { struct FBuildWorkerActionCompleteParams; }
+namespace UE::DerivedData { struct FBuildWorkerFilesCompleteParams; }
+namespace UE::DerivedData { enum class EBuildPolicy : uint8; }
+
+namespace UE::DerivedData
+{
+
+using FOnBuildWorkerActionComplete = TUniqueFunction<void (FBuildWorkerActionCompleteParams&& Params)>;
+using FOnBuildWorkerFilesCompleteParams = TUniqueFunction<void (FBuildWorkerFilesCompleteParams&& Params)>;
+
+class FBuildWorker
+{
+public:
+	virtual ~FBuildWorker() = default;
+	virtual FStringView GetName() const = 0;
+	virtual FStringView GetPath() const = 0;
+	virtual FStringView GetHostPlatform() const = 0;
+	virtual FGuid GetBuildSystemVersion() const = 0;
+	virtual void IterateFunctions(TFunctionRef<void (FStringView Name, const FGuid& Version)> Visitor) const = 0;
+	virtual void IterateFiles(TFunctionRef<void (FStringView Path, const FIoHash& RawHash, uint64 RawSize)> Visitor) const = 0;
+	virtual void IterateExecutables(TFunctionRef<void (FStringView Path, const FIoHash& RawHash, uint64 RawSize)> Visitor) const = 0;
+	virtual void IterateEnvironment(TFunctionRef<void (FStringView Name, FStringView Value)> Visitor) const = 0;
+};
+
+class FBuildWorkerBuilder
+{
+public:
+	virtual ~FBuildWorkerBuilder() = default;
+	virtual void SetName(FStringView Name) = 0;
+	virtual void SetPath(FStringView Path) = 0;
+	virtual void SetHostPlatform(FStringView Name) = 0;
+	virtual void SetBuildSystemVersion(const FGuid& Version) = 0;
+	virtual void AddFunction(FStringView Name, const FGuid& Version) = 0;
+	virtual void AddFile(FStringView Path, const FIoHash& RawHash, uint64 RawSize) = 0;
+	virtual void AddExecutable(FStringView Path, const FIoHash& RawHash, uint64 RawSize) = 0;
+	virtual void SetEnvironment(FStringView Name, FStringView Value) = 0;
+};
+
+class IBuildWorkerFactory : public IModularFeature
+{
+public:
+	virtual ~IBuildWorkerFactory() = default;
+
+	virtual void Build(FBuildWorkerBuilder& Builder) = 0;
+	virtual FRequest GetFiles(TConstArrayView<FIoHash> Files, EPriority Priority, FOnBuildWorkerFilesCompleteParams&& OnComplete) = 0;
+
+	/** Returns the name of the build worker factory modular feature. */
+	static FName GetFeatureName()
+	{
+		return FName("BuildWorkerFactory");
+	}
+};
+
+class IBuildWorkerExecutor : public IModularFeature
+{
+public:
+	virtual ~IBuildWorkerExecutor() = default;
+
+	virtual FRequest BuildAction(
+		const FBuildAction& Action,
+		const FOptionalBuildInputs& Inputs,
+		const FBuildWorker& Worker,
+		IBuildWorkerFactory& WorkerFactory,
+		EBuildPolicy Policy,
+		EPriority Priority,
+		FOnBuildWorkerActionComplete&& OnComplete) = 0;
+
+	/** Returns the name of the build worker executor modular feature. */
+	static FName GetFeatureName()
+	{
+		return FName("BuildWorkerExecutor");
+	}
+};
+
+struct FBuildWorkerActionCompleteParams
+{
+	FBuildActionKey Key;
+	FOptionalBuildOutput&& Output;
+	TConstArrayView<FStringView> MissingInputs;
+	EStatus Status = EStatus::Error;
+};
+
+struct FBuildWorkerFilesCompleteParams
+{
+	TConstArrayView<FCompressedBuffer> Files;
+	EStatus Status = EStatus::Error;
+};
+
+} // UE::DerivedData
diff --git a/Engine/Source/Developer/DerivedDataCache/Private/DerivedDataBuildWorkerRegistry.h b/Engine/Source/Developer/DerivedDataCache/Private/DerivedDataBuildWorkerRegistry.h
new file mode 100644
--- /dev/null
+++ b/Engine/Source/Developer/DerivedDataCache/Private/DerivedDataBuildWorkerRegistry.h
@@ -0,0 +1,34 @@
+// Copyright Epic Games, Inc. All Rights Reserved.
+
+#pragma once
+
+template <typename FuncType> class TUniqueFunction;
+
+namespace UE::DerivedData { class FBuildAction; }
+namespace UE::DerivedData { class FOptionalBuildInputs; }
+namespace UE::DerivedData { class IRequest; }
+namespace UE::DerivedData { struct FBuildWorkerActionCompleteParams; }
+namespace UE::DerivedData { enum class EBuildPolicy : uint8; }
+namespace UE::DerivedData { enum class EPriority : uint8; }
+namespace UE::DerivedData { template <typename RequestType> class TRequest; }
+namespace UE::DerivedData { using FRequest = TRequest<IRequest>; }
+
+namespace UE::DerivedData
+{
+
+using FOnBuildWorkerActionComplete = TUniqueFunction<void (FBuildWorkerActionCompleteParams&& Params)>;
+
+class IBuildWorkerRegistry
+{
+public:
+	virtual ~IBuildWorkerRegistry() = default;
+
+	virtual FRequest BuildAction(
+		const FBuildAction& Action,
+		const FOptionalBuildInputs& Inputs,
+		EBuildPolicy Policy,
+		EPriority Priority,
+		FOnBuildWorkerActionComplete&& Params) = 0;
+};
+
+} // UE::DerivedData
