DDC: Added serialization of FCacheRecord to/from FCbPackage

#rnx

diff --git a/Engine/Source/Developer/DerivedDataCache/Private/DerivedDataBackends.cpp b/Engine/Source/Developer/DerivedDataCache/Private/DerivedDataBackends.cpp
--- a/Engine/Source/Developer/DerivedDataCache/Private/DerivedDataBackends.cpp
+++ b/Engine/Source/Developer/DerivedDataCache/Private/DerivedDataBackends.cpp
@@ -31,6 +31,7 @@
 #include "Math/UnitConversion.h"
 #include "Internationalization/FastDecimalFormat.h"
 #include "Math/BasicMathExpressionEvaluator.h"
+#include "Serialization/CompactBinaryPackage.h"
 #include <atomic>
 
 DEFINE_LOG_CATEGORY(LogDerivedDataCache);
@@ -1129,6 +1130,16 @@ public:
 		return Private::CreateCacheRecordBuilder(Key);
 	}
 
+	virtual FCbPackage SaveRecord(const FCacheRecord& Record) override
+	{
+		return Private::SaveCacheRecord(Record);
+	}
+
+	virtual FOptionalCacheRecord LoadRecord(const FCbPackage& Package) override
+	{
+		return Private::LoadCacheRecord(Package);
+	}
+
 private:
 
 	/** Delete all created backends in the reversed order they were created. */
diff --git a/Engine/Source/Developer/DerivedDataCache/Private/DerivedDataCache.cpp b/Engine/Source/Developer/DerivedDataCache/Private/DerivedDataCache.cpp
--- a/Engine/Source/Developer/DerivedDataCache/Private/DerivedDataCache.cpp
+++ b/Engine/Source/Developer/DerivedDataCache/Private/DerivedDataCache.cpp
@@ -17,6 +17,7 @@
 #include "Misc/CommandLine.h"
 #include "Misc/ScopeLock.h"
 #include "ProfilingDebugging/CookStats.h"
+#include "Serialization/CompactBinaryPackage.h"
 #include "Stats/Stats.h"
 #include "Stats/StatsMisc.h"
 #include <atomic>
@@ -662,6 +663,16 @@ public:
 
 	FCacheRecordBuilder CreateRecord(const FCacheKey& Key) final { return CreateCacheRecordBuilder(Key); }
 
+	FCbPackage SaveRecord(const FCacheRecord& Record) final
+	{
+		return SaveCacheRecord(Record);
+	}
+
+	FOptionalCacheRecord LoadRecord(const FCbPackage& Package) final
+	{
+		return LoadCacheRecord(Package);
+	}
+
 	FRequest Put(
 		TConstArrayView<FCacheRecord> Records,
 		FStringView Context,
diff --git a/Engine/Source/Developer/DerivedDataCache/Private/DerivedDataCachePrivate.h b/Engine/Source/Developer/DerivedDataCache/Private/DerivedDataCachePrivate.h
--- a/Engine/Source/Developer/DerivedDataCache/Private/DerivedDataCachePrivate.h
+++ b/Engine/Source/Developer/DerivedDataCache/Private/DerivedDataCachePrivate.h
@@ -4,10 +4,13 @@
 
 #include "Containers/StringFwd.h"
 
+class FCbPackage;
 class FDerivedDataCacheInterface;
+enum class ECachePolicy : uint8;
 
 namespace UE::DerivedData { class FCacheBucket; }
 namespace UE::DerivedData { class FCacheRecordBuilder; }
+namespace UE::DerivedData { class FOptionalCacheRecord; }
 namespace UE::DerivedData { struct FCacheKey; }
 
 namespace UE::DerivedData::Private
@@ -21,5 +24,7 @@ FCacheBucket CreateCacheBucket(FStringView Name);
 
 // Implemented in DerivedDataCacheRecord.cpp
 FCacheRecordBuilder CreateCacheRecordBuilder(const FCacheKey& Key);
+FCbPackage SaveCacheRecord(const FCacheRecord& Record);
+FOptionalCacheRecord LoadCacheRecord(const FCbPackage& Package);
 
 } // UE::DerivedData::Private
diff --git a/Engine/Source/Developer/DerivedDataCache/Private/DerivedDataCacheRecord.cpp b/Engine/Source/Developer/DerivedDataCache/Private/DerivedDataCacheRecord.cpp
--- a/Engine/Source/Developer/DerivedDataCache/Private/DerivedDataCacheRecord.cpp
+++ b/Engine/Source/Developer/DerivedDataCache/Private/DerivedDataCacheRecord.cpp
@@ -2,13 +2,19 @@
 
 #include "DerivedDataCacheRecord.h"
 
+#include "Algo/AllOf.h"
+#include "Containers/StringConv.h"
 #include "Containers/UnrealString.h"
+#include "DerivedDataCache.h"
 #include "DerivedDataCacheKey.h"
+#include "DerivedDataCachePrivate.h"
 #include "HAL/CriticalSection.h"
 #include "Misc/ScopeExit.h"
 #include "Misc/ScopeRWLock.h"
 #include "Misc/StringBuilder.h"
 #include "Serialization/CompactBinary.h"
+#include "Serialization/CompactBinaryPackage.h"
+#include "Serialization/CompactBinaryWriter.h"
 #include "UObject/NameTypes.h"
 #include <atomic>
 
@@ -263,6 +269,132 @@ FCacheRecordBuilder CreateCacheRecordBuilder(const FCacheKey& Key)
 	return CreateCacheRecordBuilder(new FCacheRecordBuilderInternal(Key));
 }
 
+FCbPackage SaveCacheRecord(const FCacheRecord& Record)
+{
+	FCbPackage Package;
+	FCbWriter Writer;
+	Writer.BeginObject();
+	{
+		const FCacheKey& Key = Record.GetKey();
+		Writer.BeginObject("Key"_ASV);
+		Writer.AddString("Bucket"_ASV, Key.Bucket.ToString<ANSICHAR>());
+		Writer.AddHash("Hash"_ASV, Key.Hash);
+		Writer.EndObject();
+	}
+
+	if (const FCbObject& Meta = Record.GetMeta())
+	{
+		Writer.AddObject("Meta"_ASV, Meta);
+	}
+	auto SavePayload = [&Package, &Writer](const FPayload& Payload)
+	{
+		Writer.BeginObject();
+		Writer.AddObjectId("Id"_ASV, FCbObjectId(Payload.GetId().GetView()));
+		Writer.AddHash("RawHash"_ASV, Payload.GetRawHash());
+		Writer.AddInteger("RawSize"_ASV, Payload.GetRawSize());
+		const FCompositeBuffer CompressedBuffer = Payload.GetData().GetCompressed();
+		FCbAttachment AttachedCompressed(CompressedBuffer.Flatten());
+		Writer.AddAttachment("CompressedHash", AttachedCompressed);
+		Package.AddAttachment(AttachedCompressed);
+		Writer.EndObject();
+	};
+	if (const FPayload& Value = Record.GetValuePayload())
+	{
+		Writer.SetName("Value"_ASV);
+		SavePayload(Value);
+	}
+	TConstArrayView<FPayload> Attachments = Record.GetAttachmentPayloads();
+	if (!Attachments.IsEmpty())
+	{
+		Writer.BeginArray("Attachments"_ASV);
+		for (const FPayload& Attachment : Attachments)
+		{
+			SavePayload(Attachment);
+		}
+		Writer.EndArray();
+	}
+	Writer.EndObject();
+
+	Package.SetObject(Writer.Save().AsObject());
+	return Package;
+}
+
+FOptionalCacheRecord LoadCacheRecord(const FCbPackage& Package)
+{
+	FCbObjectView RecordObject = Package.GetObject();
+	
+	FCacheKey Key;
+	FCbObjectView KeyObject = RecordObject["Key"_ASV].AsObjectView();
+	auto TrySetBucketName = [](FStringView Name, FCacheKey& Key)
+	{
+		if (FCString::IsPureAnsi(Name.GetData(), Name.Len())
+			&& Algo::AllOf(Name, FChar::IsAlnum) && !Name.IsEmpty() && Name.Len() < 256)
+		{
+			Key.Bucket = CreateCacheBucket(Name);
+			return true;
+		}
+		return false;
+	};
+	if (!TrySetBucketName(FUTF8ToTCHAR(KeyObject["Bucket"_ASV].AsString()), Key))
+	{
+		return FOptionalCacheRecord();
+	}
+	Key.Hash = KeyObject["Hash"_ASV].AsHash();
+
+	FCacheRecordBuilder Builder = CreateCacheRecordBuilder(Key);
+
+	Builder.SetMeta(Package.GetObject()["Meta"_ASV].AsObject());
+
+	auto LoadPayload = [&Package](const FCbObjectView& PayloadObject)
+	{
+		const FPayloadId PayloadId(PayloadObject["Id"].AsObjectId().GetView());
+		if (PayloadId.IsNull())
+		{
+			return FPayload();
+		}
+
+		if (const FCbAttachment* Attachment = Package.FindAttachment(PayloadObject["CompressedHash"].AsHash()))
+		{
+			if (FCompressedBuffer Compressed = FCompressedBuffer::FromCompressed(Attachment->AsBinaryView()))
+			{
+				return FPayload(PayloadId, Compressed);
+			}
+		}
+		FIoHash RawHash = PayloadObject["RawHash"].AsHash();
+		uint64 RawSize = PayloadObject["RawSize"].AsUInt64(MAX_uint64);
+		if (!RawHash.IsZero() && RawSize != MAX_uint64)
+		{
+			return FPayload(PayloadId, RawHash, RawSize);
+		}
+		else
+		{
+			return FPayload();
+		}
+	};
+
+	if (FCbObjectView ValueObject = RecordObject["Value"_ASV].AsObjectView())
+	{
+		FPayload Value = LoadPayload(ValueObject);
+		if (!Value)
+		{
+			return FOptionalCacheRecord();
+		}
+		Builder.SetValue(Value);
+	}
+
+	for (FCbFieldView AttachmentField : RecordObject["Attachments"_ASV])
+	{
+		FPayload Attachment = LoadPayload(AttachmentField.AsObjectView());
+		if (!Attachment)
+		{
+			return FOptionalCacheRecord();
+		}
+		Builder.AddAttachment(Attachment);
+	}
+
+	return Builder.Build();
+}
+
 ///////////////////////////////////////////////////////////////////////////////////////////////////////////////////////
 
 } // UE::DerivedData::Private
diff --git a/Engine/Source/Developer/DerivedDataCache/Public/DerivedDataCache.h b/Engine/Source/Developer/DerivedDataCache/Public/DerivedDataCache.h
--- a/Engine/Source/Developer/DerivedDataCache/Public/DerivedDataCache.h
+++ b/Engine/Source/Developer/DerivedDataCache/Public/DerivedDataCache.h
@@ -8,12 +8,15 @@
 #include "Misc/EnumClassFlags.h"
 #include "Templates/Function.h"
 
+class FCbPackage;
+
 namespace UE::DerivedData
 {
 
 class FCacheBucket;
 class FCacheRecord;
 class FCacheRecordBuilder;
+class FOptionalCacheRecord;
 class FPayload;
 struct FCacheGetCompleteParams;
 struct FCacheGetPayloadCompleteParams;
@@ -110,6 +113,19 @@ public:
 	 * Create a cache record builder from a cache key.
 	 */
 	virtual FCacheRecordBuilder CreateRecord(const FCacheKey& Key) = 0;
+
+	/**
+	 * Save the cache record to a compact binary package.
+	 */
+	virtual FCbPackage SaveRecord(const FCacheRecord& Record) = 0;
+
+	/**
+	 * Load a cache record from compact binary.
+	 *
+	 * @param Package   A package saved from a cache record.
+	 * @return A valid cache record, or null on error.
+	 */
+	virtual FOptionalCacheRecord LoadRecord(const FCbPackage& Package) = 0;
 };
 
 ///////////////////////////////////////////////////////////////////////////////////////////////////////////////////////
