DDC: Minor fixes and tweaks discovered while integrating remote execution

#rnx

diff --git a/Engine/Source/Developer/DerivedDataCache/Private/DerivedDataBuildAction.cpp b/Engine/Source/Developer/DerivedDataCache/Private/DerivedDataBuildAction.cpp
--- a/Engine/Source/Developer/DerivedDataCache/Private/DerivedDataBuildAction.cpp
+++ b/Engine/Source/Developer/DerivedDataCache/Private/DerivedDataBuildAction.cpp
@@ -66,8 +66,8 @@ private:
 	{
 		const uint32 KeyHash = GetTypeHash(Key);
 		checkf(!Key.IsEmpty(), TEXT("Empty key used in action for build of '%s' by %s."), *Name, *Function);
-		checkf(!Inputs.ContainsByHash(KeyHash, Key), TEXT("Duplicate key '%.*s' used in action for "),
-			TEXT("build of '%s' by %s."), Key.Len(), Key.GetData(), *Name, *Function);
+		checkf(!Inputs.ContainsByHash(KeyHash, Key), TEXT("Duplicate key '%.*s' used in action ")
+			TEXT("for build of '%s' by %s."), Key.Len(), Key.GetData(), *Name, *Function);
 		Inputs.EmplaceByHash(KeyHash, Key, InputType(TInPlaceType<ValueType>(), Value));
 	}
 };
diff --git a/Engine/Source/Developer/DerivedDataCache/Private/DerivedDataBuildDefinition.cpp b/Engine/Source/Developer/DerivedDataCache/Private/DerivedDataBuildDefinition.cpp
--- a/Engine/Source/Developer/DerivedDataCache/Private/DerivedDataBuildDefinition.cpp
+++ b/Engine/Source/Developer/DerivedDataCache/Private/DerivedDataBuildDefinition.cpp
@@ -71,8 +71,8 @@ private:
 	{
 		const uint32 KeyHash = GetTypeHash(Key);
 		checkf(!Key.IsEmpty(), TEXT("Empty key used in definition for build of '%s' by %s."), *Name, *Function);
-		checkf(!Inputs.ContainsByHash(KeyHash, Key), TEXT("Duplicate key '%.*s' used in definition for "),
-			TEXT("build of '%s' by %s."), Key.Len(), Key.GetData(), *Name, *Function);
+		checkf(!Inputs.ContainsByHash(KeyHash, Key), TEXT("Duplicate key '%.*s' used in definition ")
+			TEXT("for build of '%s' by %s."), Key.Len(), Key.GetData(), *Name, *Function);
 		Inputs.EmplaceByHash(KeyHash, Key, InputType(TInPlaceType<ValueType>(), Forward<ArgType>(Value)));
 	}
 };
diff --git a/Engine/Source/Developer/DerivedDataCache/Public/DerivedDataBuildFunctionFactory.h b/Engine/Source/Developer/DerivedDataCache/Public/DerivedDataBuildFunctionFactory.h
--- a/Engine/Source/Developer/DerivedDataCache/Public/DerivedDataBuildFunctionFactory.h
+++ b/Engine/Source/Developer/DerivedDataCache/Public/DerivedDataBuildFunctionFactory.h
@@ -34,7 +34,7 @@ public:
  * static const TBuildFunctionFactory<TExampleFunction<FType>> ExampleFunctionFactory(Name, Version);
  */
 template <typename FunctionType>
-class TBuildFunctionFactory final : private IBuildFunctionFactory
+class TBuildFunctionFactory final : public IBuildFunctionFactory
 {
 	static_assert(sizeof(FunctionType) == sizeof(IBuildFunction), "IBuildFunction must be pure and maintain no state.");
 
@@ -51,12 +51,12 @@ public:
 		IModularFeatures::Get().UnregisterModularFeature(GetFeatureName(), this);
 	}
 
-private:
 	const IBuildFunction& GetFunction() const final
 	{
 		return Function;
 	}
 
+private:
 	const FunctionType Function;
 };
 
