DDC: Added FBuildInputs::IterateInputs

#rnx

diff --git a/Engine/Source/Developer/DerivedDataCache/Private/DerivedDataBuildInputs.cpp b/Engine/Source/Developer/DerivedDataCache/Private/DerivedDataBuildInputs.cpp
--- a/Engine/Source/Developer/DerivedDataCache/Private/DerivedDataBuildInputs.cpp
+++ b/Engine/Source/Developer/DerivedDataCache/Private/DerivedDataBuildInputs.cpp
@@ -26,9 +26,10 @@ public:
 	void AddInput(FStringView Key, const FCompressedBuffer& Buffer) final
 	{
 		const uint32 KeyHash = GetTypeHash(Key);
+		checkf(Buffer, TEXT("Null buffer used in input for build of '%s'."), *Name);
 		checkf(!Key.IsEmpty(), TEXT("Empty key used in input for build of '%s'."), *Name);
-		checkf(!Inputs.ContainsByHash(KeyHash, Key), TEXT("Duplicate key '%.*s' used in input for "),
-			TEXT("build of '%s'."), Key.Len(), Key.GetData(), *Name);
+		checkf(!Inputs.ContainsByHash(KeyHash, Key), TEXT("Duplicate key '%.*s' used in input for build of '%s'."),
+			Key.Len(), Key.GetData(), *Name);
 		Inputs.EmplaceByHash(KeyHash, Key, Buffer);
 	}
 
@@ -49,7 +50,8 @@ public:
 
 	FStringView GetName() const final { return Name; }
 
-	const FCompressedBuffer& GetInput(FStringView Key) const final;
+	const FCompressedBuffer& FindInput(FStringView Key) const final;
+	void IterateInputs(TFunctionRef<void (FStringView Key, const FCompressedBuffer& Buffer)> Visitor) const final;
 
 	inline void AddRef() const final
 	{
@@ -79,7 +81,7 @@ FBuildInputsInternal::FBuildInputsInternal(FBuildInputsBuilderInternal&& InputsB
 	Inputs.KeySort(TLess<>());
 }
 
-const FCompressedBuffer& FBuildInputsInternal::GetInput(FStringView Key) const
+const FCompressedBuffer& FBuildInputsInternal::FindInput(FStringView Key) const
 {
 	if (const FCompressedBuffer* Buffer = Inputs.FindByHash(GetTypeHash(Key), Key))
 	{
@@ -88,6 +90,14 @@ const FCompressedBuffer& FBuildInputsInternal::GetInput(FStringView Key) const
 	return FCompressedBuffer::Null;
 }
 
+void FBuildInputsInternal::IterateInputs(TFunctionRef<void (FStringView Key, const FCompressedBuffer& Buffer)> Visitor) const
+{
+	for (const TPair<FString, FCompressedBuffer>& Input : Inputs)
+	{
+		Input.ApplyAfter(Visitor);
+	}
+}
+
 ///////////////////////////////////////////////////////////////////////////////////////////////////////////////////////
 
 FBuildInputs FBuildInputsBuilderInternal::Build()
diff --git a/Engine/Source/Developer/DerivedDataCache/Public/DerivedDataBuildInputs.h b/Engine/Source/Developer/DerivedDataCache/Public/DerivedDataBuildInputs.h
--- a/Engine/Source/Developer/DerivedDataCache/Public/DerivedDataBuildInputs.h
+++ b/Engine/Source/Developer/DerivedDataCache/Public/DerivedDataBuildInputs.h
@@ -21,7 +21,8 @@ class IBuildInputsInternal
 public:
 	virtual ~IBuildInputsInternal() = default;
 	virtual FStringView GetName() const = 0;
-	virtual const FCompressedBuffer& GetInput(FStringView Key) const = 0;
+	virtual const FCompressedBuffer& FindInput(FStringView Key) const = 0;
+	virtual void IterateInputs(TFunctionRef<void (FStringView Key, const FCompressedBuffer& Buffer)> Visitor) const = 0;
 	virtual void AddRef() const = 0;
 	virtual void Release() const = 0;
 };
@@ -56,8 +57,14 @@ public:
 	/** Returns the name by which to identify the inputs for logging and profiling. */
 	inline FStringView GetName() const { return Inputs->GetName(); }
 
+	/** Finds an input by key, or a null buffer if not found. */
+	inline const FCompressedBuffer& FindInput(FStringView Key) const { return Inputs->FindInput(Key); }
+
 	/** Visits every input in order by key. */
-	inline const FCompressedBuffer& GetInput(FStringView Key) const { return Inputs->GetInput(Key); }
+	inline void IterateInputs(TFunctionRef<void (FStringView Key, const FCompressedBuffer& Buffer)> Visitor) const
+	{
+		Inputs->IterateInputs(Visitor);
+	}
 
 private:
 	friend class FOptionalBuildInputs;
