DDC: Added UE::DerivedData::FBuildWorkerRegistry

#rnx

diff --git a/Engine/Source/Developer/DerivedDataCache/Private/DerivedDataBuild.cpp b/Engine/Source/Developer/DerivedDataCache/Private/DerivedDataBuild.cpp
--- a/Engine/Source/Developer/DerivedDataCache/Private/DerivedDataBuild.cpp
+++ b/Engine/Source/Developer/DerivedDataCache/Private/DerivedDataBuild.cpp
@@ -10,6 +10,7 @@
 #include "DerivedDataBuildPrivate.h"
 #include "DerivedDataBuildScheduler.h"
 #include "DerivedDataBuildSession.h"
+#include "DerivedDataBuildWorkerRegistry.h"
 #include "DerivedDataCache.h"
 #include "Misc/Guid.h"
 
@@ -126,9 +127,15 @@ public:
 		return *FunctionRegistry;
 	}
 
+	IBuildWorkerRegistry& GetWorkerRegistry() const final
+	{
+		return *WorkerRegistry;
+	}
+
 private:
 	ICache& Cache;
 	TUniquePtr<IBuildFunctionRegistry> FunctionRegistry{CreateBuildFunctionRegistry()};
+	TUniquePtr<IBuildWorkerRegistry> WorkerRegistry{CreateBuildWorkerRegistry()};
 	TUniquePtr<IBuildScheduler> DefaultScheduler{CreateBuildScheduler()};
 	const FGuid Version{TEXT("ac0574e5-62bd-4c2e-84ec-f2efe48c0fef")};
 };
diff --git a/Engine/Source/Developer/DerivedDataCache/Private/DerivedDataBuildPrivate.h b/Engine/Source/Developer/DerivedDataCache/Private/DerivedDataBuildPrivate.h
--- a/Engine/Source/Developer/DerivedDataCache/Private/DerivedDataBuildPrivate.h
+++ b/Engine/Source/Developer/DerivedDataCache/Private/DerivedDataBuildPrivate.h
@@ -25,6 +25,7 @@ namespace UE::DerivedData { class IBuildFunctionRegistry; }
 namespace UE::DerivedData { class IBuildInputResolver; }
 namespace UE::DerivedData { class IBuildJob; }
 namespace UE::DerivedData { class IBuildScheduler; }
+namespace UE::DerivedData { class IBuildWorkerRegistry; }
 namespace UE::DerivedData { class ICache; }
 namespace UE::DerivedData { struct FBuildKey; }
 namespace UE::DerivedData { template <typename RequestType> class TRequest; }
@@ -40,6 +41,9 @@ IBuild* CreateBuild(ICache& Cache);
 // Implemented in DerivedDataBuildFunctionRegistry.cpp
 IBuildFunctionRegistry* CreateBuildFunctionRegistry();
 
+// Implemented in DerivedDataBuildWorkerRegistry.cpp
+IBuildWorkerRegistry* CreateBuildWorkerRegistry();
+
 // Implemented in DerivedDataBuildScheduler.cpp
 IBuildScheduler* CreateBuildScheduler();
 
diff --git a/Engine/Source/Developer/DerivedDataCache/Private/DerivedDataBuildWorkerRegistry.cpp b/Engine/Source/Developer/DerivedDataCache/Private/DerivedDataBuildWorkerRegistry.cpp
new file mode 100644
--- /dev/null
+++ b/Engine/Source/Developer/DerivedDataCache/Private/DerivedDataBuildWorkerRegistry.cpp
@@ -0,0 +1,274 @@
+// Copyright Epic Games, Inc. All Rights Reserved.
+
+#include "DerivedDataBuildWorkerRegistry.h"
+
+#include "Algo/Find.h"
+#include "Containers/StringView.h"
+#include "Containers/UnrealString.h"
+#include "DerivedDataBuildKey.h"
+#include "DerivedDataBuildPrivate.h"
+#include "DerivedDataBuildWorker.h"
+#include "DerivedDataRequest.h"
+#include "Features/IModularFeatures.h"
+#include "HAL/CriticalSection.h"
+#include "IO/IoHash.h"
+#include "Misc/Guid.h"
+#include "Misc/ScopeRWLock.h"
+#include "Templates/Function.h"
+#include "Templates/Tuple.h"
+
+namespace UE::DerivedData::Private
+{
+
+class FBuildWorkerInternal final : public FBuildWorker, public FBuildWorkerBuilder
+{
+public:
+	inline explicit FBuildWorkerInternal(IBuildWorkerFactory* InFactory)
+		: Factory(InFactory)
+	{
+	}
+
+	void Build();
+
+	inline FStringView GetName() const final { return WorkerName; }
+	inline FStringView GetPath() const final { return WorkerPath; }
+	inline FStringView GetHostPlatform() const final { return HostPlatform; }
+	inline FGuid GetBuildSystemVersion() const final { return BuildSystemVersion; }
+
+	FRequest GetFileData(TConstArrayView<FIoHash> RawHashes, EPriority Priority, FOnBuildWorkerFileDataComplete&& OnComplete) const final;
+
+	void IterateFunctions(TFunctionRef<void (FStringView WorkerName, const FGuid& Version)> Visitor) const final;
+	void IterateFiles(TFunctionRef<void (FStringView WorkerPath, const FIoHash& RawHash, uint64 RawSize)> Visitor) const final;
+	void IterateExecutables(TFunctionRef<void (FStringView WorkerPath, const FIoHash& RawHash, uint64 RawSize)> Visitor) const final;
+	void IterateEnvironment(TFunctionRef<void (FStringView WorkerName, FStringView Value)> Visitor) const final;
+
+private:
+	inline void SetName(FStringView Name) final { WorkerName = Name; }
+	inline void SetPath(FStringView Path) final { WorkerPath = Path; }
+	inline void SetHostPlatform(FStringView Name) final { HostPlatform = Name; }
+	inline void SetBuildSystemVersion(const FGuid& Version) final { BuildSystemVersion = Version; }
+
+	void AddFunction(FStringView Name, const FGuid& Version) final;
+	void AddFile(FStringView Path, const FIoHash& RawHash, uint64 RawSize) final;
+	void AddExecutable(FStringView Path, const FIoHash& RawHash, uint64 RawSize) final;
+	void SetEnvironment(FStringView Name, FStringView Value) final;
+
+private:
+	FString WorkerName;
+	FString WorkerPath;
+	FString HostPlatform;
+	FGuid BuildSystemVersion;
+	TArray<TTuple<FString, FGuid>> Functions;
+	TArray<TTuple<FString, FIoHash, uint64>> Files;
+	TArray<TTuple<FString, FIoHash, uint64>> Executables;
+	TArray<TTuple<FString, FString>> Environment;
+	IBuildWorkerFactory* Factory;
+};
+
+void FBuildWorkerInternal::Build()
+{
+	Functions.Sort();
+	Files.Sort();
+	Executables.Sort();
+	Environment.Sort();
+}
+
+FRequest FBuildWorkerInternal::GetFileData(TConstArrayView<FIoHash> RawHashes, EPriority Priority, FOnBuildWorkerFileDataComplete&& OnComplete) const
+{
+	return Factory->GetFileData(RawHashes, Priority, MoveTemp(OnComplete));
+}
+
+void FBuildWorkerInternal::IterateFunctions(TFunctionRef<void (FStringView Name, const FGuid& Version)> Visitor) const
+{
+	for (const TTuple<FString, FGuid>& Function : Functions)
+	{
+		Function.ApplyAfter(Visitor);
+	}
+}
+
+void FBuildWorkerInternal::IterateFiles(TFunctionRef<void (FStringView Path, const FIoHash& RawHash, uint64 RawSize)> Visitor) const
+{
+	for (const TTuple<FString, FIoHash, uint64>& File : Files)
+	{
+		File.ApplyAfter(Visitor);
+	}
+}
+
+void FBuildWorkerInternal::IterateExecutables(TFunctionRef<void (FStringView Path, const FIoHash& RawHash, uint64 RawSize)> Visitor) const
+{
+	for (const TTuple<FString, FIoHash, uint64>& Executable : Executables)
+	{
+		Executable.ApplyAfter(Visitor);
+	}
+}
+
+void FBuildWorkerInternal::IterateEnvironment(TFunctionRef<void (FStringView Name, FStringView Value)> Visitor) const
+{
+	for (const TTuple<FString, FString>& Variable : Environment)
+	{
+		Variable.ApplyAfter(Visitor);
+	}
+}
+
+void FBuildWorkerInternal::AddFunction(FStringView Name, const FGuid& Version)
+{
+	UE_CLOG(!Version.IsValid(), LogDerivedDataBuild, Error,
+		TEXT("Version of zero is not allowed in build function with the name %.*s in build worker '%.s'."),
+		Name.Len(), Name.GetData(), *WorkerName);
+	Functions.Emplace(Name, Version);
+}
+
+void FBuildWorkerInternal::AddFile(FStringView Path, const FIoHash& RawHash, uint64 RawSize)
+{
+	Files.Emplace(Path, RawHash, RawSize);
+}
+
+void FBuildWorkerInternal::AddExecutable(FStringView Path, const FIoHash& RawHash, uint64 RawSize)
+{
+	Executables.Emplace(Path, RawHash, RawSize);
+}
+
+void FBuildWorkerInternal::SetEnvironment(FStringView Name, FStringView Value)
+{
+	Environment.Emplace(Name, Value);
+}
+
+///////////////////////////////////////////////////////////////////////////////////////////////////////////////////////
+
+class FBuildWorkerRegistry final : public IBuildWorkerRegistry
+{
+public:
+	FBuildWorkerRegistry();
+	~FBuildWorkerRegistry();
+
+	FBuildWorker* FindWorker(
+		FStringView Function,
+		const FGuid& FunctionVersion,
+		const FGuid& BuildSystemVersion,
+		IBuildWorkerExecutor*& OutWorkerExecutor) const final;
+
+private:
+	void OnModularFeatureRegistered(const FName& Type, IModularFeature* ModularFeature);
+	void OnModularFeatureUnregistered(const FName& Type, IModularFeature* ModularFeature);
+
+	void AddWorker(IBuildWorkerFactory* Factory);
+	void RemoveWorker(IBuildWorkerFactory* Factory);
+
+private:
+	mutable FRWLock Lock;
+	IBuildWorkerExecutor* Executor = nullptr;
+	TMap<IBuildWorkerFactory*, TUniquePtr<FBuildWorker>> Workers;
+	TMultiMap<TTuple<FString, FGuid>, FBuildWorker*> Functions;
+};
+
+FBuildWorkerRegistry::FBuildWorkerRegistry()
+{
+	IModularFeatures& ModularFeatures = IModularFeatures::Get();
+	if (ModularFeatures.IsModularFeatureAvailable(IBuildWorkerExecutor::GetFeatureName()))
+	{
+		Executor = &ModularFeatures.GetModularFeature<IBuildWorkerExecutor>(IBuildWorkerExecutor::GetFeatureName());
+	}
+	for (IBuildWorkerFactory* Worker : ModularFeatures.GetModularFeatureImplementations<IBuildWorkerFactory>(IBuildWorkerFactory::GetFeatureName()))
+	{
+		AddWorker(Worker);
+	}
+	ModularFeatures.OnModularFeatureRegistered().AddRaw(this, &FBuildWorkerRegistry::OnModularFeatureRegistered);
+	ModularFeatures.OnModularFeatureUnregistered().AddRaw(this, &FBuildWorkerRegistry::OnModularFeatureUnregistered);
+}
+
+FBuildWorkerRegistry::~FBuildWorkerRegistry()
+{
+	IModularFeatures& ModularFeatures = IModularFeatures::Get();
+	ModularFeatures.OnModularFeatureUnregistered().RemoveAll(this);
+	ModularFeatures.OnModularFeatureRegistered().RemoveAll(this);
+}
+
+void FBuildWorkerRegistry::OnModularFeatureRegistered(const FName& Type, IModularFeature* ModularFeature)
+{
+	if (!Executor && Type == IBuildWorkerExecutor::GetFeatureName())
+	{
+		FWriteScopeLock WriteLock(Lock);
+		Executor = static_cast<IBuildWorkerExecutor*>(ModularFeature);
+	}
+	else if (Type == IBuildWorkerFactory::GetFeatureName())
+	{
+		AddWorker(static_cast<IBuildWorkerFactory*>(ModularFeature));
+	}
+}
+
+void FBuildWorkerRegistry::OnModularFeatureUnregistered(const FName& Type, IModularFeature* ModularFeature)
+{
+	if (Executor == ModularFeature && Type == IBuildWorkerExecutor::GetFeatureName())
+	{
+		IModularFeature* NextExecutor = nullptr;
+		IModularFeatures& ModularFeatures = IModularFeatures::Get();
+		if (ModularFeatures.IsModularFeatureAvailable(IBuildWorkerExecutor::GetFeatureName()))
+		{
+			NextExecutor = &ModularFeatures.GetModularFeature<IBuildWorkerExecutor>(IBuildWorkerExecutor::GetFeatureName());
+		}
+		FWriteScopeLock WriteLock(Lock);
+		Executor = static_cast<IBuildWorkerExecutor*>(NextExecutor);
+	}
+	else if (Type == IBuildWorkerFactory::GetFeatureName())
+	{
+		RemoveWorker(static_cast<IBuildWorkerFactory*>(ModularFeature));
+	}
+}
+
+void FBuildWorkerRegistry::AddWorker(IBuildWorkerFactory* Factory)
+{
+	TUniquePtr<FBuildWorkerInternal> Worker = MakeUnique<FBuildWorkerInternal>(Factory);
+	Factory->Build(*Worker);
+	Worker->Build();
+
+	FWriteScopeLock WriteLock(Lock);
+	Worker->IterateFunctions([this, Worker = Worker.Get()](FStringView Name, const FGuid& Version)
+	{
+		Functions.Emplace(TTuple<FString, FGuid>(Name, Version), Worker);
+	});
+	Workers.Emplace(Factory, MoveTemp(Worker));
+}
+
+void FBuildWorkerRegistry::RemoveWorker(IBuildWorkerFactory* Factory)
+{
+	FWriteScopeLock WriteLock(Lock);
+	TUniquePtr<FBuildWorker>& Worker = Workers.FindChecked(Factory);
+	Worker->IterateFunctions([this, Worker = Worker.Get()] (FStringView Name, const FGuid& Version)
+	{
+		Functions.Remove(TTuple<FString, FGuid>(Name, Version), Worker);
+	});
+	Workers.Remove(Factory);
+}
+
+FBuildWorker* FBuildWorkerRegistry::FindWorker(
+	FStringView Function,
+	const FGuid& FunctionVersion,
+	const FGuid& BuildSystemVersion,
+	IBuildWorkerExecutor*& OutWorkerExecutor) const
+{
+	FReadScopeLock ReadLock(Lock);
+	if (Executor)
+	{
+		TConstArrayView<FStringView> ExecutorHostPlatforms = Executor->GetHostPlatforms();
+		TArray<FBuildWorker*, TInlineAllocator<8>> FunctionWorkers;
+		Functions.MultiFind(TTuple<FString, FGuid>(Function, FunctionVersion), FunctionWorkers);
+		for (FBuildWorker* Worker : FunctionWorkers)
+		{
+			if (Worker->GetBuildSystemVersion() == BuildSystemVersion &&
+				Algo::Find(ExecutorHostPlatforms, Worker->GetHostPlatform()))
+			{
+				OutWorkerExecutor = Executor;
+				return Worker;
+			}
+		}
+	}
+	OutWorkerExecutor = nullptr;
+	return nullptr;
+}
+
+IBuildWorkerRegistry* CreateBuildWorkerRegistry()
+{
+	return new FBuildWorkerRegistry();
+}
+
+} // UE::DerivedData::Private
diff --git a/Engine/Source/Developer/DerivedDataCache/Public/DerivedDataBuild.h b/Engine/Source/Developer/DerivedDataCache/Public/DerivedDataBuild.h
--- a/Engine/Source/Developer/DerivedDataCache/Public/DerivedDataBuild.h
+++ b/Engine/Source/Developer/DerivedDataCache/Public/DerivedDataBuild.h
@@ -19,6 +19,7 @@ namespace UE::DerivedData { class FOptionalBuildOutput; }
 namespace UE::DerivedData { class IBuildFunctionRegistry; }
 namespace UE::DerivedData { class IBuildInputResolver; }
 namespace UE::DerivedData { class IBuildScheduler; }
+namespace UE::DerivedData { class IBuildWorkerRegistry; }
 
 namespace UE::DerivedData
 {
@@ -127,6 +128,11 @@ public:
 	 * Returns the build function registry used by the build system.
 	 */
 	virtual IBuildFunctionRegistry& GetFunctionRegistry() const = 0;
+
+	/**
+	 * Returns the build worker registry used by the build system.
+	 */
+	virtual IBuildWorkerRegistry& GetWorkerRegistry() const = 0;
 };
 
 } // UE::DerivedData
