DDC: Added remote execution to the build job

#rnx

diff --git a/Engine/Source/Developer/DerivedDataCache/Private/DerivedDataBuildJob.cpp b/Engine/Source/Developer/DerivedDataCache/Private/DerivedDataBuildJob.cpp
--- a/Engine/Source/Developer/DerivedDataCache/Private/DerivedDataBuildJob.cpp
+++ b/Engine/Source/Developer/DerivedDataCache/Private/DerivedDataBuildJob.cpp
@@ -2,6 +2,7 @@
 
 #include "DerivedDataBuildJob.h"
 
+#include "Algo/Find.h"
 #include "Compression/CompressedBuffer.h"
 #include "DerivedDataBuild.h"
 #include "DerivedDataBuildAction.h"
@@ -16,6 +17,8 @@
 #include "DerivedDataBuildPrivate.h"
 #include "DerivedDataBuildScheduler.h"
 #include "DerivedDataBuildSession.h"
+#include "DerivedDataBuildWorker.h"
+#include "DerivedDataBuildWorkerRegistry.h"
 #include "DerivedDataCache.h"
 #include "DerivedDataCacheRecord.h"
 #include "DerivedDataPayload.h"
@@ -53,14 +56,26 @@ enum class EBuildJobState : uint32
 	CacheQuery,
 	/** Wait for the cache query to finish. */
 	CacheQueryWait,
-	/** Resolve the definition input data. */
+	/** Try to execute remotely without loading additional inputs. */
+	ExecuteRemote,
+	/** Wait for the remote execution request to finish. */
+	ExecuteRemoteWait,
+	/** Resolve the input data that was missing from remote execution. */
+	ResolveRemoteInputs,
+	/** Wait for the input resolver to resolve the data. */
+	ResolveRemoteInputsWait,
+	/** Try to execute remotely after resolving missing inputs. */
+	ExecuteRemoteRetry,
+	/** Wait for the remote execution request to finish. */
+	ExecuteRemoteRetryWait,
+	/** Resolve any unresolved input data for the action. */
 	ResolveInputs,
 	/** Wait for the input resolver to resolve the data. */
 	ResolveInputsWait,
-	/** Execute the function to build the output. */
-	ExecuteBuild,
-	/** Wait for an async function to execute. */
-	ExecuteBuildWait,
+	/** Execute the function locally to build the output. */
+	ExecuteLocal,
+	/** Wait for the async function to execute locally. */
+	ExecuteLocalWait,
 	/** Store the output in the cache. */
 	CacheStore,
 	/** Wait for the cache store to finish. */
@@ -69,6 +84,12 @@ enum class EBuildJobState : uint32
 	Complete,
 };
 
+/** Returns the state after the given state. */
+static constexpr EBuildJobState GetNextState(EBuildJobState State)
+{
+	return State == EBuildJobState::Complete ? State : EBuildJobState(uint32(State) + 1);
+}
+
 class FBuildJob final : public IBuildJob
 {
 public:
@@ -108,27 +129,30 @@ private:
 	void Configure();
 	void EndJob();
 
+	void EnterResolveKey();
+	void EnterResolveAction();
+	void EnterCacheQuery();
+	void EnterResolveInputs();
+	void EnterExecuteRemote();
+	void EnterExecuteLocal();
+	void EnterCacheStore();
+
 	void BeginResolveKey() final;
 	void BeginResolveAction() final;
 	void BeginCacheQuery() final;
 	void BeginResolveInputs() final;
-	void BeginExecuteBuild() final;
+	void BeginExecuteRemote() final;
+	void BeginExecuteLocal() final;
 	void BeginCacheStore() final;
 
 	void EndResolveKey(FBuildDefinitionResolvedParams&& Params);
 	void EndResolveAction(FBuildInputResolvedParams&& Params);
 	void EndCacheQuery(FCacheGetCompleteParams&& Params);
 	void EndResolveInputs(FBuildInputResolvedParams&& Params);
-	void EndExecuteBuild();
+	void EndExecuteRemote(FBuildWorkerActionCompleteParams&& Params);
+	void EndExecuteLocal();
 	void EndCacheStore(FCachePutCompleteParams&& Params);
 
-	void EnterResolveKey();
-	void EnterResolveAction();
-	void EnterCacheQuery();
-	void EnterResolveInputs();
-	void EnterExecuteBuild();
-	void EnterCacheStore();
-
 	void SetDefinition(FBuildDefinition&& Definition);
 	void SetAction(FBuildAction&& Action);
 	void SetInputs(FBuildInputs&& Inputs);
@@ -167,6 +191,8 @@ private:
 	bool bIsCached{};
 	/** True if AdvanceToState is executing. */
 	bool bInAdvanceToState{};
+	/** True if remote execution was attempted, whether or not it failed. */
+	bool bTriedRemoteExecution{};
 
 	mutable std::atomic<uint32> ReferenceCount{0};
 
@@ -174,11 +200,11 @@ private:
 	FBuildKey DefinitionKey;
 	/** Available in [CacheQuery, Complete]. */
 	FBuildActionKey ActionKey;
-	/** Available in [ResolveAction, ExecuteBuild) for jobs created from a key or definition. */
+	/** Available in [ResolveAction, ExecuteLocal) for jobs created from a key or definition. */
 	FOptionalBuildDefinition Definition;
 	/** Available in [CacheQuery, CacheStore). */
 	FOptionalBuildAction Action;
-	/** Available in [ExecuteBuild, CacheStore) for jobs with non-constant inputs. */
+	/** Available in [ExecuteLocal, CacheStore) for jobs with non-constant inputs. */
 	FOptionalBuildInputs Inputs;
 	/** Available in [CacheStore, Complete). */
 	FOptionalBuildOutput Output;
@@ -199,9 +225,16 @@ private:
 	IBuildScheduler* Scheduler{};
 	/** Resolver for definitions and inputs. May be null for a job with nothing to resolve. */
 	IBuildInputResolver* InputResolver{};
+	/** Worker to use for remote execution. Available in [ExecuteRemote, ExecuteRemoteRetryWait]. */
+	FBuildWorker* Worker{};
+	/** Worker executor to use for remote execution. Available in [ExecuteRemote, ExecuteRemoteRetryWait]. */
+	IBuildWorkerExecutor* WorkerExecutor{};
 	/** Invoked exactly once when the output is complete or when the job fails. */
 	FOnBuildJobComplete OnComplete;
 
+	/** Keys for missing inputs. All inputs will be resolved when this is empty. */
+	TArray<FString> MissingInputs;
+
 	ICache& Cache;
 	IBuild& BuildSystem;
 };
@@ -212,20 +245,26 @@ const TCHAR* LexToString(EBuildJobState State)
 {
 	switch (State)
 	{
-	case EBuildJobState::ResolveKey:        return TEXT("ResolveKey");
-	case EBuildJobState::ResolveKeyWait:    return TEXT("ResolveKeyWait");
-	case EBuildJobState::ResolveAction:     return TEXT("ResolveAction");
-	case EBuildJobState::ResolveActionWait: return TEXT("ResolveActionWait");
-	case EBuildJobState::CacheQuery:        return TEXT("CacheQuery");
-	case EBuildJobState::CacheQueryWait:    return TEXT("CacheQueryWait");
-	case EBuildJobState::ResolveInputs:     return TEXT("ResolveInputs");
-	case EBuildJobState::ResolveInputsWait: return TEXT("ResolveInputsWait");
-	case EBuildJobState::ExecuteBuild:      return TEXT("ExecuteBuild");
-	case EBuildJobState::ExecuteBuildWait:  return TEXT("ExecuteBuildWait");
-	case EBuildJobState::CacheStore:        return TEXT("CacheStore");
-	case EBuildJobState::CacheStoreWait:    return TEXT("CacheStoreWait");
-	case EBuildJobState::Complete:          return TEXT("Complete");
-	default: checkNoEntry();                return TEXT("Unknown");
+	case EBuildJobState::ResolveKey:              return TEXT("ResolveKey");
+	case EBuildJobState::ResolveKeyWait:          return TEXT("ResolveKeyWait");
+	case EBuildJobState::ResolveAction:           return TEXT("ResolveAction");
+	case EBuildJobState::ResolveActionWait:       return TEXT("ResolveActionWait");
+	case EBuildJobState::CacheQuery:              return TEXT("CacheQuery");
+	case EBuildJobState::CacheQueryWait:          return TEXT("CacheQueryWait");
+	case EBuildJobState::ExecuteRemote:           return TEXT("ExecuteRemote");
+	case EBuildJobState::ExecuteRemoteWait:       return TEXT("ExecuteRemoteWait");
+	case EBuildJobState::ResolveRemoteInputs:     return TEXT("ResolveRemoteInputs");
+	case EBuildJobState::ResolveRemoteInputsWait: return TEXT("ResolveRemoteInputsWait");
+	case EBuildJobState::ExecuteRemoteRetry:      return TEXT("ExecuteRemoteRetry");
+	case EBuildJobState::ExecuteRemoteRetryWait:  return TEXT("ExecuteRemoteRetryWait");
+	case EBuildJobState::ResolveInputs:           return TEXT("ResolveInputs");
+	case EBuildJobState::ResolveInputsWait:       return TEXT("ResolveInputsWait");
+	case EBuildJobState::ExecuteLocal:            return TEXT("ExecuteLocal");
+	case EBuildJobState::ExecuteLocalWait:        return TEXT("ExecuteLocalWait");
+	case EBuildJobState::CacheStore:              return TEXT("CacheStore");
+	case EBuildJobState::CacheStoreWait:          return TEXT("CacheStoreWait");
+	case EBuildJobState::Complete:                return TEXT("Complete");
+	default: checkNoEntry();                      return TEXT("Unknown");
 	}
 }
 
@@ -335,7 +374,7 @@ void FBuildJob::Cancel()
 		bIsCanceled = true;
 		LocalRequest = WaitRequest;
 		LocalScheduler = Scheduler;
-		if (State == EBuildJobState::ExecuteBuildWait)
+		if (State == EBuildJobState::ExecuteLocalWait)
 		{
 			LocalContext = Context;
 		}
@@ -353,7 +392,7 @@ void FBuildJob::Cancel()
 		LocalRequest.Cancel();
 	}
 
-	// Cancel the async build, which invokes EndExecuteBuild if the build is not complete.
+	// Cancel the async build, which invokes EndExecuteLocal if the build is not complete.
 	if (LocalContext)
 	{
 		LocalContext->GetFunction().CancelAsyncBuild(*LocalContext);
@@ -583,7 +622,7 @@ void FBuildJob::Configure()
 	{
 		const FCacheKey CacheKey{Cache.CreateBucket(FunctionName), Action.Get().GetKey().Hash};
 		Context = new FBuildJobContext(*this, CacheKey, *Function, OutputBuilder, DefaultBuildPolicy,
-			[this] { EndExecuteBuild(); });
+			[this] { EndExecuteLocal(); });
 		Function->Configure(*Context);
 	}
 }
@@ -597,7 +636,7 @@ void FBuildJob::EnterCacheQuery()
 	if (!EnumHasAnyFlags(CachePolicy, ECachePolicy::Query) ||
 		EnumHasAnyFlags(BuildPolicy, EBuildPolicy::SkipCacheGet))
 	{
-		return AdvanceToState(EBuildJobState::ResolveInputs);
+		return AdvanceToState(EBuildJobState::ExecuteRemote);
 	}
 	Scheduler->DispatchCacheQuery(this);
 }
@@ -630,7 +669,7 @@ void FBuildJob::EndCacheQuery(FCacheGetCompleteParams&& Params)
 				SetOutputNoCheck(MoveTemp(CacheOutput).Get());
 			}
 		}
-		return AdvanceToState(EBuildJobState::ResolveInputs);
+		return AdvanceToState(EBuildJobState::ExecuteRemote);
 	}
 }
 
@@ -638,10 +677,27 @@ void FBuildJob::EndCacheQuery(FCacheGetCompleteParams&& Params)
 
 void FBuildJob::EnterResolveInputs()
 {
-	if (Inputs || !Action.Get().HasInputs())
+	// Identify missing inputs when resolving every input for the action.
+	if (State == EBuildJobState::ResolveInputs)
 	{
-		return AdvanceToState(EBuildJobState::ExecuteBuild);
+		MissingInputs.Reset();
+		Action.Get().IterateInputs([this](FStringView Key, const FIoHash& RawHash, uint64 RawSize)
+		{
+			if (!Inputs || Inputs.Get().FindInput(Key).IsNull())
+			{
+				MissingInputs.Emplace(Key);
+			}
+		});
+		if (MissingInputs.IsEmpty())
+		{
+			return AdvanceToState(EBuildJobState::ExecuteLocal);
+		}
 	}
+
+	checkf(!MissingInputs.IsEmpty(),
+		TEXT("Job is not expected to be in state %s without missing inputs for build of '%s' by %s"),
+		LexToString(State), *Name, *FunctionName);
+
 	if (!InputResolver)
 	{
 		return CompleteWithError(TEXT("Failed to resolve input data due to null input resolver."_SV));
@@ -657,9 +713,9 @@ void FBuildJob::BeginResolveInputs()
 {
 	if (CanExecuteState(EBuildJobState::ResolveInputs))
 	{
-		const FBuildDefinition LocalDefinition = MoveTemp(Definition).Get();
-		return AdvanceToState(EBuildJobState::ResolveInputsWait, InputResolver->GetData(LocalDefinition, Priority,
-			[this](FBuildInputResolvedParams&& Params) { EndResolveInputs(MoveTemp(Params)); }));
+		return AdvanceToState(GetNextState(State), InputResolver->GetData(Definition.Get(), Priority,
+			[this](FBuildInputResolvedParams&& Params) { EndResolveInputs(MoveTemp(Params)); },
+			[this](FStringView Key) { return !!Algo::Find(MissingInputs, Key); }));
 	}
 }
 
@@ -674,6 +730,13 @@ void FBuildJob::EndResolveInputs(FBuildInputResolvedParams&& Params)
 			{
 				Builder.AddInput(InputByKey.Key, InputByKey.Data);
 			}
+			if (Inputs)
+			{
+				Inputs.Get().IterateInputs([&Builder](FStringView Key, const FCompressedBuffer& Buffer)
+				{
+					Builder.AddInput(Key, Buffer);
+				});
+			}
 			SetInputs(Builder.Build());
 		}
 		else
@@ -685,17 +748,107 @@ void FBuildJob::EndResolveInputs(FBuildInputResolvedParams&& Params)
 
 void FBuildJob::SetInputs(FBuildInputs&& InInputs)
 {
-	checkf(Inputs.IsNull(), TEXT("Job already has inputs for build of '%s' by %s."), *Name, *FunctionName);
-	checkf(State == EBuildJobState::ResolveInputs || State == EBuildJobState::ResolveInputsWait,
-		TEXT("Job is not expecting inputs in state %s for build of '%s' by %s"),
-		LexToString(State), *Name, *FunctionName);
+	EBuildJobState ExecuteState;
+	switch (State)
+	{
+	case EBuildJobState::ResolveRemoteInputs:
+	case EBuildJobState::ResolveRemoteInputsWait:
+		ExecuteState = EBuildJobState::ExecuteRemoteRetry;
+		break;
+	case EBuildJobState::ResolveInputs:
+	case EBuildJobState::ResolveInputsWait:
+		ExecuteState = EBuildJobState::ExecuteLocal;
+		break;
+	default:
+		checkf(false, TEXT("Job is not expecting inputs in state %s for build of '%s' by %s"),
+			LexToString(State), *Name, *FunctionName);
+		return;
+	}
 	Inputs = MoveTemp(InInputs);
-	return AdvanceToState(EBuildJobState::ExecuteBuild);
+	return AdvanceToState(ExecuteState);
 }
 
 ///////////////////////////////////////////////////////////////////////////////////////////////////////////////////////
 
-void FBuildJob::EnterExecuteBuild()
+void FBuildJob::EnterExecuteRemote()
+{
+	const EBuildPolicy BuildPolicy = Context->GetBuildPolicy();
+	if (EnumHasAnyFlags(BuildPolicy, EBuildPolicy::SkipBuild))
+	{
+		return CompleteWithError(TEXT("Failed because build policy skipped building."_SV));
+	}
+	if (!EnumHasAnyFlags(BuildPolicy, EBuildPolicy::Remote))
+	{
+		return AdvanceToState(EBuildJobState::ResolveInputs);
+	}
+
+	if (!Worker || !WorkerExecutor)
+	{
+		Worker = BuildSystem.GetWorkerRegistry().FindWorker(
+			Action.Get().GetFunction(),
+			Action.Get().GetFunctionVersion(),
+			Action.Get().GetBuildSystemVersion(),
+			WorkerExecutor);
+	}
+
+	if (!Worker || !WorkerExecutor)
+	{
+		return AdvanceToState(EBuildJobState::ResolveInputs);
+	}
+
+	Scheduler->DispatchExecuteRemote(this);
+}
+
+void FBuildJob::BeginExecuteRemote()
+{
+	if (CanExecuteState(EBuildJobState::ExecuteRemote))
+	{
+		checkf(Worker && WorkerExecutor, TEXT("Job requires a worker in state %s for build of '%s' by %s"),
+			LexToString(State), *Name, *FunctionName);
+		bTriedRemoteExecution = true;
+		const EBuildPolicy BuildPolicy = Context->GetBuildPolicy();
+		AdvanceToState(GetNextState(State), WorkerExecutor->BuildAction(Action.Get(), Inputs, *Worker, BuildPolicy, Priority,
+			[this](FBuildWorkerActionCompleteParams&& Params) { EndExecuteRemote(MoveTemp(Params)); }));
+	}
+}
+
+void FBuildJob::EndExecuteRemote(FBuildWorkerActionCompleteParams&& Params)
+{
+	if (CanExecuteState(EBuildJobState::ExecuteRemoteWait))
+	{
+		if (Params.Output)
+		{
+			SetOutputNoCheck(MoveTemp(Params.Output).Get());
+		}
+		else
+		{
+			switch (State)
+			{
+			case EBuildJobState::ExecuteRemote:
+			case EBuildJobState::ExecuteRemoteWait:
+				MissingInputs.Reset(Params.MissingInputs.Num());
+				for (const FStringView& Key : Params.MissingInputs)
+				{
+					MissingInputs.Emplace(Key);
+				}
+				if (!MissingInputs.IsEmpty())
+				{
+					MissingInputs.Sort();
+					return AdvanceToState(EBuildJobState::ResolveRemoteInputs);
+				}
+				[[fallthrough]];
+			case EBuildJobState::ExecuteRemoteRetry:
+			case EBuildJobState::ExecuteRemoteRetryWait:
+			default:
+				return AdvanceToState(EBuildJobState::ResolveInputs);
+			}
+		}
+	}
+}
+
+///////////////////////////////////////////////////////////////////////////////////////////////////////////////////////
+
+void FBuildJob::EnterExecuteLocal()
 {
 	const EBuildPolicy BuildPolicy = Context->GetBuildPolicy();
 	if (EnumHasAnyFlags(BuildPolicy, EBuildPolicy::SkipBuild))
@@ -704,7 +857,16 @@ void FBuildJob::EnterExecuteBuild()
 	}
 	if (!EnumHasAnyFlags(BuildPolicy, EBuildPolicy::Local))
 	{
-		return CompleteWithError(TEXT("Failed because build policy does not allow local execution."_SV));
+		if (bTriedRemoteExecution)
+		{
+			return CompleteWithError(TEXT("Failed because build policy does not allow local execution, ")
+				TEXT("and remote execution failed to build."_SV));
+		}
+		else
+		{
+			return CompleteWithError(TEXT("Failed because build policy does not allow local execution, ")
+				TEXT("and remote execution was not available."_SV));
+		}
 	}
 	Action.Get().IterateConstants([this](FStringView Key, FCbObject&& Value)
 	{
@@ -713,7 +875,7 @@ void FBuildJob::EnterExecuteBuild()
 	bool bMatchedInputs = true;
 	Action.Get().IterateInputs([this, &bMatchedInputs](FStringView Key, const FIoHash& RawHash, uint64 RawSize)
 	{
-		const FCompressedBuffer& Buffer = Inputs.Get().GetInput(Key);
+		const FCompressedBuffer& Buffer = Inputs.Get().FindInput(Key);
 		if (Buffer.GetRawHash() == RawHash && Buffer.GetRawSize() == RawSize)
 		{
 			Context->AddInput(Key, Buffer);
@@ -730,18 +892,18 @@ void FBuildJob::EnterExecuteBuild()
 	Inputs.Reset();
 	if (bMatchedInputs)
 	{
-		Scheduler->DispatchExecuteBuild(this);
+		Scheduler->DispatchExecuteLocal(this);
 	}
 }
 
-void FBuildJob::BeginExecuteBuild()
+void FBuildJob::BeginExecuteLocal()
 {
-	if (CanExecuteState(EBuildJobState::ExecuteBuild))
+	if (CanExecuteState(EBuildJobState::ExecuteLocal))
 	{
 		Context->GetFunction().Build(*Context);
 		if (Context->IsAsyncBuild())
 		{
-			return AdvanceToState(EBuildJobState::ExecuteBuildWait);
+			return AdvanceToState(EBuildJobState::ExecuteLocalWait);
 		}
 		else
 		{
@@ -750,9 +912,9 @@ void FBuildJob::BeginExecuteBuild()
 	}
 }
 
-void FBuildJob::EndExecuteBuild()
+void FBuildJob::EndExecuteLocal()
 {
-	if (CanExecuteState(EBuildJobState::ExecuteBuildWait))
+	if (CanExecuteState(EBuildJobState::ExecuteLocalWait))
 	{
 		SetOutputNoCheck(OutputBuilder.Build());
 	}
@@ -762,7 +924,8 @@ void FBuildJob::SetOutput(const FBuildOutput& InOutput)
 {
 	checkf(State == EBuildJobState::ResolveKey || State == EBuildJobState::ResolveAction ||
 		State == EBuildJobState::CacheQuery || State == EBuildJobState::ResolveInputs ||
-		State == EBuildJobState::ExecuteBuild || State == EBuildJobState::ExecuteBuildWait,
+		State == EBuildJobState::ExecuteRemote || State == EBuildJobState::ExecuteRemoteRetry ||
+		State == EBuildJobState::ExecuteLocal,
 		TEXT("Job is not expecting an output in state %s for build of '%s' by %s"),
 		LexToString(State), *Name, *FunctionName);
 	SetOutputNoCheck(FBuildOutput(InOutput));
@@ -861,20 +1024,35 @@ bool FBuildJob::CanExecuteState(EBuildJobState RequestedState)
 		AdvanceToState(EBuildJobState::Complete);
 		return false;
 	}
-	const EBuildJobState LocalState = State;
+
+	EBuildJobState LocalState = State;
+
+	// Handle states that execute multiple times under a different name.
+	switch (LocalState)
+	{
+	case EBuildJobState::ResolveRemoteInputs:     LocalState = EBuildJobState::ResolveInputs; break;
+	case EBuildJobState::ResolveRemoteInputsWait: LocalState = EBuildJobState::ResolveInputsWait; break;
+	case EBuildJobState::ExecuteRemoteRetry:      LocalState = EBuildJobState::ExecuteRemote; break;
+	case EBuildJobState::ExecuteRemoteRetryWait:  LocalState = EBuildJobState::ExecuteRemoteWait; break;
+	}
+
 	if (LocalState == RequestedState)
 	{
 		return true;
 	}
+
+	// Handle wait states as the base state because the async operation can finish before advancing to the wait state.
 	switch (RequestedState)
 	{
-	case EBuildJobState::ResolveKeyWait:    return LocalState == EBuildJobState::ResolveKey;
-	case EBuildJobState::ResolveActionWait: return LocalState == EBuildJobState::ResolveAction;
-	case EBuildJobState::CacheQueryWait:    return LocalState == EBuildJobState::CacheQuery;
-	case EBuildJobState::ResolveInputsWait: return LocalState == EBuildJobState::ResolveInputs;
-	case EBuildJobState::ExecuteBuildWait:  return LocalState == EBuildJobState::ExecuteBuild;
-	case EBuildJobState::CacheStoreWait:    return LocalState == EBuildJobState::CacheStore;
+	case EBuildJobState::ResolveKeyWait:     return LocalState == EBuildJobState::ResolveKey;
+	case EBuildJobState::ResolveActionWait:  return LocalState == EBuildJobState::ResolveAction;
+	case EBuildJobState::CacheQueryWait:     return LocalState == EBuildJobState::CacheQuery;
+	case EBuildJobState::ResolveInputsWait:  return LocalState == EBuildJobState::ResolveInputs;
+	case EBuildJobState::ExecuteRemoteWait:  return LocalState == EBuildJobState::ExecuteRemote;
+	case EBuildJobState::ExecuteLocalWait:   return LocalState == EBuildJobState::ExecuteLocal;
+	case EBuildJobState::CacheStoreWait:     return LocalState == EBuildJobState::CacheStore;
 	}
+
 	return false;
 }
 
@@ -941,16 +1119,25 @@ void FBuildJob::ExecuteTransition(EBuildJobState OldState, EBuildJobState NewSta
 	{
 		Configure();
 	}
+	if (OldState <= EBuildJobState::ExecuteRemoteRetryWait && EBuildJobState::ExecuteRemoteRetryWait < NewState)
+	{
+		Worker = nullptr;
+		WorkerExecutor = nullptr;
+	}
 	if (OldState < EBuildJobState::ResolveInputsWait && EBuildJobState::ResolveInputsWait <= NewState)
 	{
 		Definition.Reset();
 	}
-	if (OldState < EBuildJobState::ExecuteBuildWait && EBuildJobState::ExecuteBuildWait <= NewState)
+	if (OldState <= EBuildJobState::ResolveInputsWait && EBuildJobState::ResolveInputsWait < NewState)
+	{
+		MissingInputs.Empty();
+	}
+	if (OldState < EBuildJobState::ExecuteLocalWait && EBuildJobState::ExecuteLocalWait <= NewState)
 	{
 		Action.Reset();
 		Inputs.Reset();
 	}
-	if (OldState <= EBuildJobState::ExecuteBuildWait && EBuildJobState::ExecuteBuildWait < NewState && !Output)
+	if (OldState <= EBuildJobState::ExecuteLocalWait && EBuildJobState::ExecuteLocalWait < NewState && !Output)
 	{
 		SetOutputNoCheck(OutputBuilder.Build());
 	}
@@ -966,13 +1153,16 @@ void FBuildJob::ExecuteTransition(EBuildJobState OldState, EBuildJobState NewSta
 
 	switch (NewState)
 	{
-	case EBuildJobState::ResolveKey:    return EnterResolveKey();
-	case EBuildJobState::ResolveAction: return EnterResolveAction();
-	case EBuildJobState::CacheQuery:    return EnterCacheQuery();
-	case EBuildJobState::ResolveInputs: return EnterResolveInputs();
-	case EBuildJobState::ExecuteBuild:  return EnterExecuteBuild();
-	case EBuildJobState::CacheStore:    return EnterCacheStore();
-	case EBuildJobState::Complete:      return EndJob();
+	case EBuildJobState::ResolveKey:          return EnterResolveKey();
+	case EBuildJobState::ResolveAction:       return EnterResolveAction();
+	case EBuildJobState::CacheQuery:          return EnterCacheQuery();
+	case EBuildJobState::ExecuteRemote:       return EnterExecuteRemote();
+	case EBuildJobState::ResolveRemoteInputs: return EnterResolveInputs();
+	case EBuildJobState::ExecuteRemoteRetry:  return EnterExecuteRemote();
+	case EBuildJobState::ResolveInputs:       return EnterResolveInputs();
+	case EBuildJobState::ExecuteLocal:        return EnterExecuteLocal();
+	case EBuildJobState::CacheStore:          return EnterCacheStore();
+	case EBuildJobState::Complete:            return EndJob();
 	}
 }
 
diff --git a/Engine/Source/Developer/DerivedDataCache/Public/DerivedDataBuildJob.h b/Engine/Source/Developer/DerivedDataCache/Public/DerivedDataBuildJob.h
--- a/Engine/Source/Developer/DerivedDataCache/Public/DerivedDataBuildJob.h
+++ b/Engine/Source/Developer/DerivedDataCache/Public/DerivedDataBuildJob.h
@@ -58,7 +58,8 @@ public:
 	virtual void BeginResolveKey() = 0;
 	virtual void BeginResolveAction() = 0;
 	virtual void BeginResolveInputs() = 0;
-	virtual void BeginExecuteBuild() = 0;
+	virtual void BeginExecuteRemote() = 0;
+	virtual void BeginExecuteLocal() = 0;
 
 	/** Called by the scheduler if it has cached output compatible with the build policy. */
 	virtual void SetOutput(const FBuildOutput& Output) = 0;
diff --git a/Engine/Source/Developer/DerivedDataCache/Public/DerivedDataBuildScheduler.h b/Engine/Source/Developer/DerivedDataCache/Public/DerivedDataBuildScheduler.h
--- a/Engine/Source/Developer/DerivedDataCache/Public/DerivedDataBuildScheduler.h
+++ b/Engine/Source/Developer/DerivedDataCache/Public/DerivedDataBuildScheduler.h
@@ -38,7 +38,8 @@ public:
 	virtual void DispatchResolveKey(IBuildJob* Job) { Job->BeginResolveKey(); }
 	virtual void DispatchResolveAction(IBuildJob* Job) { Job->BeginResolveAction(); }
 	virtual void DispatchResolveInputs(IBuildJob* Job) { Job->BeginResolveInputs(); }
-	virtual void DispatchExecuteBuild(IBuildJob* Job) { Job->BeginExecuteBuild(); }
+	virtual void DispatchExecuteRemote(IBuildJob* Job) { Job->BeginExecuteRemote(); }
+	virtual void DispatchExecuteLocal(IBuildJob* Job) { Job->BeginExecuteLocal(); }
 
 	/** Called when the output of the job is available. Always called once, and always between BeginJob and EndJob. */
 	virtual void CompleteJob(const FBuildJobCompleteParams& Params) {}
