DDC: Added the missing file from 16492343

#rnx

diff --git a/Engine/Source/Developer/DerivedDataCache/Public/DerivedDataBuildWorkerRegistry.h b/Engine/Source/Developer/DerivedDataCache/Public/DerivedDataBuildWorkerRegistry.h
--- a/Engine/Source/Developer/DerivedDataCache/Public/DerivedDataBuildWorkerRegistry.h
+++ b/Engine/Source/Developer/DerivedDataCache/Public/DerivedDataBuildWorkerRegistry.h
@@ -18,9 +18,20 @@ class IBuildWorkerRegistry
 public:
 	virtual ~IBuildWorkerRegistry() = default;
 
-	virtual IBuildWorkerExecutor* GetWorkerExecutor() const = 0;
-
-	virtual FBuildWorker* FindWorker(FStringView Function, const FGuid& FunctionVersion) const = 0;
+	/**
+	 * Finds a build worker that can execute the function at the version.
+	 *
+	 * @param Function            The function to find a worker for.
+	 * @param FunctionVersion     The version required for the function.
+	 * @param BuildSystemVersion  The version required for the build system.
+	 * @param OutWorkerExecutor   The executor to use to execute the worker.
+	 * @returns A build worker and executor if a compatible pair was found, or null for both.
+	 */
+	virtual FBuildWorker* FindWorker(
+		FStringView Function,
+		const FGuid& FunctionVersion,
+		const FGuid& BuildSystemVersion,
+		IBuildWorkerExecutor*& OutWorkerExecutor) const = 0;
 };
 
 } // UE::DerivedData
