DDC: Fixed an incorrect comment tag in DerivedDataBuildWorkerRegistry.h

#rnx

diff --git a/Engine/Source/Developer/DerivedDataCache/Public/DerivedDataBuildWorkerRegistry.h b/Engine/Source/Developer/DerivedDataCache/Public/DerivedDataBuildWorkerRegistry.h
--- a/Engine/Source/Developer/DerivedDataCache/Public/DerivedDataBuildWorkerRegistry.h
+++ b/Engine/Source/Developer/DerivedDataCache/Public/DerivedDataBuildWorkerRegistry.h
@@ -13,6 +13,9 @@ namespace UE::DerivedData { class IBuildWorkerExecutor; }
 namespace UE::DerivedData
 {
 
+/**
+ * A build worker registry maintains a collection of build workers.
+ */
 class IBuildWorkerRegistry
 {
 public:
@@ -25,7 +28,7 @@ public:
 	 * @param FunctionVersion     The version required for the function.
 	 * @param BuildSystemVersion  The version required for the build system.
 	 * @param OutWorkerExecutor   The executor to use to execute the worker.
-	 * @returns A build worker and executor if a compatible pair was found, or null for both.
+	 * @return A build worker and executor if a compatible pair was found, or null for both.
 	 */
 	virtual FBuildWorker* FindWorker(
 		FStringView Function,
