DDC: Accept an input resolver for every variation of creating a build job

diff --git a/Engine/Source/Developer/DerivedDataCache/Private/DerivedDataBuildJob.cpp b/Engine/Source/Developer/DerivedDataCache/Private/DerivedDataBuildJob.cpp
--- a/Engine/Source/Developer/DerivedDataCache/Private/DerivedDataBuildJob.cpp
+++ b/Engine/Source/Developer/DerivedDataCache/Private/DerivedDataBuildJob.cpp
@@ -94,11 +94,11 @@ class FBuildJob final : public IBuildJob
 {
 public:
 	/** Resolve the key to a definition, then build like the definition constructor. */
-	FBuildJob(ICache& Cache, IBuild& BuildSystem, const FBuildKey& Key, IBuildInputResolver* InputResolver);
-	/** Resolve the definition to an action, query the cache, resolve inputs, then build like the action constructor. */
-	FBuildJob(ICache& Cache, IBuild& BuildSystem, const FBuildDefinition& Definition, IBuildInputResolver* InputResolver);
-	/** Execute the build function to build the output, store to the cache, then complete the build. */
-	FBuildJob(ICache& Cache, IBuild& BuildSystem, const FBuildAction& Action, const FOptionalBuildInputs& Inputs);
+	FBuildJob(ICache& Cache, IBuild& BuildSystem, IBuildInputResolver* InputResolver, const FBuildKey& Key);
+	/** Resolve the definition to an action, then build like the action constructor. */
+	FBuildJob(ICache& Cache, IBuild& BuildSystem, IBuildInputResolver* InputResolver, const FBuildDefinition& Definition);
+	/** Query the cache, attempt remote execution, resolve inputs, fall back to local execution, store to the cache. */
+	FBuildJob(ICache& Cache, IBuild& BuildSystem, IBuildInputResolver* InputResolver, const FBuildAction& Action, const FOptionalBuildInputs& Inputs);
 
 	/** Destroy the job, which must be complete or not started. */
 	~FBuildJob();
@@ -273,8 +273,8 @@ const TCHAR* LexToString(EBuildJobState State)
 FBuildJob::FBuildJob(
 	ICache& InCache,
 	IBuild& InBuildSystem,
-	const FBuildKey& InKey,
-	IBuildInputResolver* InInputResolver)
+	IBuildInputResolver* InInputResolver,
+	const FBuildKey& InKey)
 	: Name(WriteToString<64>(TEXT("Resolve: "_SV), InKey))
 	, DefinitionKey(InKey)
 	, OutputBuilder(InBuildSystem.CreateOutput(Name, FunctionName))
@@ -287,8 +287,8 @@ FBuildJob::FBuildJob(
 FBuildJob::FBuildJob(
 	ICache& InCache,
 	IBuild& InBuildSystem,
-	const FBuildDefinition& InDefinition,
-	IBuildInputResolver* InInputResolver)
+	IBuildInputResolver* InInputResolver,
+	const FBuildDefinition& InDefinition)
 	: Name(InDefinition.GetName())
 	, FunctionName(InDefinition.GetFunction())
 	, DefinitionKey(InDefinition.GetKey())
@@ -303,6 +303,7 @@ FBuildJob::FBuildJob(
 FBuildJob::FBuildJob(
 	ICache& InCache,
 	IBuild& InBuildSystem,
+	IBuildInputResolver* InInputResolver,
 	const FBuildAction& InAction,
 	const FOptionalBuildInputs& InInputs)
 	: Name(InAction.GetName())
@@ -311,6 +312,7 @@ FBuildJob::FBuildJob(
 	, Action(InAction)
 	, Inputs(InInputs)
 	, OutputBuilder(InBuildSystem.CreateOutput(Name, FunctionName))
+	, InputResolver(InInputResolver)
 	, Cache(InCache)
 	, BuildSystem(InBuildSystem)
 {
@@ -1168,19 +1170,19 @@ void FBuildJob::ExecuteTransition(EBuildJobState OldState, EBuildJobState NewSta
 
 ///////////////////////////////////////////////////////////////////////////////////////////////////////////////////////
 
-TRequest<IBuildJob> CreateBuildJob(ICache& Cache, IBuild& BuildSystem, const FBuildKey& Key, IBuildInputResolver* InputResolver)
+TRequest<IBuildJob> CreateBuildJob(ICache& Cache, IBuild& BuildSystem, IBuildInputResolver* InputResolver, const FBuildKey& Key)
 {
-	return TRequest(new FBuildJob(Cache, BuildSystem, Key, InputResolver));
+	return TRequest(new FBuildJob(Cache, BuildSystem, InputResolver, Key));
 }
 
-TRequest<IBuildJob> CreateBuildJob(ICache& Cache, IBuild& BuildSystem, const FBuildDefinition& Definition, IBuildInputResolver* InputResolver)
+TRequest<IBuildJob> CreateBuildJob(ICache& Cache, IBuild& BuildSystem, IBuildInputResolver* InputResolver, const FBuildDefinition& Definition)
 {
-	return TRequest(new FBuildJob(Cache, BuildSystem, Definition, InputResolver));
+	return TRequest(new FBuildJob(Cache, BuildSystem, InputResolver, Definition));
 }
 
-TRequest<IBuildJob> CreateBuildJob(ICache& Cache, IBuild& BuildSystem, const FBuildAction& Action, const FOptionalBuildInputs& Inputs)
+TRequest<IBuildJob> CreateBuildJob(ICache& Cache, IBuild& BuildSystem, IBuildInputResolver* InputResolver, const FBuildAction& Action, const FOptionalBuildInputs& Inputs)
 {
-	return TRequest(new FBuildJob(Cache, BuildSystem, Action, Inputs));
+	return TRequest(new FBuildJob(Cache, BuildSystem, InputResolver, Action, Inputs));
 }
 
 } // UE::DerivedData::Private
diff --git a/Engine/Source/Developer/DerivedDataCache/Private/DerivedDataBuildPrivate.h b/Engine/Source/Developer/DerivedDataCache/Private/DerivedDataBuildPrivate.h
--- a/Engine/Source/Developer/DerivedDataCache/Private/DerivedDataBuildPrivate.h
+++ b/Engine/Source/Developer/DerivedDataCache/Private/DerivedDataBuildPrivate.h
@@ -72,8 +72,8 @@ FBuildSession CreateBuildSession(
 	IBuildInputResolver* InputResolver);
 
 // Implemented in DerivedDataBuildJob.cpp
-TRequest<IBuildJob> CreateBuildJob(ICache& Cache, IBuild& BuildSystem, const FBuildKey& Key, IBuildInputResolver* InputResolver);
-TRequest<IBuildJob> CreateBuildJob(ICache& Cache, IBuild& BuildSystem, const FBuildDefinition& Definition, IBuildInputResolver* InputResolver);
-TRequest<IBuildJob> CreateBuildJob(ICache& Cache, IBuild& BuildSystem, const FBuildAction& Action, const FOptionalBuildInputs& Inputs);
+TRequest<IBuildJob> CreateBuildJob(ICache& Cache, IBuild& BuildSystem, IBuildInputResolver* InputResolver, const FBuildKey& Key);
+TRequest<IBuildJob> CreateBuildJob(ICache& Cache, IBuild& BuildSystem, IBuildInputResolver* InputResolver, const FBuildDefinition& Definition);
+TRequest<IBuildJob> CreateBuildJob(ICache& Cache, IBuild& BuildSystem, IBuildInputResolver* InputResolver, const FBuildAction& Action, const FOptionalBuildInputs& Inputs);
 
 } // UE::DerivedData::Private
diff --git a/Engine/Source/Developer/DerivedDataCache/Private/DerivedDataBuildSession.cpp b/Engine/Source/Developer/DerivedDataCache/Private/DerivedDataBuildSession.cpp
--- a/Engine/Source/Developer/DerivedDataCache/Private/DerivedDataBuildSession.cpp
+++ b/Engine/Source/Developer/DerivedDataCache/Private/DerivedDataBuildSession.cpp
@@ -65,7 +65,7 @@ FRequest FBuildSessionInternal::Build(
 	EPriority Priority,
 	FOnBuildComplete&& OnComplete)
 {
-	TRequest Job(CreateBuildJob(Cache, BuildSystem, Definition, InputResolver));
+	TRequest Job(CreateBuildJob(Cache, BuildSystem, InputResolver, Definition));
 	FOnBuildJobComplete OnJobComplete;
 	if (OnComplete)
 	{
@@ -85,7 +85,7 @@ FRequest FBuildSessionInternal::BuildAction(
 	EPriority Priority,
 	FOnBuildActionComplete&& OnComplete)
 {
-	TRequest Job(CreateBuildJob(Cache, BuildSystem, Action, Inputs));
+	TRequest Job(CreateBuildJob(Cache, BuildSystem, InputResolver, Action, Inputs));
 	FOnBuildJobComplete OnJobComplete;
 	if (OnComplete)
 	{
@@ -105,7 +105,7 @@ FRequest FBuildSessionInternal::BuildPayload(
 	FOnBuildPayloadComplete&& OnComplete)
 {
 	// This requests the entire output to get one payload. It will be optimized later to request only one payload.
-	TRequest Job(CreateBuildJob(Cache, BuildSystem, PayloadKey.BuildKey, InputResolver));
+	TRequest Job(CreateBuildJob(Cache, BuildSystem, InputResolver, PayloadKey.BuildKey));
 	FOnBuildJobComplete OnJobComplete;
 	if (OnComplete)
 	{
