DDC: Renamed the input resolver states in the build job for clarity

diff --git a/Engine/Source/Developer/DerivedDataCache/Private/DerivedDataBuildJob.cpp b/Engine/Source/Developer/DerivedDataCache/Private/DerivedDataBuildJob.cpp
--- a/Engine/Source/Developer/DerivedDataCache/Private/DerivedDataBuildJob.cpp
+++ b/Engine/Source/Developer/DerivedDataCache/Private/DerivedDataBuildJob.cpp
@@ -48,10 +48,10 @@ enum class EBuildJobState : uint32
 	ResolveKey,
 	/** Wait for the input resolver to resolve the key. */
 	ResolveKeyWait,
-	/** Resolve the definition to an action. */
-	ResolveAction,
+	/** Resolve each input to its raw hash and raw size. */
+	ResolveInputMeta,
 	/** Wait for the input resolver to resolve the metadata. */
-	ResolveActionWait,
+	ResolveInputMetaWait,
 	/** Query the cache for the action. */
 	CacheQuery,
 	/** Wait for the cache query to finish. */
@@ -60,21 +60,21 @@ enum class EBuildJobState : uint32
 	ExecuteRemote,
 	/** Wait for the remote execution request to finish. */
 	ExecuteRemoteWait,
-	/** Resolve the input data that was missing from remote execution. */
-	ResolveRemoteInputs,
+	/** Resolve each input that was missing from remote execution. */
+	ResolveRemoteInputData,
 	/** Wait for the input resolver to resolve the data. */
-	ResolveRemoteInputsWait,
+	ResolveRemoteInputDataWait,
 	/** Try to execute remotely after resolving missing inputs. */
 	ExecuteRemoteRetry,
 	/** Wait for the remote execution request to finish. */
 	ExecuteRemoteRetryWait,
-	/** Resolve any unresolved input data for the action. */
-	ResolveInputs,
+	/** Resolve each input that is not yet resolved. */
+	ResolveInputData,
 	/** Wait for the input resolver to resolve the data. */
-	ResolveInputsWait,
+	ResolveInputDataWait,
 	/** Execute the function locally to build the output. */
 	ExecuteLocal,
-	/** Wait for the async function to execute locally. */
+	/** Wait for the async function to finish executing locally. */
 	ExecuteLocalWait,
 	/** Store the output in the cache. */
 	CacheStore,
@@ -130,25 +130,25 @@ private:
 	void EndJob();
 
 	void EnterResolveKey();
-	void EnterResolveAction();
+	void EnterResolveInputMeta();
 	void EnterCacheQuery();
-	void EnterResolveInputs();
+	void EnterResolveInputData();
 	void EnterExecuteRemote();
 	void EnterExecuteLocal();
 	void EnterCacheStore();
 
 	void BeginResolveKey() final;
-	void BeginResolveAction() final;
+	void BeginResolveInputMeta() final;
 	void BeginCacheQuery() final;
-	void BeginResolveInputs() final;
+	void BeginResolveInputData() final;
 	void BeginExecuteRemote() final;
 	void BeginExecuteLocal() final;
 	void BeginCacheStore() final;
 
 	void EndResolveKey(FBuildDefinitionResolvedParams&& Params);
-	void EndResolveAction(FBuildInputResolvedParams&& Params);
+	void EndResolveInputMeta(FBuildInputResolvedParams&& Params);
 	void EndCacheQuery(FCacheGetCompleteParams&& Params);
-	void EndResolveInputs(FBuildInputResolvedParams&& Params);
+	void EndResolveInputData(FBuildInputResolvedParams&& Params);
 	void EndExecuteRemote(FBuildWorkerActionCompleteParams&& Params);
 	void EndExecuteLocal();
 	void EndCacheStore(FCachePutCompleteParams&& Params);
@@ -245,26 +245,26 @@ const TCHAR* LexToString(EBuildJobState State)
 {
 	switch (State)
 	{
-	case EBuildJobState::ResolveKey:              return TEXT("ResolveKey");
-	case EBuildJobState::ResolveKeyWait:          return TEXT("ResolveKeyWait");
-	case EBuildJobState::ResolveAction:           return TEXT("ResolveAction");
-	case EBuildJobState::ResolveActionWait:       return TEXT("ResolveActionWait");
-	case EBuildJobState::CacheQuery:              return TEXT("CacheQuery");
-	case EBuildJobState::CacheQueryWait:          return TEXT("CacheQueryWait");
-	case EBuildJobState::ExecuteRemote:           return TEXT("ExecuteRemote");
-	case EBuildJobState::ExecuteRemoteWait:       return TEXT("ExecuteRemoteWait");
-	case EBuildJobState::ResolveRemoteInputs:     return TEXT("ResolveRemoteInputs");
-	case EBuildJobState::ResolveRemoteInputsWait: return TEXT("ResolveRemoteInputsWait");
-	case EBuildJobState::ExecuteRemoteRetry:      return TEXT("ExecuteRemoteRetry");
-	case EBuildJobState::ExecuteRemoteRetryWait:  return TEXT("ExecuteRemoteRetryWait");
-	case EBuildJobState::ResolveInputs:           return TEXT("ResolveInputs");
-	case EBuildJobState::ResolveInputsWait:       return TEXT("ResolveInputsWait");
-	case EBuildJobState::ExecuteLocal:            return TEXT("ExecuteLocal");
-	case EBuildJobState::ExecuteLocalWait:        return TEXT("ExecuteLocalWait");
-	case EBuildJobState::CacheStore:              return TEXT("CacheStore");
-	case EBuildJobState::CacheStoreWait:          return TEXT("CacheStoreWait");
-	case EBuildJobState::Complete:                return TEXT("Complete");
-	default: checkNoEntry();                      return TEXT("Unknown");
+	case EBuildJobState::ResolveKey:                 return TEXT("ResolveKey");
+	case EBuildJobState::ResolveKeyWait:             return TEXT("ResolveKeyWait");
+	case EBuildJobState::ResolveInputMeta:           return TEXT("ResolveInputMeta");
+	case EBuildJobState::ResolveInputMetaWait:       return TEXT("ResolveInputMetaWait");
+	case EBuildJobState::CacheQuery:                 return TEXT("CacheQuery");
+	case EBuildJobState::CacheQueryWait:             return TEXT("CacheQueryWait");
+	case EBuildJobState::ExecuteRemote:              return TEXT("ExecuteRemote");
+	case EBuildJobState::ExecuteRemoteWait:          return TEXT("ExecuteRemoteWait");
+	case EBuildJobState::ResolveRemoteInputData:     return TEXT("ResolveRemoteInputData");
+	case EBuildJobState::ResolveRemoteInputDataWait: return TEXT("ResolveRemoteInputDataWait");
+	case EBuildJobState::ExecuteRemoteRetry:         return TEXT("ExecuteRemoteRetry");
+	case EBuildJobState::ExecuteRemoteRetryWait:     return TEXT("ExecuteRemoteRetryWait");
+	case EBuildJobState::ResolveInputData:           return TEXT("ResolveInputData");
+	case EBuildJobState::ResolveInputDataWait:       return TEXT("ResolveInputDataWait");
+	case EBuildJobState::ExecuteLocal:               return TEXT("ExecuteLocal");
+	case EBuildJobState::ExecuteLocalWait:           return TEXT("ExecuteLocalWait");
+	case EBuildJobState::CacheStore:                 return TEXT("CacheStore");
+	case EBuildJobState::CacheStoreWait:             return TEXT("CacheStoreWait");
+	case EBuildJobState::Complete:                   return TEXT("Complete");
+	default: checkNoEntry();                         return TEXT("Unknown");
 	}
 }
 
@@ -477,7 +477,7 @@ void FBuildJob::EnterResolveKey()
 	}
 	if (Definition)
 	{
-		return AdvanceToState(EBuildJobState::ResolveAction);
+		return AdvanceToState(EBuildJobState::ResolveInputMeta);
 	}
 	if (DefinitionKey != FBuildKey::Empty)
 	{
@@ -524,36 +524,36 @@ void FBuildJob::SetDefinition(FBuildDefinition&& InDefinition)
 		LexToString(State), *Name, *FunctionName);
 	OutputBuilder = BuildSystem.CreateOutput(Name, FunctionName);
 	Definition = MoveTemp(InDefinition);
-	return AdvanceToState(EBuildJobState::ResolveAction);
+	return AdvanceToState(EBuildJobState::ResolveInputMeta);
 }
 
 ///////////////////////////////////////////////////////////////////////////////////////////////////////////////////////
 
-void FBuildJob::EnterResolveAction()
+void FBuildJob::EnterResolveInputMeta()
 {
-	if (Definition.Get().HasInputs() && !InputResolver)
-	{
-		return CompleteWithError(TEXT("Failed to resolve input metadata due to null input resolver."_SV));
-	}
 	if (!Definition.Get().HasInputs())
 	{
 		return CreateAction({});
 	}
-	Scheduler->DispatchResolveAction(this);
+	if (!InputResolver)
+	{
+		return CompleteWithError(TEXT("Failed to resolve input metadata due to null input resolver."_SV));
+	}
+	Scheduler->DispatchResolveInputMeta(this);
 }
 
-void FBuildJob::BeginResolveAction()
+void FBuildJob::BeginResolveInputMeta()
 {
-	if (CanExecuteState(EBuildJobState::ResolveAction))
+	if (CanExecuteState(EBuildJobState::ResolveInputMeta))
 	{
-		return AdvanceToState(EBuildJobState::ResolveActionWait, InputResolver->GetMeta(Definition.Get(), Priority,
-			[this](FBuildInputResolvedParams&& Params) { EndResolveAction(MoveTemp(Params)); }));
+		return AdvanceToState(EBuildJobState::ResolveInputMetaWait, InputResolver->GetMeta(Definition.Get(), Priority,
+			[this](FBuildInputResolvedParams&& Params) { EndResolveInputMeta(MoveTemp(Params)); }));
 	}
 }
 
-void FBuildJob::EndResolveAction(FBuildInputResolvedParams&& Params)
+void FBuildJob::EndResolveInputMeta(FBuildInputResolvedParams&& Params)
 {
-	if (CanExecuteState(EBuildJobState::ResolveActionWait))
+	if (CanExecuteState(EBuildJobState::ResolveInputMetaWait))
 	{
 		if (Params.Status == EStatus::Ok)
 		{
@@ -585,7 +585,7 @@ void FBuildJob::SetAction(FBuildAction&& InAction)
 {
 	checkf(Action.IsNull(), TEXT("Job already has an action for build of '%s' by %s."), *Name, *FunctionName);
 	checkf(State == EBuildJobState::ResolveKey ||
-		State == EBuildJobState::ResolveAction || State == EBuildJobState::ResolveActionWait,
+		State == EBuildJobState::ResolveInputMeta || State == EBuildJobState::ResolveInputMetaWait,
 		TEXT("Job is not expecting an action in state %s for build of '%s' by %s"),
 		LexToString(State), *Name, *FunctionName);
 	ActionKey = InAction.GetKey();
@@ -677,10 +677,10 @@ void FBuildJob::EndCacheQuery(FCacheGetCompleteParams&& Params)
 
 ///////////////////////////////////////////////////////////////////////////////////////////////////////////////////////
 
-void FBuildJob::EnterResolveInputs()
+void FBuildJob::EnterResolveInputData()
 {
 	// Identify missing inputs when resolving every input for the action.
-	if (State == EBuildJobState::ResolveInputs)
+	if (State == EBuildJobState::ResolveInputData)
 	{
 		MissingInputs.Reset();
 		Action.Get().IterateInputs([this](FStringView Key, const FIoHash& RawHash, uint64 RawSize)
@@ -708,22 +708,22 @@ void FBuildJob::EnterResolveInputs()
 	{
 		return CompleteWithError(TEXT("Failed to resolve input data due to null definition."_SV));
 	}
-	Scheduler->DispatchResolveInputs(this);
+	Scheduler->DispatchResolveInputData(this);
 }
 
-void FBuildJob::BeginResolveInputs()
+void FBuildJob::BeginResolveInputData()
 {
-	if (CanExecuteState(EBuildJobState::ResolveInputs))
+	if (CanExecuteState(EBuildJobState::ResolveInputData))
 	{
 		return AdvanceToState(GetNextState(State), InputResolver->GetData(Definition.Get(), Priority,
-			[this](FBuildInputResolvedParams&& Params) { EndResolveInputs(MoveTemp(Params)); },
+			[this](FBuildInputResolvedParams&& Params) { EndResolveInputData(MoveTemp(Params)); },
 			[this](FStringView Key) { return !!Algo::Find(MissingInputs, Key); }));
 	}
 }
 
-void FBuildJob::EndResolveInputs(FBuildInputResolvedParams&& Params)
+void FBuildJob::EndResolveInputData(FBuildInputResolvedParams&& Params)
 {
-	if (CanExecuteState(EBuildJobState::ResolveInputsWait))
+	if (CanExecuteState(EBuildJobState::ResolveInputDataWait))
 	{
 		if (Params.Status == EStatus::Ok)
 		{
@@ -753,12 +753,12 @@ void FBuildJob::SetInputs(FBuildInputs&& InInputs)
 	EBuildJobState ExecuteState;
 	switch (State)
 	{
-	case EBuildJobState::ResolveRemoteInputs:
-	case EBuildJobState::ResolveRemoteInputsWait:
+	case EBuildJobState::ResolveRemoteInputData:
+	case EBuildJobState::ResolveRemoteInputDataWait:
 		ExecuteState = EBuildJobState::ExecuteRemoteRetry;
 		break;
-	case EBuildJobState::ResolveInputs:
-	case EBuildJobState::ResolveInputsWait:
+	case EBuildJobState::ResolveInputData:
+	case EBuildJobState::ResolveInputDataWait:
 		ExecuteState = EBuildJobState::ExecuteLocal;
 		break;
 	default:
@@ -781,7 +781,7 @@ void FBuildJob::EnterExecuteRemote()
 	}
 	if (!EnumHasAnyFlags(BuildPolicy, EBuildPolicy::Remote))
 	{
-		return AdvanceToState(EBuildJobState::ResolveInputs);
+		return AdvanceToState(EBuildJobState::ResolveInputData);
 	}
 
 	if (!Worker || !WorkerExecutor)
@@ -795,7 +795,7 @@ void FBuildJob::EnterExecuteRemote()
 
 	if (!Worker || !WorkerExecutor)
 	{
-		return AdvanceToState(EBuildJobState::ResolveInputs);
+		return AdvanceToState(EBuildJobState::ResolveInputData);
 	}
 
 	Scheduler->DispatchExecuteRemote(this);
@@ -836,13 +836,13 @@ void FBuildJob::EndExecuteRemote(FBuildWorkerActionCompleteParams&& Params)
 				if (!MissingInputs.IsEmpty())
 				{
 					MissingInputs.Sort();
-					return AdvanceToState(EBuildJobState::ResolveRemoteInputs);
+					return AdvanceToState(EBuildJobState::ResolveRemoteInputData);
 				}
 				[[fallthrough]];
 			case EBuildJobState::ExecuteRemoteRetry:
 			case EBuildJobState::ExecuteRemoteRetryWait:
 			default:
-				return AdvanceToState(EBuildJobState::ResolveInputs);
+				return AdvanceToState(EBuildJobState::ResolveInputData);
 			}
 		}
 	}
@@ -924,8 +924,8 @@ void FBuildJob::EndExecuteLocal()
 
 void FBuildJob::SetOutput(const FBuildOutput& InOutput)
 {
-	checkf(State == EBuildJobState::ResolveKey || State == EBuildJobState::ResolveAction ||
-		State == EBuildJobState::CacheQuery || State == EBuildJobState::ResolveInputs ||
+	checkf(State == EBuildJobState::ResolveKey || State == EBuildJobState::ResolveInputMeta ||
+		State == EBuildJobState::CacheQuery || State == EBuildJobState::ResolveInputData ||
 		State == EBuildJobState::ExecuteRemote || State == EBuildJobState::ExecuteRemoteRetry ||
 		State == EBuildJobState::ExecuteLocal,
 		TEXT("Job is not expecting an output in state %s for build of '%s' by %s"),
@@ -1029,13 +1029,13 @@ bool FBuildJob::CanExecuteState(EBuildJobState RequestedState)
 
 	EBuildJobState LocalState = State;
 
-	// Handle states that execute multiple times under a different name.
+	// Handle states that execute multiple times under different names.
 	switch (LocalState)
 	{
-	case EBuildJobState::ResolveRemoteInputs:     LocalState = EBuildJobState::ResolveInputs; break;
-	case EBuildJobState::ResolveRemoteInputsWait: LocalState = EBuildJobState::ResolveInputsWait; break;
-	case EBuildJobState::ExecuteRemoteRetry:      LocalState = EBuildJobState::ExecuteRemote; break;
-	case EBuildJobState::ExecuteRemoteRetryWait:  LocalState = EBuildJobState::ExecuteRemoteWait; break;
+	case EBuildJobState::ResolveRemoteInputData:     LocalState = EBuildJobState::ResolveInputData; break;
+	case EBuildJobState::ResolveRemoteInputDataWait: LocalState = EBuildJobState::ResolveInputDataWait; break;
+	case EBuildJobState::ExecuteRemoteRetry:         LocalState = EBuildJobState::ExecuteRemote; break;
+	case EBuildJobState::ExecuteRemoteRetryWait:     LocalState = EBuildJobState::ExecuteRemoteWait; break;
 	}
 
 	if (LocalState == RequestedState)
@@ -1046,13 +1046,13 @@ bool FBuildJob::CanExecuteState(EBuildJobState RequestedState)
 	// Handle wait states as the base state because the async operation can finish before advancing to the wait state.
 	switch (RequestedState)
 	{
-	case EBuildJobState::ResolveKeyWait:     return LocalState == EBuildJobState::ResolveKey;
-	case EBuildJobState::ResolveActionWait:  return LocalState == EBuildJobState::ResolveAction;
-	case EBuildJobState::CacheQueryWait:     return LocalState == EBuildJobState::CacheQuery;
-	case EBuildJobState::ResolveInputsWait:  return LocalState == EBuildJobState::ResolveInputs;
-	case EBuildJobState::ExecuteRemoteWait:  return LocalState == EBuildJobState::ExecuteRemote;
-	case EBuildJobState::ExecuteLocalWait:   return LocalState == EBuildJobState::ExecuteLocal;
-	case EBuildJobState::CacheStoreWait:     return LocalState == EBuildJobState::CacheStore;
+	case EBuildJobState::ResolveKeyWait:       return LocalState == EBuildJobState::ResolveKey;
+	case EBuildJobState::ResolveInputMetaWait: return LocalState == EBuildJobState::ResolveInputMeta;
+	case EBuildJobState::CacheQueryWait:       return LocalState == EBuildJobState::CacheQuery;
+	case EBuildJobState::ResolveInputDataWait: return LocalState == EBuildJobState::ResolveInputData;
+	case EBuildJobState::ExecuteRemoteWait:    return LocalState == EBuildJobState::ExecuteRemote;
+	case EBuildJobState::ExecuteLocalWait:     return LocalState == EBuildJobState::ExecuteLocal;
+	case EBuildJobState::CacheStoreWait:       return LocalState == EBuildJobState::CacheStore;
 	}
 
 	return false;
@@ -1126,11 +1126,11 @@ void FBuildJob::ExecuteTransition(EBuildJobState OldState, EBuildJobState NewSta
 		Worker = nullptr;
 		WorkerExecutor = nullptr;
 	}
-	if (OldState < EBuildJobState::ResolveInputsWait && EBuildJobState::ResolveInputsWait <= NewState)
+	if (OldState < EBuildJobState::ResolveInputDataWait && EBuildJobState::ResolveInputDataWait <= NewState)
 	{
 		Definition.Reset();
 	}
-	if (OldState <= EBuildJobState::ResolveInputsWait && EBuildJobState::ResolveInputsWait < NewState)
+	if (OldState <= EBuildJobState::ResolveInputDataWait && EBuildJobState::ResolveInputDataWait < NewState)
 	{
 		MissingInputs.Empty();
 	}
@@ -1155,16 +1155,16 @@ void FBuildJob::ExecuteTransition(EBuildJobState OldState, EBuildJobState NewSta
 
 	switch (NewState)
 	{
-	case EBuildJobState::ResolveKey:          return EnterResolveKey();
-	case EBuildJobState::ResolveAction:       return EnterResolveAction();
-	case EBuildJobState::CacheQuery:          return EnterCacheQuery();
-	case EBuildJobState::ExecuteRemote:       return EnterExecuteRemote();
-	case EBuildJobState::ResolveRemoteInputs: return EnterResolveInputs();
-	case EBuildJobState::ExecuteRemoteRetry:  return EnterExecuteRemote();
-	case EBuildJobState::ResolveInputs:       return EnterResolveInputs();
-	case EBuildJobState::ExecuteLocal:        return EnterExecuteLocal();
-	case EBuildJobState::CacheStore:          return EnterCacheStore();
-	case EBuildJobState::Complete:            return EndJob();
+	case EBuildJobState::ResolveKey:             return EnterResolveKey();
+	case EBuildJobState::ResolveInputMeta:       return EnterResolveInputMeta();
+	case EBuildJobState::CacheQuery:             return EnterCacheQuery();
+	case EBuildJobState::ExecuteRemote:          return EnterExecuteRemote();
+	case EBuildJobState::ResolveRemoteInputData: return EnterResolveInputData();
+	case EBuildJobState::ExecuteRemoteRetry:     return EnterExecuteRemote();
+	case EBuildJobState::ResolveInputData:       return EnterResolveInputData();
+	case EBuildJobState::ExecuteLocal:           return EnterExecuteLocal();
+	case EBuildJobState::CacheStore:             return EnterCacheStore();
+	case EBuildJobState::Complete:               return EndJob();
 	}
 }
 
diff --git a/Engine/Source/Developer/DerivedDataCache/Public/DerivedDataBuildJob.h b/Engine/Source/Developer/DerivedDataCache/Public/DerivedDataBuildJob.h
--- a/Engine/Source/Developer/DerivedDataCache/Public/DerivedDataBuildJob.h
+++ b/Engine/Source/Developer/DerivedDataCache/Public/DerivedDataBuildJob.h
@@ -56,8 +56,8 @@ public:
 	virtual void BeginCacheQuery() = 0;
 	virtual void BeginCacheStore() = 0;
 	virtual void BeginResolveKey() = 0;
-	virtual void BeginResolveAction() = 0;
-	virtual void BeginResolveInputs() = 0;
+	virtual void BeginResolveInputMeta() = 0;
+	virtual void BeginResolveInputData() = 0;
 	virtual void BeginExecuteRemote() = 0;
 	virtual void BeginExecuteLocal() = 0;
 
diff --git a/Engine/Source/Developer/DerivedDataCache/Public/DerivedDataBuildScheduler.h b/Engine/Source/Developer/DerivedDataCache/Public/DerivedDataBuildScheduler.h
--- a/Engine/Source/Developer/DerivedDataCache/Public/DerivedDataBuildScheduler.h
+++ b/Engine/Source/Developer/DerivedDataCache/Public/DerivedDataBuildScheduler.h
@@ -36,8 +36,8 @@ public:
 	virtual void DispatchCacheQuery(IBuildJob* Job) { Job->BeginCacheQuery(); }
 	virtual void DispatchCacheStore(IBuildJob* Job) { Job->BeginCacheStore(); }
 	virtual void DispatchResolveKey(IBuildJob* Job) { Job->BeginResolveKey(); }
-	virtual void DispatchResolveAction(IBuildJob* Job) { Job->BeginResolveAction(); }
-	virtual void DispatchResolveInputs(IBuildJob* Job) { Job->BeginResolveInputs(); }
+	virtual void DispatchResolveInputMeta(IBuildJob* Job) { Job->BeginResolveInputMeta(); }
+	virtual void DispatchResolveInputData(IBuildJob* Job) { Job->BeginResolveInputData(); }
 	virtual void DispatchExecuteRemote(IBuildJob* Job) { Job->BeginExecuteRemote(); }
 	virtual void DispatchExecuteLocal(IBuildJob* Job) { Job->BeginExecuteLocal(); }
 
