DDC: Made the input resolver optional when constructing a build session

diff --git a/Engine/Source/Developer/DerivedDataCache/Public/DerivedDataBuild.h b/Engine/Source/Developer/DerivedDataCache/Public/DerivedDataBuild.h
--- a/Engine/Source/Developer/DerivedDataCache/Public/DerivedDataBuild.h
+++ b/Engine/Source/Developer/DerivedDataCache/Public/DerivedDataBuild.h
@@ -109,11 +109,17 @@ public:
 	/**
 	 * Create a build session.
 	 *
+	 * An input resolver is required for the session to perform builds with unresolved inputs, or
+	 * to resolve build payload keys. An input resolver is optional when the session is used only
+	 * to build actions with any inputs provided to the session directly.
+	 *
+	 * A default scheduler is used if one is not provided. Using the default is recommended.
+	 *
 	 * @param Name            The name by which to identify this session for logging and profiling.
-	 * @param InputResolver   The input resolver to resolve definitions and inputs for requested builds.
+	 * @param InputResolver   The input resolver to resolve inputs for requested builds. Optional.
 	 * @param Scheduler       The scheduler for builds created through the session. Optional.
 	 */
-	virtual FBuildSession CreateSession(FStringView Name, IBuildInputResolver* InputResolver, IBuildScheduler* Scheduler = nullptr) = 0;
+	virtual FBuildSession CreateSession(FStringView Name, IBuildInputResolver* InputResolver = nullptr, IBuildScheduler* Scheduler = nullptr) = 0;
 
 	/**
 	 * Returns the version of the build system.
