DDC: Reworked IBuildInputResolver to clarify function names and callback parameters

diff --git a/Engine/Source/Developer/DerivedDataCache/Private/DerivedDataBuildJob.cpp b/Engine/Source/Developer/DerivedDataCache/Private/DerivedDataBuildJob.cpp
--- a/Engine/Source/Developer/DerivedDataCache/Private/DerivedDataBuildJob.cpp
+++ b/Engine/Source/Developer/DerivedDataCache/Private/DerivedDataBuildJob.cpp
@@ -145,10 +145,10 @@ private:
 	void BeginExecuteLocal() final;
 	void BeginCacheStore() final;
 
-	void EndResolveKey(FBuildDefinitionResolvedParams&& Params);
-	void EndResolveInputMeta(FBuildInputResolvedParams&& Params);
+	void EndResolveKey(FBuildKeyResolvedParams&& Params);
+	void EndResolveInputMeta(FBuildInputMetaResolvedParams&& Params);
 	void EndCacheQuery(FCacheGetCompleteParams&& Params);
-	void EndResolveInputData(FBuildInputResolvedParams&& Params);
+	void EndResolveInputData(FBuildInputDataResolvedParams&& Params);
 	void EndExecuteRemote(FBuildWorkerActionCompleteParams&& Params);
 	void EndExecuteLocal();
 	void EndCacheStore(FCachePutCompleteParams&& Params);
@@ -159,7 +159,7 @@ private:
 	void SetOutput(const FBuildOutput& Output) final;
 	void SetOutputNoCheck(FBuildOutput&& Output);
 
-	void CreateAction(TConstArrayView<FBuildInputByKey> InputMeta);
+	void CreateAction(TConstArrayView<FBuildInputMetaByKey> InputMeta);
 
 	/** Terminate the job and send the error to the output complete callback. */
 	void CompleteWithError(FStringView Error);
@@ -494,12 +494,12 @@ void FBuildJob::BeginResolveKey()
 {
 	if (CanExecuteState(EBuildJobState::ResolveKey))
 	{
-		return AdvanceToState(EBuildJobState::ResolveKeyWait, InputResolver->GetDefinition(DefinitionKey,
-			[this](FBuildDefinitionResolvedParams&& Params) { EndResolveKey(MoveTemp(Params)); }));
+		return AdvanceToState(EBuildJobState::ResolveKeyWait, InputResolver->ResolveKey(DefinitionKey,
+			[this](FBuildKeyResolvedParams&& Params) { EndResolveKey(MoveTemp(Params)); }));
 	}
 }
 
-void FBuildJob::EndResolveKey(FBuildDefinitionResolvedParams&& Params)
+void FBuildJob::EndResolveKey(FBuildKeyResolvedParams&& Params)
 {
 	if (CanExecuteState(EBuildJobState::ResolveKeyWait))
 	{
@@ -546,12 +546,12 @@ void FBuildJob::BeginResolveInputMeta()
 {
 	if (CanExecuteState(EBuildJobState::ResolveInputMeta))
 	{
-		return AdvanceToState(EBuildJobState::ResolveInputMetaWait, InputResolver->GetMeta(Definition.Get(), Priority,
-			[this](FBuildInputResolvedParams&& Params) { EndResolveInputMeta(MoveTemp(Params)); }));
+		return AdvanceToState(EBuildJobState::ResolveInputMetaWait, InputResolver->ResolveInputMeta(Definition.Get(),
+			Priority, [this](FBuildInputMetaResolvedParams&& Params) { EndResolveInputMeta(MoveTemp(Params)); }));
 	}
 }
 
-void FBuildJob::EndResolveInputMeta(FBuildInputResolvedParams&& Params)
+void FBuildJob::EndResolveInputMeta(FBuildInputMetaResolvedParams&& Params)
 {
 	if (CanExecuteState(EBuildJobState::ResolveInputMetaWait))
 	{
@@ -566,7 +566,7 @@ void FBuildJob::EndResolveInputMeta(FBuildInputResolvedParams&& Params)
 	}
 }
 
-void FBuildJob::CreateAction(TConstArrayView<FBuildInputByKey> InputMeta)
+void FBuildJob::CreateAction(TConstArrayView<FBuildInputMetaByKey> InputMeta)
 {
 	const FBuildDefinition& LocalDefinition = Definition.Get();
 	FBuildActionBuilder Builder = BuildSystem.CreateAction(Name, FunctionName);
@@ -574,9 +574,9 @@ void FBuildJob::CreateAction(TConstArrayView<FBuildInputByKey> InputMeta)
 	{
 		Builder.AddConstant(Key, Value);
 	});
-	for (const FBuildInputByKey& InputByKey : InputMeta)
+	for (const FBuildInputMetaByKey& Input : InputMeta)
 	{
-		Builder.AddInput(InputByKey.Key, InputByKey.Meta.RawHash, InputByKey.Meta.RawSize);
+		Builder.AddInput(Input.Key, Input.RawHash, Input.RawSize);
 	}
 	SetAction(Builder.Build());
 }
@@ -715,22 +715,22 @@ void FBuildJob::BeginResolveInputData()
 {
 	if (CanExecuteState(EBuildJobState::ResolveInputData))
 	{
-		return AdvanceToState(GetNextState(State), InputResolver->GetData(Definition.Get(), Priority,
-			[this](FBuildInputResolvedParams&& Params) { EndResolveInputData(MoveTemp(Params)); },
+		return AdvanceToState(GetNextState(State), InputResolver->ResolveInputData(Definition.Get(), Priority,
+			[this](FBuildInputDataResolvedParams&& Params) { EndResolveInputData(MoveTemp(Params)); },
 			[this](FStringView Key) { return !!Algo::Find(MissingInputs, Key); }));
 	}
 }
 
-void FBuildJob::EndResolveInputData(FBuildInputResolvedParams&& Params)
+void FBuildJob::EndResolveInputData(FBuildInputDataResolvedParams&& Params)
 {
 	if (CanExecuteState(EBuildJobState::ResolveInputDataWait))
 	{
 		if (Params.Status == EStatus::Ok)
 		{
 			FBuildInputsBuilder Builder = BuildSystem.CreateInputs(Name);
-			for (const FBuildInputByKey& InputByKey : Params.Inputs)
+			for (const FBuildInputDataByKey& Input : Params.Inputs)
 			{
-				Builder.AddInput(InputByKey.Key, InputByKey.Data);
+				Builder.AddInput(Input.Key, Input.Data);
 			}
 			if (Inputs)
 			{
diff --git a/Engine/Source/Developer/DerivedDataCache/Public/DerivedDataBuildInputResolver.h b/Engine/Source/Developer/DerivedDataCache/Public/DerivedDataBuildInputResolver.h
--- a/Engine/Source/Developer/DerivedDataCache/Public/DerivedDataBuildInputResolver.h
+++ b/Engine/Source/Developer/DerivedDataCache/Public/DerivedDataBuildInputResolver.h
@@ -14,68 +14,83 @@ struct FGuid;
 
 namespace UE::DerivedData { class FBuildDefinition; }
 namespace UE::DerivedData { class FOptionalBuildDefinition; }
-namespace UE::DerivedData { struct FBuildDefinitionResolvedParams; }
-namespace UE::DerivedData { struct FBuildInputResolvedParams; }
+namespace UE::DerivedData { struct FBuildInputDataResolvedParams; }
+namespace UE::DerivedData { struct FBuildInputMetaResolvedParams; }
+namespace UE::DerivedData { struct FBuildKeyResolvedParams; }
 
 namespace UE::DerivedData
 {
 
 using FBuildInputFilter = TUniqueFunction<bool (FStringView Key)>;
-using FOnBuildDefinitionResolved = TUniqueFunction<void (FBuildDefinitionResolvedParams&& Params)>;
-using FOnBuildInputResolved = TUniqueFunction<void (FBuildInputResolvedParams&& Params)>;
+using FOnBuildInputDataResolved = TUniqueFunction<void (FBuildInputDataResolvedParams&& Params)>;
+using FOnBuildInputMetaResolved = TUniqueFunction<void (FBuildInputMetaResolvedParams&& Params)>;
+using FOnBuildKeyResolved = TUniqueFunction<void (FBuildKeyResolvedParams&& Params)>;
 
-/** Interface to resolve input references from a build definition into the referenced data. */
+/** Interface to resolve references to build inputs. */
 class IBuildInputResolver
 {
 public:
+	virtual ~IBuildInputResolver() = default;
+
 	/**
 	 * Asynchronous request to resolve a definition from a key.
 	 *
 	 * @param Key          The key of the definition to resolve.
 	 * @param OnResolved   A callback invoked when resolving completes or is canceled.
 	 */
-	virtual FRequest GetDefinition(
+	virtual FRequest ResolveKey(
 		const FBuildKey& Key,
-		FOnBuildDefinitionResolved&& OnResolved) = 0;
+		FOnBuildKeyResolved&& OnResolved) = 0;
 
 	/**
 	 * Asynchronous request to resolve metadata for the inputs from the definition.
 	 *
 	 * @param Definition   The definition to resolve input metadata for.
 	 * @param Priority     A priority to consider when scheduling the request. See EPriority.
-	 * @param OnResolver   A callback invoked when resolving completes or is canceled.
+	 * @param OnResolved   A callback invoked when resolving completes or is canceled.
 	 */
-	virtual FRequest GetMeta(
+	virtual FRequest ResolveInputMeta(
 		const FBuildDefinition& Definition,
 		EPriority Priority,
-		FOnBuildInputResolved&& OnResolved) = 0;
+		FOnBuildInputMetaResolved&& OnResolved) = 0;
 
 	/**
 	 * Asynchronous request to resolve data for the inputs from the definition.
 	 *
 	 * @param Definition   The definition to resolve input data for.
 	 * @param Priority     A priority to consider when scheduling the request. See EPriority.
-	 * @param OnResolver   A callback invoked when resolving completes or is canceled.
+	 * @param OnResolved   A callback invoked when resolving completes or is canceled.
 	 * @param Filter       An optional predicate to filter which input keys have data resolved.
 	 */
-	virtual FRequest GetData(
+	virtual FRequest ResolveInputData(
 		const FBuildDefinition& Definition,
 		EPriority Priority,
-		FOnBuildInputResolved&& OnResolved,
+		FOnBuildInputDataResolved&& OnResolved,
 		FBuildInputFilter&& Filter = FBuildInputFilter()) = 0;
 };
 
-/** Metadata for build inputs. */
-struct FBuildInputMeta
+/** Metadata for build inputs with the input key. */
+struct FBuildInputMetaByKey
 {
+	/** Key used to identify this input. */
+	FStringView Key;
 	/** Hash of the raw (uncompressed) input. */
 	FIoHash RawHash;
 	/** Size of the raw (uncompressed) input in bytes. */
 	uint64 RawSize = 0;
 };
 
+/** Data for build inputs with the input key. */
+struct FBuildInputDataByKey
+{
+	/** Key used to identify this input. */
+	FStringView Key;
+	/** Data for the input. */
+	FCompressedBuffer Data;
+};
+
 /** Parameters for the resolved callback for build definition requests. */
-struct FBuildDefinitionResolvedParams
+struct FBuildKeyResolvedParams
 {
 	/** Key for the build definition request that resolved or was canceled. */
 	FBuildKey Key;
@@ -87,27 +102,21 @@ struct FBuildDefinitionResolvedParams
 	EStatus Status = EStatus::Error;
 };
 
-/** Data and metadata for build inputs with the input key. */
-struct FBuildInputByKey
+/** Parameters for the resolved callback for build input metadata requests. */
+struct FBuildInputMetaResolvedParams
 {
-	/** Key used to identify this input in the definition. */
-	FStringView Key;
-
-	/** Metadata for the input. */
-	FBuildInputMeta Meta;
+	/** All of the requested inputs sorted by key. Only available when Status is Ok. */
+	TConstArrayView<FBuildInputMetaByKey> Inputs;
 
-	/** Data for the input. Null if only metadata was requested. */
-	FCompressedBuffer Data;
+	/** Status of the input request. */
+	EStatus Status = EStatus::Error;
 };
 
-/** Parameters for the resolved callback for build input requests. */
-struct FBuildInputResolvedParams
+/** Parameters for the resolved callback for build input data requests. */
+struct FBuildInputDataResolvedParams
 {
-	/** Key for the build input request that resolved or was canceled. */
-	FBuildKey Key;
-
-	/** All of the requested inputs sorted by key. Only available when Status is Ok. */
-	TConstArrayView<FBuildInputByKey> Inputs;
+	/** All of the requested inputs sorted by raw hash. Only available when Status is Ok. */
+	TConstArrayView<FBuildInputDataByKey> Inputs;
 
 	/** Status of the input request. */
 	EStatus Status = EStatus::Error;
