DDC: Attempt to resolve inputs from the action if the definition is not available

diff --git a/Engine/Source/Developer/DerivedDataCache/Private/DerivedDataBuildJob.cpp b/Engine/Source/Developer/DerivedDataCache/Private/DerivedDataBuildJob.cpp
--- a/Engine/Source/Developer/DerivedDataCache/Private/DerivedDataBuildJob.cpp
+++ b/Engine/Source/Developer/DerivedDataCache/Private/DerivedDataBuildJob.cpp
@@ -704,10 +704,6 @@ void FBuildJob::EnterResolveInputData()
 	{
 		return CompleteWithError(TEXT("Failed to resolve input data due to null input resolver."_SV));
 	}
-	if (!Definition)
-	{
-		return CompleteWithError(TEXT("Failed to resolve input data due to null definition."_SV));
-	}
 	Scheduler->DispatchResolveInputData(this);
 }
 
@@ -715,9 +711,18 @@ void FBuildJob::BeginResolveInputData()
 {
 	if (CanExecuteState(EBuildJobState::ResolveInputData))
 	{
-		return AdvanceToState(GetNextState(State), InputResolver->ResolveInputData(Definition.Get(), Priority,
-			[this](FBuildInputDataResolvedParams&& Params) { EndResolveInputData(MoveTemp(Params)); },
-			[this](FStringView Key) { return !!Algo::Find(MissingInputs, Key); }));
+		if (Definition)
+		{
+			return AdvanceToState(GetNextState(State), InputResolver->ResolveInputData(Definition.Get(), Priority,
+				[this](FBuildInputDataResolvedParams&& Params) { EndResolveInputData(MoveTemp(Params)); },
+				[this](FStringView Key) { return !!Algo::Find(MissingInputs, Key); }));
+		}
+		else
+		{
+			return AdvanceToState(GetNextState(State), InputResolver->ResolveInputData(Action.Get(), Priority,
+				[this](FBuildInputDataResolvedParams&& Params) { EndResolveInputData(MoveTemp(Params)); },
+				[this](FStringView Key) { return !!Algo::Find(MissingInputs, Key); }));
+		}
 	}
 }
 
@@ -1126,21 +1131,15 @@ void FBuildJob::ExecuteTransition(EBuildJobState OldState, EBuildJobState NewSta
 		Worker = nullptr;
 		WorkerExecutor = nullptr;
 	}
-	if (OldState < EBuildJobState::ResolveInputDataWait && EBuildJobState::ResolveInputDataWait <= NewState)
-	{
-		Definition.Reset();
-	}
 	if (OldState <= EBuildJobState::ResolveInputDataWait && EBuildJobState::ResolveInputDataWait < NewState)
 	{
 		MissingInputs.Empty();
+		Definition.Reset();
 	}
-	if (OldState < EBuildJobState::ExecuteLocalWait && EBuildJobState::ExecuteLocalWait <= NewState)
+	if (OldState <= EBuildJobState::ExecuteLocalWait && EBuildJobState::ExecuteLocalWait < NewState && !Output)
 	{
 		Action.Reset();
 		Inputs.Reset();
-	}
-	if (OldState <= EBuildJobState::ExecuteLocalWait && EBuildJobState::ExecuteLocalWait < NewState && !Output)
-	{
 		SetOutputNoCheck(OutputBuilder.Build());
 	}
 	if (OldState < EBuildJobState::CacheStore && EBuildJobState::CacheStore <= NewState)
diff --git a/Engine/Source/Developer/DerivedDataCache/Public/DerivedDataBuildInputResolver.h b/Engine/Source/Developer/DerivedDataCache/Public/DerivedDataBuildInputResolver.h
--- a/Engine/Source/Developer/DerivedDataCache/Public/DerivedDataBuildInputResolver.h
+++ b/Engine/Source/Developer/DerivedDataCache/Public/DerivedDataBuildInputResolver.h
@@ -5,6 +5,7 @@
 #include "Compression/CompressedBuffer.h"
 #include "Containers/ArrayView.h"
 #include "Containers/StringView.h"
+#include "DerivedDataBuildDefinition.h"
 #include "DerivedDataBuildKey.h"
 #include "DerivedDataRequest.h"
 #include "IO/IoHash.h"
@@ -12,8 +13,7 @@
 
 struct FGuid;
 
-namespace UE::DerivedData { class FBuildDefinition; }
-namespace UE::DerivedData { class FOptionalBuildDefinition; }
+namespace UE::DerivedData { class FBuildAction; }
 namespace UE::DerivedData { struct FBuildInputDataResolvedParams; }
 namespace UE::DerivedData { struct FBuildInputMetaResolvedParams; }
 namespace UE::DerivedData { struct FBuildKeyResolvedParams; }
@@ -26,49 +26,6 @@
 using FOnBuildInputMetaResolved = TUniqueFunction<void (FBuildInputMetaResolvedParams&& Params)>;
 using FOnBuildKeyResolved = TUniqueFunction<void (FBuildKeyResolvedParams&& Params)>;
 
-/** Interface to resolve references to build inputs. */
-class IBuildInputResolver
-{
-public:
-	virtual ~IBuildInputResolver() = default;
-
-	/**
-	 * Asynchronous request to resolve a definition from a key.
-	 *
-	 * @param Key          The key of the definition to resolve.
-	 * @param OnResolved   A callback invoked when resolving completes or is canceled.
-	 */
-	virtual FRequest ResolveKey(
-		const FBuildKey& Key,
-		FOnBuildKeyResolved&& OnResolved) = 0;
-
-	/**
-	 * Asynchronous request to resolve metadata for the inputs from the definition.
-	 *
-	 * @param Definition   The definition to resolve input metadata for.
-	 * @param Priority     A priority to consider when scheduling the request. See EPriority.
-	 * @param OnResolved   A callback invoked when resolving completes or is canceled.
-	 */
-	virtual FRequest ResolveInputMeta(
-		const FBuildDefinition& Definition,
-		EPriority Priority,
-		FOnBuildInputMetaResolved&& OnResolved) = 0;
-
-	/**
-	 * Asynchronous request to resolve data for the inputs from the definition.
-	 *
-	 * @param Definition   The definition to resolve input data for.
-	 * @param Priority     A priority to consider when scheduling the request. See EPriority.
-	 * @param OnResolved   A callback invoked when resolving completes or is canceled.
-	 * @param Filter       An optional predicate to filter which input keys have data resolved.
-	 */
-	virtual FRequest ResolveInputData(
-		const FBuildDefinition& Definition,
-		EPriority Priority,
-		FOnBuildInputDataResolved&& OnResolved,
-		FBuildInputFilter&& Filter = FBuildInputFilter()) = 0;
-};
-
 /** Metadata for build inputs with the input key. */
 struct FBuildInputMetaByKey
 {
@@ -122,4 +79,77 @@
 	EStatus Status = EStatus::Error;
 };
 
+/** Interface to resolve references to build inputs. */
+class IBuildInputResolver
+{
+public:
+	virtual ~IBuildInputResolver() = default;
+
+	/**
+	 * Asynchronous request to resolve a definition from a key.
+	 *
+	 * @param Key          The key of the definition to resolve.
+	 * @param OnResolved   A required callback invoked when resolving completes or is canceled.
+	 */
+	virtual FRequest ResolveKey(
+		const FBuildKey& Key,
+		FOnBuildKeyResolved&& OnResolved)
+	{
+		OnResolved({Key, {}, EStatus::Error});
+		return FRequest();
+	}
+
+	/**
+	 * Asynchronous request to resolve metadata for the inputs from the definition.
+	 *
+	 * @param Definition   The definition to resolve input metadata for.
+	 * @param Priority     A priority to consider when scheduling the request. See EPriority.
+	 * @param OnResolved   A required callback invoked when resolving completes or is canceled.
+	 */
+	virtual FRequest ResolveInputMeta(
+		const FBuildDefinition& Definition,
+		EPriority Priority,
+		FOnBuildInputMetaResolved&& OnResolved)
+	{
+		OnResolved({{}, EStatus::Error});
+		return FRequest();
+	}
+
+	/**
+	 * Asynchronous request to resolve data for the inputs from the definition.
+	 *
+	 * @param Definition   The definition to resolve input data for.
+	 * @param Priority     A priority to consider when scheduling the request. See EPriority.
+	 * @param OnResolved   A required callback invoked when resolving completes or is canceled.
+	 * @param Filter       An optional predicate to filter which input keys have data resolved.
+	 */
+	virtual FRequest ResolveInputData(
+		const FBuildDefinition& Definition,
+		EPriority Priority,
+		FOnBuildInputDataResolved&& OnResolved,
+		FBuildInputFilter&& Filter = FBuildInputFilter())
+	{
+		OnResolved({{}, EStatus::Error});
+		return FRequest();
+	}
+
+	/**
+	 * Asynchronous request to resolve data for the inputs from the action.
+	 *
+	 * @param Action       The action to resolve input data for.
+	 * @param Priority     A priority to consider when scheduling the request. See EPriority.
+	 * @param OnResolved   A required callback invoked when resolving completes or is canceled.
+	 * @param Filter       An optional predicate to filter which input keys have data resolved.
+	 */
+	virtual FRequest ResolveInputData(
+		const FBuildAction& Action,
+		EPriority Priority,
+		FOnBuildInputDataResolved&& OnResolved,
+		FBuildInputFilter&& Filter = FBuildInputFilter())
+	{
+		OnResolved({{}, EStatus::Error});
+		return FRequest();
+	}
+};
+
 } // UE::DerivedData
