DDC: Moved EBuildPolicy to its own header to remove private dependencies on the build session

diff --git a/Engine/Source/Developer/DerivedDataCache/Private/DerivedDataBuildJob.cpp b/Engine/Source/Developer/DerivedDataCache/Private/DerivedDataBuildJob.cpp
--- a/Engine/Source/Developer/DerivedDataCache/Private/DerivedDataBuildJob.cpp
+++ b/Engine/Source/Developer/DerivedDataCache/Private/DerivedDataBuildJob.cpp
@@ -14,9 +14,9 @@
 #include "DerivedDataBuildJobContext.h"
 #include "DerivedDataBuildKey.h"
 #include "DerivedDataBuildOutput.h"
+#include "DerivedDataBuildPolicy.h"
 #include "DerivedDataBuildPrivate.h"
 #include "DerivedDataBuildScheduler.h"
-#include "DerivedDataBuildSession.h"
 #include "DerivedDataBuildWorker.h"
 #include "DerivedDataBuildWorkerRegistry.h"
 #include "DerivedDataCache.h"
diff --git a/Engine/Source/Developer/DerivedDataCache/Private/DerivedDataBuildJobContext.cpp b/Engine/Source/Developer/DerivedDataCache/Private/DerivedDataBuildJobContext.cpp
--- a/Engine/Source/Developer/DerivedDataCache/Private/DerivedDataBuildJobContext.cpp
+++ b/Engine/Source/Developer/DerivedDataCache/Private/DerivedDataBuildJobContext.cpp
@@ -4,7 +4,7 @@
 
 #include "DerivedDataBuildJob.h"
 #include "DerivedDataBuildOutput.h"
-#include "DerivedDataBuildSession.h"
+#include "DerivedDataBuildPolicy.h"
 #include "DerivedDataCache.h"
 #include "Misc/StringBuilder.h"
 #include "UObject/NameTypes.h"
diff --git a/Engine/Source/Developer/DerivedDataCache/Public/DerivedDataBuildSession.h b/Engine/Source/Developer/DerivedDataCache/Public/DerivedDataBuildSession.h
--- a/Engine/Source/Developer/DerivedDataCache/Public/DerivedDataBuildSession.h
+++ b/Engine/Source/Developer/DerivedDataCache/Public/DerivedDataBuildSession.h
@@ -5,8 +5,8 @@
 #include "CoreTypes.h"
 #include "Containers/StringView.h"
 #include "DerivedDataBuildKey.h"
+#include "DerivedDataBuildPolicy.h"
 #include "DerivedDataRequest.h"
-#include "Misc/EnumClassFlags.h"
 #include "Templates/Function.h"
 #include "Templates/UniquePtr.h"
 
@@ -16,8 +16,8 @@ namespace UE::DerivedData { class FBuildOutput; }
 namespace UE::DerivedData { class FBuildSession; }
 namespace UE::DerivedData { class FOptionalBuildInputs; }
 namespace UE::DerivedData { class FPayload; }
-namespace UE::DerivedData { struct FBuildCompleteParams; }
 namespace UE::DerivedData { struct FBuildActionCompleteParams; }
+namespace UE::DerivedData { struct FBuildCompleteParams; }
 namespace UE::DerivedData { struct FBuildPayloadCompleteParams; }
 
 namespace UE::DerivedData
@@ -27,44 +27,6 @@ using FOnBuildComplete = TUniqueFunction<void (FBuildCompleteParams&& Params)>;
 using FOnBuildActionComplete = TUniqueFunction<void (FBuildActionCompleteParams&& Params)>;
 using FOnBuildPayloadComplete = TUniqueFunction<void (FBuildPayloadCompleteParams&& Params)>;
 
-/**
- * Flags to control the behavior of build requests.
- *
- * The build policy flags can be combined to support a variety of usage patterns. Examples:
- *
- * Build(Default): allow cache get; build if missing; return every payload.
- * Build(SkipCacheGet): never get from the cache; always build; return every payload.
- * Build(SkipBuild | SkipData): allow cache get; never build; skip every payload.
- *
- * BuildPayload(Default): allow cache get; build if missing; return one payload.
- * BuildPayload(SkipCacheGet): never get from the cache; always build; return one payload.
- * BuildPayload(SkipBuild | SkipData): allow cache get; never build; skip every payload.
- */
-enum class EBuildPolicy : uint8
-{
-	/** A value without any flags set. */
-	None            = 0,
-
-	/** Allow local execution. */
-	Local           = 1 << 0,
-	/** Allow remote execution if the function has a registered build worker. */
-	Remote          = 1 << 1,
-
-	/** Skip attempting a cache get before building. */
-	SkipCacheGet    = 1 << 2,
-	/** Skip attempting a cache get before building. */
-	SkipCachePut    = 1 << 3,
-	/** Skip attempting a build when the cache get is missed or skipped. */
-	SkipBuild       = 1 << 4,
-	/** Skip fetching the payload data from the cache. */
-	SkipData        = 1 << 5,
-
-	/** Allow cache get+put, local+remote build when missed or skipped, and fetch the payload(s). */
-	Default         = Local | Remote,
-};
-
-ENUM_CLASS_FLAGS(EBuildPolicy);
-
 } // UE::DerivedData
 
 namespace UE::DerivedData::Private
