DDC: Added missing file DerivedDataBuildPolicy.h

diff --git a/Engine/Source/Developer/DerivedDataCache/Public/DerivedDataBuildPolicy.h b/Engine/Source/Developer/DerivedDataCache/Public/DerivedDataBuildPolicy.h
new file mode 100644
--- /dev/null
+++ b/Engine/Source/Developer/DerivedDataCache/Public/DerivedDataBuildPolicy.h
@@ -0,0 +1,49 @@
+// Copyright Epic Games, Inc. All Rights Reserved.
+
+#pragma once
+
+#include "CoreTypes.h"
+#include "Misc/EnumClassFlags.h"
+
+namespace UE::DerivedData
+{
+
+/**
+ * Flags to control the behavior of build requests.
+ *
+ * The build policy flags can be combined to support a variety of usage patterns. Examples:
+ *
+ * Build(Default): allow cache get; build if missing; return every payload.
+ * Build(SkipCacheGet): never get from the cache; always build; return every payload.
+ * Build(SkipBuild | SkipData): allow cache get; never build; skip every payload.
+ *
+ * BuildPayload(Default): allow cache get; build if missing; return one payload.
+ * BuildPayload(SkipCacheGet): never get from the cache; always build; return one payload.
+ * BuildPayload(SkipBuild | SkipData): allow cache get; never build; skip every payload.
+ */
+enum class EBuildPolicy : uint8
+{
+	/** A value without any flags set. */
+	None            = 0,
+
+	/** Allow local execution. */
+	Local           = 1 << 0,
+	/** Allow remote execution if the function has a registered build worker. */
+	Remote          = 1 << 1,
+
+	/** Skip attempting a cache get before building. */
+	SkipCacheGet    = 1 << 2,
+	/** Skip attempting a cache get before building. */
+	SkipCachePut    = 1 << 3,
+	/** Skip attempting a build when the cache get is missed or skipped. */
+	SkipBuild       = 1 << 4,
+	/** Skip fetching the payload data from the cache. */
+	SkipData        = 1 << 5,
+
+	/** Allow cache get+put, local+remote build when missed or skipped, and fetch the payload(s). */
+	Default         = Local | Remote,
+};
+
+ENUM_CLASS_FLAGS(EBuildPolicy);
+
+} // UE::DerivedData
