DDC: Verify the hash of build inputs after decompression

diff --git a/Engine/Source/Developer/DerivedDataCache/Private/DerivedDataBuildJobContext.cpp b/Engine/Source/Developer/DerivedDataCache/Private/DerivedDataBuildJobContext.cpp
--- a/Engine/Source/Developer/DerivedDataCache/Private/DerivedDataBuildJobContext.cpp
+++ b/Engine/Source/Developer/DerivedDataCache/Private/DerivedDataBuildJobContext.cpp
@@ -5,7 +5,9 @@
 #include "DerivedDataBuildJob.h"
 #include "DerivedDataBuildOutput.h"
 #include "DerivedDataBuildPolicy.h"
+#include "DerivedDataBuildPrivate.h"
 #include "DerivedDataCache.h"
+#include "Hash/Blake3.h"
 #include "Misc/StringBuilder.h"
 #include "UObject/NameTypes.h"
 
@@ -58,7 +60,24 @@ FSharedBuffer FBuildJobContext::GetInput(FStringView Key) const
 {
 	if (const FCompressedBuffer* Input = Inputs.FindByHash(GetTypeHash(Key), Key))
 	{
-		return Input->Decompress();
+		FSharedBuffer Buffer = Input->Decompress();
+		const FBlake3Hash RawHash = FBlake3::HashBuffer(Buffer);
+		if (RawHash == Input->GetRawHash() && Buffer.GetSize() == Input->GetRawSize())
+		{
+			return Buffer;
+		}
+		else
+		{
+			TStringBuilder<32> Category;
+			Category << ImplicitConv<FName>(LogDerivedDataBuild.GetCategoryName());
+			TStringBuilder<256> Error;
+			Error << TEXT("Input '") << Key << TEXT("' was expected to have raw hash ") << Input->GetRawHash()
+				<< TEXT(" and raw size ") << Input->GetRawSize() << TEXT(" but has raw hash ") << RawHash
+				<< TEXT(" and raw size ") << Buffer.GetSize() << TEXT(" after decompression for build of '")
+				<< Job.GetName() << TEXT("' by ") << Job.GetFunction() << TEXT(".");
+			OutputBuilder.AddError(Category, Error);
+			UE_LOG(LogDerivedDataBuild, Error, TEXT("%.*s"), Error.Len(), Error.GetData());
+		}
 	}
 	return FSharedBuffer();
 }
