DDC: Renamed GetFileData to FindFileData to clarify that it can fail

diff --git a/Engine/Source/Developer/DerivedDataCache/Private/DerivedDataBuildWorkerRegistry.cpp b/Engine/Source/Developer/DerivedDataCache/Private/DerivedDataBuildWorkerRegistry.cpp
--- a/Engine/Source/Developer/DerivedDataCache/Private/DerivedDataBuildWorkerRegistry.cpp
+++ b/Engine/Source/Developer/DerivedDataCache/Private/DerivedDataBuildWorkerRegistry.cpp
@@ -35,7 +35,7 @@ public:
 	inline FStringView GetHostPlatform() const final { return HostPlatform; }
 	inline FGuid GetBuildSystemVersion() const final { return BuildSystemVersion; }
 
-	FRequest GetFileData(TConstArrayView<FIoHash> RawHashes, EPriority Priority, FOnBuildWorkerFileDataComplete&& OnComplete) const final;
+	FRequest FindFileData(TConstArrayView<FIoHash> RawHashes, EPriority Priority, FOnBuildWorkerFileDataComplete&& OnComplete) const final;
 
 	void IterateFunctions(TFunctionRef<void (FStringView WorkerName, const FGuid& Version)> Visitor) const final;
 	void IterateFiles(TFunctionRef<void (FStringView WorkerPath, const FIoHash& RawHash, uint64 RawSize)> Visitor) const final;
@@ -73,9 +73,9 @@ void FBuildWorkerInternal::Build()
 	Environment.Sort();
 }
 
-FRequest FBuildWorkerInternal::GetFileData(TConstArrayView<FIoHash> RawHashes, EPriority Priority, FOnBuildWorkerFileDataComplete&& OnComplete) const
+FRequest FBuildWorkerInternal::FindFileData(TConstArrayView<FIoHash> RawHashes, EPriority Priority, FOnBuildWorkerFileDataComplete&& OnComplete) const
 {
-	return Factory->GetFileData(RawHashes, Priority, MoveTemp(OnComplete));
+	return Factory->FindFileData(RawHashes, Priority, MoveTemp(OnComplete));
 }
 
 void FBuildWorkerInternal::IterateFunctions(TFunctionRef<void (FStringView Name, const FGuid& Version)> Visitor) const
diff --git a/Engine/Source/Developer/DerivedDataCache/Public/DerivedDataBuildWorker.h b/Engine/Source/Developer/DerivedDataCache/Public/DerivedDataBuildWorker.h
--- a/Engine/Source/Developer/DerivedDataCache/Public/DerivedDataBuildWorker.h
+++ b/Engine/Source/Developer/DerivedDataCache/Public/DerivedDataBuildWorker.h
@@ -38,7 +38,7 @@ public:
 	virtual FStringView GetPath() const = 0;
 	virtual FStringView GetHostPlatform() const = 0;
 	virtual FGuid GetBuildSystemVersion() const = 0;
-	virtual FRequest GetFileData(TConstArrayView<FIoHash> RawHashes, EPriority Priority, FOnBuildWorkerFileDataComplete&& OnComplete) const = 0;
+	virtual FRequest FindFileData(TConstArrayView<FIoHash> RawHashes, EPriority Priority, FOnBuildWorkerFileDataComplete&& OnComplete) const = 0;
 	virtual void IterateFunctions(TFunctionRef<void (FStringView Name, const FGuid& Version)> Visitor) const = 0;
 	virtual void IterateFiles(TFunctionRef<void (FStringView Path, const FIoHash& RawHash, uint64 RawSize)> Visitor) const = 0;
 	virtual void IterateExecutables(TFunctionRef<void (FStringView Path, const FIoHash& RawHash, uint64 RawSize)> Visitor) const = 0;
@@ -65,7 +65,7 @@ public:
 	virtual ~IBuildWorkerFactory() = default;
 
 	virtual void Build(FBuildWorkerBuilder& Builder) = 0;
-	virtual FRequest GetFileData(TConstArrayView<FIoHash> RawHashes, EPriority Priority, FOnBuildWorkerFileDataComplete&& OnComplete) = 0;
+	virtual FRequest FindFileData(TConstArrayView<FIoHash> RawHashes, EPriority Priority, FOnBuildWorkerFileDataComplete&& OnComplete) = 0;
 
 	/** Returns the name of the build worker factory modular feature. */
 	static FName GetFeatureName()
