DDC: Factored out function name validation

diff --git a/Engine/Source/Developer/DerivedDataCache/Private/DerivedDataBuildAction.cpp b/Engine/Source/Developer/DerivedDataCache/Private/DerivedDataBuildAction.cpp
--- a/Engine/Source/Developer/DerivedDataCache/Private/DerivedDataBuildAction.cpp
+++ b/Engine/Source/Developer/DerivedDataCache/Private/DerivedDataBuildAction.cpp
@@ -34,8 +34,7 @@ public:
 		, BuildSystemVersion(InBuildSystemVersion)
 	{
 		checkf(!Name.IsEmpty(), TEXT("A build action requires a non-empty name."));
-		checkf(!Function.IsEmpty() && Algo::AllOf(Function, FChar::IsAlnum),
-			TEXT("A build function name must be alphanumeric and non-empty for build of '%s' by %s."), *Name, *Function);
+		AssertValidBuildFunctionName(Function, Name);
 	}
 
 	~FBuildActionBuilderInternal() final = default;
@@ -197,7 +196,7 @@ FBuildActionInternal::FBuildActionInternal(FStringView InName, FCbObject&& InAct
 	checkf(!Name.IsEmpty(), TEXT("A build action requires a non-empty name."));
 	Action.MakeOwned();
 	bOutIsValid = Action
-		&& !Function.IsEmpty() && Algo::AllOf(Function, FChar::IsAlnum)
+		&& IsValidBuildFunctionName(Function)
 		&& FunctionVersion.IsValid()
 		&& BuildSystemVersion.IsValid()
 		&& Algo::AllOf(Action.FindView("Constants"_ASV),
diff --git a/Engine/Source/Developer/DerivedDataCache/Private/DerivedDataBuildDefinition.cpp b/Engine/Source/Developer/DerivedDataCache/Private/DerivedDataBuildDefinition.cpp
--- a/Engine/Source/Developer/DerivedDataCache/Private/DerivedDataBuildDefinition.cpp
+++ b/Engine/Source/Developer/DerivedDataCache/Private/DerivedDataBuildDefinition.cpp
@@ -26,8 +26,7 @@ public:
 		, Function(InFunction)
 	{
 		checkf(!Name.IsEmpty(), TEXT("A build definition requires a non-empty name."));
-		checkf(!Function.IsEmpty() && Algo::AllOf(Function, FChar::IsAlnum),
-			TEXT("A build function name must be alphanumeric and non-empty for build of '%s' by %s."), *Name, *Function);
+		AssertValidBuildFunctionName(Function, Name);
 	}
 
 	~FBuildDefinitionBuilderInternal() final = default;
@@ -261,7 +260,7 @@ FBuildDefinitionInternal::FBuildDefinitionInternal(FStringView InName, FCbObject
 	checkf(!Name.IsEmpty(), TEXT("A build definition requires a non-empty name."));
 	Definition.MakeOwned();
 	bOutIsValid = Definition
-		&& !Function.IsEmpty() && Algo::AllOf(Function, FChar::IsAlnum)
+		&& IsValidBuildFunctionName(Function)
 		&& Algo::AllOf(Definition.FindView("Constants"_ASV),
 			[](FCbFieldView Field) { return Field.GetName().Len() > 0 && Field.IsObject(); })
 		&& Algo::AllOf(Definition.FindView("Inputs"_ASV).AsObjectView().FindView("Builds"_ASV), [](FCbFieldView Field)
diff --git a/Engine/Source/Developer/DerivedDataCache/Private/DerivedDataBuildOutput.cpp b/Engine/Source/Developer/DerivedDataCache/Private/DerivedDataBuildOutput.cpp
--- a/Engine/Source/Developer/DerivedDataCache/Private/DerivedDataBuildOutput.cpp
+++ b/Engine/Source/Developer/DerivedDataCache/Private/DerivedDataBuildOutput.cpp
@@ -52,8 +52,7 @@ public:
 		, bHasDiagnostics(false)
 	{
 		checkf(!Name.IsEmpty(), TEXT("A build output requires a non-empty name."));
-		checkf(!Function.IsEmpty() && Algo::AllOf(Function, FChar::IsAlnum),
-			TEXT("A build function name must be alphanumeric and non-empty for build of '%s' by %s."), *Name, *Function);
+		AssertValidBuildFunctionName(Function, Name);
 		DiagnosticsWriter.BeginArray();
 	}
 
@@ -142,8 +141,7 @@ FBuildOutputInternal::FBuildOutputInternal(FStringView InName, FStringView InFun
 {
 	bOutIsValid = false;
 	checkf(!Name.IsEmpty(), TEXT("A build output requires a non-empty name."));
-	checkf(!Function.IsEmpty() && Algo::AllOf(Function, FChar::IsAlnum),
-		TEXT("A build function name must be alphanumeric and non-empty for build of '%s' by %s."), *Name, *Function);
+	AssertValidBuildFunctionName(Function, Name);
 	Meta = InOutput["Meta"_ASV].AsObject();
 	Meta.MakeOwned();
 	for (FCbFieldView PayloadField : InOutput["Payloads"_ASV])
@@ -169,8 +167,7 @@ FBuildOutputInternal::FBuildOutputInternal(FStringView InName, FStringView InFun
 	, Payloads(InOutput.GetAttachmentPayloads())
 {
 	checkf(!Name.IsEmpty(), TEXT("A build output requires a non-empty name."));
-	checkf(!Function.IsEmpty() && Algo::AllOf(Function, FChar::IsAlnum),
-		TEXT("A build function name must be alphanumeric and non-empty for build of '%s' by %s."), *Name, *Function);
+	AssertValidBuildFunctionName(Function, Name);
 	if (FSharedBuffer Buffer = InOutput.GetValue())
 	{
 		Diagnostics = FCbObject(MoveTemp(Buffer))["Diagnostics"_ASV];
@@ -258,7 +255,7 @@ void FBuildOutputInternal::Save(FCacheRecordBuilder& RecordBuilder) const
 
 bool FBuildOutputInternal::IsValid() const
 {
-	return !Function.IsEmpty() && Algo::AllOf(Function, FChar::IsAlnum)
+	return IsValidBuildFunctionName(Function)
 		&& Algo::AllOf(Payloads, [](const FPayload& Payload) { return !Payload.GetRawHash().IsZero(); })
 		&& (!Diagnostics || Diagnostics.IsArray())
 		&& Algo::AllOf(Diagnostics.CreateViewIterator(), [](FCbFieldView Field)
diff --git a/Engine/Source/Developer/DerivedDataCache/Private/DerivedDataBuildPrivate.h b/Engine/Source/Developer/DerivedDataCache/Private/DerivedDataBuildPrivate.h
--- a/Engine/Source/Developer/DerivedDataCache/Private/DerivedDataBuildPrivate.h
+++ b/Engine/Source/Developer/DerivedDataCache/Private/DerivedDataBuildPrivate.h
@@ -2,8 +2,10 @@
 
 #pragma once
 
+#include "Algo/AllOf.h"
 #include "Containers/StringFwd.h"
 #include "Logging/LogMacros.h"
+#include "Misc/Char.h"
 
 class FCbObject;
 struct FGuid;
@@ -76,4 +78,18 @@ TRequest<IBuildJob> CreateBuildJob(ICache& Cache, IBuild& BuildSystem, IBuildInp
 TRequest<IBuildJob> CreateBuildJob(ICache& Cache, IBuild& BuildSystem, IBuildInputResolver* InputResolver, const FBuildDefinition& Definition);
 TRequest<IBuildJob> CreateBuildJob(ICache& Cache, IBuild& BuildSystem, IBuildInputResolver* InputResolver, const FBuildAction& Action, const FOptionalBuildInputs& Inputs);
 
+///////////////////////////////////////////////////////////////////////////////////////////////////////////////////////
+
+inline bool IsValidBuildFunctionName(FStringView Function)
+{
+	return !Function.IsEmpty() && Algo::AllOf(Function, FChar::IsAlnum);
+}
+
+inline void AssertValidBuildFunctionName(FStringView Function, FStringView Name)
+{
+	checkf(IsValidBuildFunctionName(Function),
+		TEXT("A build function name must be alphanumeric and non-empty for build of '%.*s' by '%.*s'."),
+		Name.Len(), Name.GetData(), Function.Len(), Function.GetData());
+}
+
 } // UE::DerivedData::Private
