DDC: Factored out cache bucket name validation and aligned it more closely with build function name validation

diff --git a/Engine/Source/Developer/DerivedDataCache/Private/DerivedDataBuildPrivate.h b/Engine/Source/Developer/DerivedDataCache/Private/DerivedDataBuildPrivate.h
--- a/Engine/Source/Developer/DerivedDataCache/Private/DerivedDataBuildPrivate.h
+++ b/Engine/Source/Developer/DerivedDataCache/Private/DerivedDataBuildPrivate.h
@@ -2,10 +2,9 @@
 
 #pragma once
 
-#include "Algo/AllOf.h"
 #include "Containers/StringFwd.h"
 #include "Logging/LogMacros.h"
-#include "Misc/Char.h"
+#include "Misc/AsciiSet.h"
 
 class FCbObject;
 struct FGuid;
@@ -82,7 +81,8 @@ TRequest<IBuildJob> CreateBuildJob(ICache& Cache, IBuild& BuildSystem, IBuildInp
 
 inline bool IsValidBuildFunctionName(FStringView Function)
 {
-	return !Function.IsEmpty() && Algo::AllOf(Function, FChar::IsAlnum);
+	constexpr FAsciiSet Valid("ABCDEFGHIJKLMNOPQRSTUVWXYZabcdefghijklmnopqrstuvwxyz0123456789");
+	return !Function.IsEmpty() && FAsciiSet::HasOnly(Function, Valid);
 }
 
 inline void AssertValidBuildFunctionName(FStringView Function, FStringView Name)
diff --git a/Engine/Source/Developer/DerivedDataCache/Private/DerivedDataCacheKey.cpp b/Engine/Source/Developer/DerivedDataCache/Private/DerivedDataCacheKey.cpp
--- a/Engine/Source/Developer/DerivedDataCache/Private/DerivedDataCacheKey.cpp
+++ b/Engine/Source/Developer/DerivedDataCache/Private/DerivedDataCacheKey.cpp
@@ -4,6 +4,7 @@
 
 #include "Algo/AllOf.h"
 #include "Containers/Set.h"
+#include "DerivedDataCachePrivate.h"
 #include "Math/UnrealMathUtility.h"
 #include "Misc/ScopeRWLock.h"
 
@@ -97,10 +98,7 @@ inline FCacheBucket FCacheBuckets::FindOrAdd(FStringView Name)
 		return *Bucket;
 	}
 
-	checkf(FCString::IsPureAnsi(Name.GetData(), Name.Len()) && Algo::AllOf(Name, FChar::IsAlnum) &&
-		!Name.IsEmpty() && Name.Len() < 256,
-		TEXT("Invalid cache bucket name '%.*s' must be alphanumeric, non-empty, and contain fewer than 256 code units."),
-		Name.Len(), Name.GetData());
+	AssertValidCacheBucketName(Name);
 
 	FCacheBucketOwner LocalBucket(Name);
 	FWriteScopeLock WriteLock(Lock);
diff --git a/Engine/Source/Developer/DerivedDataCache/Private/DerivedDataCachePrivate.h b/Engine/Source/Developer/DerivedDataCache/Private/DerivedDataCachePrivate.h
--- a/Engine/Source/Developer/DerivedDataCache/Private/DerivedDataCachePrivate.h
+++ b/Engine/Source/Developer/DerivedDataCache/Private/DerivedDataCachePrivate.h
@@ -3,12 +3,14 @@
 #pragma once
 
 #include "Containers/StringFwd.h"
+#include "Misc/AsciiSet.h"
 
 class FCbPackage;
 class FDerivedDataCacheInterface;
 enum class ECachePolicy : uint8;
 
 namespace UE::DerivedData { class FCacheBucket; }
+namespace UE::DerivedData { class FCacheRecord; }
 namespace UE::DerivedData { class FCacheRecordBuilder; }
 namespace UE::DerivedData { class FOptionalCacheRecord; }
 namespace UE::DerivedData { struct FCacheKey; }
@@ -27,4 +29,19 @@ FCacheRecordBuilder CreateCacheRecordBuilder(const FCacheKey& Key);
 FCbPackage SaveCacheRecord(const FCacheRecord& Record);
 FOptionalCacheRecord LoadCacheRecord(const FCbPackage& Package);
 
+///////////////////////////////////////////////////////////////////////////////////////////////////////////////////////
+
+inline bool IsValidCacheBucketName(FStringView Name)
+{
+	constexpr FAsciiSet Valid("ABCDEFGHIJKLMNOPQRSTUVWXYZabcdefghijklmnopqrstuvwxyz0123456789");
+	return !Name.IsEmpty() && Name.Len() < 256 && FAsciiSet::HasOnly(Name, Valid);
+}
+
+inline void AssertValidCacheBucketName(FStringView Name)
+{
+	checkf(IsValidCacheBucketName(Name),
+		TEXT("A cache bucket name must be alphanumeric, non-empty, and contain fewer than 256 code units. ")
+		TEXT("Name: '%.*s'"), Name.Len(), Name.GetData());
+}
+
 } // UE::DerivedData::Private
diff --git a/Engine/Source/Developer/DerivedDataCache/Private/DerivedDataCacheRecord.cpp b/Engine/Source/Developer/DerivedDataCache/Private/DerivedDataCacheRecord.cpp
--- a/Engine/Source/Developer/DerivedDataCache/Private/DerivedDataCacheRecord.cpp
+++ b/Engine/Source/Developer/DerivedDataCache/Private/DerivedDataCacheRecord.cpp
@@ -327,8 +327,7 @@ FOptionalCacheRecord LoadCacheRecord(const FCbPackage& Package)
 	FCbObjectView KeyObject = RecordObject["Key"_ASV].AsObjectView();
 	auto TrySetBucketName = [](FStringView Name, FCacheKey& Key)
 	{
-		if (FCString::IsPureAnsi(Name.GetData(), Name.Len())
-			&& Algo::AllOf(Name, FChar::IsAlnum) && !Name.IsEmpty() && Name.Len() < 256)
+		if (IsValidCacheBucketName(Name))
 		{
 			Key.Bucket = CreateCacheBucket(Name);
 			return true;
