DDC: Reorganized functions in the build job

The code is now organized based on what states do, rather than what order they typically execute in. This now matches the function order in the scheduler and the job interface. There are no changes aside from reordering of functions.

#rnx

diff --git a/Engine/Source/Developer/DerivedDataCache/Private/DerivedDataBuildJob.cpp b/Engine/Source/Developer/DerivedDataCache/Private/DerivedDataBuildJob.cpp
--- a/Engine/Source/Developer/DerivedDataCache/Private/DerivedDataBuildJob.cpp
+++ b/Engine/Source/Developer/DerivedDataCache/Private/DerivedDataBuildJob.cpp
@@ -126,32 +126,35 @@
 
 private:
 	void BeginJob();
-	void Configure();
 	void EndJob();
 
+	void Configure();
+
+	void EnterCacheQuery();
+	void EnterCacheStore();
 	void EnterResolveKey();
 	void EnterResolveInputMeta();
-	void EnterCacheQuery();
 	void EnterResolveInputData();
 	void EnterExecuteRemote();
 	void EnterExecuteLocal();
-	void EnterCacheStore();
 
+	void BeginCacheQuery() final;
+	void BeginCacheStore() final;
 	void BeginResolveKey() final;
 	void BeginResolveInputMeta() final;
-	void BeginCacheQuery() final;
 	void BeginResolveInputData() final;
 	void BeginExecuteRemote() final;
 	void BeginExecuteLocal() final;
-	void BeginCacheStore() final;
 
+	void EndCacheQuery(FCacheGetCompleteParams&& Params);
+	void EndCacheStore(FCachePutCompleteParams&& Params);
 	void EndResolveKey(FBuildKeyResolvedParams&& Params);
 	void EndResolveInputMeta(FBuildInputMetaResolvedParams&& Params);
-	void EndCacheQuery(FCacheGetCompleteParams&& Params);
 	void EndResolveInputData(FBuildInputDataResolvedParams&& Params);
 	void EndExecuteRemote(FBuildWorkerActionCompleteParams&& Params);
 	void EndExecuteLocal();
-	void EndCacheStore(FCachePutCompleteParams&& Params);
 
+	void CreateAction(TConstArrayView<FBuildInputMetaByKey> InputMeta);
+
 	void SetDefinition(FBuildDefinition&& Definition);
 	void SetAction(FBuildAction&& Action);
@@ -159,8 +162,6 @@
 	void SetOutput(const FBuildOutput& Output) final;
 	void SetOutputNoCheck(FBuildOutput&& Output);
 
-	void CreateAction(TConstArrayView<FBuildInputMetaByKey> InputMeta);
-
 	/** Terminate the job and send the error to the output complete callback. */
 	void CompleteWithError(FStringView Error);
 
@@ -467,130 +468,20 @@
 	Scheduler->BeginJob(this);
 }
 
-///////////////////////////////////////////////////////////////////////////////////////////////////////////////////////
+void FBuildJob::EndJob()
+{
+	TRequest Self(this);
 
-void FBuildJob::EnterResolveKey()
-{
-	if (Action)
+	if (Event)
 	{
-		return AdvanceToState(EBuildJobState::CacheQuery);
+		Event->Trigger();
 	}
-	if (Definition)
-	{
-		return AdvanceToState(EBuildJobState::ResolveInputMeta);
-	}
-	if (DefinitionKey != FBuildKey::Empty)
-	{
-		return CompleteWithError(TEXT("Failed to resolve null key."_SV));
-	}
-	if (!InputResolver)
-	{
-		return CompleteWithError(TEXT("Failed to resolve key due to null input resolver."_SV));
-	}
-	Scheduler->DispatchResolveKey(this);
-}
+	Scheduler->EndJob(this);
+	Release();
 
-void FBuildJob::BeginResolveKey()
-{
-	if (CanExecuteState(EBuildJobState::ResolveKey))
-	{
-		return AdvanceToState(EBuildJobState::ResolveKeyWait, InputResolver->ResolveKey(DefinitionKey,
-			[this](FBuildKeyResolvedParams&& Params) { EndResolveKey(MoveTemp(Params)); }));
-	}
-}
-
-void FBuildJob::EndResolveKey(FBuildKeyResolvedParams&& Params)
-{
-	if (CanExecuteState(EBuildJobState::ResolveKeyWait))
-	{
-		if (Params.Status == EStatus::Ok && Params.Definition)
-		{
-			SetDefinition(MoveTemp(Params.Definition).Get());
-		}
-		else
-		{
-			CompleteWithError(WriteToString<128>(TEXT("Failed to resolve key "_SV), Params.Key, TEXT("."_SV)));
-		}
-	}
-}
-
-void FBuildJob::SetDefinition(FBuildDefinition&& InDefinition)
-{
-	Name = InDefinition.GetName();
-	FunctionName = InDefinition.GetFunction();
-	checkf(Definition.IsNull(), TEXT("Job already has a definition for build of '%s' by %s."), *Name, *FunctionName);
-	checkf(State == EBuildJobState::ResolveKey || State == EBuildJobState::ResolveKeyWait,
-		TEXT("Job is not expecting a definition in state %s for build of '%s' by %s"),
-		LexToString(State), *Name, *FunctionName);
-	OutputBuilder = BuildSystem.CreateOutput(Name, FunctionName);
-	Definition = MoveTemp(InDefinition);
-	return AdvanceToState(EBuildJobState::ResolveInputMeta);
+	FWriteScopeLock WriteLock(Lock);
+	Scheduler = nullptr;
+	InputResolver = nullptr;
 }
 
 ///////////////////////////////////////////////////////////////////////////////////////////////////////////////////////
-
-void FBuildJob::EnterResolveInputMeta()
-{
-	if (!Definition.Get().HasInputs())
-	{
-		return CreateAction({});
-	}
-	if (!InputResolver)
-	{
-		return CompleteWithError(TEXT("Failed to resolve input metadata due to null input resolver."_SV));
-	}
-	Scheduler->DispatchResolveInputMeta(this);
-}
-
-void FBuildJob::BeginResolveInputMeta()
-{
-	if (CanExecuteState(EBuildJobState::ResolveInputMeta))
-	{
-		return AdvanceToState(EBuildJobState::ResolveInputMetaWait, InputResolver->ResolveInputMeta(Definition.Get(),
-			Priority, [this](FBuildInputMetaResolvedParams&& Params) { EndResolveInputMeta(MoveTemp(Params)); }));
-	}
-}
-
-void FBuildJob::EndResolveInputMeta(FBuildInputMetaResolvedParams&& Params)
-{
-	if (CanExecuteState(EBuildJobState::ResolveInputMetaWait))
-	{
-		if (Params.Status == EStatus::Ok)
-		{
-			CreateAction(Params.Inputs);
-		}
-		else
-		{
-			CompleteWithError(TEXT("Failed to resolve input metadata."_SV));
-		}
-	}
-}
-
-void FBuildJob::CreateAction(TConstArrayView<FBuildInputMetaByKey> InputMeta)
-{
-	const FBuildDefinition& LocalDefinition = Definition.Get();
-	FBuildActionBuilder Builder = BuildSystem.CreateAction(Name, FunctionName);
-	LocalDefinition.IterateConstants([this, &Builder](FStringView Key, FCbObject&& Value)
-	{
-		Builder.AddConstant(Key, Value);
-	});
-	for (const FBuildInputMetaByKey& Input : InputMeta)
-	{
-		Builder.AddInput(Input.Key, Input.RawHash, Input.RawSize);
-	}
-	SetAction(Builder.Build());
-}
-
-void FBuildJob::SetAction(FBuildAction&& InAction)
-{
-	checkf(Action.IsNull(), TEXT("Job already has an action for build of '%s' by %s."), *Name, *FunctionName);
-	checkf(State == EBuildJobState::ResolveKey ||
-		State == EBuildJobState::ResolveInputMeta || State == EBuildJobState::ResolveInputMetaWait,
-		TEXT("Job is not expecting an action in state %s for build of '%s' by %s"),
-		LexToString(State), *Name, *FunctionName);
-	ActionKey = InAction.GetKey();
-	Action = MoveTemp(InAction);
-	return AdvanceToState(EBuildJobState::CacheQuery);
-}
-
-///////////////////////////////////////////////////////////////////////////////////////////////////////////////////////
@@ -677,6 +568,129 @@
 
 ///////////////////////////////////////////////////////////////////////////////////////////////////////////////////////
 
+void FBuildJob::EnterCacheStore()
+{
+	const EBuildPolicy BuildPolicy = Context->GetBuildPolicy();
+	ECachePolicy CachePolicy = Context->GetCachePolicy();
+	if (bIsCached ||
+		!EnumHasAnyFlags(CachePolicy, ECachePolicy::Store) ||
+		EnumHasAnyFlags(BuildPolicy, EBuildPolicy::SkipCachePut) ||
+		Output.Get().HasError())
+	{
+		return AdvanceToState(EBuildJobState::Complete);
+	}
+	Scheduler->DispatchCacheStore(this);
+}
+
+void FBuildJob::BeginCacheStore()
+{
+	if (CanExecuteState(EBuildJobState::CacheStore))
+	{
+		FCacheRecordBuilder RecordBuilder = Cache.CreateRecord(Context->GetCacheKey());
+		Output.Get().Save(RecordBuilder);
+		return AdvanceToState(EBuildJobState::CacheStoreWait,
+			Cache.Put({RecordBuilder.Build()}, Name, Context->GetCachePolicy(), FMath::Min(Priority, EPriority::Highest),
+				[this](FCachePutCompleteParams&& Params) { EndCacheStore(MoveTemp(Params)); }));
+	}
+}
+
+void FBuildJob::EndCacheStore(FCachePutCompleteParams&& Params)
+{
+	if (CanExecuteState(EBuildJobState::CacheStoreWait))
+	{
+		bIsCached = Params.Status == EStatus::Ok;
+		return AdvanceToState(EBuildJobState::Complete);
+	}
+}
+
+///////////////////////////////////////////////////////////////////////////////////////////////////////////////////////
+
+void FBuildJob::EnterResolveKey()
+{
+	if (Action)
+	{
+		return AdvanceToState(EBuildJobState::CacheQuery);
+	}
+	if (Definition)
+	{
+		return AdvanceToState(EBuildJobState::ResolveInputMeta);
+	}
+	if (DefinitionKey != FBuildKey::Empty)
+	{
+		return CompleteWithError(TEXT("Failed to resolve null key."_SV));
+	}
+	if (!InputResolver)
+	{
+		return CompleteWithError(TEXT("Failed to resolve key due to null input resolver."_SV));
+	}
+	Scheduler->DispatchResolveKey(this);
+}
+
+void FBuildJob::BeginResolveKey()
+{
+	if (CanExecuteState(EBuildJobState::ResolveKey))
+	{
+		return AdvanceToState(EBuildJobState::ResolveKeyWait, InputResolver->ResolveKey(DefinitionKey,
+			[this](FBuildKeyResolvedParams&& Params) { EndResolveKey(MoveTemp(Params)); }));
+	}
+}
+
+void FBuildJob::EndResolveKey(FBuildKeyResolvedParams&& Params)
+{
+	if (CanExecuteState(EBuildJobState::ResolveKeyWait))
+	{
+		if (Params.Status == EStatus::Ok && Params.Definition)
+		{
+			SetDefinition(MoveTemp(Params.Definition).Get());
+		}
+		else
+		{
+			CompleteWithError(WriteToString<128>(TEXT("Failed to resolve key "_SV), Params.Key, TEXT("."_SV)));
+		}
+	}
+}
+
+///////////////////////////////////////////////////////////////////////////////////////////////////////////////////////
+
+void FBuildJob::EnterResolveInputMeta()
+{
+	if (!Definition.Get().HasInputs())
+	{
+		return CreateAction({});
+	}
+	if (!InputResolver)
+	{
+		return CompleteWithError(TEXT("Failed to resolve input metadata due to null input resolver."_SV));
+	}
+	Scheduler->DispatchResolveInputMeta(this);
+}
+
+void FBuildJob::BeginResolveInputMeta()
+{
+	if (CanExecuteState(EBuildJobState::ResolveInputMeta))
+	{
+		return AdvanceToState(EBuildJobState::ResolveInputMetaWait, InputResolver->ResolveInputMeta(Definition.Get(),
+			Priority, [this](FBuildInputMetaResolvedParams&& Params) { EndResolveInputMeta(MoveTemp(Params)); }));
+	}
+}
+
+void FBuildJob::EndResolveInputMeta(FBuildInputMetaResolvedParams&& Params)
+{
+	if (CanExecuteState(EBuildJobState::ResolveInputMetaWait))
+	{
+		if (Params.Status == EStatus::Ok)
+		{
+			CreateAction(Params.Inputs);
+		}
+		else
+		{
+			CompleteWithError(TEXT("Failed to resolve input metadata."_SV));
+		}
+	}
+}
+
+///////////////////////////////////////////////////////////////////////////////////////////////////////////////////////
+
 void FBuildJob::EnterResolveInputData()
 {
 	// Identify missing inputs when resolving every input for the action.
@@ -753,28 +767,6 @@
 	}
 }
 
-void FBuildJob::SetInputs(FBuildInputs&& InInputs)
-{
-	EBuildJobState ExecuteState;
-	switch (State)
-	{
-	case EBuildJobState::ResolveRemoteInputData:
-	case EBuildJobState::ResolveRemoteInputDataWait:
-		ExecuteState = EBuildJobState::ExecuteRemoteRetry;
-		break;
-	case EBuildJobState::ResolveInputData:
-	case EBuildJobState::ResolveInputDataWait:
-		ExecuteState = EBuildJobState::ExecuteLocal;
-		break;
-	default:
-		checkf(false, TEXT("Job is not expecting inputs in state %s for build of '%s' by %s"),
-			LexToString(State), *Name, *FunctionName);
-		return;
-	}
-	Inputs = MoveTemp(InInputs);
-	return AdvanceToState(ExecuteState);
-}
-
 ///////////////////////////////////////////////////////////////////////////////////////////////////////////////////////
 
 void FBuildJob::EnterExecuteRemote()
@@ -927,6 +919,72 @@
 	}
 }
 
+///////////////////////////////////////////////////////////////////////////////////////////////////////////////////////
+
+void FBuildJob::CreateAction(TConstArrayView<FBuildInputMetaByKey> InputMeta)
+{
+	const FBuildDefinition& LocalDefinition = Definition.Get();
+	FBuildActionBuilder Builder = BuildSystem.CreateAction(Name, FunctionName);
+	LocalDefinition.IterateConstants([this, &Builder](FStringView Key, FCbObject&& Value)
+	{
+		Builder.AddConstant(Key, Value);
+	});
+	for (const FBuildInputMetaByKey& Input : InputMeta)
+	{
+		Builder.AddInput(Input.Key, Input.RawHash, Input.RawSize);
+	}
+	SetAction(Builder.Build());
+}
+
+///////////////////////////////////////////////////////////////////////////////////////////////////////////////////////
+
+void FBuildJob::SetDefinition(FBuildDefinition&& InDefinition)
+{
+	Name = InDefinition.GetName();
+	FunctionName = InDefinition.GetFunction();
+	checkf(Definition.IsNull(), TEXT("Job already has a definition for build of '%s' by %s."), *Name, *FunctionName);
+	checkf(State == EBuildJobState::ResolveKey || State == EBuildJobState::ResolveKeyWait,
+		TEXT("Job is not expecting a definition in state %s for build of '%s' by %s"),
+		LexToString(State), *Name, *FunctionName);
+	OutputBuilder = BuildSystem.CreateOutput(Name, FunctionName);
+	Definition = MoveTemp(InDefinition);
+	return AdvanceToState(EBuildJobState::ResolveInputMeta);
+}
+
+void FBuildJob::SetAction(FBuildAction&& InAction)
+{
+	checkf(Action.IsNull(), TEXT("Job already has an action for build of '%s' by %s."), *Name, *FunctionName);
+	checkf(State == EBuildJobState::ResolveKey ||
+		State == EBuildJobState::ResolveInputMeta || State == EBuildJobState::ResolveInputMetaWait,
+		TEXT("Job is not expecting an action in state %s for build of '%s' by %s"),
+		LexToString(State), *Name, *FunctionName);
+	ActionKey = InAction.GetKey();
+	Action = MoveTemp(InAction);
+	return AdvanceToState(EBuildJobState::CacheQuery);
+}
+
+void FBuildJob::SetInputs(FBuildInputs&& InInputs)
+{
+	EBuildJobState ExecuteState;
+	switch (State)
+	{
+	case EBuildJobState::ResolveRemoteInputData:
+	case EBuildJobState::ResolveRemoteInputDataWait:
+		ExecuteState = EBuildJobState::ExecuteRemoteRetry;
+		break;
+	case EBuildJobState::ResolveInputData:
+	case EBuildJobState::ResolveInputDataWait:
+		ExecuteState = EBuildJobState::ExecuteLocal;
+		break;
+	default:
+		checkf(false, TEXT("Job is not expecting inputs in state %s for build of '%s' by %s"),
+			LexToString(State), *Name, *FunctionName);
+		return;
+	}
+	Inputs = MoveTemp(InInputs);
+	return AdvanceToState(ExecuteState);
+}
+
 void FBuildJob::SetOutput(const FBuildOutput& InOutput)
 {
 	checkf(State == EBuildJobState::ResolveKey || State == EBuildJobState::ResolveInputMeta ||
@@ -955,58 +1013,3 @@
 }
 
 ///////////////////////////////////////////////////////////////////////////////////////////////////////////////////////
-
-void FBuildJob::EnterCacheStore()
-{
-	const EBuildPolicy BuildPolicy = Context->GetBuildPolicy();
-	ECachePolicy CachePolicy = Context->GetCachePolicy();
-	if (bIsCached ||
-		!EnumHasAnyFlags(CachePolicy, ECachePolicy::Store) ||
-		EnumHasAnyFlags(BuildPolicy, EBuildPolicy::SkipCachePut) ||
-		Output.Get().HasError())
-	{
-		return AdvanceToState(EBuildJobState::Complete);
-	}
-	Scheduler->DispatchCacheStore(this);
-}
-
-void FBuildJob::BeginCacheStore()
-{
-	if (CanExecuteState(EBuildJobState::CacheStore))
-	{
-		FCacheRecordBuilder RecordBuilder = Cache.CreateRecord(Context->GetCacheKey());
-		Output.Get().Save(RecordBuilder);
-		return AdvanceToState(EBuildJobState::CacheStoreWait,
-			Cache.Put({RecordBuilder.Build()}, Name, Context->GetCachePolicy(), FMath::Min(Priority, EPriority::Highest),
-				[this](FCachePutCompleteParams&& Params) { EndCacheStore(MoveTemp(Params)); }));
-	}
-}
-
-void FBuildJob::EndCacheStore(FCachePutCompleteParams&& Params)
-{
-	if (CanExecuteState(EBuildJobState::CacheStoreWait))
-	{
-		bIsCached = Params.Status == EStatus::Ok;
-		return AdvanceToState(EBuildJobState::Complete);
-	}
-}
-
-///////////////////////////////////////////////////////////////////////////////////////////////////////////////////////
-
-void FBuildJob::EndJob()
-{
-	TRequest Self(this);
-
-	if (Event)
-	{
-		Event->Trigger();
-	}
-	Scheduler->EndJob(this);
-	Release();
-
-	FWriteScopeLock WriteLock(Lock);
-	Scheduler = nullptr;
-	InputResolver = nullptr;
-}
-
-///////////////////////////////////////////////////////////////////////////////////////////////////////////////////////
