DDC: Extended the cache record builder to take composite buffers

#rnx

diff --git a/Engine/Source/Developer/DerivedDataCache/Private/DerivedDataCacheRecord.cpp b/Engine/Source/Developer/DerivedDataCache/Private/DerivedDataCacheRecord.cpp
--- a/Engine/Source/Developer/DerivedDataCache/Private/DerivedDataCacheRecord.cpp
+++ b/Engine/Source/Developer/DerivedDataCache/Private/DerivedDataCacheRecord.cpp
@@ -31,9 +31,11 @@ public:
 
 	void SetMeta(FCbObject&& Meta) final;
 
+	FPayloadId SetValue(const FCompositeBuffer& Buffer, const FPayloadId& Id) final;
 	FPayloadId SetValue(const FSharedBuffer& Buffer, const FPayloadId& Id) final;
 	FPayloadId SetValue(const FPayload& Payload) final;
 
+	FPayloadId AddAttachment(const FCompositeBuffer& Buffer, const FPayloadId& Id) final;
 	FPayloadId AddAttachment(const FSharedBuffer& Buffer, const FPayloadId& Id) final;
 	FPayloadId AddAttachment(const FPayload& Payload) final;
 
@@ -94,9 +96,9 @@ private:
 
 ///////////////////////////////////////////////////////////////////////////////////////////////////////////////////////
 
-static FPayloadId GetOrCreatePayloadId(const FPayloadId& Id, const FIoHash& RawHash)
+static FPayloadId GetOrCreatePayloadId(const FPayloadId& Id, const FCompressedBuffer& Buffer)
 {
-	return Id.IsValid() ? Id : FPayloadId::FromHash(RawHash);
+	return Id.IsValid() ? Id : FPayloadId::FromHash(Buffer.GetRawHash());
 }
 
 ///////////////////////////////////////////////////////////////////////////////////////////////////////////////////////
@@ -196,10 +198,17 @@ void FCacheRecordBuilderInternal::SetMeta(FCbObject&& InMeta)
 	Meta.MakeOwned();
 }
 
+FPayloadId FCacheRecordBuilderInternal::SetValue(const FCompositeBuffer& Buffer, const FPayloadId& Id)
+{
+	FCompressedBuffer CompressedBuffer = FCompressedBuffer::Compress(NAME_Default, Buffer);
+	const FPayloadId ValueId = GetOrCreatePayloadId(Id, CompressedBuffer);
+	return SetValue(FPayload(ValueId, MoveTemp(CompressedBuffer)));
+}
+
 FPayloadId FCacheRecordBuilderInternal::SetValue(const FSharedBuffer& Buffer, const FPayloadId& Id)
 {
 	FCompressedBuffer CompressedBuffer = FCompressedBuffer::Compress(NAME_Default, Buffer);
-	const FPayloadId ValueId = GetOrCreatePayloadId(Id, CompressedBuffer.GetRawHash());
+	const FPayloadId ValueId = GetOrCreatePayloadId(Id, CompressedBuffer);
 	return SetValue(FPayload(ValueId, MoveTemp(CompressedBuffer)));
 }
 
@@ -217,10 +226,17 @@ FPayloadId FCacheRecordBuilderInternal::SetValue(const FPayload& Payload)
 	return Id;
 }
 
+FPayloadId FCacheRecordBuilderInternal::AddAttachment(const FCompositeBuffer& Buffer, const FPayloadId& Id)
+{
+	FCompressedBuffer CompressedBuffer = FCompressedBuffer::Compress(NAME_Default, Buffer);
+	const FPayloadId AttachmentId = GetOrCreatePayloadId(Id, CompressedBuffer);
+	return AddAttachment(FPayload(AttachmentId, MoveTemp(CompressedBuffer)));
+}
+
 FPayloadId FCacheRecordBuilderInternal::AddAttachment(const FSharedBuffer& Buffer, const FPayloadId& Id)
 {
 	FCompressedBuffer CompressedBuffer = FCompressedBuffer::Compress(NAME_Default, Buffer);
-	const FPayloadId AttachmentId = GetOrCreatePayloadId(Id, CompressedBuffer.GetRawHash());
+	const FPayloadId AttachmentId = GetOrCreatePayloadId(Id, CompressedBuffer);
 	return AddAttachment(FPayload(AttachmentId, MoveTemp(CompressedBuffer)));
 }
 
diff --git a/Engine/Source/Developer/DerivedDataCache/Public/DerivedDataCacheRecord.h b/Engine/Source/Developer/DerivedDataCache/Public/DerivedDataCacheRecord.h
--- a/Engine/Source/Developer/DerivedDataCache/Public/DerivedDataCacheRecord.h
+++ b/Engine/Source/Developer/DerivedDataCache/Public/DerivedDataCacheRecord.h
@@ -44,8 +44,10 @@ class ICacheRecordBuilderInternal
 public:
 	virtual ~ICacheRecordBuilderInternal() = default;
 	virtual void SetMeta(FCbObject&& Meta) = 0;
+	virtual FPayloadId SetValue(const FCompositeBuffer& Buffer, const FPayloadId& Id) = 0;
 	virtual FPayloadId SetValue(const FSharedBuffer& Buffer, const FPayloadId& Id) = 0;
 	virtual FPayloadId SetValue(const FPayload& Payload) = 0;
+	virtual FPayloadId AddAttachment(const FCompositeBuffer& Buffer, const FPayloadId& Id) = 0;
 	virtual FPayloadId AddAttachment(const FSharedBuffer& Buffer, const FPayloadId& Id) = 0;
 	virtual FPayloadId AddAttachment(const FPayload& Payload) = 0;
 	virtual FCacheRecord Build() = 0;
@@ -148,6 +150,10 @@ public:
 	 *                 the buffer will be hashed to create an ID.
 	 * @return The ID that was provided or created.
 	 */
+	inline FPayloadId SetValue(const FCompositeBuffer& Buffer, const FPayloadId& Id = FPayloadId())
+	{
+		return RecordBuilder->SetValue(Buffer, Id);
+	}
 	inline FPayloadId SetValue(const FSharedBuffer& Buffer, const FPayloadId& Id = FPayloadId())
 	{
 		return RecordBuilder->SetValue(Buffer, Id);
@@ -172,6 +178,10 @@ public:
 	 *                 the buffer will be hashed to create an ID.
 	 * @return The ID that was provided or created.
 	 */
+	inline FPayloadId AddAttachment(const FCompositeBuffer& Buffer, const FPayloadId& Id = FPayloadId())
+	{
+		return RecordBuilder->AddAttachment(Buffer, Id);
+	}
 	inline FPayloadId AddAttachment(const FSharedBuffer& Buffer, const FPayloadId& Id = FPayloadId())
 	{
 		return RecordBuilder->AddAttachment(Buffer, Id);
