Added conversions and constructors to/from 12-byte arrays for FCbObjectId and DerivedData::FPayloadId

#rnx

diff --git a/Engine/Source/Developer/DerivedDataCache/Private/DerivedDataBuildDefinition.cpp b/Engine/Source/Developer/DerivedDataCache/Private/DerivedDataBuildDefinition.cpp
--- a/Engine/Source/Developer/DerivedDataCache/Private/DerivedDataBuildDefinition.cpp
+++ b/Engine/Source/Developer/DerivedDataCache/Private/DerivedDataBuildDefinition.cpp
@@ -195,7 +195,7 @@ FBuildDefinitionInternal::FBuildDefinitionInternal(FBuildDefinitionBuilderIntern
 				const FBuildPayloadKey& PayloadKey = Pair.Value.Get<FBuildPayloadKey>();
 				Writer.BeginObject(FTCHARToUTF8(Pair.Key));
 				Writer.AddHash("Build"_ASV, PayloadKey.BuildKey.Hash);
-				Writer.AddObjectId("Payload"_ASV, FCbObjectId(PayloadKey.Id.GetView()));
+				Writer.AddObjectId("Payload"_ASV, FCbObjectId(PayloadKey.Id));
 				Writer.EndObject();
 			}
 		}
@@ -301,7 +301,7 @@ void FBuildDefinitionInternal::IterateInputBuilds(TFunctionRef<void (FStringView
 	{
 		FCbObjectView Build = Field.AsObjectView();
 		const FBuildKey BuildKey{Build["Build"_ASV].AsHash()};
-		const FPayloadId Id = FPayloadId(Build["Payload"_ASV].AsObjectId().GetView());
+		const FPayloadId Id = FPayloadId(Build["Payload"_ASV].AsObjectId());
 		Visitor(FUTF8ToTCHAR(Field.GetName()), FBuildPayloadKey{BuildKey, Id});
 	}
 }
diff --git a/Engine/Source/Developer/DerivedDataCache/Private/DerivedDataBuildOutput.cpp b/Engine/Source/Developer/DerivedDataCache/Private/DerivedDataBuildOutput.cpp
--- a/Engine/Source/Developer/DerivedDataCache/Private/DerivedDataBuildOutput.cpp
+++ b/Engine/Source/Developer/DerivedDataCache/Private/DerivedDataBuildOutput.cpp
@@ -147,7 +147,7 @@ FBuildOutputInternal::FBuildOutputInternal(FStringView InName, FStringView InFun
 	for (FCbFieldView PayloadField : InOutput["Payloads"_ASV])
 	{
 		const FCbObjectView Payload = PayloadField.AsObjectView();
-		const FPayloadId Id(Payload["Id"_ASV].AsObjectId().GetView());
+		const FPayloadId Id(Payload["Id"_ASV].AsObjectId());
 		const FIoHash& RawHash = Payload["RawHash"_ASV].AsAttachment();
 		if (Id.IsNull() || RawHash.IsZero() || !Payload["RawSize"_ASV].IsInteger())
 		{
@@ -218,7 +218,7 @@ void FBuildOutputInternal::Save(FCbWriter& Writer) const
 		for (const FPayload& Payload : Payloads)
 		{
 			Writer.BeginObject();
-			Writer.AddObjectId("Id"_ASV, FCbObjectId(Payload.GetId().GetView()));
+			Writer.AddObjectId("Id"_ASV, FCbObjectId(Payload.GetId()));
 			Writer.AddBinaryAttachment("RawHash"_ASV, Payload.GetRawHash());
 			Writer.AddInteger("RawSize"_ASV, Payload.GetRawSize());
 			Writer.EndObject();
diff --git a/Engine/Source/Developer/DerivedDataCache/Private/DerivedDataCacheRecord.cpp b/Engine/Source/Developer/DerivedDataCache/Private/DerivedDataCacheRecord.cpp
--- a/Engine/Source/Developer/DerivedDataCache/Private/DerivedDataCacheRecord.cpp
+++ b/Engine/Source/Developer/DerivedDataCache/Private/DerivedDataCacheRecord.cpp
@@ -305,7 +305,7 @@ FCbPackage SaveCacheRecord(const FCacheRecord& Record)
 	auto SavePayload = [&Package, &Writer](const FPayload& Payload)
 	{
 		Writer.BeginObject();
-		Writer.AddObjectId("Id"_ASV, FCbObjectId(Payload.GetId().GetView()));
+		Writer.AddObjectId("Id"_ASV, FCbObjectId(Payload.GetId()));
 		Writer.AddHash("RawHash"_ASV, Payload.GetRawHash());
 		Writer.AddInteger("RawSize"_ASV, Payload.GetRawSize());
 		const FCompositeBuffer CompressedBuffer = Payload.GetData().GetCompressed();
@@ -362,7 +362,7 @@ FOptionalCacheRecord LoadCacheRecord(const FCbPackage& Package)
 
 	auto LoadPayload = [&Package](const FCbObjectView& PayloadObject)
 	{
-		const FPayloadId PayloadId(PayloadObject["Id"].AsObjectId().GetView());
+		const FPayloadId PayloadId(PayloadObject["Id"].AsObjectId());
 		if (PayloadId.IsNull())
 		{
 			return FPayload();
diff --git a/Engine/Source/Developer/DerivedDataCache/Private/FileSystemDerivedDataBackend.cpp b/Engine/Source/Developer/DerivedDataCache/Private/FileSystemDerivedDataBackend.cpp
--- a/Engine/Source/Developer/DerivedDataCache/Private/FileSystemDerivedDataBackend.cpp
+++ b/Engine/Source/Developer/DerivedDataCache/Private/FileSystemDerivedDataBackend.cpp
@@ -1207,7 +1207,7 @@ private:
 		}
 
 		Writer.BeginObject();
-		Writer.AddObjectId("Id"_ASV, FCbObjectId(Payload.GetId().GetView()));
+		Writer.AddObjectId("Id"_ASV, FCbObjectId(Payload.GetId()));
 		Writer.AddInteger("RawSize"_ASV, Payload.GetRawSize());
 		if (bStoreInline)
 		{
@@ -1226,7 +1226,7 @@ private:
 
 	FPayload GetCachePayload(const FCacheKey& Key, FStringView Context, ECachePolicy Policy, const FCbObject& Object, bool bAlwaysLoadInlineData = false) const
 	{
-		const FPayloadId Id(Object.FindView("Id"_ASV).AsObjectId().GetView());
+		const FPayloadId Id(Object.FindView("Id"_ASV).AsObjectId());
 		const uint64 RawSize = Object.FindView("RawSize"_ASV).AsUInt64(MAX_uint64);
 		const FIoHash RawHash = Object.FindView("RawHash"_ASV).AsHash();
 		FIoHash CompressedHash = Object.FindView("CompressedHash"_ASV).AsHash();
diff --git a/Engine/Source/Developer/DerivedDataCache/Public/DerivedDataPayload.h b/Engine/Source/Developer/DerivedDataCache/Public/DerivedDataPayload.h
--- a/Engine/Source/Developer/DerivedDataCache/Public/DerivedDataPayload.h
+++ b/Engine/Source/Developer/DerivedDataCache/Public/DerivedDataPayload.h
@@ -25,6 +25,9 @@ public:
 	/** Construct a null ID. */
 	FPayloadId() = default;
 
+	/** Construct an ID from an array of 12 bytes. */
+	inline explicit FPayloadId(const ByteArray& Id);
+
 	/** Construct an ID from a view of 12 bytes. */
 	inline explicit FPayloadId(FMemoryView Id);
 
@@ -37,6 +40,7 @@ public:
 
 	/** Returns a reference to the raw byte array for the ID. */
 	inline const ByteArray& GetBytes() const { return Bytes; }
+	inline operator const ByteArray&() const { return Bytes; }
 
 	/** Returns a view of the raw byte array for the ID. */
 	inline FMemoryView GetView() const { return MakeMemoryView(Bytes); }
@@ -121,11 +125,16 @@ inline const FPayload FPayload::Null;
 
 ///////////////////////////////////////////////////////////////////////////////////////////////////////////////////////
 
-inline FPayloadId::FPayloadId(const FMemoryView InId)
+inline FPayloadId::FPayloadId(const ByteArray& Id)
+{
+	FMemory::Memcpy(Bytes, Id, sizeof(ByteArray));
+}
+
+inline FPayloadId::FPayloadId(const FMemoryView Id)
 {
-	checkf(InId.GetSize() == sizeof(ByteArray),
-		TEXT("FPayloadId cannot be constructed from a view of %" UINT64_FMT " bytes."), InId.GetSize());
-	FMemory::Memcpy(Bytes, InId.GetData(), sizeof(ByteArray));
+	checkf(Id.GetSize() == sizeof(ByteArray),
+		TEXT("FPayloadId cannot be constructed from a view of %" UINT64_FMT " bytes."), Id.GetSize());
+	FMemory::Memcpy(Bytes, Id.GetData(), sizeof(ByteArray));
 }
 
 inline FPayloadId FPayloadId::FromHash(const FIoHash& Hash)
diff --git a/Engine/Source/Runtime/Core/Public/Serialization/CompactBinary.h b/Engine/Source/Runtime/Core/Public/Serialization/CompactBinary.h
--- a/Engine/Source/Runtime/Core/Public/Serialization/CompactBinary.h
+++ b/Engine/Source/Runtime/Core/Public/Serialization/CompactBinary.h
@@ -463,9 +463,14 @@ enum class ECbFieldError : uint8
 class FCbObjectId
 {
 public:
+	using ByteArray = uint8[12];
+
 	/** Construct an ObjectId with every byte initialized to zero. */
 	FCbObjectId() = default;
 
+	/** Construct an ObjectId from an array of 12 bytes. */
+	inline explicit FCbObjectId(const ByteArray& ObjectId);
+
 	/** Construct an ObjectId from a view of 12 bytes. */
 	CORE_API explicit FCbObjectId(FMemoryView ObjectId);
 
@@ -473,13 +478,22 @@ public:
 	CORE_API void ToString(FAnsiStringBuilderBase& Builder) const;
 	CORE_API void ToString(FWideStringBuilderBase& Builder) const;
 
+	/** Returns a reference to the raw byte array for the ObjectId. */
+	inline const ByteArray& GetBytes() const { return Bytes; }
+	inline operator const ByteArray&() const { return Bytes; }
+
 	/** Returns a view of the raw byte array for the ObjectId. */
 	constexpr inline FMemoryView GetView() const { return MakeMemoryView(Bytes); }
 
 private:
-	alignas(uint32) uint8 Bytes[12]{};
+	alignas(uint32) ByteArray Bytes{};
 };
 
+inline FCbObjectId::FCbObjectId(const ByteArray& ObjectId)
+{
+	FMemory::Memcpy(Bytes, ObjectId, sizeof(ByteArray));
+}
+
 inline bool operator==(const FCbObjectId& A, const FCbObjectId& B)
 {
 	return FMemory::Memcmp(&A, &B, sizeof(FCbObjectId)) == 0;
