DDC: Added CreateCacheBucket to FBuildConfigContext

#rnx

diff --git a/Engine/Source/Developer/DerivedDataCache/Private/DerivedDataBuildJob.cpp b/Engine/Source/Developer/DerivedDataCache/Private/DerivedDataBuildJob.cpp
--- a/Engine/Source/Developer/DerivedDataCache/Private/DerivedDataBuildJob.cpp
+++ b/Engine/Source/Developer/DerivedDataCache/Private/DerivedDataBuildJob.cpp
@@ -113,6 +113,9 @@ public:
 	inline const FBuildKey& GetKey() const final { return DefinitionKey; }
 	inline const FBuildActionKey& GetActionKey() const final { return ActionKey; }
 
+	inline ICache& GetCache() const final { return Cache; }
+	inline IBuild& GetBuild() const final { return BuildSystem; }
+
 	void Schedule(IBuildScheduler& Scheduler, EBuildPolicy Policy, EPriority Priority, FOnBuildJobComplete&& OnComplete) final;
 
 	// IRequest Interface
diff --git a/Engine/Source/Developer/DerivedDataCache/Private/DerivedDataBuildJobContext.cpp b/Engine/Source/Developer/DerivedDataCache/Private/DerivedDataBuildJobContext.cpp
--- a/Engine/Source/Developer/DerivedDataCache/Private/DerivedDataBuildJobContext.cpp
+++ b/Engine/Source/Developer/DerivedDataCache/Private/DerivedDataBuildJobContext.cpp
@@ -119,6 +119,11 @@ void FBuildJobContext::EndAsyncBuild()
 	OnEndAsyncBuild();
 }
 
+FCacheBucket FBuildJobContext::CreateCacheBucket(FStringView Name) const
+{
+	return Job.GetCache().CreateBucket(Name);
+}
+
 void FBuildJobContext::SetCacheBucket(FCacheBucket Bucket)
 {
 	checkf(!Bucket.IsNull(), TEXT("Null cache bucket not allowed for build of '%s' by %s. ")
diff --git a/Engine/Source/Developer/DerivedDataCache/Private/DerivedDataBuildJobContext.h b/Engine/Source/Developer/DerivedDataCache/Private/DerivedDataBuildJobContext.h
--- a/Engine/Source/Developer/DerivedDataCache/Private/DerivedDataBuildJobContext.h
+++ b/Engine/Source/Developer/DerivedDataCache/Private/DerivedDataBuildJobContext.h
@@ -58,6 +58,7 @@ private:
 	void BeginAsyncBuild() final;
 	void EndAsyncBuild() final;
 
+	FCacheBucket CreateCacheBucket(FStringView Name) const final;
 	void SetCacheBucket(FCacheBucket Bucket) final;
 	void SetCachePolicy(ECachePolicy Policy) final;
 	void SetBuildPolicy(EBuildPolicy Policy) final;
diff --git a/Engine/Source/Developer/DerivedDataCache/Public/DerivedDataBuildFunction.h b/Engine/Source/Developer/DerivedDataCache/Public/DerivedDataBuildFunction.h
--- a/Engine/Source/Developer/DerivedDataCache/Public/DerivedDataBuildFunction.h
+++ b/Engine/Source/Developer/DerivedDataCache/Public/DerivedDataBuildFunction.h
@@ -127,6 +127,9 @@ public:
 	/** Returns the constant with the matching key, or an object with no fields if not found. */
 	virtual FCbObject FindConstant(FStringView Key) const = 0;
 
+	/** Create a cache bucket from a name. See ICacheFactory::CreateBucket. */
+	virtual FCacheBucket CreateCacheBucket(FStringView Name) const = 0;
+
 	/** Overrides the cache bucket used when reading or writing this build in the cache. */
 	virtual void SetCacheBucket(FCacheBucket Bucket) = 0;
 
diff --git a/Engine/Source/Developer/DerivedDataCache/Public/DerivedDataBuildJob.h b/Engine/Source/Developer/DerivedDataCache/Public/DerivedDataBuildJob.h
--- a/Engine/Source/Developer/DerivedDataCache/Public/DerivedDataBuildJob.h
+++ b/Engine/Source/Developer/DerivedDataCache/Public/DerivedDataBuildJob.h
@@ -11,7 +11,9 @@ class FEvent;
 template <typename FuncType> class TUniqueFunction;
 
 namespace UE::DerivedData { class FBuildOutput; }
+namespace UE::DerivedData { class IBuild; }
 namespace UE::DerivedData { class IBuildScheduler; }
+namespace UE::DerivedData { class ICache; }
 namespace UE::DerivedData { struct FBuildJobCompleteParams; }
 namespace UE::DerivedData { struct FBuildActionKey; }
 namespace UE::DerivedData { struct FBuildKey; }
@@ -49,6 +51,11 @@ public:
 	/** Returns the action key, or null if the action has not been resolved yet. */
 	virtual const FBuildActionKey& GetActionKey() const = 0;
 
+	/** Returns the cache associated with this job. */
+	virtual ICache& GetCache() const = 0;
+	/** Returns the build system associated with this job. */
+	virtual IBuild& GetBuild() const = 0;
+
 	/** Start the build by dispatching it to the scheduler. */
 	virtual void Schedule(IBuildScheduler& Scheduler, EBuildPolicy Policy, EPriority Priority, FOnBuildJobComplete&& OnComplete) = 0;
 
