DDC: Fixed TryToPrefetch to work when there is only one backend

#rnx

diff --git a/Engine/Source/Developer/DerivedDataCache/Private/FileSystemDerivedDataBackend.cpp b/Engine/Source/Developer/DerivedDataCache/Private/FileSystemDerivedDataBackend.cpp
--- a/Engine/Source/Developer/DerivedDataCache/Private/FileSystemDerivedDataBackend.cpp
+++ b/Engine/Source/Developer/DerivedDataCache/Private/FileSystemDerivedDataBackend.cpp
@@ -871,7 +871,7 @@ public:
 
 	bool TryToPrefetch(TConstArrayView<FString> CacheKeys) override
 	{
-		return false;
+		return GetSpeedClass() >= ESpeedClass::Fast && CachedDataProbablyExistsBatch(CacheKeys).CountSetBits() == CacheKeys.Num();
 	}
 
 	bool ApplyDebugOptions(FBackendDebugOptions& InOptions) override
diff --git a/Engine/Source/Developer/DerivedDataCache/Private/HttpDerivedDataBackend.cpp b/Engine/Source/Developer/DerivedDataCache/Private/HttpDerivedDataBackend.cpp
--- a/Engine/Source/Developer/DerivedDataCache/Private/HttpDerivedDataBackend.cpp
+++ b/Engine/Source/Developer/DerivedDataCache/Private/HttpDerivedDataBackend.cpp
@@ -1806,7 +1806,7 @@ FString FHttpDerivedDataBackend::GetName() const
 
 bool FHttpDerivedDataBackend::TryToPrefetch(TConstArrayView<FString> CacheKeys)
 {
-	return false;
+	return GetSpeedClass() >= ESpeedClass::Fast && CachedDataProbablyExistsBatch(CacheKeys).CountSetBits() == CacheKeys.Num();
 }
 
 bool FHttpDerivedDataBackend::WouldCache(const TCHAR* CacheKey, TArrayView<const uint8> InData)
diff --git a/Engine/Source/Developer/DerivedDataCache/Private/MemoryDerivedDataBackend.cpp b/Engine/Source/Developer/DerivedDataCache/Private/MemoryDerivedDataBackend.cpp
--- a/Engine/Source/Developer/DerivedDataCache/Private/MemoryDerivedDataBackend.cpp
+++ b/Engine/Source/Developer/DerivedDataCache/Private/MemoryDerivedDataBackend.cpp
@@ -97,7 +97,7 @@ bool FMemoryDerivedDataBackend::GetCachedData(const TCHAR* CacheKey, TArray<uint
 
 bool FMemoryDerivedDataBackend::TryToPrefetch(TConstArrayView<FString> CacheKeys)
 {
-	return false;
+	return GetSpeedClass() >= ESpeedClass::Fast && CachedDataProbablyExistsBatch(CacheKeys).CountSetBits() == CacheKeys.Num();
 }
 
 bool FMemoryDerivedDataBackend::WouldCache(const TCHAR* CacheKey, TArrayView<const uint8> InData)
diff --git a/Engine/Source/Developer/DerivedDataCache/Private/PakFileDerivedDataBackend.h b/Engine/Source/Developer/DerivedDataCache/Private/PakFileDerivedDataBackend.h
--- a/Engine/Source/Developer/DerivedDataCache/Private/PakFileDerivedDataBackend.h
+++ b/Engine/Source/Developer/DerivedDataCache/Private/PakFileDerivedDataBackend.h
@@ -95,7 +95,10 @@ public:
 
 	virtual TSharedRef<FDerivedDataCacheStatsNode> GatherUsageStats() const override;
 
-	virtual bool TryToPrefetch(TConstArrayView<FString> CacheKeys) override { return false; }
+	virtual bool TryToPrefetch(TConstArrayView<FString> CacheKeys) override
+	{
+		return GetSpeedClass() >= ESpeedClass::Fast && CachedDataProbablyExistsBatch(CacheKeys).CountSetBits() == CacheKeys.Num();
+	}
 
 	virtual bool WouldCache(const TCHAR* CacheKey, TArrayView<const uint8> InData) override { return true; }
 
diff --git a/Engine/Source/Developer/DerivedDataCache/Private/S3DerivedDataBackend.cpp b/Engine/Source/Developer/DerivedDataCache/Private/S3DerivedDataBackend.cpp
--- a/Engine/Source/Developer/DerivedDataCache/Private/S3DerivedDataBackend.cpp
+++ b/Engine/Source/Developer/DerivedDataCache/Private/S3DerivedDataBackend.cpp
@@ -962,7 +962,7 @@ FDerivedDataBackendInterface::ESpeedClass FS3DerivedDataBackend::GetSpeedClass()
 
 bool FS3DerivedDataBackend::TryToPrefetch(TConstArrayView<FString> CacheKeys)
 {
-	return false;
+	return GetSpeedClass() >= ESpeedClass::Fast && CachedDataProbablyExistsBatch(CacheKeys).CountSetBits() == CacheKeys.Num();
 }
 
 bool FS3DerivedDataBackend::WouldCache(const TCHAR* CacheKey, TArrayView<const uint8> InData)
