DDC: Fixed MeasureCacheRecord to only count payloads with data

#rnx

diff --git a/Engine/Source/Developer/DerivedDataCache/Private/FileSystemDerivedDataBackend.cpp b/Engine/Source/Developer/DerivedDataCache/Private/FileSystemDerivedDataBackend.cpp
--- a/Engine/Source/Developer/DerivedDataCache/Private/FileSystemDerivedDataBackend.cpp
+++ b/Engine/Source/Developer/DerivedDataCache/Private/FileSystemDerivedDataBackend.cpp
@@ -999,8 +999,9 @@ private:
 	uint64 MeasureCacheRecord(const FCacheRecord& Record) const
 	{
 		return Record.GetMeta().GetSize() +
-			Record.GetValuePayload().GetRawSize() +
-			Algo::TransformAccumulate(Record.GetAttachmentPayloads(), &FPayload::GetRawSize, uint64(0));
+			Record.GetValuePayload().GetData().GetRawSize() +
+			Algo::TransformAccumulate(Record.GetAttachmentPayloads(),
+				[](const FPayload& Payload) { return Payload.GetData().GetRawSize(); }, uint64(0));
 	}
 
 	bool PutCacheRecord(const FCacheRecord& Record, FStringView Context, ECachePolicy Policy) const
