DDC: Switched to operator[] and AsView() on compact binary types

#rnx

diff --git a/Engine/Source/Developer/DerivedDataCache/Private/DerivedDataBuildAction.cpp b/Engine/Source/Developer/DerivedDataCache/Private/DerivedDataBuildAction.cpp
--- a/Engine/Source/Developer/DerivedDataCache/Private/DerivedDataBuildAction.cpp
+++ b/Engine/Source/Developer/DerivedDataCache/Private/DerivedDataBuildAction.cpp
@@ -199,13 +199,13 @@ FBuildActionInternal::FBuildActionInternal(FStringView InName, FCbObject&& InAct
 		&& IsValidBuildFunctionName(Function)
 		&& FunctionVersion.IsValid()
 		&& BuildSystemVersion.IsValid()
-		&& Algo::AllOf(Action.FindView("Constants"_ASV),
+		&& Algo::AllOf(Action.AsView()["Constants"_ASV],
 			[](FCbFieldView Field) { return Field.GetName().Len() > 0 && Field.IsObject(); })
-		&& Algo::AllOf(Action.FindView("Inputs"_ASV), [](FCbFieldView Field)
+		&& Algo::AllOf(Action.AsView()["Inputs"_ASV], [](FCbFieldView Field)
 			{
 				return Field.GetName().Len() > 0 && Field.IsObject()
-					&& Field.AsObjectView()["RawHash"_ASV].IsBinaryAttachment()
-					&& Field.AsObjectView()["RawSize"_ASV].IsInteger();
+					&& Field["RawHash"_ASV].IsBinaryAttachment()
+					&& Field["RawSize"_ASV].IsInteger();
 			});
 }
 
@@ -229,10 +229,9 @@ void FBuildActionInternal::IterateConstants(TFunctionRef<void (FStringView Key,
 
 void FBuildActionInternal::IterateInputs(TFunctionRef<void (FStringView Key, const FIoHash& RawHash, uint64 RawSize)> Visitor) const
 {
-	for (FCbFieldView Field : Action.FindView("Inputs"_ASV))
+	for (FCbFieldView Field : Action.AsView()["Inputs"_ASV])
 	{
-		FCbObjectView Input = Field.AsObjectView();
-		Visitor(FUTF8ToTCHAR(Field.GetName()), Input["RawHash"_ASV].AsHash(), Input["RawSize"_ASV].AsUInt64());
+		Visitor(FUTF8ToTCHAR(Field.GetName()), Field["RawHash"_ASV].AsHash(), Field["RawSize"_ASV].AsUInt64());
 	}
 }
 
diff --git a/Engine/Source/Developer/DerivedDataCache/Private/DerivedDataBuildDefinition.cpp b/Engine/Source/Developer/DerivedDataCache/Private/DerivedDataBuildDefinition.cpp
--- a/Engine/Source/Developer/DerivedDataCache/Private/DerivedDataBuildDefinition.cpp
+++ b/Engine/Source/Developer/DerivedDataCache/Private/DerivedDataBuildDefinition.cpp
@@ -261,19 +261,19 @@ FBuildDefinitionInternal::FBuildDefinitionInternal(FStringView InName, FCbObject
 	Definition.MakeOwned();
 	bOutIsValid = Definition
 		&& IsValidBuildFunctionName(Function)
-		&& Algo::AllOf(Definition.FindView("Constants"_ASV),
+		&& Algo::AllOf(Definition.AsView()["Constants"_ASV],
 			[](FCbFieldView Field) { return Field.GetName().Len() > 0 && Field.IsObject(); })
-		&& Algo::AllOf(Definition.FindView("Inputs"_ASV).AsObjectView().FindView("Builds"_ASV), [](FCbFieldView Field)
+		&& Algo::AllOf(Definition.AsView()["Inputs"_ASV]["Builds"_ASV], [](FCbFieldView Field)
 			{
 				return Field.GetName().Len() > 0 && Field.IsObject()
-					&& Field.AsObjectView()["Build"_ASV].IsHash()
-					&& Field.AsObjectView()["Payload"_ASV].IsObjectId();
+					&& Field["Build"_ASV].IsHash()
+					&& Field["Payload"_ASV].IsObjectId();
 			})
-		&& Algo::AllOf(Definition.FindView("Inputs"_ASV).AsObjectView().FindView("BulkData"_ASV),
+		&& Algo::AllOf(Definition.AsView()["Inputs"_ASV]["BulkData"_ASV],
 			[](FCbFieldView Field) { return Field.GetName().Len() > 0 && Field.IsUuid(); })
-		&& Algo::AllOf(Definition.FindView("Inputs"_ASV).AsObjectView().FindView("Files"_ASV),
+		&& Algo::AllOf(Definition.AsView()["Inputs"_ASV]["Files"_ASV],
 			[](FCbFieldView Field) { return Field.GetName().Len() > 0 && Field.AsString().Len() > 0; })
-		&& Algo::AllOf(Definition.FindView("Inputs"_ASV).AsObjectView().FindView("Hashes"_ASV),
+		&& Algo::AllOf(Definition.AsView()["Inputs"_ASV]["Hashes"_ASV],
 			[](FCbFieldView Field) { return Field.GetName().Len() > 0 && Field.IsBinaryAttachment(); });
 }
 
@@ -297,18 +297,17 @@ void FBuildDefinitionInternal::IterateConstants(TFunctionRef<void (FStringView K
 
 void FBuildDefinitionInternal::IterateInputBuilds(TFunctionRef<void (FStringView Key, const FBuildPayloadKey& PayloadKey)> Visitor) const
 {
-	for (FCbFieldView Field : Definition.FindView("Inputs"_ASV).AsObjectView().FindView("Builds"_ASV))
+	for (FCbFieldView Field : Definition.AsView()["Inputs"_ASV]["Builds"_ASV])
 	{
-		FCbObjectView Build = Field.AsObjectView();
-		const FBuildKey BuildKey{Build["Build"_ASV].AsHash()};
-		const FPayloadId Id = FPayloadId(Build["Payload"_ASV].AsObjectId());
+		const FBuildKey BuildKey{Field["Build"_ASV].AsHash()};
+		const FPayloadId Id = FPayloadId(Field["Payload"_ASV].AsObjectId());
 		Visitor(FUTF8ToTCHAR(Field.GetName()), FBuildPayloadKey{BuildKey, Id});
 	}
 }
 
 void FBuildDefinitionInternal::IterateInputBulkData(TFunctionRef<void (FStringView Key, const FGuid& BulkDataId)> Visitor) const
 {
-	for (FCbFieldView Field : Definition.FindView("Inputs"_ASV).AsObjectView().FindView("BulkData"_ASV))
+	for (FCbFieldView Field : Definition.AsView()["Inputs"_ASV]["BulkData"_ASV])
 	{
 		Visitor(FUTF8ToTCHAR(Field.GetName()), Field.AsUuid());
 	}
@@ -316,7 +315,7 @@ void FBuildDefinitionInternal::IterateInputBulkData(TFunctionRef<void (FStringVi
 
 void FBuildDefinitionInternal::IterateInputFiles(TFunctionRef<void (FStringView Key, FStringView Path)> Visitor) const
 {
-	for (FCbFieldView Field : Definition.FindView("Inputs"_ASV).AsObjectView().FindView("Files"_ASV))
+	for (FCbFieldView Field : Definition.AsView()["Inputs"_ASV]["Files"_ASV])
 	{
 		Visitor(FUTF8ToTCHAR(Field.GetName()), FUTF8ToTCHAR(Field.AsString()));
 	}
@@ -324,7 +323,7 @@ void FBuildDefinitionInternal::IterateInputFiles(TFunctionRef<void (FStringView
 
 void FBuildDefinitionInternal::IterateInputHashes(TFunctionRef<void (FStringView Key, const FIoHash& RawHash)> Visitor) const
 {
-	for (FCbFieldView Field : Definition.FindView("Inputs"_ASV).AsObjectView().FindView("Hashes"_ASV))
+	for (FCbFieldView Field : Definition.AsView()["Inputs"_ASV]["Hashes"_ASV])
 	{
 		Visitor(FUTF8ToTCHAR(Field.GetName()), Field.AsBinaryAttachment());
 	}
diff --git a/Engine/Source/Developer/DerivedDataCache/Private/DerivedDataBuildOutput.cpp b/Engine/Source/Developer/DerivedDataCache/Private/DerivedDataBuildOutput.cpp
--- a/Engine/Source/Developer/DerivedDataCache/Private/DerivedDataBuildOutput.cpp
+++ b/Engine/Source/Developer/DerivedDataCache/Private/DerivedDataBuildOutput.cpp
@@ -144,9 +144,8 @@ FBuildOutputInternal::FBuildOutputInternal(FStringView InName, FStringView InFun
 	AssertValidBuildFunctionName(Function, Name);
 	Meta = InOutput["Meta"_ASV].AsObject();
 	Meta.MakeOwned();
-	for (FCbFieldView PayloadField : InOutput["Payloads"_ASV])
+	for (FCbFieldView Payload : InOutput["Payloads"_ASV])
 	{
-		const FCbObjectView Payload = PayloadField.AsObjectView();
 		const FPayloadId Id(Payload["Id"_ASV].AsObjectId());
 		const FIoHash& RawHash = Payload["RawHash"_ASV].AsAttachment();
 		if (Id.IsNull() || RawHash.IsZero() || !Payload["RawSize"_ASV].IsInteger())
@@ -157,7 +156,7 @@ FBuildOutputInternal::FBuildOutputInternal(FStringView InName, FStringView InFun
 	}
 	Diagnostics = InOutput["Diagnostics"_ASV];
 	Diagnostics.MakeOwned();
-	bOutIsValid = IsValid() && (!InOutput.FindView("Meta"_ASV) || InOutput.FindView("Meta"_ASV).IsObject());
+	bOutIsValid = IsValid() && (!InOutput.AsView()["Meta"_ASV] || InOutput.AsView()["Meta"_ASV].IsObject());
 }
 
 FBuildOutputInternal::FBuildOutputInternal(FStringView InName, FStringView InFunction, const FCacheRecord& InOutput, bool& bOutIsValid)
@@ -183,10 +182,9 @@ const FPayload& FBuildOutputInternal::GetPayload(const FPayloadId& Id) const
 
 void FBuildOutputInternal::IterateDiagnostics(TFunctionRef<void (const FBuildDiagnostic& Diagnostic)> Visitor) const
 {
-	for (FCbFieldView Field : Diagnostics.CreateViewIterator())
+	for (FCbFieldView Diagnostic : Diagnostics.CreateViewIterator())
 	{
 		EBuildDiagnosticLevel Level;
-		const FCbObjectView Diagnostic = Field.AsObjectView();
 		LexFromString(Level, Diagnostic["Level"_ASV].AsString());
 		const FAnsiStringView Category = Diagnostic["Category"_ASV].AsString();
 		const FAnsiStringView Message = Diagnostic["Message"_ASV].AsString();
@@ -261,9 +259,9 @@ bool FBuildOutputInternal::IsValid() const
 		&& Algo::AllOf(Diagnostics.CreateViewIterator(), [](FCbFieldView Field)
 			{
 				return Field.IsObject()
-					&& Field.AsObjectView()["Level"_ASV].AsString().Len() > 0
-					&& Field.AsObjectView()["Category"_ASV].AsString().Len() > 0
-					&& Field.AsObjectView()["Message"_ASV].AsString().Len() > 0;
+					&& Field["Level"_ASV].AsString().Len() > 0
+					&& Field["Category"_ASV].AsString().Len() > 0
+					&& Field["Message"_ASV].AsString().Len() > 0;
 			});
 }
 
