DDC: Fixed the build job to create and destroy the build context only when necessary, and always reset the action and inputs

#rnx

diff --git a/Engine/Source/Developer/DerivedDataCache/Private/DerivedDataBuildJob.cpp b/Engine/Source/Developer/DerivedDataCache/Private/DerivedDataBuildJob.cpp
--- a/Engine/Source/Developer/DerivedDataCache/Private/DerivedDataBuildJob.cpp
+++ b/Engine/Source/Developer/DerivedDataCache/Private/DerivedDataBuildJob.cpp
@@ -131,7 +131,7 @@ private:
 	void BeginJob();
 	void EndJob();
 
-	void Configure();
+	void CreateContext();
 
 	void EnterCacheQuery();
 	void EnterCacheStore();
@@ -489,7 +489,7 @@ void FBuildJob::EndJob()
 
 ///////////////////////////////////////////////////////////////////////////////////////////////////////////////////////
 
-void FBuildJob::Configure()
+void FBuildJob::CreateContext()
 {
 	const IBuildFunction* Function = BuildSystem.GetFunctionRegistry().FindFunction(FunctionName);
 	if (BuildSystem.GetVersion() != Action.Get().GetBuildSystemVersion())
@@ -1128,9 +1128,9 @@ void FBuildJob::ExecuteTransition(EBuildJobState OldState, EBuildJobState NewSta
 	{
 		BeginJob();
 	}
-	if (OldState < EBuildJobState::CacheQuery && EBuildJobState::CacheQuery <= NewState)
+	if (OldState < EBuildJobState::CacheQuery && EBuildJobState::CacheQuery <= NewState && NewState <= EBuildJobState::CacheStore)
 	{
-		Configure();
+		CreateContext();
 	}
 	if (OldState <= EBuildJobState::ExecuteRemoteRetryWait && EBuildJobState::ExecuteRemoteRetryWait < NewState)
 	{
@@ -1142,13 +1142,16 @@ void FBuildJob::ExecuteTransition(EBuildJobState OldState, EBuildJobState NewSta
 		MissingInputs.Empty();
 		Definition.Reset();
 	}
-	if (OldState <= EBuildJobState::ExecuteLocalWait && EBuildJobState::ExecuteLocalWait < NewState && !Output)
+	if (OldState <= EBuildJobState::ExecuteLocalWait && EBuildJobState::ExecuteLocalWait < NewState)
 	{
 		Action.Reset();
 		Inputs.Reset();
+	}
+	if (OldState <= EBuildJobState::ExecuteLocalWait && EBuildJobState::ExecuteLocalWait < NewState && !Output)
+	{
 		SetOutputNoCheck(OutputBuilder.Build());
 	}
-	if (OldState < EBuildJobState::CacheStore && EBuildJobState::CacheStore <= NewState)
+	if (OldState < EBuildJobState::CacheStore && EBuildJobState::CacheStore <= NewState && Context)
 	{
 		Context->ResetInputs();
 	}
