DDC: Extended IBuildJob and IBuildScheduler to give the scheduler more control and more data

#rnx

diff --git a/Engine/Source/Developer/DerivedDataCache/Private/DerivedDataBuildJob.cpp b/Engine/Source/Developer/DerivedDataCache/Private/DerivedDataBuildJob.cpp
--- a/Engine/Source/Developer/DerivedDataCache/Private/DerivedDataBuildJob.cpp
+++ b/Engine/Source/Developer/DerivedDataCache/Private/DerivedDataBuildJob.cpp
@@ -107,12 +107,9 @@ public:
 
 	inline FStringView GetName() const final { return Name; }
 	inline FStringView GetFunction() const final { return FunctionName; }
-	inline EBuildPolicy GetPolicy() const final { return Context ? Context->GetBuildPolicy() : DefaultBuildPolicy; }
+	inline EBuildPolicy GetPolicy() const final { return BuildPolicy; }
 	inline EPriority GetPriority() const final { return Priority; }
 
-	inline const FBuildKey& GetKey() const final { return DefinitionKey; }
-	inline const FBuildActionKey& GetActionKey() const final { return ActionKey; }
-
 	inline ICache& GetCache() const final { return Cache; }
 	inline IBuild& GetBuild() const final { return BuildSystem; }
 
@@ -157,6 +154,8 @@ private:
 	void EndExecuteRemote(FBuildWorkerActionCompleteParams&& Params);
 	void EndExecuteLocal();
 
+	void SkipExecuteRemote() final;
+
 	void CreateAction(TConstArrayView<FBuildInputMetaByKey> InputMeta);
 
 	void SetDefinition(FBuildDefinition&& Definition);
@@ -186,7 +185,7 @@ private:
 	/** Next state for the job. Used to handle re-entrant calls to AdvanceToState. */
 	EBuildJobState NextState{EBuildJobState::NotStarted};
 
-	EBuildPolicy DefaultBuildPolicy{};
+	EBuildPolicy BuildPolicy{};
 	EPriority Priority{};
 
 	/** True if the build was canceled before it was complete. */
@@ -202,8 +201,6 @@ private:
 
 	/** Available in [ResolveKey, Complete] for jobs created from a key or definition. */
 	FBuildKey DefinitionKey;
-	/** Available in [CacheQuery, Complete]. */
-	FBuildActionKey ActionKey;
 	/** Available in [ResolveAction, ExecuteLocal) for jobs created from a key or definition. */
 	FOptionalBuildDefinition Definition;
 	/** Available in [CacheQuery, CacheStore). */
@@ -219,6 +216,8 @@ private:
 	/** Available in [CacheQuery, CacheStoreWait). Context used by IBuildFunction. */
 	TRefCountPtr<FBuildJobContext> Context;
 
+	FBuildSchedulerParams SchedulerParams;
+
 	/** Lock to synchronize writes to State, NextState, bIsCanceled, Event, WaitRequest, Scheduler. */
 	FRWLock Lock;
 	/** Event used by and created by Wait(). Null until Wait() has been called on an incomplete job. */
@@ -236,7 +235,7 @@ private:
 	/** Invoked exactly once when the output is complete or when the job fails. */
 	FOnBuildJobComplete OnComplete;
 
-	/** Keys for missing inputs. All inputs will be resolved when this is empty. */
+	/** Keys for missing inputs. */
 	TArray<FString> MissingInputs;
 
 	ICache& Cache;
@@ -312,7 +311,6 @@ FBuildJob::FBuildJob(
 	const FOptionalBuildInputs& InInputs)
 	: Name(InAction.GetName())
 	, FunctionName(InAction.GetFunction())
-	, ActionKey(InAction.GetKey())
 	, Action(InAction)
 	, Inputs(InInputs)
 	, OutputBuilder(InBuildSystem.CreateOutput(Name, FunctionName))
@@ -343,7 +341,7 @@ void FBuildJob::Schedule(
 		TEXT("Job in state %s was previously scheduled for build of '%s' by %s."),
 		LexToString(State), *Name, *FunctionName);
 	Scheduler = &InScheduler;
-	DefaultBuildPolicy = InPolicy;
+	BuildPolicy = InPolicy;
 	Priority = InPriority;
 	OnComplete = MoveTemp(InOnComplete);
 	return AdvanceToState(EBuildJobState::ResolveKey);
@@ -517,9 +515,30 @@ void FBuildJob::CreateContext()
 	else
 	{
 		const FCacheKey CacheKey{Cache.CreateBucket(FunctionName), Action.Get().GetKey().Hash};
-		Context = new FBuildJobContext(*this, CacheKey, *Function, OutputBuilder, DefaultBuildPolicy,
+		Context = new FBuildJobContext(*this, CacheKey, *Function, OutputBuilder, BuildPolicy,
 			[this] { EndExecuteLocal(); });
 		Function->Configure(*Context);
+		BuildPolicy = Context->GetBuildPolicy();
+	}
+
+	// Populate the scheduler params with the information that is available now.
+	SchedulerParams.Key = Action.Get().GetKey();
+	Action.Get().IterateConstants([this](FStringView Key, FCbObject&& Value)
+	{
+		const uint64 ValueSize = Value.GetSize();
+		SchedulerParams.TotalInputsSize += ValueSize;
+		SchedulerParams.ResolvedInputsSize += ValueSize;
+	});
+	Action.Get().IterateInputs([this](FStringView Key, const FIoHash& RawHash, uint64 RawSize)
+	{
+		SchedulerParams.TotalInputsSize += RawSize;
+	});
+	if (Inputs)
+	{
+		Inputs.Get().IterateInputs([this](FStringView Key, const FCompressedBuffer& Buffer)
+		{
+			SchedulerParams.ResolvedInputsSize += Buffer.GetRawSize();
+		});
 	}
 }
 
@@ -527,21 +546,19 @@ void FBuildJob::CreateContext()
 
 void FBuildJob::EnterCacheQuery()
 {
-	const EBuildPolicy BuildPolicy = Context->GetBuildPolicy();
 	ECachePolicy CachePolicy = Context->GetCachePolicy();
 	if (!EnumHasAnyFlags(CachePolicy, ECachePolicy::Query) ||
 		EnumHasAnyFlags(BuildPolicy, EBuildPolicy::SkipCacheGet))
 	{
 		return AdvanceToState(EBuildJobState::ExecuteRemote);
 	}
-	Scheduler->DispatchCacheQuery(this);
+	Scheduler->DispatchCacheQuery(this, SchedulerParams);
 }
 
 void FBuildJob::BeginCacheQuery()
 {
 	if (CanExecuteState(EBuildJobState::CacheQuery))
 	{
-		const EBuildPolicy BuildPolicy = Context->GetBuildPolicy();
 		ECachePolicy CachePolicy = Context->GetCachePolicy();
 		if (EnumHasAnyFlags(BuildPolicy, EBuildPolicy::SkipData))
 		{
@@ -573,7 +590,6 @@ void FBuildJob::EndCacheQuery(FCacheGetCompleteParams&& Params)
 
 void FBuildJob::EnterCacheStore()
 {
-	const EBuildPolicy BuildPolicy = Context->GetBuildPolicy();
 	ECachePolicy CachePolicy = Context->GetCachePolicy();
 	if (bIsCached ||
 		!EnumHasAnyFlags(CachePolicy, ECachePolicy::Store) ||
@@ -582,7 +598,7 @@ void FBuildJob::EnterCacheStore()
 	{
 		return AdvanceToState(EBuildJobState::Complete);
 	}
-	Scheduler->DispatchCacheStore(this);
+	Scheduler->DispatchCacheStore(this, SchedulerParams);
 }
 
 void FBuildJob::BeginCacheStore()
@@ -696,7 +712,8 @@ void FBuildJob::EndResolveInputMeta(FBuildInputMetaResolvedParams&& Params)
 
 void FBuildJob::EnterResolveInputData()
 {
-	// Identify missing inputs when resolving every input for the action.
+	SchedulerParams.MissingLocalInputsSize = 0;
+	SchedulerParams.MissingRemoteInputsSize = 0;
 	if (State == EBuildJobState::ResolveInputData)
 	{
 		MissingInputs.Reset();
@@ -705,6 +722,7 @@ void FBuildJob::EnterResolveInputData()
 			if (!Inputs || Inputs.Get().FindInput(Key).IsNull())
 			{
 				MissingInputs.Emplace(Key);
+				SchedulerParams.MissingLocalInputsSize += RawSize;
 			}
 		});
 		if (MissingInputs.IsEmpty())
@@ -712,6 +730,16 @@ void FBuildJob::EnterResolveInputData()
 			return AdvanceToState(EBuildJobState::ExecuteLocal);
 		}
 	}
+	else
+	{
+		Action.Get().IterateInputs([this](FStringView Key, const FIoHash& RawHash, uint64 RawSize)
+		{
+			if (MissingInputs.FindByKey(Key))
+			{
+				SchedulerParams.MissingRemoteInputsSize += RawSize;
+			}
+		});
+	}
 
 	checkf(!MissingInputs.IsEmpty(),
 		TEXT("Job is not expected to be in state %s without missing inputs for build of '%s' by %s"),
@@ -721,7 +749,7 @@ void FBuildJob::EnterResolveInputData()
 	{
 		return CompleteWithError(TEXT("Failed to resolve input data due to null input resolver."_SV));
 	}
-	Scheduler->DispatchResolveInputData(this);
+	Scheduler->DispatchResolveInputData(this, SchedulerParams);
 }
 
 void FBuildJob::BeginResolveInputData()
@@ -753,6 +781,7 @@ void FBuildJob::EndResolveInputData(FBuildInputDataResolvedParams&& Params)
 			for (const FBuildInputDataByKey& Input : Params.Inputs)
 			{
 				Builder.AddInput(Input.Key, Input.Data);
+				SchedulerParams.ResolvedInputsSize += Input.Data.GetRawSize();
 			}
 			if (Inputs)
 			{
@@ -774,7 +803,6 @@ void FBuildJob::EndResolveInputData(FBuildInputDataResolvedParams&& Params)
 
 void FBuildJob::EnterExecuteRemote()
 {
-	const EBuildPolicy BuildPolicy = Context->GetBuildPolicy();
 	if (EnumHasAnyFlags(BuildPolicy, EBuildPolicy::SkipBuild))
 	{
 		return CompleteWithError(TEXT("Failed because build policy skipped building."_SV));
@@ -798,7 +826,7 @@ void FBuildJob::EnterExecuteRemote()
 		return AdvanceToState(EBuildJobState::ResolveInputData);
 	}
 
-	Scheduler->DispatchExecuteRemote(this);
+	Scheduler->DispatchExecuteRemote(this, SchedulerParams);
 }
 
 void FBuildJob::BeginExecuteRemote()
@@ -808,7 +836,6 @@ void FBuildJob::BeginExecuteRemote()
 		checkf(Worker && WorkerExecutor, TEXT("Job requires a worker in state %s for build of '%s' by %s"),
 			LexToString(State), *Name, *FunctionName);
 		bTriedRemoteExecution = true;
-		const EBuildPolicy BuildPolicy = Context->GetBuildPolicy();
 		AdvanceToState(GetNextState(State), WorkerExecutor->BuildAction(Action.Get(), Inputs, *Worker, BuildPolicy, Priority,
 			[this](FBuildWorkerActionCompleteParams&& Params) { EndExecuteRemote(MoveTemp(Params)); }));
 	}
@@ -848,11 +875,22 @@ void FBuildJob::EndExecuteRemote(FBuildWorkerActionCompleteParams&& Params)
 	}
 }
 
+void FBuildJob::SkipExecuteRemote()
+{
+	checkf(State == EBuildJobState::ResolveRemoteInputData ||
+		State == EBuildJobState::ExecuteRemote || State == EBuildJobState::ExecuteRemoteRetry,
+		TEXT("Job is not expecting SkipExecuteRemote to be called in state %s for build of '%s' by %s"),
+		LexToString(State), *Name, *FunctionName);
+	if (CanExecuteState(EBuildJobState::ResolveInputData) || CanExecuteState(EBuildJobState::ExecuteRemote))
+	{
+		return AdvanceToState(EBuildJobState::ResolveInputData);
+	}
+}
+
 ///////////////////////////////////////////////////////////////////////////////////////////////////////////////////////
 
 void FBuildJob::EnterExecuteLocal()
 {
-	const EBuildPolicy BuildPolicy = Context->GetBuildPolicy();
 	if (EnumHasAnyFlags(BuildPolicy, EBuildPolicy::SkipBuild))
 	{
 		return CompleteWithError(TEXT("Failed because build policy skipped building."_SV));
@@ -892,9 +930,10 @@ void FBuildJob::EnterExecuteLocal()
 	});
 	Action.Reset();
 	Inputs.Reset();
+	SchedulerParams.ResolvedInputsSize = 0;
 	if (bMatchedInputs)
 	{
-		Scheduler->DispatchExecuteLocal(this);
+		Scheduler->DispatchExecuteLocal(this, SchedulerParams);
 	}
 }
 
@@ -961,7 +1000,6 @@ void FBuildJob::SetAction(FBuildAction&& InAction)
 		State == EBuildJobState::ResolveInputMeta || State == EBuildJobState::ResolveInputMetaWait,
 		TEXT("Job is not expecting an action in state %s for build of '%s' by %s"),
 		LexToString(State), *Name, *FunctionName);
-	ActionKey = InAction.GetKey();
 	Action = MoveTemp(InAction);
 	return AdvanceToState(EBuildJobState::CacheQuery);
 }
@@ -985,6 +1023,8 @@ void FBuildJob::SetInputs(FBuildInputs&& InInputs)
 		return;
 	}
 	Inputs = MoveTemp(InInputs);
+	SchedulerParams.MissingLocalInputsSize = 0;
+	SchedulerParams.MissingRemoteInputsSize = 0;
 	return AdvanceToState(ExecuteState);
 }
 
@@ -1004,10 +1044,14 @@ void FBuildJob::SetOutputNoCheck(FBuildOutput&& InOutput)
 	checkf(Output.IsNull(), TEXT("Job already has an output for build of '%s' by %s."), *Name, *FunctionName);
 	Output = MoveTemp(InOutput);
 
-	const EStatus Status = bIsCanceled ? EStatus::Canceled : Output.Get().HasError() ? EStatus::Error : EStatus::Ok;
-	Scheduler->CompleteJob({*this, FBuildOutput(Output.Get()), Status});
+	if (!bIsCanceled)
+	{
+		Scheduler->SetJobOutput(this, SchedulerParams, Output.Get());
+	}
+
 	if (OnComplete)
 	{
+		const EStatus Status = bIsCanceled ? EStatus::Canceled : Output.Get().HasError() ? EStatus::Error : EStatus::Ok;
 		OnComplete({*this, FBuildOutput(Output.Get()), Status});
 		OnComplete = nullptr;
 	}
@@ -1146,6 +1190,7 @@ void FBuildJob::ExecuteTransition(EBuildJobState OldState, EBuildJobState NewSta
 	{
 		Action.Reset();
 		Inputs.Reset();
+		SchedulerParams.ResolvedInputsSize = 0;
 	}
 	if (OldState <= EBuildJobState::ExecuteLocalWait && EBuildJobState::ExecuteLocalWait < NewState && !Output)
 	{
diff --git a/Engine/Source/Developer/DerivedDataCache/Public/DerivedDataBuildJob.h b/Engine/Source/Developer/DerivedDataCache/Public/DerivedDataBuildJob.h
--- a/Engine/Source/Developer/DerivedDataCache/Public/DerivedDataBuildJob.h
+++ b/Engine/Source/Developer/DerivedDataCache/Public/DerivedDataBuildJob.h
@@ -46,11 +46,6 @@ public:
 	/** Returns the priority of this job. */
 	virtual EPriority GetPriority() const = 0;
 
-	/** Returns the key, or null if the job was created directly from an action. */
-	virtual const FBuildKey& GetKey() const = 0;
-	/** Returns the action key, or null if the action has not been resolved yet. */
-	virtual const FBuildActionKey& GetActionKey() const = 0;
-
 	/** Returns the cache associated with this job. */
 	virtual ICache& GetCache() const = 0;
 	/** Returns the build system associated with this job. */
@@ -68,6 +63,9 @@ public:
 	virtual void BeginExecuteRemote() = 0;
 	virtual void BeginExecuteLocal() = 0;
 
+	/** Called by the scheduler to skip remote execution and fall back to local execution if permitted. */
+	virtual void SkipExecuteRemote() = 0;
+
 	/** Called by the scheduler if it has cached output compatible with the build policy. */
 	virtual void SetOutput(const FBuildOutput& Output) = 0;
 };
diff --git a/Engine/Source/Developer/DerivedDataCache/Public/DerivedDataBuildScheduler.h b/Engine/Source/Developer/DerivedDataCache/Public/DerivedDataBuildScheduler.h
--- a/Engine/Source/Developer/DerivedDataCache/Public/DerivedDataBuildScheduler.h
+++ b/Engine/Source/Developer/DerivedDataCache/Public/DerivedDataBuildScheduler.h
@@ -3,6 +3,9 @@
 #pragma once
 
 #include "DerivedDataBuildJob.h"
+#include "DerivedDataBuildKey.h"
+
+namespace UE::DerivedData { struct FBuildSchedulerParams; }
 
 namespace UE::DerivedData
 {
@@ -32,17 +35,49 @@ public:
 	/** Update the priority of the job if it is queued. May be called multiple times and/or concurrently. */
 	virtual void UpdateJobPriority(IBuildJob* Job) {}
 
-	/** Dispatch the job by executing the requested state now or queuing it for later. */
-	virtual void DispatchCacheQuery(IBuildJob* Job) { Job->BeginCacheQuery(); }
-	virtual void DispatchCacheStore(IBuildJob* Job) { Job->BeginCacheStore(); }
+	/** Dispatch by calling BeginCacheQuery or SetOutput, either now or later. */
+	virtual void DispatchCacheQuery(IBuildJob* Job, const FBuildSchedulerParams& Params) { Job->BeginCacheQuery(); }
+
+	/** Dispatch by calling BeginCacheStore, either now or later. */
+	virtual void DispatchCacheStore(IBuildJob* Job, const FBuildSchedulerParams& Params) { Job->BeginCacheStore(); }
+
+	/** Dispatch by calling BeginResolveKey, either now or later. */
 	virtual void DispatchResolveKey(IBuildJob* Job) { Job->BeginResolveKey(); }
+
+	/** Dispatch by calling BeginResolveInputMeta, either now or later. */
 	virtual void DispatchResolveInputMeta(IBuildJob* Job) { Job->BeginResolveInputMeta(); }
-	virtual void DispatchResolveInputData(IBuildJob* Job) { Job->BeginResolveInputData(); }
-	virtual void DispatchExecuteRemote(IBuildJob* Job) { Job->BeginExecuteRemote(); }
-	virtual void DispatchExecuteLocal(IBuildJob* Job) { Job->BeginExecuteLocal(); }
 
-	/** Called when the output of the job is available. Always called once, and always between BeginJob and EndJob. */
-	virtual void CompleteJob(const FBuildJobCompleteParams& Params) {}
+	/**
+	 * Dispatch by calling BeginResolveInputData, SetOutput, or SkipExecuteRemote, either now or later.
+	 *
+	 * SkipExecuteRemote is only valid to call when MissingRemoteInputsSize is non-zero.
+	 */
+	virtual void DispatchResolveInputData(IBuildJob* Job, const FBuildSchedulerParams& Params) { Job->BeginResolveInputData(); }
+
+	/** Dispatch by calling BeginExecuteRemote, SetOutput, or SkipExecuteRemote, either now or later. */
+	virtual void DispatchExecuteRemote(IBuildJob* Job, const FBuildSchedulerParams& Params) { Job->BeginExecuteRemote(); }
+
+	/** Dispatch by calling BeginExecuteLocal or SetOutput, either now or later. */
+	virtual void DispatchExecuteLocal(IBuildJob* Job, const FBuildSchedulerParams& Params) { Job->BeginExecuteLocal(); }
+
+	/** Set the output of the job. Always called once between BeginJob and EndJob unless canceled. */
+	virtual void SetJobOutput(IBuildJob* Job, const FBuildSchedulerParams& Params, const FBuildOutput& Output) {}
+};
+
+/** Parameters that describe a build job to the build scheduler. */
+struct FBuildSchedulerParams
+{
+	FBuildActionKey Key;
+
+	/** Total size of constants and inputs, whether resolved or not. */
+	uint64 TotalInputsSize = 0;
+	/** Total size of constants and resolved inputs that are in memory now. */
+	uint64 ResolvedInputsSize = 0;
+
+	/** Total size of inputs that need to be resolved for local execution. Available in ResolveInputData. */
+	uint64 MissingLocalInputsSize = 0;
+	/** Total size of inputs that need to be resolved for remote execution. Available in ResolveInputData. */
+	uint64 MissingRemoteInputsSize = 0;
 };
 
 } // UE::DerivedData
