Integrating //UE5/Dev-Cooker @ 16678003 to //UE5/Main (Zousar.Shaker-YEG-0943-Quaternary)

Non-DevIteration_ShooterGame

diff --git a/Engine/Config/BaseEditor.ini b/Engine/Config/BaseEditor.ini
--- a/Engine/Config/BaseEditor.ini
+++ b/Engine/Config/BaseEditor.ini
@@ -372,6 +372,7 @@ MaxMemoryAllowance=16384
 [CookSettings]
 ; EditorDomain is in prototype and currently is a significant deoptimization
 EditorDomainEnabled=false
+EditorDomainSaveUnversioned=false
 PackagesPerGC=0
 IdleTimeToGC=20
 MemoryMaxUsedPhysical=16384
diff --git a/Engine/Plugins/Editor/ContentBrowser/ContentBrowserAssetDataSource/Source/ContentBrowserAssetDataSource/Private/AssetFileContextMenu.cpp b/Engine/Plugins/Editor/ContentBrowser/ContentBrowserAssetDataSource/Source/ContentBrowserAssetDataSource/Private/AssetFileContextMenu.cpp
--- a/Engine/Plugins/Editor/ContentBrowser/ContentBrowserAssetDataSource/Source/ContentBrowserAssetDataSource/Private/AssetFileContextMenu.cpp
+++ b/Engine/Plugins/Editor/ContentBrowser/ContentBrowserAssetDataSource/Source/ContentBrowserAssetDataSource/Private/AssetFileContextMenu.cpp
@@ -77,6 +77,7 @@
 
 #include "Commandlets/TextAssetCommandlet.h"
 #include "Misc/FileHelper.h"
+#include "Internationalization/GatherableTextData.h"
 
 #include "ContentBrowserDataMenuContexts.h"
 
diff --git a/Engine/Source/Developer/DerivedDataCache/Private/DerivedDataBackends.cpp b/Engine/Source/Developer/DerivedDataCache/Private/DerivedDataBackends.cpp
--- a/Engine/Source/Developer/DerivedDataCache/Private/DerivedDataBackends.cpp
+++ b/Engine/Source/Developer/DerivedDataCache/Private/DerivedDataBackends.cpp
@@ -17,6 +17,8 @@
 #include "DerivedDataBackendAsyncPutWrapper.h"
 #include "PakFileDerivedDataBackend.h"
 #include "S3DerivedDataBackend.h"
+#include "ZenKvDerivedDataBackend.h"
+#include "ZenDerivedDataBackend.h"
 #include "HierarchicalDerivedDataBackend.h"
 #include "DerivedDataLimitKeyLengthWrapper.h"
 #include "DerivedDataBackendCorruptionWrapper.h"
@@ -236,6 +238,14 @@ public:
 				{
 					ParsedNode = ParseHttpCache(NodeName, *Entry);
 				}
+				else if (NodeType == TEXT("ZenHttp"))
+				{
+					ParsedNode = ParseZenHttpCache(NodeName, *Entry);
+				}
+				else if (NodeType == TEXT("Zen"))
+				{
+					ParsedNode = ParseZenCache(NodeName, *Entry);
+				}
 			}
 		}
 
@@ -838,6 +848,74 @@ public:
 #endif
 	}
 
+	/**
+	 * Creates a Zen / HTTP data cache interface (DDC1 style)
+	 */
+	FDerivedDataBackendInterface * ParseZenHttpCache(const TCHAR * NodeName, const TCHAR * Entry)
+	{
+#if WITH_ZEN_DDC_BACKEND
+		FString ServiceUrl;
+		if (!FParse::Value(Entry, TEXT("Host="), ServiceUrl))
+		{
+			UE_LOG(LogDerivedDataCache, Error, TEXT("Node %s does not specify 'Host'."), NodeName);
+			return nullptr;
+		}
+
+		FString Namespace;
+		if (!FParse::Value(Entry, TEXT("Namespace="), Namespace))
+		{
+			Namespace = FApp::GetProjectName();
+			UE_LOG(LogDerivedDataCache, Warning, TEXT("Node %s does not specify 'Namespace', falling back to '%s'"), NodeName, *Namespace);
+		}
+
+		FZenHttpDerivedDataBackend* backend = new FZenHttpDerivedDataBackend(*this, *ServiceUrl, *Namespace);
+		if (!backend->IsUsable())
+		{
+			UE_LOG(LogDerivedDataCache, Warning, TEXT("%s could not contact the service (%s), will not use it."), NodeName, *ServiceUrl);
+			delete backend;
+			return nullptr;
+		}
+		return backend;
+#else
+		UE_LOG(LogDerivedDataCache, Warning, TEXT("ZenHTTP backend is not yet supported in the current build configuration."));
+		return nullptr;
+#endif
+	}
+
+	/**
+	 * Creates a Zen data cache interface
+	 */
+	FDerivedDataBackendInterface* ParseZenCache(const TCHAR* NodeName, const TCHAR* Entry)
+	{
+#if WITH_ZEN_DDC_BACKEND
+		FString ServiceUrl;
+		if (!FParse::Value(Entry, TEXT("Host="), ServiceUrl))
+		{
+			UE_LOG(LogDerivedDataCache, Error, TEXT("Node %s does not specify 'Host'."), NodeName);
+			return nullptr;
+		}
+
+		FString Namespace;
+		if (!FParse::Value(Entry, TEXT("Namespace="), Namespace))
+		{
+			Namespace = FApp::GetProjectName();
+			UE_LOG(LogDerivedDataCache, Warning, TEXT("Node %s does not specify 'Namespace', falling back to '%s'"), NodeName, *Namespace);
+		}
+
+		FZenDerivedDataBackend* backend = new FZenDerivedDataBackend(*this, *ServiceUrl, *Namespace);
+		if (!backend->IsUsable())
+		{
+			UE_LOG(LogDerivedDataCache, Warning, TEXT("%s could not contact the service (%s), will not use it."), NodeName, *ServiceUrl);
+			delete backend;
+			return nullptr;
+		}
+		return backend;
+#else
+		UE_LOG(LogDerivedDataCache, Warning, TEXT("Zen backend is not yet supported in the current build configuration."));
+		return nullptr;
+#endif
+	}
+
 	/**
 	 * Creates Boot data cache interface from ini settings.
 	 *
diff --git a/Engine/Source/Developer/DerivedDataCache/Private/ZenDerivedDataBackend.cpp b/Engine/Source/Developer/DerivedDataCache/Private/ZenDerivedDataBackend.cpp
new file mode 100644
--- /dev/null
+++ b/Engine/Source/Developer/DerivedDataCache/Private/ZenDerivedDataBackend.cpp
@@ -0,0 +1,1794 @@
+// Copyright Epic Games, Inc. All Rights Reserved.
+#include "ZenDerivedDataBackend.h"
+
+#if WITH_ZEN_DDC_BACKEND
+
+#if PLATFORM_WINDOWS || PLATFORM_HOLOLENS
+#	include "Windows/WindowsHWrapper.h"
+#	include "Windows/AllowWindowsPlatformTypes.h"
+#endif
+
+#include "curl/curl.h"
+
+#if PLATFORM_WINDOWS || PLATFORM_HOLOLENS
+#	include "Windows/HideWindowsPlatformTypes.h"
+#endif
+
+#include "Algo/Accumulate.h"
+#include "Containers/StaticArray.h"
+#include "Containers/Ticker.h"
+#include "DerivedDataCacheRecord.h"
+#include "Dom/JsonObject.h"
+#include "GenericPlatform/GenericPlatformFile.h"
+#include "HAL/PlatformFileManager.h"
+#include "Misc/ConfigCacheIni.h"
+#include "Misc/FileHelper.h"
+#include "Misc/ScopeLock.h"
+#include "Misc/SecureHash.h"
+#include "ProfilingDebugging/CpuProfilerTrace.h"
+#include "ProfilingDebugging/CountersTrace.h"
+#include "Serialization/CompactBinary.h"
+#include "Serialization/CompactBinaryPackage.h"
+#include "Serialization/CompactBinaryValidation.h"
+#include "Serialization/CompactBinaryWriter.h"
+#include "Serialization/JsonReader.h"
+#include "Serialization/JsonSerializer.h"
+#include "Serialization/MemoryReader.h"
+#include "ZenBackendUtils.h"
+
+#define UE_ZENDDC_BACKEND_WAIT_INTERVAL			0.01f
+#define UE_ZENDDC_HTTP_REQUEST_TIMEOUT_SECONDS	30L
+#define UE_ZENDDC_HTTP_REQUEST_TIMOUT_ENABLED	1
+#define UE_ZENDDC_HTTP_DEBUG					0
+#define UE_ZENDDC_REQUEST_POOL_SIZE				16
+#define UE_ZENDDC_MAX_FAILED_LOGIN_ATTEMPTS		16
+#define UE_ZENDDC_MAX_ATTEMPTS					4
+
+TRACE_DECLARE_INT_COUNTER(ZenDDC_Exist,			TEXT("ZenDDC Exist"));
+TRACE_DECLARE_INT_COUNTER(ZenDDC_ExistHit,		TEXT("ZenDDC Exist Hit"));
+TRACE_DECLARE_INT_COUNTER(ZenDDC_Get,			TEXT("ZenDDC Get"));
+TRACE_DECLARE_INT_COUNTER(ZenDDC_GetHit,		TEXT("ZenDDC Get Hit"));
+TRACE_DECLARE_INT_COUNTER(ZenDDC_Put,			TEXT("ZenDDC Put"));
+TRACE_DECLARE_INT_COUNTER(ZenDDC_PutHit,		TEXT("ZenDDC Put Hit"));
+TRACE_DECLARE_INT_COUNTER(ZenDDC_BytesReceived, TEXT("ZenDDC Bytes Received"));
+TRACE_DECLARE_INT_COUNTER(ZenDDC_BytesSent,		TEXT("ZenDDC Bytes Sent"));
+
+namespace zen {
+
+	static bool IsSuccessCode(int64 ResponseCode)
+	{
+		return 200 <= ResponseCode && ResponseCode < 300;
+	}
+
+	enum class EContentType
+	{
+		Binary,
+		CompactBinary,
+		CompactBinaryPackage
+	};
+
+/**
+ * Encapsulation for access token shared by all requests.
+ */
+
+	/**
+	 * Minimal HTTP request type wrapping CURL without the need for managers. This request
+	 * is written to allow reuse of request objects, in order to allow connections to be reused.
+	 *
+	 * CURL has a global library initialization (curl_global_init). We rely on this happening in
+	 * the Online/HTTP library which is a dependency on this module.
+	 */
+	class FRequest
+	{
+	public:
+		/**
+		 * Supported request verbs
+		 */
+		enum RequestVerb
+		{
+			Get,
+			Put,
+			Post,
+			PostJson,
+			Delete,
+			Head
+		};
+
+		/**
+		 * Convenience result type interpreted from HTTP response code.
+		 */
+		enum Result
+		{
+			Success,
+			Failed,
+			FailedTimeout
+		};
+
+		FRequest(const TCHAR* InDomain, bool bInLogErrors)
+			: bLogErrors(bInLogErrors)
+			, Domain(InDomain)
+		{
+			Curl = curl_easy_init();
+			Reset();
+		}
+
+		~FRequest()
+		{
+			curl_easy_cleanup(Curl);
+		}
+
+		/**
+		 * Resets all options on the request except those that should always be set.
+		 */
+		void Reset()
+		{
+			Headers.Reset();
+			ResponseHeader.Reset();
+			ResponseBuffer.Reset();
+			ResponseCode = 0;
+			bResponseFormatValid = true;
+			ReadDataView = nullptr;
+			WriteDataBufferPtr = nullptr;
+			WriteHeaderBufferPtr = nullptr;
+			BytesSent = 0;
+			BytesReceived = 0;
+			CurlResult = CURL_LAST;
+
+			curl_easy_reset(Curl);
+
+			// Options that are always set for all connections.
+#if UE_ZENDDC_HTTP_REQUEST_TIMOUT_ENABLED
+			curl_easy_setopt(Curl, CURLOPT_CONNECTTIMEOUT, UE_ZENDDC_HTTP_REQUEST_TIMEOUT_SECONDS);
+#endif
+			curl_easy_setopt(Curl, CURLOPT_FOLLOWLOCATION, 1L);
+			curl_easy_setopt(Curl, CURLOPT_NOSIGNAL, 1L);
+			curl_easy_setopt(Curl, CURLOPT_BUFFERSIZE, 256 * 1024L);
+			//curl_easy_setopt(Curl, CURLOPT_UPLOAD_BUFFERSIZE, 256 * 1024L);
+			// Response functions
+			curl_easy_setopt(Curl, CURLOPT_HEADERDATA, this);
+			curl_easy_setopt(Curl, CURLOPT_HEADERFUNCTION, &FRequest::StaticWriteHeaderFn);
+			curl_easy_setopt(Curl, CURLOPT_WRITEDATA, this);
+			curl_easy_setopt(Curl, CURLOPT_WRITEFUNCTION, StaticWriteBodyFn);
+			// Rewind method, handle special error case where request need to rewind data stream
+			curl_easy_setopt(Curl, CURLOPT_SEEKFUNCTION, StaticSeekFn);
+			curl_easy_setopt(Curl, CURLOPT_SEEKDATA, this);
+			// Debug hooks
+#if UE_ZENDDC_HTTP_DEBUG
+			curl_easy_setopt(Curl, CURLOPT_DEBUGDATA, this);
+			curl_easy_setopt(Curl, CURLOPT_DEBUGFUNCTION, StaticDebugCallback);
+			curl_easy_setopt(Curl, CURLOPT_VERBOSE, 1L);
+#endif
+		}
+
+		/** Gets the domain name for this request */
+		const FString& GetDomain() const
+		{
+			return Domain;
+		}
+
+		/** Returns the HTTP response code.*/
+		const int64 GetResponseCode() const
+		{
+			return ResponseCode;
+		}
+
+		const bool GetResponseFormatValid() const
+		{
+			return bResponseFormatValid;
+		}
+
+		/** Returns the number of bytes received this request (headers withstanding). */
+		const size_t GetBytesReceived() const
+		{
+			return BytesReceived;
+		}
+
+		/** Returns the number of bytes sent during this request (headers withstanding). */
+		const size_t  GetBytesSent() const
+		{
+			return BytesSent;
+		}
+
+		/**
+		 * Upload buffer using the request, using either "Put" or "Post" verbs.
+		 * @param Uri Url to use.
+		 * @param Buffer Data to upload
+		 * @return Result of the request
+		 */
+		template<RequestVerb V>
+		Result PerformBlockingUpload(const TCHAR* Uri, const FCompositeBuffer& Buffer, EContentType ContentType)
+		{
+			static_assert(V == Put || V == Post || V == PostJson, "Upload should use either Put or Post verbs.");
+
+			uint32 ContentLength = 0u;
+
+			if (V == Put)
+			{
+				ContentLength = Buffer.GetSize();
+				curl_easy_setopt(Curl, CURLOPT_UPLOAD, 1L);
+				curl_easy_setopt(Curl, CURLOPT_INFILESIZE, ContentLength);
+				curl_easy_setopt(Curl, CURLOPT_READDATA, this);
+				curl_easy_setopt(Curl, CURLOPT_READFUNCTION, StaticReadFn);
+
+				switch (ContentType)
+				{
+				case EContentType::Binary:
+					Headers.Add(FString(TEXT("Content-Type: application/octet-stream")));
+					break;
+				case EContentType::CompactBinary:
+					Headers.Add(FString(TEXT("Content-Type: application/x-ue-cb")));
+					break;
+				case EContentType::CompactBinaryPackage:
+					Headers.Add(FString(TEXT("Content-Type: application/x-ue-cbpkg")));
+					break;
+				default:
+					checkNoEntry();
+					break;
+				}
+				ReadDataView = &Buffer;
+			}
+			else if (V == Post || V == PostJson)
+			{
+				ContentLength = Buffer.GetSize();
+				curl_easy_setopt(Curl, CURLOPT_POST, 1L);
+				curl_easy_setopt(Curl, CURLOPT_INFILESIZE, ContentLength);
+				curl_easy_setopt(Curl, CURLOPT_READDATA, this);
+				curl_easy_setopt(Curl, CURLOPT_READFUNCTION, StaticReadFn);
+
+				switch (ContentType)
+				{
+				case EContentType::Binary:
+					Headers.Add(V == Post ? FString(TEXT("Content-Type: application/x-www-form-urlencoded")) : FString(TEXT("Content-Type: application/json")));
+					break;
+				case EContentType::CompactBinary:
+					Headers.Add(FString(TEXT("Content-Type: application/x-ue-cb")));
+					break;
+				case EContentType::CompactBinaryPackage:
+					Headers.Add(FString(TEXT("Content-Type: application/x-ue-cbpkg")));
+					break;
+				default:
+					checkNoEntry();
+					break;
+				}
+
+				ReadDataView = &Buffer;
+			}
+
+			return PerformBlocking(Uri, V, ContentLength);
+		}
+
+		/**
+		 * Download an url into a buffer using the request.
+		 * @param Uri Url to use.
+		 * @param Buffer Optional buffer where data should be downloaded to. If empty downloaded data will
+		 * be stored in an internal buffer and accessed GetResponse* methods.
+		 * @return Result of the request
+		 */
+		Result PerformBlockingDownload(const TCHAR* Uri, TArray<uint8>* Buffer)
+		{
+			curl_easy_setopt(Curl, CURLOPT_HTTPGET, 1L);
+			WriteDataBufferPtr = Buffer;
+
+			return PerformBlocking(Uri, Get, 0u);
+		}
+
+		Result PerformBlockingDownload(const TCHAR* Uri, FCbPackage& OutPackage)
+		{
+			curl_easy_setopt(Curl, CURLOPT_HTTPGET, 1L);
+			OutPackage.Reset();
+			// TODO: When PackageBytes can be written in segments directly, set the WritePtr to the OutPackage and use that
+			TArray<uint8> PackageBytes;
+			WriteDataBufferPtr = &PackageBytes;
+			Result LocalResult = PerformBlocking(Uri, Get, 0u);
+			if (IsSuccessCode(ResponseCode))
+			{
+				if (ValidateCompactBinaryPackage(FMemoryView(PackageBytes.GetData(), PackageBytes.Num()), ECbValidateMode::All) != ECbValidateError::None)
+				{
+					bResponseFormatValid = false;
+				}
+				else
+				{
+					FMemoryReader PackageLoader(PackageBytes);
+					bResponseFormatValid = OutPackage.TryLoad(PackageLoader);
+				}
+			}
+			return LocalResult;
+		}
+
+		/**
+		 * Query an url using the request. Queries can use either "Head" or "Delete" verbs.
+		 * @param Uri Url to use.
+		 * @return Result of the request
+		 */
+		template<RequestVerb V>
+		Result PerformBlockingQuery(const TCHAR* Uri)
+		{
+			static_assert(V == Head || V == Delete, "Queries should use either Head or Delete verbs.");
+
+			if (V == Delete)
+			{
+				curl_easy_setopt(Curl, CURLOPT_POST, 1L);
+				curl_easy_setopt(Curl, CURLOPT_CUSTOMREQUEST, "DELETE");
+			}
+			else if (V == Head)
+			{
+				curl_easy_setopt(Curl, CURLOPT_NOBODY, 1L);
+			}
+
+			return PerformBlocking(Uri, V, 0u);
+		}
+
+		/**
+		 * Set a header to send with the request.
+		 */
+		void SetHeader(const TCHAR* Header, const TCHAR* Value)
+		{
+			check(CurlResult == CURL_LAST); // Cannot set header after request is sent
+			Headers.Add(FString::Printf(TEXT("%s: %s"), Header, Value));
+		}
+
+		/**
+		 * Attempts to find the header from the response. Returns false if header is not present.
+		 */
+		bool GetHeader(const ANSICHAR* Header, FString& OutValue) const
+		{
+			check(CurlResult != CURL_LAST);  // Cannot query headers before request is sent
+
+			const ANSICHAR* HeadersBuffer = (const ANSICHAR*)ResponseHeader.GetData();
+			size_t HeaderLen = strlen(Header);
+
+			// Find the header key in the (ANSI) response buffer. If not found we can exist immediately
+			if (const ANSICHAR* Found = strstr(HeadersBuffer, Header))
+			{
+				const ANSICHAR* Linebreak = strchr(Found, '\r');
+				const ANSICHAR* ValueStart = Found + HeaderLen + 2; //colon and space
+				const size_t ValueSize = Linebreak - ValueStart;
+				FUTF8ToTCHAR TCHARData(ValueStart, ValueSize);
+				OutValue = FString(TCHARData.Length(), TCHARData.Get());
+				return true;
+			}
+			return false;
+		}
+
+		/**
+		 * Returns the response buffer. Note that is the request is performed
+		 * with an external buffer as target buffer this string will be empty.
+		 */
+		const TArray<uint8>& GetResponseBuffer() const
+		{
+			return ResponseBuffer;
+		}
+
+		/**
+		 * Returns the response buffer as a string. Note that is the request is performed
+		 * with an external buffer as target buffer this string will be empty.
+		 */
+		FString GetResponseAsString() const
+		{
+			return GetAnsiBufferAsString(ResponseBuffer);
+		}
+
+		/**
+		 * Returns the response header as a string.
+		 */
+		FString GetResponseHeaderAsString()
+		{
+			return GetAnsiBufferAsString(ResponseHeader);
+		}
+
+		/**
+		 * Tries to parse the response buffer as a JsonObject. Return empty pointer if
+		 * parse error occurs.
+		 */
+		TSharedPtr<FJsonObject> GetResponseAsJsonObject() const
+		{
+			FString Response = GetAnsiBufferAsString(ResponseBuffer);
+
+			TSharedPtr<FJsonObject> JsonObject;
+			TSharedRef<TJsonReader<> > JsonReader = TJsonReaderFactory<>::Create(Response);
+			if (!FJsonSerializer::Deserialize(JsonReader, JsonObject) || !JsonObject.IsValid())
+			{
+				return TSharedPtr<FJsonObject>(nullptr);
+			}
+
+			return JsonObject;
+		}
+
+	private:
+
+		CURL* Curl;
+		CURLcode				CurlResult;
+		long					ResponseCode;
+		size_t					BytesSent;
+		size_t					BytesReceived;
+		bool					bLogErrors;
+		bool					bResponseFormatValid;
+
+		const FCompositeBuffer*	ReadDataView;
+		TArray<uint8>*			WriteDataBufferPtr;
+		FCbPackage*				WriteDataPackage;
+		TArray<uint8>*			WriteHeaderBufferPtr;
+
+		TArray<uint8>			ResponseHeader;
+		TArray<uint8>			ResponseBuffer;
+		TArray<FString>			Headers;
+		FString					Domain;
+
+		/**
+		 * Performs the request, blocking until finished.
+		 * @param Uri Address on the domain to query
+		 * @param Verb HTTP verb to use
+		 * @param Buffer Optional buffer to directly receive the result of the request.
+		 * If unset the response body will be stored in the request.
+		 */
+		Result PerformBlocking(const TCHAR* Uri, RequestVerb Verb, uint32 ContentLength)
+		{
+			static const char* CommonHeaders[] = {
+				"User-Agent: UE4",
+				nullptr
+			};
+
+			TRACE_CPUPROFILER_EVENT_SCOPE(ZenDDC_CurlPerform);
+
+			// Setup request options
+			FString Url = FString::Printf(TEXT("%s/%s"), *Domain, Uri);
+			curl_easy_setopt(Curl, CURLOPT_URL, TCHAR_TO_ANSI(*Url));
+
+			// Setup response header buffer. If caller has not setup a response data buffer, use internal.
+			WriteHeaderBufferPtr = &ResponseHeader;
+			if (WriteDataBufferPtr == nullptr)
+			{
+				WriteDataBufferPtr = &ResponseBuffer;
+			}
+
+			// Content-Length should always be set
+			Headers.Add(FString::Printf(TEXT("Content-Length: %d"), ContentLength));
+
+			// Build headers list
+			curl_slist* CurlHeaders = nullptr;
+			// Add common headers
+			for (uint8 i = 0; CommonHeaders[i] != nullptr; ++i)
+			{
+				CurlHeaders = curl_slist_append(CurlHeaders, CommonHeaders[i]);
+			}
+			// Setup added headers
+			for (const FString& Header : Headers)
+			{
+				CurlHeaders = curl_slist_append(CurlHeaders, TCHAR_TO_ANSI(*Header));
+			}
+			curl_easy_setopt(Curl, CURLOPT_HTTPHEADER, CurlHeaders);
+
+			// Shots fired!
+			CurlResult = curl_easy_perform(Curl);
+
+			// Get response code
+			bool bRedirected = false;
+			if (CURLE_OK == curl_easy_getinfo(Curl, CURLINFO_RESPONSE_CODE, &ResponseCode))
+			{
+				bRedirected = (ResponseCode >= 300 && ResponseCode < 400);
+			}
+
+			LogResult(CurlResult, Uri, Verb);
+
+			// Clean up
+			curl_slist_free_all(CurlHeaders);
+
+			return CurlResult == CURLE_OK ? Success : Failed;
+		}
+
+		void LogResult(CURLcode Result, const TCHAR* Uri, RequestVerb Verb) const
+		{
+			if (Result == CURLE_OK)
+			{
+				bool bSuccess = false;
+				const TCHAR* VerbStr = nullptr;
+				FString AdditionalInfo;
+
+				const bool Is400 = (ResponseCode >= 400) && (ResponseCode <= 499);
+				const bool Is200 = (ResponseCode >= 200) && (ResponseCode <= 299);
+
+				switch (Verb)
+				{
+				case Head:
+					bSuccess = Is400 || Is200;
+					VerbStr = TEXT("querying");
+					break;
+				case Get:
+					bSuccess = Is400 || Is200;
+					VerbStr = TEXT("fetching");
+					AdditionalInfo = FString::Printf(TEXT("Received: %d bytes."), BytesReceived);
+					break;
+				case Put:
+					bSuccess = Is200;
+					VerbStr = TEXT("updating");
+					AdditionalInfo = FString::Printf(TEXT("Sent: %d bytes."), BytesSent);
+					break;
+				case Post:
+				case PostJson:
+					bSuccess = Is200;
+					VerbStr = TEXT("posting");
+					break;
+				case Delete:
+					bSuccess = Is200;
+					VerbStr = TEXT("deleting");
+					break;
+				}
+
+				if (bSuccess)
+				{
+					UE_LOG(
+						LogDerivedDataCache,
+						Verbose,
+						TEXT("Finished %s HTTP cache entry (response %d) from %s. %s"),
+						VerbStr,
+						ResponseCode,
+						Uri,
+						*AdditionalInfo
+					);
+				}
+				else if (bLogErrors)
+				{
+					// Print the response body if we got one, otherwise print header.
+					FString Response = GetAnsiBufferAsString(ResponseBuffer.Num() > 0 ? ResponseBuffer : ResponseHeader);
+					Response.ReplaceCharInline('\n', ' ');
+					Response.ReplaceCharInline('\r', ' ');
+					UE_LOG(
+						LogDerivedDataCache,
+						Error,
+						TEXT("Failed %s HTTP cache entry (response %d) from %s. Response: %s"),
+						VerbStr,
+						ResponseCode,
+						Uri,
+						*Response
+					);
+				}
+			}
+			else if (bLogErrors)
+			{
+				UE_LOG(
+					LogDerivedDataCache,
+					Error,
+					TEXT("Error while connecting to %s: %s"),
+					*Domain,
+					ANSI_TO_TCHAR(curl_easy_strerror(Result))
+				);
+			}
+		}
+
+		FString GetAnsiBufferAsString(const TArray<uint8>& Buffer) const
+		{
+			// Content is NOT null-terminated; we need to specify lengths here
+			FUTF8ToTCHAR TCHARData(reinterpret_cast<const ANSICHAR*>(Buffer.GetData()), Buffer.Num());
+			return FString(TCHARData.Length(), TCHARData.Get());
+		}
+
+		static size_t StaticDebugCallback(CURL* Handle, curl_infotype DebugInfoType, char* DebugInfo, size_t DebugInfoSize, void* UserData)
+		{
+			FRequest* Request = static_cast<FRequest*>(UserData);
+
+			switch (DebugInfoType)
+			{
+			case CURLINFO_TEXT:
+			{
+				// Truncate at 1023 characters. This is just an arbitrary number based on a buffer size seen in
+				// the libcurl code.
+				DebugInfoSize = FMath::Min(DebugInfoSize, (size_t)1023);
+
+				// Calculate the actual length of the string due to incorrect use of snprintf() in lib/vtls/openssl.c.
+				char* FoundNulPtr = (char*)memchr(DebugInfo, 0, DebugInfoSize);
+				int CalculatedSize = FoundNulPtr != nullptr ? FoundNulPtr - DebugInfo : DebugInfoSize;
+
+				auto ConvertedString = StringCast<TCHAR>(static_cast<const ANSICHAR*>(DebugInfo), CalculatedSize);
+				FString DebugText(ConvertedString.Length(), ConvertedString.Get());
+				DebugText.ReplaceInline(TEXT("\n"), TEXT(""), ESearchCase::CaseSensitive);
+				DebugText.ReplaceInline(TEXT("\r"), TEXT(""), ESearchCase::CaseSensitive);
+				UE_LOG(LogDerivedDataCache, VeryVerbose, TEXT("%p: '%s'"), Request, *DebugText);
+			}
+			break;
+
+			case CURLINFO_HEADER_IN:
+				UE_LOG(LogDerivedDataCache, VeryVerbose, TEXT("%p: Received header (%d bytes)"), Request, DebugInfoSize);
+				break;
+
+			case CURLINFO_DATA_IN:
+				UE_LOG(LogDerivedDataCache, VeryVerbose, TEXT("%p: Received data (%d bytes)"), Request, DebugInfoSize);
+				break;
+
+			case CURLINFO_DATA_OUT:
+				UE_LOG(LogDerivedDataCache, VeryVerbose, TEXT("%p: Sent data (%d bytes)"), Request, DebugInfoSize);
+				break;
+			}
+
+			return 0;
+		}
+
+		static size_t StaticReadFn(void* Ptr, size_t SizeInBlocks, size_t BlockSizeInBytes, void* UserData)
+		{
+			FRequest* Request = static_cast<FRequest*>(UserData);
+			check(Request->ReadDataView);
+			const FCompositeBuffer& ReadDataView = *Request->ReadDataView;
+
+			const size_t Offset = Request->BytesSent;
+			const size_t ReadSize = FMath::Min((size_t)ReadDataView.GetSize() - Offset, SizeInBlocks * BlockSizeInBytes);
+			check(ReadDataView.GetSize() >= Offset + ReadSize);
+
+			Memcpy(Ptr, ReadDataView, Offset, ReadSize);
+			Request->BytesSent += ReadSize;
+			return ReadSize;
+
+			return 0;
+		}
+
+		static size_t StaticWriteHeaderFn(void* Ptr, size_t SizeInBlocks, size_t BlockSizeInBytes, void* UserData)
+		{
+			FRequest* Request = static_cast<FRequest*>(UserData);
+			const size_t WriteSize = SizeInBlocks * BlockSizeInBytes;
+			TArray<uint8>* WriteHeaderBufferPtr = Request->WriteHeaderBufferPtr;
+			if (WriteHeaderBufferPtr && WriteSize > 0)
+			{
+				const size_t CurrentBufferLength = WriteHeaderBufferPtr->Num();
+				if (CurrentBufferLength > 0)
+				{
+					// Remove the previous zero termination
+					(*WriteHeaderBufferPtr)[CurrentBufferLength - 1] = ' ';
+				}
+
+				// Write the header
+				WriteHeaderBufferPtr->Append((const uint8*)Ptr, WriteSize + 1);
+				(*WriteHeaderBufferPtr)[WriteHeaderBufferPtr->Num() - 1] = 0; // Zero terminate string
+				return WriteSize;
+			}
+			return 0;
+		}
+
+		static size_t StaticWriteBodyFn(void* Ptr, size_t SizeInBlocks, size_t BlockSizeInBytes, void* UserData)
+		{
+			FRequest* Request = static_cast<FRequest*>(UserData);
+			const size_t WriteSize = SizeInBlocks * BlockSizeInBytes;
+			TArray<uint8>* WriteDataBufferPtr = Request->WriteDataBufferPtr;
+
+			if (WriteDataBufferPtr && WriteSize > 0)
+			{
+				// If this is the first part of the body being received, try to reserve 
+				// memory if content length is defined in the header.
+				if (Request->BytesReceived == 0 && Request->WriteHeaderBufferPtr)
+				{
+					static const ANSICHAR* ContentLengthHeaderStr = "Content-Length: ";
+					const ANSICHAR* Header = (const ANSICHAR*)Request->WriteHeaderBufferPtr->GetData();
+
+					if (const ANSICHAR* ContentLengthHeader = FCStringAnsi::Strstr(Header, ContentLengthHeaderStr))
+					{
+						size_t ContentLength = (size_t)FCStringAnsi::Atoi64(ContentLengthHeader + strlen(ContentLengthHeaderStr));
+						if (ContentLength > 0)
+						{
+							WriteDataBufferPtr->Reserve(ContentLength);
+						}
+					}
+				}
+
+				// Write to the target buffer
+				WriteDataBufferPtr->Append((const uint8*)Ptr, WriteSize);
+				Request->BytesReceived += WriteSize;
+				return WriteSize;
+			}
+
+			return 0;
+		}
+
+		static size_t StaticSeekFn(void* UserData, curl_off_t Offset, int Origin)
+		{
+			FRequest* Request = static_cast<FRequest*>(UserData);
+			size_t NewPosition = 0;
+			uint64 ReadDataSize = Request->ReadDataView ? Request->ReadDataView->GetSize() : 0;
+
+			switch (Origin)
+			{
+			case SEEK_SET: NewPosition = Offset; break;
+			case SEEK_CUR: NewPosition = Request->BytesSent + Offset; break;
+			case SEEK_END: NewPosition = ReadDataSize + Offset; break;
+			}
+
+			// Make sure we don't seek outside of the buffer
+			if (NewPosition < 0 || NewPosition >= ReadDataSize)
+			{
+				return CURL_SEEKFUNC_FAIL;
+			}
+
+			// Update the used offset
+			Request->BytesSent = NewPosition;
+			return CURL_SEEKFUNC_OK;
+		}
+
+	};
+
+	//----------------------------------------------------------------------------------------------------------
+	// Request pool
+	//----------------------------------------------------------------------------------------------------------
+
+	/**
+	 * Pool that manages a fixed set of requests. Users are required to release requests that have been
+	 * acquired. Usable with \ref zen::FScopedRequestPtr which handles this automatically.
+	 */
+	struct FRequestPool
+	{
+		FRequestPool(const TCHAR* InServiceUrl)
+		{
+			for (uint8 i = 0; i < Pool.Num(); ++i)
+			{
+				Pool[i].Usage = 0u;
+				Pool[i].Request = new FRequest(InServiceUrl, true);
+			}
+		}
+
+		~FRequestPool()
+		{
+			for (uint8 i = 0; i < Pool.Num(); ++i)
+			{
+				// No requests should be in use by now.
+				check(Pool[i].Usage.Load(EMemoryOrder::Relaxed) == 0u);
+				delete Pool[i].Request;
+			}
+		}
+
+		/**
+		 * Block until a request is free. Once a request has been returned it is
+		 * "owned by the caller and need to release it to the pool when work has been completed.
+		 * @return Usable request instance.
+		 */
+		FRequest* WaitForFreeRequest()
+		{
+			TRACE_CPUPROFILER_EVENT_SCOPE(ZenDDC_WaitForConnPool);
+			while (true)
+			{
+				for (uint8 i = 0; i < Pool.Num(); ++i)
+				{
+					if (!Pool[i].Usage.Load(EMemoryOrder::Relaxed))
+					{
+						uint8 Expected = 0u;
+						if (Pool[i].Usage.CompareExchange(Expected, 1u))
+						{
+							return Pool[i].Request;
+						}
+					}
+				}
+				FPlatformProcess::Sleep(UE_ZENDDC_BACKEND_WAIT_INTERVAL);
+			}
+		}
+
+		/**
+		 * Release request to the pool.
+		 * @param Request Request that should be freed. Note that any buffer owened by the request can now be reset.
+		 */
+		void ReleaseRequestToPool(FRequest* Request)
+		{
+			for (uint8 i = 0; i < Pool.Num(); ++i)
+			{
+				if (Pool[i].Request == Request)
+				{
+					Request->Reset();
+					uint8 Expected = 1u;
+					Pool[i].Usage.CompareExchange(Expected, 0u);
+					return;
+				}
+			}
+			check(false);
+		}
+
+	private:
+
+		struct FEntry
+		{
+			TAtomic<uint8> Usage;
+			FRequest* Request;
+		};
+
+		TStaticArray<FEntry, UE_ZENDDC_REQUEST_POOL_SIZE> Pool;
+
+		FRequestPool() {}
+	};
+
+	//----------------------------------------------------------------------------------------------------------
+	// FScopedRequestPtr
+	//----------------------------------------------------------------------------------------------------------
+
+	/**
+	 * Utility class to manage requesting and releasing requests from the \ref FRequestPool.
+	 */
+	struct FScopedRequestPtr
+	{
+	public:
+		FScopedRequestPtr(FRequestPool* InPool)
+			: Request(InPool->WaitForFreeRequest())
+			, Pool(InPool)
+		{}
+
+		~FScopedRequestPtr()
+		{
+			Pool->ReleaseRequestToPool(Request);
+		}
+
+		bool IsValid() const
+		{
+			return Request != nullptr;
+		}
+
+		FRequest* Get() const
+		{
+			check(IsValid());
+			return Request;
+		}
+
+		FRequest* operator->()
+		{
+			check(IsValid());
+			return Request;
+		}
+
+	private:
+		FRequest* Request;
+		FRequestPool* Pool;
+	};
+}
+
+namespace UE::DerivedData {
+
+//----------------------------------------------------------------------------------------------------------
+// FZenDerivedDataBackend
+//----------------------------------------------------------------------------------------------------------
+
+FZenDerivedDataBackend::FZenDerivedDataBackend(
+	ICacheFactory& InFactory, 
+	const TCHAR* InServiceUrl,
+	const TCHAR* InNamespace)
+	: Factory(InFactory)
+	, Domain(InServiceUrl)
+	, Namespace(InNamespace)
+{
+	if (IsServiceReady())
+	{
+		RequestPool = MakeUnique<zen::FRequestPool>(InServiceUrl);
+		bIsUsable = true;
+	}
+	bCacheRecordEndpointEnabled = false;
+	GConfig->GetBool(TEXT("Zen"), TEXT("CacheRecordEndpointEnabled"), bCacheRecordEndpointEnabled, GEngineIni);
+}
+
+FZenDerivedDataBackend::~FZenDerivedDataBackend()
+{
+}
+
+FString FZenDerivedDataBackend::GetName() const
+{
+	return Domain;
+}
+
+
+bool FZenDerivedDataBackend::IsServiceReady()
+{
+	zen::FRequest Request(*Domain, false);
+	zen::FRequest::Result Result = Request.PerformBlockingDownload(TEXT("health/ready"), nullptr);
+	
+	if (Result == zen::FRequest::Success && zen::IsSuccessCode(Request.GetResponseCode()))
+	{
+		UE_LOG(LogDerivedDataCache, Display, TEXT("Z$ HTTP DDC service status: %s."), *Request.GetResponseAsString());
+		return true;
+	}
+	else
+	{
+		UE_LOG(LogDerivedDataCache, Warning, TEXT("Unable to reach Z$ HTTP DDC service at %s. Status: %d . Response: %s"), *Domain, Request.GetResponseCode(), *Request.GetResponseAsString());
+	}
+
+	return false;
+}
+
+bool FZenDerivedDataBackend::ShouldRetryOnError(int64 ResponseCode)
+{
+	// Access token might have expired, request a new token and try again.
+	if (ResponseCode == 401)
+	{
+		return true;
+	}
+
+	// Too many requests, make a new attempt
+	if (ResponseCode == 429)
+	{
+		return true;
+	}
+
+	return false;
+}
+
+bool FZenDerivedDataBackend::CachedDataProbablyExists(const TCHAR* CacheKey)
+{
+	TRACE_CPUPROFILER_EVENT_SCOPE(ZenDDC_Exist);
+	TRACE_COUNTER_ADD(ZenDDC_Exist, int64(1));
+	COOK_STAT(auto Timer = UsageStats.TimeProbablyExists());
+
+	if (ShouldSimulateMiss(CacheKey))
+	{
+		return false;
+	}
+	
+	FString Uri = MakeLegacyZenKey(CacheKey);
+	long ResponseCode = 0; uint32 Attempts = 0;
+
+	// Retry request until we get an accepted response or exhaust allowed number of attempts.
+	while (ResponseCode == 0 && ++Attempts < UE_ZENDDC_MAX_ATTEMPTS)
+	{
+		zen::FScopedRequestPtr Request(RequestPool.Get());
+		zen::FRequest::Result Result = Request->PerformBlockingQuery<zen::FRequest::Head>(*Uri);
+		ResponseCode = Request->GetResponseCode();
+
+		if (zen::IsSuccessCode(ResponseCode) || ResponseCode == 400)
+		{
+			const bool bIsHit = (Result == zen::FRequest::Success && zen::IsSuccessCode(ResponseCode));
+			if (bIsHit)
+			{
+				TRACE_COUNTER_ADD(ZenDDC_ExistHit, int64(1));
+			}
+			return bIsHit;
+		}
+
+		if (!ShouldRetryOnError(ResponseCode))
+		{
+			return false;
+		}
+
+		ResponseCode = 0;
+	}
+
+	return false;
+}
+
+bool FZenDerivedDataBackend::GetCachedData(const TCHAR* CacheKey, TArray<uint8>& OutData)
+{
+	TRACE_CPUPROFILER_EVENT_SCOPE(ZenDDC_Get);
+	TRACE_COUNTER_ADD(ZenDDC_Get, int64(1));
+	COOK_STAT(auto Timer = UsageStats.TimeGet());
+
+	if (ShouldSimulateMiss(CacheKey))
+	{
+		return false;
+	}
+
+	double StartTime = FPlatformTime::Seconds();
+
+	EGetResult Result = GetZenData(*MakeLegacyZenKey(CacheKey), &OutData);
+	if (Result != EGetResult::Success)
+	{
+		switch (Result)
+		{
+		default:
+			UE_LOG(LogDerivedDataCache, Verbose, TEXT("%s: Cache miss on %s"), *GetName(), CacheKey);
+			break;
+		case EGetResult::Corrupted:
+			UE_LOG(LogDerivedDataCache, Warning,
+				TEXT("Checksum from server on %s did not match recieved data. Discarding cached result."), CacheKey);
+			break;
+		}
+		return false;
+	}
+
+	TRACE_COUNTER_ADD(ZenDDC_GetHit, int64(1));
+	TRACE_COUNTER_ADD(ZenDDC_BytesReceived, int64(OutData.Num()));
+	COOK_STAT(Timer.AddHit(OutData.Num()));
+	double ReadDuration = FPlatformTime::Seconds() - StartTime;
+	double ReadSpeed = (OutData.Num() / ReadDuration) / (1024.0 * 1024.0);
+	UE_LOG(LogDerivedDataCache, Verbose, TEXT("%s: Cache hit on %s (%d bytes, %.02f secs, %.2fMB/s)"),
+		*GetName(), CacheKey, OutData.Num(), ReadDuration, ReadSpeed);
+	return true;
+}
+
+FZenDerivedDataBackend::EGetResult
+FZenDerivedDataBackend::GetZenData(const TCHAR* Uri, TArray<uint8>* OutData) const
+{
+	// Retry request until we get an accepted response or exhaust allowed number of attempts.
+	EGetResult GetResult = EGetResult::NotFound;
+	for (uint32 Attempts = 0; Attempts < UE_ZENDDC_MAX_ATTEMPTS; ++Attempts)
+	{
+		zen::FScopedRequestPtr Request(RequestPool.Get());
+		if (Request.IsValid())
+		{
+			zen::FRequest::Result Result = Request->PerformBlockingDownload(Uri, OutData);
+			int64 ResponseCode = Request->GetResponseCode();
+
+			// Request was successful, make sure we got all the expected data.
+			if (zen::IsSuccessCode(ResponseCode))
+			{
+				return EGetResult::Success;
+			}
+
+			if (!ShouldRetryOnError(ResponseCode))
+			{
+				break;
+			}
+		}
+	}
+
+	if (OutData)
+	{
+		OutData->Reset();
+	}
+	return GetResult;
+}
+
+FZenDerivedDataBackend::EGetResult
+FZenDerivedDataBackend::GetZenData(const FCacheKey& CacheKey, ECachePolicy CachePolicy, FCbPackage& OutPackage) const
+{
+	TStringBuilder<256> QueryUri;
+	AppendZenUri(CacheKey, QueryUri);
+	AppendPolicyQueryString(CachePolicy, QueryUri);
+
+	// Retry request until we get an accepted response or exhaust allowed number of attempts.
+	EGetResult GetResult = EGetResult::NotFound;
+	for (uint32 Attempts = 0; Attempts < UE_ZENDDC_MAX_ATTEMPTS; ++Attempts)
+	{
+		zen::FScopedRequestPtr Request(RequestPool.Get());
+		if (Request.IsValid())
+		{
+			zen::FRequest::Result Result = Request->PerformBlockingDownload(QueryUri.ToString(), OutPackage);
+			int64 ResponseCode = Request->GetResponseCode();
+			bool bPackageValid = Request->GetResponseFormatValid();
+
+			// Request was successful, make sure we got all the expected data.
+			if (zen::IsSuccessCode(ResponseCode))
+			{
+				if (bPackageValid)
+				{
+					GetResult = EGetResult::Success;
+				}
+				else
+				{
+					GetResult = EGetResult::Corrupted;
+				}
+				break;
+			}
+
+			if (!ShouldRetryOnError(ResponseCode))
+			{
+				break;
+			}
+		}
+	}
+
+	if (GetResult != EGetResult::Success)
+	{
+		OutPackage.Reset();
+	}
+	return GetResult;
+}
+
+FDerivedDataBackendInterface::EPutStatus
+FZenDerivedDataBackend::PutCachedData(const TCHAR* CacheKey, TArrayView<const uint8> InData, bool bPutEvenIfExists)
+{
+	if (DidSimulateMiss(CacheKey))
+	{
+		return EPutStatus::NotCached;
+	}
+
+	FSharedBuffer DataBuffer = FSharedBuffer::MakeView(InData.GetData(), InData.Num());
+	return PutZenData(*MakeLegacyZenKey(CacheKey), FCompositeBuffer(DataBuffer), zen::EContentType::Binary);
+}
+
+FDerivedDataBackendInterface::EPutStatus
+FZenDerivedDataBackend::PutZenData(const TCHAR* Uri, const FCompositeBuffer& InData, zen::EContentType ContentType)
+{
+	TRACE_CPUPROFILER_EVENT_SCOPE(ZenDDC_Put);
+	COOK_STAT(auto Timer = UsageStats.TimePut());
+
+	int64 ResponseCode = 0; uint32 Attempts = 0;
+
+	// Retry request until we get an accepted response or exhaust allowed number of attempts.
+	while (ResponseCode == 0 && ++Attempts < UE_ZENDDC_MAX_ATTEMPTS)
+	{
+		zen::FScopedRequestPtr Request(RequestPool.Get());
+		if (Request.IsValid())
+		{
+			zen::FRequest::Result Result = Request->PerformBlockingUpload<zen::FRequest::Put>(Uri, InData, ContentType);
+			ResponseCode = Request->GetResponseCode();
+
+			if (zen::IsSuccessCode(ResponseCode))
+			{
+				TRACE_COUNTER_ADD(ZenDDC_BytesSent, int64(Request->GetBytesSent()));
+				COOK_STAT(Timer.AddHit(Request->GetBytesSent()));
+				return EPutStatus::Cached;
+			}
+
+			if (!ShouldRetryOnError(ResponseCode))
+			{
+				return EPutStatus::NotCached;
+			}
+
+			ResponseCode = 0;
+		}
+	}
+
+	return EPutStatus::NotCached;
+}
+
+FString FZenDerivedDataBackend::MakeLegacyZenKey(const TCHAR* CacheKey)
+{
+	FIoHash KeyHash = FIoHash::HashBuffer(CacheKey, FCString::Strlen(CacheKey) * sizeof(TCHAR));
+	return FString::Printf(TEXT("/z$/legacy/%s"), *LexToString(KeyHash));
+}
+
+void FZenDerivedDataBackend::AppendZenUri(const FCacheKey& CacheKey, FStringBuilderBase& Out)
+{
+	FStringView BucketStr = CacheKey.Bucket.ToString<TCHAR>();
+	Out.Appendf(TEXT("/z$/%.*s/"), BucketStr.Len(), BucketStr.GetData());
+	Out << CacheKey.Hash;
+}
+
+void FZenDerivedDataBackend::AppendZenUri(const FCacheKey& CacheKey, const FPayloadId& PayloadId, FStringBuilderBase& Out)
+{
+	AppendZenUri(CacheKey, Out);
+	Out << TEXT("/");
+	UE::String::BytesToHexLower(PayloadId.GetBytes(), Out);
+}
+
+void FZenDerivedDataBackend::AppendPolicyQueryString(ECachePolicy Policy, FStringBuilderBase& Uri)
+{
+	bool bQueryEmpty = true;
+	bool bValueEmpty = true;
+	auto AppendKey = [&Uri, &bQueryEmpty, &bValueEmpty](const TCHAR* Key)
+	{
+		if (bQueryEmpty)
+		{
+			TCHAR LastChar = Uri.Len() == 0 ? '\0' : Uri.LastChar();
+			if (LastChar != '?' && LastChar != '&')
+			{
+				Uri << '?';
+			}
+			bQueryEmpty = false;
+		}
+		else
+		{
+			Uri << '&';
+		}
+		bValueEmpty = true;
+		Uri << Key;
+	};
+	auto AppendValue = [&Uri, &bValueEmpty](const TCHAR* Value)
+	{
+		if (bValueEmpty)
+		{
+			bValueEmpty = false;
+		}
+		else
+		{
+			Uri << ',';
+		}
+		Uri << Value;
+	};
+
+	if (!EnumHasAllFlags(Policy, ECachePolicy::Query))
+	{
+		AppendKey(TEXT("query="));
+		if (EnumHasAnyFlags(Policy, ECachePolicy::QueryLocal)) { AppendValue(TEXT("local")); }
+		if (EnumHasAnyFlags(Policy, ECachePolicy::QueryRemote)) { AppendValue(TEXT("remote")); }
+		if (!EnumHasAnyFlags(Policy, ECachePolicy::Query)) { AppendValue(TEXT("none")); }
+	}
+	if (!EnumHasAllFlags(Policy, ECachePolicy::Store))
+	{
+		AppendKey(TEXT("store="));
+		if (EnumHasAnyFlags(Policy, ECachePolicy::StoreLocal)) { AppendValue(TEXT("local")); }
+		if (EnumHasAnyFlags(Policy, ECachePolicy::StoreRemote)) { AppendValue(TEXT("remote")); }
+		if (!EnumHasAnyFlags(Policy, ECachePolicy::Store)) { AppendValue(TEXT("none")); }
+	}
+	if (EnumHasAnyFlags(Policy, ECachePolicy::SkipData))
+	{
+		AppendKey(TEXT("skip="));
+		if (EnumHasAllFlags(Policy, ECachePolicy::SkipData)) { AppendValue(TEXT("data")); }
+		else
+		{
+			if (EnumHasAnyFlags(Policy, ECachePolicy::SkipMeta)) { AppendValue(TEXT("meta")); }
+			if (EnumHasAnyFlags(Policy, ECachePolicy::SkipValue)) { AppendValue(TEXT("value")); }
+			if (EnumHasAnyFlags(Policy, ECachePolicy::SkipAttachments)) { AppendValue(TEXT("attachments")); }
+		}
+	}
+}
+
+uint64 FZenDerivedDataBackend::MeasureCacheRecord(const FCacheRecord& Record)
+{
+	return Record.GetMeta().GetSize() +
+		Record.GetValuePayload().GetRawSize() +
+		Algo::TransformAccumulate(Record.GetAttachmentPayloads(), &FPayload::GetRawSize, uint64(0));
+}
+
+void FZenDerivedDataBackend::RemoveCachedData(const TCHAR* CacheKey, bool bTransient)
+{
+	TRACE_CPUPROFILER_EVENT_SCOPE(ZenDDC_Remove);
+	FString Uri = MakeLegacyZenKey(CacheKey);
+
+	int64 ResponseCode = 0; uint32 Attempts = 0;
+
+	// Retry request until we get an accepted response or exhaust allowed number of attempts.
+	while (ResponseCode == 0 && ++Attempts < UE_ZENDDC_MAX_ATTEMPTS)
+	{
+		zen::FScopedRequestPtr Request(RequestPool.Get());
+		if (Request.IsValid())
+		{
+			zen::FRequest::Result Result = Request->PerformBlockingQuery<zen::FRequest::Delete>(*Uri);
+			ResponseCode = Request->GetResponseCode();
+
+			if (zen::IsSuccessCode(ResponseCode))
+			{
+				return;
+			}
+
+			if (!ShouldRetryOnError(ResponseCode))
+			{
+				return;
+			}
+
+			ResponseCode = 0;
+		}
+	}
+}
+
+bool FZenDerivedDataBackend::IsWritable() const 
+{
+	return true;
+}
+
+FDerivedDataBackendInterface::ESpeedClass 
+FZenDerivedDataBackend::GetSpeedClass() const
+{
+	return ESpeedClass::Fast;
+}
+
+TSharedRef<FDerivedDataCacheStatsNode> FZenDerivedDataBackend::GatherUsageStats() const
+{
+	return MakeShared<FDerivedDataCacheStatsNode>(this, "foo");
+}
+
+bool FZenDerivedDataBackend::TryToPrefetch(TConstArrayView<FString> CacheKeys)
+{
+	return GetSpeedClass() >= ESpeedClass::Fast && CachedDataProbablyExistsBatch(CacheKeys).CountSetBits() == CacheKeys.Num();
+}
+
+bool FZenDerivedDataBackend::WouldCache(const TCHAR* CacheKey, TArrayView<const uint8> InData)
+{
+	return true;
+}
+
+bool FZenDerivedDataBackend::ApplyDebugOptions(FBackendDebugOptions& InOptions)
+{
+	DebugOptions = InOptions;
+	return true;
+}
+
+bool FZenDerivedDataBackend::DidSimulateMiss(const TCHAR* InKey)
+{
+	if (DebugOptions.RandomMissRate == 0 || DebugOptions.SimulateMissTypes.Num() == 0)
+	{
+		return false;
+	}
+	FScopeLock Lock(&MissedKeysCS);
+	return DebugMissedKeys.Contains(FName(InKey));
+}
+
+bool FZenDerivedDataBackend::ShouldSimulateMiss(const TCHAR* InKey)
+{
+	// once missed, always missed
+	if (DidSimulateMiss(InKey))
+	{
+		return true;
+	}
+
+	if (DebugOptions.ShouldSimulateMiss(InKey))
+	{
+		FScopeLock Lock(&MissedKeysCS);
+		UE_LOG(LogDerivedDataCache, Verbose, TEXT("Simulating miss in %s for %s"), *GetName(), InKey);
+		DebugMissedKeys.Add(FName(InKey));
+		return true;
+	}
+
+	return false;
+}
+
+FRequest FZenDerivedDataBackend::Put(
+	TConstArrayView<UE::DerivedData::FCacheRecord> Records,
+	FStringView Context,
+	ECachePolicy Policy,
+	EPriority Priority,
+	FOnCachePutComplete&& OnComplete)
+{
+	for (const FCacheRecord& Record : Records)
+	{
+		COOK_STAT(auto Timer = UsageStats.TimePut());
+		bool bResult;
+		if (bCacheRecordEndpointEnabled)
+		{
+			bResult = PutCacheRecord(Record, Context, Policy);
+		}
+		else
+		{
+			bResult = LegacyPutCacheRecord(Record, Context, Policy);
+		}
+		if (bResult)
+		{
+			UE_LOG(LogDerivedDataCache, Verbose, TEXT("%s: Cache put complete for %s from '%.*s'"),
+				*GetName(), *WriteToString<96>(Record.GetKey()), Context.Len(), Context.GetData());
+			COOK_STAT(Timer.AddHit(MeasureCacheRecord(Record)));
+			if (OnComplete)
+			{
+				OnComplete({ Record.GetKey(), EStatus::Ok });
+			}
+		}
+		else
+		{
+			COOK_STAT(Timer.AddMiss(MeasureCacheRecord(Record)));
+			if (OnComplete)
+			{
+				OnComplete({ Record.GetKey(), EStatus::Error });
+			}
+		}
+	}
+	return FRequest();
+}
+
+FRequest FZenDerivedDataBackend::Get(
+	TConstArrayView<UE::DerivedData::FCacheKey> Keys,
+	FStringView Context,
+	UE::DerivedData::ECachePolicy Policy,
+	UE::DerivedData::EPriority Priority,
+	UE::DerivedData::FOnCacheGetComplete&& OnComplete)
+{
+	for (const UE::DerivedData::FCacheKey& Key : Keys)
+	{
+		TRACE_CPUPROFILER_EVENT_SCOPE(ZenDDC_Get);
+		TRACE_COUNTER_ADD(ZenDDC_Get, int64(1));
+		COOK_STAT(auto Timer = UsageStats.TimeGet());
+
+		FOptionalCacheRecord Record;
+		if (bCacheRecordEndpointEnabled)
+		{
+			Record = GetCacheRecord(Key, Context, Policy);
+		}
+		else
+		{
+			Record = LegacyGetCacheRecord(Key, Context, Policy);
+		}
+		if (Record)
+		{
+			UE_LOG(LogDerivedDataCache, Verbose, TEXT("%s: Cache hit for %s from '%.*s'"),
+				*GetName(), *WriteToString<96>(Key), Context.Len(), Context.GetData());
+			TRACE_COUNTER_ADD(ZenDDC_GetHit, int64(1));
+			int64 ReceivedSize = MeasureCacheRecord(Record.Get());
+			TRACE_COUNTER_ADD(ZenDDC_BytesReceived, ReceivedSize);
+			COOK_STAT(Timer.AddHit(ReceivedSize));
+
+			if (OnComplete)
+			{
+				OnComplete({ MoveTemp(Record).Get(), EStatus::Ok });
+			}
+		}
+		else
+		{
+			if (OnComplete)
+			{
+				OnComplete({ Factory.CreateRecord(Key).Build(), EStatus::Error });
+			}
+		}
+	}
+	return FRequest();
+}
+
+FRequest FZenDerivedDataBackend::GetPayload(
+	TConstArrayView<UE::DerivedData::FCachePayloadKey> Keys,
+	FStringView Context,
+	UE::DerivedData::ECachePolicy Policy,
+	UE::DerivedData::EPriority Priority,
+	UE::DerivedData::FOnCacheGetPayloadComplete&& OnComplete)
+{
+	if (bCacheRecordEndpointEnabled)
+	{
+		for (const FCachePayloadKey& Key : Keys)
+		{
+			TRACE_CPUPROFILER_EVENT_SCOPE(ZenDDC_Get);
+			TRACE_COUNTER_ADD(ZenDDC_Get, int64(1));
+			COOK_STAT(auto Timer = UsageStats.TimeGet());
+
+			TStringBuilder<256> QueryUri;
+			AppendZenUri(Key.CacheKey, Key.Id, QueryUri);
+			AppendPolicyQueryString(Policy, QueryUri);
+
+			TArray<uint8> PayloadData;
+			EGetResult GetResult = GetZenData(QueryUri.ToString(), &PayloadData);
+			FCompressedBuffer CompressedBuffer;
+			if (GetResult == EGetResult::Success)
+			{
+				CompressedBuffer = FCompressedBuffer::FromCompressed(MakeSharedBufferFromArray(MoveTemp(PayloadData)));
+			}
+			if (CompressedBuffer)
+			{
+				FPayload Payload = FPayload(Key.Id, MoveTemp(CompressedBuffer));
+				UE_LOG(LogDerivedDataCache, Verbose, TEXT("%s: Cache hit for %s from '%.*s'"),
+					*GetName(), *WriteToString<96>(Key), Context.Len(), Context.GetData());
+				TRACE_COUNTER_ADD(ZenDDC_GetHit, int64(1));
+				TRACE_COUNTER_ADD(ZenDDC_BytesReceived, Payload.GetRawSize());
+				COOK_STAT(Timer.AddHit(Payload.GetRawSize()));
+				if (OnComplete)
+				{
+					OnComplete({ Key.CacheKey, MoveTemp(Payload), EStatus::Ok });
+				}
+			}
+			else
+			{
+				UE_LOG(LogDerivedDataCache, Display, TEXT("%s: Cache miss with %s payload %s for %s from '%.*s'"),
+					*GetName(), (GetResult == EGetResult::NotFound ? TEXT("missing") : TEXT("corrupted")),
+					*WriteToString<16>(Key.Id), *WriteToString<96>(Key), Context.Len(), Context.GetData());
+				if (OnComplete)
+				{
+					OnComplete({ Key.CacheKey, FPayload(Key.Id), EStatus::Error });
+				}
+			}
+		}
+	}
+	else
+	{
+		// We have to load the CacheRecord for each Payload, so sort all the Payloads sharing
+		// the same key to occur together and we can load the CacheRecord at the start of each run of Payloads
+		TArray<FCachePayloadKey, TInlineAllocator<16>> SortedKeys(Keys);
+		SortedKeys.StableSort();
+
+		FOptionalCacheRecord Record;
+		for (const FCachePayloadKey& Key : SortedKeys)
+		{
+			TRACE_CPUPROFILER_EVENT_SCOPE(ZenDDC_Get);
+			TRACE_COUNTER_ADD(ZenDDC_Get, int64(1));
+			COOK_STAT(auto Timer = UsageStats.TimeGet());
+
+			if (!Record || Record.Get().GetKey() != Key.CacheKey)
+			{
+				Record = LegacyGetCacheRecord(Key.CacheKey, Context, Policy | ECachePolicy::SkipData, /*bAlwaysLoadInlineData*/ true);
+			}
+			if (FPayload Payload = Record ? LegacyGetCachePayload(Key.CacheKey, Context, Policy, Record.Get().GetPayload(Key.Id)) : FPayload::Null)
+			{
+				UE_LOG(LogDerivedDataCache, Verbose, TEXT("%s: Cache hit for %s from '%.*s'"),
+					*GetName(), *WriteToString<96>(Key), Context.Len(), Context.GetData());
+				TRACE_COUNTER_ADD(ZenDDC_GetHit, int64(1));
+				TRACE_COUNTER_ADD(ZenDDC_BytesReceived, Payload.GetRawSize());
+				COOK_STAT(Timer.AddHit(Payload.GetRawSize()));
+				if (OnComplete)
+				{
+					OnComplete({ Key.CacheKey, MoveTemp(Payload), EStatus::Ok });
+				}
+			}
+			else
+			{
+				if (OnComplete)
+				{
+					OnComplete({ Key.CacheKey, FPayload(Key.Id), EStatus::Error });
+				}
+			}
+		}
+	}
+	return FRequest();
+}
+
+void FZenDerivedDataBackend::CancelAll()
+{
+}
+
+bool FZenDerivedDataBackend::PutCacheRecord(const FCacheRecord& Record, FStringView Context, ECachePolicy Policy)
+{
+	const FCacheKey& Key = Record.GetKey();
+	FCbPackage Package = Factory.SaveRecord(Record);
+	FCbWriter PackageWriter;
+	Package.Save(PackageWriter);
+	FUniqueBuffer Buffer = FUniqueBuffer::Alloc(PackageWriter.GetSaveSize());
+	PackageWriter.Save(Buffer);
+	TStringBuilder<256> Uri;
+	AppendZenUri(Record.GetKey(), Uri);
+	AppendPolicyQueryString(Policy, Uri);
+	if (PutZenData(Uri.ToString(), FCompositeBuffer(Buffer.MoveToShared()), zen::EContentType::CompactBinaryPackage)
+		!= FDerivedDataBackendInterface::EPutStatus::Cached)
+	{
+		return false;
+	}
+
+	return true;
+}
+
+FOptionalCacheRecord FZenDerivedDataBackend::GetCacheRecord(const FCacheKey& Key, FStringView Context,
+	ECachePolicy Policy) const
+{
+	FCbPackage Package;
+	EGetResult GetResult = GetZenData(Key, Policy, Package);
+	if (GetResult != EGetResult::Success)
+	{
+		switch (GetResult)
+		{
+		default:
+			UE_LOG(LogDerivedDataCache, Verbose, TEXT("%s: Cache miss with missing record for %s from '%.*s'"),
+				*GetName(), *WriteToString<96>(Key), Context.Len(), Context.GetData());
+			break;
+		case EGetResult::Corrupted:
+			UE_LOG(LogDerivedDataCache, Verbose, TEXT("%s: Cache miss with corrupted record (corrupted package) for %s from '%.*s'"),
+				*GetName(), *WriteToString<96>(Key), Context.Len(), Context.GetData());
+			break;
+		}
+		return FOptionalCacheRecord();
+	}
+	FOptionalCacheRecord Record = Factory.LoadRecord(Package);
+	if (!Record)
+	{
+		UE_LOG(LogDerivedDataCache, Verbose, TEXT("%s: Cache miss with corrupted record (corrupted record) for %s from '%.*s'"),
+			*GetName(), *WriteToString<96>(Key), Context.Len(), Context.GetData());
+	}
+	return Record;
+}
+
+bool FZenDerivedDataBackend::LegacyPutCacheRecord(const FCacheRecord& Record, FStringView Context, ECachePolicy Policy)
+{
+	const FCacheKey& Key = Record.GetKey();
+	// Save the payloads and build the cache object.
+	const TConstArrayView<FPayload> Attachments = Record.GetAttachmentPayloads();
+	TCbWriter<512> Writer;
+	Writer.BeginObject();
+	if (const FCbObject& Meta = Record.GetMeta())
+	{
+		Writer.AddObject("Meta"_ASV, Meta);
+		Writer.AddHash("MetaHash"_ASV, Meta.GetHash());
+	}
+	if (const FPayload& Value = Record.GetValuePayload())
+	{
+		Writer.SetName("Value"_ASV);
+		if (!LegacyPutCachePayload(Key, Context, Value, Writer))
+		{
+			return false;
+		}
+	}
+	if (!Attachments.IsEmpty())
+	{
+		Writer.BeginArray("Attachments"_ASV);
+		for (const FPayload& Attachment : Attachments)
+		{
+			if (!LegacyPutCachePayload(Key, Context, Attachment, Writer))
+			{
+				return false;
+			}
+		}
+		Writer.EndArray();
+	}
+	Writer.EndObject();
+
+	// Save the record to storage.
+	FUniqueBuffer Buffer = FUniqueBuffer::Alloc(Writer.GetSaveSize());
+	Writer.Save(Buffer);
+	TStringBuilder<256> Uri;
+	LegacyMakeZenKey(Key, Uri);
+	if (PutZenData(Uri.ToString(), FCompositeBuffer(Buffer.MoveToShared()), zen::EContentType::CompactBinary)
+		!= FDerivedDataBackendInterface::EPutStatus::Cached)
+	{
+		return false;
+	}
+
+	return true;
+}
+
+bool FZenDerivedDataBackend::LegacyPutCachePayload(const FCacheKey& Key, FStringView Context, const FPayload& Payload, FCbWriter& Writer)
+{
+	const FIoHash& RawHash = Payload.GetRawHash();
+	const FCompositeBuffer CompressedBuffer = Payload.GetData().GetCompressed();
+
+	const bool bStoreInline = !Payload.HasData();
+	if (!bStoreInline)
+	{
+		TStringBuilder<256> Uri;
+		LegacyMakePayloadKey(Key, RawHash, Uri);
+		if (PutZenData(Uri.ToString(), CompressedBuffer, zen::EContentType::Binary) != FDerivedDataBackendInterface::EPutStatus::Cached)
+		{
+			return false;
+		}
+	}
+
+	Writer.BeginObject();
+	Writer.AddObjectId("Id"_ASV, FCbObjectId(Payload.GetId().GetView()));
+	Writer.AddInteger("RawSize"_ASV, Payload.GetRawSize());
+	if (bStoreInline)
+	{
+		Writer.AddHash("RawHash"_ASV, RawHash);
+		Writer.AddHash("CompressedHash"_ASV, FIoHash::HashBuffer(CompressedBuffer));
+		Writer.AddBinary("CompressedData"_ASV, CompressedBuffer);
+	}
+	else
+	{
+		Writer.AddBinaryAttachment("RawHash"_ASV, RawHash);
+	}
+	Writer.EndObject();
+	return true;
+}
+
+FOptionalCacheRecord FZenDerivedDataBackend::LegacyGetCacheRecord(const FCacheKey& Key,
+	FStringView Context, ECachePolicy Policy, bool bAlwaysLoadInlineData) const
+{
+	TStringBuilder<256> Uri;
+	LegacyMakeZenKey(Key, Uri);
+	TArray<uint8> Buffer;
+	EGetResult Result = GetZenData(Uri.ToString(), &Buffer);
+	if (Result != EGetResult::Success)
+	{
+		return FOptionalCacheRecord();
+	}
+	return FZenDerivedDataBackend::LegacyCreateRecord(MakeSharedBufferFromArray(MoveTemp(Buffer)),
+		Key, Context, Policy, bAlwaysLoadInlineData);
+}
+
+FPayload FZenDerivedDataBackend::LegacyGetCachePayload(const FCacheKey& Key, FStringView Context, ECachePolicy Policy,
+	const FPayload& PayloadIdOnly) const
+{
+	auto GetDescriptor = [&PayloadIdOnly, &Key, &Context]()
+	{
+		return FString::Printf(TEXT("%s with hash %s for %s from '%.*s'"),
+			*WriteToString<16>(PayloadIdOnly.GetId()), *WriteToString<48>(PayloadIdOnly.GetRawHash()),
+			*WriteToString<96>(Key), Context.Len(), Context.GetData());
+	};
+
+	if (PayloadIdOnly.HasData())
+	{
+		return PayloadIdOnly;
+	}
+
+	TStringBuilder<256> Uri;
+	LegacyMakePayloadKey(Key, PayloadIdOnly.GetRawHash(), Uri);
+
+	bool bSkipData = EnumHasAllFlags(Policy, ECachePolicy::SkipData);
+	EGetResult GetResult;
+	TArray<uint8> ArrayBuffer;
+	GetResult = GetZenData(Uri.ToString(), bSkipData ? nullptr : &ArrayBuffer);
+	if (GetResult != EGetResult::Success)
+	{
+		switch (GetResult)
+		{
+		default:
+			UE_LOG(LogDerivedDataCache, Verbose, TEXT("%s: Cache miss with missing payload %s"),
+				*GetName(), *GetDescriptor());
+			break;
+		case EGetResult::Corrupted:
+			UE_LOG(LogDerivedDataCache, Display, TEXT("%s: Cache miss with corrupted payload %s"),
+				*GetName(), *GetDescriptor());
+			break;
+		}
+		return FPayload();
+	}
+	if (bSkipData)
+	{
+		return PayloadIdOnly;
+	}
+
+	FCompressedBuffer CompressedBuffer = FCompressedBuffer::FromCompressed(
+		MakeSharedBufferFromArray(MoveTemp(ArrayBuffer)));
+	if (!CompressedBuffer || CompressedBuffer.GetRawHash() != PayloadIdOnly.GetRawHash())
+	{
+		UE_LOG(LogDerivedDataCache, Display, TEXT("%s: Cache miss with corrupted payload %s"),
+			*GetName(), *GetDescriptor());
+		return FPayload();
+	}
+	return FPayload(PayloadIdOnly.GetId(), MoveTemp(CompressedBuffer));
+}
+
+FOptionalCacheRecord FZenDerivedDataBackend::LegacyCreateRecord(FSharedBuffer&& RecordBytes, const FCacheKey& Key,
+	FStringView Context, ECachePolicy Policy, bool bAlwaysLoadInlineData) const
+{
+	// TODO: Temporary function copied from FFileSystemDerivedDataBackend.
+	// Eventually We will instead send the Policy to zen sever along with the request for the FCacheKey,
+	// and it will end back all payloads in the response rather than requiring separate requests.
+	const FCbObject RecordObject(MoveTemp(RecordBytes));
+	FCacheRecordBuilder RecordBuilder = Factory.CreateRecord(Key);
+
+	if (!EnumHasAnyFlags(Policy, ECachePolicy::SkipMeta))
+	{
+		if (FCbFieldView MetaHash = RecordObject.FindView("MetaHash"_ASV))
+		{
+			if (FCbObject MetaObject = RecordObject["Meta"_ASV].AsObject(); MetaObject.GetHash() == MetaHash.AsHash())
+			{
+				RecordBuilder.SetMeta(MoveTemp(MetaObject));
+			}
+			else
+			{
+				UE_LOG(LogDerivedDataCache, Display, TEXT("%s: Cache miss with corrupted metadata for %s from '%.*s'"),
+					*GetName(), *WriteToString<96>(Key), Context.Len(), Context.GetData());
+				return FOptionalCacheRecord();
+			}
+		}
+	}
+
+	if (FCbObject ValueObject = RecordObject["Value"_ASV].AsObject())
+	{
+		const ECachePolicy ValuePolicy = Policy | (ECachePolicy::SkipData & ~ECachePolicy::SkipValue);
+		FPayload Payload = LegacyGetCachePayload(Key, Context, ValuePolicy, ValueObject, bAlwaysLoadInlineData);
+		if (Payload.IsNull())
+		{
+			return FOptionalCacheRecord();
+		}
+		RecordBuilder.SetValue(MoveTemp(Payload));
+	}
+
+	for (FCbField AttachmentField : RecordObject["Attachments"_ASV])
+	{
+		const ECachePolicy AttachmentsPolicy = Policy | (ECachePolicy::SkipData & ~ECachePolicy::SkipAttachments);
+		FPayload Payload = LegacyGetCachePayload(Key, Context, AttachmentsPolicy, AttachmentField.AsObject(),
+			bAlwaysLoadInlineData);
+		if (Payload.IsNull())
+		{
+			return FOptionalCacheRecord();
+		}
+		RecordBuilder.AddAttachment(MoveTemp(Payload));
+	}
+
+	return RecordBuilder.Build();
+}
+
+FPayload FZenDerivedDataBackend::LegacyGetCachePayload(const FCacheKey& Key, FStringView Context, ECachePolicy Policy,
+	const FCbObject& Object, bool bAlwaysLoadInlineData) const
+{
+	// TODO: Temporary function copied from FFileSystemDerivedDataBackend, helper for CreateRecord
+	const FPayloadId Id(Object.FindView("Id"_ASV).AsObjectId().GetView());
+	const uint64 RawSize = Object.FindView("RawSize"_ASV).AsUInt64(MAX_uint64);
+	const FIoHash RawHash = Object.FindView("RawHash"_ASV).AsHash();
+	FIoHash CompressedHash = Object.FindView("CompressedHash"_ASV).AsHash();
+	FSharedBuffer CompressedData = Object["CompressedData"_ASV].AsBinary();
+
+	if (Id.IsNull() || RawSize == MAX_uint64 || RawHash.IsZero() || !(CompressedHash.IsZero() == CompressedData.IsNull()))
+	{
+		UE_LOG(LogDerivedDataCache, Display, TEXT("%s: Cache miss with invalid record format for %s from '%.*s'"),
+			*GetName(), *WriteToString<96>(Key), Context.Len(), Context.GetData());
+		return FPayload();
+	}
+
+	FPayload Payload(Id, RawHash, RawSize);
+
+	if (CompressedData)
+	{
+		if (EnumHasAllFlags(Policy, ECachePolicy::SkipData) && !bAlwaysLoadInlineData)
+		{
+			return Payload;
+		}
+		else
+		{
+			return LegacyValidateCachePayload(Key, Context, Payload, CompressedHash, MoveTemp(CompressedData));
+		}
+	}
+
+	return LegacyGetCachePayload(Key, Context, Policy, Payload);
+}
+
+FPayload FZenDerivedDataBackend::LegacyValidateCachePayload(const FCacheKey& Key, FStringView Context,
+	const FPayload& Payload, const FIoHash& CompressedHash, FSharedBuffer&& CompressedData) const
+{
+	// TODO: Temporary function copied from FFileSystemDerivedDataBackend, helper for CreateRecord
+	if (FCompressedBuffer CompressedBuffer = FCompressedBuffer::FromCompressed(MoveTemp(CompressedData));
+		CompressedBuffer &&
+		CompressedBuffer.GetRawHash() == Payload.GetRawHash() &&
+		FIoHash::HashBuffer(CompressedBuffer.GetCompressed()) == CompressedHash)
+	{
+		return FPayload(Payload.GetId(), MoveTemp(CompressedBuffer));
+	}
+	UE_LOG(LogDerivedDataCache, Display,
+		TEXT("%s: Cache miss with corrupted payload %s with hash %s for %s from '%.*s'"),
+		*GetName(), *WriteToString<16>(Payload.GetId()), *WriteToString<48>(Payload.GetRawHash()),
+		*WriteToString<96>(Key), Context.Len(), Context.GetData());
+	return FPayload();
+}
+
+void FZenDerivedDataBackend::LegacyMakeZenKey(const FCacheKey& CacheKey, FStringBuilderBase& Out) const
+{
+	Out.Reset();
+	FStringView Bucket = CacheKey.Bucket.ToString<TCHAR>();
+	Out.Appendf(TEXT("/z$/%.*s/%s"), Bucket.Len(), Bucket.GetData(), *LexToString(CacheKey.Hash));
+}
+
+void FZenDerivedDataBackend::LegacyMakePayloadKey(const FCacheKey& CacheKey, const FIoHash& RawHash, FStringBuilderBase& Out) const
+{
+	Out.Reset();
+	FStringView Bucket = CacheKey.Bucket.ToString<TCHAR>();
+	Out.Appendf(TEXT("/z$/%.*s/%s/%s"), Bucket.Len(), Bucket.GetData(), *LexToString(CacheKey.Hash), *LexToString(RawHash));
+}
+
+
+} // namespace UE::DerivedData
+
+#endif //WITH_HTTP_DDC_BACKEND
diff --git a/Engine/Source/Developer/DerivedDataCache/Private/ZenDerivedDataBackend.h b/Engine/Source/Developer/DerivedDataCache/Private/ZenDerivedDataBackend.h
new file mode 100644
--- /dev/null
+++ b/Engine/Source/Developer/DerivedDataCache/Private/ZenDerivedDataBackend.h
@@ -0,0 +1,177 @@
+// Copyright Epic Games, Inc. All Rights Reserved.
+
+#pragma once
+
+#include "CoreMinimal.h"
+
+// Macro for whether to enable the S3 backend. libcurl is not currently available on Mac.
+#if PLATFORM_WINDOWS
+#	define WITH_ZEN_DDC_BACKEND 1
+#else
+#	define WITH_ZEN_DDC_BACKEND 0
+#endif
+
+#if WITH_ZEN_DDC_BACKEND
+
+#include "Containers/Set.h"
+#include "Containers/StringFwd.h"
+#include "DerivedDataBackendInterface.h"
+#include "DerivedDataCacheUsageStats.h"
+#include "DerivedDataRequest.h"
+#include "HAL/CriticalSection.h"
+
+class FCbObject;
+class FCbPackage;
+class FCbWriter;
+class FCompositeBuffer;
+struct FIoHash;
+
+namespace zen {
+	enum class EContentType;
+	struct FRequestPool;
+}
+
+namespace UE::DerivedData {
+	struct FCacheKey;
+	class FOptionalCacheRecord;
+	class FPayload;
+	struct FPayloadId;
+}
+
+namespace UE::DerivedData {
+
+/**
+ * Backend for a HTTP based caching service (Zen)
+ **/
+class FZenDerivedDataBackend : public FDerivedDataBackendInterface
+{
+public:
+	
+	/**
+	 * Creates the backend, checks health status and attempts to acquire an access token.
+	 *
+	 * @param ServiceUrl	Base url to the service including schema.
+	 * @param Namespace		Namespace to use.
+	 */
+	FZenDerivedDataBackend(ICacheFactory& Factory, const TCHAR* ServiceUrl, const TCHAR* Namespace);
+
+	~FZenDerivedDataBackend();
+
+	/**
+	 * Checks is backend is usable (reachable and accessible).
+	 * @return true if usable
+	 */
+	bool IsUsable() const { return bIsUsable; }
+
+	virtual bool IsWritable() const override;
+	virtual ESpeedClass GetSpeedClass() const override;
+	virtual TSharedRef<FDerivedDataCacheStatsNode> GatherUsageStats() const override;
+
+	/**
+	 * Synchronous attempt to make sure the cached data will be available as optimally as possible.
+	 *
+	 * @param	CacheKeys	Alphanumeric+underscore keys of the cache items
+	 * @return				true if the data will probably be found in a fast backend on a future request.
+	 */
+	virtual bool TryToPrefetch(TConstArrayView<FString> CacheKeys) override;
+
+	virtual bool CachedDataProbablyExists(const TCHAR* CacheKey) override;
+	virtual bool GetCachedData(const TCHAR* CacheKey, TArray<uint8>& OutData) override;
+	virtual EPutStatus PutCachedData(const TCHAR* CacheKey, TArrayView<const uint8> InData, bool bPutEvenIfExists) override;
+
+	virtual void RemoveCachedData(const TCHAR* CacheKey, bool bTransient) override;
+	
+	virtual FString GetName() const override;
+	virtual bool WouldCache(const TCHAR* CacheKey, TArrayView<const uint8> InData) override;
+	virtual bool ApplyDebugOptions(FBackendDebugOptions& InOptions) override;
+
+	// ICacheStore
+
+	virtual FRequest Put(
+		TConstArrayView<FCacheRecord> Records,
+		FStringView Context,
+		ECachePolicy Policy = ECachePolicy::Default,
+		EPriority Priority = EPriority::Normal,
+		FOnCachePutComplete&& OnComplete = FOnCachePutComplete()) override;
+
+	virtual FRequest Get(
+		TConstArrayView<FCacheKey> Keys,
+		FStringView Context,
+		ECachePolicy Policy,
+		EPriority Priority,
+		FOnCacheGetComplete&& OnComplete) override;
+
+	virtual FRequest GetPayload(
+		TConstArrayView<FCachePayloadKey> Keys,
+		FStringView Context,
+		ECachePolicy Policy,
+		EPriority Priority,
+		FOnCacheGetPayloadComplete&& OnComplete) override;
+
+	virtual void CancelAll() override;
+
+private:
+	enum class EGetResult
+	{
+		Success,
+		NotFound,
+		Corrupted
+	};
+	EGetResult GetZenData(const TCHAR* Uri, TArray<uint8>* OutData) const;
+
+	// TODO: need ability to specify content type
+	FDerivedDataBackendInterface::EPutStatus PutZenData(const TCHAR* Uri, const FCompositeBuffer& InData, zen::EContentType ContentType);
+	EGetResult GetZenData(const FCacheKey& Key, ECachePolicy CachePolicy, FCbPackage& OutPackage) const;
+
+	bool PutCacheRecord(const FCacheRecord& Record, FStringView Context, ECachePolicy Policy);
+	FOptionalCacheRecord GetCacheRecord(const FCacheKey& Key, FStringView Context,
+		ECachePolicy Policy) const;
+
+	bool IsServiceReady();
+	static FString MakeLegacyZenKey(const TCHAR* CacheKey);
+	static void AppendZenUri(const FCacheKey& CacheKey, FStringBuilderBase& Out);
+	static void AppendZenUri(const FCacheKey& CacheKey, const FPayloadId& PayloadId, FStringBuilderBase& Out);
+	static void AppendPolicyQueryString(ECachePolicy Policy, FStringBuilderBase& Out);
+
+	static bool ShouldRetryOnError(int64 ResponseCode);
+	static uint64 MeasureCacheRecord(const FCacheRecord& Record);
+
+	// Legacy CacheRecord endpoint
+	bool LegacyPutCacheRecord(const FCacheRecord& Record, FStringView Context, ECachePolicy Policy);
+	bool LegacyPutCachePayload(const FCacheKey& Key, FStringView Context, const FPayload& Payload, FCbWriter& Writer);
+	FOptionalCacheRecord LegacyGetCacheRecord(const FCacheKey& Key, FStringView Context,
+		ECachePolicy Policy, bool bAlwaysLoadInlineData = false) const;
+	void LegacyMakeZenKey(const FCacheKey& CacheKey, FStringBuilderBase& Out) const;
+	void LegacyMakePayloadKey(const FCacheKey& CacheKey, const FIoHash& RawHash, FStringBuilderBase& Out) const;
+	FPayload LegacyGetCachePayload(const FCacheKey& Key, FStringView Context, ECachePolicy Policy, const FPayload& Payload) const;
+	FOptionalCacheRecord LegacyCreateRecord(FSharedBuffer&& RecordBytes, const FCacheKey& Key, FStringView Context,
+		ECachePolicy Policy, bool bAlwaysLoadInlineData) const;
+	FPayload LegacyGetCachePayload(const FCacheKey& Key, FStringView Context, ECachePolicy Policy, const FCbObject& Object,
+		bool bAlwaysLoadInlineData = false) const;
+	FPayload LegacyValidateCachePayload(const FCacheKey& Key, FStringView Context, const FPayload& Payload,
+		const FIoHash& CompressedHash, FSharedBuffer&& CompressedData) const;
+
+	/* Debug helpers */
+	bool DidSimulateMiss(const TCHAR* InKey);
+	bool ShouldSimulateMiss(const TCHAR* InKey);
+private:
+	ICacheFactory& Factory;
+	FString Domain;
+	FString Namespace;
+	mutable FDerivedDataCacheUsageStats UsageStats;
+	TUniquePtr<struct zen::FRequestPool> RequestPool;
+	bool bIsUsable = false;
+	uint32 FailedLoginAttempts = 0;
+	bool bCacheRecordEndpointEnabled = false;
+
+	/** Debug Options */
+	FBackendDebugOptions DebugOptions;
+
+	/** Keys we ignored due to miss rate settings */
+	FCriticalSection MissedKeysCS;
+	TSet<FName> DebugMissedKeys;
+};
+
+} // namespace UE::DerivedData
+
+#endif // WITH_ZEN_DDC_BACKEND
