Integrating //UE5/Dev-Cooker @ 16678003 to //UE5/Main (Zousar.Shaker-YEG-0943-Quaternary)

diff --git a/Engine/Source/Developer/DerivedDataCache/Private/ZenKvDerivedDataBackend.cpp b/Engine/Source/Developer/DerivedDataCache/Private/ZenKvDerivedDataBackend.cpp
new file mode 100644
--- /dev/null
+++ b/Engine/Source/Developer/DerivedDataCache/Private/ZenKvDerivedDataBackend.cpp
@@ -0,0 +1,1123 @@
+// Copyright Epic Games, Inc. All Rights Reserved.
+
+#include "ZenKvDerivedDataBackend.h"
+
+#if WITH_ZEN_DDC_BACKEND
+
+#if PLATFORM_WINDOWS || PLATFORM_HOLOLENS
+#	include "Windows/WindowsHWrapper.h"
+#	include "Windows/AllowWindowsPlatformTypes.h"
+#endif
+
+#include "curl/curl.h"
+
+#if PLATFORM_WINDOWS || PLATFORM_HOLOLENS
+#	include "Windows/HideWindowsPlatformTypes.h"
+#endif
+
+#include "Containers/StaticArray.h"
+#include "Containers/Ticker.h"
+#include "DerivedDataCacheRecord.h"
+#include "Dom/JsonObject.h"
+#include "GenericPlatform/GenericPlatformFile.h"
+#include "HAL/PlatformFileManager.h"
+#include "Misc/FileHelper.h"
+#include "Misc/ScopeLock.h"
+#include "Misc/SecureHash.h"
+#include "ProfilingDebugging/CpuProfilerTrace.h"
+#include "ProfilingDebugging/CountersTrace.h"
+#include "Serialization/JsonReader.h"
+#include "Serialization/JsonSerializer.h"
+
+#define UE_ZENHTTPDDC_BACKEND_WAIT_INTERVAL			0.01f
+#define UE_ZENHTTPDDC_HTTP_REQUEST_TIMEOUT_SECONDS	30L
+#define UE_ZENHTTPDDC_HTTP_REQUEST_TIMOUT_ENABLED	1
+#define UE_ZENHTTPDDC_HTTP_DEBUG					0
+#define UE_ZENHTTPDDC_REQUEST_POOL_SIZE				16
+#define UE_ZENHTTPDDC_MAX_FAILED_LOGIN_ATTEMPTS		16
+#define UE_ZENHTTPDDC_MAX_ATTEMPTS					4
+#define UE_ZENHTTPDDC_MAX_BUFFER_RESERVE			104857600u
+
+TRACE_DECLARE_INT_COUNTER(ZenHTTPDDC_Exist,			TEXT("ZenHTTPDDC Exist"));
+TRACE_DECLARE_INT_COUNTER(ZenHTTPDDC_ExistHit,		TEXT("ZenHTTPDDC Exist Hit"));
+TRACE_DECLARE_INT_COUNTER(ZenHTTPDDC_Get,			TEXT("ZenHTTPDDC Get"));
+TRACE_DECLARE_INT_COUNTER(ZenHTTPDDC_GetHit,		TEXT("ZenHTTPDDC Get Hit"));
+TRACE_DECLARE_INT_COUNTER(ZenHTTPDDC_Put,			TEXT("ZenHTTPDDC Put"));
+TRACE_DECLARE_INT_COUNTER(ZenHTTPDDC_PutHit,		TEXT("ZenHTTPDDC Put Hit"));
+TRACE_DECLARE_INT_COUNTER(ZenHTTPDDC_BytesReceived, TEXT("ZenHTTPDDC Bytes Received"));
+TRACE_DECLARE_INT_COUNTER(ZenHTTPDDC_BytesSent,		TEXT("ZenHTTPDDC Bytes Sent"));
+
+namespace zenhttp {
+
+/**
+ * Encapsulation for access token shared by all requests.
+ */
+
+	/**
+	 * Minimal HTTP request type wrapping CURL without the need for managers. This request
+	 * is written to allow reuse of request objects, in order to allow connections to be reused.
+	 *
+	 * CURL has a global library initialization (curl_global_init). We rely on this happening in
+	 * the Online/HTTP library which is a dependency on this module.
+	 */
+	class FRequest
+	{
+	public:
+		/**
+		 * Supported request verbs
+		 */
+		enum RequestVerb
+		{
+			Get,
+			Put,
+			Post,
+			PostJson,
+			Delete,
+			Head
+		};
+
+		/**
+		 * Convenience result type interpreted from HTTP response code.
+		 */
+		enum Result
+		{
+			Success,
+			Failed,
+			FailedTimeout
+		};
+
+		FRequest(const TCHAR* InDomain, bool bInLogErrors)
+			: bLogErrors(bInLogErrors)
+			, Domain(InDomain)
+		{
+			Curl = curl_easy_init();
+			Reset();
+		}
+
+		~FRequest()
+		{
+			curl_easy_cleanup(Curl);
+		}
+
+		/**
+		 * Resets all options on the request except those that should always be set.
+		 */
+		void Reset()
+		{
+			Headers.Reset();
+			ResponseHeader.Reset();
+			ResponseBuffer.Reset();
+			ResponseCode = 0;
+			ReadDataView = TArrayView<const uint8>();
+			WriteDataBufferPtr = nullptr;
+			WriteHeaderBufferPtr = nullptr;
+			BytesSent = 0;
+			BytesReceived = 0;
+			CurlResult = CURL_LAST;
+
+			curl_easy_reset(Curl);
+
+			// Options that are always set for all connections.
+#if UE_ZENHTTPDDC_HTTP_REQUEST_TIMOUT_ENABLED
+			curl_easy_setopt(Curl, CURLOPT_CONNECTTIMEOUT, UE_ZENHTTPDDC_HTTP_REQUEST_TIMEOUT_SECONDS);
+#endif
+			curl_easy_setopt(Curl, CURLOPT_FOLLOWLOCATION, 1L);
+			curl_easy_setopt(Curl, CURLOPT_NOSIGNAL, 1L);
+			// Response functions
+			curl_easy_setopt(Curl, CURLOPT_HEADERDATA, this);
+			curl_easy_setopt(Curl, CURLOPT_HEADERFUNCTION, &FRequest::StaticWriteHeaderFn);
+			curl_easy_setopt(Curl, CURLOPT_WRITEDATA, this);
+			curl_easy_setopt(Curl, CURLOPT_WRITEFUNCTION, StaticWriteBodyFn);
+			// Allow compressed data
+			curl_easy_setopt(Curl, CURLOPT_ACCEPT_ENCODING, "gzip");
+			// Rewind method, handle special error case where request need to rewind data stream
+			curl_easy_setopt(Curl, CURLOPT_SEEKFUNCTION, StaticSeekFn);
+			curl_easy_setopt(Curl, CURLOPT_SEEKDATA, this);
+			// Debug hooks
+#if UE_ZENHTTPDDC_HTTP_DEBUG
+			curl_easy_setopt(Curl, CURLOPT_DEBUGDATA, this);
+			curl_easy_setopt(Curl, CURLOPT_DEBUGFUNCTION, StaticDebugCallback);
+			curl_easy_setopt(Curl, CURLOPT_VERBOSE, 1L);
+#endif
+		}
+
+		/** Gets the domain name for this request */
+		const FString& GetDomain() const
+		{
+			return Domain;
+		}
+
+		/** Returns the HTTP response code.*/
+		const int64 GetResponseCode() const
+		{
+			return ResponseCode;
+		}
+
+		/** Returns the number of bytes received this request (headers withstanding). */
+		const size_t GetBytesReceived() const
+		{
+			return BytesReceived;
+		}
+
+		/** Returns the number of bytes sent during this request (headers withstanding). */
+		const size_t  GetBytesSent() const
+		{
+			return BytesSent;
+		}
+
+		/**
+		 * Upload buffer using the request, using either "Put" or "Post" verbs.
+		 * @param Uri Url to use.
+		 * @param Buffer Data to upload
+		 * @return Result of the request
+		 */
+		template<RequestVerb V>
+		Result PerformBlockingUpload(const TCHAR* Uri, TArrayView<const uint8> Buffer)
+		{
+			static_assert(V == Put || V == Post || V == PostJson, "Upload should use either Put or Post verbs.");
+
+			uint32 ContentLength = 0u;
+
+			if (V == Put)
+			{
+				curl_easy_setopt(Curl, CURLOPT_UPLOAD, 1L);
+				curl_easy_setopt(Curl, CURLOPT_INFILESIZE, Buffer.Num());
+				curl_easy_setopt(Curl, CURLOPT_READDATA, this);
+				curl_easy_setopt(Curl, CURLOPT_READFUNCTION, StaticReadFn);
+				Headers.Add(FString(TEXT("Content-Type: application/octet-stream")));
+				ContentLength = Buffer.Num();
+				ReadDataView = Buffer;
+			}
+			else if (V == Post || V == PostJson)
+			{
+				curl_easy_setopt(Curl, CURLOPT_POST, 1L);
+				curl_easy_setopt(Curl, CURLOPT_INFILESIZE, Buffer.Num());
+				curl_easy_setopt(Curl, CURLOPT_READDATA, this);
+				curl_easy_setopt(Curl, CURLOPT_READFUNCTION, StaticReadFn);
+				Headers.Add(V == Post ? FString(TEXT("Content-Type: application/x-www-form-urlencoded")) : FString(TEXT("Content-Type: application/json")));
+				ContentLength = Buffer.Num();
+				ReadDataView = Buffer;
+			}
+
+			return PerformBlocking(Uri, V, ContentLength);
+		}
+
+		/**
+		 * Download an url into a buffer using the request.
+		 * @param Uri Url to use.
+		 * @param Buffer Optional buffer where data should be downloaded to. If empty downloaded data will
+		 * be stored in an internal buffer and accessed GetResponse* methods.
+		 * @return Result of the request
+		 */
+		Result PerformBlockingDownload(const TCHAR* Uri, TArray<uint8>* Buffer)
+		{
+			curl_easy_setopt(Curl, CURLOPT_HTTPGET, 1L);
+			WriteDataBufferPtr = Buffer;
+
+			return PerformBlocking(Uri, Get, 0u);
+		}
+
+		/**
+		 * Query an url using the request. Queries can use either "Head" or "Delete" verbs.
+		 * @param Uri Url to use.
+		 * @return Result of the request
+		 */
+		template<RequestVerb V>
+		Result PerformBlockingQuery(const TCHAR* Uri)
+		{
+			static_assert(V == Head || V == Delete, "Queries should use either Head or Delete verbs.");
+
+			if (V == Delete)
+			{
+				curl_easy_setopt(Curl, CURLOPT_POST, 1L);
+				curl_easy_setopt(Curl, CURLOPT_CUSTOMREQUEST, "DELETE");
+			}
+			else if (V == Head)
+			{
+				curl_easy_setopt(Curl, CURLOPT_NOBODY, 1L);
+			}
+
+			return PerformBlocking(Uri, V, 0u);
+		}
+
+		/**
+		 * Set a header to send with the request.
+		 */
+		void SetHeader(const TCHAR* Header, const TCHAR* Value)
+		{
+			check(CurlResult == CURL_LAST); // Cannot set header after request is sent
+			Headers.Add(FString::Printf(TEXT("%s: %s"), Header, Value));
+		}
+
+		/**
+		 * Attempts to find the header from the response. Returns false if header is not present.
+		 */
+		bool GetHeader(const ANSICHAR* Header, FString& OutValue) const
+		{
+			check(CurlResult != CURL_LAST);  // Cannot query headers before request is sent
+
+			const ANSICHAR* HeadersBuffer = (const ANSICHAR*)ResponseHeader.GetData();
+			size_t HeaderLen = strlen(Header);
+
+			// Find the header key in the (ANSI) response buffer. If not found we can exist immediately
+			if (const ANSICHAR* Found = strstr(HeadersBuffer, Header))
+			{
+				const ANSICHAR* Linebreak = strchr(Found, '\r');
+				const ANSICHAR* ValueStart = Found + HeaderLen + 2; //colon and space
+				const size_t ValueSize = Linebreak - ValueStart;
+				FUTF8ToTCHAR TCHARData(ValueStart, ValueSize);
+				OutValue = FString(TCHARData.Length(), TCHARData.Get());
+				return true;
+			}
+			return false;
+		}
+
+		/**
+		 * Returns the response buffer. Note that is the request is performed
+		 * with an external buffer as target buffer this string will be empty.
+		 */
+		const TArray<uint8>& GetResponseBuffer() const
+		{
+			return ResponseBuffer;
+		}
+
+		/**
+		 * Returns the response buffer as a string. Note that is the request is performed
+		 * with an external buffer as target buffer this string will be empty.
+		 */
+		FString GetResponseAsString() const
+		{
+			return GetAnsiBufferAsString(ResponseBuffer);
+		}
+
+		/**
+		 * Returns the response header as a string.
+		 */
+		FString GetResponseHeaderAsString()
+		{
+			return GetAnsiBufferAsString(ResponseHeader);
+		}
+
+		/**
+		 * Tries to parse the response buffer as a JsonObject. Return empty pointer if
+		 * parse error occurs.
+		 */
+		TSharedPtr<FJsonObject> GetResponseAsJsonObject() const
+		{
+			FString Response = GetAnsiBufferAsString(ResponseBuffer);
+
+			TSharedPtr<FJsonObject> JsonObject;
+			TSharedRef<TJsonReader<> > JsonReader = TJsonReaderFactory<>::Create(Response);
+			if (!FJsonSerializer::Deserialize(JsonReader, JsonObject) || !JsonObject.IsValid())
+			{
+				return TSharedPtr<FJsonObject>(nullptr);
+			}
+
+			return JsonObject;
+		}
+
+	private:
+
+		CURL* Curl;
+		CURLcode				CurlResult;
+		long					ResponseCode;
+		size_t					BytesSent;
+		size_t					BytesReceived;
+		bool					bLogErrors;
+
+		TArrayView<const uint8>	ReadDataView;
+		TArray<uint8>* WriteDataBufferPtr;
+		TArray<uint8>* WriteHeaderBufferPtr;
+
+		TArray<uint8>			ResponseHeader;
+		TArray<uint8>			ResponseBuffer;
+		TArray<FString>			Headers;
+		FString					Domain;
+
+		/**
+		 * Performs the request, blocking until finished.
+		 * @param Uri Address on the domain to query
+		 * @param Verb HTTP verb to use
+		 * @param Buffer Optional buffer to directly receive the result of the request.
+		 * If unset the response body will be stored in the request.
+		 */
+		Result PerformBlocking(const TCHAR* Uri, RequestVerb Verb, uint32 ContentLength)
+		{
+			static const char* CommonHeaders[] = {
+				"User-Agent: UE4",
+				nullptr
+			};
+
+			TRACE_CPUPROFILER_EVENT_SCOPE(ZenHTTPDDC_CurlPerform);
+
+			// Setup request options
+			FString Url = FString::Printf(TEXT("%s%s"), *Domain, Uri);
+			curl_easy_setopt(Curl, CURLOPT_URL, TCHAR_TO_ANSI(*Url));
+
+			// Setup response header buffer. If caller has not setup a response data buffer, use interal.
+			WriteHeaderBufferPtr = &ResponseHeader;
+			if (WriteDataBufferPtr == nullptr)
+			{
+				WriteDataBufferPtr = &ResponseBuffer;
+			}
+
+			// Content-Length should always be set
+			Headers.Add(FString::Printf(TEXT("Content-Length: %d"), ContentLength));
+
+			// Build headers list
+			curl_slist* CurlHeaders = nullptr;
+			// Add common headers
+			for (uint8 i = 0; CommonHeaders[i] != nullptr; ++i)
+			{
+				CurlHeaders = curl_slist_append(CurlHeaders, CommonHeaders[i]);
+			}
+			// Setup added headers
+			for (const FString& Header : Headers)
+			{
+				CurlHeaders = curl_slist_append(CurlHeaders, TCHAR_TO_ANSI(*Header));
+			}
+			curl_easy_setopt(Curl, CURLOPT_HTTPHEADER, CurlHeaders);
+
+			// Shots fired!
+			CurlResult = curl_easy_perform(Curl);
+
+			// Get response code
+			bool bRedirected = false;
+			if (CURLE_OK == curl_easy_getinfo(Curl, CURLINFO_RESPONSE_CODE, &ResponseCode))
+			{
+				bRedirected = (ResponseCode >= 300 && ResponseCode < 400);
+			}
+
+			LogResult(CurlResult, Uri, Verb);
+
+			// Clean up
+			curl_slist_free_all(CurlHeaders);
+
+			return CurlResult == CURLE_OK ? Success : Failed;
+		}
+
+		void LogResult(CURLcode Result, const TCHAR* Uri, RequestVerb Verb) const
+		{
+			if (Result == CURLE_OK)
+			{
+				bool bSuccess = false;
+				const TCHAR* VerbStr = nullptr;
+				FString AdditionalInfo;
+
+				const bool Is400 = (ResponseCode >= 400) && (ResponseCode <= 499);
+				const bool Is200 = (ResponseCode >= 200) && (ResponseCode <= 299);
+
+				switch (Verb)
+				{
+				case Head:
+					bSuccess = Is400 || Is200;
+					VerbStr = TEXT("querying");
+					break;
+				case Get:
+					bSuccess = Is400 || Is200;
+					VerbStr = TEXT("fetching");
+					AdditionalInfo = FString::Printf(TEXT("Received: %d bytes."), BytesReceived);
+					break;
+				case Put:
+					bSuccess = Is200;
+					VerbStr = TEXT("updating");
+					AdditionalInfo = FString::Printf(TEXT("Sent: %d bytes."), BytesSent);
+					break;
+				case Post:
+				case PostJson:
+					bSuccess = Is200;
+					VerbStr = TEXT("posting");
+					break;
+				case Delete:
+					bSuccess = Is200;
+					VerbStr = TEXT("deleting");
+					break;
+				}
+
+				if (bSuccess)
+				{
+					UE_LOG(
+						LogDerivedDataCache,
+						Verbose,
+						TEXT("Finished %s HTTP cache entry (response %d) from %s. %s"),
+						VerbStr,
+						ResponseCode,
+						Uri,
+						*AdditionalInfo
+					);
+				}
+				else if (bLogErrors)
+				{
+					// Print the response body if we got one, otherwise print header.
+					FString Response = GetAnsiBufferAsString(ResponseBuffer.Num() > 0 ? ResponseBuffer : ResponseHeader);
+					Response.ReplaceCharInline('\n', ' ');
+					Response.ReplaceCharInline('\r', ' ');
+					UE_LOG(
+						LogDerivedDataCache,
+						Error,
+						TEXT("Failed %s HTTP cache entry (response %d) from %s. Response: %s"),
+						VerbStr,
+						ResponseCode,
+						Uri,
+						*Response
+					);
+				}
+			}
+			else if (bLogErrors)
+			{
+				UE_LOG(
+					LogDerivedDataCache,
+					Error,
+					TEXT("Error while connecting to %s: %s"),
+					*Domain,
+					ANSI_TO_TCHAR(curl_easy_strerror(Result))
+				);
+			}
+		}
+
+		FString GetAnsiBufferAsString(const TArray<uint8>& Buffer) const
+		{
+			// Content is NOT null-terminated; we need to specify lengths here
+			FUTF8ToTCHAR TCHARData(reinterpret_cast<const ANSICHAR*>(Buffer.GetData()), Buffer.Num());
+			return FString(TCHARData.Length(), TCHARData.Get());
+		}
+
+		static size_t StaticDebugCallback(CURL* Handle, curl_infotype DebugInfoType, char* DebugInfo, size_t DebugInfoSize, void* UserData)
+		{
+			FRequest* Request = static_cast<FRequest*>(UserData);
+
+			switch (DebugInfoType)
+			{
+			case CURLINFO_TEXT:
+			{
+				// Truncate at 1023 characters. This is just an arbitrary number based on a buffer size seen in
+				// the libcurl code.
+				DebugInfoSize = FMath::Min(DebugInfoSize, (size_t)1023);
+
+				// Calculate the actual length of the string due to incorrect use of snprintf() in lib/vtls/openssl.c.
+				char* FoundNulPtr = (char*)memchr(DebugInfo, 0, DebugInfoSize);
+				int CalculatedSize = FoundNulPtr != nullptr ? FoundNulPtr - DebugInfo : DebugInfoSize;
+
+				auto ConvertedString = StringCast<TCHAR>(static_cast<const ANSICHAR*>(DebugInfo), CalculatedSize);
+				FString DebugText(ConvertedString.Length(), ConvertedString.Get());
+				DebugText.ReplaceInline(TEXT("\n"), TEXT(""), ESearchCase::CaseSensitive);
+				DebugText.ReplaceInline(TEXT("\r"), TEXT(""), ESearchCase::CaseSensitive);
+				UE_LOG(LogDerivedDataCache, VeryVerbose, TEXT("%p: '%s'"), Request, *DebugText);
+			}
+			break;
+
+			case CURLINFO_HEADER_IN:
+				UE_LOG(LogDerivedDataCache, VeryVerbose, TEXT("%p: Received header (%d bytes)"), Request, DebugInfoSize);
+				break;
+
+			case CURLINFO_DATA_IN:
+				UE_LOG(LogDerivedDataCache, VeryVerbose, TEXT("%p: Received data (%d bytes)"), Request, DebugInfoSize);
+				break;
+
+			case CURLINFO_DATA_OUT:
+				UE_LOG(LogDerivedDataCache, VeryVerbose, TEXT("%p: Sent data (%d bytes)"), Request, DebugInfoSize);
+				break;
+			}
+
+			return 0;
+		}
+
+		static size_t StaticReadFn(void* Ptr, size_t SizeInBlocks, size_t BlockSizeInBytes, void* UserData)
+		{
+			FRequest* Request = static_cast<FRequest*>(UserData);
+			TArrayView<const uint8>& ReadDataView = Request->ReadDataView;
+
+			const size_t Offset = Request->BytesSent;
+			const size_t ReadSize = FMath::Min((size_t)ReadDataView.Num() - Offset, SizeInBlocks * BlockSizeInBytes);
+			check(ReadDataView.Num() >= Offset + ReadSize);
+
+			FMemory::Memcpy(Ptr, ReadDataView.GetData() + Offset, ReadSize);
+			Request->BytesSent += ReadSize;
+			return ReadSize;
+
+			return 0;
+		}
+
+		static size_t StaticWriteHeaderFn(void* Ptr, size_t SizeInBlocks, size_t BlockSizeInBytes, void* UserData)
+		{
+			FRequest* Request = static_cast<FRequest*>(UserData);
+			const size_t WriteSize = SizeInBlocks * BlockSizeInBytes;
+			TArray<uint8>* WriteHeaderBufferPtr = Request->WriteHeaderBufferPtr;
+			if (WriteHeaderBufferPtr && WriteSize > 0)
+			{
+				const size_t CurrentBufferLength = WriteHeaderBufferPtr->Num();
+				if (CurrentBufferLength > 0)
+				{
+					// Remove the previous zero termination
+					(*WriteHeaderBufferPtr)[CurrentBufferLength - 1] = ' ';
+				}
+
+				// Write the header
+				WriteHeaderBufferPtr->Append((const uint8*)Ptr, WriteSize + 1);
+				(*WriteHeaderBufferPtr)[WriteHeaderBufferPtr->Num() - 1] = 0; // Zero terminate string
+				return WriteSize;
+			}
+			return 0;
+		}
+
+		static size_t StaticWriteBodyFn(void* Ptr, size_t SizeInBlocks, size_t BlockSizeInBytes, void* UserData)
+		{
+			FRequest* Request = static_cast<FRequest*>(UserData);
+			const size_t WriteSize = SizeInBlocks * BlockSizeInBytes;
+			TArray<uint8>* WriteDataBufferPtr = Request->WriteDataBufferPtr;
+
+			if (WriteDataBufferPtr && WriteSize > 0)
+			{
+				// If this is the first part of the body being received, try to reserve 
+				// memory if content length is defined in the header.
+				if (Request->BytesReceived == 0 && Request->WriteHeaderBufferPtr)
+				{
+					static const ANSICHAR* ContentLengthHeaderStr = "Content-Length: ";
+					const ANSICHAR* Header = (const ANSICHAR*)Request->WriteHeaderBufferPtr->GetData();
+
+					if (const ANSICHAR* ContentLengthHeader = FCStringAnsi::Strstr(Header, ContentLengthHeaderStr))
+					{
+						size_t ContentLength = (size_t)FCStringAnsi::Atoi64(ContentLengthHeader + strlen(ContentLengthHeaderStr));
+						if (ContentLength > 0u && ContentLength < UE_ZENHTTPDDC_MAX_BUFFER_RESERVE)
+						{
+							WriteDataBufferPtr->Reserve(ContentLength);
+						}
+					}
+				}
+
+				// Write to the target buffer
+				WriteDataBufferPtr->Append((const uint8*)Ptr, WriteSize);
+				Request->BytesReceived += WriteSize;
+				return WriteSize;
+			}
+
+			return 0;
+		}
+
+		static size_t StaticSeekFn(void* UserData, curl_off_t Offset, int Origin)
+		{
+			FRequest* Request = static_cast<FRequest*>(UserData);
+			size_t NewPosition = 0;
+
+			switch (Origin)
+			{
+			case SEEK_SET: NewPosition = Offset; break;
+			case SEEK_CUR: NewPosition = Request->BytesSent + Offset; break;
+			case SEEK_END: NewPosition = Request->ReadDataView.Num() + Offset; break;
+			}
+
+			// Make sure we don't seek outside of the buffer
+			if (NewPosition < 0 || NewPosition >= Request->ReadDataView.Num())
+			{
+				return CURL_SEEKFUNC_FAIL;
+			}
+
+			// Update the used offset
+			Request->BytesSent = NewPosition;
+			return CURL_SEEKFUNC_OK;
+		}
+
+	};
+
+	//----------------------------------------------------------------------------------------------------------
+	// Request pool
+	//----------------------------------------------------------------------------------------------------------
+
+	/**
+	 * Pool that manages a fixed set of requests. Users are required to release requests that have been
+	 * acquired. Usable with \ref zenhttp::FScopedRequestPtr which handles this automatically.
+	 */
+	struct FRequestPool
+	{
+		FRequestPool(const TCHAR* InServiceUrl)
+		{
+			for (uint8 i = 0; i < Pool.Num(); ++i)
+			{
+				Pool[i].Usage = 0u;
+				Pool[i].Request = new FRequest(InServiceUrl, true);
+			}
+		}
+
+		~FRequestPool()
+		{
+			for (uint8 i = 0; i < Pool.Num(); ++i)
+			{
+				// No requests should be in use by now.
+				check(Pool[i].Usage.Load(EMemoryOrder::Relaxed) == 0u);
+				delete Pool[i].Request;
+			}
+		}
+
+		/**
+		 * Block until a request is free. Once a request has been returned it is
+		 * "owned by the caller and need to release it to the pool when work has been completed.
+		 * @return Usable request instance.
+		 */
+		FRequest* WaitForFreeRequest()
+		{
+			TRACE_CPUPROFILER_EVENT_SCOPE(ZenHTTPDDC_WaitForConnPool);
+			while (true)
+			{
+				for (uint8 i = 0; i < Pool.Num(); ++i)
+				{
+					if (!Pool[i].Usage.Load(EMemoryOrder::Relaxed))
+					{
+						uint8 Expected = 0u;
+						if (Pool[i].Usage.CompareExchange(Expected, 1u))
+						{
+							return Pool[i].Request;
+						}
+					}
+				}
+				FPlatformProcess::Sleep(UE_ZENHTTPDDC_BACKEND_WAIT_INTERVAL);
+			}
+		}
+
+		/**
+		 * Release request to the pool.
+		 * @param Request Request that should be freed. Note that any buffer owened by the request can now be reset.
+		 */
+		void ReleaseRequestToPool(FRequest* Request)
+		{
+			for (uint8 i = 0; i < Pool.Num(); ++i)
+			{
+				if (Pool[i].Request == Request)
+				{
+					Request->Reset();
+					uint8 Expected = 1u;
+					Pool[i].Usage.CompareExchange(Expected, 0u);
+					return;
+				}
+			}
+			check(false);
+		}
+
+	private:
+
+		struct FEntry
+		{
+			TAtomic<uint8> Usage;
+			FRequest* Request;
+		};
+
+		TStaticArray<FEntry, UE_ZENHTTPDDC_REQUEST_POOL_SIZE> Pool;
+
+		FRequestPool() {}
+	};
+
+	//----------------------------------------------------------------------------------------------------------
+	// FScopedRequestPtr
+	//----------------------------------------------------------------------------------------------------------
+
+	/**
+	 * Utility class to manage requesting and releasing requests from the \ref FRequestPool.
+	 */
+	struct FScopedRequestPtr
+	{
+	public:
+		FScopedRequestPtr(FRequestPool* InPool)
+			: Request(InPool->WaitForFreeRequest())
+			, Pool(InPool)
+		{}
+
+		~FScopedRequestPtr()
+		{
+			Pool->ReleaseRequestToPool(Request);
+		}
+
+		bool IsValid() const
+		{
+			return Request != nullptr;
+		}
+
+		FRequest* Get() const
+		{
+			check(IsValid());
+			return Request;
+		}
+
+		FRequest* operator->()
+		{
+			check(IsValid());
+			return Request;
+		}
+
+	private:
+		FRequest* Request;
+		FRequestPool* Pool;
+	};
+
+	//----------------------------------------------------------------------------------------------------------
+	// Content parsing and checking
+	//----------------------------------------------------------------------------------------------------------
+
+
+	/**
+	 * Verifies the integrity of the recieved data using supplied checksum.
+	 * @param Hash Recieved hash value.
+	 * @param Payload Payload recieved.
+	 * @return True if the data is correct, false if checksums doesn't match.
+	 */
+	bool VerifyPayload(FSHAHash Hash, const TArray<uint8>& Payload)
+	{
+		FSHAHash PayloadHash;
+		FSHA1::HashBuffer(Payload.GetData(), Payload.Num(), PayloadHash.Hash);
+
+		if (Hash != PayloadHash)
+		{
+			UE_LOG(LogDerivedDataCache,
+				Warning,
+				TEXT("Checksum from server did not match recieved data (%s vs %s). Discarding cached result."),
+				*Hash.ToString(),
+				*PayloadHash.ToString()
+			);
+			return false;
+		}
+
+		return true;
+	}
+
+
+	/**
+	 * Verifies the integrity of the recieved data using supplied checksum.
+	 * @param Request Request that the data was be recieved with.
+	 * @param Payload Payload recieved.
+	 * @return True if the data is correct, false if checksums doesn't match.
+	 */
+	bool VerifyRequest(const FRequest* Request, const TArray<uint8>& Payload)
+	{
+#if 0
+		FString RecievedHashStr;
+		FSHAHash PayloadHash;
+		FSHA1::HashBuffer(Payload.GetData(), Payload.Num(), PayloadHash.Hash);
+		if (Request->GetHeader("X-Jupiter-Sha1", RecievedHashStr))
+		{
+			FSHAHash RecievedHash;
+			RecievedHash.FromString(RecievedHashStr);
+			return VerifyPayload(RecievedHash, Payload);
+		}
+		UE_LOG(LogDerivedDataCache, Warning, TEXT("HTTP server did not send a content hash. Wrong server version?"));
+#endif
+		return true;
+	}
+
+	/**
+	 * Adds a checksum (as request header) for a given payload. Jupiter will use this to verify the integrity
+	 * of the recieved data.
+	 * @param Request Request that the data will be sent with.
+	 * @param Payload Payload that will be sent.
+	 * @return True on success, false on failure.
+	 */
+	bool HashPayload(FRequest* Request, const TArrayView<const uint8> Payload)
+	{
+		FSHAHash PayloadHash;
+		FSHA1::HashBuffer(Payload.GetData(), Payload.Num(), PayloadHash.Hash);
+		Request->SetHeader(TEXT("X-Jupiter-Sha1"), *PayloadHash.ToString());
+		return true;
+	}
+
+}
+
+//----------------------------------------------------------------------------------------------------------
+// FZenHttpDerivedDataBackend
+//----------------------------------------------------------------------------------------------------------
+
+namespace UE::DerivedData::Backends {
+
+FZenHttpDerivedDataBackend::FZenHttpDerivedDataBackend(
+	ICacheFactory& InFactory,
+	const TCHAR* InServiceUrl, 
+	const TCHAR* InNamespace)
+	: Factory(InFactory)
+	, Domain(InServiceUrl)
+	, Namespace(InNamespace)
+	, DefaultBucket(TEXT("default"))
+	, bIsUsable(false)
+	, FailedLoginAttempts(0)
+{
+	if (IsServiceReady())
+	{
+		RequestPool = MakeUnique<zenhttp::FRequestPool>(InServiceUrl);
+		bIsUsable = true;
+	}
+}
+
+FZenHttpDerivedDataBackend::~FZenHttpDerivedDataBackend()
+{
+}
+
+FString FZenHttpDerivedDataBackend::GetName() const
+{
+	return Domain;
+}
+
+bool FZenHttpDerivedDataBackend::IsServiceReady()
+{
+	zenhttp::FRequest Request(*Domain, false);
+	zenhttp::FRequest::Result Result = Request.PerformBlockingDownload(TEXT("/health/ready"), nullptr);
+	
+	if (Result == zenhttp::FRequest::Success && Request.GetResponseCode() == 200)
+	{
+		UE_LOG(LogDerivedDataCache, Display, TEXT("HTTP DDC service status: %s."), *Request.GetResponseAsString());
+		return true;
+	}
+	else
+	{
+		UE_LOG(LogDerivedDataCache, Warning, TEXT("Unable to reach HTTP DDC service at %s. Status: %d . Response: %s"), *Domain, Request.GetResponseCode(), *Request.GetResponseAsString());
+	}
+
+	return false;
+}
+
+bool FZenHttpDerivedDataBackend::ShouldRetryOnError(int64 ResponseCode)
+{
+	// Access token might have expired, request a new token and try again.
+	if (ResponseCode == 401)
+	{
+		return true;
+	}
+
+	// Too many requests, make a new attempt
+	if (ResponseCode == 429)
+	{
+		return true;
+	}
+
+	return false;
+}
+
+bool FZenHttpDerivedDataBackend::CachedDataProbablyExists(const TCHAR* CacheKey)
+{
+	TRACE_CPUPROFILER_EVENT_SCOPE(ZenHTTPDDC_Exist);
+	TRACE_COUNTER_ADD(ZenHTTPDDC_Exist, int64(1));
+	COOK_STAT(auto Timer = UsageStats.TimeProbablyExists());
+	
+	FString Uri = FString::Printf(TEXT("/cache/%s"), CacheKey);
+	long ResponseCode = 0; uint32 Attempts = 0;
+
+	// Retry request until we get an accepted response or exhaust allowed number of attempts.
+	while (ResponseCode == 0 && ++Attempts < UE_ZENHTTPDDC_MAX_ATTEMPTS)
+	{
+		zenhttp::FScopedRequestPtr Request(RequestPool.Get());
+		zenhttp::FRequest::Result Result = Request->PerformBlockingQuery<zenhttp::FRequest::Head>(*Uri);
+		ResponseCode = Request->GetResponseCode();
+
+		if (ResponseCode == 200 || ResponseCode == 400)
+		{
+			const bool bIsHit = (Result == zenhttp::FRequest::Success && ResponseCode == 200);
+			if (bIsHit)
+			{
+				TRACE_COUNTER_ADD(ZenHTTPDDC_ExistHit, int64(1));
+			}
+			return bIsHit;
+		}
+
+		if (!ShouldRetryOnError(ResponseCode))
+		{
+			return false;
+		}
+
+		ResponseCode = 0;
+	}
+
+	return false;
+}
+
+bool FZenHttpDerivedDataBackend::GetCachedData(const TCHAR* CacheKey, TArray<uint8>& OutData)
+{
+	TRACE_CPUPROFILER_EVENT_SCOPE(ZenHTTPDDC_Get);
+	TRACE_COUNTER_ADD(ZenHTTPDDC_Get, int64(1));
+	COOK_STAT(auto Timer = UsageStats.TimeGet());
+
+	FString Uri = FString::Printf(TEXT("/cache/%s"), CacheKey);
+	int64 ResponseCode = 0; uint32 Attempts = 0;
+
+	// Retry request until we get an accepted response or exhaust allowed number of attempts.
+	while (ResponseCode == 0 && ++Attempts < UE_ZENHTTPDDC_MAX_ATTEMPTS)
+	{
+		zenhttp::FScopedRequestPtr Request(RequestPool.Get());
+		if (Request.IsValid())
+		{
+			zenhttp::FRequest::Result Result = Request->PerformBlockingDownload(*Uri, &OutData);
+			ResponseCode = Request->GetResponseCode();
+
+			// Request was successful, make sure we got all the expected data.
+			if (ResponseCode == 200 && VerifyRequest(Request.Get(), OutData))
+			{
+				TRACE_COUNTER_ADD(ZenHTTPDDC_GetHit, int64(1));
+				TRACE_COUNTER_ADD(ZenHTTPDDC_BytesReceived, int64(Request->GetBytesReceived()));
+				COOK_STAT(Timer.AddHit(Request->GetBytesReceived()));
+				return true;
+			}
+
+			if (!ShouldRetryOnError(ResponseCode))
+			{
+				return false;
+			}
+
+			ResponseCode = 0;
+		}
+	}
+
+	return false;
+}
+
+FDerivedDataBackendInterface::EPutStatus FZenHttpDerivedDataBackend::PutCachedData(const TCHAR* CacheKey, TArrayView<const uint8> InData, bool bPutEvenIfExists)
+{
+	TRACE_CPUPROFILER_EVENT_SCOPE(ZenHTTPDDC_Put);
+	COOK_STAT(auto Timer = UsageStats.TimePut());
+
+	FString Uri = FString::Printf(TEXT("/cache/%s"), CacheKey);
+	int64 ResponseCode = 0; uint32 Attempts = 0;
+
+	// Retry request until we get an accepted response or exhaust allowed number of attempts.
+	while (ResponseCode == 0 && ++Attempts < UE_ZENHTTPDDC_MAX_ATTEMPTS)
+	{
+		zenhttp::FScopedRequestPtr Request(RequestPool.Get());
+		if (Request.IsValid())
+		{
+			// Append the content hash to the header
+			HashPayload(Request.Get(), InData);
+
+			zenhttp::FRequest::Result Result = Request->PerformBlockingUpload<zenhttp::FRequest::Put>(*Uri, InData);
+			ResponseCode = Request->GetResponseCode();
+
+			if (ResponseCode >= 200 && ResponseCode < 300)
+			{
+				TRACE_COUNTER_ADD(ZenHTTPDDC_BytesSent, int64(Request->GetBytesSent()));
+				COOK_STAT(Timer.AddHit(Request->GetBytesSent()));
+
+				return EPutStatus::Cached;
+			}
+
+			if (!ShouldRetryOnError(ResponseCode))
+			{
+				return EPutStatus::NotCached;
+			}
+
+			ResponseCode = 0;
+		}
+	}
+
+	return EPutStatus::NotCached;
+}
+
+void FZenHttpDerivedDataBackend::RemoveCachedData(const TCHAR* CacheKey, bool bTransient)
+{
+	TRACE_CPUPROFILER_EVENT_SCOPE(ZenHTTPDDC_Remove);
+	FString Uri = FString::Printf(TEXT("/cache/%s"), CacheKey);
+	int64 ResponseCode = 0; uint32 Attempts = 0;
+
+	// Retry request until we get an accepted response or exhaust allowed number of attempts.
+	while (ResponseCode == 0 && ++Attempts < UE_ZENHTTPDDC_MAX_ATTEMPTS)
+	{
+		zenhttp::FScopedRequestPtr Request(RequestPool.Get());
+		if (Request.IsValid())
+		{
+			zenhttp::FRequest::Result Result = Request->PerformBlockingQuery<zenhttp::FRequest::Delete>(*Uri);
+			ResponseCode = Request->GetResponseCode();
+
+			if (ResponseCode == 200)
+			{
+				return;
+			}
+
+			if (!ShouldRetryOnError(ResponseCode))
+			{
+				return;
+			}
+
+			ResponseCode = 0;
+		}
+	}
+}
+
+bool FZenHttpDerivedDataBackend::IsWritable() const
+{
+	return true;
+}
+
+FDerivedDataBackendInterface::ESpeedClass
+FZenHttpDerivedDataBackend::GetSpeedClass() const
+{
+	return ESpeedClass::Fast;
+}
+
+TSharedRef<FDerivedDataCacheStatsNode> FZenHttpDerivedDataBackend::GatherUsageStats() const
+{
+	return MakeShared<FDerivedDataCacheStatsNode>(this, "ZenHttp");
+}
+
+bool FZenHttpDerivedDataBackend::TryToPrefetch(TConstArrayView<FString> CacheKeys)
+{
+	return false;
+}
+
+bool FZenHttpDerivedDataBackend::WouldCache(const TCHAR* CacheKey, TArrayView<const uint8> InData)
+{
+	return true;
+}
+
+bool FZenHttpDerivedDataBackend::ApplyDebugOptions(FBackendDebugOptions& InOptions)
+{
+	return false;
+}
+
+FRequest FZenHttpDerivedDataBackend::Put(
+	TConstArrayView<UE::DerivedData::FCacheRecord> Records,
+	FStringView Context,
+	ECachePolicy Policy,
+	EPriority Priority,
+	FOnCachePutComplete&& OnComplete)
+{
+	if (OnComplete)
+	{
+		for (const FCacheRecord& Record : Records)
+		{
+			OnComplete({ Record.GetKey(), EStatus::Error });
+		}
+	}
+	return FRequest();
+}
+
+FRequest FZenHttpDerivedDataBackend::Get(
+	TConstArrayView<UE::DerivedData::FCacheKey> Keys,
+	FStringView Context,
+	UE::DerivedData::ECachePolicy Policy,
+	UE::DerivedData::EPriority Priority,
+	UE::DerivedData::FOnCacheGetComplete&& OnComplete)
+{
+	if (OnComplete)
+	{
+		for (const FCacheKey& Key : Keys)
+		{
+			OnComplete({ Factory.CreateRecord(Key).Build(), EStatus::Error });
+		}
+	}
+	return FRequest();
+}
+
+FRequest FZenHttpDerivedDataBackend::GetPayload(
+	TConstArrayView<UE::DerivedData::FCachePayloadKey> Keys,
+	FStringView Context,
+	UE::DerivedData::ECachePolicy Policy,
+	UE::DerivedData::EPriority Priority,
+	UE::DerivedData::FOnCacheGetPayloadComplete&& OnComplete)
+{
+	if (OnComplete)
+	{
+		for (const FCachePayloadKey& Key : Keys)
+		{
+			OnComplete({ Key.CacheKey, FPayload(Key.Id), EStatus::Error });
+		}
+	}
+	return FRequest();
+}
+
+void FZenHttpDerivedDataBackend::CancelAll()
+{
+}
+
+}
+
+#endif //WITH_HTTP_DDC_BACKEND
diff --git a/Engine/Source/Developer/DerivedDataCache/Private/ZenKvDerivedDataBackend.h b/Engine/Source/Developer/DerivedDataCache/Private/ZenKvDerivedDataBackend.h
new file mode 100644
--- /dev/null
+++ b/Engine/Source/Developer/DerivedDataCache/Private/ZenKvDerivedDataBackend.h
@@ -0,0 +1,112 @@
+// Copyright Epic Games, Inc. All Rights Reserved.
+
+#pragma once
+
+#include "CoreMinimal.h"
+#include "DerivedDataBackendInterface.h"
+#include "DerivedDataCacheUsageStats.h"
+
+// Macro for whether to enable the S3 backend. libcurl is not currently available on Mac.
+#if PLATFORM_WINDOWS
+#	define WITH_ZEN_DDC_BACKEND 1
+#else
+#	define WITH_ZEN_DDC_BACKEND 0
+#endif
+
+#if WITH_ZEN_DDC_BACKEND
+
+typedef TSharedPtr<class IHttpRequest> FHttpRequestPtr;
+typedef TSharedPtr<class IHttpResponse, ESPMode::ThreadSafe> FHttpResponsePtr;
+
+namespace zenhttp {
+	struct FRequestPool;
+}
+
+namespace UE::DerivedData::Backends
+{
+
+/**
+ * Backend for a HTTP based caching service (Zen)
+ **/
+class FZenHttpDerivedDataBackend : public FDerivedDataBackendInterface
+{
+public:
+
+	/**
+	 * Creates the backend, checks health status and attempts to acquire an access token.
+	 *
+	 * @param ServiceUrl	Base url to the service including schema.
+	 * @param Namespace		Namespace to use.
+	 */
+	FZenHttpDerivedDataBackend(ICacheFactory& Factory, const TCHAR* ServiceUrl, const TCHAR* Namespace);
+
+	~FZenHttpDerivedDataBackend();
+
+	/** return true if this cache is writable **/
+	virtual bool IsWritable() const override;
+	virtual ESpeedClass GetSpeedClass() const override;
+	virtual bool TryToPrefetch(TConstArrayView<FString> CacheKeys) override;
+
+	/**
+	 * Retrieve usage stats for this backend. If the backend holds inner backends, this is expected to be passed down recursively.
+	 */
+	virtual TSharedRef<FDerivedDataCacheStatsNode> GatherUsageStats() const;
+
+	/**
+	 * Checks is backend is usable (reachable and accessible).
+	 * @return true if usable
+	 */
+	bool IsUsable() const { return bIsUsable; }
+
+	virtual bool CachedDataProbablyExists(const TCHAR* CacheKey) override;
+	virtual bool GetCachedData(const TCHAR* CacheKey, TArray<uint8>& OutData) override;
+	virtual EPutStatus PutCachedData(const TCHAR* CacheKey, TArrayView<const uint8> InData, bool bPutEvenIfExists) override;
+	virtual void RemoveCachedData(const TCHAR* CacheKey, bool bTransient) override;
+
+	virtual FString GetName() const override;
+	virtual bool WouldCache(const TCHAR* CacheKey, TArrayView<const uint8> InData) override;
+	virtual bool ApplyDebugOptions(FBackendDebugOptions& InOptions) override;
+
+	// ICacheStore
+
+	virtual FRequest Put(
+		TConstArrayView<FCacheRecord> Records,
+		FStringView Context,
+		ECachePolicy Policy = ECachePolicy::Default,
+		EPriority Priority = EPriority::Normal,
+		FOnCachePutComplete&& OnComplete = FOnCachePutComplete());
+
+	virtual FRequest Get(
+		TConstArrayView<FCacheKey> Keys,
+		FStringView Context,
+		ECachePolicy Policy,
+		EPriority Priority,
+		FOnCacheGetComplete&& OnComplete);
+
+	virtual FRequest GetPayload(
+		TConstArrayView<FCachePayloadKey> Keys,
+		FStringView Context,
+		ECachePolicy Policy,
+		EPriority Priority,
+		FOnCacheGetPayloadComplete&& OnComplete);
+	
+	virtual void CancelAll();
+
+private:
+	ICacheFactory& Factory;
+	FString Domain;
+	FString Namespace;
+	FString DefaultBucket;
+	FCriticalSection AccessCs;
+	FDerivedDataCacheUsageStats UsageStats;
+	TUniquePtr<struct zenhttp::FRequestPool> RequestPool;
+	bool bIsUsable;
+	uint32 FailedLoginAttempts;
+
+	bool IsServiceReady();
+	bool ShouldRetryOnError(int64 ResponseCode);
+};
+
+} // UE::DerivedData::Backends
+
+#endif //WITH_HTTP_DDC_BACKEND
diff --git a/Engine/Source/Developer/IoStoreUtilities/IoStoreUtilities.Build.cs b/Engine/Source/Developer/IoStoreUtilities/IoStoreUtilities.Build.cs
--- a/Engine/Source/Developer/IoStoreUtilities/IoStoreUtilities.Build.cs
+++ b/Engine/Source/Developer/IoStoreUtilities/IoStoreUtilities.Build.cs
@@ -16,6 +16,7 @@ public class IoStoreUtilities : ModuleRules
             "Projects",
             "AssetRegistry",
             "RenderCore",
+			"libcurl",
         });
         PrivateDependencyModuleNames.Add("PakFile");
         PrivateDependencyModuleNames.Add("Json");
diff --git a/Engine/Source/Developer/IoStoreUtilities/Private/FilePackageStoreWriter.cpp b/Engine/Source/Developer/IoStoreUtilities/Private/FilePackageStoreWriter.cpp
new file mode 100644
--- /dev/null
+++ b/Engine/Source/Developer/IoStoreUtilities/Private/FilePackageStoreWriter.cpp
@@ -0,0 +1,358 @@
+// Copyright Epic Games, Inc. All Rights Reserved.
+
+#include "FilePackageStoreWriter.h"
+#include "PackageStoreOptimizer.h"
+#include "Containers/Map.h"
+#include "Serialization/LargeMemoryWriter.h"
+#include "Misc/FileHelper.h"
+#include "Misc/Paths.h"
+#include "HAL/PlatformFileManager.h"
+#include "Misc/App.h"
+#include "Misc/DataDrivenPlatformInfoRegistry.h"
+#include "Interfaces/ITargetPlatform.h"
+#include "Serialization/JsonSerializer.h"
+#include "Misc/ScopeLock.h"
+
+DEFINE_LOG_CATEGORY_STATIC(LogPackageStoreWriter, Log, All);
+
+void WriteFileRegions(const TCHAR* Filename, const TArray<FFileRegion>& FileRegions)
+{
+	if (FileRegions.Num() > 0)
+	{
+		TArray<FFileRegion> FileRegionsCopy(FileRegions);
+		TUniquePtr<FArchive> Ar(IFileManager::Get().CreateFileWriter(Filename));
+		FFileRegion::SerializeFileRegions(*Ar, FileRegionsCopy);
+	}
+}
+
+void FPackageStoreManifest::BeginPackage(FName PackageName)
+{
+	FScopeLock Lock(&CriticalSection);
+	FPackageInfo& PackageInfo = PackageInfoByNameMap.FindOrAdd(PackageName);
+	PackageInfo.PackageName = PackageName;
+	if (PackageInfo.PackageChunkId.IsValid())
+	{
+		FileNameByChunkIdMap.Remove(PackageInfo.PackageChunkId);
+		PackageInfo.PackageChunkId = FIoChunkId();
+	}
+	for (const FIoChunkId& BulkDataChunkId : PackageInfo.BulkDataChunkIds)
+	{
+		FileNameByChunkIdMap.Remove(BulkDataChunkId);
+	}
+	PackageInfo.BulkDataChunkIds.Empty();
+}
+
+void FPackageStoreManifest::AddPackageData(FName PackageName, const FString& FileName, const FIoChunkId& ChunkId)
+{
+	FScopeLock Lock(&CriticalSection);
+	FPackageInfo* PackageInfo = PackageInfoByNameMap.Find(PackageName);
+	check(PackageInfo);
+	PackageInfo->PackageChunkId = ChunkId;
+	if (!FileName.IsEmpty())
+	{
+		FileNameByChunkIdMap.Add(ChunkId, FileName);
+	}
+}
+
+void FPackageStoreManifest::AddBulkData(FName PackageName, const FString& FileName, const FIoChunkId& ChunkId)
+{
+	FScopeLock Lock(&CriticalSection);
+	FPackageInfo* PackageInfo = PackageInfoByNameMap.Find(PackageName);
+	check(PackageInfo);
+	PackageInfo->BulkDataChunkIds.Add(ChunkId);
+	if (!FileName.IsEmpty())
+	{
+		FileNameByChunkIdMap.Add(ChunkId, FileName);
+	}
+}
+
+FIoStatus FPackageStoreManifest::Save(const TCHAR* Filename) const
+{
+	FScopeLock Lock(&CriticalSection);
+	TStringBuilder<64> ChunkIdStringBuilder;
+	auto ChunkIdToString = [&ChunkIdStringBuilder](const FIoChunkId& ChunkId)
+	{
+		ChunkIdStringBuilder.Reset();
+		ChunkIdStringBuilder << ChunkId;
+		return *ChunkIdStringBuilder;
+	};
+
+	FString JsonTcharText;
+	TSharedRef<TJsonWriter<TCHAR, TPrettyJsonPrintPolicy<TCHAR>>> Writer = TJsonWriterFactory<TCHAR, TPrettyJsonPrintPolicy<TCHAR>>::Create(&JsonTcharText);
+	Writer->WriteObjectStart();
+
+	if (ZenServerInfo)
+	{
+		Writer->WriteObjectStart(TEXT("ZenServer"));
+		Writer->WriteValue(TEXT("HostName"), ZenServerInfo->HostName);
+		Writer->WriteValue(TEXT("Port"), ZenServerInfo->Port);
+		Writer->WriteValue(TEXT("ProjectId"), ZenServerInfo->ProjectId);
+		Writer->WriteValue(TEXT("OplogId"), ZenServerInfo->OplogId);
+		Writer->WriteObjectEnd();
+	}
+	
+	Writer->WriteArrayStart(TEXT("Files"));
+	for (const auto& KV : FileNameByChunkIdMap)
+	{
+		Writer->WriteObjectStart();
+		Writer->WriteValue(TEXT("Path"), KV.Value);
+		Writer->WriteValue(TEXT("ChunkId"), ChunkIdToString(KV.Key));
+		Writer->WriteObjectEnd();
+	}
+	Writer->WriteArrayEnd();
+
+	Writer->WriteArrayStart(TEXT("Packages"));
+	for (const auto& KV : PackageInfoByNameMap)
+	{
+		const FPackageInfo& PackageInfo = KV.Value;
+		Writer->WriteObjectStart();
+		Writer->WriteValue(TEXT("Name"), PackageInfo.PackageName.ToString());
+		Writer->WriteValue(TEXT("PackageChunkId"), ChunkIdToString(PackageInfo.PackageChunkId));
+		if (!PackageInfo.BulkDataChunkIds.IsEmpty())
+		{
+			Writer->WriteArrayStart(TEXT("BulkDataChunkIds"));
+			for (const FIoChunkId& ChunkId : PackageInfo.BulkDataChunkIds)
+			{
+				Writer->WriteValue(ChunkIdToString(ChunkId));
+			}
+			Writer->WriteArrayEnd();
+		}
+		Writer->WriteObjectEnd();
+	}
+	Writer->WriteArrayEnd();
+
+	Writer->WriteObjectEnd();
+	Writer->Close();
+
+	if (!FFileHelper::SaveStringToFile(JsonTcharText, Filename))
+	{
+		return FIoStatus(EIoErrorCode::FileOpenFailed);
+	}
+
+	return FIoStatus::Ok;
+}
+
+FIoStatus FPackageStoreManifest::Load(const TCHAR* Filename)
+{
+	FScopeLock Lock(&CriticalSection);
+	PackageInfoByNameMap.Empty();
+	FileNameByChunkIdMap.Empty();
+
+	auto ChunkIdFromString = [](const FString& ChunkIdString)
+	{
+		FStringView ChunkIdStringView(*ChunkIdString, 24);
+		uint8 Data[12];
+		UE::String::HexToBytes(ChunkIdStringView, Data);
+		FIoChunkId ChunkId;
+		ChunkId.Set(Data, 12);
+		return ChunkId;
+	};
+
+	FString JsonText;
+	if (!FFileHelper::LoadFileToString(JsonText, Filename))
+	{
+		return FIoStatus(EIoErrorCode::FileOpenFailed);
+	}
+
+	TSharedPtr<FJsonObject> JsonObject;
+	TSharedRef<TJsonReader<TCHAR>> Reader = TJsonReaderFactory<TCHAR>::Create(JsonText);
+	if (!FJsonSerializer::Deserialize(Reader, JsonObject) || !JsonObject.IsValid())
+	{
+		return FIoStatus(EIoErrorCode::Unknown);
+	}
+
+	TSharedPtr<FJsonValue> ZenServerValue = JsonObject->Values.FindRef(TEXT("ZenServer"));
+	if (ZenServerValue)
+	{
+		ZenServerInfo = MakeUnique<FZenServerInfo>();
+		TSharedPtr<FJsonObject> ZenServerObject = ZenServerValue->AsObject();
+		ZenServerInfo->HostName = ZenServerObject->Values.FindRef(TEXT("HostName"))->AsString();
+		ZenServerInfo->Port = uint16(ZenServerObject->Values.FindRef(TEXT("Port"))->AsNumber());
+		ZenServerInfo->ProjectId = ZenServerObject->Values.FindRef(TEXT("ProjectId"))->AsString();
+		ZenServerInfo->OplogId = ZenServerObject->Values.FindRef(TEXT("OplogId"))->AsString();
+	}
+
+	TSharedPtr<FJsonValue> FilesArrayValue = JsonObject->Values.FindRef(TEXT("Files"));
+	TArray<TSharedPtr<FJsonValue>> FilesArray = FilesArrayValue->AsArray();
+	FileNameByChunkIdMap.Reserve(FilesArray.Num());
+	for (const TSharedPtr<FJsonValue>& FileValue : FilesArray)
+	{
+		TSharedPtr<FJsonObject> FileObject = FileValue->AsObject();
+		FIoChunkId ChunkId = ChunkIdFromString(FileObject->Values.FindRef(TEXT("ChunkId"))->AsString());
+		FileNameByChunkIdMap.Add(ChunkId, FileObject->Values.FindRef(TEXT("Path"))->AsString());
+	}
+
+	TArray<TSharedPtr<FJsonValue>> PackagesArray = JsonObject->Values.FindRef(TEXT("Packages"))->AsArray();
+	PackageInfoByNameMap.Reserve(PackagesArray.Num());
+	for (const TSharedPtr<FJsonValue>& PackageValue : PackagesArray)
+	{
+		TSharedPtr<FJsonObject> PackageObject = PackageValue->AsObject();
+		FName PackageName = FName(PackageObject->Values.FindRef(TEXT("Name"))->AsString());
+
+		FPackageInfo& PackageInfo = PackageInfoByNameMap.FindOrAdd(PackageName);
+		PackageInfo.PackageName = PackageName;
+		PackageInfo.PackageChunkId = ChunkIdFromString(PackageObject->Values.FindRef(TEXT("PackageChunkId"))->AsString());
+
+		TSharedPtr<FJsonValue> BulkDataChunkIdsValue = PackageObject->Values.FindRef(TEXT("BulkDataChunkIds"));
+		if (BulkDataChunkIdsValue.IsValid())
+		{
+			TArray<TSharedPtr<FJsonValue>> BulkDataChunkIdsArray = BulkDataChunkIdsValue->AsArray();
+			PackageInfo.BulkDataChunkIds.Reserve(BulkDataChunkIdsArray.Num());
+			for (const TSharedPtr<FJsonValue>& BulkDataChunkIdValue : BulkDataChunkIdsArray)
+			{
+				PackageInfo.BulkDataChunkIds.Add(ChunkIdFromString(BulkDataChunkIdValue->AsString()));
+			}
+		}
+	}
+
+	return FIoStatus::Ok;
+}
+
+TArray<FPackageStoreManifest::FFileInfo> FPackageStoreManifest::GetFiles() const
+{
+	FScopeLock Lock(&CriticalSection);
+	TArray<FFileInfo> Files;
+	Files.Reserve(FileNameByChunkIdMap.Num());
+	for (const auto& KV : FileNameByChunkIdMap)
+	{
+		Files.Add({ KV.Value, KV.Key });
+	}
+	return Files;
+}
+
+TArray<FPackageStoreManifest::FPackageInfo> FPackageStoreManifest::GetPackages() const
+{
+	FScopeLock Lock(&CriticalSection);
+	TArray<FPackageInfo> Packages;
+	PackageInfoByNameMap.GenerateValueArray(Packages);
+	return Packages;
+}
+
+FPackageStoreManifest::FZenServerInfo& FPackageStoreManifest::EditZenServerInfo()
+{
+	FScopeLock Lock(&CriticalSection);
+	if (!ZenServerInfo)
+	{
+		ZenServerInfo = MakeUnique<FZenServerInfo>();
+	}
+	return *ZenServerInfo;
+}
+
+const FPackageStoreManifest::FZenServerInfo* FPackageStoreManifest::ReadZenServerInfo() const
+{
+	FScopeLock Lock(&CriticalSection);
+	return ZenServerInfo.Get();
+}
+
+FFilePackageStoreWriter::FFilePackageStoreWriter(const FString& InOutputPath, const FString& InMetadataDirectoryPath, const ITargetPlatform* InTargetPlatform)
+	: PackageStoreOptimizer(new FPackageStoreOptimizer())
+	, IoStoreWriterContext(new FIoStoreWriterContext())
+	, IoStoreWriter(new FIoStoreWriter(*(InOutputPath / TEXT("global"))))
+	, TargetPlatform(*InTargetPlatform)
+	, OutputPath(InOutputPath)
+	, MetadataDirectoryPath(InMetadataDirectoryPath)
+
+{
+	FIoStoreWriterSettings WriterSettings;
+	WriterSettings.CompressionMethod = NAME_Zlib;
+	WriterSettings.InitializePlatformSpecificSettings(InTargetPlatform);
+	FIoStatus IoStatus = IoStoreWriterContext->Initialize(WriterSettings);
+	check(IoStatus.IsOk());
+	FIoContainerSettings ContainerSettings;
+	ContainerSettings.ContainerId = ContainerId;
+	ContainerSettings.ContainerFlags |= EIoContainerFlags::Compressed;
+	ContainerSettings.ContainerFlags |= EIoContainerFlags::Indexed;
+	IoStatus = IoStoreWriter->Initialize(*IoStoreWriterContext, ContainerSettings);
+	check(IoStatus.IsOk());
+	PackageStoreOptimizer->Initialize(InTargetPlatform);
+}
+
+FFilePackageStoreWriter::~FFilePackageStoreWriter()
+{
+}
+
+void FFilePackageStoreWriter::WritePackageData(const FPackageInfo& Info, const FIoBuffer& PackageData, const TArray<FFileRegion>& FileRegions)
+{
+	TRACE_CPUPROFILER_EVENT_SCOPE(WritePackageData);
+	FIoBuffer CookedHeaderBuffer = FIoBuffer(PackageData.Data(), Info.HeaderSize, PackageData);
+	FIoBuffer CookedExportsBuffer = FIoBuffer(PackageData.Data() + Info.HeaderSize, PackageData.DataSize() - Info.HeaderSize, PackageData);
+	FPackageStorePackage* Package = PackageStoreOptimizer->CreatePackageFromCookedHeader(Info.PackageName, CookedHeaderBuffer);
+	PackageStoreOptimizer->FinalizePackage(Package);
+	TArray<FFileRegion> FileRegionsCopy(FileRegions);
+	for (FFileRegion& Region : FileRegionsCopy)
+	{
+		// Adjust regions so they are relative to the start of the exports buffer
+		Region.Offset -= Info.HeaderSize;
+	}
+	FIoBuffer PackageBuffer = PackageStoreOptimizer->CreatePackageBuffer(Package, CookedExportsBuffer, &FileRegionsCopy);
+	FIoWriteOptions WriteOptions;
+	WriteOptions.FileName = Info.LooseFilePath;
+	IoStoreWriter->Append(Info.ChunkId, PackageBuffer, WriteOptions);
+	PackageStoreManifest.AddPackageData(Info.PackageName, Info.LooseFilePath, Info.ChunkId);
+	PackageStoreEntries.Add(PackageStoreOptimizer->CreatePackageStoreEntry(Package));
+	for (FFileRegion& Region : FileRegionsCopy)
+	{
+		// Adjust regions once more so they are relative to the exports bundle buffer
+		Region.Offset -= Package->GetHeaderSize();
+	}
+	WriteFileRegions(*FPaths::ChangeExtension(Info.LooseFilePath, FString(".uexp") + FFileRegion::RegionsFileExtension), FileRegionsCopy);
+	delete Package;
+}
+
+void FFilePackageStoreWriter::WriteBulkdata(const FBulkDataInfo& Info, const FIoBuffer& BulkData, const TArray<FFileRegion>& FileRegions)
+{
+	TRACE_CPUPROFILER_EVENT_SCOPE(WriteBulkdata);
+	FIoWriteOptions WriteOptions;
+	WriteOptions.FileName = Info.LooseFilePath;
+	FPackageId PackageId = FPackageId::FromName(Info.PackageName);
+	if (Info.BulkdataType == FBulkDataInfo::Mmap)
+	{
+		WriteOptions.bForceUncompressed = true;
+		WriteOptions.bIsMemoryMapped = true;
+	}
+	IoStoreWriter->Append(Info.ChunkId, BulkData, WriteOptions);
+	PackageStoreManifest.AddBulkData(Info.PackageName, Info.LooseFilePath, Info.ChunkId);
+	WriteFileRegions(*(Info.LooseFilePath + FFileRegion::RegionsFileExtension), FileRegions);
+}
+
+void FFilePackageStoreWriter::BeginCook(const FCookInfo& Info)
+{
+	check(Info.CookMode == FCookInfo::CookByTheBookMode);
+}
+
+void FFilePackageStoreWriter::EndCook()
+{
+	FIoBuffer ScriptObjectsBuffer = PackageStoreOptimizer->CreateScriptObjectsBuffer();
+	FIoWriteOptions WriteOptions;
+	WriteOptions.DebugName = TEXT("ScriptObjects");
+	IoStoreWriter->Append(CreateIoChunkId(0, 0, EIoChunkType::ScriptObjects), ScriptObjectsBuffer, WriteOptions);
+
+	FContainerHeader Header = PackageStoreOptimizer->CreateContainerHeader(ContainerId, PackageStoreEntries);
+	FLargeMemoryWriter HeaderAr(0, true);
+	HeaderAr << Header;
+	int64 DataSize = HeaderAr.TotalSize();
+	FIoBuffer HeaderBuffer(FIoBuffer::AssumeOwnership, HeaderAr.ReleaseOwnership(), DataSize);
+	WriteOptions.DebugName = TEXT("ContainerHeader");
+	IoStoreWriter->Append(CreateIoChunkId(ContainerId.Value(), 0, EIoChunkType::ContainerHeader), HeaderBuffer, WriteOptions);
+
+	TIoStatusOr<FIoStoreWriterResult> Result = IoStoreWriter->Flush();
+	check(Result.IsOk());
+
+	PackageStoreManifest.Save(*(MetadataDirectoryPath / TEXT("packagestore.manifest")));
+
+	UE_LOG(LogPackageStoreWriter, Display, TEXT("Input:\t%d Packages"), PackageStoreOptimizer->GetTotalPackageCount());
+	UE_LOG(LogPackageStoreWriter, Display, TEXT("Output:\t%d Export bundles"), PackageStoreOptimizer->GetTotalExportBundleCount());
+	UE_LOG(LogPackageStoreWriter, Display, TEXT("Output:\t%d Export bundle entries"), PackageStoreOptimizer->GetTotalExportBundleEntryCount());
+	UE_LOG(LogPackageStoreWriter, Display, TEXT("Output:\t%d Internal export bundle arcs"), PackageStoreOptimizer->GetTotalInternalBundleArcsCount());
+	UE_LOG(LogPackageStoreWriter, Display, TEXT("Output:\t%d External export bundle arcs"), PackageStoreOptimizer->GetTotalExternalBundleArcsCount());
+	UE_LOG(LogPackageStoreWriter, Display, TEXT("Output:\t%d Public runtime script objects"), PackageStoreOptimizer->GetTotalScriptObjectCount());
+}
+
+void FFilePackageStoreWriter::BeginPackage(const FPackageBaseInfo& Info)
+{
+	PackageStoreManifest.BeginPackage(Info.PackageName);
+}
+
+void FFilePackageStoreWriter::CommitPackage(const FPackageBaseInfo& Info)
+{
+}
