Fix DDC prefetch logic error causing useless data loads from disk
  - 46GB to 10GB data loads from disk during map load of special project

#rnx

diff --git a/Engine/Source/Developer/DerivedDataCache/Private/HierarchicalDerivedDataBackend.h b/Engine/Source/Developer/DerivedDataCache/Private/HierarchicalDerivedDataBackend.h
--- a/Engine/Source/Developer/DerivedDataCache/Private/HierarchicalDerivedDataBackend.h
+++ b/Engine/Source/Developer/DerivedDataCache/Private/HierarchicalDerivedDataBackend.h
@@ -254,12 +254,17 @@ public:
 				TBitArray<> Hits = Interface->CachedDataProbablyExistsBatch(SearchKeys);
 				TArray<FString, TInlineAllocator<16>>::TIterator It = SearchKeys.CreateIterator();
 				It.SetToEnd();
-				for (--It; It; --It)
+				for (--It; It; )
 				{
 					if (Hits[It.GetIndex()])
 					{
+						// RemoveCurrent will decrement the iterator
 						It.RemoveCurrent();
 					}
+					else
+					{
+						--It;
+					}
 				}
 
 				// No fetch is necessary if every key already exists in a fast backend.
