DDC: Added -ExportBuilds[=Type1+Type2] to export actions and inputs for builds that miss or skip the cache

#rnx

diff --git a/Engine/Source/Developer/DerivedDataCache/Private/DerivedDataBuildJob.cpp b/Engine/Source/Developer/DerivedDataCache/Private/DerivedDataBuildJob.cpp
--- a/Engine/Source/Developer/DerivedDataCache/Private/DerivedDataBuildJob.cpp
+++ b/Engine/Source/Developer/DerivedDataCache/Private/DerivedDataBuildJob.cpp
@@ -24,12 +24,19 @@
 #include "DerivedDataPayload.h"
 #include "HAL/CriticalSection.h"
 #include "HAL/Event.h"
+#include "HAL/FileManager.h"
 #include "HAL/PlatformProcess.h"
 #include "HAL/PlatformMath.h"
+#include "Misc/CommandLine.h"
 #include "Misc/Guid.h"
+#include "Misc/Parse.h"
+#include "Misc/Paths.h"
+#include "Misc/PathViews.h"
 #include "Misc/ScopeRWLock.h"
 #include "Misc/StringBuilder.h"
 #include "Serialization/CompactBinary.h"
+#include "Serialization/CompactBinaryWriter.h"
+#include "String/ParseTokens.h"
 #include "Templates/Function.h"
 #include "Templates/RefCounting.h"
 #include <atomic>
@@ -176,6 +183,15 @@ private:
 	/** Execute a transition from the old state to the new state. */
 	void ExecuteTransition(EBuildJobState OldState, EBuildJobState NewState);
 
+	/** Exports the action and inputs for this build to disk. */
+	void ExportBuild() const;
+
+	/** Determine whether the action and inputs for this build should be exported to disk. */
+	bool ShouldExportBuild() const;
+
+	/** Parse the types to export from the command line. Sorted by FNameFastLess. */
+	static TArray<FName> ParseExportBuildTypes(bool& bOutExportAll);
+
 private:
 	FString Name;
 	FString FunctionName{TEXT("Unknown")};
@@ -196,6 +212,8 @@ private:
 	bool bInAdvanceToState{};
 	/** True if remote execution was attempted, whether or not it failed. */
 	bool bTriedRemoteExecution{};
+	/** True if the build action and inputs are being exported. */
+	bool bIsExportingBuild{};
 
 	mutable std::atomic<uint32> ReferenceCount{0};
 
@@ -519,6 +537,7 @@ void FBuildJob::CreateContext()
 			[this] { EndExecuteLocal(); });
 		Function->Configure(*Context);
 		BuildPolicy = Context->GetBuildPolicy();
+		bIsExportingBuild = ShouldExportBuild();
 	}
 
 	// Populate the scheduler params with the information that is available now.
@@ -807,7 +826,7 @@ void FBuildJob::EnterExecuteRemote()
 	{
 		return CompleteWithError(TEXT("Failed because build policy skipped building."_SV));
 	}
-	if (!EnumHasAnyFlags(BuildPolicy, EBuildPolicy::Remote))
+	if (!EnumHasAnyFlags(BuildPolicy, EBuildPolicy::Remote) || bIsExportingBuild)
 	{
 		return AdvanceToState(EBuildJobState::ResolveInputData);
 	}
@@ -891,9 +910,9 @@ void FBuildJob::SkipExecuteRemote()
 
 void FBuildJob::EnterExecuteLocal()
 {
-	if (EnumHasAnyFlags(BuildPolicy, EBuildPolicy::SkipBuild))
+	if (bIsExportingBuild)
 	{
-		return CompleteWithError(TEXT("Failed because build policy skipped building."_SV));
+		ExportBuild();
 	}
 	if (!EnumHasAnyFlags(BuildPolicy, EBuildPolicy::Local))
 	{
@@ -1221,6 +1240,104 @@ void FBuildJob::ExecuteTransition(EBuildJobState OldState, EBuildJobState NewSta
 	}
 }
 
+void FBuildJob::ExportBuild() const
+{
+	// Export to <SavedDir>/DerivedDataBuildExport/<Bucket>[/<Function>]/<Action>
+	TStringBuilder<256> ExportPath;
+	const FCacheKey& Key = Context->GetCacheKey();
+	FPathViews::Append(ExportPath, FPaths::ProjectSavedDir(), TEXT("DerivedDataBuildExport"), Key.Bucket);
+	if (FunctionName != Key.Bucket.ToString<TCHAR>())
+	{
+		FPathViews::Append(ExportPath, FunctionName);
+	}
+	FPathViews::Append(ExportPath, Key.Hash);
+	int32 ExportRootLen = ExportPath.Len();
+
+	TAnsiStringBuilder<512> Meta;
+	Meta << "Name: " << FTCHARToUTF8(Name) << LINE_TERMINATOR_ANSI;
+	Meta << "Cache: " << Key << LINE_TERMINATOR_ANSI;
+	Meta << "Function: " << FTCHARToUTF8(FunctionName) << LINE_TERMINATOR_ANSI;
+	Meta << "FunctionVersion: " << Action.Get().GetFunctionVersion() << LINE_TERMINATOR_ANSI;
+	Meta << "BuildSystemVersion: " << Action.Get().GetBuildSystemVersion() << LINE_TERMINATOR_ANSI;
+	if (Action.Get().HasConstants())
+	{
+		Meta << "Constants:" << LINE_TERMINATOR_ANSI;
+		Action.Get().IterateConstants([&Meta](FStringView Key, FCbObject&& Value)
+		{
+			Meta << "  " << FTCHARToUTF8(Key) << ":" LINE_TERMINATOR_ANSI;
+			Meta << "    RawHash: " << FIoHash(Value.GetHash()) << LINE_TERMINATOR_ANSI;
+			Meta << "    RawSize: " << Value.GetSize() << LINE_TERMINATOR_ANSI;
+		});
+	}
+	if (Action.Get().HasInputs())
+	{
+		Meta << "Inputs:" << LINE_TERMINATOR_ANSI;
+		Action.Get().IterateInputs([&Meta](FStringView Key, const FIoHash& RawHash, uint64 RawSize)
+		{
+			Meta << "  " << FTCHARToUTF8(Key) << ":" LINE_TERMINATOR_ANSI;
+			Meta << "    RawHash: " << RawHash << LINE_TERMINATOR_ANSI;
+			Meta << "    RawSize: " << RawSize << LINE_TERMINATOR_ANSI;
+		});
+	}
+
+	FPathViews::Append(ExportPath, TEXT("Meta.yaml"));
+	if (TUniquePtr<FArchive> Ar{IFileManager::Get().CreateFileWriter(*ExportPath)})
+	{
+		Ar->Serialize(Meta.GetData(), Meta.Len());
+	}
+	ExportPath.RemoveSuffix(ExportPath.Len() - ExportRootLen);
+
+	FPathViews::Append(ExportPath, TEXT("Build.action"));
+	if (TUniquePtr<FArchive> Ar{IFileManager::Get().CreateFileWriter(*ExportPath)})
+	{
+		FCbWriter Writer;
+		Action.Get().Save(Writer);
+		Writer.Save(*Ar);
+	}
+	ExportPath.RemoveSuffix(ExportPath.Len() - ExportRootLen);
+
+	if (Inputs)
+	{
+		Meta << "Inputs:" << LINE_TERMINATOR_ANSI;
+		FPathViews::Append(ExportPath, TEXT("Inputs"));
+		ExportRootLen = ExportPath.Len();
+		Inputs.Get().IterateInputs([&ExportPath, ExportRootLen](FStringView Key, const FCompressedBuffer& Buffer)
+		{
+			FPathViews::Append(ExportPath, FIoHash(Buffer.GetRawHash()));
+			if (TUniquePtr<FArchive> Ar{IFileManager::Get().CreateFileWriter(*ExportPath)})
+			{
+				*Ar << const_cast<FCompressedBuffer&>(Buffer);
+			}
+			ExportPath.RemoveSuffix(ExportPath.Len() - ExportRootLen);
+		});
+	}
+}
+
+bool FBuildJob::ShouldExportBuild() const
+{
+	static bool bExportAll;
+	static TArray<FName> ExportTypes = ParseExportBuildTypes(bExportAll);
+	return bExportAll
+		|| Algo::BinarySearch(ExportTypes, FName(FunctionName), FNameFastLess()) != INDEX_NONE
+		|| Algo::BinarySearch(ExportTypes, FName(Context->GetCacheKey().Bucket.ToString<ANSICHAR>()), FNameFastLess()) != INDEX_NONE;
+}
+
+TArray<FName> FBuildJob::ParseExportBuildTypes(bool& bOutExportAll)
+{
+	TArray<FName> ExportTypes;
+	if (FString ExportTypesArg; FParse::Value(FCommandLine::Get(), TEXT("-ExportBuilds="), ExportTypesArg))
+	{
+		String::ParseTokens(ExportTypesArg, TEXT('+'), [&ExportTypes](FStringView Type) { ExportTypes.Emplace(Type); });
+		ExportTypes.Sort(FNameFastLess());
+		bOutExportAll = false;
+	}
+	else
+	{
+		bOutExportAll = FParse::Param(FCommandLine::Get(), TEXT("ExportBuilds"));
+	}
+	return ExportTypes;
+}
+
 ///////////////////////////////////////////////////////////////////////////////////////////////////////////////////////
 
 TRequest<IBuildJob> CreateBuildJob(ICache& Cache, IBuild& BuildSystem, IBuildInputResolver* InputResolver, const FBuildKey& Key)
