DDC: Removed a few temporary strings from the build job logging

#rnx

diff --git a/Engine/Source/Developer/DerivedDataCache/Private/DerivedDataBuildJob.cpp b/Engine/Source/Developer/DerivedDataCache/Private/DerivedDataBuildJob.cpp
--- a/Engine/Source/Developer/DerivedDataCache/Private/DerivedDataBuildJob.cpp
+++ b/Engine/Source/Developer/DerivedDataCache/Private/DerivedDataBuildJob.cpp
@@ -511,8 +511,8 @@ void FBuildJob::CreateContext()
 	if (BuildSystem.GetVersion() != Action.Get().GetBuildSystemVersion())
 	{
 		CompleteWithError(WriteToString<192>(TEXT("Failed because the build system is version "_SV),
-			BuildSystem.GetVersion().ToString(), TEXT(" when version "_SV),
-			Action.Get().GetBuildSystemVersion().ToString(), TEXT(" is expected."_SV)));
+			BuildSystem.GetVersion(), TEXT(" when version "_SV),
+			Action.Get().GetBuildSystemVersion(), TEXT(" is expected."_SV)));
 	}
 	else if (!Function)
 	{
@@ -527,8 +527,8 @@ void FBuildJob::CreateContext()
 	else if (Function->GetVersion() != Action.Get().GetFunctionVersion())
 	{
 		CompleteWithError(WriteToString<192>(TEXT("Failed because the function "_SV), FunctionName,
-			TEXT(" is version "_SV), Function->GetVersion().ToString(), TEXT(" when version "_SV),
-			Action.Get().GetFunctionVersion().ToString(), TEXT(" is expected."_SV)));
+			TEXT(" is version "_SV), Function->GetVersion(), TEXT(" when version "_SV),
+			Action.Get().GetFunctionVersion(), TEXT(" is expected."_SV)));
 	}
 	else
 	{
