DDC: Exposed the name of the build job to the build function

#rnx

diff --git a/Engine/Source/Developer/DerivedDataCache/Private/DerivedDataBuildJobContext.cpp b/Engine/Source/Developer/DerivedDataCache/Private/DerivedDataBuildJobContext.cpp
--- a/Engine/Source/Developer/DerivedDataCache/Private/DerivedDataBuildJobContext.cpp
+++ b/Engine/Source/Developer/DerivedDataCache/Private/DerivedDataBuildJobContext.cpp
@@ -31,6 +31,11 @@ FBuildJobContext::FBuildJobContext(
 {
 }
 
+FStringView FBuildJobContext::GetName() const
+{
+	return Job.GetName();
+}
+
 void FBuildJobContext::AddConstant(FStringView Key, FCbObject&& Value)
 {
 	Constants.EmplaceByHash(GetTypeHash(Key), Key, MoveTemp(Value));
diff --git a/Engine/Source/Developer/DerivedDataCache/Private/DerivedDataBuildJobContext.h b/Engine/Source/Developer/DerivedDataCache/Private/DerivedDataBuildJobContext.h
--- a/Engine/Source/Developer/DerivedDataCache/Private/DerivedDataBuildJobContext.h
+++ b/Engine/Source/Developer/DerivedDataCache/Private/DerivedDataBuildJobContext.h
@@ -32,6 +32,8 @@ public:
 		EBuildPolicy InBuildPolicy,
 		TUniqueFunction<void ()>&& InOnEndAsyncBuild);
 
+	FStringView GetName() const final;
+
 	inline const FCacheKey& GetCacheKey() const { return CacheKey; }
 	inline const IBuildFunction& GetFunction() const { return Function; }
 
diff --git a/Engine/Source/Developer/DerivedDataCache/Public/DerivedDataBuildFunction.h b/Engine/Source/Developer/DerivedDataCache/Public/DerivedDataBuildFunction.h
--- a/Engine/Source/Developer/DerivedDataCache/Public/DerivedDataBuildFunction.h
+++ b/Engine/Source/Developer/DerivedDataCache/Public/DerivedDataBuildFunction.h
@@ -81,6 +81,9 @@ class FBuildContext
 public:
 	virtual ~FBuildContext() = default;
 
+	/** Returns the name by which to identify this build for logging and profiling. */
+	virtual FStringView GetName() const = 0;
+
 	/** Returns the constant with the matching key, or an object with no fields if not found. */
 	virtual FCbObject FindConstant(FStringView Key) const = 0;
 
@@ -124,6 +127,9 @@ class FBuildConfigContext
 public:
 	virtual ~FBuildConfigContext() = default;
 
+	/** Returns the name by which to identify this build for logging and profiling. */
+	virtual FStringView GetName() const = 0;
+
 	/** Returns the constant with the matching key, or an object with no fields if not found. */
 	virtual FCbObject FindConstant(FStringView Key) const = 0;
 
