DDC: Added FBuildVersionBuilder as a standard way to create build versions

#rnx

diff --git a/Engine/Source/Developer/DerivedDataCache/Public/DerivedDataBuildVersion.h b/Engine/Source/Developer/DerivedDataCache/Public/DerivedDataBuildVersion.h
new file mode 100644
--- /dev/null
+++ b/Engine/Source/Developer/DerivedDataCache/Public/DerivedDataBuildVersion.h
@@ -0,0 +1,86 @@
+// Copyright Epic Games, Inc. All Rights Reserved.
+
+#pragma once
+
+#include "CoreTypes.h"
+#include "Hash/xxhash.h"
+#include "Misc/Guid.h"
+#include "Serialization/Archive.h"
+#include <type_traits>
+
+namespace UE::DerivedData
+{
+
+/**
+ * A type that builds a build version from arbitrary values serialized to it.
+ *
+ * FGuid GetVersion() const final
+ * {
+ *     FBuildVersionBuilder Builder;
+ *     Builder << FGuid(TEXT("a3ae79ff-6a89-4124-afd6-dc095e000488"));
+ *     Builder << ThirdPartyLibraryVersion;
+ *     return Builder.Build();
+ * }
+ */
+class FBuildVersionBuilder final : public FArchive
+{
+public:
+	inline FBuildVersionBuilder();
+
+	inline FGuid Build() const;
+
+	inline FString GetArchiveName() const final { return TEXT("FBuildVersionBuilder"); }
+
+	inline void Serialize(void* Data, int64 Num) final;
+
+	using FArchive::operator<<;
+	FArchive& operator<<(FName& Name) final;
+
+private:
+	FXxHash128Builder HashBuilder;
+};
+
+inline FBuildVersionBuilder::FBuildVersionBuilder()
+{
+	SetIsSaving(true);
+	SetIsPersistent(false);
+
+	// Hash this as the namespace of the Version 3 UUID created in Build().
+	static FGuid BaseVersion(TEXT("b9080e6a-9eb2-4c25-8810-e957009fa0a1"));
+	HashBuilder.Update(&BaseVersion, sizeof(FGuid));
+}
+
+inline FGuid FBuildVersionBuilder::Build() const
+{
+	FXxHash128 Hash = HashBuilder.Finalize();
+	uint32 A = uint32(Hash.HashHigh >> 32);
+	uint32 B = uint32(Hash.HashHigh);
+	uint32 C = uint32(Hash.HashLow >> 32);
+	uint32 D = uint32(Hash.HashLow);
+	// Convert to a Variant 1 Version 3 UUID, which is meant to be created
+	// by hashing the namespace and name with MD5, but this is using XXH3.
+	B = (B & ~0x0000f000) | 0x00003000; // Version 3 (MD5)
+	C = (C & ~0xc0000000) | 0x80000000; // Variant 1 (RFC 4122 UUID)
+	return FGuid(A, B, C, D);
+}
+
+inline void FBuildVersionBuilder::Serialize(void* Data, int64 Num)
+{
+	HashBuilder.Update(Data, uint64(Num));
+}
+
+inline FArchive& FBuildVersionBuilder::operator<<(FName& Name)
+{
+	FString NameString = Name.ToString();
+	return *this << NameString;
+}
+
+template <typename ArgType>
+inline FBuildVersionBuilder& operator<<(FBuildVersionBuilder& Builder, ArgType&& Arg)
+{
+	FArchive& Ar = Builder;
+	Ar << const_cast<std::decay_t<ArgType>&>(Arg);
+	return Builder;
+}
+
+} // UE::DerivedData
