DDC: Fixed errors with NO_LOGGING

#rnx

diff --git a/Engine/Source/Developer/DerivedDataCache/Private/DerivedDataBuildJob.cpp b/Engine/Source/Developer/DerivedDataCache/Private/DerivedDataBuildJob.cpp
--- a/Engine/Source/Developer/DerivedDataCache/Private/DerivedDataBuildJob.cpp
+++ b/Engine/Source/Developer/DerivedDataCache/Private/DerivedDataBuildJob.cpp
@@ -761,7 +761,7 @@ void FBuildJob::EnterResolveInputData()
 	}
 
 	checkf(!MissingInputs.IsEmpty(),
-		TEXT("Job is not expected to be in state %s without missing inputs for build of '%s' by %s"),
+		TEXT("Job is not expected to be in state %s without missing inputs for build of '%s' by %s."),
 		LexToString(State), *Name, *FunctionName);
 
 	if (!InputResolver)
@@ -852,7 +852,7 @@ void FBuildJob::BeginExecuteRemote()
 {
 	if (CanExecuteState(EBuildJobState::ExecuteRemote))
 	{
-		checkf(Worker && WorkerExecutor, TEXT("Job requires a worker in state %s for build of '%s' by %s"),
+		checkf(Worker && WorkerExecutor, TEXT("Job requires a worker in state %s for build of '%s' by %s."),
 			LexToString(State), *Name, *FunctionName);
 		bTriedRemoteExecution = true;
 		AdvanceToState(GetNextState(State), WorkerExecutor->BuildAction(Action.Get(), Inputs, *Worker, BuildPolicy, Priority,
@@ -898,7 +898,7 @@ void FBuildJob::SkipExecuteRemote()
 {
 	checkf(State == EBuildJobState::ResolveRemoteInputData ||
 		State == EBuildJobState::ExecuteRemote || State == EBuildJobState::ExecuteRemoteRetry,
-		TEXT("Job is not expecting SkipExecuteRemote to be called in state %s for build of '%s' by %s"),
+		TEXT("Job is not expecting SkipExecuteRemote to be called in state %s for build of '%s' by %s."),
 		LexToString(State), *Name, *FunctionName);
 	if (CanExecuteState(EBuildJobState::ResolveInputData) || CanExecuteState(EBuildJobState::ExecuteRemote))
 	{
@@ -1005,7 +1005,7 @@ void FBuildJob::SetDefinition(FBuildDefinition&& InDefinition)
 	FunctionName = InDefinition.GetFunction();
 	checkf(Definition.IsNull(), TEXT("Job already has a definition for build of '%s' by %s."), *Name, *FunctionName);
 	checkf(State == EBuildJobState::ResolveKey || State == EBuildJobState::ResolveKeyWait,
-		TEXT("Job is not expecting a definition in state %s for build of '%s' by %s"),
+		TEXT("Job is not expecting a definition in state %s for build of '%s' by %s."),
 		LexToString(State), *Name, *FunctionName);
 	OutputBuilder = BuildSystem.CreateOutput(Name, FunctionName);
 	Definition = MoveTemp(InDefinition);
@@ -1017,7 +1017,7 @@ void FBuildJob::SetAction(FBuildAction&& InAction)
 	checkf(Action.IsNull(), TEXT("Job already has an action for build of '%s' by %s."), *Name, *FunctionName);
 	checkf(State == EBuildJobState::ResolveKey ||
 		State == EBuildJobState::ResolveInputMeta || State == EBuildJobState::ResolveInputMetaWait,
-		TEXT("Job is not expecting an action in state %s for build of '%s' by %s"),
+		TEXT("Job is not expecting an action in state %s for build of '%s' by %s."),
 		LexToString(State), *Name, *FunctionName);
 	Action = MoveTemp(InAction);
 	return AdvanceToState(EBuildJobState::CacheQuery);
@@ -1037,7 +1037,7 @@ void FBuildJob::SetInputs(FBuildInputs&& InInputs)
 		ExecuteState = EBuildJobState::ExecuteLocal;
 		break;
 	default:
-		checkf(false, TEXT("Job is not expecting inputs in state %s for build of '%s' by %s"),
+		checkf(false, TEXT("Job is not expecting inputs in state %s for build of '%s' by %s."),
 			LexToString(State), *Name, *FunctionName);
 		return;
 	}
@@ -1053,7 +1053,7 @@ void FBuildJob::SetOutput(const FBuildOutput& InOutput)
 		State == EBuildJobState::CacheQuery || State == EBuildJobState::ResolveInputData ||
 		State == EBuildJobState::ExecuteRemote || State == EBuildJobState::ExecuteRemoteRetry ||
 		State == EBuildJobState::ExecuteLocal,
-		TEXT("Job is not expecting an output in state %s for build of '%s' by %s"),
+		TEXT("Job is not expecting an output in state %s for build of '%s' by %s."),
 		LexToString(State), *Name, *FunctionName);
 	SetOutputNoCheck(FBuildOutput(InOutput));
 }
@@ -1086,9 +1086,7 @@ void FBuildJob::CompleteWithError(FStringView Error)
 	{
 		return;
 	}
-	TStringBuilder<32> Category;
-	Category << ImplicitConv<FName>(LogDerivedDataBuild.GetCategoryName());
-	OutputBuilder.AddError(Category, Error);
+	OutputBuilder.AddError(TEXT("LogDerivedDataBuild"_SV), Error);
 	SetOutputNoCheck(OutputBuilder.Build());
 	return AdvanceToState(EBuildJobState::Complete);
 }
diff --git a/Engine/Source/Developer/DerivedDataCache/Private/DerivedDataBuildJobContext.cpp b/Engine/Source/Developer/DerivedDataCache/Private/DerivedDataBuildJobContext.cpp
--- a/Engine/Source/Developer/DerivedDataCache/Private/DerivedDataBuildJobContext.cpp
+++ b/Engine/Source/Developer/DerivedDataCache/Private/DerivedDataBuildJobContext.cpp
@@ -73,14 +73,12 @@ FSharedBuffer FBuildJobContext::FindInput(FStringView Key) const
 		}
 		else
 		{
-			TStringBuilder<32> Category;
-			Category << ImplicitConv<FName>(LogDerivedDataBuild.GetCategoryName());
 			TStringBuilder<256> Error;
 			Error << TEXT("Input '") << Key << TEXT("' was expected to have raw hash ") << Input->GetRawHash()
 				<< TEXT(" and raw size ") << Input->GetRawSize() << TEXT(" but has raw hash ") << RawHash
 				<< TEXT(" and raw size ") << Buffer.GetSize() << TEXT(" after decompression for build of '")
 				<< Job.GetName() << TEXT("' by ") << Job.GetFunction() << TEXT(".");
-			OutputBuilder.AddError(Category, Error);
+			OutputBuilder.AddError(TEXT("LogDerivedDataBuild"_SV), Error);
 			UE_LOG(LogDerivedDataBuild, Error, TEXT("%.*s"), Error.Len(), Error.GetData());
 		}
 	}
