DDC: Added cpu trace scopes to the build job

#rnx

diff --git a/Engine/Source/Developer/DerivedDataCache/Private/DerivedDataBuildJob.cpp b/Engine/Source/Developer/DerivedDataCache/Private/DerivedDataBuildJob.cpp
--- a/Engine/Source/Developer/DerivedDataCache/Private/DerivedDataBuildJob.cpp
+++ b/Engine/Source/Developer/DerivedDataCache/Private/DerivedDataBuildJob.cpp
@@ -34,6 +34,7 @@
 #include "Misc/PathViews.h"
 #include "Misc/ScopeRWLock.h"
 #include "Misc/StringBuilder.h"
+#include "ProfilingDebugging/CpuProfilerTrace.h"
 #include "Serialization/CompactBinary.h"
 #include "Serialization/CompactBinaryWriter.h"
 #include "String/ParseTokens.h"
@@ -457,6 +458,7 @@ void FBuildJob::Wait()
 		}
 	}
 
+	TRACE_CPUPROFILER_EVENT_SCOPE(FBuildJob::Wait);
 	LocalEvent->Wait();
 }
 
@@ -507,6 +509,8 @@ void FBuildJob::EndJob()
 
 void FBuildJob::CreateContext()
 {
+	TRACE_CPUPROFILER_EVENT_SCOPE(FBuildJob::CreateContext);
+
 	const IBuildFunction* Function = BuildSystem.GetFunctionRegistry().FindFunction(FunctionName);
 	if (BuildSystem.GetVersion() != Action.Get().GetBuildSystemVersion())
 	{
@@ -578,6 +582,7 @@ void FBuildJob::BeginCacheQuery()
 {
 	if (CanExecuteState(EBuildJobState::CacheQuery))
 	{
+		TRACE_CPUPROFILER_EVENT_SCOPE(FBuildJob::CacheQuery);
 		ECachePolicy CachePolicy = Context->GetCachePolicy();
 		if (EnumHasAnyFlags(BuildPolicy, EBuildPolicy::SkipData))
 		{
@@ -624,6 +629,7 @@ void FBuildJob::BeginCacheStore()
 {
 	if (CanExecuteState(EBuildJobState::CacheStore))
 	{
+		TRACE_CPUPROFILER_EVENT_SCOPE(FBuildJob::CacheStore);
 		FCacheRecordBuilder RecordBuilder = Cache.CreateRecord(Context->GetCacheKey());
 		Output.Get().Save(RecordBuilder);
 		return AdvanceToState(EBuildJobState::CacheStoreWait,
@@ -668,6 +674,7 @@ void FBuildJob::BeginResolveKey()
 {
 	if (CanExecuteState(EBuildJobState::ResolveKey))
 	{
+		TRACE_CPUPROFILER_EVENT_SCOPE(FBuildJob::ResolveKey);
 		return AdvanceToState(EBuildJobState::ResolveKeyWait, InputResolver->ResolveKey(DefinitionKey,
 			[this](FBuildKeyResolvedParams&& Params) { EndResolveKey(MoveTemp(Params)); }));
 	}
@@ -707,6 +714,7 @@ void FBuildJob::BeginResolveInputMeta()
 {
 	if (CanExecuteState(EBuildJobState::ResolveInputMeta))
 	{
+		TRACE_CPUPROFILER_EVENT_SCOPE(FBuildJob::ResolveInputMeta);
 		return AdvanceToState(EBuildJobState::ResolveInputMetaWait, InputResolver->ResolveInputMeta(Definition.Get(),
 			Priority, [this](FBuildInputMetaResolvedParams&& Params) { EndResolveInputMeta(MoveTemp(Params)); }));
 	}
@@ -775,6 +783,7 @@ void FBuildJob::BeginResolveInputData()
 {
 	if (CanExecuteState(EBuildJobState::ResolveInputData))
 	{
+		TRACE_CPUPROFILER_EVENT_SCOPE(FBuildJob::ResolveInputData);
 		if (Definition)
 		{
 			return AdvanceToState(GetNextState(State), InputResolver->ResolveInputData(Definition.Get(), Priority,
@@ -852,6 +861,7 @@ void FBuildJob::BeginExecuteRemote()
 {
 	if (CanExecuteState(EBuildJobState::ExecuteRemote))
 	{
+		TRACE_CPUPROFILER_EVENT_SCOPE(FBuildJob::ExecuteRemote);
 		checkf(Worker && WorkerExecutor, TEXT("Job requires a worker in state %s for build of '%s' by %s."),
 			LexToString(State), *Name, *FunctionName);
 		bTriedRemoteExecution = true;
@@ -960,6 +970,7 @@ void FBuildJob::BeginExecuteLocal()
 {
 	if (CanExecuteState(EBuildJobState::ExecuteLocal))
 	{
+		TRACE_CPUPROFILER_EVENT_SCOPE(FBuildJob::ExecuteLocal);
 		Context->GetFunction().Build(*Context);
 		if (Context->IsAsyncBuild())
 		{
