DDC: Added FOptionalBuildSession

#rnx

diff --git a/Engine/Source/Developer/DerivedDataCache/Public/DerivedDataBuildSession.h b/Engine/Source/Developer/DerivedDataCache/Public/DerivedDataBuildSession.h
--- a/Engine/Source/Developer/DerivedDataCache/Public/DerivedDataBuildSession.h
+++ b/Engine/Source/Developer/DerivedDataCache/Public/DerivedDataBuildSession.h
@@ -138,6 +138,7 @@ public:
 	}
 
 private:
+	friend class FOptionalBuildSession;
 	friend FBuildSession Private::CreateBuildSession(Private::IBuildSessionInternal* Session);
 
 	/** Construct a build session. Use IBuild::CreateSession(). */
@@ -149,6 +150,33 @@ private:
 	TUniquePtr<Private::IBuildSessionInternal> Session;
 };
 
+/**
+ * A build session that can be null.
+ *
+ * @see FBuildSession
+ */
+class FOptionalBuildSession : private FBuildSession
+{
+public:
+	inline FOptionalBuildSession() : FBuildSession(nullptr) {}
+
+	inline FOptionalBuildSession(FBuildSession&& InSession) : FBuildSession(MoveTemp(InSession)) {}
+	inline FOptionalBuildSession& operator=(FBuildSession&& InSession) { FBuildSession::operator=(MoveTemp(InSession)); return *this; }
+
+	inline FOptionalBuildSession(const FBuildSession& InSession) = delete;
+	inline FOptionalBuildSession& operator=(const FBuildSession& InSession) = delete;
+
+	/** Returns the build session. The caller must check for null before using this accessor. */
+	inline FBuildSession& Get() & { return *this; }
+	inline FBuildSession&& Get() && { return MoveTemp(*this); }
+
+	inline bool IsNull() const { return !IsValid(); }
+	inline bool IsValid() const { return Session.IsValid(); }
+	inline explicit operator bool() const { return IsValid(); }
+
+	inline void Reset() { *this = FOptionalBuildSession(); }
+};
+
 /** Parameters for the completion callback for build requests. */
 struct FBuildCompleteParams
 {
