DDC: Fixed race conditions in the build job

#rnx

diff --git a/Engine/Source/Developer/DerivedDataCache/Private/DerivedDataBuildJob.cpp b/Engine/Source/Developer/DerivedDataCache/Private/DerivedDataBuildJob.cpp
--- a/Engine/Source/Developer/DerivedDataCache/Private/DerivedDataBuildJob.cpp
+++ b/Engine/Source/Developer/DerivedDataCache/Private/DerivedDataBuildJob.cpp
@@ -32,6 +32,7 @@
 #include "Misc/Parse.h"
 #include "Misc/Paths.h"
 #include "Misc/PathViews.h"
+#include "Misc/ScopeExit.h"
 #include "Misc/ScopeRWLock.h"
 #include "Misc/StringBuilder.h"
 #include "ProfilingDebugging/CpuProfilerTrace.h"
@@ -92,12 +93,6 @@ enum class EBuildJobState : uint32
 	Complete,
 };
 
-/** Returns the state after the given state. */
-static constexpr EBuildJobState GetNextState(EBuildJobState State)
-{
-	return State == EBuildJobState::Complete ? State : EBuildJobState(uint32(State) + 1);
-}
-
 class FBuildJob final : public IBuildJob
 {
 public:
@@ -121,7 +116,9 @@ public:
 	inline ICache& GetCache() const final { return Cache; }
 	inline IBuild& GetBuild() const final { return BuildSystem; }
 
-	void Schedule(IBuildScheduler& Scheduler, EBuildPolicy Policy, EPriority Priority, FOnBuildJobComplete&& OnComplete) final;
+	void Start(IBuildScheduler& Scheduler, EBuildPolicy Policy, EPriority Priority, FOnBuildJobComplete&& OnComplete) final;
+
+	void Schedule() final;
 
 	// IRequest Interface
 
@@ -146,13 +143,13 @@ private:
 	void EnterExecuteRemote();
 	void EnterExecuteLocal();
 
-	void BeginCacheQuery() final;
-	void BeginCacheStore() final;
-	void BeginResolveKey() final;
-	void BeginResolveInputMeta() final;
-	void BeginResolveInputData() final;
-	void BeginExecuteRemote() final;
-	void BeginExecuteLocal() final;
+	void BeginCacheQuery();
+	void BeginCacheStore();
+	void BeginResolveKey();
+	void BeginResolveInputMeta();
+	void BeginResolveInputData();
+	void BeginExecuteRemote();
+	void BeginExecuteLocal();
 
 	void EndCacheQuery(FCacheGetCompleteParams&& Params);
 	void EndCacheStore(FCachePutCompleteParams&& Params);
@@ -175,15 +172,18 @@ private:
 	/** Terminate the job and send the error to the output complete callback. */
 	void CompleteWithError(FStringView Error);
 
-	/** Process cancellation requests and return whether the requested state is valid to execute. */
-	bool CanExecuteState(EBuildJobState RequestedState);
+	/** Start execution of an async operation that is managed by a request. */
+	void ExecuteAsync(TFunctionRef<FRequest ()> Operation);
 
 	/** Advance to the new state, dispatching to the scheduler and invoking callbacks as appropriate. */
-	void AdvanceToState(EBuildJobState NewState, FRequest NewRequest = FRequest());
+	void AdvanceToState(EBuildJobState NewState);
 
 	/** Execute a transition from the old state to the new state. */
 	void ExecuteTransition(EBuildJobState OldState, EBuildJobState NewState);
 
+	/** Execute the new state. */
+	void ExecuteState(EBuildJobState NewState);
+
 	/** Exports the action and inputs for this build to disk. */
 	void ExportBuild() const;
 
@@ -350,7 +350,7 @@ FBuildJob::~FBuildJob()
 	FPlatformProcess::ReturnSynchEventToPool(Event);
 }
 
-void FBuildJob::Schedule(
+void FBuildJob::Start(
 	IBuildScheduler& InScheduler,
 	EBuildPolicy InPolicy,
 	EPriority InPriority,
@@ -366,6 +366,26 @@ void FBuildJob::Schedule(
 	return AdvanceToState(EBuildJobState::ResolveKey);
 }
 
+void FBuildJob::Schedule()
+{
+	switch (State)
+	{
+	case EBuildJobState::ResolveKey:             return AdvanceToState(EBuildJobState::ResolveKeyWait);
+	case EBuildJobState::ResolveInputMeta:       return AdvanceToState(EBuildJobState::ResolveInputMetaWait);
+	case EBuildJobState::CacheQuery:             return AdvanceToState(EBuildJobState::CacheQueryWait);
+	case EBuildJobState::ExecuteRemote:          return AdvanceToState(EBuildJobState::ExecuteRemoteWait);
+	case EBuildJobState::ResolveRemoteInputData: return AdvanceToState(EBuildJobState::ResolveRemoteInputDataWait);
+	case EBuildJobState::ExecuteRemoteRetry:     return AdvanceToState(EBuildJobState::ExecuteRemoteRetryWait);
+	case EBuildJobState::ResolveInputData:       return AdvanceToState(EBuildJobState::ResolveInputDataWait);
+	case EBuildJobState::ExecuteLocal:           return AdvanceToState(EBuildJobState::ExecuteLocalWait);
+	case EBuildJobState::CacheStore:             return AdvanceToState(EBuildJobState::CacheStoreWait);
+	default:
+		checkf(false,
+			TEXT("Job in state %s is not valid to be scheduled for build of '%s' by %s."),
+			LexToString(State), *Name, *FunctionName);
+	}
+}
+
 ///////////////////////////////////////////////////////////////////////////////////////////////////////////////////////
 
 void FBuildJob::SetPriority(EPriority InPriority)
@@ -403,7 +423,7 @@ void FBuildJob::Cancel()
 		}
 	}
 
-	// Cancel the job on the scheduler, which invokes Begin[State] if the job was queued.
+	// Cancel the job on the scheduler, which invokes Schedule if the job was queued.
 	if (LocalScheduler)
 	{
 		LocalScheduler->CancelJob(this);
@@ -423,7 +443,7 @@ void FBuildJob::Cancel()
 
 	// Most jobs will be complete at this point, but it is possible for a job to reach this point
 	// because of a race with the scheduler. A job can be removed from every scheduler queue, and
-	// missed by CancelJob, but not yet have called into Begin[State] or Set[Type].
+	// missed by CancelJob, but not yet have called back into the job.
 	Wait();
 }
 
@@ -491,14 +511,13 @@ void FBuildJob::BeginJob()
 
 void FBuildJob::EndJob()
 {
-	TRequest Self(this);
+	ON_SCOPE_EXIT { Release(); };
 
 	if (Event)
 	{
 		Event->Trigger();
 	}
 	Scheduler->EndJob(this);
-	Release();
 
 	FWriteScopeLock WriteLock(Lock);
 	Scheduler = nullptr;
@@ -575,39 +594,34 @@ void FBuildJob::EnterCacheQuery()
 	{
 		return AdvanceToState(EBuildJobState::ExecuteRemote);
 	}
-	Scheduler->DispatchCacheQuery(this, SchedulerParams);
 }
 
 void FBuildJob::BeginCacheQuery()
 {
-	if (CanExecuteState(EBuildJobState::CacheQuery))
+	TRACE_CPUPROFILER_EVENT_SCOPE(FBuildJob::CacheQuery);
+	ExecuteAsync([this]
 	{
-		TRACE_CPUPROFILER_EVENT_SCOPE(FBuildJob::CacheQuery);
 		ECachePolicy CachePolicy = Context->GetCachePolicy();
 		if (EnumHasAnyFlags(BuildPolicy, EBuildPolicy::SkipData))
 		{
 			CachePolicy |= ECachePolicy::SkipAttachments;
 		}
-		return AdvanceToState(EBuildJobState::CacheQueryWait,
-			Cache.Get({Context->GetCacheKey()}, Name, CachePolicy, Priority,
-				[this](FCacheGetCompleteParams&& Params) { EndCacheQuery(MoveTemp(Params)); }));
-	}
+		return Cache.Get({Context->GetCacheKey()}, Name, CachePolicy, Priority,
+			[this](FCacheGetCompleteParams&& Params) { EndCacheQuery(MoveTemp(Params)); });
+	});
 }
 
 void FBuildJob::EndCacheQuery(FCacheGetCompleteParams&& Params)
 {
-	if (CanExecuteState(EBuildJobState::CacheQueryWait))
+	if (Params.Status == EStatus::Ok)
 	{
-		if (Params.Status == EStatus::Ok)
+		if (FOptionalBuildOutput CacheOutput = BuildSystem.LoadOutput(Name, FunctionName, Params.Record))
 		{
-			if (FOptionalBuildOutput CacheOutput = BuildSystem.LoadOutput(Name, FunctionName, Params.Record))
-			{
-				bIsCached = true;
-				SetOutputNoCheck(MoveTemp(CacheOutput).Get());
-			}
+			bIsCached = true;
+			return SetOutputNoCheck(MoveTemp(CacheOutput).Get());
 		}
-		return AdvanceToState(EBuildJobState::ExecuteRemote);
 	}
+	return AdvanceToState(EBuildJobState::ExecuteRemote);
 }
 
 ///////////////////////////////////////////////////////////////////////////////////////////////////////////////////////
@@ -622,29 +636,25 @@ void FBuildJob::EnterCacheStore()
 	{
 		return AdvanceToState(EBuildJobState::Complete);
 	}
-	Scheduler->DispatchCacheStore(this, SchedulerParams);
 }
 
 void FBuildJob::BeginCacheStore()
 {
-	if (CanExecuteState(EBuildJobState::CacheStore))
+	TRACE_CPUPROFILER_EVENT_SCOPE(FBuildJob::CacheStore);
+	ExecuteAsync([this]
 	{
-		TRACE_CPUPROFILER_EVENT_SCOPE(FBuildJob::CacheStore);
 		FCacheRecordBuilder RecordBuilder = Cache.CreateRecord(Context->GetCacheKey());
 		Output.Get().Save(RecordBuilder);
-		return AdvanceToState(EBuildJobState::CacheStoreWait,
-			Cache.Put({RecordBuilder.Build()}, Name, Context->GetCachePolicy(), FMath::Min(Priority, EPriority::Highest),
-				[this](FCachePutCompleteParams&& Params) { EndCacheStore(MoveTemp(Params)); }));
-	}
+		return Cache.Put({RecordBuilder.Build()}, Name, Context->GetCachePolicy(),
+			FMath::Min(Priority, EPriority::Highest),
+			[this](FCachePutCompleteParams&& Params) { EndCacheStore(MoveTemp(Params)); });
+	});
 }
 
 void FBuildJob::EndCacheStore(FCachePutCompleteParams&& Params)
 {
-	if (CanExecuteState(EBuildJobState::CacheStoreWait))
-	{
-		bIsCached = Params.Status == EStatus::Ok;
-		return AdvanceToState(EBuildJobState::Complete);
-	}
+	bIsCached = Params.Status == EStatus::Ok;
+	return AdvanceToState(EBuildJobState::Complete);
 }
 
 ///////////////////////////////////////////////////////////////////////////////////////////////////////////////////////
@@ -667,31 +677,27 @@ void FBuildJob::EnterResolveKey()
 	{
 		return CompleteWithError(TEXT("Failed to resolve key due to null input resolver."_SV));
 	}
-	Scheduler->DispatchResolveKey(this);
 }
 
 void FBuildJob::BeginResolveKey()
 {
-	if (CanExecuteState(EBuildJobState::ResolveKey))
+	TRACE_CPUPROFILER_EVENT_SCOPE(FBuildJob::ResolveKey);
+	ExecuteAsync([this]
 	{
-		TRACE_CPUPROFILER_EVENT_SCOPE(FBuildJob::ResolveKey);
-		return AdvanceToState(EBuildJobState::ResolveKeyWait, InputResolver->ResolveKey(DefinitionKey,
-			[this](FBuildKeyResolvedParams&& Params) { EndResolveKey(MoveTemp(Params)); }));
-	}
+		return InputResolver->ResolveKey(DefinitionKey,
+			[this](FBuildKeyResolvedParams&& Params) { EndResolveKey(MoveTemp(Params)); });
+	});
 }
 
 void FBuildJob::EndResolveKey(FBuildKeyResolvedParams&& Params)
 {
-	if (CanExecuteState(EBuildJobState::ResolveKeyWait))
+	if (Params.Status == EStatus::Ok && Params.Definition)
 	{
-		if (Params.Status == EStatus::Ok && Params.Definition)
-		{
-			SetDefinition(MoveTemp(Params.Definition).Get());
-		}
-		else
-		{
-			CompleteWithError(WriteToString<128>(TEXT("Failed to resolve key "_SV), Params.Key, TEXT("."_SV)));
-		}
+		SetDefinition(MoveTemp(Params.Definition).Get());
+	}
+	else
+	{
+		CompleteWithError(WriteToString<128>(TEXT("Failed to resolve key "_SV), Params.Key, TEXT("."_SV)));
 	}
 }
 
@@ -707,31 +713,27 @@ void FBuildJob::EnterResolveInputMeta()
 	{
 		return CompleteWithError(TEXT("Failed to resolve input metadata due to null input resolver."_SV));
 	}
-	Scheduler->DispatchResolveInputMeta(this);
 }
 
 void FBuildJob::BeginResolveInputMeta()
 {
-	if (CanExecuteState(EBuildJobState::ResolveInputMeta))
+	TRACE_CPUPROFILER_EVENT_SCOPE(FBuildJob::ResolveInputMeta);
+	ExecuteAsync([this]
 	{
-		TRACE_CPUPROFILER_EVENT_SCOPE(FBuildJob::ResolveInputMeta);
-		return AdvanceToState(EBuildJobState::ResolveInputMetaWait, InputResolver->ResolveInputMeta(Definition.Get(),
-			Priority, [this](FBuildInputMetaResolvedParams&& Params) { EndResolveInputMeta(MoveTemp(Params)); }));
-	}
+		return InputResolver->ResolveInputMeta(Definition.Get(), Priority,
+			[this](FBuildInputMetaResolvedParams&& Params) { EndResolveInputMeta(MoveTemp(Params)); });
+	});
 }
 
 void FBuildJob::EndResolveInputMeta(FBuildInputMetaResolvedParams&& Params)
 {
-	if (CanExecuteState(EBuildJobState::ResolveInputMetaWait))
+	if (Params.Status == EStatus::Ok)
 	{
-		if (Params.Status == EStatus::Ok)
-		{
-			CreateAction(Params.Inputs);
-		}
-		else
-		{
-			CompleteWithError(TEXT("Failed to resolve input metadata."_SV));
-		}
+		CreateAction(Params.Inputs);
+	}
+	else
+	{
+		CompleteWithError(TEXT("Failed to resolve input metadata."_SV));
 	}
 }
 
@@ -776,54 +778,50 @@ void FBuildJob::EnterResolveInputData()
 	{
 		return CompleteWithError(TEXT("Failed to resolve input data due to null input resolver."_SV));
 	}
-	Scheduler->DispatchResolveInputData(this, SchedulerParams);
 }
 
 void FBuildJob::BeginResolveInputData()
 {
-	if (CanExecuteState(EBuildJobState::ResolveInputData))
+	TRACE_CPUPROFILER_EVENT_SCOPE(FBuildJob::ResolveInputData);
+	ExecuteAsync([this]
 	{
-		TRACE_CPUPROFILER_EVENT_SCOPE(FBuildJob::ResolveInputData);
 		if (Definition)
 		{
-			return AdvanceToState(GetNextState(State), InputResolver->ResolveInputData(Definition.Get(), Priority,
+			return InputResolver->ResolveInputData(Definition.Get(), Priority,
 				[this](FBuildInputDataResolvedParams&& Params) { EndResolveInputData(MoveTemp(Params)); },
-				[this](FStringView Key) { return !!Algo::Find(MissingInputs, Key); }));
+				[this](FStringView Key) { return !!Algo::Find(MissingInputs, Key); });
 		}
 		else
 		{
-			return AdvanceToState(GetNextState(State), InputResolver->ResolveInputData(Action.Get(), Priority,
+			return InputResolver->ResolveInputData(Action.Get(), Priority,
 				[this](FBuildInputDataResolvedParams&& Params) { EndResolveInputData(MoveTemp(Params)); },
-				[this](FStringView Key) { return !!Algo::Find(MissingInputs, Key); }));
+				[this](FStringView Key) { return !!Algo::Find(MissingInputs, Key); });
 		}
-	}
+	});
 }
 
 void FBuildJob::EndResolveInputData(FBuildInputDataResolvedParams&& Params)
 {
-	if (CanExecuteState(EBuildJobState::ResolveInputDataWait))
+	if (Params.Status == EStatus::Ok)
 	{
-		if (Params.Status == EStatus::Ok)
+		FBuildInputsBuilder Builder = BuildSystem.CreateInputs(Name);
+		for (const FBuildInputDataByKey& Input : Params.Inputs)
 		{
-			FBuildInputsBuilder Builder = BuildSystem.CreateInputs(Name);
-			for (const FBuildInputDataByKey& Input : Params.Inputs)
-			{
-				Builder.AddInput(Input.Key, Input.Data);
-				SchedulerParams.ResolvedInputsSize += Input.Data.GetRawSize();
-			}
-			if (Inputs)
-			{
-				Inputs.Get().IterateInputs([&Builder](FStringView Key, const FCompressedBuffer& Buffer)
-				{
-					Builder.AddInput(Key, Buffer);
-				});
-			}
-			SetInputs(Builder.Build());
+			Builder.AddInput(Input.Key, Input.Data);
+			SchedulerParams.ResolvedInputsSize += Input.Data.GetRawSize();
 		}
-		else
+		if (Inputs)
 		{
-			CompleteWithError(TEXT("Failed to resolve input data."_SV));
+			Inputs.Get().IterateInputs([&Builder](FStringView Key, const FCompressedBuffer& Buffer)
+			{
+				Builder.AddInput(Key, Buffer);
+			});
 		}
+		SetInputs(Builder.Build());
+	}
+	else
+	{
+		CompleteWithError(TEXT("Failed to resolve input data."_SV));
 	}
 }
 
@@ -853,53 +851,48 @@ void FBuildJob::EnterExecuteRemote()
 	{
 		return AdvanceToState(EBuildJobState::ResolveInputData);
 	}
-
-	Scheduler->DispatchExecuteRemote(this, SchedulerParams);
 }
 
 void FBuildJob::BeginExecuteRemote()
 {
-	if (CanExecuteState(EBuildJobState::ExecuteRemote))
+	TRACE_CPUPROFILER_EVENT_SCOPE(FBuildJob::ExecuteRemote);
+	checkf(Worker && WorkerExecutor, TEXT("Job requires a worker in state %s for build of '%s' by %s."),
+		LexToString(State), *Name, *FunctionName);
+	bTriedRemoteExecution = true;
+	ExecuteAsync([this]
 	{
-		TRACE_CPUPROFILER_EVENT_SCOPE(FBuildJob::ExecuteRemote);
-		checkf(Worker && WorkerExecutor, TEXT("Job requires a worker in state %s for build of '%s' by %s."),
-			LexToString(State), *Name, *FunctionName);
-		bTriedRemoteExecution = true;
-		AdvanceToState(GetNextState(State), WorkerExecutor->BuildAction(Action.Get(), Inputs, *Worker, BuildPolicy, Priority,
-			[this](FBuildWorkerActionCompleteParams&& Params) { EndExecuteRemote(MoveTemp(Params)); }));
-	}
+		return WorkerExecutor->BuildAction(Action.Get(), Inputs, *Worker, BuildPolicy, Priority,
+			[this](FBuildWorkerActionCompleteParams&& Params) { EndExecuteRemote(MoveTemp(Params)); });
+	});
 }
 
 void FBuildJob::EndExecuteRemote(FBuildWorkerActionCompleteParams&& Params)
 {
-	if (CanExecuteState(EBuildJobState::ExecuteRemoteWait))
+	if (Params.Output)
 	{
-		if (Params.Output)
-		{
-			SetOutputNoCheck(MoveTemp(Params.Output).Get());
-		}
-		else
+		return SetOutputNoCheck(MoveTemp(Params.Output).Get());
+	}
+	else
+	{
+		switch (State)
 		{
-			switch (State)
+		case EBuildJobState::ExecuteRemote:
+		case EBuildJobState::ExecuteRemoteWait:
+			MissingInputs.Reset(Params.MissingInputs.Num());
+			for (const FStringView& Key : Params.MissingInputs)
 			{
-			case EBuildJobState::ExecuteRemote:
-			case EBuildJobState::ExecuteRemoteWait:
-				MissingInputs.Reset(Params.MissingInputs.Num());
-				for (const FStringView& Key : Params.MissingInputs)
-				{
-					MissingInputs.Emplace(Key);
-				}
-				if (!MissingInputs.IsEmpty())
-				{
-					MissingInputs.Sort();
-					return AdvanceToState(EBuildJobState::ResolveRemoteInputData);
-				}
-				[[fallthrough]];
-			case EBuildJobState::ExecuteRemoteRetry:
-			case EBuildJobState::ExecuteRemoteRetryWait:
-			default:
-				return AdvanceToState(EBuildJobState::ResolveInputData);
+				MissingInputs.Emplace(Key);
 			}
+			if (!MissingInputs.IsEmpty())
+			{
+				MissingInputs.Sort();
+				return AdvanceToState(EBuildJobState::ResolveRemoteInputData);
+			}
+			[[fallthrough]];
+		case EBuildJobState::ExecuteRemoteRetry:
+		case EBuildJobState::ExecuteRemoteRetryWait:
+		default:
+			return AdvanceToState(EBuildJobState::ResolveInputData);
 		}
 	}
 }
@@ -910,10 +903,7 @@ void FBuildJob::SkipExecuteRemote()
 		State == EBuildJobState::ExecuteRemote || State == EBuildJobState::ExecuteRemoteRetry,
 		TEXT("Job is not expecting SkipExecuteRemote to be called in state %s for build of '%s' by %s."),
 		LexToString(State), *Name, *FunctionName);
-	if (CanExecuteState(EBuildJobState::ResolveInputData) || CanExecuteState(EBuildJobState::ExecuteRemote))
-	{
-		return AdvanceToState(EBuildJobState::ResolveInputData);
-	}
+	return AdvanceToState(EBuildJobState::ResolveInputData);
 }
 
 ///////////////////////////////////////////////////////////////////////////////////////////////////////////////////////
@@ -941,8 +931,7 @@ void FBuildJob::EnterExecuteLocal()
 	{
 		Context->AddConstant(Key, MoveTemp(Value));
 	});
-	bool bMatchedInputs = true;
-	Action.Get().IterateInputs([this, &bMatchedInputs](FStringView Key, const FIoHash& RawHash, uint64 RawSize)
+	Action.Get().IterateInputs([this](FStringView Key, const FIoHash& RawHash, uint64 RawSize)
 	{
 		const FCompressedBuffer& Buffer = Inputs.Get().FindInput(Key);
 		if (Buffer.GetRawHash() == RawHash && Buffer.GetRawSize() == RawSize)
@@ -951,7 +940,6 @@ void FBuildJob::EnterExecuteLocal()
 		}
 		else
 		{
-			bMatchedInputs = false;
 			CompleteWithError(WriteToString<256>(TEXT("Failed because input '"_SV), Key, TEXT("' with hash "_SV),
 				RawHash, TEXT(" ("_SV), RawSize, TEXT(" bytes) was resolved with hash "_SV), Buffer.GetRawHash(),
 				TEXT(" ("_SV), Buffer.GetRawSize(), TEXT(" bytes)."_SV)));
@@ -960,35 +948,25 @@ void FBuildJob::EnterExecuteLocal()
 	Action.Reset();
 	Inputs.Reset();
 	SchedulerParams.ResolvedInputsSize = 0;
-	if (bMatchedInputs)
-	{
-		Scheduler->DispatchExecuteLocal(this, SchedulerParams);
-	}
 }
 
 void FBuildJob::BeginExecuteLocal()
 {
-	if (CanExecuteState(EBuildJobState::ExecuteLocal))
+	TRACE_CPUPROFILER_EVENT_SCOPE(FBuildJob::ExecuteLocal);
+	ExecuteAsync([this]
 	{
-		TRACE_CPUPROFILER_EVENT_SCOPE(FBuildJob::ExecuteLocal);
 		Context->GetFunction().Build(*Context);
-		if (Context->IsAsyncBuild())
-		{
-			return AdvanceToState(EBuildJobState::ExecuteLocalWait);
-		}
-		else
+		if (!Context->IsAsyncBuild())
 		{
-			return SetOutputNoCheck(OutputBuilder.Build());
+			EndExecuteLocal();
 		}
-	}
+		return FRequest();
+	});
 }
 
 void FBuildJob::EndExecuteLocal()
 {
-	if (CanExecuteState(EBuildJobState::ExecuteLocalWait))
-	{
-		SetOutputNoCheck(OutputBuilder.Build());
-	}
+	return SetOutputNoCheck(OutputBuilder.Build());
 }
 
 ///////////////////////////////////////////////////////////////////////////////////////////////////////////////////////
@@ -1066,7 +1044,7 @@ void FBuildJob::SetOutput(const FBuildOutput& InOutput)
 		State == EBuildJobState::ExecuteLocal,
 		TEXT("Job is not expecting an output in state %s for build of '%s' by %s."),
 		LexToString(State), *Name, *FunctionName);
-	SetOutputNoCheck(FBuildOutput(InOutput));
+	return SetOutputNoCheck(FBuildOutput(InOutput));
 }
 
 void FBuildJob::SetOutputNoCheck(FBuildOutput&& InOutput)
@@ -1098,72 +1076,55 @@ void FBuildJob::CompleteWithError(FStringView Error)
 		return;
 	}
 	OutputBuilder.AddError(TEXT("LogDerivedDataBuild"_SV), Error);
-	SetOutputNoCheck(OutputBuilder.Build());
-	return AdvanceToState(EBuildJobState::Complete);
+	return SetOutputNoCheck(OutputBuilder.Build());
 }
 
-bool FBuildJob::CanExecuteState(EBuildJobState RequestedState)
+void FBuildJob::ExecuteAsync(TFunctionRef<FRequest ()> Operation)
 {
-	if (bIsCanceled)
-	{
-		AdvanceToState(EBuildJobState::Complete);
-		return false;
-	}
-
-	EBuildJobState LocalState = State;
-
-	// Handle states that execute multiple times under different names.
-	switch (LocalState)
-	{
-	case EBuildJobState::ResolveRemoteInputData:     LocalState = EBuildJobState::ResolveInputData; break;
-	case EBuildJobState::ResolveRemoteInputDataWait: LocalState = EBuildJobState::ResolveInputDataWait; break;
-	case EBuildJobState::ExecuteRemoteRetry:         LocalState = EBuildJobState::ExecuteRemote; break;
-	case EBuildJobState::ExecuteRemoteRetryWait:     LocalState = EBuildJobState::ExecuteRemoteWait; break;
-	}
-
-	if (LocalState == RequestedState)
-	{
-		return true;
-	}
-
-	// Handle wait states as the base state because the async operation can finish before advancing to the wait state.
-	switch (RequestedState)
+	// An operation may complete and advance to another state before its request has been stored.
+	// Hold a reference to allow access to this after starting the operation.
+	const TRequest Self(this);
+	const EBuildJobState LocalState = State;
+	if (FRequest LocalRequest = Operation())
 	{
-	case EBuildJobState::ResolveKeyWait:       return LocalState == EBuildJobState::ResolveKey;
-	case EBuildJobState::ResolveInputMetaWait: return LocalState == EBuildJobState::ResolveInputMeta;
-	case EBuildJobState::CacheQueryWait:       return LocalState == EBuildJobState::CacheQuery;
-	case EBuildJobState::ResolveInputDataWait: return LocalState == EBuildJobState::ResolveInputData;
-	case EBuildJobState::ExecuteRemoteWait:    return LocalState == EBuildJobState::ExecuteRemote;
-	case EBuildJobState::ExecuteLocalWait:     return LocalState == EBuildJobState::ExecuteLocal;
-	case EBuildJobState::CacheStoreWait:       return LocalState == EBuildJobState::CacheStore;
+		if (FWriteScopeLock WriteLock(Lock); LocalState == State && !bIsCanceled)
+		{
+			// The state flow through the job requires any previous WaitRequest to be complete by
+			// the time that ExecuteAsync is called because the completion of the async operation
+			// will request the state transition that leads to another new async operation. It is
+			// not possible to poll the previous request to validate that it is complete, because
+			// execution likely reached this point from the completion callback, and a request is
+			// not complete until its completion callback has returned.
+			WaitRequest = MoveTemp(LocalRequest);
+		}
+		// Cancel the request if it was not moved above due to cancellation or change of state.
+		LocalRequest.Cancel();
 	}
-
-	return false;
 }
 
-void FBuildJob::AdvanceToState(EBuildJobState NewState, FRequest NewRequest)
+void FBuildJob::AdvanceToState(EBuildJobState NewState)
 {
-	bool bLocalIsCanceled = false;
 	EBuildJobState OldState = EBuildJobState::Complete;
 
 	if (FWriteScopeLock WriteLock(Lock); NextState < NewState)
 	{
-		// The structure of the job requires that the any previous WaitRequest is complete by
-		// the time that AdvanceToState is run, because nothing can happen in waiting states,
-		// and the callback from the request should be triggering this state transition. This
-		// is not possible to validate by polling the request because the request contract is
-		// that it is not considered complete until the callback finishes executing.
-		WaitRequest = MoveTemp(NewRequest);
+		// TODO: Improve CompleteWithError to avoid unexpected state transitions.
+		//checkf(NextState < NewState,
+		//	TEXT("Job in state %s is requesting an invalid transition from %s to %s for build of '%s' by %s."),
+		//	LexToString(State), LexToString(NextState), LexToString(NewState), *Name, *FunctionName);
 
-		bLocalIsCanceled = bIsCanceled;
-		if (bLocalIsCanceled)
+		if (bIsCanceled)
 		{
 			NewState = EBuildJobState::Complete;
 		}
 
 		NextState = NewState;
 
-		if (!bInAdvanceToState)
+		if (bInAdvanceToState)
+		{
+			return;
+		}
+		else
 		{
 			bInAdvanceToState = true;
 			OldState = State;
@@ -1192,6 +1153,8 @@ void FBuildJob::AdvanceToState(EBuildJobState NewState, FRequest NewRequest)
 			break;
 		}
 	}
+
+	ExecuteState(NewState);
 }
 
 void FBuildJob::ExecuteTransition(EBuildJobState OldState, EBuildJobState NewState)
@@ -1200,7 +1163,7 @@ void FBuildJob::ExecuteTransition(EBuildJobState OldState, EBuildJobState NewSta
 	{
 		BeginJob();
 	}
-	if (OldState < EBuildJobState::CacheQuery && EBuildJobState::CacheQuery <= NewState && NewState <= EBuildJobState::CacheStore)
+	if (OldState < EBuildJobState::CacheQuery && EBuildJobState::CacheQuery <= NewState && NewState <= EBuildJobState::CacheStoreWait)
 	{
 		CreateContext();
 	}
@@ -1218,17 +1181,17 @@ void FBuildJob::ExecuteTransition(EBuildJobState OldState, EBuildJobState NewSta
 	{
 		Action.Reset();
 		Inputs.Reset();
+		if (Context)
+		{
+			Context->ResetInputs();
+		}
 		SchedulerParams.ResolvedInputsSize = 0;
 	}
 	if (OldState <= EBuildJobState::ExecuteLocalWait && EBuildJobState::ExecuteLocalWait < NewState && !Output)
 	{
 		SetOutputNoCheck(OutputBuilder.Build());
 	}
-	if (OldState < EBuildJobState::CacheStore && EBuildJobState::CacheStore <= NewState && Context)
-	{
-		Context->ResetInputs();
-	}
-	if (OldState <= EBuildJobState::CacheStore && EBuildJobState::CacheStore < NewState)
+	if (NewState == EBuildJobState::Complete)
 	{
 		Output.Reset();
 		Context = nullptr;
@@ -1236,16 +1199,41 @@ void FBuildJob::ExecuteTransition(EBuildJobState OldState, EBuildJobState NewSta
 
 	switch (NewState)
 	{
-	case EBuildJobState::ResolveKey:             return EnterResolveKey();
-	case EBuildJobState::ResolveInputMeta:       return EnterResolveInputMeta();
-	case EBuildJobState::CacheQuery:             return EnterCacheQuery();
-	case EBuildJobState::ExecuteRemote:          return EnterExecuteRemote();
-	case EBuildJobState::ResolveRemoteInputData: return EnterResolveInputData();
-	case EBuildJobState::ExecuteRemoteRetry:     return EnterExecuteRemote();
-	case EBuildJobState::ResolveInputData:       return EnterResolveInputData();
-	case EBuildJobState::ExecuteLocal:           return EnterExecuteLocal();
-	case EBuildJobState::CacheStore:             return EnterCacheStore();
-	case EBuildJobState::Complete:               return EndJob();
+	case EBuildJobState::ResolveKey:                 return EnterResolveKey();
+	case EBuildJobState::ResolveKeyWait:             return BeginResolveKey();
+	case EBuildJobState::ResolveInputMeta:           return EnterResolveInputMeta();
+	case EBuildJobState::ResolveInputMetaWait:       return BeginResolveInputMeta();
+	case EBuildJobState::CacheQuery:                 return EnterCacheQuery();
+	case EBuildJobState::CacheQueryWait:             return BeginCacheQuery();
+	case EBuildJobState::ExecuteRemote:              return EnterExecuteRemote();
+	case EBuildJobState::ExecuteRemoteWait:          return BeginExecuteRemote();
+	case EBuildJobState::ResolveRemoteInputData:     return EnterResolveInputData();
+	case EBuildJobState::ResolveRemoteInputDataWait: return BeginResolveInputData();
+	case EBuildJobState::ExecuteRemoteRetry:         return EnterExecuteRemote();
+	case EBuildJobState::ExecuteRemoteRetryWait:     return BeginExecuteRemote();
+	case EBuildJobState::ResolveInputData:           return EnterResolveInputData();
+	case EBuildJobState::ResolveInputDataWait:       return BeginResolveInputData();
+	case EBuildJobState::ExecuteLocal:               return EnterExecuteLocal();
+	case EBuildJobState::ExecuteLocalWait:           return BeginExecuteLocal();
+	case EBuildJobState::CacheStore:                 return EnterCacheStore();
+	case EBuildJobState::CacheStoreWait:             return BeginCacheStore();
+	case EBuildJobState::Complete:                   return EndJob();
+	}
+}
+
+void FBuildJob::ExecuteState(EBuildJobState NewState)
+{
+	switch (NewState)
+	{
+	case EBuildJobState::ResolveKey:                 return Scheduler->DispatchResolveKey(this);
+	case EBuildJobState::ResolveInputMeta:           return Scheduler->DispatchResolveInputMeta(this);
+	case EBuildJobState::CacheQuery:                 return Scheduler->DispatchCacheQuery(this, SchedulerParams);
+	case EBuildJobState::ExecuteRemote:              return Scheduler->DispatchExecuteRemote(this, SchedulerParams);
+	case EBuildJobState::ResolveRemoteInputData:     return Scheduler->DispatchResolveInputData(this, SchedulerParams);
+	case EBuildJobState::ExecuteRemoteRetry:         return Scheduler->DispatchExecuteRemote(this, SchedulerParams);
+	case EBuildJobState::ResolveInputData:           return Scheduler->DispatchResolveInputData(this, SchedulerParams);
+	case EBuildJobState::ExecuteLocal:               return Scheduler->DispatchExecuteLocal(this, SchedulerParams);
+	case EBuildJobState::CacheStore:                 return Scheduler->DispatchCacheStore(this, SchedulerParams);
 	}
 }
 
diff --git a/Engine/Source/Developer/DerivedDataCache/Private/DerivedDataBuildSession.cpp b/Engine/Source/Developer/DerivedDataCache/Private/DerivedDataBuildSession.cpp
--- a/Engine/Source/Developer/DerivedDataCache/Private/DerivedDataBuildSession.cpp
+++ b/Engine/Source/Developer/DerivedDataCache/Private/DerivedDataBuildSession.cpp
@@ -74,7 +74,7 @@ FRequest FBuildSessionInternal::Build(
 			OnComplete({Definition.GetKey(), MoveTemp(Params.Output), Params.Status});
 		};
 	}
-	Job->Schedule(Scheduler, Policy, Priority, MoveTemp(OnJobComplete));
+	Job->Start(Scheduler, Policy, Priority, MoveTemp(OnJobComplete));
 	return Job;
 }
 
@@ -94,7 +94,7 @@ FRequest FBuildSessionInternal::BuildAction(
 			OnComplete({Action.GetKey(), MoveTemp(Params.Output), Params.Status});
 		};
 	}
-	Job->Schedule(Scheduler, Policy, Priority, MoveTemp(OnJobComplete));
+	Job->Start(Scheduler, Policy, Priority, MoveTemp(OnJobComplete));
 	return Job;
 }
 
@@ -125,7 +125,7 @@ FRequest FBuildSessionInternal::BuildPayload(
 			OnComplete({PayloadKey.BuildKey, MoveTemp(Payload), Status});
 		};
 	}
-	Job->Schedule(Scheduler, Policy, Priority, MoveTemp(OnJobComplete));
+	Job->Start(Scheduler, Policy, Priority, MoveTemp(OnJobComplete));
 	return Job;
 }
 
diff --git a/Engine/Source/Developer/DerivedDataCache/Public/DerivedDataBuildJob.h b/Engine/Source/Developer/DerivedDataCache/Public/DerivedDataBuildJob.h
--- a/Engine/Source/Developer/DerivedDataCache/Public/DerivedDataBuildJob.h
+++ b/Engine/Source/Developer/DerivedDataCache/Public/DerivedDataBuildJob.h
@@ -52,16 +52,10 @@ public:
 	virtual IBuild& GetBuild() const = 0;
 
 	/** Start the build by dispatching it to the scheduler. */
-	virtual void Schedule(IBuildScheduler& Scheduler, EBuildPolicy Policy, EPriority Priority, FOnBuildJobComplete&& OnComplete) = 0;
+	virtual void Start(IBuildScheduler& Scheduler, EBuildPolicy Policy, EPriority Priority, FOnBuildJobComplete&& OnComplete) = 0;
 
-	/** Called by the scheduler to begin execution for the corresponding state. */
-	virtual void BeginCacheQuery() = 0;
-	virtual void BeginCacheStore() = 0;
-	virtual void BeginResolveKey() = 0;
-	virtual void BeginResolveInputMeta() = 0;
-	virtual void BeginResolveInputData() = 0;
-	virtual void BeginExecuteRemote() = 0;
-	virtual void BeginExecuteLocal() = 0;
+	/** Called by the scheduler to continue this job on the calling thread. */
+	virtual void Schedule() = 0;
 
 	/** Called by the scheduler to skip remote execution and fall back to local execution if permitted. */
 	virtual void SkipExecuteRemote() = 0;
diff --git a/Engine/Source/Developer/DerivedDataCache/Public/DerivedDataBuildScheduler.h b/Engine/Source/Developer/DerivedDataCache/Public/DerivedDataBuildScheduler.h
--- a/Engine/Source/Developer/DerivedDataCache/Public/DerivedDataBuildScheduler.h
+++ b/Engine/Source/Developer/DerivedDataCache/Public/DerivedDataBuildScheduler.h
@@ -35,30 +35,30 @@ public:
 	/** Update the priority of the job if it is queued. May be called multiple times and/or concurrently. */
 	virtual void UpdateJobPriority(IBuildJob* Job) {}
 
-	/** Dispatch by calling BeginCacheQuery or SetOutput, either now or later. */
-	virtual void DispatchCacheQuery(IBuildJob* Job, const FBuildSchedulerParams& Params) { Job->BeginCacheQuery(); }
+	/** Dispatch by calling Schedule or SetOutput, either now or later. */
+	virtual void DispatchCacheQuery(IBuildJob* Job, const FBuildSchedulerParams& Params) { Job->Schedule(); }
 
-	/** Dispatch by calling BeginCacheStore, either now or later. */
-	virtual void DispatchCacheStore(IBuildJob* Job, const FBuildSchedulerParams& Params) { Job->BeginCacheStore(); }
+	/** Dispatch by calling Schedule, either now or later. */
+	virtual void DispatchCacheStore(IBuildJob* Job, const FBuildSchedulerParams& Params) { Job->Schedule(); }
 
-	/** Dispatch by calling BeginResolveKey, either now or later. */
-	virtual void DispatchResolveKey(IBuildJob* Job) { Job->BeginResolveKey(); }
+	/** Dispatch by calling Schedule, either now or later. */
+	virtual void DispatchResolveKey(IBuildJob* Job) { Job->Schedule(); }
 
-	/** Dispatch by calling BeginResolveInputMeta, either now or later. */
-	virtual void DispatchResolveInputMeta(IBuildJob* Job) { Job->BeginResolveInputMeta(); }
+	/** Dispatch by calling Schedule, either now or later. */
+	virtual void DispatchResolveInputMeta(IBuildJob* Job) { Job->Schedule(); }
 
 	/**
-	 * Dispatch by calling BeginResolveInputData, SetOutput, or SkipExecuteRemote, either now or later.
+	 * Dispatch by calling Schedule, SetOutput, or SkipExecuteRemote, either now or later.
 	 *
 	 * SkipExecuteRemote is only valid to call when MissingRemoteInputsSize is non-zero.
 	 */
-	virtual void DispatchResolveInputData(IBuildJob* Job, const FBuildSchedulerParams& Params) { Job->BeginResolveInputData(); }
+	virtual void DispatchResolveInputData(IBuildJob* Job, const FBuildSchedulerParams& Params) { Job->Schedule(); }
 
-	/** Dispatch by calling BeginExecuteRemote, SetOutput, or SkipExecuteRemote, either now or later. */
-	virtual void DispatchExecuteRemote(IBuildJob* Job, const FBuildSchedulerParams& Params) { Job->BeginExecuteRemote(); }
+	/** Dispatch by calling Schedule, SetOutput, or SkipExecuteRemote, either now or later. */
+	virtual void DispatchExecuteRemote(IBuildJob* Job, const FBuildSchedulerParams& Params) { Job->Schedule(); }
 
-	/** Dispatch by calling BeginExecuteLocal or SetOutput, either now or later. */
-	virtual void DispatchExecuteLocal(IBuildJob* Job, const FBuildSchedulerParams& Params) { Job->BeginExecuteLocal(); }
+	/** Dispatch by calling Schedule or SetOutput, either now or later. */
+	virtual void DispatchExecuteLocal(IBuildJob* Job, const FBuildSchedulerParams& Params) { Job->Schedule(); }
 
 	/** Set the output of the job. Always called once between BeginJob and EndJob unless canceled. */
 	virtual void SetJobOutput(IBuildJob* Job, const FBuildSchedulerParams& Params, const FBuildOutput& Output) {}
