DDC: Schedule non-blocking builds using tasks

#rnx

diff --git a/Engine/Source/Developer/DerivedDataCache/Private/DerivedDataBuildScheduler.cpp b/Engine/Source/Developer/DerivedDataCache/Private/DerivedDataBuildScheduler.cpp
--- a/Engine/Source/Developer/DerivedDataCache/Private/DerivedDataBuildScheduler.cpp
+++ b/Engine/Source/Developer/DerivedDataCache/Private/DerivedDataBuildScheduler.cpp
@@ -6,14 +6,103 @@
 #include "DerivedDataBuildFunctionFactory.h"
 #include "DerivedDataBuildFunctionRegistry.h"
 #include "Misc/Guid.h"
+#include "Tasks/Task.h"
 
 namespace UE::DerivedData::Private
 {
 
 class FBuildScheduler final : public IBuildScheduler
 {
+public:
+	void DispatchCacheQuery(IBuildJob* Job, const FBuildSchedulerParams& Params) final;
+	void DispatchCacheStore(IBuildJob* Job, const FBuildSchedulerParams& Params) final;
+	void DispatchResolveKey(IBuildJob* Job) final;
+	void DispatchResolveInputMeta(IBuildJob* Job) final;
+	void DispatchResolveInputData(IBuildJob* Job, const FBuildSchedulerParams& Params) final;
+	void DispatchExecuteRemote(IBuildJob* Job, const FBuildSchedulerParams& Params) final;
+	void DispatchExecuteLocal(IBuildJob* Job, const FBuildSchedulerParams& Params) final;
 };
 
+void FBuildScheduler::DispatchCacheQuery(IBuildJob* Job, const FBuildSchedulerParams& Params)
+{
+	Job->Schedule();
+}
+
+void FBuildScheduler::DispatchCacheStore(IBuildJob* Job, const FBuildSchedulerParams& Params)
+{
+	Job->Schedule();
+}
+
+void FBuildScheduler::DispatchResolveKey(IBuildJob* Job)
+{
+	if (Job->GetPriority() == EPriority::Blocking)
+	{
+		Job->Schedule();
+	}
+	else
+	{
+		Tasks::Launch(TEXT("FBuildScheduler::DispatchResolveKey"),
+			[Job = TRequest(Job)] { Job->Schedule(); },
+			LowLevelTasks::ETaskPriority::BackgroundNormal);
+	}
+}
+
+void FBuildScheduler::DispatchResolveInputMeta(IBuildJob* Job)
+{
+	if (Job->GetPriority() == EPriority::Blocking)
+	{
+		Job->Schedule();
+	}
+	else
+	{
+		Tasks::Launch(TEXT("FBuildScheduler::DispatchResolveInputMeta"),
+			[Job = TRequest(Job)] { Job->Schedule(); },
+			LowLevelTasks::ETaskPriority::BackgroundNormal);
+	}
+}
+
+void FBuildScheduler::DispatchResolveInputData(IBuildJob* Job, const FBuildSchedulerParams& Params)
+{
+	if (Job->GetPriority() == EPriority::Blocking)
+	{
+		Job->Schedule();
+	}
+	else
+	{
+		Tasks::Launch(TEXT("FBuildScheduler::DispatchResolveInputData"),
+			[Job = TRequest(Job)] { Job->Schedule(); },
+			LowLevelTasks::ETaskPriority::BackgroundNormal);
+	}
+}
+
+void FBuildScheduler::DispatchExecuteRemote(IBuildJob* Job, const FBuildSchedulerParams& Params)
+{
+	if (Job->GetPriority() == EPriority::Blocking)
+	{
+		Job->Schedule();
+	}
+	else
+	{
+		Tasks::Launch(TEXT("FBuildScheduler::DispatchExecuteRemote"),
+			[Job = TRequest(Job)] { Job->Schedule(); },
+			LowLevelTasks::ETaskPriority::BackgroundNormal);
+	}
+}
+
+void FBuildScheduler::DispatchExecuteLocal(IBuildJob* Job, const FBuildSchedulerParams& Params)
+{
+	if (Job->GetPriority() == EPriority::Blocking)
+	{
+		Job->Schedule();
+	}
+	else
+	{
+		Tasks::Launch(TEXT("FBuildScheduler::DispatchExecuteLocal"),
+			[Job = TRequest(Job)] { Job->Schedule(); },
+			LowLevelTasks::ETaskPriority::BackgroundNormal);
+	}
+}
+
 IBuildScheduler* CreateBuildScheduler()
 {
 	return new FBuildScheduler();
