DDC: Fixed a missing include in DerivedDataBuildVersion.h

#rnx

diff --git a/Engine/Source/Developer/DerivedDataCache/Public/DerivedDataBuildVersion.h b/Engine/Source/Developer/DerivedDataCache/Public/DerivedDataBuildVersion.h
--- a/Engine/Source/Developer/DerivedDataCache/Public/DerivedDataBuildVersion.h
+++ b/Engine/Source/Developer/DerivedDataCache/Public/DerivedDataBuildVersion.h
@@ -6,6 +6,7 @@
 #include "Hash/xxhash.h"
 #include "Misc/Guid.h"
 #include "Serialization/Archive.h"
+#include "UObject/NameTypes.h"
 #include <type_traits>
 
 namespace UE::DerivedData
