DDC: Added EPutStatus::Skipped to avoid slow retry when simulating cache misses

diff --git a/Engine/Source/Developer/DerivedDataCache/Private/DerivedDataBackendAsyncPutWrapper.cpp b/Engine/Source/Developer/DerivedDataCache/Private/DerivedDataBackendAsyncPutWrapper.cpp
--- a/Engine/Source/Developer/DerivedDataCache/Private/DerivedDataBackendAsyncPutWrapper.cpp
+++ b/Engine/Source/Developer/DerivedDataCache/Private/DerivedDataBackendAsyncPutWrapper.cpp
@@ -82,7 +82,7 @@ public:
 		if (InflightCache)
 		{
 			// if the data was not cached synchronously, retry
-			if (Status != EPutStatus::Cached)
+			if (Status != EPutStatus::Cached && Status != EPutStatus::Skipped)
 			{
 				// retry after a brief wait
 				FPlatformProcess::SleepNoStats(0.2f);
@@ -99,8 +99,9 @@ public:
 
 			switch (Status)
 			{
+			case EPutStatus::Skipped:
 			case EPutStatus::Cached:
-				// remove this from the in-flight cache because the inner cache contains the data
+				// remove this from the in-flight cache because the inner cache contains the data or it was intentionally skipped
 				InflightCache->RemoveCachedData(*CacheKey, /*bTransient*/ false);
 				break;
 			case EPutStatus::NotCached:
diff --git a/Engine/Source/Developer/DerivedDataCache/Private/FileSystemDerivedDataBackend.cpp b/Engine/Source/Developer/DerivedDataCache/Private/FileSystemDerivedDataBackend.cpp
--- a/Engine/Source/Developer/DerivedDataCache/Private/FileSystemDerivedDataBackend.cpp
+++ b/Engine/Source/Developer/DerivedDataCache/Private/FileSystemDerivedDataBackend.cpp
@@ -730,7 +730,7 @@ public:
 			// don't put anything we pretended didn't exist
 			if (ShouldSimulateMiss(CacheKey))
 			{
-				return EPutStatus::NotCached;
+				return EPutStatus::Skipped;
 			}
 
 			EPutStatus Status = EPutStatus::NotCached;
diff --git a/Engine/Source/Developer/DerivedDataCache/Private/MemoryDerivedDataBackend.cpp b/Engine/Source/Developer/DerivedDataCache/Private/MemoryDerivedDataBackend.cpp
--- a/Engine/Source/Developer/DerivedDataCache/Private/MemoryDerivedDataBackend.cpp
+++ b/Engine/Source/Developer/DerivedDataCache/Private/MemoryDerivedDataBackend.cpp
@@ -117,7 +117,7 @@ FDerivedDataBackendInterface::EPutStatus FMemoryDerivedDataBackend::PutCachedDat
 
 	if (DidSimulateMiss(CacheKey))
 	{
-		return EPutStatus::NotCached;
+		return EPutStatus::Skipped;
 	}
 	
 	// Should never hit this as higher level code should be checking..
diff --git a/Engine/Source/Developer/DerivedDataCache/Public/DerivedDataBackendInterface.h b/Engine/Source/Developer/DerivedDataCache/Public/DerivedDataBackendInterface.h
--- a/Engine/Source/Developer/DerivedDataCache/Public/DerivedDataBackendInterface.h
+++ b/Engine/Source/Developer/DerivedDataCache/Public/DerivedDataBackendInterface.h
@@ -42,6 +42,8 @@ public:
 		NotCached,
 		/** The put completed synchronously and the data was cached. */
 		Cached,
+		/** The put was skipped and should not be retried. */
+		Skipped,
 	};
 
 	/*
