DDC: Fixed a crash when creating the build context fails

#rnx

diff --git a/Engine/Source/Developer/DerivedDataCache/Private/DerivedDataBuildJob.cpp b/Engine/Source/Developer/DerivedDataCache/Private/DerivedDataBuildJob.cpp
--- a/Engine/Source/Developer/DerivedDataCache/Private/DerivedDataBuildJob.cpp
+++ b/Engine/Source/Developer/DerivedDataCache/Private/DerivedDataBuildJob.cpp
@@ -267,6 +267,7 @@ const TCHAR* LexToString(EBuildJobState State)
 {
 	switch (State)
 	{
+	case EBuildJobState::NotStarted:                 return TEXT("NotStarted");
 	case EBuildJobState::ResolveKey:                 return TEXT("ResolveKey");
 	case EBuildJobState::ResolveKeyWait:             return TEXT("ResolveKeyWait");
 	case EBuildJobState::ResolveInputMeta:           return TEXT("ResolveInputMeta");
@@ -533,23 +534,23 @@ void FBuildJob::CreateContext()
 	const IBuildFunction* Function = BuildSystem.GetFunctionRegistry().FindFunction(FunctionName);
 	if (BuildSystem.GetVersion() != Action.Get().GetBuildSystemVersion())
 	{
-		CompleteWithError(WriteToString<192>(TEXT("Failed because the build system is version "_SV),
+		return CompleteWithError(WriteToString<192>(TEXT("Failed because the build system is version "_SV),
 			BuildSystem.GetVersion(), TEXT(" when version "_SV),
 			Action.Get().GetBuildSystemVersion(), TEXT(" is expected."_SV)));
 	}
 	else if (!Function)
 	{
-		CompleteWithError(WriteToString<128>(TEXT("Failed because the function "_SV), FunctionName,
+		return CompleteWithError(WriteToString<128>(TEXT("Failed because the function "_SV), FunctionName,
 			TEXT(" was not found."_SV)));
 	}
 	else if (!Function->GetVersion().IsValid())
 	{
-		CompleteWithError(WriteToString<128>(TEXT("Failed because the function "_SV), FunctionName,
+		return CompleteWithError(WriteToString<128>(TEXT("Failed because the function "_SV), FunctionName,
 			TEXT(" has a version of zero."_SV)));
 	}
 	else if (Function->GetVersion() != Action.Get().GetFunctionVersion())
 	{
-		CompleteWithError(WriteToString<192>(TEXT("Failed because the function "_SV), FunctionName,
+		return CompleteWithError(WriteToString<192>(TEXT("Failed because the function "_SV), FunctionName,
 			TEXT(" is version "_SV), Function->GetVersion(), TEXT(" when version "_SV),
 			Action.Get().GetFunctionVersion(), TEXT(" is expected."_SV)));
 	}
@@ -588,7 +589,7 @@ void FBuildJob::CreateContext()
 
 void FBuildJob::EnterCacheQuery()
 {
-	ECachePolicy CachePolicy = Context->GetCachePolicy();
+	ECachePolicy CachePolicy = Context ? Context->GetCachePolicy() : ECachePolicy::None;
 	if (!EnumHasAnyFlags(CachePolicy, ECachePolicy::Query) ||
 		EnumHasAnyFlags(BuildPolicy, EBuildPolicy::SkipCacheGet))
 	{
@@ -599,7 +600,7 @@ void FBuildJob::EnterCacheQuery()
 void FBuildJob::BeginCacheQuery()
 {
 	TRACE_CPUPROFILER_EVENT_SCOPE(FBuildJob::CacheQuery);
-	ExecuteAsync([this]
+	return ExecuteAsync([this]
 	{
 		ECachePolicy CachePolicy = Context->GetCachePolicy();
 		if (EnumHasAnyFlags(BuildPolicy, EBuildPolicy::SkipData))
@@ -628,7 +629,7 @@ void FBuildJob::EndCacheQuery(FCacheGetCompleteParams&& Params)
 
 void FBuildJob::EnterCacheStore()
 {
-	ECachePolicy CachePolicy = Context->GetCachePolicy();
+	ECachePolicy CachePolicy = Context ? Context->GetCachePolicy() : ECachePolicy::None;
 	if (bIsCached ||
 		!EnumHasAnyFlags(CachePolicy, ECachePolicy::Store) ||
 		EnumHasAnyFlags(BuildPolicy, EBuildPolicy::SkipCachePut) ||
@@ -641,7 +642,7 @@ void FBuildJob::EnterCacheStore()
 void FBuildJob::BeginCacheStore()
 {
 	TRACE_CPUPROFILER_EVENT_SCOPE(FBuildJob::CacheStore);
-	ExecuteAsync([this]
+	return ExecuteAsync([this]
 	{
 		FCacheRecordBuilder RecordBuilder = Cache.CreateRecord(Context->GetCacheKey());
 		Output.Get().Save(RecordBuilder);
@@ -682,7 +683,7 @@ void FBuildJob::EnterResolveKey()
 void FBuildJob::BeginResolveKey()
 {
 	TRACE_CPUPROFILER_EVENT_SCOPE(FBuildJob::ResolveKey);
-	ExecuteAsync([this]
+	return ExecuteAsync([this]
 	{
 		return InputResolver->ResolveKey(DefinitionKey,
 			[this](FBuildKeyResolvedParams&& Params) { EndResolveKey(MoveTemp(Params)); });
@@ -693,11 +694,11 @@ void FBuildJob::EndResolveKey(FBuildKeyResolvedParams&& Params)
 {
 	if (Params.Status == EStatus::Ok && Params.Definition)
 	{
-		SetDefinition(MoveTemp(Params.Definition).Get());
+		return SetDefinition(MoveTemp(Params.Definition).Get());
 	}
 	else
 	{
-		CompleteWithError(WriteToString<128>(TEXT("Failed to resolve key "_SV), Params.Key, TEXT("."_SV)));
+		return CompleteWithError(WriteToString<128>(TEXT("Failed to resolve key "_SV), Params.Key, TEXT("."_SV)));
 	}
 }
 
@@ -718,7 +719,7 @@ void FBuildJob::EnterResolveInputMeta()
 void FBuildJob::BeginResolveInputMeta()
 {
 	TRACE_CPUPROFILER_EVENT_SCOPE(FBuildJob::ResolveInputMeta);
-	ExecuteAsync([this]
+	return ExecuteAsync([this]
 	{
 		return InputResolver->ResolveInputMeta(Definition.Get(), Priority,
 			[this](FBuildInputMetaResolvedParams&& Params) { EndResolveInputMeta(MoveTemp(Params)); });
@@ -729,11 +730,11 @@ void FBuildJob::EndResolveInputMeta(FBuildInputMetaResolvedParams&& Params)
 {
 	if (Params.Status == EStatus::Ok)
 	{
-		CreateAction(Params.Inputs);
+		return CreateAction(Params.Inputs);
 	}
 	else
 	{
-		CompleteWithError(TEXT("Failed to resolve input metadata."_SV));
+		return CompleteWithError(TEXT("Failed to resolve input metadata."_SV));
 	}
 }
 
@@ -783,7 +784,7 @@ void FBuildJob::EnterResolveInputData()
 void FBuildJob::BeginResolveInputData()
 {
 	TRACE_CPUPROFILER_EVENT_SCOPE(FBuildJob::ResolveInputData);
-	ExecuteAsync([this]
+	return ExecuteAsync([this]
 	{
 		if (Definition)
 		{
@@ -817,11 +818,11 @@ void FBuildJob::EndResolveInputData(FBuildInputDataResolvedParams&& Params)
 				Builder.AddInput(Key, Buffer);
 			});
 		}
-		SetInputs(Builder.Build());
+		return SetInputs(Builder.Build());
 	}
 	else
 	{
-		CompleteWithError(TEXT("Failed to resolve input data."_SV));
+		return CompleteWithError(TEXT("Failed to resolve input data."_SV));
 	}
 }
 
@@ -859,7 +860,7 @@ void FBuildJob::BeginExecuteRemote()
 	checkf(Worker && WorkerExecutor, TEXT("Job requires a worker in state %s for build of '%s' by %s."),
 		LexToString(State), *Name, *FunctionName);
 	bTriedRemoteExecution = true;
-	ExecuteAsync([this]
+	return ExecuteAsync([this]
 	{
 		return WorkerExecutor->BuildAction(Action.Get(), Inputs, *Worker, BuildPolicy, Priority,
 			[this](FBuildWorkerActionCompleteParams&& Params) { EndExecuteRemote(MoveTemp(Params)); });
@@ -953,7 +954,7 @@ void FBuildJob::EnterExecuteLocal()
 void FBuildJob::BeginExecuteLocal()
 {
 	TRACE_CPUPROFILER_EVENT_SCOPE(FBuildJob::ExecuteLocal);
-	ExecuteAsync([this]
+	return ExecuteAsync([this]
 	{
 		Context->GetFunction().Build(*Context);
 		if (!Context->IsAsyncBuild())
@@ -983,7 +984,7 @@ void FBuildJob::CreateAction(TConstArrayView<FBuildInputMetaByKey> InputMeta)
 	{
 		Builder.AddInput(Input.Key, Input.RawHash, Input.RawSize);
 	}
-	SetAction(Builder.Build());
+	return SetAction(Builder.Build());
 }
 
 ///////////////////////////////////////////////////////////////////////////////////////////////////////////////////////
