Fix bug where queued batch request failure in HttpDerivedDataBackend would cause incorrect DDC output (erroneous cache misses).
Adding basic unit test for HttpDerivedDataBackend.

diff --git a/Engine/Source/Developer/DerivedDataCache/Private/HttpDerivedDataBackend.cpp b/Engine/Source/Developer/DerivedDataCache/Private/HttpDerivedDataBackend.cpp
--- a/Engine/Source/Developer/DerivedDataCache/Private/HttpDerivedDataBackend.cpp
+++ b/Engine/Source/Developer/DerivedDataCache/Private/HttpDerivedDataBackend.cpp
@@ -945,14 +945,33 @@ struct FDataRequestHelper
 		{
 			// We have exceeded the threshold for concurrent connections, start or add this request
 			// to a batched request.
-			Request = QueueBatchRequest(
-				InPool, 
-				InNamespace,
-				InBucket, 
-				TConstArrayView<const TCHAR*>({ InCacheKey }),
-				OutData ? TConstArrayView<TArray<uint8>*>({ OutData }) : TConstArrayView<TArray<uint8>*>(),
-				bVerified
-			);
+			if (IsQueueCandidate(1, OutData && !OutData->IsEmpty()))
+			{
+				Request = QueueBatchRequest(
+					InPool,
+					InNamespace,
+					InBucket,
+					TConstArrayView<const TCHAR*>({InCacheKey}),
+					OutData ? TConstArrayView<TArray<uint8>*>({OutData}) : TConstArrayView<TArray<uint8>*>(),
+					bVerified
+				);
+			}
+
+			if (!Request)
+			{
+				Request = Pool->WaitForFreeRequest();
+
+				FQueuedBatchEntry Entry{
+					InNamespace,
+					InBucket,
+					TConstArrayView<const TCHAR*>({InCacheKey}),
+					OutData ? TConstArrayView<TArray<uint8>*>({OutData}) : TConstArrayView<TArray<uint8>*>(),
+					OutData && !OutData->IsEmpty() ? FHttpRequest::RequestVerb::Get : FHttpRequest::RequestVerb::Head,
+					&bVerified
+				};
+
+				PerformBatchQuery(Request, TArrayView<FQueuedBatchEntry>(&Entry, 1));
+			}
 		}
 	}
 
@@ -967,7 +986,8 @@ struct FDataRequestHelper
 		Algo::Transform(InCacheKeys, CacheKeys, [](const FString& Key) { return *Key; });
 		
 		Request = Pool->GetFreeRequest();
-		if (Request || InCacheKeys.Num() > UE_HTTPDDC_BATCH_SIZE)
+
+		if (Request || !IsQueueCandidate(InCacheKeys.Num(), false))
 		{
 			// If the request is too big for existing batches, wait for a free connection and create our own.
 			if (!Request)
@@ -996,6 +1016,22 @@ struct FDataRequestHelper
 				TConstArrayView<TArray<uint8>*>(), 
 				bVerified
 			);
+
+			if (!Request)
+			{
+				Request = Pool->WaitForFreeRequest();
+
+				FQueuedBatchEntry Entry{
+					InNamespace,
+					InBucket,
+					CacheKeys,
+					TConstArrayView<TArray<uint8>*>(),
+					FHttpRequest::RequestVerb::Head,
+					&bVerified
+				};
+
+				PerformBatchQuery(Request, TArrayView<FQueuedBatchEntry>(&Entry, 1));
+			}
 		}
 	}
 
@@ -1062,6 +1098,25 @@ private:
 	static std::atomic<uint32> FirstAvailableBatch;
 	static TStaticArray<FBatch, UE_HTTPDDC_BATCH_NUM> Batches;
 
+	static uint32 ComputeWeight(int32 NumKeys, bool bHasDatas)
+	{
+		return NumKeys * (bHasDatas ? UE_HTTPDDC_BATCH_GET_WEIGHT : UE_HTTPDDC_BATCH_HEAD_WEIGHT);
+	}
+
+	static bool IsQueueCandidate(int32 NumKeys, bool bHasDatas)
+	{
+		if (NumKeys > UE_HTTPDDC_BATCH_SIZE)
+		{
+			return false;
+		}
+		const uint32 Weight = ComputeWeight(NumKeys, bHasDatas);
+		if (Weight > UE_HTTPDDC_BATCH_WEIGHT_HINT)
+		{
+			return false;
+		}
+		return true;
+	}
+
 	/**
 	 * Queues up a request to be batched. Blocks until the query is made.
 	 */
@@ -1075,7 +1130,7 @@ private:
 		TRACE_CPUPROFILER_EVENT_SCOPE(HttpDDC_BatchQuery);
 		check(InCacheKeys.Num() == OutDatas.Num() || OutDatas.Num() == 0);
 		const uint32 RequestNum = InCacheKeys.Num();
-		const uint32 RequestWeight = InCacheKeys.Num() * (OutDatas.Num() ? UE_HTTPDDC_BATCH_HEAD_WEIGHT : UE_HTTPDDC_BATCH_GET_WEIGHT);
+		const uint32 RequestWeight = ComputeWeight(InCacheKeys.Num(), !OutDatas.IsEmpty());
 
 		for (int32 i = 0; i < Batches.Num(); i++)
 		{
@@ -1793,10 +1848,16 @@ FHttpDerivedDataBackend::FHttpDerivedDataBackend(
 		PutRequestPools[1] = MakeUnique<FRequestPool>(InServiceUrl, Access.Get(), 1);
 		bIsUsable = true;
 	}
+
+	AnyInstance = this;
 }
 
 FHttpDerivedDataBackend::~FHttpDerivedDataBackend()
 {
+	if (AnyInstance == this)
+	{
+		AnyInstance = nullptr;
+	}
 }
 
 FString FHttpDerivedDataBackend::GetName() const
diff --git a/Engine/Source/Developer/DerivedDataCache/Private/HttpDerivedDataBackend.h b/Engine/Source/Developer/DerivedDataCache/Private/HttpDerivedDataBackend.h
--- a/Engine/Source/Developer/DerivedDataCache/Private/HttpDerivedDataBackend.h
+++ b/Engine/Source/Developer/DerivedDataCache/Private/HttpDerivedDataBackend.h
@@ -97,6 +97,11 @@ public:
 	{
 	}
 
+	static FHttpDerivedDataBackend* GetAny()
+	{
+		return AnyInstance;
+	}
+
 private:
 	ICacheFactory& Factory;
 	FString Domain;
@@ -114,6 +119,7 @@ private:
 	bool bReadOnly;
 	uint32 FailedLoginAttempts;
 	ESpeedClass SpeedClass;
+	static inline FHttpDerivedDataBackend* AnyInstance = nullptr;
 
 	bool IsServiceReady();
 	bool AcquireAccessToken();
diff --git a/Engine/Source/Developer/DerivedDataCache/Private/Tests/HttpDerivedDataBackendTest.cpp b/Engine/Source/Developer/DerivedDataCache/Private/Tests/HttpDerivedDataBackendTest.cpp
new file mode 100644
--- /dev/null
+++ b/Engine/Source/Developer/DerivedDataCache/Private/Tests/HttpDerivedDataBackendTest.cpp
@@ -0,0 +1,113 @@
+// Copyright Epic Games, Inc. All Rights Reserved.
+
+#include "HttpDerivedDataBackend.h"
+
+#include "Async/ParallelFor.h"
+#include "DerivedDataBackendInterface.h"
+#include "Misc/AutomationTest.h"
+#include "Misc/SecureHash.h"
+
+// Test is targeted at HttpDerivedDataBackend but with some backend test interface it could be generalized
+// to function against all backends.
+
+#if WITH_DEV_AUTOMATION_TESTS && WITH_HTTP_DDC_BACKEND
+
+#define TEST_NAME_ROOT TEXT("System.DerivedDataCache.HttpDerivedDataBackend")
+
+#define IMPLEMENT_HTTPDERIVEDDATA_AUTOMATION_TEST( TClass, PrettyName, TFlags ) \
+	IMPLEMENT_CUSTOM_COMPLEX_AUTOMATION_TEST(TClass, FHttpDerivedDataTestBase, TEST_NAME_ROOT PrettyName, TFlags) \
+	void TClass::GetTests(TArray<FString>& OutBeautifiedNames, TArray <FString>& OutTestCommands) const \
+	{ \
+		if (CheckPrequisites()) \
+		{ \
+			OutBeautifiedNames.Add(TEST_NAME_ROOT PrettyName); \
+			OutTestCommands.Add(FString()); \
+		} \
+	}
+
+namespace HttpDerivedDataBackendTest
+{
+
+class FHttpDerivedDataTestBase : public FAutomationTestBase
+{
+public:
+	FHttpDerivedDataTestBase(const FString& InName, const bool bInComplexTask)
+	: FAutomationTestBase(InName, bInComplexTask)
+	{
+	}
+
+	bool CheckPrequisites() const
+	{
+		if (UE::DerivedData::Backends::FHttpDerivedDataBackend* Backend = GetTestBackend())
+		{
+			if (Backend->IsUsable())
+			{
+				return true;
+			}
+		}
+		return false;
+	}
+
+protected:
+	UE::DerivedData::Backends::FHttpDerivedDataBackend* GetTestBackend() const
+	{
+		static UE::DerivedData::Backends::FHttpDerivedDataBackend* CachedBackend = FetchTestBackend_Internal();
+		return CachedBackend;
+	}
+
+private:
+	UE::DerivedData::Backends::FHttpDerivedDataBackend* FetchTestBackend_Internal() const
+	{
+		return UE::DerivedData::Backends::FHttpDerivedDataBackend::GetAny();
+	}
+};
+
+IMPLEMENT_HTTPDERIVEDDATA_AUTOMATION_TEST(FConcurrentCachedDataProbablyExistsBatch, TEXT(".FHeavyConcurrentCachedDataProbablyExistsBatch"), EAutomationTestFlags::EditorContext | EAutomationTestFlags::ProductFilter)
+bool FConcurrentCachedDataProbablyExistsBatch::RunTest(const FString& Parameters)
+{
+	UE::DerivedData::Backends::FHttpDerivedDataBackend* TestBackend = GetTestBackend();
+
+	const int32 ParallelTasks = 32;
+	const uint32 Iterations = 20;
+	const uint32 KeysInBatch = 4;
+
+	TArray<FString> Keys;
+	TArray<uint8> KeyContents;
+	KeyContents.Add(42);
+
+	FSHA1 HashState;
+	HashState.Update(KeyContents.GetData(), KeyContents.Num());
+	HashState.Final();
+	uint8 Hash[FSHA1::DigestSize];
+	HashState.GetHash(Hash);
+	const FString HashString = BytesToHex(Hash, FSHA1::DigestSize);
+
+	for (uint32 KeyIndex = 0; KeyIndex < KeysInBatch; ++KeyIndex)
+	{
+		FString NewKey = FString::Printf(TEXT("__AutoTest_Dummy_%u__%s"), KeyIndex, *HashString);
+		Keys.Add(NewKey);
+		TestBackend->PutCachedData(*NewKey, KeyContents, false);
+	}
+
+	std::atomic<uint32> MismatchedResults = 0;
+	ParallelFor(ParallelTasks,
+		[&](int32 TaskIndex)
+		{
+			for (uint32 Iteration = 0; Iteration < Iterations; ++Iteration)
+			{
+				TConstArrayView<FString> BatchView = MakeArrayView(Keys.GetData(), KeysInBatch);
+				TBitArray<> Result = TestBackend->CachedDataProbablyExistsBatch(BatchView);
+				if (Result.CountSetBits() != BatchView.Num())
+				{
+					MismatchedResults.fetch_add(BatchView.Num() - Result.CountSetBits(), std::memory_order_relaxed);
+				}
+			}
+		}
+	);
+	TestEqual(TEXT("Concurrent calls to CachedDataProbablyExistsBatch for a batch of keys that were put are not reliably found"), MismatchedResults, 0);
+
+	return true;
+}
+
+}
+#endif // WITH_DEV_AUTOMATION_TESTS
