DDC: Exposed the cache key and build status in the build complete callback params

#rnx

diff --git a/Engine/Source/Developer/DerivedDataCache/Private/DerivedDataBuildJob.cpp b/Engine/Source/Developer/DerivedDataCache/Private/DerivedDataBuildJob.cpp
--- a/Engine/Source/Developer/DerivedDataCache/Private/DerivedDataBuildJob.cpp
+++ b/Engine/Source/Developer/DerivedDataCache/Private/DerivedDataBuildJob.cpp
@@ -216,6 +216,9 @@ private:
 	/** True if the build action and inputs are being exported. */
 	bool bIsExportingBuild{};
 
+	/** Status flags that are added to as the job moves through its states. */
+	EBuildStatus BuildStatus{EBuildStatus::None};
+
 	mutable std::atomic<uint32> ReferenceCount{0};
 
 	/** Available in [ResolveKey, Complete] for jobs created from a key or definition. */
@@ -591,7 +594,7 @@ void FBuildJob::EnterCacheQuery()
 {
 	ECachePolicy CachePolicy = Context ? Context->GetCachePolicy() : ECachePolicy::None;
 	if (!EnumHasAnyFlags(CachePolicy, ECachePolicy::Query) ||
-		EnumHasAnyFlags(BuildPolicy, EBuildPolicy::SkipCacheGet))
+		!EnumHasAnyFlags(BuildPolicy, EBuildPolicy::CacheQuery))
 	{
 		return AdvanceToState(EBuildJobState::ExecuteRemote);
 	}
@@ -614,8 +617,10 @@ void FBuildJob::BeginCacheQuery()
 
 void FBuildJob::EndCacheQuery(FCacheGetCompleteParams&& Params)
 {
+	BuildStatus |= EBuildStatus::CacheQuery;
 	if (Params.Status == EStatus::Ok)
 	{
+		BuildStatus |= EBuildStatus::CacheQueryHit;
 		if (FOptionalBuildOutput CacheOutput = BuildSystem.LoadOutput(Name, FunctionName, Params.Record))
 		{
 			bIsCached = true;
@@ -632,7 +637,7 @@ void FBuildJob::EnterCacheStore()
 	ECachePolicy CachePolicy = Context ? Context->GetCachePolicy() : ECachePolicy::None;
 	if (bIsCached ||
 		!EnumHasAnyFlags(CachePolicy, ECachePolicy::Store) ||
-		EnumHasAnyFlags(BuildPolicy, EBuildPolicy::SkipCachePut) ||
+		!EnumHasAnyFlags(BuildPolicy, EBuildPolicy::CacheStore) ||
 		Output.Get().HasError())
 	{
 		return AdvanceToState(EBuildJobState::Complete);
@@ -830,11 +835,7 @@ void FBuildJob::EndResolveInputData(FBuildInputDataResolvedParams&& Params)
 
 void FBuildJob::EnterExecuteRemote()
 {
-	if (EnumHasAnyFlags(BuildPolicy, EBuildPolicy::SkipBuild))
-	{
-		return CompleteWithError(TEXT("Failed because build policy skipped building."_SV));
-	}
-	if (!EnumHasAnyFlags(BuildPolicy, EBuildPolicy::Remote) || bIsExportingBuild)
+	if (!EnumHasAnyFlags(BuildPolicy, EBuildPolicy::BuildRemote) || bIsExportingBuild)
 	{
 		return AdvanceToState(EBuildJobState::ResolveInputData);
 	}
@@ -871,6 +872,7 @@ void FBuildJob::EndExecuteRemote(FBuildWorkerActionCompleteParams&& Params)
 {
 	if (Params.Output)
 	{
+		BuildStatus |= EBuildStatus::BuildRemote;
 		return SetOutputNoCheck(MoveTemp(Params.Output).Get());
 	}
 	else
@@ -915,9 +917,13 @@ void FBuildJob::EnterExecuteLocal()
 	{
 		ExportBuild();
 	}
-	if (!EnumHasAnyFlags(BuildPolicy, EBuildPolicy::Local))
+	if (!EnumHasAnyFlags(BuildPolicy, EBuildPolicy::BuildLocal))
 	{
-		if (bTriedRemoteExecution)
+		if (!EnumHasAnyFlags(BuildPolicy, EBuildPolicy::BuildRemote))
+		{
+			return CompleteWithError(TEXT("Failed because build policy does not allow local or remote execution."_SV));
+		}
+		else if (bTriedRemoteExecution)
 		{
 			return CompleteWithError(TEXT("Failed because build policy does not allow local execution, ")
 				TEXT("and remote execution failed to build."_SV));
@@ -967,6 +973,7 @@ void FBuildJob::BeginExecuteLocal()
 
 void FBuildJob::EndExecuteLocal()
 {
+	BuildStatus |= EBuildStatus::BuildLocal;
 	return SetOutputNoCheck(OutputBuilder.Build());
 }
 
@@ -1060,8 +1067,9 @@ void FBuildJob::SetOutputNoCheck(FBuildOutput&& InOutput)
 
 	if (OnComplete)
 	{
+		const FCacheKey& CacheKey = Context ? Context->GetCacheKey() : FCacheKey::Empty;
 		const EStatus Status = bIsCanceled ? EStatus::Canceled : Output.Get().HasError() ? EStatus::Error : EStatus::Ok;
-		OnComplete({*this, FBuildOutput(Output.Get()), Status});
+		OnComplete({*this, CacheKey, FBuildOutput(Output.Get()), BuildStatus, Status});
 		OnComplete = nullptr;
 	}
 
diff --git a/Engine/Source/Developer/DerivedDataCache/Private/DerivedDataBuildJobContext.cpp b/Engine/Source/Developer/DerivedDataCache/Private/DerivedDataBuildJobContext.cpp
--- a/Engine/Source/Developer/DerivedDataCache/Private/DerivedDataBuildJobContext.cpp
+++ b/Engine/Source/Developer/DerivedDataCache/Private/DerivedDataBuildJobContext.cpp
@@ -138,7 +138,7 @@ void FBuildJobContext::SetCacheBucket(FCacheBucket Bucket)
 void FBuildJobContext::SetCachePolicy(ECachePolicy Policy)
 {
 	checkf(!EnumHasAnyFlags(Policy, ECachePolicy::SkipData),
-		TEXT("SkipData flags not allowed on cache policy for build of '%s' by %s. ")
+		TEXT("SkipData flags not allowed on the cache policy for build of '%s' by %s. ")
 		TEXT("Flags for skipping data may be set indirectly through EBuildPolicy."),
 		*WriteToString<128>(Job.GetName()), *WriteToString<32>(Job.GetFunction()));
 	CachePolicy = Policy;
@@ -146,13 +146,13 @@ void FBuildJobContext::SetCachePolicy(ECachePolicy Policy)
 
 void FBuildJobContext::SetBuildPolicy(EBuildPolicy Policy)
 {
-	checkf(!EnumHasAnyFlags(BuildPolicy ^ Policy, EBuildPolicy::SkipCacheGet | EBuildPolicy::SkipCachePut),
-		TEXT("SkipCache flags may not be modified on build policy for build of '%s' by %s. ")
-		TEXT("Flags for skipping cache operations may be set through ECachePolicy."),
+	checkf(!EnumHasAnyFlags(BuildPolicy ^ Policy, EBuildPolicy::Cache),
+		TEXT("Cache flags may not be modified on the build policy for build of '%s' by %s. ")
+		TEXT("Flags for modifying cache operations may be set through ECachePolicy."),
 		*WriteToString<128>(Job.GetName()), *WriteToString<32>(Job.GetFunction()));
-	checkf(!EnumHasAnyFlags(BuildPolicy ^ Policy, EBuildPolicy::SkipBuild | EBuildPolicy::SkipData),
-		TEXT("Skip flags may not be modified on build policy for build of '%s' by %s. ")
-		TEXT("Flags for skipping the build or the data may only be set through the session."),
+	checkf(!EnumHasAnyFlags(BuildPolicy ^ Policy, EBuildPolicy::SkipData),
+		TEXT("SkipData flag may not be modified on the build policy for build of '%s' by %s. ")
+		TEXT("Flags for skipping the data may only be set through the build session."),
 		*WriteToString<128>(Job.GetName()), *WriteToString<32>(Job.GetFunction()));
 	BuildPolicy = Policy;
 }
diff --git a/Engine/Source/Developer/DerivedDataCache/Private/DerivedDataBuildSession.cpp b/Engine/Source/Developer/DerivedDataCache/Private/DerivedDataBuildSession.cpp
--- a/Engine/Source/Developer/DerivedDataCache/Private/DerivedDataBuildSession.cpp
+++ b/Engine/Source/Developer/DerivedDataCache/Private/DerivedDataBuildSession.cpp
@@ -71,7 +71,7 @@ FRequest FBuildSessionInternal::Build(
 	{
 		OnJobComplete = [Definition, OnComplete = MoveTemp(OnComplete)](FBuildJobCompleteParams&& Params)
 		{
-			OnComplete({Definition.GetKey(), MoveTemp(Params.Output), Params.Status});
+			OnComplete({Definition.GetKey(), Params.CacheKey, MoveTemp(Params.Output), Params.BuildStatus, Params.Status});
 		};
 	}
 	Job->Start(Scheduler, Policy, Priority, MoveTemp(OnJobComplete));
@@ -91,7 +91,7 @@ FRequest FBuildSessionInternal::BuildAction(
 	{
 		OnJobComplete = [Action, OnComplete = MoveTemp(OnComplete)](FBuildJobCompleteParams&& Params)
 		{
-			OnComplete({Action.GetKey(), MoveTemp(Params.Output), Params.Status});
+			OnComplete({Action.GetKey(), Params.CacheKey, MoveTemp(Params.Output), Params.BuildStatus, Params.Status});
 		};
 	}
 	Job->Start(Scheduler, Policy, Priority, MoveTemp(OnJobComplete));
diff --git a/Engine/Source/Developer/DerivedDataCache/Public/DerivedDataBuildJob.h b/Engine/Source/Developer/DerivedDataCache/Public/DerivedDataBuildJob.h
--- a/Engine/Source/Developer/DerivedDataCache/Public/DerivedDataBuildJob.h
+++ b/Engine/Source/Developer/DerivedDataCache/Public/DerivedDataBuildJob.h
@@ -17,7 +17,9 @@ namespace UE::DerivedData { class ICache; }
 namespace UE::DerivedData { struct FBuildJobCompleteParams; }
 namespace UE::DerivedData { struct FBuildActionKey; }
 namespace UE::DerivedData { struct FBuildKey; }
+namespace UE::DerivedData { struct FCacheKey; }
 namespace UE::DerivedData { enum class EBuildPolicy : uint8; }
+namespace UE::DerivedData { enum class EBuildStatus : uint32; }
 
 namespace UE::DerivedData
 {
@@ -69,9 +71,13 @@ struct FBuildJobCompleteParams
 {
 	/** Job that is complete. */
 	const IBuildJob& Job;
+	/** Key for the job in the cache. Empty if the build completes before the key is assigned. */
+	const FCacheKey& CacheKey;
 	/** Output for the job that completed or was canceled. */
 	FBuildOutput&& Output;
-	/** Status of the job. */
+	/** Detailed status of the job. */
+	EBuildStatus BuildStatus{};
+	/** Basic status of the job. */
 	EStatus Status = EStatus::Error;
 };
 
diff --git a/Engine/Source/Developer/DerivedDataCache/Public/DerivedDataBuildPolicy.h b/Engine/Source/Developer/DerivedDataCache/Public/DerivedDataBuildPolicy.h
--- a/Engine/Source/Developer/DerivedDataCache/Public/DerivedDataBuildPolicy.h
+++ b/Engine/Source/Developer/DerivedDataCache/Public/DerivedDataBuildPolicy.h
@@ -14,36 +14,58 @@ namespace UE::DerivedData
  * The build policy flags can be combined to support a variety of usage patterns. Examples:
  *
  * Build(Default): allow cache get; build if missing; return every payload.
- * Build(SkipCacheGet): never get from the cache; always build; return every payload.
- * Build(SkipBuild | SkipData): allow cache get; never build; skip every payload.
+ * Build(Build | CacheStore): never query the cache; always build; return every payload.
+ * Build(Cache | CacheSkipData): allow cache get; never build; skip every payload.
  *
  * BuildPayload(Default): allow cache get; build if missing; return one payload.
- * BuildPayload(SkipCacheGet): never get from the cache; always build; return one payload.
- * BuildPayload(SkipBuild | SkipData): allow cache get; never build; skip every payload.
+ * BuildPayload(Build | CacheStore): never get from the cache; always build; return one payload.
+ * BuildPayload(Cache | CacheSkipData): allow cache get; never build; skip every payload.
  */
 enum class EBuildPolicy : uint8
 {
 	/** A value without any flags set. */
 	None            = 0,
 
-	/** Allow local execution. */
-	Local           = 1 << 0,
-	/** Allow remote execution if the function has a registered build worker. */
-	Remote          = 1 << 1,
-
-	/** Skip attempting a cache get before building. */
-	SkipCacheGet    = 1 << 2,
-	/** Skip attempting a cache get before building. */
-	SkipCachePut    = 1 << 3,
-	/** Skip attempting a build when the cache get is missed or skipped. */
-	SkipBuild       = 1 << 4,
+	/** Allow local execution of the build function. */
+	BuildLocal      = 1 << 0,
+	/** Allow remote execution of the build function if it has a registered build worker. */
+	BuildRemote     = 1 << 1,
+	/** Allow local and remote execution of the build function. */
+	Build           = BuildLocal | BuildRemote,
+
+	/** Allow a cache query to avoid having to build. */
+	CacheQuery      = 1 << 2,
+	/** Allow a cache store to persist the build output. */
+	CacheStore      = 1 << 3,
+	/** Allow a cache query and a cache store for the build. */
+	Cache           = CacheQuery | CacheStore,
+
 	/** Skip fetching the payload data from the cache. */
-	SkipData        = 1 << 5,
+	SkipData        = 1 << 4,
 
-	/** Allow cache get+put, local+remote build when missed or skipped, and fetch the payload(s). */
-	Default         = Local | Remote,
+	/** Allow cache query+store, allow local+remote build when missed or skipped, and fetch the payload(s). */
+	Default         = Build | Cache,
 };
 
 ENUM_CLASS_FLAGS(EBuildPolicy);
 
+/** Flags for build request completion callbacks. */
+enum class EBuildStatus : uint32
+{
+	/** A value without any flags set. */
+	None            = 0,
+
+	/** The build function was executed locally. */
+	BuildLocal      = 1 << 0,
+	/** The build function was executed remotely. */
+	BuildRemote     = 1 << 1,
+
+	/** The build made a cache query request. */
+	CacheQuery      = 1 << 2,
+	/** Valid build output was found in the cache. */
+	CacheQueryHit   = 1 << 3,
+};
+
+ENUM_CLASS_FLAGS(EBuildStatus);
+
 } // UE::DerivedData
diff --git a/Engine/Source/Developer/DerivedDataCache/Public/DerivedDataBuildSession.h b/Engine/Source/Developer/DerivedDataCache/Public/DerivedDataBuildSession.h
--- a/Engine/Source/Developer/DerivedDataCache/Public/DerivedDataBuildSession.h
+++ b/Engine/Source/Developer/DerivedDataCache/Public/DerivedDataBuildSession.h
@@ -19,6 +19,7 @@ namespace UE::DerivedData { class FPayload; }
 namespace UE::DerivedData { struct FBuildActionCompleteParams; }
 namespace UE::DerivedData { struct FBuildCompleteParams; }
 namespace UE::DerivedData { struct FBuildPayloadCompleteParams; }
+namespace UE::DerivedData { struct FCacheKey; }
 
 namespace UE::DerivedData
 {
@@ -181,7 +182,10 @@ public:
 struct FBuildCompleteParams
 {
 	/** Key for the build request that completed or was canceled. */
-	FBuildKey Key;
+	const FBuildKey& BuildKey;
+
+	/** Key for the build in the cache. Empty if the build completes before the key is assigned. */
+	const FCacheKey& CacheKey;
 
 	/**
 	 * Output for the build request that completed or was canceled.
@@ -192,7 +196,10 @@ struct FBuildCompleteParams
 	 */
 	FBuildOutput&& Output;
 
-	/** Status of the build request. */
+	/** Detailed status of the build request. */
+	EBuildStatus BuildStatus = EBuildStatus::None;
+
+	/** Basic status of the build request. */
 	EStatus Status = EStatus::Error;
 };
 
@@ -200,7 +207,10 @@ struct FBuildCompleteParams
 struct FBuildActionCompleteParams
 {
 	/** Key for the build action request that completed or was canceled. */
-	FBuildActionKey Key;
+	FBuildActionKey BuildKey;
+
+	/** Key for the build in the cache. Empty if the build completes before the key is assigned. */
+	const FCacheKey& CacheKey;
 
 	/**
 	 * Output for the build action request that completed or was canceled.
@@ -211,7 +221,10 @@ struct FBuildActionCompleteParams
 	 */
 	FBuildOutput&& Output;
 
-	/** Status of the build request. */
+	/** Detailed status of the build request. */
+	EBuildStatus BuildStatus = EBuildStatus::None;
+
+	/** Basic status of the build request. */
 	EStatus Status = EStatus::Error;
 };
 
diff --git a/Engine/Source/Programs/DerivedDataBuildWorker/Private/DerivedDataBuildWorker.cpp b/Engine/Source/Programs/DerivedDataBuildWorker/Private/DerivedDataBuildWorker.cpp
--- a/Engine/Source/Programs/DerivedDataBuildWorker/Private/DerivedDataBuildWorker.cpp
+++ b/Engine/Source/Programs/DerivedDataBuildWorker/Private/DerivedDataBuildWorker.cpp
@@ -233,8 +233,7 @@ bool FBuildWorkerProgram::Build()
 		}
 		else
 		{
-			Builds.Add(Session.BuildAction(Action.Get(), {},
-				EBuildPolicy::Local | EBuildPolicy::SkipCacheGet | EBuildPolicy::SkipCachePut, EPriority::Normal,
+			Builds.Add(Session.BuildAction(Action.Get(), {}, EBuildPolicy::BuildLocal, EPriority::Normal,
 				[this](FBuildActionCompleteParams&& Params) { BuildComplete(MoveTemp(Params)); }));
 		}
 	}
