Add methods for iterating through build functions in the function registry.

diff --git a/Engine/Source/Developer/DerivedDataCache/Private/DerivedDataBuildFunctionRegistry.cpp b/Engine/Source/Developer/DerivedDataCache/Private/DerivedDataBuildFunctionRegistry.cpp
--- a/Engine/Source/Developer/DerivedDataCache/Private/DerivedDataBuildFunctionRegistry.cpp
+++ b/Engine/Source/Developer/DerivedDataCache/Private/DerivedDataBuildFunctionRegistry.cpp
@@ -23,6 +23,7 @@ public:
 
 	const IBuildFunction* FindFunction(FStringView Function) const;
 	FGuid FindFunctionVersion(FStringView Function) const;
+	void IterateFunctionVersions(TFunctionRef<void(FStringView Function, const FGuid& Version)> Visitor) const;
 
 private:
 	void OnModularFeatureRegistered(const FName& Type, IModularFeature* ModularFeature);
@@ -113,6 +114,24 @@ FGuid FBuildFunctionRegistry::FindFunctionVersion(FStringView Function) const
 	return FGuid();
 }
 
+void FBuildFunctionRegistry::IterateFunctionVersions(TFunctionRef<void(FStringView Function, const FGuid& Version)> Visitor) const
+{
+	TArray<TPair<FStringView, FGuid>> FunctionVersions;
+	{
+		FReadScopeLock ReadLock(Lock);
+		FunctionVersions.Reserve(Functions.Num());
+		for (const TPair<FStringView, IBuildFunctionFactory*>& Function : Functions)
+		{
+			FunctionVersions.Emplace(Function.Key, Function.Value->GetFunction().GetVersion());
+		}
+	}
+
+	for (const TPair<FStringView, FGuid>& Function : FunctionVersions)
+	{
+		Visitor(Function.Key, Function.Value);
+	}
+}
+
 IBuildFunctionRegistry* CreateBuildFunctionRegistry()
 {
 	return new FBuildFunctionRegistry();
diff --git a/Engine/Source/Developer/DerivedDataCache/Public/DerivedDataBuildFunctionRegistry.h b/Engine/Source/Developer/DerivedDataCache/Public/DerivedDataBuildFunctionRegistry.h
--- a/Engine/Source/Developer/DerivedDataCache/Public/DerivedDataBuildFunctionRegistry.h
+++ b/Engine/Source/Developer/DerivedDataCache/Public/DerivedDataBuildFunctionRegistry.h
@@ -3,6 +3,7 @@
 #pragma once
 
 #include "Containers/StringFwd.h"
+#include "Templates/Function.h"
 
 struct FGuid;
 
@@ -23,6 +24,8 @@ public:
 	virtual const IBuildFunction* FindFunction(FStringView Function) const = 0;
 	/** Find a function version by name. Returns zero if not found. Safe to call from any thread. */
 	virtual FGuid FindFunctionVersion(FStringView Function) const = 0;
+	/** Iterate the complete list of build function versions. Safe to call from any thread. */
+	virtual void IterateFunctionVersions(TFunctionRef<void(FStringView Name, const FGuid& Version)> Visitor) const = 0;
 };
 
 } // UE::DerivedData
