DDC: Switched to XXH3 for FPayloadId::FromName

#rnx

diff --git a/Engine/Source/Developer/DerivedDataCache/Private/DerivedDataBuild.cpp b/Engine/Source/Developer/DerivedDataCache/Private/DerivedDataBuild.cpp
--- a/Engine/Source/Developer/DerivedDataCache/Private/DerivedDataBuild.cpp
+++ b/Engine/Source/Developer/DerivedDataCache/Private/DerivedDataBuild.cpp
@@ -137,7 +137,7 @@ private:
 	TUniquePtr<IBuildFunctionRegistry> FunctionRegistry{CreateBuildFunctionRegistry()};
 	TUniquePtr<IBuildWorkerRegistry> WorkerRegistry{CreateBuildWorkerRegistry()};
 	TUniquePtr<IBuildScheduler> DefaultScheduler{CreateBuildScheduler()};
-	const FGuid Version{TEXT("ac0574e5-62bd-4c2e-84ec-f2efe48c0fef")};
+	const FGuid Version{TEXT("ed89db6a-335d-45d3-8599-39d3b09f7ad7")};
 };
 
 ///////////////////////////////////////////////////////////////////////////////////////////////////////////////////////
diff --git a/Engine/Source/Developer/DerivedDataCache/Public/DerivedDataPayload.h b/Engine/Source/Developer/DerivedDataCache/Public/DerivedDataPayload.h
--- a/Engine/Source/Developer/DerivedDataCache/Public/DerivedDataPayload.h
+++ b/Engine/Source/Developer/DerivedDataCache/Public/DerivedDataPayload.h
@@ -7,6 +7,7 @@
 #include "Containers/StringConv.h"
 #include "Containers/StringFwd.h"
 #include "Containers/StringView.h"
+#include "Hash/xxhash.h"
 #include "IO/IoHash.h"
 #include "Memory/MemoryView.h"
 #include "String/BytesToHex.h"
@@ -146,7 +147,9 @@ inline FPayloadId FPayloadId::FromHash(const FIoHash& Hash)
 inline FPayloadId FPayloadId::FromName(const FAnsiStringView Name)
 {
 	checkf(!Name.IsEmpty(), TEXT("FPayloadId requires a non-empty name."));
-	return FPayloadId::FromHash(FIoHash::HashBuffer(Name.GetData(), Name.Len()));
+	uint8 HashBytes[16];
+	FXxHash128::HashBuffer(Name.GetData(), Name.Len()).ToByteArray(HashBytes);
+	return FPayloadId(MakeMemoryView(HashBytes, sizeof(ByteArray)));
 }
 
 inline FPayloadId FPayloadId::FromName(const FWideStringView Name)
