DDC: Changed TRequest to cast to IRequest before invoking its functions

This allows the request type to make its functions on IRequest protected or private, or allows use with a request type that hides its IRequest behind an implicit conversion.

Also disabled operator-> and operator* for the base IRequest because it is preferred to use the functions on TRequest to the IRequest interface.

#rnx

diff --git a/Engine/Source/Developer/DerivedDataCache/Public/DerivedDataRequest.h b/Engine/Source/Developer/DerivedDataCache/Public/DerivedDataRequest.h
--- a/Engine/Source/Developer/DerivedDataCache/Public/DerivedDataRequest.h
+++ b/Engine/Source/Developer/DerivedDataCache/Public/DerivedDataRequest.h
@@ -5,6 +5,7 @@
 #include "CoreTypes.h"
 #include "Templates/RefCounting.h"
 #include "Templates/UnrealTemplate.h"
+#include <type_traits>
 
 namespace UE::DerivedData
 {
@@ -62,22 +63,38 @@ enum class EStatus : uint8
 
 ///////////////////////////////////////////////////////////////////////////////////////////////////////////////////////
 
-/** Interface to a request. Use via FRequest to automate reference counting. */
+/** Interface to a request. Use via FRequest or TRequest to automate reference counting. */
 class IRequest
 {
 public:
 	virtual ~IRequest() = default;
 
-	/** Set the priority of the request. */
+	/**
+	 * Set the priority of the request.
+	 *
+	 * @param Priority   The priority, which may be higher or lower than the existing priority.
+	 */
 	virtual void SetPriority(EPriority Priority) = 0;
 
-	/** Cancel the request and invoke its callback. */
+	/**
+	 * Cancel the request and invoke any associated callback.
+	 *
+	 * Does not return until any callback for the request has finished executing.
+	 */
 	virtual void Cancel() = 0;
 
-	/** Block the calling thread until the request and its callback complete. */
+	/**
+	 * Block the calling thread until the request is complete.
+	 *
+	 * Does not return until any callback for the request has finished executing.
+	 */
 	virtual void Wait() = 0;
 
-	/** Poll the request and return true if complete and false otherwise. */
+	/**
+	 * Poll the request and return true if complete and false otherwise.
+	 *
+	 * Does not return true until any callback for the request has finished executing.
+	 */
 	virtual bool Poll() = 0;
 
 	/** Add a reference to the request. */
@@ -124,8 +141,8 @@ public:
 
 	/** Access the referenced request. */
 	inline RequestType* Get() const { return Request; }
-	inline RequestType* operator->() const { return Request; }
-	inline RequestType& operator*() const { return *Request; }
+	inline std::conditional_t<!std::is_same_v<IRequest, RequestType>, RequestType*, void> operator->() const { return Request; }
+	inline std::conditional_t<!std::is_same_v<IRequest, RequestType>, RequestType&, void> operator*() const { return *Request; }
 
 	/**
 	 * Set the priority of the request.
@@ -134,9 +151,9 @@ public:
 	 */
 	inline void SetPriority(EPriority Priority)
 	{
-		if (Request)
+		if (IRequest* LocalRequest = Request)
 		{
-			Request->SetPriority(Priority);
+			LocalRequest->SetPriority(Priority);
 		}
 	}
 
@@ -147,9 +164,9 @@ public:
 	 */
 	inline void Cancel() const
 	{
-		if (Request)
+		if (IRequest* LocalRequest = Request)
 		{
-			Request->Cancel();
+			LocalRequest->Cancel();
 		}
 	}
 
@@ -161,9 +178,9 @@ public:
 	 */
 	inline void Wait() const
 	{
-		if (Request)
+		if (IRequest* LocalRequest = Request)
 		{
-			Request->Wait();
+			LocalRequest->Wait();
 		}
 	}
 
@@ -175,7 +192,11 @@ public:
 	 */
 	inline bool Poll() const
 	{
-		return !Request || Request->Poll();
+		if (IRequest* LocalRequest = Request)
+		{
+			return LocalRequest->Poll();
+		}
+		return true;
 	}
 
 	template <typename OtherRequestType, decltype(ImplicitConv<IRequest*>(DeclVal<OtherRequestType*>()))* = nullptr>
