Add a DDC throttling layer that can be used to approximate Jupiter performance locally by adding latency and limiting bandwidth
This is useful for benchmarking as it reduce noise caused by variable Internet performance

#rnx

diff --git a/Engine/Source/Developer/DerivedDataCache/Private/DerivedDataBackendThrottleWrapper.h b/Engine/Source/Developer/DerivedDataCache/Private/DerivedDataBackendThrottleWrapper.h
new file mode 100644
--- /dev/null
+++ b/Engine/Source/Developer/DerivedDataCache/Private/DerivedDataBackendThrottleWrapper.h
@@ -0,0 +1,262 @@
+// Copyright Epic Games, Inc. All Rights Reserved.
+
+#pragma once
+
+#include "CoreMinimal.h"
+#include "DerivedDataBackendInterface.h"
+#include <atomic>
+
+namespace UE::DerivedData::Backends
+{
+
+/** 
+ * A backend wrapper to limit speed and add latency to each call to approximate a slower backend using a faster one
+ *    1. Reproduce Jupiter timings using a local cache to reduce Internet usage and make benchmarks more reproducible (less noisy)
+ *    2. Reproduce HDD latency and bandwidth even when data is stored on faster drives (i.e. SSDs)
+**/
+class FDerivedDataBackendThrottleWrapper : public FDerivedDataBackendInterface
+{
+public:
+
+	/**
+	 * Constructor
+	 *
+	 * @param	InInnerBackend	Backend to use for storage, my responsibilities are about throttling
+	 */
+	FDerivedDataBackendThrottleWrapper(FDerivedDataBackendInterface* InInnerBackend, uint32 InLatencyMS, uint32 InMaxBytesPerSecond, const TCHAR* InParams)
+		: InnerBackend(InInnerBackend)
+		, Latency(float(InLatencyMS) / 1000.0f)
+		, MaxBytesPerSecond(InMaxBytesPerSecond)
+	{
+		/* Speeds faster than this are considered fast*/
+		const float ConsiderFastAtMS = 10;
+		/* Speeds faster than this are ok. Everything else is slow. This value can be overridden in the ini file */
+		float ConsiderSlowAtMS = 50;
+		FParse::Value(InParams, TEXT("ConsiderSlowAt="), ConsiderSlowAtMS);
+
+		// classify and report on these times
+		if (InLatencyMS < 1)
+		{
+			SpeedClass = ESpeedClass::Local;
+		}
+		else if (InLatencyMS <= ConsiderFastAtMS)
+		{
+			SpeedClass = ESpeedClass::Fast;
+		}
+		else if (InLatencyMS >= ConsiderSlowAtMS)
+		{
+			SpeedClass = ESpeedClass::Slow;
+		}
+		else
+		{
+			SpeedClass = ESpeedClass::Ok;
+		}
+
+		check(InnerBackend);
+	}
+
+	/** Return a name for this interface */
+	virtual FString GetName() const override
+	{
+		return FString::Printf(TEXT("ThrottleWrapper (%s)"), *InnerBackend->GetName());
+	}
+
+	/** return true if this cache is writable **/
+	virtual bool IsWritable() const override
+	{
+		return InnerBackend->IsWritable();
+	}
+
+	/** Returns a class of speed for this interface **/
+	virtual ESpeedClass GetSpeedClass() const override
+	{
+		return (ESpeedClass)FMath::Min((int32)SpeedClass, (int32)InnerBackend->GetSpeedClass());
+	}
+
+	/**
+	 * Synchronous test for the existence of a cache item
+	 *
+	 * @param	CacheKey	Alphanumeric+underscore key of this cache item
+	 * @return				true if the data probably will be found, this can't be guaranteed because of concurrency in the backends, corruption, etc
+	 */
+	virtual bool CachedDataProbablyExists(const TCHAR* CacheKey) override
+	{
+		ThrottlingScope Scope(this);
+
+		return InnerBackend->CachedDataProbablyExists(CacheKey);
+	}
+
+	/**
+	 * Synchronous test for the existence of multiple cache items
+	 *
+	 * @param	CacheKeys	Alphanumeric+underscore key of the cache items
+	 * @return				A bit array with bits indicating whether the data for the corresponding key will probably be found
+	 */
+	virtual TBitArray<> CachedDataProbablyExistsBatch(TConstArrayView<FString> CacheKeys) override
+	{
+		ThrottlingScope Scope(this);
+
+		return InnerBackend->CachedDataProbablyExistsBatch(CacheKeys);
+	}
+
+	/**
+	 * Synchronous retrieve of a cache item
+	 *
+	 * @param	CacheKey	Alphanumeric+underscore key of this cache item
+	 * @param	OutData		Buffer to receive the results, if any were found
+	 * @return				true if any data was found, and in this case OutData is non-empty
+	 */
+	virtual bool GetCachedData(const TCHAR* CacheKey, TArray<uint8>& OutData) override
+	{
+		ThrottlingScope Scope(this, [&OutData]() { return OutData.Num(); });
+
+		return InnerBackend->GetCachedData(CacheKey, OutData);
+	}
+
+	/**
+	 * Asynchronous, fire-and-forget placement of a cache item
+	 *
+	 * @param	CacheKey	Alphanumeric+underscore key of this cache item
+	 * @param	InData		Buffer containing the data to cache, can be destroyed after the call returns, immediately
+	 * @param	bPutEvenIfExists	If true, then do not attempt skip the put even if CachedDataProbablyExists returns true
+	 */
+	virtual EPutStatus PutCachedData(const TCHAR* CacheKey, TArrayView<const uint8> InData, bool bPutEvenIfExists) override
+	{
+		ThrottlingScope Scope(this, [&InData]() { return InData.Num(); });
+
+		return InnerBackend->PutCachedData(CacheKey, InData, bPutEvenIfExists);
+	}
+	
+	virtual void RemoveCachedData(const TCHAR* CacheKey, bool bTransient) override
+	{
+		ThrottlingScope Scope(this);
+
+		return InnerBackend->RemoveCachedData(CacheKey, bTransient);
+	}
+
+	virtual TSharedRef<FDerivedDataCacheStatsNode> GatherUsageStats() const override
+	{
+		TSharedRef<FDerivedDataCacheStatsNode> Inner = InnerBackend->GatherUsageStats();
+		TSharedRef<FDerivedDataCacheStatsNode> Usage = MakeShared<FDerivedDataCacheStatsNode>(this, TEXT("ThrottleWrapper"));
+		Usage->Children.Add(Inner);
+
+		return Usage;
+	}
+
+	virtual bool TryToPrefetch(TConstArrayView<FString> CacheKeys) override
+	{
+		ThrottlingScope Scope(this);
+
+		return InnerBackend->TryToPrefetch(CacheKeys);
+	}
+
+	virtual bool WouldCache(const TCHAR* CacheKey, TArrayView<const uint8> InData) override
+	{
+		return InnerBackend->WouldCache(CacheKey, InData);
+	}
+
+	virtual bool ApplyDebugOptions(FBackendDebugOptions& InOptions) 
+	{ 
+		return InnerBackend->ApplyDebugOptions(InOptions);
+	}
+
+	virtual FRequest Put(
+		TConstArrayView<FCacheRecord> Records,
+		FStringView Context,
+		ECachePolicy Policy,
+		EPriority Priority,
+		FOnCachePutComplete&& OnComplete) override
+	{
+		ThrottlingScope Scope(this);
+
+		return InnerBackend->Put(Records, Context, Policy, Priority, MoveTemp(OnComplete));
+	}
+
+	virtual FRequest Get(
+		TConstArrayView<FCacheKey> Keys,
+		FStringView Context,
+		ECachePolicy Policy,
+		EPriority Priority,
+		FOnCacheGetComplete&& OnComplete) override
+	{
+		ThrottlingScope Scope(this);
+
+		return InnerBackend->Get(Keys, Context, Policy, Priority, MoveTemp(OnComplete));
+	}
+
+	virtual FRequest GetPayload(
+		TConstArrayView<FCachePayloadKey> Keys,
+		FStringView Context,
+		ECachePolicy Policy,
+		EPriority Priority,
+		FOnCacheGetPayloadComplete&& OnComplete) override
+	{
+		ThrottlingScope Scope(this);
+
+		return InnerBackend->GetPayload(Keys, Context, Policy, Priority, MoveTemp(OnComplete));
+	}
+
+	virtual void CancelAll() override
+	{
+		InnerBackend->CancelAll();
+	}
+
+private:
+
+	void EnterThrottlingScope(uint64& PreviousBytesTransferred)
+	{
+		if (Latency > 0)
+		{
+			TRACE_CPUPROFILER_EVENT_SCOPE(ThrottlingLatency);
+			FPlatformProcess::Sleep(Latency);
+		}
+		PreviousBytesTransferred = TotalBytesTransferred.load();
+	}
+
+	void CloseThrottlingScope(double PreviousTime, uint64 PreviousBytesTransferred, TFunction<int64()> GetTransferredBytes = nullptr)
+	{
+		if (MaxBytesPerSecond && GetTransferredBytes)
+		{
+			uint64 BytesTransferred = GetTransferredBytes();
+			// Take into account any other transfer that might have happened during that time from any other thread so we have a global limit
+			uint64 NewBytesTransferred = TotalBytesTransferred += BytesTransferred;
+			double ActualTime = FPlatformTime::Seconds() - PreviousTime;
+			double ExpectedTime = double(NewBytesTransferred - PreviousBytesTransferred) / MaxBytesPerSecond;
+			if (ExpectedTime > ActualTime)
+			{
+				TRACE_CPUPROFILER_EVENT_SCOPE(ThrottlingBandwidth);
+				FPlatformProcess::Sleep(ExpectedTime - ActualTime);
+			}
+		}
+	}
+
+	struct ThrottlingScope
+	{
+		FDerivedDataBackendThrottleWrapper* ThrottleWrapper;
+		TFunction<int64()> GetTransferredBytes;
+		double PreviousTime;
+		uint64 PreviousBytesTransferred;
+
+		ThrottlingScope(FDerivedDataBackendThrottleWrapper* InThrottleWrapper, TFunction<int64()> InGetTransferredBytes = nullptr)
+			: ThrottleWrapper(InThrottleWrapper)
+			, GetTransferredBytes(InGetTransferredBytes)
+			, PreviousTime(FPlatformTime::Seconds())
+		{
+			ThrottleWrapper->EnterThrottlingScope(PreviousBytesTransferred);
+		}
+
+		~ThrottlingScope()
+		{
+			ThrottleWrapper->CloseThrottlingScope(PreviousTime, PreviousBytesTransferred, GetTransferredBytes);
+		}
+	};
+
+	/** Backend to use for storage, my responsibilities are about throttling **/
+	FDerivedDataBackendInterface* InnerBackend;
+	float       Latency;
+	uint32      MaxBytesPerSecond;
+	std::atomic<uint64> TotalBytesTransferred{ 0 };
+	ESpeedClass SpeedClass;
+};
+
+} // UE::DerivedData::Backends
diff --git a/Engine/Source/Developer/DerivedDataCache/Private/DerivedDataBackends.cpp b/Engine/Source/Developer/DerivedDataCache/Private/DerivedDataBackends.cpp
--- a/Engine/Source/Developer/DerivedDataCache/Private/DerivedDataBackends.cpp
+++ b/Engine/Source/Developer/DerivedDataCache/Private/DerivedDataBackends.cpp
@@ -22,6 +22,7 @@
 #include "HierarchicalDerivedDataBackend.h"
 #include "DerivedDataLimitKeyLengthWrapper.h"
 #include "DerivedDataBackendCorruptionWrapper.h"
+#include "DerivedDataBackendThrottleWrapper.h"
 #include "DerivedDataBackendVerifyWrapper.h"
 #include "DerivedDataUtilsInterface.h"
 #include "Misc/EngineBuildSettings.h"
@@ -246,6 +247,21 @@ public:
 				{
 					ParsedNode = ParseZenCache(NodeName, *Entry);
 				}
+				
+				if (ParsedNode)
+				{
+					// Add a throttling layer if parameters are found
+					uint32 LatencyMS = 0;
+					FParse::Value(*Entry, TEXT("LatencyMS="), LatencyMS);
+
+					uint32 MaxBytesPerSecond = 0;
+					FParse::Value(*Entry, TEXT("MaxBytesPerSecond="), MaxBytesPerSecond);
+
+					if (LatencyMS != 0 || MaxBytesPerSecond != 0)
+					{
+						ParsedNode = new FDerivedDataBackendThrottleWrapper(ParsedNode, LatencyMS, MaxBytesPerSecond, *Entry);
+					}
+				}
 			}
 		}
 
