Fix crash in http request batch when the batch contains both a head and get request for the same cache keys
Added a test that could reliably reproduce the problem to validate it is properly fixed

diff --git a/Engine/Source/Developer/DerivedDataCache/Private/HttpDerivedDataBackend.cpp b/Engine/Source/Developer/DerivedDataCache/Private/HttpDerivedDataBackend.cpp
--- a/Engine/Source/Developer/DerivedDataCache/Private/HttpDerivedDataBackend.cpp
+++ b/Engine/Source/Developer/DerivedDataCache/Private/HttpDerivedDataBackend.cpp
@@ -1285,17 +1285,27 @@ private:
 		}
 	}
 
+	// Above result value
+	enum class OpResult : uint8
+	{
+		Ok = 0,			// Op finished succesfully
+		Error = 1,		// Error during op
+		NotFound = 2,	// Key was not found
+		Exists = 3		// Used to indicate head op success
+	};
+
 	// Searches for potentially multiple key requests that are satisfied the given cache key result
 	// Search strategy is exhaustive forward search from the last found entry.  If the results come in ordered the same as the requests,
 	//  and there are no duplicates, the search will be somewhat efficient (still has to do exhaustive searching looking for duplicates).
 	//  If the results are unordered or there are duplicates, search will become more inefficient.
 	struct FRequestSearchHelper
 	{
-		FRequestSearchHelper(TArrayView<FQueuedBatchEntry> InRequests, const FUTF8ToTCHAR& InCacheKey, int32 InEntryIdx, int32 InKeyIdx)
+		FRequestSearchHelper(TArrayView<FQueuedBatchEntry> InRequests, const FUTF8ToTCHAR& InCacheKey, int32 InEntryIdx, int32 InKeyIdx, OpResult InRequestResult)
 			: Requests(InRequests)
 			, CacheKey(InCacheKey)
 			, StartEntryIdx(InEntryIdx)
 			, StartKeyIdx(InKeyIdx)
+			, RequestResult(InRequestResult)
 		{}
 
 		bool FindNext(int32& EntryIdx, int32& KeyIdx)
@@ -1304,7 +1314,12 @@ private:
 			int32 CurrentKeyIdx = KeyIdx;
 			do
 			{
-				if (FCString::Stricmp(Requests[CurrentEntryIdx].CacheKeys[CurrentKeyIdx], CacheKey.Get()) == 0)
+				// Do not match a get request with a head response code (i.e. Exists) 
+				// or a head request with a get response code (i.e. Ok)
+				// if the response code is an error or not found they can be matched to both head or get request it doesn't matter
+				if (((Requests[CurrentEntryIdx].Verb == FHttpRequest::Get) ^ (RequestResult == OpResult::Exists)) &&
+					((Requests[CurrentEntryIdx].Verb == FHttpRequest::Head) ^ (RequestResult == OpResult::Ok)) &&
+					FCString::Stricmp(Requests[CurrentEntryIdx].CacheKeys[CurrentKeyIdx], CacheKey.Get()) == 0)
 				{
 					EntryIdx = CurrentEntryIdx;
 					KeyIdx = CurrentKeyIdx;
@@ -1330,6 +1345,7 @@ private:
 		const FUTF8ToTCHAR& CacheKey;
 		int32 StartEntryIdx;
 		int32 StartKeyIdx;
+		OpResult RequestResult;
 	};
 
 	/**
@@ -1344,15 +1360,6 @@ private:
 		// The expected data stream is structured accordingly
 		// {"JPTR"} {PayloadCount:uint32} {{"JPEE"} {Name:cstr} {Result:uint8} {Hash:IoHash} {Size:uint64} {Payload...}} ...
 
-		// Above result value
-		enum OpResult : uint8
-		{
-			Ok = 0,			// Op finished succesfully
-			Error = 1,		// Error during op
-			NotFound = 2,	// Key was not found
-			Exists = 3		// Used to indicate head op success
-		};
-
 		const TCHAR ResponseErrorMessage[] = TEXT("Malformed response from server.");
 		const ANSICHAR* ProtocolMagic = "JPTR";
 		const ANSICHAR* PayloadMagic = "JPEE";
@@ -1385,13 +1392,18 @@ private:
 			Response += FCStringAnsi::Strlen(PayloadNameA) + 1; //String and zero termination
 			const ANSICHAR* CacheKeyA = FCStringAnsi::Strrchr(PayloadNameA, '.') + 1; // "namespace.bucket.cachekey"
 
+			// Result of the operation is used to match to the appropriate request (i.e. get or head)
+			OpResult PayloadResult = static_cast<OpResult>(*Response);
+			Response += sizeof(uint8);
+
+			const uint8* ResponseRewindMark = Response;
+
 			// Find the payload among the requests.  Payloads may be returned in any order and if the same cache key was part of two requests,
 			// a single payload may satisfy multiple cache keys in multiple requests.
 			FUTF8ToTCHAR CacheKey(CacheKeyA);
-			FRequestSearchHelper RequestSearch(Requests, CacheKey, EntryIdx, KeyIdx);
+			FRequestSearchHelper RequestSearch(Requests, CacheKey, EntryIdx, KeyIdx, PayloadResult);
 			bool bFoundAny = false;
 
-			const uint8* ResponseRewindMark = Response;
 			while (RequestSearch.FindNext(EntryIdx, KeyIdx))
 			{
 				Response = ResponseRewindMark;
@@ -1400,13 +1412,8 @@ private:
 				FQueuedBatchEntry& RequestOp = Requests[EntryIdx];
 				TBitArray<>& bSuccess = *RequestOp.bSuccess;
 
-				// Result of the operation
-				uint8 PayloadResult = *Response;
-				Response += sizeof(uint8);
-
 				switch (PayloadResult)
 				{
-				
 				case OpResult::Ok:
 					{
 						// Payload hash of the following payload data
@@ -1429,13 +1436,14 @@ private:
 							{
 								Response += PayloadSize;
 							}
-							else
+							// Head requests will explicitly have no associcated OutData array
+							else //if (RequestOp.OutDatas.IsValidIndex(KeyIdx))
 							{
 								TArray<uint8>* OutData = RequestOp.OutDatas[KeyIdx];
 
 								OutData->Append(Response, PayloadSize);
 								Response += PayloadSize;
-								// Verify the recieved and parsed payload
+								// Verify the received and parsed payload
 								if (VerifyPayload(PayloadHash, *OutData))
 								{
 									TRACE_COUNTER_ADD(HttpDDC_GetHit, int64(1));
diff --git a/Engine/Source/Developer/DerivedDataCache/Private/Tests/HttpDerivedDataBackendTest.cpp b/Engine/Source/Developer/DerivedDataCache/Private/Tests/HttpDerivedDataBackendTest.cpp
--- a/Engine/Source/Developer/DerivedDataCache/Private/Tests/HttpDerivedDataBackendTest.cpp
+++ b/Engine/Source/Developer/DerivedDataCache/Private/Tests/HttpDerivedDataBackendTest.cpp
@@ -62,15 +62,9 @@ private:
 	}
 };
 
-IMPLEMENT_HTTPDERIVEDDATA_AUTOMATION_TEST(FConcurrentCachedDataProbablyExistsBatch, TEXT(".FHeavyConcurrentCachedDataProbablyExistsBatch"), EAutomationTestFlags::EditorContext | EAutomationTestFlags::ProductFilter)
-bool FConcurrentCachedDataProbablyExistsBatch::RunTest(const FString& Parameters)
+// Helper function to create a number of dummy cache keys for testing
+TArray<FString> CreateTestCacheKeys(UE::DerivedData::Backends::FHttpDerivedDataBackend* InTestBackend, uint32 InNumKeys)
 {
-	UE::DerivedData::Backends::FHttpDerivedDataBackend* TestBackend = GetTestBackend();
-
-	const int32 ParallelTasks = 32;
-	const uint32 Iterations = 20;
-	const uint32 KeysInBatch = 4;
-
 	TArray<FString> Keys;
 	TArray<uint8> KeyContents;
 	KeyContents.Add(42);
@@ -82,12 +76,25 @@ bool FConcurrentCachedDataProbablyExistsBatch::RunTest(const FString& Parameters
 	HashState.GetHash(Hash);
 	const FString HashString = BytesToHex(Hash, FSHA1::DigestSize);
 
-	for (uint32 KeyIndex = 0; KeyIndex < KeysInBatch; ++KeyIndex)
+	for (uint32 KeyIndex = 0; KeyIndex < InNumKeys; ++KeyIndex)
 	{
 		FString NewKey = FString::Printf(TEXT("__AutoTest_Dummy_%u__%s"), KeyIndex, *HashString);
 		Keys.Add(NewKey);
-		TestBackend->PutCachedData(*NewKey, KeyContents, false);
+		InTestBackend->PutCachedData(*NewKey, KeyContents, false);
 	}
+	return Keys;
+}
+
+IMPLEMENT_HTTPDERIVEDDATA_AUTOMATION_TEST(FConcurrentCachedDataProbablyExistsBatch, TEXT(".FConcurrentCachedDataProbablyExistsBatch"), EAutomationTestFlags::EditorContext | EAutomationTestFlags::ProductFilter)
+bool FConcurrentCachedDataProbablyExistsBatch::RunTest(const FString& Parameters)
+{
+	UE::DerivedData::Backends::FHttpDerivedDataBackend* TestBackend = GetTestBackend();
+
+	const int32 ParallelTasks = 32;
+	const uint32 Iterations = 20;
+	const uint32 KeysInBatch = 4;
+
+	TArray<FString> Keys = CreateTestCacheKeys(TestBackend, KeysInBatch);
 
 	std::atomic<uint32> MismatchedResults = 0;
 	ParallelFor(ParallelTasks,
@@ -109,5 +116,40 @@ bool FConcurrentCachedDataProbablyExistsBatch::RunTest(const FString& Parameters
 	return true;
 }
 
+// This test validate that batch requests wont mismatch head and get request for the same keys in the same batch
+IMPLEMENT_HTTPDERIVEDDATA_AUTOMATION_TEST(FConcurrentExistsAndGetForSameKeyBatch, TEXT(".FConcurrentExistsAndGetForSameKeyBatch"), EAutomationTestFlags::EditorContext | EAutomationTestFlags::ProductFilter)
+bool FConcurrentExistsAndGetForSameKeyBatch::RunTest(const FString& Parameters)
+{
+	UE::DerivedData::Backends::FHttpDerivedDataBackend* TestBackend = GetTestBackend();
+
+	const int32 ParallelTasks = 32;
+	const uint32 Iterations = 20;
+	const uint32 KeysInBatch = 4;
+
+	TArray<FString> Keys = CreateTestCacheKeys(TestBackend, KeysInBatch);
+
+	ParallelFor(ParallelTasks,
+		[&](int32 TaskIndex)
+		{
+			for (uint32 Iteration = 0; Iteration < Iterations; ++Iteration)
+			{
+				for (int32 KeyIndex = 0; KeyIndex < Keys.Num(); ++KeyIndex)
+				{
+					if ((TaskIndex % 2) ^ (KeyIndex % 2))
+					{
+						TestBackend->CachedDataProbablyExists(*Keys[KeyIndex]);
+					}
+					else
+					{
+						TArray<uint8> OutData;
+						TestBackend->GetCachedData(*Keys[KeyIndex], OutData);
+					}
+				}
+			}
+		}
+	);
+	return true;
+}
+
 }
 #endif // WITH_DEV_AUTOMATION_TESTS
