Add trace counter increments for more cache interaction prototypes in FileSystemDDC.

This fixes bytes read not matching the reported get hit count, and likely the same for put hits.

diff --git a/Engine/Source/Developer/DerivedDataCache/Private/FileSystemDerivedDataBackend.cpp b/Engine/Source/Developer/DerivedDataCache/Private/FileSystemDerivedDataBackend.cpp
--- a/Engine/Source/Developer/DerivedDataCache/Private/FileSystemDerivedDataBackend.cpp
+++ b/Engine/Source/Developer/DerivedDataCache/Private/FileSystemDerivedDataBackend.cpp
@@ -896,6 +896,7 @@ public:
 			{
 				UE_LOG(LogDerivedDataCache, Verbose, TEXT("%s: Cache put complete for %s from '%.*s'"),
 					*CachePath, *WriteToString<96>(Record.GetKey()), Context.Len(), Context.GetData());
+				TRACE_COUNTER_INCREMENT(FileSystemDDC_PutHit);
 				TRACE_COUNTER_ADD(FileSystemDDC_BytesWritten, MeasureCacheRecord(Record));
 				COOK_STAT(Timer.AddHit(MeasureCacheRecord(Record)));
 				if (OnComplete)
@@ -931,6 +932,7 @@ public:
 			{
 				UE_LOG(LogDerivedDataCache, Verbose, TEXT("%s: Cache hit for %s from '%.*s'"),
 					*CachePath, *WriteToString<96>(Key), Context.Len(), Context.GetData());
+				TRACE_COUNTER_INCREMENT(FileSystemDDC_GetHit);
 				TRACE_COUNTER_ADD(FileSystemDDC_BytesRead, MeasureCacheRecord(Record.Get()));
 				COOK_STAT(Timer.AddHit(MeasureCacheRecord(Record.Get())));
 				if (OnComplete)
@@ -973,6 +975,7 @@ public:
 			{
 				UE_LOG(LogDerivedDataCache, Verbose, TEXT("%s: Cache hit for %s from '%.*s'"),
 					*CachePath, *WriteToString<96>(Key), Context.Len(), Context.GetData());
+				TRACE_COUNTER_INCREMENT(FileSystemDDC_GetHit);
 				TRACE_COUNTER_ADD(FileSystemDDC_BytesRead, Payload.GetRawSize());
 				COOK_STAT(Timer.AddHit(Payload.GetRawSize()));
 				if (OnComplete)
