Fix DDCUtils build error in engine

diff --git a/Engine/Source/Developer/DerivedDataCache/DerivedDataCache.Build.cs b/Engine/Source/Developer/DerivedDataCache/DerivedDataCache.Build.cs
--- a/Engine/Source/Developer/DerivedDataCache/DerivedDataCache.Build.cs
+++ b/Engine/Source/Developer/DerivedDataCache/DerivedDataCache.Build.cs
@@ -17,7 +17,7 @@ public class DerivedDataCache : ModuleRules
 
 		// Internal (NotForLicensees) module
 		var DDCUtilsModule = Path.Combine(EngineDirectory, "Restricted/NotForLicensees/Source/Developer/DDCUtils/DDCUtils.Build.cs");
-		if (File.Exists(DDCUtilsModule))
+		if (Target.Type == TargetType.Editor &&  File.Exists(DDCUtilsModule))
 		{
 			DynamicallyLoadedModuleNames.Add("DDCUtils");
 		}
