DDC: Compress payloads using Oodle Mermaid VeryFast

#rnx

diff --git a/Engine/Source/Developer/DerivedDataCache/Private/DerivedDataBuildJobContext.cpp b/Engine/Source/Developer/DerivedDataCache/Private/DerivedDataBuildJobContext.cpp
--- a/Engine/Source/Developer/DerivedDataCache/Private/DerivedDataBuildJobContext.cpp
+++ b/Engine/Source/Developer/DerivedDataCache/Private/DerivedDataBuildJobContext.cpp
@@ -7,6 +7,7 @@
 #include "DerivedDataBuildPolicy.h"
 #include "DerivedDataBuildPrivate.h"
 #include "DerivedDataCache.h"
+#include "DerivedDataPayloadPrivate.h"
 #include "Hash/Blake3.h"
 #include "Misc/StringBuilder.h"
 #include "UObject/NameTypes.h"
@@ -95,14 +96,20 @@ void FBuildJobContext::AddPayload(const FPayloadId& Id, const FCompressedBuffer&
 	AddPayload(FPayload(Id, Buffer));
 }
 
+void FBuildJobContext::AddPayload(const FPayloadId& Id, const FCompositeBuffer& Buffer)
+{
+	AddPayload(FPayload(Id, FCompressedBuffer::Compress(Buffer, GDefaultCompressor, GDefaultCompressionLevel)));
+}
+
 void FBuildJobContext::AddPayload(const FPayloadId& Id, const FSharedBuffer& Buffer)
 {
-	AddPayload(FPayload(Id, FCompressedBuffer::Compress(NAME_Default, Buffer)));
+	AddPayload(FPayload(Id, FCompressedBuffer::Compress(Buffer, GDefaultCompressor, GDefaultCompressionLevel)));
 }
 
 void FBuildJobContext::AddPayload(const FPayloadId& Id, const FCbObject& Object)
 {
-	AddPayload(FPayload(Id, FCompressedBuffer::Compress(NAME_Default, Object.GetBuffer())));
+	const FCompositeBuffer& Buffer = Object.GetBuffer();
+	AddPayload(FPayload(Id, FCompressedBuffer::Compress(Buffer, GDefaultCompressor, GDefaultCompressionLevel)));
 }
 
 void FBuildJobContext::BeginAsyncBuild()
diff --git a/Engine/Source/Developer/DerivedDataCache/Private/DerivedDataBuildJobContext.h b/Engine/Source/Developer/DerivedDataCache/Private/DerivedDataBuildJobContext.h
--- a/Engine/Source/Developer/DerivedDataCache/Private/DerivedDataBuildJobContext.h
+++ b/Engine/Source/Developer/DerivedDataCache/Private/DerivedDataBuildJobContext.h
@@ -54,6 +54,7 @@ private:
 
 	void AddPayload(const FPayload& Payload) final;
 	void AddPayload(const FPayloadId& Id, const FCompressedBuffer& Buffer) final;
+	void AddPayload(const FPayloadId& Id, const FCompositeBuffer& Buffer) final;
 	void AddPayload(const FPayloadId& Id, const FSharedBuffer& Buffer) final;
 	void AddPayload(const FPayloadId& Id, const FCbObject& Object) final;
 
diff --git a/Engine/Source/Developer/DerivedDataCache/Private/DerivedDataCacheRecord.cpp b/Engine/Source/Developer/DerivedDataCache/Private/DerivedDataCacheRecord.cpp
--- a/Engine/Source/Developer/DerivedDataCache/Private/DerivedDataCacheRecord.cpp
+++ b/Engine/Source/Developer/DerivedDataCache/Private/DerivedDataCacheRecord.cpp
@@ -8,6 +8,7 @@
 #include "DerivedDataCache.h"
 #include "DerivedDataCacheKey.h"
 #include "DerivedDataCachePrivate.h"
+#include "DerivedDataPayloadPrivate.h"
 #include "HAL/CriticalSection.h"
 #include "Misc/ScopeExit.h"
 #include "Misc/ScopeRWLock.h"
@@ -200,16 +201,16 @@ void FCacheRecordBuilderInternal::SetMeta(FCbObject&& InMeta)
 
 FPayloadId FCacheRecordBuilderInternal::SetValue(const FCompositeBuffer& Buffer, const FPayloadId& Id)
 {
-	FCompressedBuffer CompressedBuffer = FCompressedBuffer::Compress(NAME_Default, Buffer);
-	const FPayloadId ValueId = GetOrCreatePayloadId(Id, CompressedBuffer);
-	return SetValue(FPayload(ValueId, MoveTemp(CompressedBuffer)));
+	FCompressedBuffer Compressed = FCompressedBuffer::Compress(Buffer, GDefaultCompressor, GDefaultCompressionLevel);
+	const FPayloadId ValueId = GetOrCreatePayloadId(Id, Compressed);
+	return SetValue(FPayload(ValueId, MoveTemp(Compressed)));
 }
 
 FPayloadId FCacheRecordBuilderInternal::SetValue(const FSharedBuffer& Buffer, const FPayloadId& Id)
 {
-	FCompressedBuffer CompressedBuffer = FCompressedBuffer::Compress(NAME_Default, Buffer);
-	const FPayloadId ValueId = GetOrCreatePayloadId(Id, CompressedBuffer);
-	return SetValue(FPayload(ValueId, MoveTemp(CompressedBuffer)));
+	FCompressedBuffer Compressed = FCompressedBuffer::Compress(Buffer, GDefaultCompressor, GDefaultCompressionLevel);
+	const FPayloadId ValueId = GetOrCreatePayloadId(Id, Compressed);
+	return SetValue(FPayload(ValueId, MoveTemp(Compressed)));
 }
 
 FPayloadId FCacheRecordBuilderInternal::SetValue(const FPayload& Payload)
@@ -228,16 +229,16 @@ FPayloadId FCacheRecordBuilderInternal::SetValue(const FPayload& Payload)
 
 FPayloadId FCacheRecordBuilderInternal::AddAttachment(const FCompositeBuffer& Buffer, const FPayloadId& Id)
 {
-	FCompressedBuffer CompressedBuffer = FCompressedBuffer::Compress(NAME_Default, Buffer);
-	const FPayloadId AttachmentId = GetOrCreatePayloadId(Id, CompressedBuffer);
-	return AddAttachment(FPayload(AttachmentId, MoveTemp(CompressedBuffer)));
+	FCompressedBuffer Compressed = FCompressedBuffer::Compress(Buffer, GDefaultCompressor, GDefaultCompressionLevel);
+	const FPayloadId AttachmentId = GetOrCreatePayloadId(Id, Compressed);
+	return AddAttachment(FPayload(AttachmentId, MoveTemp(Compressed)));
 }
 
 FPayloadId FCacheRecordBuilderInternal::AddAttachment(const FSharedBuffer& Buffer, const FPayloadId& Id)
 {
-	FCompressedBuffer CompressedBuffer = FCompressedBuffer::Compress(NAME_Default, Buffer);
-	const FPayloadId AttachmentId = GetOrCreatePayloadId(Id, CompressedBuffer);
-	return AddAttachment(FPayload(AttachmentId, MoveTemp(CompressedBuffer)));
+	FCompressedBuffer Compressed = FCompressedBuffer::Compress(Buffer, GDefaultCompressor, GDefaultCompressionLevel);
+	const FPayloadId AttachmentId = GetOrCreatePayloadId(Id, Compressed);
+	return AddAttachment(FPayload(AttachmentId, MoveTemp(Compressed)));
 }
 
 FPayloadId FCacheRecordBuilderInternal::AddAttachment(const FPayload& Payload)
diff --git a/Engine/Source/Developer/DerivedDataCache/Private/DerivedDataPayloadPrivate.h b/Engine/Source/Developer/DerivedDataCache/Private/DerivedDataPayloadPrivate.h
new file mode 100644
--- /dev/null
+++ b/Engine/Source/Developer/DerivedDataCache/Private/DerivedDataPayloadPrivate.h
@@ -0,0 +1,14 @@
+// Copyright Epic Games, Inc. All Rights Reserved.
+
+#pragma once
+
+#include "CoreTypes.h"
+#include "Compression/OodleDataCompression.h"
+
+namespace UE::DerivedData
+{
+
+constexpr FOodleDataCompression::ECompressor GDefaultCompressor = FOodleDataCompression::ECompressor::Mermaid;
+constexpr FOodleDataCompression::ECompressionLevel GDefaultCompressionLevel = FOodleDataCompression::ECompressionLevel::VeryFast;
+
+} // UE::DerivedData
diff --git a/Engine/Source/Developer/DerivedDataCache/Public/DerivedDataBuildFunction.h b/Engine/Source/Developer/DerivedDataCache/Public/DerivedDataBuildFunction.h
--- a/Engine/Source/Developer/DerivedDataCache/Public/DerivedDataBuildFunction.h
+++ b/Engine/Source/Developer/DerivedDataCache/Public/DerivedDataBuildFunction.h
@@ -93,6 +93,7 @@ public:
 	/** Adds a payload to the build output. Must have a non-null buffer and a unique ID. */
 	virtual void AddPayload(const FPayload& Payload) = 0;
 	virtual void AddPayload(const FPayloadId& Id, const FCompressedBuffer& Buffer) = 0;
+	virtual void AddPayload(const FPayloadId& Id, const FCompositeBuffer& Buffer) = 0;
 	virtual void AddPayload(const FPayloadId& Id, const FSharedBuffer& Buffer) = 0;
 	virtual void AddPayload(const FPayloadId& Id, const FCbObject& Object) = 0;
 
