Fix thread starvation for request in http backend by using a FIFO for thread waiting for a connection to be freed
Replace the sleep by a wait on an event triggered by the thread freeing the connection
Improve the test to output stats seen below and use the proper amount of threads on low core systems

Automation RunTests System.DerivedDataCache.HttpDerivedDataBackend.ConcurrentCachedDataProbablyExistsBatch

Before
   RPS: 238, AvgLatency: 90.09 ms, MaxLatency: 7.77 s
   RPS: 231, AvgLatency: 87.66 ms, MaxLatency: 7.77 s
   RPS: 245, AvgLatency: 139.07 ms, MaxLatency: 10.08 s

After
   RPS: 634, AvgLatency: 102.97 ms, MaxLatency: 0.26 s
   RPS: 653, AvgLatency: 102.51 ms, MaxLatency: 0.26 s
   RPS: 630, AvgLatency: 102.29 ms, MaxLatency: 0.26 s

#rnx

diff --git a/Engine/Source/Developer/DerivedDataCache/Private/HttpDerivedDataBackend.cpp b/Engine/Source/Developer/DerivedDataCache/Private/HttpDerivedDataBackend.cpp
--- a/Engine/Source/Developer/DerivedDataCache/Private/HttpDerivedDataBackend.cpp
+++ b/Engine/Source/Developer/DerivedDataCache/Private/HttpDerivedDataBackend.cpp
@@ -23,6 +23,7 @@
 #include "Misc/FileHelper.h"
 #include "Misc/ScopeLock.h"
 #include "Misc/SecureHash.h"
+#include "Experimental/Containers/FAAArrayQueue.h"
 #include "Misc/StringBuilder.h"
 #include "ProfilingDebugging/CpuProfilerTrace.h"
 #include "ProfilingDebugging/CountersTrace.h"
@@ -45,6 +46,7 @@
 #endif
 
 #define UE_HTTPDDC_BACKEND_WAIT_INTERVAL 0.01f
+#define UE_HTTPDDC_BACKEND_WAIT_INTERVAL_MS ((uint32)(UE_HTTPDDC_BACKEND_WAIT_INTERVAL*1000))
 #define UE_HTTPDDC_HTTP_REQUEST_TIMEOUT_SECONDS 30L
 #define UE_HTTPDDC_HTTP_REQUEST_TIMOUT_ENABLED 1
 #define UE_HTTPDDC_HTTP_DEBUG 0
@@ -788,6 +790,41 @@ struct FRequestPool
 		return nullptr;
 	}
 
+	class FWaiter : public FThreadSafeRefCountedObject
+	{
+	public:
+		std::atomic<FHttpRequest*> Request{ nullptr };
+
+		FWaiter(FRequestPool* InPool)
+			: Event(FPlatformProcess::GetSynchEventFromPool(true))
+			, Pool(InPool)
+		{
+		}
+		
+		bool Wait(uint32 TimeMS)
+		{
+			return Event->Wait(TimeMS);
+		}
+
+		void Trigger()
+		{
+			Event->Trigger();
+		}
+	private:
+		~FWaiter()
+		{
+			FPlatformProcess::ReturnSynchEventToPool(Event);
+
+			if (Request)
+			{
+				Pool->ReleaseRequestToPool(Request.exchange(nullptr));
+			}
+		}
+
+		FEvent* Event;
+		FRequestPool* Pool;
+	};
+
 	/**
 	 * Block until a request is free. Once a request has been returned it is 
 	 * "owned by the caller and need to release it to the pool when work has been completed.
@@ -796,14 +833,36 @@ struct FRequestPool
 	FHttpRequest* WaitForFreeRequest()
 	{
 		TRACE_CPUPROFILER_EVENT_SCOPE(HttpDDC_WaitForConnPool);
-		FHttpRequest* Request = nullptr;
-		while (true)
+
+		FHttpRequest* Request = GetFreeRequest();
+		if (Request == nullptr)
 		{
-			Request = GetFreeRequest();
-			if (Request != nullptr)
-				break;
-			FPlatformProcess::Sleep(UE_HTTPDDC_BACKEND_WAIT_INTERVAL);
+			// Make it a fair by allowing each thread to register itself in a FIFO
+			// so that the first thread to start waiting is the first one to get a request.
+			FWaiter* Waiter = new FWaiter(this);
+			Waiter->AddRef(); // One ref for the thread that will dequeue
+			Waiter->AddRef(); // One ref for us
+
+			Waiters.enqueue(Waiter);
+
+			while (!Waiter->Wait(UE_HTTPDDC_BACKEND_WAIT_INTERVAL_MS))
+			{
+				// While waiting, allow us to check if a race occurred and a request has been freed
+				// between the time we checked for free requests and the time we queued ourself as a Waiter.
+				if ((Request = GetFreeRequest()) != nullptr)
+				{
+					// We abandon the FWaiter, it will be freed by the next dequeue
+					// and if it has a request, it will be queued back to the pool.
+					Waiter->Release();
+					return Request;
+				}
+			}
+
+			Request = Waiter->Request.exchange(nullptr);
+			Request->Reset();
+			Waiter->Release();
 		}
+		check(Request);
 		return Request;
 	}
 
@@ -817,6 +876,18 @@ struct FRequestPool
 		{
 			if (Pool[i].Request == Request)
 			{
+				// If only 1 user is remaining, we can give it to a waiter
+				// instead of releasing it back to the pool.
+				if (Pool[i].Usage == 1u)
+				{
+					if (FWaiter* Waiter = Waiters.dequeue())
+					{
+						Waiter->Request = Request;
+						Waiter->Trigger();
+						Waiter->Release();
+						return;
+					}
+				}
 				
 				Pool[i].Usage--;
 				return;
@@ -835,7 +906,6 @@ struct FRequestPool
 		{
 			if (Pool[i].Request == Request)
 			{
-
 				Pool[i].Usage = Users;
 				return;
 			}
@@ -852,6 +922,7 @@ private:
 	};
 
 	TArray<FEntry> Pool;
+	FAAArrayQueue<FWaiter> Waiters;
 
 	FRequestPool() = delete;
 };
@@ -1210,7 +1281,10 @@ private:
 			else
 			{
 				// Wait until "driver" has done query
-				Batch.Complete->Wait(~0);
+				{
+					TRACE_CPUPROFILER_EVENT_SCOPE(WaitForMasterOfBatch);
+					Batch.Complete->Wait(~0);
+				}
 
 				// Store away request and signal we are done
 				Request = Batch.Request;
diff --git a/Engine/Source/Developer/DerivedDataCache/Private/Tests/HttpDerivedDataBackendTest.cpp b/Engine/Source/Developer/DerivedDataCache/Private/Tests/HttpDerivedDataBackendTest.cpp
--- a/Engine/Source/Developer/DerivedDataCache/Private/Tests/HttpDerivedDataBackendTest.cpp
+++ b/Engine/Source/Developer/DerivedDataCache/Private/Tests/HttpDerivedDataBackendTest.cpp
@@ -12,6 +12,8 @@
 
 #if WITH_DEV_AUTOMATION_TESTS && WITH_HTTP_DDC_BACKEND
 
+DEFINE_LOG_CATEGORY_STATIC(LogHttpDerivedDataBackendTests, Log, All);
+
 #define TEST_NAME_ROOT TEXT("System.DerivedDataCache.HttpDerivedDataBackend")
 
 #define IMPLEMENT_HTTPDERIVEDDATA_AUTOMATION_TEST( TClass, PrettyName, TFlags ) \
@@ -49,6 +51,93 @@ public:
 	}
 
 protected:
+	void ConcurrentTestWithStats(TFunctionRef<void()> TestFunction, int32 ThreadCount, double Duration)
+	{
+		std::atomic<uint64> Requests{ 0 };
+		std::atomic<uint64> MaxLatency{ 0 };
+		std::atomic<uint64> TotalMS{ 0 };
+		std::atomic<uint64> TotalRequests{ 0 };
+
+		FEvent* StartEvent = FPlatformProcess::GetSynchEventFromPool(true);
+		FEvent* LastEvent = FPlatformProcess::GetSynchEventFromPool(true);
+		std::atomic<double> StopTime{ 0.0 };
+		std::atomic<uint64> ActiveCount{ 0 };
+
+		for (int32 ThreadIndex = 0; ThreadIndex < ThreadCount; ++ThreadIndex)
+		{
+			ActiveCount++;
+			Async(
+				ThreadIndex < FTaskGraphInterface::Get().GetNumWorkerThreads() ? EAsyncExecution::TaskGraph : EAsyncExecution::Thread,
+				[&]()
+				{
+					// No false start, wait until everyone is ready before starting the test
+					StartEvent->Wait();
+
+					while (FPlatformTime::Seconds() < StopTime.load(std::memory_order_relaxed))
+					{
+						uint64 Before = FPlatformTime::Cycles64();
+						TestFunction();
+						uint64 Delta = FPlatformTime::Cycles64() - Before;
+						Requests++;
+						TotalMS += FPlatformTime::ToMilliseconds64(Delta);
+						TotalRequests++;
+
+						// Compare exchange loop until we either succeed to set the maximum value
+						// or we bail out because we don't have the maximum value anymore.
+						while (true)
+						{
+							uint64 Snapshot = MaxLatency.load();
+							if (Delta > Snapshot)
+							{
+								// Only do the exchange if the value has not changed since we confirmed
+								// we had a bigger one.
+								if (MaxLatency.compare_exchange_strong(Snapshot, Delta))
+								{
+									// Exchange succeeded
+									break;
+								}
+							}
+							else
+							{
+								// We don't have the maximum
+								break;
+							}
+						}
+					}
+					
+					if (--ActiveCount == 0)
+					{
+						LastEvent->Trigger();
+					}
+				}
+			);
+		}
+
+		StopTime = FPlatformTime::Seconds() + Duration;
+
+		// GO!
+		StartEvent->Trigger();
+
+		while (FPlatformTime::Seconds() < StopTime)
+		{
+			FPlatformProcess::Sleep(1.0f);
+
+			if (TotalRequests)
+			{
+				UE_LOG(LogHttpDerivedDataBackendTests, Display, TEXT("RPS: %llu, AvgLatency: %.02f ms, MaxLatency: %.02f s"), Requests.exchange(0), double(TotalMS) / TotalRequests, FPlatformTime::ToSeconds(MaxLatency));
+			}
+			else
+			{
+				UE_LOG(LogHttpDerivedDataBackendTests, Display, TEXT("RPS: %llu, AvgLatency: N/A, MaxLatency: %.02f s"), Requests.exchange(0), FPlatformTime::ToSeconds(MaxLatency));
+			}
+		}
+
+		LastEvent->Wait();
+
+		FPlatformProcess::ReturnSynchEventToPool(StartEvent);
+		FPlatformProcess::ReturnSynchEventToPool(LastEvent);
+	}
+
 	UE::DerivedData::Backends::FHttpDerivedDataBackend* GetTestBackend() const
 	{
 		static UE::DerivedData::Backends::FHttpDerivedDataBackend* CachedBackend = FetchTestBackend_Internal();
@@ -62,13 +151,13 @@ private:
 	}
 };
 
-IMPLEMENT_HTTPDERIVEDDATA_AUTOMATION_TEST(FConcurrentCachedDataProbablyExistsBatch, TEXT(".FHeavyConcurrentCachedDataProbablyExistsBatch"), EAutomationTestFlags::EditorContext | EAutomationTestFlags::ProductFilter)
+IMPLEMENT_HTTPDERIVEDDATA_AUTOMATION_TEST(FConcurrentCachedDataProbablyExistsBatch, TEXT(".ConcurrentCachedDataProbablyExistsBatch"), EAutomationTestFlags::EditorContext | EAutomationTestFlags::ProductFilter)
 bool FConcurrentCachedDataProbablyExistsBatch::RunTest(const FString& Parameters)
 {
 	UE::DerivedData::Backends::FHttpDerivedDataBackend* TestBackend = GetTestBackend();
 
-	const int32 ParallelTasks = 32;
-	const uint32 Iterations = 20;
+	const int32  ThreadCount = 64;
+	const double Duration = 10;
 	const uint32 KeysInBatch = 4;
 
 	TArray<FString> Keys;
@@ -90,24 +179,25 @@ bool FConcurrentCachedDataProbablyExistsBatch::RunTest(const FString& Parameters
 	}
 
 	std::atomic<uint32> MismatchedResults = 0;
-	ParallelFor(ParallelTasks,
-		[&](int32 TaskIndex)
+
+	ConcurrentTestWithStats(
+		[&]()
 		{
-			for (uint32 Iteration = 0; Iteration < Iterations; ++Iteration)
+			TConstArrayView<FString> BatchView = MakeArrayView(Keys.GetData(), KeysInBatch);
+			TBitArray<> Result = TestBackend->CachedDataProbablyExistsBatch(BatchView);
+			if (Result.CountSetBits() != BatchView.Num())
 			{
-				TConstArrayView<FString> BatchView = MakeArrayView(Keys.GetData(), KeysInBatch);
-				TBitArray<> Result = TestBackend->CachedDataProbablyExistsBatch(BatchView);
-				if (Result.CountSetBits() != BatchView.Num())
-				{
-					MismatchedResults.fetch_add(BatchView.Num() - Result.CountSetBits(), std::memory_order_relaxed);
-				}
+				MismatchedResults.fetch_add(BatchView.Num() - Result.CountSetBits(), std::memory_order_relaxed);
 			}
-		}
+		},
+		ThreadCount,
+		Duration
 	);
+
 	TestEqual(TEXT("Concurrent calls to CachedDataProbablyExistsBatch for a batch of keys that were put are not reliably found"), MismatchedResults, 0);
 
 	return true;
 }
 
 }
-#endif // WITH_DEV_AUTOMATION_TESTS
+#endif // #if WITH_DEV_AUTOMATION_TESTS && WITH_HTTP_DDC_BACKEND
