Fix double Jupiter download caused by improper usage of async put underneath each backend in hierarchy
  - Remove one copy of the data in the AsyncPut task when used in conjonction with an Inflight cache
  - Make sure each async put in the hierarchical backend can fullfill request from it's Inflight cache
  - Make sure to hit the asyncput backends for exists, trytoprefetch and get queries to avoid repetitive miss while the put is still in flight
  - Remove the code that tried to delete cache miss as it caused many deletes and redownload of the same asset, that duty is fulfilled by the corruption wrapper anyway
  - Fix race condition when mounting and unmounting Pak files or other modifications applied to the hierarchical backend

#rnx

diff --git a/Engine/Source/Developer/DerivedDataCache/Private/DerivedDataBackendAsyncPutWrapper.cpp b/Engine/Source/Developer/DerivedDataCache/Private/DerivedDataBackendAsyncPutWrapper.cpp
--- a/Engine/Source/Developer/DerivedDataCache/Private/DerivedDataBackendAsyncPutWrapper.cpp
+++ b/Engine/Source/Developer/DerivedDataCache/Private/DerivedDataBackendAsyncPutWrapper.cpp
@@ -33,13 +33,17 @@ public:
 	*/
 	FCachePutAsyncWorker(const TCHAR* InCacheKey, TArrayView<const uint8> InData, FDerivedDataBackendInterface* InInnerBackend, bool InbPutEvenIfExists, FDerivedDataBackendInterface* InInflightCache, FThreadSet* InInFilesInFlight, FDerivedDataCacheUsageStats& InUsageStats)
 		: CacheKey(InCacheKey)
-		, Data(InData.GetData(), InData.Num())
 		, InnerBackend(InInnerBackend)
 		, InflightCache(InInflightCache)
 		, FilesInFlight(InInFilesInFlight)
 		, bPutEvenIfExists(InbPutEvenIfExists)
 		, UsageStats(InUsageStats)
 	{
+		// Only make a copy if it's not going to be available from the Inflight cache
+		if (!InInflightCache || !InInflightCache->CachedDataProbablyExists(InCacheKey))
+		{
+			Data = InData;
+		}
 		check(InnerBackend);
 	}
 
@@ -75,6 +79,11 @@ public:
 		}
 		else
 		{
+			if (InflightCache && Data.Num() == 0)
+			{
+				// We verified at construction time that we would be able to get the data from the Inflight cache
+				verify(InflightCache->GetCachedData(*CacheKey, Data));
+			}
 			Status = InnerBackend->PutCachedData(*CacheKey, Data, bPutEvenIfExists);
 			COOK_STAT(Timer.AddHit(Data.Num()));
 		}
@@ -93,6 +102,10 @@ public:
 				}
 				else
 				{
+					if (Data.Num() == 0)
+					{
+						verify(InflightCache->GetCachedData(*CacheKey, Data));
+					}
 					Status = InnerBackend->PutCachedData(*CacheKey, Data, /*bPutEvenIfExists*/ false);
 				}
 			}
@@ -169,6 +182,11 @@ FDerivedDataBackendInterface::ESpeedClass FDerivedDataBackendAsyncPutWrapper::Ge
 	return InnerBackend->GetSpeedClass();
 }
 
+bool FDerivedDataBackendAsyncPutWrapper::BackfillLowerCacheLevels() const
+{
+	return InnerBackend->BackfillLowerCacheLevels();
+}
+
 bool FDerivedDataBackendAsyncPutWrapper::CachedDataProbablyExists(const TCHAR* CacheKey)
 {
 	COOK_STAT(auto Timer = UsageStats.TimeProbablyExists());
@@ -279,7 +297,7 @@ FDerivedDataBackendInterface::EPutStatus FDerivedDataBackendAsyncPutWrapper::Put
 		InflightCache->PutCachedData(CacheKey, InData, true); // temp copy stored in memory while the async task waits to complete
 		COOK_STAT(Timer.AddHit(InData.Num()));
 	}
-
+	
 	UE_LOG(LogDerivedDataCache, Verbose, TEXT("%s queueing %s for put"), *GetName(), CacheKey);
 
 	FDerivedDataBackend::Get().AddToAsyncCompletionCounter(1);
diff --git a/Engine/Source/Developer/DerivedDataCache/Private/DerivedDataBackendAsyncPutWrapper.h b/Engine/Source/Developer/DerivedDataCache/Private/DerivedDataBackendAsyncPutWrapper.h
--- a/Engine/Source/Developer/DerivedDataCache/Private/DerivedDataBackendAsyncPutWrapper.h
+++ b/Engine/Source/Developer/DerivedDataCache/Private/DerivedDataBackendAsyncPutWrapper.h
@@ -78,6 +78,9 @@ public:
 	/** Returns a class of speed for this interface **/
 	virtual ESpeedClass GetSpeedClass() const override;
 
+	/** Return true if hits on this cache should propagate to lower cache level. */
+	virtual bool BackfillLowerCacheLevels() const override;
+
 	/**
 	 * Synchronous test for the existence of a cache item
 	 *
diff --git a/Engine/Source/Developer/DerivedDataCache/Private/HierarchicalDerivedDataBackend.h b/Engine/Source/Developer/DerivedDataCache/Private/HierarchicalDerivedDataBackend.h
--- a/Engine/Source/Developer/DerivedDataCache/Private/HierarchicalDerivedDataBackend.h
+++ b/Engine/Source/Developer/DerivedDataCache/Private/HierarchicalDerivedDataBackend.h
@@ -8,6 +8,7 @@
 #include "DerivedDataCacheRecord.h"
 #include "DerivedDataCacheUsageStats.h"
 #include "DerivedDataBackendAsyncPutWrapper.h"
+#include "Misc/ScopeRWLock.h"
 #include "Templates/UniquePtr.h"
 #include "Containers/ArrayView.h"
 
@@ -54,6 +55,7 @@ public:
 		return ESpeedClass::Local;
 	}
 
+private:
 	void UpdateAsyncInnerBackends()
 	{
 		bIsWritable = false;
@@ -80,19 +82,22 @@ public:
 				bHasRemoteBackends = true;
 			}
 		}
-		if (bIsWritable)
+
+		for (int32 CacheIndex = 0; CacheIndex < InnerBackends.Num(); CacheIndex++)
 		{
-			for (int32 CacheIndex = 0; CacheIndex < InnerBackends.Num(); CacheIndex++)
-			{
-				// async puts to allow us to fill all levels without holding up the engine
-				AsyncPutInnerBackends.Emplace(new FDerivedDataBackendAsyncPutWrapper(Factory, InnerBackends[CacheIndex], false));
-			}
+			// async puts to allow us to fill all levels without holding up the engine
+			// we need to cache inflight puts to avoid having inconsistent miss and redownload on lower cache levels while puts are still async
+			const bool bCacheInFlightPuts = true;
+			AsyncPutInnerBackends.Emplace(new FDerivedDataBackendAsyncPutWrapper(Factory, InnerBackends[CacheIndex], bCacheInFlightPuts));
 		}
 	}
 
+public:
 	/** Adds inner backend. */
 	void AddInnerBackend(FDerivedDataBackendInterface* InInner) 
 	{
+		FWriteScopeLock LockScope(Lock);
+
 		InnerBackends.Add(InInner);
 		AsyncPutInnerBackends.Empty();
 		UpdateAsyncInnerBackends();
@@ -101,6 +106,8 @@ public:
 	/** Removes inner backend. */
 	bool RemoveInnerBackend(FDerivedDataBackendInterface* InInner) 
 	{
+		FWriteScopeLock LockScope(Lock);
+
 		int32 NumRemoved = InnerBackends.Remove(InInner);
 		AsyncPutInnerBackends.Empty();
 		UpdateAsyncInnerBackends();
@@ -122,20 +129,16 @@ public:
 	virtual bool CachedDataProbablyExists(const TCHAR* CacheKey) override
 	{
 		COOK_STAT(auto Timer = UsageStats.TimeProbablyExists());
-		for (int32 CacheIndex = 0; CacheIndex < InnerBackends.Num(); CacheIndex++)
+
+		FReadScopeLock LockScope(Lock);
+
+		for (int32 CacheIndex = 0; CacheIndex < AsyncPutInnerBackends.Num(); CacheIndex++)
 		{
-			if (InnerBackends[CacheIndex]->CachedDataProbablyExists(CacheKey))
+			if (AsyncPutInnerBackends[CacheIndex]->CachedDataProbablyExists(CacheKey))
 			{
 				COOK_STAT(Timer.AddHit(0));
 				return true;
 			}
-			else
-			{
-				if (GVerifyDDC)
-				{
-					ensureMsgf(!AsyncPutInnerBackends[CacheIndex]->CachedDataProbablyExists(CacheKey), TEXT("%s did not exist in sync interface for CachedDataProbablyExists but was found in async wrapper"), CacheKey);
-				}
-			}
 		}
 		return false;
 	}
@@ -149,46 +152,49 @@ public:
 	virtual TBitArray<> CachedDataProbablyExistsBatch(TConstArrayView<FString> CacheKeys) override
 	{
 		COOK_STAT(auto Timer = UsageStats.TimeProbablyExists());
-		check(InnerBackends.Num() > 0);
 
 		TBitArray<> Result;
-
-		for (int32 CacheIndex = 0; CacheIndex < InnerBackends.Num(); ++CacheIndex)
 		{
-			const int32 MissingKeys = CacheKeys.Num() - Result.CountSetBits();
-			if (MissingKeys == 0)
-			{
-				break;
-			}
+			FReadScopeLock LockScope(Lock);
 
-			if (MissingKeys == CacheKeys.Num())
-			{
-				Result = InnerBackends[CacheIndex]->CachedDataProbablyExistsBatch(CacheKeys);
-				check(Result.Num() == CacheKeys.Num());
-			}
-			else
+			check(InnerBackends.Num() > 0);
+			for (int32 CacheIndex = 0; CacheIndex < AsyncPutInnerBackends.Num(); ++CacheIndex)
 			{
-				TArray<FString> RemainingKeys;
-				for (int32 KeyIndex = 0; KeyIndex < CacheKeys.Num(); ++KeyIndex)
+				const int32 MissingKeys = CacheKeys.Num() - Result.CountSetBits();
+				if (MissingKeys == 0)
 				{
-					if (!Result[KeyIndex])
-					{
-						RemainingKeys.Add(CacheKeys[KeyIndex]);
-					}
+					break;
 				}
 
-				TBitArray<>::FIterator ResultIt(Result);
-				TBitArray<> NewResult = InnerBackends[CacheIndex]->CachedDataProbablyExistsBatch(RemainingKeys);
-				check(NewResult.Num() == RemainingKeys.Num());
-				for (TBitArray<>::FConstIterator NewIt(NewResult); NewIt; ++NewIt)
+				if (MissingKeys == CacheKeys.Num())
 				{
-					// Skip to the bit for the next remaining key.
-					while (ResultIt.GetValue())
+					Result = AsyncPutInnerBackends[CacheIndex]->CachedDataProbablyExistsBatch(CacheKeys);
+					check(Result.Num() == CacheKeys.Num());
+				}
+				else
+				{
+					TArray<FString> RemainingKeys;
+					for (int32 KeyIndex = 0; KeyIndex < CacheKeys.Num(); ++KeyIndex)
+					{
+						if (!Result[KeyIndex])
+						{
+							RemainingKeys.Add(CacheKeys[KeyIndex]);
+						}
+					}
+
+					TBitArray<>::FIterator ResultIt(Result);
+					TBitArray<> NewResult = AsyncPutInnerBackends[CacheIndex]->CachedDataProbablyExistsBatch(RemainingKeys);
+					check(NewResult.Num() == RemainingKeys.Num());
+					for (TBitArray<>::FConstIterator NewIt(NewResult); NewIt; ++NewIt)
 					{
+						// Skip to the bit for the next remaining key.
+						while (ResultIt.GetValue())
+						{
+							++ResultIt;
+						}
+						ResultIt.GetValue() = NewIt.GetValue();
 						++ResultIt;
 					}
-					ResultIt.GetValue() = NewIt.GetValue();
-					++ResultIt;
 				}
 			}
 		}
@@ -220,39 +226,43 @@ public:
 		bool bHasFastBackendToWrite = false;
 		bool bHasSlowBackend = false;
 
-		for (FDerivedDataBackendInterface* Interface : InnerBackends)
 		{
-			if (Interface->GetSpeedClass() < ESpeedClass::Fast)
-			{
-				bHasSlowBackend = true;
-			}
-			else
-			{
-				bHasFastBackendToWrite = bHasFastBackendToWrite || Interface->IsWritable();
+			FReadScopeLock LockScope(Lock);
 
-				// Remove keys that exist in a fast backend because we already have fast access to them.
-				TBitArray<> Hits = Interface->CachedDataProbablyExistsBatch(SearchKeys);
-				TArray<FString, TInlineAllocator<16>>::TIterator It = SearchKeys.CreateIterator();
-				It.SetToEnd();
-				for (--It; It; )
+			for (const TUniquePtr<FDerivedDataBackendInterface>& Interface : AsyncPutInnerBackends)
+			{
+				if (Interface->GetSpeedClass() < ESpeedClass::Fast)
+				{
+					bHasSlowBackend = true;
+				}
+				else
 				{
-					if (Hits[It.GetIndex()])
+					bHasFastBackendToWrite = bHasFastBackendToWrite || Interface->IsWritable();
+
+					// Remove keys that exist in a fast backend because we already have fast access to them.
+					TBitArray<> Hits = Interface->CachedDataProbablyExistsBatch(SearchKeys);
+					TArray<FString, TInlineAllocator<16>>::TIterator It = SearchKeys.CreateIterator();
+					It.SetToEnd();
+					for (--It; It; )
 					{
-						// RemoveCurrent will decrement the iterator
-						It.RemoveCurrent();
+						if (Hits[It.GetIndex()])
+						{
+							// RemoveCurrent will decrement the iterator
+							It.RemoveCurrent();
+						}
+						else
+						{
+							--It;
+						}
 					}
-					else
+
+					// No fetch is necessary if every key already exists in a fast backend.
+					if (SearchKeys.IsEmpty())
 					{
-						--It;
+						COOK_STAT(Timer.AddHit(0));
+						return true;
 					}
 				}
-
-				// No fetch is necessary if every key already exists in a fast backend.
-				if (SearchKeys.IsEmpty())
-				{
-					COOK_STAT(Timer.AddHit(0));
-					return true;
-				}
 			}
 		}
 
@@ -278,9 +288,11 @@ public:
 	*/
 	virtual bool WouldCache(const TCHAR* CacheKey, TArrayView<const uint8> InData) override
 	{
-		for (int32 CacheIndex = 0; CacheIndex < InnerBackends.Num(); CacheIndex++)
+		FReadScopeLock LockScope(Lock);
+
+		for (int32 CacheIndex = 0; CacheIndex < AsyncPutInnerBackends.Num(); CacheIndex++)
 		{
-			if (InnerBackends[CacheIndex]->WouldCache(CacheKey, InData))
+			if (AsyncPutInnerBackends[CacheIndex]->WouldCache(CacheKey, InData))
 			{
 				return true;
 			}
@@ -291,6 +303,8 @@ public:
 
 	bool ApplyDebugOptions(FBackendDebugOptions& InOptions) override
 	{
+		FReadScopeLock LockScope(Lock);
+		
 		bool bSuccess = true;
 		for (int32 CacheIndex = 0; CacheIndex < InnerBackends.Num(); CacheIndex++)
 		{
@@ -316,9 +330,11 @@ public:
 	{
 		COOK_STAT(auto Timer = UsageStats.TimeGet());
 
-		for (int32 CacheIndex = 0; CacheIndex < InnerBackends.Num(); CacheIndex++)
+		FReadScopeLock LockScope(Lock);
+
+		for (int32 CacheIndex = 0; CacheIndex < AsyncPutInnerBackends.Num(); CacheIndex++)
 		{
-			FDerivedDataBackendInterface* GetInterface = InnerBackends[CacheIndex];
+			const TUniquePtr<FDerivedDataBackendInterface>& GetInterface = AsyncPutInnerBackends[CacheIndex];
 
 			// just try and get the cached data. It's faster to try and fail than it is to check and succeed. 			
 			if (GetInterface->GetCachedData(CacheKey, OutData))
@@ -334,33 +350,20 @@ public:
 
 						if (MissedCache->IsWritable())
 						{
-							// We want to make sure that the relationship between ProbablyExists and GetCachedData is valid but
-							// only if we have a fast cache. Mismatches are edge cases caused by failed writes or corruption. 
-							// They get handled, so can be left to eventually be rectified by a faster machine
-							bool bFastCache = MissedCache->GetSpeedClass() >= ESpeedClass::Fast;
-							bool bDidExist = bFastCache ? MissedCache->CachedDataProbablyExists(CacheKey) : false;
-							bool bForcePut = false;
-
-							// the cache failed to return data it thinks it has, so clean it up. (todo - can it just be stomped?)
-							if (bDidExist)
-							{
-								MissedCache->RemoveCachedData(CacheKey, /*bTransient=*/ false); // it apparently failed, so lets delete what is there				
-								bForcePut = true;
-							}
-
 							// use the async interface to perform the put
-							AsyncPutInnerBackends[MissedCacheIndex]->PutCachedData(CacheKey, OutData, bForcePut);
-							UE_LOG(LogDerivedDataCache, Verbose, TEXT("Forward-filling cache %s with: %s (%d bytes) (force=%d)"), *MissedCache->GetName(), CacheKey, OutData.Num(), bForcePut);
+							AsyncPutInnerBackends[MissedCacheIndex]->PutCachedData(CacheKey, OutData, false);
+
+							UE_LOG(LogDerivedDataCache, Verbose, TEXT("Forward-filling cache %s with: %s (%d bytes)"), *MissedCache->GetName(), CacheKey, OutData.Num());
 						}
 					}
 
 					// cascade this data to any lower level back ends that may be missing the data
-					if (InnerBackends[CacheIndex]->BackfillLowerCacheLevels())
+					if (GetInterface->BackfillLowerCacheLevels())
 					{
 						// fill in the lower level caches
 						for (int32 PutCacheIndex = CacheIndex + 1; PutCacheIndex < AsyncPutInnerBackends.Num(); PutCacheIndex++)
 						{
-							FDerivedDataBackendInterface* PutBackend = InnerBackends[PutCacheIndex];
+							const TUniquePtr<FDerivedDataBackendInterface>& PutBackend = AsyncPutInnerBackends[PutCacheIndex];
 
 							// If the key is in a distributed cache (e.g. Pak or S3) then don't backfill any further. 
 							bool IsInDistributedCache = !PutBackend->IsWritable() && !PutBackend->BackfillLowerCacheLevels() && PutBackend->CachedDataProbablyExists(CacheKey);
@@ -374,7 +377,7 @@ public:
 								if (bFastCache && PutBackend->IsWritable())
 								{								
 									AsyncPutInnerBackends[PutCacheIndex]->PutCachedData(CacheKey, OutData, false); // we do not need to force a put here
-									UE_LOG(LogDerivedDataCache, Verbose, TEXT("Back-filling cache %s with: %s (%d bytes) (force=%d)"), *PutBackend->GetName(), CacheKey, OutData.Num(), false);
+									UE_LOG(LogDerivedDataCache, Verbose, TEXT("Back-filling cache %s with: %s (%d bytes)"), *PutBackend->GetName(), CacheKey, OutData.Num());
 								}
 							}
 							else
@@ -388,14 +391,6 @@ public:
 				COOK_STAT(Timer.AddHit(OutData.Num()));
 				return true;
 			}
-			else
-			{
-				if (GVerifyDDC)
-				{
-					TArray<uint8> TempData;
-					ensureMsgf(!AsyncPutInnerBackends[CacheIndex]->GetCachedData(CacheKey, TempData), TEXT("CacheKey %s did not exist in sync interface for GetCachedData but was found in async wrapper"), CacheKey);
-				}
-			}
 		}
 		return false;
 	}
@@ -414,6 +409,9 @@ public:
 		{
 			return Status; // no point in continuing down the chain
 		}
+
+		FReadScopeLock LockScope(Lock);
+
 		for (int32 PutCacheIndex = 0; PutCacheIndex < InnerBackends.Num(); PutCacheIndex++)
 		{
 			if (!InnerBackends[PutCacheIndex]->IsWritable() && !InnerBackends[PutCacheIndex]->BackfillLowerCacheLevels() && InnerBackends[PutCacheIndex]->CachedDataProbablyExists(CacheKey))
@@ -443,9 +441,12 @@ public:
 		{
 			return; // no point in continuing down the chain
 		}
-		for (int32 PutCacheIndex = 0; PutCacheIndex < InnerBackends.Num(); PutCacheIndex++)
+
+		FReadScopeLock LockScope(Lock);
+
+		for (int32 PutCacheIndex = 0; PutCacheIndex < AsyncPutInnerBackends.Num(); PutCacheIndex++)
 		{
-			InnerBackends[PutCacheIndex]->RemoveCachedData(CacheKey, bTransient);
+			AsyncPutInnerBackends[PutCacheIndex]->RemoveCachedData(CacheKey, bTransient);
 		}
 	}
 
@@ -454,20 +455,12 @@ public:
 		TSharedRef<FDerivedDataCacheStatsNode> Usage = MakeShared<FDerivedDataCacheStatsNode>(this, TEXT("Hierarchical"));
 		Usage->Stats.Add(TEXT(""), UsageStats);
 
-		// All the inner backends are actually wrapped by AsyncPut backends in writable cases (most cases in practice)
-		if (AsyncPutInnerBackends.Num() > 0)
-		{
-			for (const auto& InnerBackend : AsyncPutInnerBackends)
-			{
-				Usage->Children.Add(InnerBackend->GatherUsageStats());
-			}
-		}
-		else
+		FReadScopeLock LockScope(Lock);
+
+		// All the inner backends are actually wrapped by AsyncPut backends
+		for (const auto& InnerBackend : AsyncPutInnerBackends)
 		{
-			for (auto InnerBackend : InnerBackends)
-			{
-				Usage->Children.Add(InnerBackend->GatherUsageStats());
-			}
+			Usage->Children.Add(InnerBackend->GatherUsageStats());
 		}
 
 		return Usage;
@@ -481,30 +474,35 @@ public:
 		FOnCachePutComplete&& OnComplete) override
 	{
 		TSet<FCacheKey> RecordsOk;
-		for (int32 PutCacheIndex = 0; PutCacheIndex < InnerBackends.Num(); PutCacheIndex++)
+
 		{
-			if (InnerBackends[PutCacheIndex]->IsWritable())
+			FReadScopeLock LockScope(Lock);
+
+			for (int32 PutCacheIndex = 0; PutCacheIndex < InnerBackends.Num(); PutCacheIndex++)
 			{
-				// Every record must put synchronously before switching to async calls.
-				if (RecordsOk.Num() < Records.Num())
+				if (InnerBackends[PutCacheIndex]->IsWritable())
 				{
-					InnerBackends[PutCacheIndex]->Put(Records, Context, Policy, Priority,
-						[&OnComplete, &RecordsOk](FCachePutCompleteParams&& Params)
-						{
-							if (Params.Status == EStatus::Ok)
+					// Every record must put synchronously before switching to async calls.
+					if (RecordsOk.Num() < Records.Num())
+					{
+						InnerBackends[PutCacheIndex]->Put(Records, Context, Policy, Priority,
+							[&OnComplete, &RecordsOk](FCachePutCompleteParams&& Params)
 							{
-								bool bIsAlreadyInSet = false;
-								RecordsOk.FindOrAdd(Params.Key, &bIsAlreadyInSet);
-								if (OnComplete && !bIsAlreadyInSet)
+								if (Params.Status == EStatus::Ok)
 								{
-									OnComplete(MoveTemp(Params));
+									bool bIsAlreadyInSet = false;
+									RecordsOk.FindOrAdd(Params.Key, &bIsAlreadyInSet);
+									if (OnComplete && !bIsAlreadyInSet)
+									{
+										OnComplete(MoveTemp(Params));
+									}
 								}
-							}
-						}).Wait();
-				}
-				else
-				{
-					AsyncPutInnerBackends[PutCacheIndex]->Put(Records, Context, Policy, Priority);
+							}).Wait();
+					}
+					else
+					{
+						AsyncPutInnerBackends[PutCacheIndex]->Put(Records, Context, Policy, Priority);
+					}
 				}
 			}
 		}
@@ -539,47 +537,52 @@ public:
 		TArray<FCacheKey, TInlineAllocator<16>> RemainingKeys(Keys);
 
 		TSet<FCacheKey> KeysOk;
-		for (int32 GetCacheIndex = 0; GetCacheIndex < InnerBackends.Num() && !RemainingKeys.IsEmpty(); ++GetCacheIndex)
+
 		{
-			// Remove SkipData flags when possibly filling other backends because only complete records can be written.
-			const bool bIsLocalGet = InnerBackends[GetCacheIndex]->GetSpeedClass() == ESpeedClass::Local;
-			const bool bFill =
-				(bIsLocalGet && bStoreLocalCopy) ||
-				(bIsLocalGet && bStoreRemote) ||
-				(!bIsLocalGet && bStoreLocal) ||
-				(!bIsLocalGet && bStoreRemote && bHasMultipleRemoteBackends);
-			const ECachePolicy FillPolicy = bFill ? (Policy & ~ECachePolicy::SkipData) : Policy;
-
-			// Block on this because backends in this hierarchy are not expected to be asynchronous.
-			InnerBackends[GetCacheIndex]->Get(RemainingKeys, Context, Policy, Priority,
-				[this, Context, GetCacheIndex, bStoreLocal, bStoreRemote, bStoreLocalCopy, bIsLocalGet, bFill, &OnComplete, &KeysOk](FCacheGetCompleteParams&& Params)
-				{
-					if (Params.Status == EStatus::Ok)
+			FReadScopeLock LockScope(Lock);
+
+			for (int32 GetCacheIndex = 0; GetCacheIndex < InnerBackends.Num() && !RemainingKeys.IsEmpty(); ++GetCacheIndex)
+			{
+				// Remove SkipData flags when possibly filling other backends because only complete records can be written.
+				const bool bIsLocalGet = InnerBackends[GetCacheIndex]->GetSpeedClass() == ESpeedClass::Local;
+				const bool bFill =
+					(bIsLocalGet && bStoreLocalCopy) ||
+					(bIsLocalGet && bStoreRemote) ||
+					(!bIsLocalGet && bStoreLocal) ||
+					(!bIsLocalGet && bStoreRemote && bHasMultipleRemoteBackends);
+				const ECachePolicy FillPolicy = bFill ? (Policy & ~ECachePolicy::SkipData) : Policy;
+
+				// Block on this because backends in this hierarchy are not expected to be asynchronous.
+				InnerBackends[GetCacheIndex]->Get(RemainingKeys, Context, Policy, Priority,
+					[this, Context, GetCacheIndex, bStoreLocal, bStoreRemote, bStoreLocalCopy, bIsLocalGet, bFill, &OnComplete, &KeysOk](FCacheGetCompleteParams&& Params)
 					{
-						if (bFill)
+						if (Params.Status == EStatus::Ok)
 						{
-							for (int32 FillCacheIndex = 0; FillCacheIndex < InnerBackends.Num(); ++FillCacheIndex)
+							if (bFill)
 							{
-								if (GetCacheIndex != FillCacheIndex)
+								for (int32 FillCacheIndex = 0; FillCacheIndex < InnerBackends.Num(); ++FillCacheIndex)
 								{
-									const bool bIsLocalFill = InnerBackends[FillCacheIndex]->GetSpeedClass() == ESpeedClass::Local;
-									if ((bIsLocalFill && bStoreLocal && bIsLocalGet <= bStoreLocalCopy) || (!bIsLocalFill && bStoreRemote))
+									if (GetCacheIndex != FillCacheIndex)
 									{
-										AsyncPutInnerBackends[FillCacheIndex]->Put({Params.Record}, Context);
+										const bool bIsLocalFill = InnerBackends[FillCacheIndex]->GetSpeedClass() == ESpeedClass::Local;
+										if ((bIsLocalFill && bStoreLocal && bIsLocalGet <= bStoreLocalCopy) || (!bIsLocalFill && bStoreRemote))
+										{
+											AsyncPutInnerBackends[FillCacheIndex]->Put({Params.Record}, Context);
+										}
 									}
 								}
 							}
-						}
 
-						KeysOk.Add(Params.Record.GetKey());
-						if (OnComplete)
-						{
-							OnComplete(MoveTemp(Params));
+							KeysOk.Add(Params.Record.GetKey());
+							if (OnComplete)
+							{
+								OnComplete(MoveTemp(Params));
+							}
 						}
-					}
-				}).Wait();
+					}).Wait();
 
-			RemainingKeys.RemoveAll([&KeysOk](const FCacheKey& Key) { return KeysOk.Contains(Key); });
+				RemainingKeys.RemoveAll([&KeysOk](const FCacheKey& Key) { return KeysOk.Contains(Key); });
+			}
 		}
 
 		if (OnComplete)
@@ -603,22 +606,27 @@ public:
 		TArray<FCachePayloadKey, TInlineAllocator<16>> RemainingKeys(Keys);
 
 		TSet<FCachePayloadKey> KeysOk;
-		for (int32 GetCacheIndex = 0; GetCacheIndex < InnerBackends.Num() && !RemainingKeys.IsEmpty(); ++GetCacheIndex)
+
 		{
-			InnerBackends[GetCacheIndex]->GetPayload(RemainingKeys, Context, Policy, Priority,
-				[&OnComplete, &KeysOk](FCacheGetPayloadCompleteParams&& Params)
-				{
-					if (Params.Status == EStatus::Ok)
+			FReadScopeLock LockScope(Lock);
+
+			for (int32 GetCacheIndex = 0; GetCacheIndex < InnerBackends.Num() && !RemainingKeys.IsEmpty(); ++GetCacheIndex)
+			{
+				InnerBackends[GetCacheIndex]->GetPayload(RemainingKeys, Context, Policy, Priority,
+					[&OnComplete, &KeysOk](FCacheGetPayloadCompleteParams&& Params)
 					{
-						KeysOk.Add({Params.Key, Params.Payload.GetId()});
-						if (OnComplete)
+						if (Params.Status == EStatus::Ok)
 						{
-							OnComplete(MoveTemp(Params));
+							KeysOk.Add({Params.Key, Params.Payload.GetId()});
+							if (OnComplete)
+							{
+								OnComplete(MoveTemp(Params));
+							}
 						}
-					}
-				}).Wait();
+					}).Wait();
 
-			RemainingKeys.RemoveAll([&KeysOk](const FCachePayloadKey& Key) { return KeysOk.Contains(Key); });
+				RemainingKeys.RemoveAll([&KeysOk](const FCachePayloadKey& Key) { return KeysOk.Contains(Key); });
+			}
 		}
 
 		if (OnComplete)
@@ -639,6 +647,9 @@ public:
 private:
 	FDerivedDataCacheUsageStats UsageStats;
 
+	/** To avoid race while manipulating the arrays */
+	mutable FRWLock Lock;
+
 	ICacheFactory& Factory;
 	/** Array of backends forming the hierarchical cache...the first element is the fastest cache. **/
 	TArray<FDerivedDataBackendInterface*> InnerBackends;
