Fixed static events initialization crash

diff --git a/Engine/Source/Developer/DerivedDataCache/Private/HttpDerivedDataBackend.cpp b/Engine/Source/Developer/DerivedDataCache/Private/HttpDerivedDataBackend.cpp
--- a/Engine/Source/Developer/DerivedDataCache/Private/HttpDerivedDataBackend.cpp
+++ b/Engine/Source/Developer/DerivedDataCache/Private/HttpDerivedDataBackend.cpp
@@ -3,6 +3,8 @@
 
 #if WITH_HTTP_DDC_BACKEND
 
+#include <memory>
+
 #if PLATFORM_WINDOWS || PLATFORM_HOLOLENS
 #include "Windows/WindowsHWrapper.h"
 #include "Windows/AllowWindowsPlatformTypes.h"
@@ -1122,6 +1124,7 @@ struct FDataRequestHelper
 		{
 			Batches[i].Reserved = 0;
 			Batches[i].Ready = 0;
+			Batches[i].Complete = std::unique_ptr<FEvent, FBatch::EventDeleter>(FPlatformProcess::GetSynchEventFromPool(true));
 		}
 		Initialized = true;
 	}
@@ -1155,12 +1158,20 @@ private:
 
 	struct FBatch
 	{
+		struct EventDeleter
+		{
+			void operator()(FEvent* Event)
+			{
+				FPlatformProcess::ReturnSynchEventToPool(Event);
+			}
+		};
+
 		FQueuedBatchEntry Entries[UE_HTTPDDC_BATCH_SIZE];
 		std::atomic<uint32> Reserved;
 		std::atomic<uint32> Ready;
 		std::atomic<uint32> WeightHint;
 		FHttpRequest* Request;
-		FEventRef Complete{ EEventMode::ManualReset };
+		std::unique_ptr<FEvent, EventDeleter> Complete;
 	};
 
 	FHttpRequest* Request;
