DDC: Moved new stats functions into FDerivedDataCacheInterface to remove module dependency

#rnx

diff --git a/Engine/Source/Developer/DerivedDataCache/Private/DerivedDataCache.cpp b/Engine/Source/Developer/DerivedDataCache/Private/DerivedDataCache.cpp
--- a/Engine/Source/Developer/DerivedDataCache/Private/DerivedDataCache.cpp
+++ b/Engine/Source/Developer/DerivedDataCache/Private/DerivedDataCache.cpp
@@ -149,6 +149,9 @@ namespace DerivedDataCacheCookStats
 }
 #endif
 
+void GatherDerivedDataCacheResourceStats(TArray<FDerivedDataCacheResourceStat>& DDCResourceStats);
+void GatherDerivedDataCacheSummaryStats(FDerivedDataCacheSummaryStats& DDCSummaryStats);
+
 /** Whether we want to verify the DDC (pass in -VerifyDDC on the command line)*/
 bool GVerifyDDC = false;
 
@@ -624,6 +627,16 @@ public:
 	}
 	PRAGMA_ENABLE_DEPRECATION_WARNINGS
 
+	virtual void GatherResourceStats(TArray<FDerivedDataCacheResourceStat>& DDCResourceStats) const override
+	{
+		GatherDerivedDataCacheResourceStats(DDCResourceStats);
+	}
+
+	virtual void GatherSummaryStats(FDerivedDataCacheSummaryStats& DDCSummaryStats) const override
+	{
+		GatherDerivedDataCacheSummaryStats(DDCSummaryStats);
+	}
+
 	/** Get event delegate for data cache notifications */
 	virtual FOnDDCNotification& GetDDCNotificationEvent()
 	{
diff --git a/Engine/Source/Developer/DerivedDataCache/Private/DerivedDataCacheUsageStats.cpp b/Engine/Source/Developer/DerivedDataCache/Private/DerivedDataCacheUsageStats.cpp
--- a/Engine/Source/Developer/DerivedDataCache/Private/DerivedDataCacheUsageStats.cpp
+++ b/Engine/Source/Developer/DerivedDataCache/Private/DerivedDataCacheUsageStats.cpp
@@ -4,7 +4,7 @@
 #include "DerivedDataCacheInterface.h"
 #include "DerivedDataBackendInterface.h"
 
-DERIVEDDATACACHE_API void GatherDerivedDataCacheResourceStats(TArray<FDerivedDataCacheResourceStat>& DDCResourceStats)
+void GatherDerivedDataCacheResourceStats(TArray<FDerivedDataCacheResourceStat>& DDCResourceStats)
 {
 #if ENABLE_COOK_STATS
 	// Temporary Set to allow accumulation by asset type
@@ -102,21 +102,21 @@ DERIVEDDATACACHE_API void GatherDerivedDataCacheResourceStats(TArray<FDerivedDat
 #endif
 }
 
-#if ENABLE_COOK_STATS
-DERIVEDDATACACHE_API void GatherDerivedDataCacheSummaryStats(TArray<FDerivedDataCacheSummaryStat>& DDCSummaryStats)
+void GatherDerivedDataCacheSummaryStats(FDerivedDataCacheSummaryStats& DDCSummaryStats)
 {
+#if ENABLE_COOK_STATS
 	/** this functor will take a collected cooker stat and log it out using some custom formatting based on known stats that are collected.. */
 	auto LogStatsFunc = [&DDCSummaryStats]
 	(const FString& StatName, const TArray<FCookStatsManager::StringKeyValue>& StatAttributes)
 	{
 		if (StatName == TEXT("DDC.Summary"))
 		{
-			DDCSummaryStats.Append(StatAttributes);
+			DDCSummaryStats.Stats.Append(StatAttributes);
 		}
 	};
 
 	// Grab the DDC stats
 	FCookStatsManager::LogCookStats(LogStatsFunc);
-}
 #endif
+}
 
diff --git a/Engine/Source/Developer/DerivedDataCache/Public/DerivedDataCacheInterface.h b/Engine/Source/Developer/DerivedDataCache/Public/DerivedDataCacheInterface.h
--- a/Engine/Source/Developer/DerivedDataCache/Public/DerivedDataCacheInterface.h
+++ b/Engine/Source/Developer/DerivedDataCache/Public/DerivedDataCacheInterface.h
@@ -12,6 +12,8 @@
 class FDerivedDataCacheUsageStats;
 class FDerivedDataCacheStatsNode;
 class IDDCCleanup;
+struct FDerivedDataCacheResourceStat;
+struct FDerivedDataCacheSummaryStats;
 
 /** 
  * Interface for the derived data cache
@@ -292,6 +294,10 @@ public:
 	UE_DEPRECATED(4.27, "This overload of GatherUsageStats is temporary. Please use other overload.")
 	virtual TSharedRef<FDerivedDataCacheStatsNode> GatherUsageStats() const = 0;
 
+	virtual void GatherResourceStats(TArray<FDerivedDataCacheResourceStat>& DDCResourceStats) const = 0;
+
+	virtual void GatherSummaryStats(FDerivedDataCacheSummaryStats& DDCSummaryStats) const = 0;
+
 	//-----------------------
 	// Notification Interface
 	//-----------------------
diff --git a/Engine/Source/Developer/DerivedDataCache/Public/DerivedDataCacheUsageStats.h b/Engine/Source/Developer/DerivedDataCache/Public/DerivedDataCacheUsageStats.h
--- a/Engine/Source/Developer/DerivedDataCache/Public/DerivedDataCacheUsageStats.h
+++ b/Engine/Source/Developer/DerivedDataCache/Public/DerivedDataCacheUsageStats.h
@@ -219,12 +219,10 @@ struct FDerivedDataCacheResourceStatKeyFuncs : BaseKeyFuncs<FDerivedDataCacheRes
 	static uint32 GetKeyHash(const FString& Key) { return GetTypeHash(Key); }
 };
 
-DERIVEDDATACACHE_API void GatherDerivedDataCacheResourceStats(TArray<FDerivedDataCacheResourceStat>& DDCResourceStats);
+COOK_STAT(using FDerivedDataCacheSummaryStat = FCookStatsManager::StringKeyValue);
 
-#if ENABLE_COOK_STATS
-
-using FDerivedDataCacheSummaryStat = FCookStatsManager::StringKeyValue;
-DERIVEDDATACACHE_API void GatherDerivedDataCacheSummaryStats(TArray<FDerivedDataCacheSummaryStat>& DDCSummaryStats);
-
-#endif
+struct FDerivedDataCacheSummaryStats
+{
+	COOK_STAT(TArray<FDerivedDataCacheSummaryStat> Stats);
+};
 
diff --git a/Engine/Source/Editor/UnrealEd/Private/DDC/SDDCInformation.cpp b/Engine/Source/Editor/UnrealEd/Private/DDC/SDDCInformation.cpp
--- a/Engine/Source/Editor/UnrealEd/Private/DDC/SDDCInformation.cpp
+++ b/Engine/Source/Editor/UnrealEd/Private/DDC/SDDCInformation.cpp
@@ -91,7 +91,7 @@ TSharedRef<SWidget> SDDCInformation::GetAssetGrid()
 	TArray<FDerivedDataCacheResourceStat> DDCResourceStats;
 
 	// Grab the resource stats
-	GatherDerivedDataCacheResourceStats(DDCResourceStats);
+	GetDerivedDataCacheRef().GatherResourceStats(DDCResourceStats);
 
 	// Sort results on asscending Load size
 	DDCResourceStats.Sort([](const FDerivedDataCacheResourceStat& LHS, const FDerivedDataCacheResourceStat& RHS) { return LHS.LoadSizeMB > RHS.LoadSizeMB; });
@@ -326,7 +326,7 @@ TSharedRef<SWidget> SDDCInformation::GetSummaryGrid()
 	TArray<FDerivedDataCacheResourceStat> DDCResourceStats;
 
 	// Grab the latest resource stats
-	GatherDerivedDataCacheResourceStats(DDCResourceStats);
+	GetDerivedDataCacheRef().GatherResourceStats(DDCResourceStats);
 
 	FDerivedDataCacheResourceStat DDCResourceStatsTotal(TEXT("Total"));
 
diff --git a/Engine/Source/Editor/UnrealEd/UnrealEd.Build.cs b/Engine/Source/Editor/UnrealEd/UnrealEd.Build.cs
--- a/Engine/Source/Editor/UnrealEd/UnrealEd.Build.cs
+++ b/Engine/Source/Editor/UnrealEd/UnrealEd.Build.cs
@@ -207,7 +207,6 @@ public class UnrealEd : ModuleRules
 				"MaterialBaking",
 				"CookOnTheFly",
 				"BazelExecutor",
-				"DerivedDataCache",
 			}
 		);
 
diff --git a/Engine/Source/Runtime/Engine/Engine.Build.cs b/Engine/Source/Runtime/Engine/Engine.Build.cs
--- a/Engine/Source/Runtime/Engine/Engine.Build.cs
+++ b/Engine/Source/Runtime/Engine/Engine.Build.cs
@@ -61,10 +61,6 @@ public class Engine : ModuleRules
 					"EditorAnalyticsSession",
 				}
 			);
-			if (Target.Type != TargetType.Server)
-			{
-				PrivateDependencyModuleNames.Add("DerivedDataCache");
-			}
 		}
 
 		PublicDependencyModuleNames.AddRange(
diff --git a/Engine/Source/Runtime/Engine/Private/Analytics/StudioAnalytics.cpp b/Engine/Source/Runtime/Engine/Private/Analytics/StudioAnalytics.cpp
--- a/Engine/Source/Runtime/Engine/Private/Analytics/StudioAnalytics.cpp
+++ b/Engine/Source/Runtime/Engine/Private/Analytics/StudioAnalytics.cpp
@@ -18,6 +18,7 @@
 #include "Interfaces/IAnalyticsProvider.h"
 #include "Templates/SharedPointer.h"
 #include "HAL/PlatformProcess.h"
+#include "DerivedDataCacheInterface.h"
 #include "DerivedDataCacheUsageStats.h"
 
 bool FStudioAnalytics::bInitialized = false;
@@ -170,7 +171,7 @@ void FStudioAnalytics::FireEvent_Loading(const FString& LoadingName, double Seco
 		TArray<FDerivedDataCacheResourceStat> DDCResourceStats;
 
 		// Grab the latest resource stats
-		GatherDerivedDataCacheResourceStats(DDCResourceStats);
+		GetDerivedDataCacheRef().GatherResourceStats(DDCResourceStats);
 
 		FDerivedDataCacheResourceStat DDCResourceStatsTotal(TEXT("Total"));
 
@@ -217,12 +218,11 @@ void FStudioAnalytics::FireEvent_Loading(const FString& LoadingName, double Seco
 			}
 		}
 
-		TArray<FDerivedDataCacheSummaryStat> DDCSummaryStats;
-
 		// Grab the summary stats
-		GatherDerivedDataCacheSummaryStats(DDCSummaryStats);
+		FDerivedDataCacheSummaryStats DDCSummaryStats;
+		GetDerivedDataCacheRef().GatherSummaryStats(DDCSummaryStats);
 		
-		for (const FDerivedDataCacheSummaryStat& Stat : DDCSummaryStats)
+		for (const FDerivedDataCacheSummaryStat& Stat : DDCSummaryStats.Stats)
 		{
 			FString FormattedAttrName = "DDC.Summary." + Stat.Key;
 			Attributes.Emplace(FormattedAttrName, Stat.Value);
