Remove prototype "ZenHttp" DDC client code - this was never meant to propagate to main

diff --git a/Engine/Source/Developer/DerivedDataCache/Private/DerivedDataBackends.cpp b/Engine/Source/Developer/DerivedDataCache/Private/DerivedDataBackends.cpp
--- a/Engine/Source/Developer/DerivedDataCache/Private/DerivedDataBackends.cpp
+++ b/Engine/Source/Developer/DerivedDataCache/Private/DerivedDataBackends.cpp
@@ -17,7 +17,6 @@
 #include "DerivedDataBackendAsyncPutWrapper.h"
 #include "PakFileDerivedDataBackend.h"
 #include "S3DerivedDataBackend.h"
-#include "ZenKvDerivedDataBackend.h"
 #include "ZenDerivedDataBackend.h"
 #include "HierarchicalDerivedDataBackend.h"
 #include "DerivedDataLimitKeyLengthWrapper.h"
@@ -239,10 +238,6 @@ public:
 				{
 					ParsedNode = ParseHttpCache(NodeName, *Entry);
 				}
-				else if (NodeType == TEXT("ZenHttp"))
-				{
-					ParsedNode = ParseZenHttpCache(NodeName, *Entry);
-				}
 				else if (NodeType == TEXT("Zen"))
 				{
 					ParsedNode = ParseZenCache(NodeName, *Entry);
@@ -865,41 +860,7 @@ public:
 	}
 
 	/**
-	 * Creates a Zen / HTTP data cache interface (DDC1 style)
-	 */
-	FDerivedDataBackendInterface * ParseZenHttpCache(const TCHAR * NodeName, const TCHAR * Entry)
-	{
-#if WITH_ZEN_DDC_BACKEND
-		FString ServiceUrl;
-		if (!FParse::Value(Entry, TEXT("Host="), ServiceUrl))
-		{
-			UE_LOG(LogDerivedDataCache, Error, TEXT("Node %s does not specify 'Host'."), NodeName);
-			return nullptr;
-		}
-
-		FString Namespace;
-		if (!FParse::Value(Entry, TEXT("Namespace="), Namespace))
-		{
-			Namespace = FApp::GetProjectName();
-			UE_LOG(LogDerivedDataCache, Warning, TEXT("Node %s does not specify 'Namespace', falling back to '%s'"), NodeName, *Namespace);
-		}
-
-		FZenHttpDerivedDataBackend* backend = new FZenHttpDerivedDataBackend(*this, *ServiceUrl, *Namespace);
-		if (!backend->IsUsable())
-		{
-			UE_LOG(LogDerivedDataCache, Warning, TEXT("%s could not contact the service (%s), will not use it."), NodeName, *ServiceUrl);
-			delete backend;
-			return nullptr;
-		}
-		return backend;
-#else
-		UE_LOG(LogDerivedDataCache, Warning, TEXT("ZenHTTP backend is not yet supported in the current build configuration."));
-		return nullptr;
-#endif
-	}
-
-	/**
-	 * Creates a Zen data cache interface
+	 * Creates a Zen structured data cache interface
 	 */
 	FDerivedDataBackendInterface* ParseZenCache(const TCHAR* NodeName, const TCHAR* Entry)
 	{
diff --git a/Engine/Source/Developer/DerivedDataCache/Private/ZenKvDerivedDataBackend.cpp b/Engine/Source/Developer/DerivedDataCache/Private/ZenKvDerivedDataBackend.cpp
deleted file mode 100644
diff --git a/Engine/Source/Developer/DerivedDataCache/Private/ZenKvDerivedDataBackend.h b/Engine/Source/Developer/DerivedDataCache/Private/ZenKvDerivedDataBackend.h
deleted file mode 100644
