General HTTP zen client implementation simplification

* Rename FRequest -> FZenHttpRequest to make it easier to move the implementation elsewhere
* Simplified implementation of PUT/HEAD/DELETE operations by turning template functions into plain member functions to make it easier to hide the implementation details when the class is moved into a header

diff --git a/Engine/Source/Developer/DerivedDataCache/Private/ZenDerivedDataBackend.cpp b/Engine/Source/Developer/DerivedDataCache/Private/ZenDerivedDataBackend.cpp
--- a/Engine/Source/Developer/DerivedDataCache/Private/ZenDerivedDataBackend.cpp
+++ b/Engine/Source/Developer/DerivedDataCache/Private/ZenDerivedDataBackend.cpp
@@ -67,10 +67,6 @@ namespace zen {
 		CompactBinaryPackage
 	};
 
-/**
- * Encapsulation for access token shared by all requests.
- */
-
 	/**
 	 * Minimal HTTP request type wrapping CURL without the need for managers. This request
 	 * is written to allow reuse of request objects, in order to allow connections to be reused.
@@ -78,7 +74,7 @@ namespace zen {
 	 * CURL has a global library initialization (curl_global_init). We rely on this happening in
 	 * the Online/HTTP library which is a dependency on this module.
 	 */
-	class FRequest
+	class FZenHttpRequest
 	{
 	public:
 		/**
@@ -104,7 +100,7 @@ namespace zen {
 			FailedTimeout
 		};
 
-		FRequest(const TCHAR* InDomain, bool bInLogErrors)
+		FZenHttpRequest(const TCHAR* InDomain, bool bInLogErrors)
 			: bLogErrors(bInLogErrors)
 			, Domain(InDomain)
 		{
@@ -112,7 +108,7 @@ namespace zen {
 			Reset();
 		}
 
-		~FRequest()
+		~FZenHttpRequest()
 		{
 			curl_easy_cleanup(Curl);
 		}
@@ -146,7 +142,7 @@ namespace zen {
 			//curl_easy_setopt(Curl, CURLOPT_UPLOAD_BUFFERSIZE, 256 * 1024L);
 			// Response functions
 			curl_easy_setopt(Curl, CURLOPT_HEADERDATA, this);
-			curl_easy_setopt(Curl, CURLOPT_HEADERFUNCTION, &FRequest::StaticWriteHeaderFn);
+			curl_easy_setopt(Curl, CURLOPT_HEADERFUNCTION, &FZenHttpRequest::StaticWriteHeaderFn);
 			curl_easy_setopt(Curl, CURLOPT_WRITEDATA, this);
 			curl_easy_setopt(Curl, CURLOPT_WRITEFUNCTION, StaticWriteBodyFn);
 			// Rewind method, handle special error case where request need to rewind data stream
@@ -190,71 +186,39 @@ namespace zen {
 		}
 
 		/**
-		 * Upload buffer using the request, using either "Put" or "Post" verbs.
+		 * Upload buffer using the request, using PUT verb
 		 * @param Uri Url to use.
 		 * @param Buffer Data to upload
 		 * @return Result of the request
 		 */
-		template<RequestVerb V>
-		Result PerformBlockingUpload(const TCHAR* Uri, const FCompositeBuffer& Buffer, EContentType ContentType)
+		Result PerformBlockingPut(const TCHAR* Uri, const FCompositeBuffer& Buffer, EContentType ContentType)
 		{
-			static_assert(V == Put || V == Post || V == PostJson, "Upload should use either Put or Post verbs.");
-
 			uint32 ContentLength = 0u;
 
-			if (V == Put)
-			{
-				ContentLength = Buffer.GetSize();
-				curl_easy_setopt(Curl, CURLOPT_UPLOAD, 1L);
-				curl_easy_setopt(Curl, CURLOPT_INFILESIZE, ContentLength);
-				curl_easy_setopt(Curl, CURLOPT_READDATA, this);
-				curl_easy_setopt(Curl, CURLOPT_READFUNCTION, StaticReadFn);
+			ContentLength = Buffer.GetSize();
+			curl_easy_setopt(Curl, CURLOPT_UPLOAD, 1L);
+			curl_easy_setopt(Curl, CURLOPT_INFILESIZE, ContentLength);
+			curl_easy_setopt(Curl, CURLOPT_READDATA, this);
+			curl_easy_setopt(Curl, CURLOPT_READFUNCTION, StaticReadFn);
 
-				switch (ContentType)
-				{
-				case EContentType::Binary:
-					Headers.Add(FString(TEXT("Content-Type: application/octet-stream")));
-					break;
-				case EContentType::CompactBinary:
-					Headers.Add(FString(TEXT("Content-Type: application/x-ue-cb")));
-					break;
-				case EContentType::CompactBinaryPackage:
-					Headers.Add(FString(TEXT("Content-Type: application/x-ue-cbpkg")));
-					break;
-				default:
-					checkNoEntry();
-					break;
-				}
-				ReadDataView = &Buffer;
-			}
-			else if (V == Post || V == PostJson)
+			switch (ContentType)
 			{
-				ContentLength = Buffer.GetSize();
-				curl_easy_setopt(Curl, CURLOPT_POST, 1L);
-				curl_easy_setopt(Curl, CURLOPT_INFILESIZE, ContentLength);
-				curl_easy_setopt(Curl, CURLOPT_READDATA, this);
-				curl_easy_setopt(Curl, CURLOPT_READFUNCTION, StaticReadFn);
-
-				switch (ContentType)
-				{
-				case EContentType::Binary:
-					Headers.Add(V == Post ? FString(TEXT("Content-Type: application/x-www-form-urlencoded")) : FString(TEXT("Content-Type: application/json")));
-					break;
-				case EContentType::CompactBinary:
-					Headers.Add(FString(TEXT("Content-Type: application/x-ue-cb")));
-					break;
-				case EContentType::CompactBinaryPackage:
-					Headers.Add(FString(TEXT("Content-Type: application/x-ue-cbpkg")));
-					break;
-				default:
-					checkNoEntry();
-					break;
-				}
-
-				ReadDataView = &Buffer;
+			case EContentType::Binary:
+				Headers.Add(FString(TEXT("Content-Type: application/octet-stream")));
+				break;
+			case EContentType::CompactBinary:
+				Headers.Add(FString(TEXT("Content-Type: application/x-ue-cb")));
+				break;
+			case EContentType::CompactBinaryPackage:
+				Headers.Add(FString(TEXT("Content-Type: application/x-ue-cbpkg")));
+				break;
+			default:
+				checkNoEntry();
+				break;
 			}
+			ReadDataView = &Buffer;
 
-			return PerformBlocking(Uri, V, ContentLength);
+			return PerformBlocking(Uri, RequestVerb::Put, ContentLength);
 		}
 
 		/**
@@ -300,22 +264,24 @@ namespace zen {
 		 * @param Uri Url to use.
 		 * @return Result of the request
 		 */
-		template<RequestVerb V>
-		Result PerformBlockingQuery(const TCHAR* Uri)
+		Result PerformBlockingHead(const TCHAR* Uri)
 		{
-			static_assert(V == Head || V == Delete, "Queries should use either Head or Delete verbs.");
+			curl_easy_setopt(Curl, CURLOPT_NOBODY, 1L);
 
-			if (V == Delete)
-			{
-				curl_easy_setopt(Curl, CURLOPT_POST, 1L);
-				curl_easy_setopt(Curl, CURLOPT_CUSTOMREQUEST, "DELETE");
-			}
-			else if (V == Head)
-			{
-				curl_easy_setopt(Curl, CURLOPT_NOBODY, 1L);
-			}
+			return PerformBlocking(Uri, RequestVerb::Head, 0u);
+		}
+
+		/**
+		 * Query an url using the request. Queries can use either "Head" or "Delete" verbs.
+		 * @param Uri Url to use.
+		 * @return Result of the request
+		 */
+		Result PerformBlockingDelete(const TCHAR* Uri)
+		{
+			curl_easy_setopt(Curl, CURLOPT_POST, 1L);
+			curl_easy_setopt(Curl, CURLOPT_CUSTOMREQUEST, "DELETE");
 
-			return PerformBlocking(Uri, V, 0u);
+			return PerformBlocking(Uri, RequestVerb::Delete, 0u);
 		}
 
 		/**
@@ -395,7 +361,6 @@ namespace zen {
 		}
 
 	private:
-
 		CURL* Curl;
 		CURLcode				CurlResult;
 		long					ResponseCode;
@@ -564,7 +529,7 @@ namespace zen {
 
 		static size_t StaticDebugCallback(CURL* Handle, curl_infotype DebugInfoType, char* DebugInfo, size_t DebugInfoSize, void* UserData)
 		{
-			FRequest* Request = static_cast<FRequest*>(UserData);
+			FZenHttpRequest* Request = static_cast<FZenHttpRequest*>(UserData);
 
 			switch (DebugInfoType)
 			{
@@ -604,7 +569,7 @@ namespace zen {
 
 		static size_t StaticReadFn(void* Ptr, size_t SizeInBlocks, size_t BlockSizeInBytes, void* UserData)
 		{
-			FRequest* Request = static_cast<FRequest*>(UserData);
+			FZenHttpRequest* Request = static_cast<FZenHttpRequest*>(UserData);
 			check(Request->ReadDataView);
 			const FCompositeBuffer& ReadDataView = *Request->ReadDataView;
 
@@ -621,7 +586,7 @@ namespace zen {
 
 		static size_t StaticWriteHeaderFn(void* Ptr, size_t SizeInBlocks, size_t BlockSizeInBytes, void* UserData)
 		{
-			FRequest* Request = static_cast<FRequest*>(UserData);
+			FZenHttpRequest* Request = static_cast<FZenHttpRequest*>(UserData);
 			const size_t WriteSize = SizeInBlocks * BlockSizeInBytes;
 			TArray<uint8>* WriteHeaderBufferPtr = Request->WriteHeaderBufferPtr;
 			if (WriteHeaderBufferPtr && WriteSize > 0)
@@ -643,7 +608,7 @@ namespace zen {
 
 		static size_t StaticWriteBodyFn(void* Ptr, size_t SizeInBlocks, size_t BlockSizeInBytes, void* UserData)
 		{
-			FRequest* Request = static_cast<FRequest*>(UserData);
+			FZenHttpRequest* Request = static_cast<FZenHttpRequest*>(UserData);
 			const size_t WriteSize = SizeInBlocks * BlockSizeInBytes;
 			TArray<uint8>* WriteDataBufferPtr = Request->WriteDataBufferPtr;
 
@@ -677,7 +642,7 @@ namespace zen {
 
 		static size_t StaticSeekFn(void* UserData, curl_off_t Offset, int Origin)
 		{
-			FRequest* Request = static_cast<FRequest*>(UserData);
+			FZenHttpRequest* Request = static_cast<FZenHttpRequest*>(UserData);
 			size_t NewPosition = 0;
 			uint64 ReadDataSize = Request->ReadDataView ? Request->ReadDataView->GetSize() : 0;
 
@@ -716,7 +681,7 @@ namespace zen {
 			for (uint8 i = 0; i < Pool.Num(); ++i)
 			{
 				Pool[i].Usage = 0u;
-				Pool[i].Request = new FRequest(InServiceUrl, true);
+				Pool[i].Request = new FZenHttpRequest(InServiceUrl, true);
 			}
 		}
 
@@ -735,7 +700,7 @@ namespace zen {
 		 * "owned by the caller and need to release it to the pool when work has been completed.
 		 * @return Usable request instance.
 		 */
-		FRequest* WaitForFreeRequest()
+		FZenHttpRequest* WaitForFreeRequest()
 		{
 			TRACE_CPUPROFILER_EVENT_SCOPE(ZenDDC_WaitForConnPool);
 			while (true)
@@ -759,7 +724,7 @@ namespace zen {
 		 * Release request to the pool.
 		 * @param Request Request that should be freed. Note that any buffer owened by the request can now be reset.
 		 */
-		void ReleaseRequestToPool(FRequest* Request)
+		void ReleaseRequestToPool(FZenHttpRequest* Request)
 		{
 			for (uint8 i = 0; i < Pool.Num(); ++i)
 			{
@@ -779,7 +744,7 @@ namespace zen {
 		struct FEntry
 		{
 			TAtomic<uint8> Usage;
-			FRequest* Request;
+			FZenHttpRequest* Request;
 		};
 
 		TStaticArray<FEntry, UE_ZENDDC_REQUEST_POOL_SIZE> Pool;
@@ -807,25 +772,21 @@ namespace zen {
 			Pool->ReleaseRequestToPool(Request);
 		}
 
-		bool IsValid() const
+		inline bool IsValid() const
 		{
 			return Request != nullptr;
 		}
 
-		FRequest* Get() const
-		{
-			check(IsValid());
-			return Request;
-		}
+		inline operator bool() const { return IsValid(); }
 
-		FRequest* operator->()
+		FZenHttpRequest* operator->()
 		{
 			check(IsValid());
 			return Request;
 		}
 
 	private:
-		FRequest* Request;
+		FZenHttpRequest* Request;
 		FRequestPool* Pool;
 	};
 }
@@ -865,10 +826,10 @@ FString FZenDerivedDataBackend::GetName() const
 
 bool FZenDerivedDataBackend::IsServiceReady()
 {
-	zen::FRequest Request(*Domain, false);
-	zen::FRequest::Result Result = Request.PerformBlockingDownload(TEXT("health/ready"), nullptr);
+	zen::FZenHttpRequest Request(*Domain, false);
+	zen::FZenHttpRequest::Result Result = Request.PerformBlockingDownload(TEXT("health/ready"), nullptr);
 	
-	if (Result == zen::FRequest::Success && zen::IsSuccessCode(Request.GetResponseCode()))
+	if (Result == zen::FZenHttpRequest::Success && zen::IsSuccessCode(Request.GetResponseCode()))
 	{
 		UE_LOG(LogDerivedDataCache, Display, TEXT("Z$ HTTP DDC service status: %s."), *Request.GetResponseAsString());
 		return true;
@@ -916,12 +877,12 @@ bool FZenDerivedDataBackend::CachedDataProbablyExists(const TCHAR* CacheKey)
 	while (ResponseCode == 0 && ++Attempts < UE_ZENDDC_MAX_ATTEMPTS)
 	{
 		zen::FScopedRequestPtr Request(RequestPool.Get());
-		zen::FRequest::Result Result = Request->PerformBlockingQuery<zen::FRequest::Head>(*Uri);
+		zen::FZenHttpRequest::Result Result = Request->PerformBlockingHead(*Uri);
 		ResponseCode = Request->GetResponseCode();
 
 		if (zen::IsSuccessCode(ResponseCode) || ResponseCode == 400)
 		{
-			const bool bIsHit = (Result == zen::FRequest::Success && zen::IsSuccessCode(ResponseCode));
+			const bool bIsHit = (Result == zen::FZenHttpRequest::Success && zen::IsSuccessCode(ResponseCode));
 			if (bIsHit)
 			{
 				TRACE_COUNTER_ADD(ZenDDC_ExistHit, int64(1));
@@ -989,7 +950,7 @@ FZenDerivedDataBackend::GetZenData(const TCHAR* Uri, TArray<uint8>* OutData) con
 		zen::FScopedRequestPtr Request(RequestPool.Get());
 		if (Request.IsValid())
 		{
-			zen::FRequest::Result Result = Request->PerformBlockingDownload(Uri, OutData);
+			zen::FZenHttpRequest::Result Result = Request->PerformBlockingDownload(Uri, OutData);
 			int64 ResponseCode = Request->GetResponseCode();
 
 			// Request was successful, make sure we got all the expected data.
@@ -1026,7 +987,7 @@ FZenDerivedDataBackend::GetZenData(const FCacheKey& CacheKey, ECachePolicy Cache
 		zen::FScopedRequestPtr Request(RequestPool.Get());
 		if (Request.IsValid())
 		{
-			zen::FRequest::Result Result = Request->PerformBlockingDownload(QueryUri.ToString(), OutPackage);
+			zen::FZenHttpRequest::Result Result = Request->PerformBlockingDownload(QueryUri.ToString(), OutPackage);
 			int64 ResponseCode = Request->GetResponseCode();
 			bool bPackageValid = Request->GetResponseFormatValid();
 
@@ -1084,7 +1045,7 @@ FZenDerivedDataBackend::PutZenData(const TCHAR* Uri, const FCompositeBuffer& InD
 		zen::FScopedRequestPtr Request(RequestPool.Get());
 		if (Request.IsValid())
 		{
-			zen::FRequest::Result Result = Request->PerformBlockingUpload<zen::FRequest::Put>(Uri, InData, ContentType);
+			zen::FZenHttpRequest::Result Result = Request->PerformBlockingPut(Uri, InData, ContentType);
 			ResponseCode = Request->GetResponseCode();
 
 			if (zen::IsSuccessCode(ResponseCode))
@@ -1206,9 +1167,9 @@ void FZenDerivedDataBackend::RemoveCachedData(const TCHAR* CacheKey, bool bTrans
 	while (ResponseCode == 0 && ++Attempts < UE_ZENDDC_MAX_ATTEMPTS)
 	{
 		zen::FScopedRequestPtr Request(RequestPool.Get());
-		if (Request.IsValid())
+		if (Request)
 		{
-			zen::FRequest::Result Result = Request->PerformBlockingQuery<zen::FRequest::Delete>(*Uri);
+			zen::FZenHttpRequest::Result Result = Request->PerformBlockingDelete(*Uri);
 			ResponseCode = Request->GetResponseCode();
 
 			if (zen::IsSuccessCode(ResponseCode))
