Refactored ZenDDC communication. HTTP interface code is now simplified and separated out from the ZenDDC client code as a step towards a shared communication layer for the oplog/cache/etc as well as instance lifetime management. There should be no functional differences.

diff --git a/Engine/Source/Developer/DerivedDataCache/Internal/ZenServerInterface.h b/Engine/Source/Developer/DerivedDataCache/Internal/ZenServerInterface.h
new file mode 100644
--- /dev/null
+++ b/Engine/Source/Developer/DerivedDataCache/Internal/ZenServerInterface.h
@@ -0,0 +1,254 @@
+// Copyright Epic Games, Inc. All Rights Reserved.
+
+#pragma once
+
+#include "Containers/StaticArray.h"
+
+class FCompositeBuffer;
+class FCbPackage;
+
+#ifndef UE_WITH_ZEN
+#	if PLATFORM_WINDOWS
+#		define UE_WITH_ZEN 1
+#	else
+#		define UE_WITH_ZEN 0
+#	endif
+#endif
+
+#if UE_WITH_ZEN
+
+namespace UE::Zen {
+
+class FZenServiceInstance
+{
+public:
+	 DERIVEDDATACACHE_API FZenServiceInstance();
+	 DERIVEDDATACACHE_API ~FZenServiceInstance();
+
+	 DERIVEDDATACACHE_API bool IsServiceRunning();
+	 DERIVEDDATACACHE_API bool IsServiceReady();
+
+private:
+};
+
+static bool IsSuccessCode(int ResponseCode)
+{
+	return 200 <= ResponseCode && ResponseCode < 300;
+}
+
+enum class EContentType
+{
+	Binary,
+	CompactBinary,
+	CompactBinaryPackage
+};
+
+/** Minimal HTTP request type wrapping CURL without the need for managers. This request
+  * is written to allow reuse of request objects, in order to allow connections to be reused.
+	
+  * CURL has a global library initialization (curl_global_init). We rely on this happening in
+  * the Online/HTTP library which is a dependency of this module
+  */
+
+class FZenHttpRequest
+{
+public:
+	FZenHttpRequest(const TCHAR* InDomain, bool bInLogErrors);
+	~FZenHttpRequest();
+
+	/**
+		* Resets all options on the request except those that should always be set.
+		*/
+	void Reset();
+
+	/** Returns the HTTP response code.*/
+	inline const int GetResponseCode() const
+	{
+		return int(ResponseCode);
+	}
+
+	inline const bool GetResponseFormatValid() const
+	{
+		return bResponseFormatValid;
+	}
+
+	/** Returns the number of bytes sent during this request (headers withstanding). */
+	inline const size_t  GetBytesSent() const
+	{
+		return BytesSent;
+	}
+
+	/**
+		* Convenience result type interpreted from HTTP response code.
+		*/
+	enum class Result
+	{
+		Success,
+		Failed
+	};
+
+	/**
+		* Upload buffer using the request, using PUT verb
+		* @param Uri Url to use.
+		* @param Buffer Data to upload
+		* @return Result of the request
+		*/
+	Result PerformBlockingPut(const TCHAR* Uri, const FCompositeBuffer& Buffer, EContentType ContentType);
+
+	/**
+		* Download an url into a buffer using the request.
+		* @param Uri Url to use.
+		* @param Buffer Optional buffer where data should be downloaded to. If this is null then 
+		* downloaded data will be stored in an internal buffer and accessed via GetResponseAsString
+		* @return Result of the request
+		*/
+	Result PerformBlockingDownload(const TCHAR* Uri, TArray<uint8>* Buffer);
+
+	/**
+		* Download an url into a buffer using the request.
+		* @param Uri Url to use.
+		* @param OutPackage Package instance which will receive the data
+		* @result Request success/failure status
+		*/
+	Result PerformBlockingDownload(const TCHAR* Uri, FCbPackage& OutPackage);
+
+	/**
+		* Query an url using the request. Queries can use either "Head" or "Delete" verbs.
+		* @param Uri Url to use.
+		* @return Result of the request
+		*/
+	Result PerformBlockingHead(const TCHAR* Uri);
+
+	/**
+		* Query an url using the request. Queries can use either "Head" or "Delete" verbs.
+		* @param Uri Url to use.
+		* @return Result of the request
+		*/
+	Result PerformBlockingDelete(const TCHAR* Uri);
+
+	/**
+		* Returns the response buffer as a string. Note that is the request is performed
+		* with an external buffer as target buffer this string will be empty.
+		*/
+	inline FString GetResponseAsString() const
+	{
+		return GetAnsiBufferAsString(ResponseBuffer);
+	}
+
+private:
+	void* /* CURL */		Curl = nullptr;
+	long /* CURLCode */		CurlResult;
+	long					ResponseCode = 0;
+	size_t					BytesSent = 0;
+	size_t					BytesReceived = 0;
+	bool					bLogErrors = false;
+	bool					bResponseFormatValid = false;
+
+	const FCompositeBuffer*	ReadDataView = nullptr;
+	TArray<uint8>*			WriteDataBufferPtr = nullptr;
+	FCbPackage*				WriteDataPackage = nullptr;
+	TArray<uint8>*			WriteHeaderBufferPtr = nullptr;
+
+	TArray<uint8>			ResponseHeader;
+	TArray<uint8>			ResponseBuffer;
+	TArray<FString>			Headers;
+	FString					Domain;
+
+	/**
+	  * Supported request verb
+	  */
+	enum class RequestVerb
+	{
+		Get,
+		Put,
+		Post,
+		Delete,
+		Head
+	};
+
+	void LogResult(long /*CURLcode*/ Result, const TCHAR* Uri, RequestVerb Verb) const;
+
+	/**
+		* Performs the request, blocking until finished.
+		* @param Uri Address on the domain to query
+		* @param Verb HTTP verb to use
+		* @param Buffer Optional buffer to directly receive the result of the request.
+		* If unset the response body will be stored in the request.
+		*/
+	Result PerformBlocking(const TCHAR* Uri, RequestVerb Verb, uint32 ContentLength);
+
+	static FString GetAnsiBufferAsString(const TArray<uint8>& Buffer);
+
+	struct FStatics;
+};
+
+/**
+  * Pool which manages a fixed set of requests. Users are required to release requests that have been
+  * acquired. 
+  * 
+  * Intended to be used with \ref FScopedRequestPtr which handles lifetime management transparently
+  */
+struct FRequestPool
+{
+	explicit FRequestPool(const TCHAR* InServiceUrl);
+	~FRequestPool();
+
+	/** Block until a request is free. Once a request has been returned it is
+	  * "owned by the caller and need to release it to the pool when work has been completed.
+	  * @return Usable request instance.
+	  */
+	FZenHttpRequest* WaitForFreeRequest();
+
+	/** Release request to the pool.
+	  * @param Request Request that should be freed. Note that any buffer owned by the request can now be reset.
+	  */
+	void ReleaseRequestToPool(FZenHttpRequest* Request);
+
+private:
+	struct FEntry
+	{
+		TAtomic<uint8> Usage;
+		FZenHttpRequest* Request;
+	};
+
+	TStaticArray<FEntry, 16> Pool;
+};
+
+/**
+  * Utility class to manage requesting and releasing requests from the \ref FRequestPool.
+  */
+struct FScopedRequestPtr
+{
+public:
+	FScopedRequestPtr(FRequestPool* InPool)
+	:	Request(InPool->WaitForFreeRequest())
+	,	Pool(InPool)
+	{
+	}
+
+	~FScopedRequestPtr()
+	{
+		Pool->ReleaseRequestToPool(Request);
+	}
+
+	inline bool IsValid() const
+	{
+		return Request != nullptr;
+	}
+
+	inline operator bool() const { return IsValid(); }
+
+	FZenHttpRequest* operator->()
+	{
+		check(IsValid());
+		return Request;
+	}
+
+private:
+	FZenHttpRequest*	Request;
+	FRequestPool*		Pool;
+};
+
+} // namespace UE::Zen
+
+#endif // UE_WITH_ZEN
diff --git a/Engine/Source/Developer/DerivedDataCache/Private/ZenBackendUtils.h b/Engine/Source/Developer/DerivedDataCache/Private/ZenBackendUtils.h
--- a/Engine/Source/Developer/DerivedDataCache/Private/ZenBackendUtils.h
+++ b/Engine/Source/Developer/DerivedDataCache/Private/ZenBackendUtils.h
@@ -9,7 +9,7 @@
 #include "Memory/SharedBuffer.h"
 #include "Misc/AssertionMacros.h"
 
-namespace zen {
+namespace UE::Zen {
 
 /** Copy Size bytes from Source at offset SourceOffset into Dest at offset 0. */
 inline void Memcpy(void* Dest, const FCompositeBuffer& Source, size_t SourceOffset, size_t Size)
@@ -50,5 +50,4 @@ inline void Memcpy(void* Dest, const FCompositeBuffer& Source, size_t SourceOffs
 	check(Offset == EndOffset);
 }
 
-
-} // namespace zen
\ No newline at end of file
+} // namespace UE::Zen
diff --git a/Engine/Source/Developer/DerivedDataCache/Private/ZenDerivedDataBackend.cpp b/Engine/Source/Developer/DerivedDataCache/Private/ZenDerivedDataBackend.cpp
--- a/Engine/Source/Developer/DerivedDataCache/Private/ZenDerivedDataBackend.cpp
+++ b/Engine/Source/Developer/DerivedDataCache/Private/ZenDerivedDataBackend.cpp
@@ -3,46 +3,17 @@
 
 #if WITH_ZEN_DDC_BACKEND
 
-#if PLATFORM_WINDOWS || PLATFORM_HOLOLENS
-#	include "Windows/WindowsHWrapper.h"
-#	include "Windows/AllowWindowsPlatformTypes.h"
-#endif
-
-#include "curl/curl.h"
-
-#if PLATFORM_WINDOWS || PLATFORM_HOLOLENS
-#	include "Windows/HideWindowsPlatformTypes.h"
-#endif
-
 #include "Algo/Accumulate.h"
-#include "Containers/StaticArray.h"
-#include "Containers/Ticker.h"
 #include "DerivedDataCacheRecord.h"
-#include "Dom/JsonObject.h"
-#include "GenericPlatform/GenericPlatformFile.h"
-#include "HAL/PlatformFileManager.h"
-#include "Misc/ConfigCacheIni.h"
-#include "Misc/FileHelper.h"
 #include "Misc/ScopeLock.h"
-#include "Misc/SecureHash.h"
-#include "ProfilingDebugging/CpuProfilerTrace.h"
 #include "ProfilingDebugging/CountersTrace.h"
+#include "ProfilingDebugging/CpuProfilerTrace.h"
 #include "Serialization/CompactBinary.h"
 #include "Serialization/CompactBinaryPackage.h"
 #include "Serialization/CompactBinaryValidation.h"
 #include "Serialization/CompactBinaryWriter.h"
-#include "Serialization/JsonReader.h"
-#include "Serialization/JsonSerializer.h"
-#include "Serialization/MemoryReader.h"
 #include "ZenBackendUtils.h"
-
-#define UE_ZENDDC_BACKEND_WAIT_INTERVAL			0.01f
-#define UE_ZENDDC_HTTP_REQUEST_TIMEOUT_SECONDS	30L
-#define UE_ZENDDC_HTTP_REQUEST_TIMOUT_ENABLED	1
-#define UE_ZENDDC_HTTP_DEBUG					0
-#define UE_ZENDDC_REQUEST_POOL_SIZE				16
-#define UE_ZENDDC_MAX_FAILED_LOGIN_ATTEMPTS		16
-#define UE_ZENDDC_MAX_ATTEMPTS					4
+#include "ZenServerInterface.h"
 
 TRACE_DECLARE_INT_COUNTER(ZenDDC_Exist,			TEXT("ZenDDC Exist"));
 TRACE_DECLARE_INT_COUNTER(ZenDDC_ExistHit,		TEXT("ZenDDC Exist Hit"));
@@ -53,744 +24,6 @@ TRACE_DECLARE_INT_COUNTER(ZenDDC_PutHit,		TEXT("ZenDDC Put Hit"));
 TRACE_DECLARE_INT_COUNTER(ZenDDC_BytesReceived, TEXT("ZenDDC Bytes Received"));
 TRACE_DECLARE_INT_COUNTER(ZenDDC_BytesSent,		TEXT("ZenDDC Bytes Sent"));
 
-namespace zen {
-
-	static bool IsSuccessCode(int64 ResponseCode)
-	{
-		return 200 <= ResponseCode && ResponseCode < 300;
-	}
-
-	enum class EContentType
-	{
-		Binary,
-		CompactBinary,
-		CompactBinaryPackage
-	};
-
-	/**
-	 * Minimal HTTP request type wrapping CURL without the need for managers. This request
-	 * is written to allow reuse of request objects, in order to allow connections to be reused.
-	 *
-	 * CURL has a global library initialization (curl_global_init). We rely on this happening in
-	 * the Online/HTTP library which is a dependency on this module.
-	 */
-	class FZenHttpRequest
-	{
-	public:
-		/**
-		 * Supported request verbs
-		 */
-		enum RequestVerb
-		{
-			Get,
-			Put,
-			Post,
-			PostJson,
-			Delete,
-			Head
-		};
-
-		/**
-		 * Convenience result type interpreted from HTTP response code.
-		 */
-		enum Result
-		{
-			Success,
-			Failed,
-			FailedTimeout
-		};
-
-		FZenHttpRequest(const TCHAR* InDomain, bool bInLogErrors)
-			: bLogErrors(bInLogErrors)
-			, Domain(InDomain)
-		{
-			Curl = curl_easy_init();
-			Reset();
-		}
-
-		~FZenHttpRequest()
-		{
-			curl_easy_cleanup(Curl);
-		}
-
-		/**
-		 * Resets all options on the request except those that should always be set.
-		 */
-		void Reset()
-		{
-			Headers.Reset();
-			ResponseHeader.Reset();
-			ResponseBuffer.Reset();
-			ResponseCode = 0;
-			bResponseFormatValid = true;
-			ReadDataView = nullptr;
-			WriteDataBufferPtr = nullptr;
-			WriteHeaderBufferPtr = nullptr;
-			BytesSent = 0;
-			BytesReceived = 0;
-			CurlResult = CURL_LAST;
-
-			curl_easy_reset(Curl);
-
-			// Options that are always set for all connections.
-#if UE_ZENDDC_HTTP_REQUEST_TIMOUT_ENABLED
-			curl_easy_setopt(Curl, CURLOPT_CONNECTTIMEOUT, UE_ZENDDC_HTTP_REQUEST_TIMEOUT_SECONDS);
-#endif
-			curl_easy_setopt(Curl, CURLOPT_FOLLOWLOCATION, 1L);
-			curl_easy_setopt(Curl, CURLOPT_NOSIGNAL, 1L);
-			curl_easy_setopt(Curl, CURLOPT_BUFFERSIZE, 256 * 1024L);
-			//curl_easy_setopt(Curl, CURLOPT_UPLOAD_BUFFERSIZE, 256 * 1024L);
-			// Response functions
-			curl_easy_setopt(Curl, CURLOPT_HEADERDATA, this);
-			curl_easy_setopt(Curl, CURLOPT_HEADERFUNCTION, &FZenHttpRequest::StaticWriteHeaderFn);
-			curl_easy_setopt(Curl, CURLOPT_WRITEDATA, this);
-			curl_easy_setopt(Curl, CURLOPT_WRITEFUNCTION, StaticWriteBodyFn);
-			// Rewind method, handle special error case where request need to rewind data stream
-			curl_easy_setopt(Curl, CURLOPT_SEEKFUNCTION, StaticSeekFn);
-			curl_easy_setopt(Curl, CURLOPT_SEEKDATA, this);
-			// Debug hooks
-#if UE_ZENDDC_HTTP_DEBUG
-			curl_easy_setopt(Curl, CURLOPT_DEBUGDATA, this);
-			curl_easy_setopt(Curl, CURLOPT_DEBUGFUNCTION, StaticDebugCallback);
-			curl_easy_setopt(Curl, CURLOPT_VERBOSE, 1L);
-#endif
-		}
-
-		/** Gets the domain name for this request */
-		const FString& GetDomain() const
-		{
-			return Domain;
-		}
-
-		/** Returns the HTTP response code.*/
-		const int64 GetResponseCode() const
-		{
-			return ResponseCode;
-		}
-
-		const bool GetResponseFormatValid() const
-		{
-			return bResponseFormatValid;
-		}
-
-		/** Returns the number of bytes received this request (headers withstanding). */
-		const size_t GetBytesReceived() const
-		{
-			return BytesReceived;
-		}
-
-		/** Returns the number of bytes sent during this request (headers withstanding). */
-		const size_t  GetBytesSent() const
-		{
-			return BytesSent;
-		}
-
-		/**
-		 * Upload buffer using the request, using PUT verb
-		 * @param Uri Url to use.
-		 * @param Buffer Data to upload
-		 * @return Result of the request
-		 */
-		Result PerformBlockingPut(const TCHAR* Uri, const FCompositeBuffer& Buffer, EContentType ContentType)
-		{
-			uint32 ContentLength = 0u;
-
-			ContentLength = Buffer.GetSize();
-			curl_easy_setopt(Curl, CURLOPT_UPLOAD, 1L);
-			curl_easy_setopt(Curl, CURLOPT_INFILESIZE, ContentLength);
-			curl_easy_setopt(Curl, CURLOPT_READDATA, this);
-			curl_easy_setopt(Curl, CURLOPT_READFUNCTION, StaticReadFn);
-
-			switch (ContentType)
-			{
-			case EContentType::Binary:
-				Headers.Add(FString(TEXT("Content-Type: application/octet-stream")));
-				break;
-			case EContentType::CompactBinary:
-				Headers.Add(FString(TEXT("Content-Type: application/x-ue-cb")));
-				break;
-			case EContentType::CompactBinaryPackage:
-				Headers.Add(FString(TEXT("Content-Type: application/x-ue-cbpkg")));
-				break;
-			default:
-				checkNoEntry();
-				break;
-			}
-			ReadDataView = &Buffer;
-
-			return PerformBlocking(Uri, RequestVerb::Put, ContentLength);
-		}
-
-		/**
-		 * Download an url into a buffer using the request.
-		 * @param Uri Url to use.
-		 * @param Buffer Optional buffer where data should be downloaded to. If empty downloaded data will
-		 * be stored in an internal buffer and accessed GetResponse* methods.
-		 * @return Result of the request
-		 */
-		Result PerformBlockingDownload(const TCHAR* Uri, TArray<uint8>* Buffer)
-		{
-			curl_easy_setopt(Curl, CURLOPT_HTTPGET, 1L);
-			WriteDataBufferPtr = Buffer;
-
-			return PerformBlocking(Uri, Get, 0u);
-		}
-
-		Result PerformBlockingDownload(const TCHAR* Uri, FCbPackage& OutPackage)
-		{
-			curl_easy_setopt(Curl, CURLOPT_HTTPGET, 1L);
-			OutPackage.Reset();
-			// TODO: When PackageBytes can be written in segments directly, set the WritePtr to the OutPackage and use that
-			TArray<uint8> PackageBytes;
-			WriteDataBufferPtr = &PackageBytes;
-			Result LocalResult = PerformBlocking(Uri, Get, 0u);
-			if (IsSuccessCode(ResponseCode))
-			{
-				if (ValidateCompactBinaryPackage(FMemoryView(PackageBytes.GetData(), PackageBytes.Num()), ECbValidateMode::All) != ECbValidateError::None)
-				{
-					bResponseFormatValid = false;
-				}
-				else
-				{
-					FMemoryReader PackageLoader(PackageBytes);
-					bResponseFormatValid = OutPackage.TryLoad(PackageLoader);
-				}
-			}
-			return LocalResult;
-		}
-
-		/**
-		 * Query an url using the request. Queries can use either "Head" or "Delete" verbs.
-		 * @param Uri Url to use.
-		 * @return Result of the request
-		 */
-		Result PerformBlockingHead(const TCHAR* Uri)
-		{
-			curl_easy_setopt(Curl, CURLOPT_NOBODY, 1L);
-
-			return PerformBlocking(Uri, RequestVerb::Head, 0u);
-		}
-
-		/**
-		 * Query an url using the request. Queries can use either "Head" or "Delete" verbs.
-		 * @param Uri Url to use.
-		 * @return Result of the request
-		 */
-		Result PerformBlockingDelete(const TCHAR* Uri)
-		{
-			curl_easy_setopt(Curl, CURLOPT_POST, 1L);
-			curl_easy_setopt(Curl, CURLOPT_CUSTOMREQUEST, "DELETE");
-
-			return PerformBlocking(Uri, RequestVerb::Delete, 0u);
-		}
-
-		/**
-		 * Set a header to send with the request.
-		 */
-		void SetHeader(const TCHAR* Header, const TCHAR* Value)
-		{
-			check(CurlResult == CURL_LAST); // Cannot set header after request is sent
-			Headers.Add(FString::Printf(TEXT("%s: %s"), Header, Value));
-		}
-
-		/**
-		 * Attempts to find the header from the response. Returns false if header is not present.
-		 */
-		bool GetHeader(const ANSICHAR* Header, FString& OutValue) const
-		{
-			check(CurlResult != CURL_LAST);  // Cannot query headers before request is sent
-
-			const ANSICHAR* HeadersBuffer = (const ANSICHAR*)ResponseHeader.GetData();
-			size_t HeaderLen = strlen(Header);
-
-			// Find the header key in the (ANSI) response buffer. If not found we can exist immediately
-			if (const ANSICHAR* Found = strstr(HeadersBuffer, Header))
-			{
-				const ANSICHAR* Linebreak = strchr(Found, '\r');
-				const ANSICHAR* ValueStart = Found + HeaderLen + 2; //colon and space
-				const size_t ValueSize = Linebreak - ValueStart;
-				FUTF8ToTCHAR TCHARData(ValueStart, ValueSize);
-				OutValue = FString(TCHARData.Length(), TCHARData.Get());
-				return true;
-			}
-			return false;
-		}
-
-		/**
-		 * Returns the response buffer. Note that is the request is performed
-		 * with an external buffer as target buffer this string will be empty.
-		 */
-		const TArray<uint8>& GetResponseBuffer() const
-		{
-			return ResponseBuffer;
-		}
-
-		/**
-		 * Returns the response buffer as a string. Note that is the request is performed
-		 * with an external buffer as target buffer this string will be empty.
-		 */
-		FString GetResponseAsString() const
-		{
-			return GetAnsiBufferAsString(ResponseBuffer);
-		}
-
-		/**
-		 * Returns the response header as a string.
-		 */
-		FString GetResponseHeaderAsString()
-		{
-			return GetAnsiBufferAsString(ResponseHeader);
-		}
-
-		/**
-		 * Tries to parse the response buffer as a JsonObject. Return empty pointer if
-		 * parse error occurs.
-		 */
-		TSharedPtr<FJsonObject> GetResponseAsJsonObject() const
-		{
-			FString Response = GetAnsiBufferAsString(ResponseBuffer);
-
-			TSharedPtr<FJsonObject> JsonObject;
-			TSharedRef<TJsonReader<> > JsonReader = TJsonReaderFactory<>::Create(Response);
-			if (!FJsonSerializer::Deserialize(JsonReader, JsonObject) || !JsonObject.IsValid())
-			{
-				return TSharedPtr<FJsonObject>(nullptr);
-			}
-
-			return JsonObject;
-		}
-
-	private:
-		CURL* Curl;
-		CURLcode				CurlResult;
-		long					ResponseCode;
-		size_t					BytesSent;
-		size_t					BytesReceived;
-		bool					bLogErrors;
-		bool					bResponseFormatValid;
-
-		const FCompositeBuffer*	ReadDataView;
-		TArray<uint8>*			WriteDataBufferPtr;
-		FCbPackage*				WriteDataPackage;
-		TArray<uint8>*			WriteHeaderBufferPtr;
-
-		TArray<uint8>			ResponseHeader;
-		TArray<uint8>			ResponseBuffer;
-		TArray<FString>			Headers;
-		FString					Domain;
-
-		/**
-		 * Performs the request, blocking until finished.
-		 * @param Uri Address on the domain to query
-		 * @param Verb HTTP verb to use
-		 * @param Buffer Optional buffer to directly receive the result of the request.
-		 * If unset the response body will be stored in the request.
-		 */
-		Result PerformBlocking(const TCHAR* Uri, RequestVerb Verb, uint32 ContentLength)
-		{
-			static const char* CommonHeaders[] = {
-				"User-Agent: UE4",
-				nullptr
-			};
-
-			TRACE_CPUPROFILER_EVENT_SCOPE(ZenDDC_CurlPerform);
-
-			// Setup request options
-			FString Url = FString::Printf(TEXT("%s/%s"), *Domain, Uri);
-			curl_easy_setopt(Curl, CURLOPT_URL, TCHAR_TO_ANSI(*Url));
-
-			// Setup response header buffer. If caller has not setup a response data buffer, use internal.
-			WriteHeaderBufferPtr = &ResponseHeader;
-			if (WriteDataBufferPtr == nullptr)
-			{
-				WriteDataBufferPtr = &ResponseBuffer;
-			}
-
-			// Content-Length should always be set
-			Headers.Add(FString::Printf(TEXT("Content-Length: %d"), ContentLength));
-
-			// Build headers list
-			curl_slist* CurlHeaders = nullptr;
-			// Add common headers
-			for (uint8 i = 0; CommonHeaders[i] != nullptr; ++i)
-			{
-				CurlHeaders = curl_slist_append(CurlHeaders, CommonHeaders[i]);
-			}
-			// Setup added headers
-			for (const FString& Header : Headers)
-			{
-				CurlHeaders = curl_slist_append(CurlHeaders, TCHAR_TO_ANSI(*Header));
-			}
-			curl_easy_setopt(Curl, CURLOPT_HTTPHEADER, CurlHeaders);
-
-			// Shots fired!
-			CurlResult = curl_easy_perform(Curl);
-
-			// Get response code
-			bool bRedirected = false;
-			if (CURLE_OK == curl_easy_getinfo(Curl, CURLINFO_RESPONSE_CODE, &ResponseCode))
-			{
-				bRedirected = (ResponseCode >= 300 && ResponseCode < 400);
-			}
-
-			LogResult(CurlResult, Uri, Verb);
-
-			// Clean up
-			curl_slist_free_all(CurlHeaders);
-
-			return CurlResult == CURLE_OK ? Success : Failed;
-		}
-
-		void LogResult(CURLcode Result, const TCHAR* Uri, RequestVerb Verb) const
-		{
-			if (Result == CURLE_OK)
-			{
-				bool bSuccess = false;
-				const TCHAR* VerbStr = nullptr;
-				FString AdditionalInfo;
-
-				const bool Is400 = (ResponseCode >= 400) && (ResponseCode <= 499);
-				const bool Is200 = (ResponseCode >= 200) && (ResponseCode <= 299);
-
-				switch (Verb)
-				{
-				case Head:
-					bSuccess = Is400 || Is200;
-					VerbStr = TEXT("querying");
-					break;
-				case Get:
-					bSuccess = Is400 || Is200;
-					VerbStr = TEXT("fetching");
-					AdditionalInfo = FString::Printf(TEXT("Received: %d bytes."), BytesReceived);
-					break;
-				case Put:
-					bSuccess = Is200;
-					VerbStr = TEXT("updating");
-					AdditionalInfo = FString::Printf(TEXT("Sent: %d bytes."), BytesSent);
-					break;
-				case Post:
-				case PostJson:
-					bSuccess = Is200;
-					VerbStr = TEXT("posting");
-					break;
-				case Delete:
-					bSuccess = Is200;
-					VerbStr = TEXT("deleting");
-					break;
-				}
-
-				if (bSuccess)
-				{
-					UE_LOG(
-						LogDerivedDataCache,
-						Verbose,
-						TEXT("Finished %s HTTP cache entry (response %d) from %s. %s"),
-						VerbStr,
-						ResponseCode,
-						Uri,
-						*AdditionalInfo
-					);
-				}
-				else if (bLogErrors)
-				{
-					// Print the response body if we got one, otherwise print header.
-					FString Response = GetAnsiBufferAsString(ResponseBuffer.Num() > 0 ? ResponseBuffer : ResponseHeader);
-					Response.ReplaceCharInline('\n', ' ');
-					Response.ReplaceCharInline('\r', ' ');
-					UE_LOG(
-						LogDerivedDataCache,
-						Error,
-						TEXT("Failed %s HTTP cache entry (response %d) from %s. Response: %s"),
-						VerbStr,
-						ResponseCode,
-						Uri,
-						*Response
-					);
-				}
-			}
-			else if (bLogErrors)
-			{
-				UE_LOG(
-					LogDerivedDataCache,
-					Error,
-					TEXT("Error while connecting to %s: %s"),
-					*Domain,
-					ANSI_TO_TCHAR(curl_easy_strerror(Result))
-				);
-			}
-		}
-
-		FString GetAnsiBufferAsString(const TArray<uint8>& Buffer) const
-		{
-			// Content is NOT null-terminated; we need to specify lengths here
-			FUTF8ToTCHAR TCHARData(reinterpret_cast<const ANSICHAR*>(Buffer.GetData()), Buffer.Num());
-			return FString(TCHARData.Length(), TCHARData.Get());
-		}
-
-		static size_t StaticDebugCallback(CURL* Handle, curl_infotype DebugInfoType, char* DebugInfo, size_t DebugInfoSize, void* UserData)
-		{
-			FZenHttpRequest* Request = static_cast<FZenHttpRequest*>(UserData);
-
-			switch (DebugInfoType)
-			{
-			case CURLINFO_TEXT:
-			{
-				// Truncate at 1023 characters. This is just an arbitrary number based on a buffer size seen in
-				// the libcurl code.
-				DebugInfoSize = FMath::Min(DebugInfoSize, (size_t)1023);
-
-				// Calculate the actual length of the string due to incorrect use of snprintf() in lib/vtls/openssl.c.
-				char* FoundNulPtr = (char*)memchr(DebugInfo, 0, DebugInfoSize);
-				int CalculatedSize = FoundNulPtr != nullptr ? FoundNulPtr - DebugInfo : DebugInfoSize;
-
-				auto ConvertedString = StringCast<TCHAR>(static_cast<const ANSICHAR*>(DebugInfo), CalculatedSize);
-				FString DebugText(ConvertedString.Length(), ConvertedString.Get());
-				DebugText.ReplaceInline(TEXT("\n"), TEXT(""), ESearchCase::CaseSensitive);
-				DebugText.ReplaceInline(TEXT("\r"), TEXT(""), ESearchCase::CaseSensitive);
-				UE_LOG(LogDerivedDataCache, VeryVerbose, TEXT("%p: '%s'"), Request, *DebugText);
-			}
-			break;
-
-			case CURLINFO_HEADER_IN:
-				UE_LOG(LogDerivedDataCache, VeryVerbose, TEXT("%p: Received header (%d bytes)"), Request, DebugInfoSize);
-				break;
-
-			case CURLINFO_DATA_IN:
-				UE_LOG(LogDerivedDataCache, VeryVerbose, TEXT("%p: Received data (%d bytes)"), Request, DebugInfoSize);
-				break;
-
-			case CURLINFO_DATA_OUT:
-				UE_LOG(LogDerivedDataCache, VeryVerbose, TEXT("%p: Sent data (%d bytes)"), Request, DebugInfoSize);
-				break;
-			}
-
-			return 0;
-		}
-
-		static size_t StaticReadFn(void* Ptr, size_t SizeInBlocks, size_t BlockSizeInBytes, void* UserData)
-		{
-			FZenHttpRequest* Request = static_cast<FZenHttpRequest*>(UserData);
-			check(Request->ReadDataView);
-			const FCompositeBuffer& ReadDataView = *Request->ReadDataView;
-
-			const size_t Offset = Request->BytesSent;
-			const size_t ReadSize = FMath::Min((size_t)ReadDataView.GetSize() - Offset, SizeInBlocks * BlockSizeInBytes);
-			check(ReadDataView.GetSize() >= Offset + ReadSize);
-
-			Memcpy(Ptr, ReadDataView, Offset, ReadSize);
-			Request->BytesSent += ReadSize;
-			return ReadSize;
-
-			return 0;
-		}
-
-		static size_t StaticWriteHeaderFn(void* Ptr, size_t SizeInBlocks, size_t BlockSizeInBytes, void* UserData)
-		{
-			FZenHttpRequest* Request = static_cast<FZenHttpRequest*>(UserData);
-			const size_t WriteSize = SizeInBlocks * BlockSizeInBytes;
-			TArray<uint8>* WriteHeaderBufferPtr = Request->WriteHeaderBufferPtr;
-			if (WriteHeaderBufferPtr && WriteSize > 0)
-			{
-				const size_t CurrentBufferLength = WriteHeaderBufferPtr->Num();
-				if (CurrentBufferLength > 0)
-				{
-					// Remove the previous zero termination
-					(*WriteHeaderBufferPtr)[CurrentBufferLength - 1] = ' ';
-				}
-
-				// Write the header
-				WriteHeaderBufferPtr->Append((const uint8*)Ptr, WriteSize + 1);
-				(*WriteHeaderBufferPtr)[WriteHeaderBufferPtr->Num() - 1] = 0; // Zero terminate string
-				return WriteSize;
-			}
-			return 0;
-		}
-
-		static size_t StaticWriteBodyFn(void* Ptr, size_t SizeInBlocks, size_t BlockSizeInBytes, void* UserData)
-		{
-			FZenHttpRequest* Request = static_cast<FZenHttpRequest*>(UserData);
-			const size_t WriteSize = SizeInBlocks * BlockSizeInBytes;
-			TArray<uint8>* WriteDataBufferPtr = Request->WriteDataBufferPtr;
-
-			if (WriteDataBufferPtr && WriteSize > 0)
-			{
-				// If this is the first part of the body being received, try to reserve 
-				// memory if content length is defined in the header.
-				if (Request->BytesReceived == 0 && Request->WriteHeaderBufferPtr)
-				{
-					static const ANSICHAR* ContentLengthHeaderStr = "Content-Length: ";
-					const ANSICHAR* Header = (const ANSICHAR*)Request->WriteHeaderBufferPtr->GetData();
-
-					if (const ANSICHAR* ContentLengthHeader = FCStringAnsi::Strstr(Header, ContentLengthHeaderStr))
-					{
-						size_t ContentLength = (size_t)FCStringAnsi::Atoi64(ContentLengthHeader + strlen(ContentLengthHeaderStr));
-						if (ContentLength > 0)
-						{
-							WriteDataBufferPtr->Reserve(ContentLength);
-						}
-					}
-				}
-
-				// Write to the target buffer
-				WriteDataBufferPtr->Append((const uint8*)Ptr, WriteSize);
-				Request->BytesReceived += WriteSize;
-				return WriteSize;
-			}
-
-			return 0;
-		}
-
-		static size_t StaticSeekFn(void* UserData, curl_off_t Offset, int Origin)
-		{
-			FZenHttpRequest* Request = static_cast<FZenHttpRequest*>(UserData);
-			size_t NewPosition = 0;
-			uint64 ReadDataSize = Request->ReadDataView ? Request->ReadDataView->GetSize() : 0;
-
-			switch (Origin)
-			{
-			case SEEK_SET: NewPosition = Offset; break;
-			case SEEK_CUR: NewPosition = Request->BytesSent + Offset; break;
-			case SEEK_END: NewPosition = ReadDataSize + Offset; break;
-			}
-
-			// Make sure we don't seek outside of the buffer
-			if (NewPosition < 0 || NewPosition >= ReadDataSize)
-			{
-				return CURL_SEEKFUNC_FAIL;
-			}
-
-			// Update the used offset
-			Request->BytesSent = NewPosition;
-			return CURL_SEEKFUNC_OK;
-		}
-
-	};
-
-	//----------------------------------------------------------------------------------------------------------
-	// Request pool
-	//----------------------------------------------------------------------------------------------------------
-
-	/**
-	 * Pool that manages a fixed set of requests. Users are required to release requests that have been
-	 * acquired. Usable with \ref zen::FScopedRequestPtr which handles this automatically.
-	 */
-	struct FRequestPool
-	{
-		FRequestPool(const TCHAR* InServiceUrl)
-		{
-			for (uint8 i = 0; i < Pool.Num(); ++i)
-			{
-				Pool[i].Usage = 0u;
-				Pool[i].Request = new FZenHttpRequest(InServiceUrl, true);
-			}
-		}
-
-		~FRequestPool()
-		{
-			for (uint8 i = 0; i < Pool.Num(); ++i)
-			{
-				// No requests should be in use by now.
-				check(Pool[i].Usage.Load(EMemoryOrder::Relaxed) == 0u);
-				delete Pool[i].Request;
-			}
-		}
-
-		/**
-		 * Block until a request is free. Once a request has been returned it is
-		 * "owned by the caller and need to release it to the pool when work has been completed.
-		 * @return Usable request instance.
-		 */
-		FZenHttpRequest* WaitForFreeRequest()
-		{
-			TRACE_CPUPROFILER_EVENT_SCOPE(ZenDDC_WaitForConnPool);
-			while (true)
-			{
-				for (uint8 i = 0; i < Pool.Num(); ++i)
-				{
-					if (!Pool[i].Usage.Load(EMemoryOrder::Relaxed))
-					{
-						uint8 Expected = 0u;
-						if (Pool[i].Usage.CompareExchange(Expected, 1u))
-						{
-							return Pool[i].Request;
-						}
-					}
-				}
-				FPlatformProcess::Sleep(UE_ZENDDC_BACKEND_WAIT_INTERVAL);
-			}
-		}
-
-		/**
-		 * Release request to the pool.
-		 * @param Request Request that should be freed. Note that any buffer owened by the request can now be reset.
-		 */
-		void ReleaseRequestToPool(FZenHttpRequest* Request)
-		{
-			for (uint8 i = 0; i < Pool.Num(); ++i)
-			{
-				if (Pool[i].Request == Request)
-				{
-					Request->Reset();
-					uint8 Expected = 1u;
-					Pool[i].Usage.CompareExchange(Expected, 0u);
-					return;
-				}
-			}
-			check(false);
-		}
-
-	private:
-
-		struct FEntry
-		{
-			TAtomic<uint8> Usage;
-			FZenHttpRequest* Request;
-		};
-
-		TStaticArray<FEntry, UE_ZENDDC_REQUEST_POOL_SIZE> Pool;
-
-		FRequestPool() {}
-	};
-
-	//----------------------------------------------------------------------------------------------------------
-	// FScopedRequestPtr
-	//----------------------------------------------------------------------------------------------------------
-
-	/**
-	 * Utility class to manage requesting and releasing requests from the \ref FRequestPool.
-	 */
-	struct FScopedRequestPtr
-	{
-	public:
-		FScopedRequestPtr(FRequestPool* InPool)
-			: Request(InPool->WaitForFreeRequest())
-			, Pool(InPool)
-		{}
-
-		~FScopedRequestPtr()
-		{
-			Pool->ReleaseRequestToPool(Request);
-		}
-
-		inline bool IsValid() const
-		{
-			return Request != nullptr;
-		}
-
-		inline operator bool() const { return IsValid(); }
-
-		FZenHttpRequest* operator->()
-		{
-			check(IsValid());
-			return Request;
-		}
-
-	private:
-		FZenHttpRequest* Request;
-		FRequestPool* Pool;
-	};
-}
-
 namespace UE::DerivedData {
 
 //----------------------------------------------------------------------------------------------------------
@@ -807,7 +40,7 @@ FZenDerivedDataBackend::FZenDerivedDataBackend(
 {
 	if (IsServiceReady())
 	{
-		RequestPool = MakeUnique<zen::FRequestPool>(InServiceUrl);
+		RequestPool = MakeUnique<Zen::FRequestPool>(InServiceUrl);
 		bIsUsable = true;
 	}
 	bCacheRecordEndpointEnabled = false;
@@ -826,10 +59,10 @@ FString FZenDerivedDataBackend::GetName() const
 
 bool FZenDerivedDataBackend::IsServiceReady()
 {
-	zen::FZenHttpRequest Request(*Domain, false);
-	zen::FZenHttpRequest::Result Result = Request.PerformBlockingDownload(TEXT("health/ready"), nullptr);
+	Zen::FZenHttpRequest Request(*Domain, false);
+	Zen::FZenHttpRequest::Result Result = Request.PerformBlockingDownload(TEXT("health/ready"), nullptr);
 	
-	if (Result == zen::FZenHttpRequest::Success && zen::IsSuccessCode(Request.GetResponseCode()))
+	if (Result == Zen::FZenHttpRequest::Result::Success && Zen::IsSuccessCode(Request.GetResponseCode()))
 	{
 		UE_LOG(LogDerivedDataCache, Display, TEXT("Z$ HTTP DDC service status: %s."), *Request.GetResponseAsString());
 		return true;
@@ -871,18 +104,20 @@ bool FZenDerivedDataBackend::CachedDataProbablyExists(const TCHAR* CacheKey)
 	}
 	
 	FString Uri = MakeLegacyZenKey(CacheKey);
-	long ResponseCode = 0; uint32 Attempts = 0;
+	int64 ResponseCode = 0; 
+	uint32 Attempts = 0;
 
 	// Retry request until we get an accepted response or exhaust allowed number of attempts.
-	while (ResponseCode == 0 && ++Attempts < UE_ZENDDC_MAX_ATTEMPTS)
+	while (ResponseCode == 0 && ++Attempts < MaxAttempts)
 	{
-		zen::FScopedRequestPtr Request(RequestPool.Get());
-		zen::FZenHttpRequest::Result Result = Request->PerformBlockingHead(*Uri);
+		Zen::FScopedRequestPtr Request(RequestPool.Get());
+		Zen::FZenHttpRequest::Result Result = Request->PerformBlockingHead(*Uri);
 		ResponseCode = Request->GetResponseCode();
 
-		if (zen::IsSuccessCode(ResponseCode) || ResponseCode == 400)
+		if (Zen::IsSuccessCode(ResponseCode) || ResponseCode == 404)
 		{
-			const bool bIsHit = (Result == zen::FZenHttpRequest::Success && zen::IsSuccessCode(ResponseCode));
+			const bool bIsHit = (Result == Zen::FZenHttpRequest::Result::Success && Zen::IsSuccessCode(ResponseCode));
+
 			if (bIsHit)
 			{
 				TRACE_COUNTER_ADD(ZenDDC_ExistHit, int64(1));
@@ -945,16 +180,16 @@ FZenDerivedDataBackend::GetZenData(const TCHAR* Uri, TArray<uint8>* OutData) con
 {
 	// Retry request until we get an accepted response or exhaust allowed number of attempts.
 	EGetResult GetResult = EGetResult::NotFound;
-	for (uint32 Attempts = 0; Attempts < UE_ZENDDC_MAX_ATTEMPTS; ++Attempts)
+	for (uint32 Attempts = 0; Attempts < MaxAttempts; ++Attempts)
 	{
-		zen::FScopedRequestPtr Request(RequestPool.Get());
+		Zen::FScopedRequestPtr Request(RequestPool.Get());
 		if (Request.IsValid())
 		{
-			zen::FZenHttpRequest::Result Result = Request->PerformBlockingDownload(Uri, OutData);
+			Zen::FZenHttpRequest::Result Result = Request->PerformBlockingDownload(Uri, OutData);
 			int64 ResponseCode = Request->GetResponseCode();
 
 			// Request was successful, make sure we got all the expected data.
-			if (zen::IsSuccessCode(ResponseCode))
+			if (Zen::IsSuccessCode(ResponseCode))
 			{
 				return EGetResult::Success;
 			}
@@ -982,17 +217,17 @@ FZenDerivedDataBackend::GetZenData(const FCacheKey& CacheKey, ECachePolicy Cache
 
 	// Retry request until we get an accepted response or exhaust allowed number of attempts.
 	EGetResult GetResult = EGetResult::NotFound;
-	for (uint32 Attempts = 0; Attempts < UE_ZENDDC_MAX_ATTEMPTS; ++Attempts)
+	for (uint32 Attempts = 0; Attempts < MaxAttempts; ++Attempts)
 	{
-		zen::FScopedRequestPtr Request(RequestPool.Get());
+		Zen::FScopedRequestPtr Request(RequestPool.Get());
 		if (Request.IsValid())
 		{
-			zen::FZenHttpRequest::Result Result = Request->PerformBlockingDownload(QueryUri.ToString(), OutPackage);
+			Zen::FZenHttpRequest::Result Result = Request->PerformBlockingDownload(QueryUri.ToString(), OutPackage);
 			int64 ResponseCode = Request->GetResponseCode();
 			bool bPackageValid = Request->GetResponseFormatValid();
 
 			// Request was successful, make sure we got all the expected data.
-			if (zen::IsSuccessCode(ResponseCode))
+			if (Zen::IsSuccessCode(ResponseCode))
 			{
 				if (bPackageValid)
 				{
@@ -1028,27 +263,28 @@ FZenDerivedDataBackend::PutCachedData(const TCHAR* CacheKey, TArrayView<const ui
 	}
 
 	FSharedBuffer DataBuffer = FSharedBuffer::MakeView(InData.GetData(), InData.Num());
-	return PutZenData(*MakeLegacyZenKey(CacheKey), FCompositeBuffer(DataBuffer), zen::EContentType::Binary);
+	return PutZenData(*MakeLegacyZenKey(CacheKey), FCompositeBuffer(DataBuffer), Zen::EContentType::Binary);
 }
 
 FDerivedDataBackendInterface::EPutStatus
-FZenDerivedDataBackend::PutZenData(const TCHAR* Uri, const FCompositeBuffer& InData, zen::EContentType ContentType)
+FZenDerivedDataBackend::PutZenData(const TCHAR* Uri, const FCompositeBuffer& InData, Zen::EContentType ContentType)
 {
 	TRACE_CPUPROFILER_EVENT_SCOPE(ZenDDC_Put);
 	COOK_STAT(auto Timer = UsageStats.TimePut());
 
-	int64 ResponseCode = 0; uint32 Attempts = 0;
+	int64 ResponseCode = 0; 
+	uint32 Attempts = 0;
 
 	// Retry request until we get an accepted response or exhaust allowed number of attempts.
-	while (ResponseCode == 0 && ++Attempts < UE_ZENDDC_MAX_ATTEMPTS)
+	while (ResponseCode == 0 && ++Attempts < MaxAttempts)
 	{
-		zen::FScopedRequestPtr Request(RequestPool.Get());
+		Zen::FScopedRequestPtr Request(RequestPool.Get());
 		if (Request.IsValid())
 		{
-			zen::FZenHttpRequest::Result Result = Request->PerformBlockingPut(Uri, InData, ContentType);
+			Zen::FZenHttpRequest::Result Result = Request->PerformBlockingPut(Uri, InData, ContentType);
 			ResponseCode = Request->GetResponseCode();
 
-			if (zen::IsSuccessCode(ResponseCode))
+			if (Zen::IsSuccessCode(ResponseCode))
 			{
 				TRACE_COUNTER_ADD(ZenDDC_BytesSent, int64(Request->GetBytesSent()));
 				COOK_STAT(Timer.AddHit(Request->GetBytesSent()));
@@ -1161,18 +397,19 @@ void FZenDerivedDataBackend::RemoveCachedData(const TCHAR* CacheKey, bool bTrans
 	TRACE_CPUPROFILER_EVENT_SCOPE(ZenDDC_Remove);
 	FString Uri = MakeLegacyZenKey(CacheKey);
 
-	int64 ResponseCode = 0; uint32 Attempts = 0;
+	int64 ResponseCode = 0; 
+	uint32 Attempts = 0;
 
 	// Retry request until we get an accepted response or exhaust allowed number of attempts.
-	while (ResponseCode == 0 && ++Attempts < UE_ZENDDC_MAX_ATTEMPTS)
+	while (ResponseCode == 0 && ++Attempts < MaxAttempts)
 	{
-		zen::FScopedRequestPtr Request(RequestPool.Get());
+		Zen::FScopedRequestPtr Request(RequestPool.Get());
 		if (Request)
 		{
-			zen::FZenHttpRequest::Result Result = Request->PerformBlockingDelete(*Uri);
+			Zen::FZenHttpRequest::Result Result = Request->PerformBlockingDelete(*Uri);
 			ResponseCode = Request->GetResponseCode();
 
-			if (zen::IsSuccessCode(ResponseCode))
+			if (Zen::IsSuccessCode(ResponseCode))
 			{
 				return;
 			}
@@ -1444,7 +681,7 @@ bool FZenDerivedDataBackend::PutCacheRecord(const FCacheRecord& Record, FStringV
 	TStringBuilder<256> Uri;
 	AppendZenUri(Record.GetKey(), Uri);
 	AppendPolicyQueryString(Policy, Uri);
-	if (PutZenData(Uri.ToString(), FCompositeBuffer(Buffer.MoveToShared()), zen::EContentType::CompactBinaryPackage)
+	if (PutZenData(Uri.ToString(), FCompositeBuffer(Buffer.MoveToShared()), Zen::EContentType::CompactBinaryPackage)
 		!= FDerivedDataBackendInterface::EPutStatus::Cached)
 	{
 		return false;
@@ -1521,7 +758,7 @@ bool FZenDerivedDataBackend::LegacyPutCacheRecord(const FCacheRecord& Record, FS
 	Writer.Save(Buffer);
 	TStringBuilder<256> Uri;
 	LegacyMakeZenKey(Key, Uri);
-	if (PutZenData(Uri.ToString(), FCompositeBuffer(Buffer.MoveToShared()), zen::EContentType::CompactBinary)
+	if (PutZenData(Uri.ToString(), FCompositeBuffer(Buffer.MoveToShared()), Zen::EContentType::CompactBinary)
 		!= FDerivedDataBackendInterface::EPutStatus::Cached)
 	{
 		return false;
@@ -1540,7 +777,7 @@ bool FZenDerivedDataBackend::LegacyPutCachePayload(const FCacheKey& Key, FString
 	{
 		TStringBuilder<256> Uri;
 		LegacyMakePayloadKey(Key, RawHash, Uri);
-		if (PutZenData(Uri.ToString(), CompressedBuffer, zen::EContentType::Binary) != FDerivedDataBackendInterface::EPutStatus::Cached)
+		if (PutZenData(Uri.ToString(), CompressedBuffer, Zen::EContentType::Binary) != FDerivedDataBackendInterface::EPutStatus::Cached)
 		{
 			return false;
 		}
diff --git a/Engine/Source/Developer/DerivedDataCache/Private/ZenDerivedDataBackend.h b/Engine/Source/Developer/DerivedDataCache/Private/ZenDerivedDataBackend.h
--- a/Engine/Source/Developer/DerivedDataCache/Private/ZenDerivedDataBackend.h
+++ b/Engine/Source/Developer/DerivedDataCache/Private/ZenDerivedDataBackend.h
@@ -4,8 +4,10 @@
 
 #include "CoreMinimal.h"
 
+#include "ZenServerInterface.h"
+
 // Macro for whether to enable the S3 backend. libcurl is not currently available on Mac.
-#if PLATFORM_WINDOWS
+#if UE_WITH_ZEN
 #	define WITH_ZEN_DDC_BACKEND 1
 #else
 #	define WITH_ZEN_DDC_BACKEND 0
@@ -26,7 +28,7 @@ class FCbWriter;
 class FCompositeBuffer;
 struct FIoHash;
 
-namespace zen {
+namespace UE::Zen {
 	enum class EContentType;
 	struct FRequestPool;
 }
@@ -120,7 +122,7 @@ private:
 	EGetResult GetZenData(const TCHAR* Uri, TArray<uint8>* OutData) const;
 
 	// TODO: need ability to specify content type
-	FDerivedDataBackendInterface::EPutStatus PutZenData(const TCHAR* Uri, const FCompositeBuffer& InData, zen::EContentType ContentType);
+	FDerivedDataBackendInterface::EPutStatus PutZenData(const TCHAR* Uri, const FCompositeBuffer& InData, Zen::EContentType ContentType);
 	EGetResult GetZenData(const FCacheKey& Key, ECachePolicy CachePolicy, FCbPackage& OutPackage) const;
 
 	bool PutCacheRecord(const FCacheRecord& Record, FStringView Context, ECachePolicy Policy);
@@ -159,9 +161,10 @@ private:
 	FString Domain;
 	FString Namespace;
 	mutable FDerivedDataCacheUsageStats UsageStats;
-	TUniquePtr<struct zen::FRequestPool> RequestPool;
+	TUniquePtr<UE::Zen::FRequestPool> RequestPool;
 	bool bIsUsable = false;
 	uint32 FailedLoginAttempts = 0;
+	uint32 MaxAttempts = 4;
 	bool bCacheRecordEndpointEnabled = false;
 
 	/** Debug Options */
diff --git a/Engine/Source/Developer/DerivedDataCache/Private/ZenServerInterface.cpp b/Engine/Source/Developer/DerivedDataCache/Private/ZenServerInterface.cpp
new file mode 100644
--- /dev/null
+++ b/Engine/Source/Developer/DerivedDataCache/Private/ZenServerInterface.cpp
@@ -0,0 +1,572 @@
+// Copyright Epic Games, Inc. All Rights Reserved.
+#include "ZenServerInterface.h"
+#include "DerivedDataBackendInterface.h"
+#include "ZenBackendUtils.h"
+
+#if PLATFORM_WINDOWS || PLATFORM_HOLOLENS
+#	include "Windows/WindowsHWrapper.h"
+#	include "Windows/AllowWindowsPlatformTypes.h"
+#endif
+
+#include "curl/curl.h"
+
+#if PLATFORM_WINDOWS || PLATFORM_HOLOLENS
+#	include "Windows/HideWindowsPlatformTypes.h"
+#endif
+
+#include "Memory/CompositeBuffer.h"
+#include "Serialization/CompactBinary.h"
+#include "Serialization/CompactBinaryPackage.h"
+#include "Serialization/CompactBinaryValidation.h"
+#include "Serialization/MemoryReader.h"
+
+#if UE_WITH_ZEN
+
+namespace UE::Zen {
+
+static bool bHasLaunched = false;
+
+FZenServiceInstance::FZenServiceInstance()
+{
+	if (bHasLaunched)
+	{
+		return;
+	}
+
+	FString Parms;
+	bool bLaunchDetached = false;
+	bool bLaunchHidden = false;
+	bool bLaunchReallyHidden = false;
+	uint32* OutProcessID = nullptr;
+	int32 PriorityModifier = 0;
+	const TCHAR* OptionalWorkingDirectory = nullptr;
+	void* PipeWriteChild = nullptr;
+	void* PipeReadChild = nullptr;
+
+	FString MainFilePath = FPaths::ConvertRelativePathToFull(FPaths::Combine(FPaths::EngineDir(), TEXT("Binaries"), TEXT("Win64"), TEXT("zenserver.exe")));
+
+	Parms.Appendf(TEXT("--owner-pid %d"), FPlatformProcess::GetCurrentProcessId());
+
+	FProcHandle Proc = FPlatformProcess::CreateProc(
+		*MainFilePath,
+		*Parms,
+		bLaunchDetached,
+		bLaunchHidden,
+		bLaunchReallyHidden,
+		OutProcessID,
+		PriorityModifier,
+		OptionalWorkingDirectory,
+		PipeWriteChild,
+		PipeReadChild);
+
+	if (!Proc.IsValid())
+	{
+		Proc = FPlatformProcess::CreateElevatedProcess(*MainFilePath, *Parms);
+	}
+
+	bHasLaunched = true;
+}
+
+FZenServiceInstance::~FZenServiceInstance()
+{
+}
+
+bool 
+FZenServiceInstance::IsServiceRunning()
+{
+	return false;
+}
+
+bool 
+FZenServiceInstance::IsServiceReady()
+{
+	return false;
+}
+
+#define UE_ZENDDC_BACKEND_WAIT_INTERVAL			0.01f
+#define UE_ZENDDC_HTTP_REQUEST_TIMEOUT_SECONDS	30L
+#define UE_ZENDDC_HTTP_DEBUG					0
+
+	struct FZenHttpRequest::FStatics
+	{
+		static size_t StaticDebugCallback(CURL* Handle, curl_infotype DebugInfoType, char* DebugInfo, size_t DebugInfoSize, void* UserData);
+		static size_t StaticReadFn(void* Ptr, size_t SizeInBlocks, size_t BlockSizeInBytes, void* UserData);
+		static size_t StaticWriteHeaderFn(void* Ptr, size_t SizeInBlocks, size_t BlockSizeInBytes, void* UserData);
+		static size_t StaticWriteBodyFn(void* Ptr, size_t SizeInBlocks, size_t BlockSizeInBytes, void* UserData);
+		static size_t StaticSeekFn(void* UserData, curl_off_t Offset, int Origin);
+	};
+
+	FZenHttpRequest::FZenHttpRequest(const TCHAR* InDomain, bool bInLogErrors)
+		: bLogErrors(bInLogErrors)
+		, Domain(InDomain)
+	{
+		Curl = curl_easy_init();
+		Reset();
+	}
+
+	FZenHttpRequest::~FZenHttpRequest()
+	{
+		curl_easy_cleanup(Curl);
+	}
+
+	/**
+	 * Resets all options on the request except those that should always be set.
+	 */
+	void FZenHttpRequest::Reset()
+	{
+		Headers.Reset();
+		ResponseHeader.Reset();
+		ResponseBuffer.Reset();
+		ResponseCode = 0;
+		bResponseFormatValid = true;
+		ReadDataView = nullptr;
+		WriteDataBufferPtr = nullptr;
+		WriteHeaderBufferPtr = nullptr;
+		BytesSent = 0;
+		BytesReceived = 0;
+		CurlResult = CURL_LAST;
+
+		curl_easy_reset(Curl);
+
+		// Options that are always set for all connections.
+		curl_easy_setopt(Curl, CURLOPT_CONNECTTIMEOUT, UE_ZENDDC_HTTP_REQUEST_TIMEOUT_SECONDS);
+		curl_easy_setopt(Curl, CURLOPT_FOLLOWLOCATION, 1L);
+		curl_easy_setopt(Curl, CURLOPT_NOSIGNAL, 1L);
+		curl_easy_setopt(Curl, CURLOPT_BUFFERSIZE, 256 * 1024L);
+		//curl_easy_setopt(Curl, CURLOPT_UPLOAD_BUFFERSIZE, 256 * 1024L);
+		// Response functions
+		curl_easy_setopt(Curl, CURLOPT_HEADERDATA, this);
+		curl_easy_setopt(Curl, CURLOPT_HEADERFUNCTION, &FZenHttpRequest::FStatics::StaticWriteHeaderFn);
+		curl_easy_setopt(Curl, CURLOPT_WRITEDATA, this);
+		curl_easy_setopt(Curl, CURLOPT_WRITEFUNCTION, &FZenHttpRequest::FStatics::StaticWriteBodyFn);
+		// Rewind method, handle special error case where request need to rewind data stream
+		curl_easy_setopt(Curl, CURLOPT_SEEKFUNCTION, &FZenHttpRequest::FStatics::StaticSeekFn);
+		curl_easy_setopt(Curl, CURLOPT_SEEKDATA, this);
+		// Debug hooks
+#if UE_ZENDDC_HTTP_DEBUG
+		curl_easy_setopt(Curl, CURLOPT_DEBUGDATA, this);
+		curl_easy_setopt(Curl, CURLOPT_DEBUGFUNCTION, StaticDebugCallback);
+		curl_easy_setopt(Curl, CURLOPT_VERBOSE, 1L);
+#endif
+	}
+
+	FZenHttpRequest::Result FZenHttpRequest::PerformBlockingPut(const TCHAR* Uri, const FCompositeBuffer& Buffer, EContentType ContentType)
+	{
+		uint32 ContentLength = 0u;
+
+		ContentLength = Buffer.GetSize();
+		curl_easy_setopt(Curl, CURLOPT_UPLOAD, 1L);
+		curl_easy_setopt(Curl, CURLOPT_INFILESIZE, ContentLength);
+		curl_easy_setopt(Curl, CURLOPT_READDATA, this);
+		curl_easy_setopt(Curl, CURLOPT_READFUNCTION, &FZenHttpRequest::FStatics::StaticReadFn);
+
+		switch (ContentType)
+		{
+		case EContentType::Binary:
+			Headers.Add(FString(TEXT("Content-Type: application/octet-stream")));
+			break;
+		case EContentType::CompactBinary:
+			Headers.Add(FString(TEXT("Content-Type: application/x-ue-cb")));
+			break;
+		case EContentType::CompactBinaryPackage:
+			Headers.Add(FString(TEXT("Content-Type: application/x-ue-cbpkg")));
+			break;
+		default:
+			checkNoEntry();
+			break;
+		}
+		ReadDataView = &Buffer;
+
+		return PerformBlocking(Uri, RequestVerb::Put, ContentLength);
+	}
+
+	FZenHttpRequest::Result FZenHttpRequest::PerformBlockingDownload(const TCHAR* Uri, TArray<uint8>* Buffer)
+	{
+		curl_easy_setopt(Curl, CURLOPT_HTTPGET, 1L);
+		WriteDataBufferPtr = Buffer;
+
+		return PerformBlocking(Uri, RequestVerb::Get, 0u);
+	}
+
+	FZenHttpRequest::Result FZenHttpRequest::PerformBlockingDownload(const TCHAR* Uri, FCbPackage& OutPackage)
+	{
+		curl_easy_setopt(Curl, CURLOPT_HTTPGET, 1L);
+		OutPackage.Reset();
+		// TODO: When PackageBytes can be written in segments directly, set the WritePtr to the OutPackage and use that
+		TArray<uint8> PackageBytes;
+		WriteDataBufferPtr = &PackageBytes;
+		Result LocalResult = PerformBlocking(Uri, RequestVerb::Get, 0u);
+		if (IsSuccessCode(ResponseCode))
+		{
+			if (ValidateCompactBinaryPackage(FMemoryView(PackageBytes.GetData(), PackageBytes.Num()), ECbValidateMode::All) != ECbValidateError::None)
+			{
+				bResponseFormatValid = false;
+			}
+			else
+			{
+				FMemoryReader PackageLoader(PackageBytes);
+				bResponseFormatValid = OutPackage.TryLoad(PackageLoader);
+			}
+		}
+		return LocalResult;
+	}
+
+	FZenHttpRequest::Result FZenHttpRequest::PerformBlockingHead(const TCHAR* Uri)
+	{
+		curl_easy_setopt(Curl, CURLOPT_NOBODY, 1L);
+
+		return PerformBlocking(Uri, RequestVerb::Head, 0u);
+	}
+
+	FZenHttpRequest::Result FZenHttpRequest::PerformBlockingDelete(const TCHAR* Uri)
+	{
+		curl_easy_setopt(Curl, CURLOPT_POST, 1L);
+		curl_easy_setopt(Curl, CURLOPT_CUSTOMREQUEST, "DELETE");
+
+		return PerformBlocking(Uri, RequestVerb::Delete, 0u);
+	}
+
+	FZenHttpRequest::Result FZenHttpRequest::PerformBlocking(const TCHAR* Uri, RequestVerb Verb, uint32 ContentLength)
+	{
+		static const char* CommonHeaders[] = {
+			"User-Agent: UE",
+			nullptr
+		};
+
+		TRACE_CPUPROFILER_EVENT_SCOPE(ZenDDC_CurlPerform);
+
+		// Setup request options
+		FString Url = FString::Printf(TEXT("%s/%s"), *Domain, Uri);
+		curl_easy_setopt(Curl, CURLOPT_URL, TCHAR_TO_ANSI(*Url));
+
+		// Setup response header buffer. If caller has not setup a response data buffer, use internal.
+		WriteHeaderBufferPtr = &ResponseHeader;
+		if (WriteDataBufferPtr == nullptr)
+		{
+			WriteDataBufferPtr = &ResponseBuffer;
+		}
+
+		// Content-Length should always be set
+		Headers.Add(FString::Printf(TEXT("Content-Length: %d"), ContentLength));
+
+		// Build headers list
+		curl_slist* CurlHeaders = nullptr;
+		// Add common headers
+		for (uint8 i = 0; CommonHeaders[i] != nullptr; ++i)
+		{
+			CurlHeaders = curl_slist_append(CurlHeaders, CommonHeaders[i]);
+		}
+		// Setup added headers
+		for (const FString& Header : Headers)
+		{
+			CurlHeaders = curl_slist_append(CurlHeaders, TCHAR_TO_ANSI(*Header));
+		}
+		curl_easy_setopt(Curl, CURLOPT_HTTPHEADER, CurlHeaders);
+
+		// Shots fired!
+		CurlResult = curl_easy_perform(Curl);
+
+		// Get response code
+		bool bRedirected = false;
+		if (CURLE_OK == curl_easy_getinfo(Curl, CURLINFO_RESPONSE_CODE, &ResponseCode))
+		{
+			bRedirected = (ResponseCode >= 300 && ResponseCode < 400);
+		}
+
+		LogResult(CurlResult, Uri, Verb);
+
+		// Clean up
+		curl_slist_free_all(CurlHeaders);
+
+		return CurlResult == CURLE_OK ? Result::Success : Result::Failed;
+	}
+
+	void FZenHttpRequest::LogResult(long InResult, const TCHAR* Uri, RequestVerb Verb) const
+	{
+		CURLcode Result = (CURLcode) InResult;
+		if (Result == CURLE_OK)
+		{
+			bool bSuccess = false;
+			const TCHAR* VerbStr = nullptr;
+			FString AdditionalInfo;
+
+			const bool Is400 = (ResponseCode >= 400) && (ResponseCode <= 499);
+			const bool Is200 = (ResponseCode >= 200) && (ResponseCode <= 299);
+
+			switch (Verb)
+			{
+			case RequestVerb::Head:
+				bSuccess = Is400 || Is200;
+				VerbStr = TEXT("querying");
+				break;
+			case RequestVerb::Get:
+				bSuccess = Is400 || Is200;
+				VerbStr = TEXT("fetching");
+				AdditionalInfo = FString::Printf(TEXT("Received: %d bytes."), BytesReceived);
+				break;
+			case RequestVerb::Put:
+				bSuccess = Is200;
+				VerbStr = TEXT("updating");
+				AdditionalInfo = FString::Printf(TEXT("Sent: %d bytes."), BytesSent);
+				break;
+			case RequestVerb::Post:
+				bSuccess = Is200;
+				VerbStr = TEXT("posting");
+				break;
+			case RequestVerb::Delete:
+				bSuccess = Is200;
+				VerbStr = TEXT("deleting");
+				break;
+			}
+
+			if (bSuccess)
+			{
+				UE_LOG(
+					LogDerivedDataCache,
+					Verbose,
+					TEXT("Finished %s HTTP cache entry (response %d) from %s. %s"),
+					VerbStr,
+					ResponseCode,
+					Uri,
+					*AdditionalInfo
+				);
+			}
+			else if (bLogErrors)
+			{
+				// Print the response body if we got one, otherwise print header.
+				FString Response = GetAnsiBufferAsString(ResponseBuffer.Num() > 0 ? ResponseBuffer : ResponseHeader);
+				Response.ReplaceCharInline('\n', ' ');
+				Response.ReplaceCharInline('\r', ' ');
+				UE_LOG(
+					LogDerivedDataCache,
+					Error,
+					TEXT("Failed %s HTTP cache entry (response %d) from %s. Response: %s"),
+					VerbStr,
+					ResponseCode,
+					Uri,
+					*Response
+				);
+			}
+		}
+		else if (bLogErrors)
+		{
+			UE_LOG(
+				LogDerivedDataCache,
+				Error,
+				TEXT("Error while connecting to %s: %s"),
+				*Domain,
+				ANSI_TO_TCHAR(curl_easy_strerror(Result))
+			);
+		}
+	}
+
+	FString FZenHttpRequest::GetAnsiBufferAsString(const TArray<uint8>& Buffer)
+	{
+		// Content is NOT null-terminated; we need to specify lengths here
+		FUTF8ToTCHAR TCHARData(reinterpret_cast<const ANSICHAR*>(Buffer.GetData()), Buffer.Num());
+		return FString(TCHARData.Length(), TCHARData.Get());
+	}
+
+	size_t FZenHttpRequest::FStatics::StaticDebugCallback(CURL* Handle, curl_infotype DebugInfoType, char* DebugInfo, size_t DebugInfoSize, void* UserData)
+	{
+		FZenHttpRequest* Request = static_cast<FZenHttpRequest*>(UserData);
+
+		switch (DebugInfoType)
+		{
+		case CURLINFO_TEXT:
+		{
+			// Truncate at 1023 characters. This is just an arbitrary number based on a buffer size seen in
+			// the libcurl code.
+			DebugInfoSize = FMath::Min(DebugInfoSize, (size_t)1023);
+
+			// Calculate the actual length of the string due to incorrect use of snprintf() in lib/vtls/openssl.c.
+			char* FoundNulPtr = (char*)memchr(DebugInfo, 0, DebugInfoSize);
+			int CalculatedSize = FoundNulPtr != nullptr ? FoundNulPtr - DebugInfo : DebugInfoSize;
+
+			auto ConvertedString = StringCast<TCHAR>(static_cast<const ANSICHAR*>(DebugInfo), CalculatedSize);
+			FString DebugText(ConvertedString.Length(), ConvertedString.Get());
+			DebugText.ReplaceInline(TEXT("\n"), TEXT(""), ESearchCase::CaseSensitive);
+			DebugText.ReplaceInline(TEXT("\r"), TEXT(""), ESearchCase::CaseSensitive);
+			UE_LOG(LogDerivedDataCache, VeryVerbose, TEXT("%p: '%s'"), Request, *DebugText);
+		}
+		break;
+
+		case CURLINFO_HEADER_IN:
+			UE_LOG(LogDerivedDataCache, VeryVerbose, TEXT("%p: Received header (%d bytes)"), Request, DebugInfoSize);
+			break;
+
+		case CURLINFO_DATA_IN:
+			UE_LOG(LogDerivedDataCache, VeryVerbose, TEXT("%p: Received data (%d bytes)"), Request, DebugInfoSize);
+			break;
+
+		case CURLINFO_DATA_OUT:
+			UE_LOG(LogDerivedDataCache, VeryVerbose, TEXT("%p: Sent data (%d bytes)"), Request, DebugInfoSize);
+			break;
+		}
+
+		return 0;
+	}
+
+	size_t FZenHttpRequest::FStatics::StaticReadFn(void* Ptr, size_t SizeInBlocks, size_t BlockSizeInBytes, void* UserData)
+	{
+		FZenHttpRequest* Request = static_cast<FZenHttpRequest*>(UserData);
+		check(Request->ReadDataView);
+		const FCompositeBuffer& ReadDataView = *Request->ReadDataView;
+
+		const size_t Offset = Request->BytesSent;
+		const size_t ReadSize = FMath::Min((size_t)ReadDataView.GetSize() - Offset, SizeInBlocks * BlockSizeInBytes);
+		check(ReadDataView.GetSize() >= Offset + ReadSize);
+
+		Memcpy(Ptr, ReadDataView, Offset, ReadSize);
+		Request->BytesSent += ReadSize;
+		return ReadSize;
+
+		return 0;
+	}
+
+	size_t FZenHttpRequest::FStatics::StaticWriteHeaderFn(void* Ptr, size_t SizeInBlocks, size_t BlockSizeInBytes, void* UserData)
+	{
+		FZenHttpRequest* Request = static_cast<FZenHttpRequest*>(UserData);
+		const size_t WriteSize = SizeInBlocks * BlockSizeInBytes;
+		TArray<uint8>* WriteHeaderBufferPtr = Request->WriteHeaderBufferPtr;
+		if (WriteHeaderBufferPtr && WriteSize > 0)
+		{
+			const size_t CurrentBufferLength = WriteHeaderBufferPtr->Num();
+			if (CurrentBufferLength > 0)
+			{
+				// Remove the previous zero termination
+				(*WriteHeaderBufferPtr)[CurrentBufferLength - 1] = ' ';
+			}
+
+			// Write the header
+			WriteHeaderBufferPtr->Append((const uint8*)Ptr, WriteSize + 1);
+			(*WriteHeaderBufferPtr)[WriteHeaderBufferPtr->Num() - 1] = 0; // Zero terminate string
+			return WriteSize;
+		}
+		return 0;
+	}
+
+	size_t FZenHttpRequest::FStatics::StaticWriteBodyFn(void* Ptr, size_t SizeInBlocks, size_t BlockSizeInBytes, void* UserData)
+	{
+		FZenHttpRequest* Request = static_cast<FZenHttpRequest*>(UserData);
+		const size_t WriteSize = SizeInBlocks * BlockSizeInBytes;
+		TArray<uint8>* WriteDataBufferPtr = Request->WriteDataBufferPtr;
+
+		if (WriteDataBufferPtr && WriteSize > 0)
+		{
+			// If this is the first part of the body being received, try to reserve 
+			// memory if content length is defined in the header.
+			if (Request->BytesReceived == 0 && Request->WriteHeaderBufferPtr)
+			{
+				static const ANSICHAR* ContentLengthHeaderStr = "Content-Length: ";
+				const ANSICHAR* Header = (const ANSICHAR*)Request->WriteHeaderBufferPtr->GetData();
+
+				if (const ANSICHAR* ContentLengthHeader = FCStringAnsi::Strstr(Header, ContentLengthHeaderStr))
+				{
+					size_t ContentLength = (size_t)FCStringAnsi::Atoi64(ContentLengthHeader + strlen(ContentLengthHeaderStr));
+					if (ContentLength > 0)
+					{
+						WriteDataBufferPtr->Reserve(ContentLength);
+					}
+				}
+			}
+
+			// Write to the target buffer
+			WriteDataBufferPtr->Append((const uint8*)Ptr, WriteSize);
+			Request->BytesReceived += WriteSize;
+			return WriteSize;
+		}
+
+		return 0;
+	}
+
+	size_t FZenHttpRequest::FStatics::StaticSeekFn(void* UserData, curl_off_t Offset, int Origin)
+	{
+		FZenHttpRequest* Request = static_cast<FZenHttpRequest*>(UserData);
+		size_t NewPosition = 0;
+		uint64 ReadDataSize = Request->ReadDataView ? Request->ReadDataView->GetSize() : 0;
+
+		switch (Origin)
+		{
+		case SEEK_SET: NewPosition = Offset; break;
+		case SEEK_CUR: NewPosition = Request->BytesSent + Offset; break;
+		case SEEK_END: NewPosition = ReadDataSize + Offset; break;
+		}
+
+		// Make sure we don't seek outside of the buffer
+		if (NewPosition < 0 || NewPosition >= ReadDataSize)
+		{
+			return CURL_SEEKFUNC_FAIL;
+		}
+
+		// Update the used offset
+		Request->BytesSent = NewPosition;
+		return CURL_SEEKFUNC_OK;
+	}
+
+	//////////////////////////////////////////////////////////////////////////
+
+	FRequestPool::FRequestPool(const TCHAR* InServiceUrl)
+	{
+		for (uint8 i = 0; i < Pool.Num(); ++i)
+		{
+			Pool[i].Usage = 0u;
+			Pool[i].Request = new FZenHttpRequest(InServiceUrl, true);
+		}
+	}
+
+	FRequestPool::~FRequestPool()
+	{
+		for (uint8 i = 0; i < Pool.Num(); ++i)
+		{
+			// No requests should be in use by now.
+			check(Pool[i].Usage.Load(EMemoryOrder::Relaxed) == 0u);
+			delete Pool[i].Request;
+		}
+	}
+
+	/**
+		* Block until a request is free. Once a request has been returned it is
+		* "owned by the caller and need to release it to the pool when work has been completed.
+		* @return Usable request instance.
+		*/
+	FZenHttpRequest* FRequestPool::WaitForFreeRequest()
+	{
+		TRACE_CPUPROFILER_EVENT_SCOPE(ZenDDC_WaitForConnPool);
+		while (true)
+		{
+			for (uint8 i = 0; i < Pool.Num(); ++i)
+			{
+				if (!Pool[i].Usage.Load(EMemoryOrder::Relaxed))
+				{
+					uint8 Expected = 0u;
+					if (Pool[i].Usage.CompareExchange(Expected, 1u))
+					{
+						return Pool[i].Request;
+					}
+				}
+			}
+			FPlatformProcess::Sleep(UE_ZENDDC_BACKEND_WAIT_INTERVAL);
+		}
+	}
+
+	/**
+		* Release request to the pool.
+		* @param Request Request that should be freed. Note that any buffer owened by the request can now be reset.
+		*/
+	void FRequestPool::ReleaseRequestToPool(FZenHttpRequest* Request)
+	{
+		for (uint8 i = 0; i < Pool.Num(); ++i)
+		{
+			if (Pool[i].Request == Request)
+			{
+				Request->Reset();
+				uint8 Expected = 1u;
+				Pool[i].Usage.CompareExchange(Expected, 0u);
+				return;
+			}
+		}
+		check(false);
+	}
+}
+#endif // UE_WITH_ZEN
diff --git a/Engine/Source/Editor/UnrealEd/UnrealEd.Build.cs b/Engine/Source/Editor/UnrealEd/UnrealEd.Build.cs
--- a/Engine/Source/Editor/UnrealEd/UnrealEd.Build.cs
+++ b/Engine/Source/Editor/UnrealEd/UnrealEd.Build.cs
@@ -132,6 +132,7 @@ public class UnrealEd : ModuleRules
 				"BlueprintGraph",
 				"CinematicCamera",
 				"CurveEditor",
+				"DerivedDataCache",
 				"DesktopPlatform",
 				"LauncherPlatform",
 				"EditorStyle",
