DDC: Replaced FRequest with FRequestGroup to handle nested and chained requests more robustly

Any function that may create a request now has a IRequestOwner& parameter, and uses the Begin and End functions on the owner to manage the lifetime of any requests that it creates, as well as using End to invoke the completion callback for any request which has one.

The new FRequestBarrier may be used to block a group from being considered complete in a scope where more requests may be added to it.

#rnx

diff --git a/Engine/Source/Developer/DerivedDataCache/Private/DerivedDataBackendAsyncPutWrapper.cpp b/Engine/Source/Developer/DerivedDataCache/Private/DerivedDataBackendAsyncPutWrapper.cpp
--- a/Engine/Source/Developer/DerivedDataCache/Private/DerivedDataBackendAsyncPutWrapper.cpp
+++ b/Engine/Source/Developer/DerivedDataCache/Private/DerivedDataBackendAsyncPutWrapper.cpp
@@ -343,11 +343,12 @@ TSharedRef<FDerivedDataCacheStatsNode> FDerivedDataBackendAsyncPutWrapper::Gathe
 	return Usage;
 }
 
-class FDerivedDataAsyncWrapperRequest final : public IRequest, private IQueuedWork
+class FDerivedDataAsyncWrapperRequest final : public FRequestBase, private IQueuedWork
 {
 public:
-	inline explicit FDerivedDataAsyncWrapperRequest(TUniqueFunction<void (bool bCancel)>&& InFunction)
-		: Function(MoveTemp(InFunction))
+	inline FDerivedDataAsyncWrapperRequest(IRequestOwner& InOwner, TUniqueFunction<void (bool bCancel)>&& InFunction)
+		: Owner(InOwner)
+		, Function(MoveTemp(InFunction))
 		, DoneEvent(FPlatformProcess::GetSynchEventFromPool(true))
 	{
 	}
@@ -360,7 +361,7 @@ public:
 	inline void Start(EPriority Priority)
 	{
 		FDerivedDataBackend::Get().AddToAsyncCompletionCounter(1);
-		AddRef(); // Keep the request alive until the queued work executes.
+		Owner.Begin(this);
 
 		DoneEvent->Reset();
 		WorkNotFinishedCounter.fetch_add(1, std::memory_order_relaxed);
@@ -370,11 +371,13 @@ public:
 	inline void Execute(bool bCancel)
 	{
 		FScopeCycleCounter Scope(GetStatId(), /*bAlways*/ true);
-		Function(/*bCancel*/ false);
-		WorkNotFinishedCounter.fetch_sub(1, std::memory_order_release);
-		DoneEvent->Trigger();
-
-		Release(); // DO NOT ACCESS ANY MEMBERS PAST THIS POINT!
+		Owner.End(this, [this, bCancel]
+		{
+			Function(bCancel);
+			WorkNotFinishedCounter.fetch_sub(1, std::memory_order_release);
+			DoneEvent->Trigger();
+		});
+		// DO NOT ACCESS ANY MEMBERS PAST THIS POINT!
 		FDerivedDataBackend::Get().AddToAsyncCompletionCounter(-1);
 	}
 
@@ -390,7 +393,7 @@ public:
 
 	inline void Cancel() final
 	{
-		if (!Poll())
+		if (WorkNotFinishedCounter.load(std::memory_order_acquire) > 0)
 		{
 			if (GDDCIOThreadPool->RetractQueuedWork(this))
 			{
@@ -405,7 +408,7 @@ public:
 
 	inline void Wait() final
 	{
-		if (!Poll())
+		if (WorkNotFinishedCounter.load(std::memory_order_acquire) > 0)
 		{
 			if (GDDCIOThreadPool->RetractQueuedWork(this))
 			{
@@ -419,24 +422,6 @@ public:
 		}
 	}
 
-	inline bool Poll() final
-	{
-		return WorkNotFinishedCounter.load(std::memory_order_acquire) == 0;
-	}
-
-	inline void AddRef() const final
-	{
-		ReferenceCount.fetch_add(1, std::memory_order_relaxed);
-	}
-
-	inline void Release() const final
-	{
-		if (ReferenceCount.fetch_sub(1, std::memory_order_acq_rel) == 1)
-		{
-			delete this;
-		}
-	}
-
 private:
 	static EQueuedWorkPriority GetPriority(EPriority Priority)
 	{
@@ -470,31 +455,33 @@ private:
 	}
 
 private:
-	mutable std::atomic<uint32> ReferenceCount{0};
 	std::atomic<uint32> WorkNotFinishedCounter{0};
+	IRequestOwner& Owner;
 	TUniqueFunction<void (bool bCancel)> Function;
 	FEvent* DoneEvent;
 };
 
-FRequest FDerivedDataBackendAsyncPutWrapper::Put(
+void FDerivedDataBackendAsyncPutWrapper::Put(
 	TConstArrayView<FCacheRecord> Records,
 	FStringView Context,
 	ECachePolicy Policy,
-	EPriority Priority,
+	IRequestOwner& Owner,
 	FOnCachePutComplete&& OnComplete)
 {
-	if (Priority == EPriority::Blocking)
+	if (Owner.GetPriority() == EPriority::Blocking)
 	{
-		return InnerBackend->Put(Records, Context, Policy, Priority, MoveTemp(OnComplete));
+		InnerBackend->Put(Records, Context, Policy, Owner, MoveTemp(OnComplete));
 	}
 	else
 	{
-		TRequest<FDerivedDataAsyncWrapperRequest> Request(new FDerivedDataAsyncWrapperRequest(
-			[Backend = InnerBackend, Records = TArray<FCacheRecord>(Records), Context = FString(Context), Policy, Priority, OnComplete = MoveTemp(OnComplete)](bool bCancel) mutable
+		FDerivedDataAsyncWrapperRequest* Request = new FDerivedDataAsyncWrapperRequest(Owner,
+			[this, Records = TArray<FCacheRecord>(Records), Context = FString(Context), Policy, OnComplete = MoveTemp(OnComplete)](bool bCancel) mutable
 			{
 				if (!bCancel)
 				{
-					Backend->Put(Records, Context, Policy, Priority, MoveTemp(OnComplete)).Wait();
+					FRequestGroup BlockingGroup = Factory.CreateGroup(EPriority::Blocking);
+					InnerBackend->Put(Records, Context, Policy, BlockingGroup, MoveTemp(OnComplete));
+					BlockingGroup.Wait();
 				}
 				else if (OnComplete)
 				{
@@ -503,31 +490,32 @@ FRequest FDerivedDataBackendAsyncPutWrapper::Put(
 						OnComplete({Record.GetKey(), EStatus::Canceled});
 					}
 				}
-			}));
-		Request->Start(Priority == EPriority::Normal ? EPriority::Low : Priority);
-		return Request;
+			});
+		Request->Start(Owner.GetPriority());
 	}
 }
 
-FRequest FDerivedDataBackendAsyncPutWrapper::Get(
+void FDerivedDataBackendAsyncPutWrapper::Get(
 	TConstArrayView<FCacheKey> Keys,
 	FStringView Context,
 	ECachePolicy Policy,
-	EPriority Priority,
+	IRequestOwner& Owner,
 	FOnCacheGetComplete&& OnComplete)
 {
-	if (Priority == EPriority::Blocking)
+	if (Owner.GetPriority() == EPriority::Blocking)
 	{
-		return InnerBackend->Get(Keys, Context, Policy, Priority, MoveTemp(OnComplete));
+		InnerBackend->Get(Keys, Context, Policy, Owner, MoveTemp(OnComplete));
 	}
 	else
 	{
-		TRequest<FDerivedDataAsyncWrapperRequest> Request(new FDerivedDataAsyncWrapperRequest(
-			[&Factory = Factory, Backend = InnerBackend, Keys = TArray<FCacheKey>(Keys), Context = FString(Context), Policy, Priority, OnComplete = MoveTemp(OnComplete)](bool bCancel) mutable
+		FDerivedDataAsyncWrapperRequest* Request = new FDerivedDataAsyncWrapperRequest(Owner,
+			[this, Keys = TArray<FCacheKey>(Keys), Context = FString(Context), Policy, OnComplete = MoveTemp(OnComplete)](bool bCancel) mutable
 			{
 				if (!bCancel)
 				{
-					Backend->Get(Keys, Context, Policy, Priority, MoveTemp(OnComplete)).Wait();
+					FRequestGroup BlockingGroup = Factory.CreateGroup(EPriority::Blocking);
+					InnerBackend->Get(Keys, Context, Policy, BlockingGroup, MoveTemp(OnComplete));
+					BlockingGroup.Wait();
 				}
 				else if (OnComplete)
 				{
@@ -536,31 +524,32 @@ FRequest FDerivedDataBackendAsyncPutWrapper::Get(
 						OnComplete({Factory.CreateRecord(Key).Build(), EStatus::Canceled});
 					}
 				}
-			}));
-		Request->Start(Priority);
-		return Request;
+			});
+		Request->Start(Owner.GetPriority());
 	}
 }
 
-FRequest FDerivedDataBackendAsyncPutWrapper::GetPayload(
+void FDerivedDataBackendAsyncPutWrapper::GetPayload(
 	TConstArrayView<FCachePayloadKey> Keys,
 	FStringView Context,
 	ECachePolicy Policy,
-	EPriority Priority,
+	IRequestOwner& Owner,
 	FOnCacheGetPayloadComplete&& OnComplete)
 {
-	if (Priority == EPriority::Blocking)
+	if (Owner.GetPriority() == EPriority::Blocking)
 	{
-		return InnerBackend->GetPayload(Keys, Context, Policy, Priority, MoveTemp(OnComplete));
+		InnerBackend->GetPayload(Keys, Context, Policy, Owner, MoveTemp(OnComplete));
 	}
 	else
 	{
-		TRequest<FDerivedDataAsyncWrapperRequest> Request(new FDerivedDataAsyncWrapperRequest(
-			[Backend = InnerBackend, Keys = TArray<FCachePayloadKey>(Keys), Context = FString(Context), Policy, Priority, OnComplete = MoveTemp(OnComplete)](bool bCancel) mutable
+		FDerivedDataAsyncWrapperRequest* Request = new FDerivedDataAsyncWrapperRequest(Owner,
+			[this, Keys = TArray<FCachePayloadKey>(Keys), Context = FString(Context), Policy, OnComplete = MoveTemp(OnComplete)](bool bCancel) mutable
 			{
 				if (!bCancel)
 				{
-					Backend->GetPayload(Keys, Context, Policy, Priority, MoveTemp(OnComplete)).Wait();
+					FRequestGroup BlockingGroup = Factory.CreateGroup(EPriority::Blocking);
+					InnerBackend->GetPayload(Keys, Context, Policy, BlockingGroup, MoveTemp(OnComplete));
+					BlockingGroup.Wait();
 				}
 				else if (OnComplete)
 				{
@@ -569,9 +558,8 @@ FRequest FDerivedDataBackendAsyncPutWrapper::GetPayload(
 						OnComplete({Key.CacheKey, FPayload(Key.Id), EStatus::Canceled});
 					}
 				}
-			}));
-		Request->Start(Priority);
-		return Request;
+			});
+		Request->Start(Owner.GetPriority());
 	}
 }
 
diff --git a/Engine/Source/Developer/DerivedDataCache/Private/DerivedDataBackendAsyncPutWrapper.h b/Engine/Source/Developer/DerivedDataCache/Private/DerivedDataBackendAsyncPutWrapper.h
--- a/Engine/Source/Developer/DerivedDataCache/Private/DerivedDataBackendAsyncPutWrapper.h
+++ b/Engine/Source/Developer/DerivedDataCache/Private/DerivedDataBackendAsyncPutWrapper.h
@@ -135,25 +135,25 @@ public:
 
 	virtual bool ApplyDebugOptions(FBackendDebugOptions& InOptions) override;
 
-	virtual FRequest Put(
+	virtual void Put(
 		TConstArrayView<FCacheRecord> Records,
 		FStringView Context,
 		ECachePolicy Policy,
-		EPriority Priority,
+		IRequestOwner& Owner,
 		FOnCachePutComplete&& OnComplete) override;
 
-	virtual FRequest Get(
+	virtual void Get(
 		TConstArrayView<FCacheKey> Keys,
 		FStringView Context,
 		ECachePolicy Policy,
-		EPriority Priority,
+		IRequestOwner& Owner,
 		FOnCacheGetComplete&& OnComplete) override;
 
-	virtual FRequest GetPayload(
+	virtual void GetPayload(
 		TConstArrayView<FCachePayloadKey> Keys,
 		FStringView Context,
 		ECachePolicy Policy,
-		EPriority Priority,
+		IRequestOwner& Owner,
 		FOnCacheGetPayloadComplete&& OnComplete) override;
 
 	virtual void CancelAll() override;
diff --git a/Engine/Source/Developer/DerivedDataCache/Private/DerivedDataBackendCorruptionWrapper.h b/Engine/Source/Developer/DerivedDataCache/Private/DerivedDataBackendCorruptionWrapper.h
--- a/Engine/Source/Developer/DerivedDataCache/Private/DerivedDataBackendCorruptionWrapper.h
+++ b/Engine/Source/Developer/DerivedDataCache/Private/DerivedDataBackendCorruptionWrapper.h
@@ -251,34 +251,34 @@ public:
 		return InnerBackend->ApplyDebugOptions(InOptions);
 	}
 
-	virtual FRequest Put(
+	virtual void Put(
 		TConstArrayView<FCacheRecord> Records,
 		FStringView Context,
 		ECachePolicy Policy,
-		EPriority Priority,
+		IRequestOwner& Owner,
 		FOnCachePutComplete&& OnComplete) override
 	{
-		return InnerBackend->Put(Records, Context, Policy, Priority, MoveTemp(OnComplete));
+		return InnerBackend->Put(Records, Context, Policy, Owner, MoveTemp(OnComplete));
 	}
 
-	virtual FRequest Get(
+	virtual void Get(
 		TConstArrayView<FCacheKey> Keys,
 		FStringView Context,
 		ECachePolicy Policy,
-		EPriority Priority,
+		IRequestOwner& Owner,
 		FOnCacheGetComplete&& OnComplete) override
 	{
-		return InnerBackend->Get(Keys, Context, Policy, Priority, MoveTemp(OnComplete));
+		return InnerBackend->Get(Keys, Context, Policy, Owner, MoveTemp(OnComplete));
 	}
 
-	virtual FRequest GetPayload(
+	virtual void GetPayload(
 		TConstArrayView<FCachePayloadKey> Keys,
 		FStringView Context,
 		ECachePolicy Policy,
-		EPriority Priority,
+		IRequestOwner& Owner,
 		FOnCacheGetPayloadComplete&& OnComplete) override
 	{
-		return InnerBackend->GetPayload(Keys, Context, Policy, Priority, MoveTemp(OnComplete));
+		return InnerBackend->GetPayload(Keys, Context, Policy, Owner, MoveTemp(OnComplete));
 	}
 
 	virtual void CancelAll() override
diff --git a/Engine/Source/Developer/DerivedDataCache/Private/DerivedDataBackendThrottleWrapper.h b/Engine/Source/Developer/DerivedDataCache/Private/DerivedDataBackendThrottleWrapper.h
--- a/Engine/Source/Developer/DerivedDataCache/Private/DerivedDataBackendThrottleWrapper.h
+++ b/Engine/Source/Developer/DerivedDataCache/Private/DerivedDataBackendThrottleWrapper.h
@@ -4,6 +4,7 @@
 
 #include "CoreMinimal.h"
 #include "DerivedDataBackendInterface.h"
+#include "DerivedDataCacheUsageStats.h"
 #include <atomic>
 
 namespace UE::DerivedData::Backends
@@ -160,40 +161,40 @@ public:
 		return InnerBackend->ApplyDebugOptions(InOptions);
 	}
 
-	virtual FRequest Put(
+	virtual void Put(
 		TConstArrayView<FCacheRecord> Records,
 		FStringView Context,
 		ECachePolicy Policy,
-		EPriority Priority,
+		IRequestOwner& Owner,
 		FOnCachePutComplete&& OnComplete) override
 	{
 		ThrottlingScope Scope(this);
 
-		return InnerBackend->Put(Records, Context, Policy, Priority, MoveTemp(OnComplete));
+		InnerBackend->Put(Records, Context, Policy, Owner, MoveTemp(OnComplete));
 	}
 
-	virtual FRequest Get(
+	virtual void Get(
 		TConstArrayView<FCacheKey> Keys,
 		FStringView Context,
 		ECachePolicy Policy,
-		EPriority Priority,
+		IRequestOwner& Owner,
 		FOnCacheGetComplete&& OnComplete) override
 	{
 		ThrottlingScope Scope(this);
 
-		return InnerBackend->Get(Keys, Context, Policy, Priority, MoveTemp(OnComplete));
+		InnerBackend->Get(Keys, Context, Policy, Owner, MoveTemp(OnComplete));
 	}
 
-	virtual FRequest GetPayload(
+	virtual void GetPayload(
 		TConstArrayView<FCachePayloadKey> Keys,
 		FStringView Context,
 		ECachePolicy Policy,
-		EPriority Priority,
+		IRequestOwner& Owner,
 		FOnCacheGetPayloadComplete&& OnComplete) override
 	{
 		ThrottlingScope Scope(this);
 
-		return InnerBackend->GetPayload(Keys, Context, Policy, Priority, MoveTemp(OnComplete));
+		InnerBackend->GetPayload(Keys, Context, Policy, Owner, MoveTemp(OnComplete));
 	}
 
 	virtual void CancelAll() override
diff --git a/Engine/Source/Developer/DerivedDataCache/Private/DerivedDataBackendVerifyWrapper.h b/Engine/Source/Developer/DerivedDataCache/Private/DerivedDataBackendVerifyWrapper.h
--- a/Engine/Source/Developer/DerivedDataCache/Private/DerivedDataBackendVerifyWrapper.h
+++ b/Engine/Source/Developer/DerivedDataCache/Private/DerivedDataBackendVerifyWrapper.h
@@ -162,34 +162,34 @@ public:
 		return InnerBackend->ApplyDebugOptions(InOptions);
 	}
 
-	virtual FRequest Put(
+	virtual void Put(
 		TConstArrayView<FCacheRecord> Records,
 		FStringView Context,
 		ECachePolicy Policy,
-		EPriority Priority,
+		IRequestOwner& Owner,
 		FOnCachePutComplete&& OnComplete) override
 	{
-		return InnerBackend->Put(Records, Context, Policy, Priority, MoveTemp(OnComplete));
+		return InnerBackend->Put(Records, Context, Policy, Owner, MoveTemp(OnComplete));
 	}
 
-	virtual FRequest Get(
+	virtual void Get(
 		TConstArrayView<FCacheKey> Keys,
 		FStringView Context,
 		ECachePolicy Policy,
-		EPriority Priority,
+		IRequestOwner& Owner,
 		FOnCacheGetComplete&& OnComplete) override
 	{
-		return InnerBackend->Get(Keys, Context, Policy, Priority, MoveTemp(OnComplete));
+		return InnerBackend->Get(Keys, Context, Policy, Owner, MoveTemp(OnComplete));
 	}
 
-	virtual FRequest GetPayload(
+	virtual void GetPayload(
 		TConstArrayView<FCachePayloadKey> Keys,
 		FStringView Context,
 		ECachePolicy Policy,
-		EPriority Priority,
+		IRequestOwner& Owner,
 		FOnCacheGetPayloadComplete&& OnComplete) override
 	{
-		return InnerBackend->GetPayload(Keys, Context, Policy, Priority, MoveTemp(OnComplete));
+		return InnerBackend->GetPayload(Keys, Context, Policy, Owner, MoveTemp(OnComplete));
 	}
 
 	virtual void CancelAll() override
diff --git a/Engine/Source/Developer/DerivedDataCache/Private/DerivedDataBackends.cpp b/Engine/Source/Developer/DerivedDataCache/Private/DerivedDataBackends.cpp
--- a/Engine/Source/Developer/DerivedDataCache/Private/DerivedDataBackends.cpp
+++ b/Engine/Source/Developer/DerivedDataCache/Private/DerivedDataBackends.cpp
@@ -23,6 +23,7 @@
 #include "DerivedDataBackendCorruptionWrapper.h"
 #include "DerivedDataBackendThrottleWrapper.h"
 #include "DerivedDataBackendVerifyWrapper.h"
+#include "DerivedDataRequest.h"
 #include "DerivedDataUtilsInterface.h"
 #include "Misc/EngineBuildSettings.h"
 #include "Modules/ModuleManager.h"
@@ -1175,6 +1176,11 @@ public:
 		return MakeShared<FDerivedDataCacheStatsNode>(nullptr, TEXT(""));
 	}
 
+	virtual FRequestGroup CreateGroup(EPriority Priority) override
+	{
+		return Private::CreateRequestGroup(Priority);
+	}
+
 	virtual FCacheBucket CreateBucket(FStringView Name) override
 	{
 		return Private::CreateCacheBucket(Name);
diff --git a/Engine/Source/Developer/DerivedDataCache/Private/DerivedDataBuild.cpp b/Engine/Source/Developer/DerivedDataCache/Private/DerivedDataBuild.cpp
--- a/Engine/Source/Developer/DerivedDataCache/Private/DerivedDataBuild.cpp
+++ b/Engine/Source/Developer/DerivedDataCache/Private/DerivedDataBuild.cpp
@@ -12,6 +12,7 @@
 #include "DerivedDataBuildSession.h"
 #include "DerivedDataBuildWorkerRegistry.h"
 #include "DerivedDataCache.h"
+#include "DerivedDataRequest.h"
 #include "Misc/Guid.h"
 
 namespace UE::DerivedData::Private
@@ -72,6 +73,11 @@ public:
 	{
 	}
 
+	FRequestGroup CreateGroup(EPriority Priority) final
+	{
+		return CreateRequestGroup(Priority);
+	}
+
 	FBuildDefinitionBuilder CreateDefinition(FStringView Name, FStringView Function) final
 	{
 		return CreateBuildDefinition(Name, Function);
diff --git a/Engine/Source/Developer/DerivedDataCache/Private/DerivedDataBuildJob.cpp b/Engine/Source/Developer/DerivedDataCache/Private/DerivedDataBuildJob.cpp
--- a/Engine/Source/Developer/DerivedDataCache/Private/DerivedDataBuildJob.cpp
+++ b/Engine/Source/Developer/DerivedDataCache/Private/DerivedDataBuildJob.cpp
@@ -25,14 +25,11 @@
 #include "HAL/CriticalSection.h"
 #include "HAL/Event.h"
 #include "HAL/FileManager.h"
-#include "HAL/PlatformProcess.h"
-#include "HAL/PlatformMath.h"
 #include "Misc/CommandLine.h"
 #include "Misc/Guid.h"
 #include "Misc/Parse.h"
 #include "Misc/Paths.h"
 #include "Misc/PathViews.h"
-#include "Misc/ScopeExit.h"
 #include "Misc/ScopeRWLock.h"
 #include "Misc/StringBuilder.h"
 #include "ProfilingDebugging/CpuProfilerTrace.h"
@@ -41,7 +38,6 @@
 #include "String/ParseTokens.h"
 #include "Templates/Function.h"
 #include "Templates/RefCounting.h"
-#include <atomic>
 
 namespace UE::DerivedData::Private
 {
@@ -97,11 +93,11 @@ class FBuildJob final : public IBuildJob
 {
 public:
 	/** Resolve the key to a definition, then build like the definition constructor. */
-	FBuildJob(ICache& Cache, IBuild& BuildSystem, IBuildInputResolver* InputResolver, const FBuildKey& Key);
+	FBuildJob(const FBuildJobCreateParams& Params, const FBuildKey& Key, FOnBuildJobComplete&& OnComplete);
 	/** Resolve the definition to an action, then build like the action constructor. */
-	FBuildJob(ICache& Cache, IBuild& BuildSystem, IBuildInputResolver* InputResolver, const FBuildDefinition& Definition);
+	FBuildJob(const FBuildJobCreateParams& Params, const FBuildDefinition& Definition, FOnBuildJobComplete&& OnComplete);
 	/** Query the cache, attempt remote execution, resolve inputs, fall back to local execution, store to the cache. */
-	FBuildJob(ICache& Cache, IBuild& BuildSystem, IBuildInputResolver* InputResolver, const FBuildAction& Action, const FOptionalBuildInputs& Inputs);
+	FBuildJob(const FBuildJobCreateParams& Params, const FBuildAction& Action, const FOptionalBuildInputs& Inputs, FOnBuildJobComplete&& OnComplete);
 
 	/** Destroy the job, which must be complete or not started. */
 	~FBuildJob();
@@ -110,26 +106,15 @@ public:
 
 	inline FStringView GetName() const final { return Name; }
 	inline FStringView GetFunction() const final { return FunctionName; }
-	inline EBuildPolicy GetPolicy() const final { return BuildPolicy; }
-	inline EPriority GetPriority() const final { return Priority; }
 
 	inline ICache& GetCache() const final { return Cache; }
 	inline IBuild& GetBuild() const final { return BuildSystem; }
 
-	void Start(IBuildScheduler& Scheduler, EBuildPolicy Policy, EPriority Priority, FOnBuildJobComplete&& OnComplete) final;
-
 	void Schedule() final;
 
-	// IRequest Interface
-
-	void SetPriority(EPriority Priority) final;
-	void Cancel() final;
-	void Wait() final;
-	bool Poll() final;
-	void AddRef() const final;
-	void Release() const final;
-
 private:
+	FBuildJob(const FBuildJobCreateParams& Params, FStringView Name, FStringView FunctionName, FOnBuildJobComplete&& OnComplete);
+
 	void BeginJob();
 	void EndJob();
 
@@ -167,14 +152,11 @@ private:
 	void SetAction(FBuildAction&& Action);
 	void SetInputs(FBuildInputs&& Inputs);
 	void SetOutput(const FBuildOutput& Output) final;
-	void SetOutputNoCheck(FBuildOutput&& Output);
+	void SetOutputNoCheck(FBuildOutput&& Output, EBuildJobState NewState = EBuildJobState::CacheStore);
 
 	/** Terminate the job and send the error to the output complete callback. */
 	void CompleteWithError(FStringView Error);
 
-	/** Start execution of an async operation that is managed by a request. */
-	void ExecuteAsync(TFunctionRef<FRequest ()> Operation);
-
 	/** Advance to the new state, dispatching to the scheduler and invoking callbacks as appropriate. */
 	void AdvanceToState(EBuildJobState NewState);
 
@@ -202,24 +184,13 @@ private:
 	/** Next state for the job. Used to handle re-entrant calls to AdvanceToState. */
 	EBuildJobState NextState{EBuildJobState::NotStarted};
 
-	EBuildPolicy BuildPolicy{};
-	EPriority Priority{};
-
-	/** True if the build was canceled before it was complete. */
-	bool bIsCanceled{};
-	/** True if the build was retrieved from or stored to the cache. */
-	bool bIsCached{};
-	/** True if AdvanceToState is executing. */
-	bool bInAdvanceToState{};
-	/** True if remote execution was attempted, whether or not it failed. */
-	bool bTriedRemoteExecution{};
-	/** True if the build action and inputs are being exported. */
-	bool bIsExportingBuild{};
-
 	/** Status flags that are added to as the job moves through its states. */
 	EBuildStatus BuildStatus{EBuildStatus::None};
+	/** Policy requested and/or configured for the job. */
+	EBuildPolicy BuildPolicy{EBuildPolicy::None};
 
-	mutable std::atomic<uint32> ReferenceCount{0};
+	/** True if AdvanceToState is executing. */
+	bool bInAdvanceToState{};
 
 	/** Available in [ResolveKey, Complete] for jobs created from a key or definition. */
 	FBuildKey DefinitionKey;
@@ -240,14 +211,12 @@ private:
 
 	FBuildSchedulerParams SchedulerParams;
 
-	/** Lock to synchronize writes to State, NextState, bIsCanceled, Event, WaitRequest, Scheduler. */
+	/** Lock to synchronize writes to State, NextState, bInAdvanceToState. */
 	FRWLock Lock;
-	/** Event used by and created by Wait(). Null until Wait() has been called on an incomplete job. */
-	FEvent* Event{};
-	/** Request that the job is waiting on in its wait states. Null outside of states that use it. */
-	FRequest WaitRequest;
+	/** Owner used to synchronize every async operation executed by the job. */
+	IRequestOwner& Owner;
 	/** Scheduler that the job is scheduled on. Null before Schedule() and after completion. */
-	IBuildScheduler* Scheduler{};
+	IBuildScheduler& Scheduler;
 	/** Resolver for definitions and inputs. May be null for a job with nothing to resolve. */
 	IBuildInputResolver* InputResolver{};
 	/** Worker to use for remote execution. Available in [ExecuteRemote, ExecuteRemoteRetryWait]. */
@@ -297,50 +266,54 @@ const TCHAR* LexToString(EBuildJobState State)
 ///////////////////////////////////////////////////////////////////////////////////////////////////////////////////////
 
 FBuildJob::FBuildJob(
-	ICache& InCache,
-	IBuild& InBuildSystem,
-	IBuildInputResolver* InInputResolver,
-	const FBuildKey& InKey)
-	: Name(WriteToString<64>(TEXT("Resolve: "_SV), InKey))
-	, DefinitionKey(InKey)
-	, OutputBuilder(InBuildSystem.CreateOutput(Name, FunctionName))
-	, InputResolver(InInputResolver)
-	, Cache(InCache)
-	, BuildSystem(InBuildSystem)
+	const FBuildJobCreateParams& Params,
+	FStringView InName,
+	FStringView InFunctionName,
+	FOnBuildJobComplete&& InOnComplete)
+	: Name(InName)
+	, FunctionName(InFunctionName)
+	, BuildPolicy(Params.Policy)
+	, OutputBuilder(Params.BuildSystem.CreateOutput(Name, FunctionName))
+	, Owner(Params.Owner)
+	, Scheduler(Params.Scheduler)
+	, InputResolver(Params.InputResolver)
+	, OnComplete(MoveTemp(InOnComplete))
+	, Cache(Params.Cache)
+	, BuildSystem(Params.BuildSystem)
 {
 }
 
 FBuildJob::FBuildJob(
-	ICache& InCache,
-	IBuild& InBuildSystem,
-	IBuildInputResolver* InInputResolver,
-	const FBuildDefinition& InDefinition)
-	: Name(InDefinition.GetName())
-	, FunctionName(InDefinition.GetFunction())
-	, DefinitionKey(InDefinition.GetKey())
-	, Definition(InDefinition)
-	, OutputBuilder(InBuildSystem.CreateOutput(Name, FunctionName))
-	, InputResolver(InInputResolver)
-	, Cache(InCache)
-	, BuildSystem(InBuildSystem)
+	const FBuildJobCreateParams& Params,
+	const FBuildKey& InKey,
+	FOnBuildJobComplete&& InOnComplete)
+	: FBuildJob(Params, WriteToString<64>(TEXT("Resolve: "_SV), InKey), TEXT("Unknown"_SV), MoveTemp(InOnComplete))
 {
+	DefinitionKey = InKey;
+	AdvanceToState(EBuildJobState::ResolveKey);
 }
 
 FBuildJob::FBuildJob(
-	ICache& InCache,
-	IBuild& InBuildSystem,
-	IBuildInputResolver* InInputResolver,
+	const FBuildJobCreateParams& Params,
+	const FBuildDefinition& InDefinition,
+	FOnBuildJobComplete&& InOnComplete)
+	: FBuildJob(Params, InDefinition.GetName(), InDefinition.GetFunction(), MoveTemp(InOnComplete))
+{
+	DefinitionKey = InDefinition.GetKey();
+	Definition = InDefinition;
+	AdvanceToState(EBuildJobState::ResolveKey);
+}
+
+FBuildJob::FBuildJob(
+	const FBuildJobCreateParams& Params,
 	const FBuildAction& InAction,
-	const FOptionalBuildInputs& InInputs)
-	: Name(InAction.GetName())
-	, FunctionName(InAction.GetFunction())
-	, Action(InAction)
-	, Inputs(InInputs)
-	, OutputBuilder(InBuildSystem.CreateOutput(Name, FunctionName))
-	, InputResolver(InInputResolver)
-	, Cache(InCache)
-	, BuildSystem(InBuildSystem)
+	const FOptionalBuildInputs& InInputs,
+	FOnBuildJobComplete&& InOnComplete)
+	: FBuildJob(Params, InAction.GetName(), InAction.GetFunction(), MoveTemp(InOnComplete))
 {
+	Action = InAction;
+	Inputs = InInputs;
+	AdvanceToState(EBuildJobState::ResolveKey);
 }
 
 FBuildJob::~FBuildJob()
@@ -351,23 +324,6 @@ FBuildJob::~FBuildJob()
 	checkf(!OnComplete,
 		TEXT("Job in state %s must invoke its completion callback before being destroyed for build of '%s' by %s."),
 		LexToString(State), *Name, *FunctionName);
-	FPlatformProcess::ReturnSynchEventToPool(Event);
-}
-
-void FBuildJob::Start(
-	IBuildScheduler& InScheduler,
-	EBuildPolicy InPolicy,
-	EPriority InPriority,
-	FOnBuildJobComplete&& InOnComplete)
-{
-	checkf(State == EBuildJobState::NotStarted,
-		TEXT("Job in state %s was previously scheduled for build of '%s' by %s."),
-		LexToString(State), *Name, *FunctionName);
-	Scheduler = &InScheduler;
-	BuildPolicy = InPolicy;
-	Priority = InPriority;
-	OnComplete = MoveTemp(InOnComplete);
-	return AdvanceToState(EBuildJobState::ResolveKey);
 }
 
 void FBuildJob::Schedule()
@@ -392,140 +348,15 @@ void FBuildJob::Schedule()
 
 ///////////////////////////////////////////////////////////////////////////////////////////////////////////////////////
 
-void FBuildJob::SetPriority(EPriority InPriority)
-{
-	IBuildScheduler* LocalScheduler;
-	{
-		FWriteScopeLock WriteLock(Lock);
-		Priority = InPriority;
-		LocalScheduler = Scheduler;
-		WaitRequest.SetPriority(InPriority);
-	}
-	if (LocalScheduler)
-	{
-		LocalScheduler->UpdateJobPriority(this);
-	}
-}
-
-void FBuildJob::Cancel()
-{
-	FRequest LocalRequest;
-	IBuildScheduler* LocalScheduler;
-	TRefCountPtr<FBuildJobContext> LocalContext;
-	if (FReadScopeLock ReadLock(Lock); State == EBuildJobState::Complete)
-	{
-		return;
-	}
-	else
-	{
-		bIsCanceled = true;
-		LocalRequest = WaitRequest;
-		LocalScheduler = Scheduler;
-		if (State == EBuildJobState::ExecuteLocalWait)
-		{
-			LocalContext = Context;
-		}
-	}
-
-	// Cancel the job on the scheduler, which invokes Schedule if the job was queued.
-	if (LocalScheduler)
-	{
-		LocalScheduler->CancelJob(this);
-	}
-
-	// Cancel the request, which invokes End[State] if the request is not complete.
-	if (LocalRequest)
-	{
-		LocalRequest.Cancel();
-	}
-
-	// Cancel the async build, which invokes EndExecuteLocal if the build is not complete.
-	if (LocalContext)
-	{
-		LocalContext->GetFunction().CancelAsyncBuild(*LocalContext);
-	}
-
-	// Most jobs will be complete at this point, but it is possible for a job to reach this point
-	// because of a race with the scheduler. A job can be removed from every scheduler queue, and
-	// missed by CancelJob, but not yet have called back into the job.
-	Wait();
-}
-
-void FBuildJob::Wait()
-{
-	FEvent* LocalEvent;
-
-	if (FReadScopeLock ReadLock(Lock); State == EBuildJobState::Complete)
-	{
-		return;
-	}
-	else
-	{
-		LocalEvent = Event;
-	}
-
-	SetPriority(EPriority::Blocking);
-
-	if (!LocalEvent)
-	{
-		if (FWriteScopeLock WriteLock(Lock); State == EBuildJobState::Complete)
-		{
-			return;
-		}
-		else if (Event)
-		{
-			LocalEvent = Event;
-		}
-		else
-		{
-			LocalEvent = Event = FPlatformProcess::GetSynchEventFromPool(/*bIsManualReset*/ true);
-		}
-	}
-
-	TRACE_CPUPROFILER_EVENT_SCOPE(FBuildJob::Wait);
-	LocalEvent->Wait();
-}
-
-bool FBuildJob::Poll()
-{
-	FReadScopeLock ReadLock(Lock);
-	return State == EBuildJobState::Complete;
-}
-
-void FBuildJob::AddRef() const
-{
-	ReferenceCount.fetch_add(1, std::memory_order_relaxed);
-}
-
-void FBuildJob::Release() const
-{
-	if (ReferenceCount.fetch_sub(1, std::memory_order_acq_rel) == 1)
-	{
-		delete this;
-	}
-}
-
-///////////////////////////////////////////////////////////////////////////////////////////////////////////////////////
-
 void FBuildJob::BeginJob()
 {
-	AddRef();
-	Scheduler->BeginJob(this);
+	Scheduler.BeginJob(this);
 }
 
 void FBuildJob::EndJob()
 {
-	ON_SCOPE_EXIT { Release(); };
-
-	if (Event)
-	{
-		Event->Trigger();
-	}
-	Scheduler->EndJob(this);
-
-	FWriteScopeLock WriteLock(Lock);
-	Scheduler = nullptr;
-	InputResolver = nullptr;
+	Scheduler.EndJob(this);
+	delete this;
 }
 
 ///////////////////////////////////////////////////////////////////////////////////////////////////////////////////////
@@ -560,11 +391,10 @@ void FBuildJob::CreateContext()
 	else
 	{
 		const FCacheKey CacheKey{Cache.CreateBucket(FunctionName), Action.Get().GetKey().Hash};
-		Context = new FBuildJobContext(*this, CacheKey, *Function, OutputBuilder, BuildPolicy,
-			[this] { EndExecuteLocal(); });
+		Context = new FBuildJobContext(*this, CacheKey, *Function, OutputBuilder, BuildPolicy);
 		Function->Configure(*Context);
 		BuildPolicy = Context->GetBuildPolicy();
-		bIsExportingBuild = ShouldExportBuild();
+		EnumAddFlags(BuildStatus, ShouldExportBuild() ? EBuildStatus::BuildTryExport : EBuildStatus::None);
 	}
 
 	// Populate the scheduler params with the information that is available now.
@@ -603,27 +433,23 @@ void FBuildJob::EnterCacheQuery()
 void FBuildJob::BeginCacheQuery()
 {
 	TRACE_CPUPROFILER_EVENT_SCOPE(FBuildJob::CacheQuery);
-	return ExecuteAsync([this]
+	ECachePolicy CachePolicy = Context->GetCachePolicy();
+	if (EnumHasAnyFlags(BuildPolicy, EBuildPolicy::SkipData))
 	{
-		ECachePolicy CachePolicy = Context->GetCachePolicy();
-		if (EnumHasAnyFlags(BuildPolicy, EBuildPolicy::SkipData))
-		{
-			CachePolicy |= ECachePolicy::SkipAttachments;
-		}
-		return Cache.Get({Context->GetCacheKey()}, Name, CachePolicy, Priority,
-			[this](FCacheGetCompleteParams&& Params) { EndCacheQuery(MoveTemp(Params)); });
-	});
+		EnumAddFlags(CachePolicy, ECachePolicy::SkipAttachments);
+	}
+	EnumAddFlags(BuildStatus, EBuildStatus::CacheQuery);
+	Cache.Get({Context->GetCacheKey()}, Name, CachePolicy, Owner,
+		[this](FCacheGetCompleteParams&& Params) { EndCacheQuery(MoveTemp(Params)); });
 }
 
 void FBuildJob::EndCacheQuery(FCacheGetCompleteParams&& Params)
 {
-	BuildStatus |= EBuildStatus::CacheQuery;
 	if (Params.Status == EStatus::Ok)
 	{
-		BuildStatus |= EBuildStatus::CacheQueryHit;
 		if (FOptionalBuildOutput CacheOutput = BuildSystem.LoadOutput(Name, FunctionName, Params.Record))
 		{
-			bIsCached = true;
+			EnumAddFlags(BuildStatus, EBuildStatus::CacheQueryHit);
 			return SetOutputNoCheck(MoveTemp(CacheOutput).Get());
 		}
 	}
@@ -635,9 +461,9 @@ void FBuildJob::EndCacheQuery(FCacheGetCompleteParams&& Params)
 void FBuildJob::EnterCacheStore()
 {
 	ECachePolicy CachePolicy = Context ? Context->GetCachePolicy() : ECachePolicy::None;
-	if (bIsCached ||
-		!EnumHasAnyFlags(CachePolicy, ECachePolicy::Store) ||
+	if (!EnumHasAnyFlags(CachePolicy, ECachePolicy::Store) ||
 		!EnumHasAnyFlags(BuildPolicy, EBuildPolicy::CacheStore) ||
+		EnumHasAnyFlags(BuildStatus, EBuildStatus::CacheQueryHit) ||
 		Output.Get().HasError())
 	{
 		return AdvanceToState(EBuildJobState::Complete);
@@ -647,19 +473,19 @@ void FBuildJob::EnterCacheStore()
 void FBuildJob::BeginCacheStore()
 {
 	TRACE_CPUPROFILER_EVENT_SCOPE(FBuildJob::CacheStore);
-	return ExecuteAsync([this]
-	{
-		FCacheRecordBuilder RecordBuilder = Cache.CreateRecord(Context->GetCacheKey());
-		Output.Get().Save(RecordBuilder);
-		return Cache.Put({RecordBuilder.Build()}, Name, Context->GetCachePolicy(),
-			FMath::Min(Priority, EPriority::Highest),
-			[this](FCachePutCompleteParams&& Params) { EndCacheStore(MoveTemp(Params)); });
-	});
+	FCacheRecordBuilder RecordBuilder = Cache.CreateRecord(Context->GetCacheKey());
+	Output.Get().Save(RecordBuilder);
+	EnumAddFlags(BuildStatus, EBuildStatus::CacheStore);
+	Cache.Put({RecordBuilder.Build()}, Name, Context->GetCachePolicy(), Owner,
+		[this](FCachePutCompleteParams&& Params) { EndCacheStore(MoveTemp(Params)); });
 }
 
 void FBuildJob::EndCacheStore(FCachePutCompleteParams&& Params)
 {
-	bIsCached = Params.Status == EStatus::Ok;
+	if (Params.Status == EStatus::Ok)
+	{
+		EnumAddFlags(BuildStatus, EBuildStatus::CacheStoreHit);
+	}
 	return AdvanceToState(EBuildJobState::Complete);
 }
 
@@ -688,11 +514,8 @@ void FBuildJob::EnterResolveKey()
 void FBuildJob::BeginResolveKey()
 {
 	TRACE_CPUPROFILER_EVENT_SCOPE(FBuildJob::ResolveKey);
-	return ExecuteAsync([this]
-	{
-		return InputResolver->ResolveKey(DefinitionKey,
-			[this](FBuildKeyResolvedParams&& Params) { EndResolveKey(MoveTemp(Params)); });
-	});
+	InputResolver->ResolveKey(DefinitionKey, Owner,
+		[this](FBuildKeyResolvedParams&& Params) { EndResolveKey(MoveTemp(Params)); });
 }
 
 void FBuildJob::EndResolveKey(FBuildKeyResolvedParams&& Params)
@@ -724,11 +547,8 @@ void FBuildJob::EnterResolveInputMeta()
 void FBuildJob::BeginResolveInputMeta()
 {
 	TRACE_CPUPROFILER_EVENT_SCOPE(FBuildJob::ResolveInputMeta);
-	return ExecuteAsync([this]
-	{
-		return InputResolver->ResolveInputMeta(Definition.Get(), Priority,
-			[this](FBuildInputMetaResolvedParams&& Params) { EndResolveInputMeta(MoveTemp(Params)); });
-	});
+	InputResolver->ResolveInputMeta(Definition.Get(), Owner,
+		[this](FBuildInputMetaResolvedParams&& Params) { EndResolveInputMeta(MoveTemp(Params)); });
 }
 
 void FBuildJob::EndResolveInputMeta(FBuildInputMetaResolvedParams&& Params)
@@ -789,21 +609,18 @@ void FBuildJob::EnterResolveInputData()
 void FBuildJob::BeginResolveInputData()
 {
 	TRACE_CPUPROFILER_EVENT_SCOPE(FBuildJob::ResolveInputData);
-	return ExecuteAsync([this]
+	if (Definition)
 	{
-		if (Definition)
-		{
-			return InputResolver->ResolveInputData(Definition.Get(), Priority,
-				[this](FBuildInputDataResolvedParams&& Params) { EndResolveInputData(MoveTemp(Params)); },
-				[this](FStringView Key) { return !!Algo::Find(MissingInputs, Key); });
-		}
-		else
-		{
-			return InputResolver->ResolveInputData(Action.Get(), Priority,
-				[this](FBuildInputDataResolvedParams&& Params) { EndResolveInputData(MoveTemp(Params)); },
-				[this](FStringView Key) { return !!Algo::Find(MissingInputs, Key); });
-		}
-	});
+		InputResolver->ResolveInputData(Definition.Get(), Owner,
+			[this](FBuildInputDataResolvedParams&& Params) { EndResolveInputData(MoveTemp(Params)); },
+			[this](FStringView Key) { return !!Algo::Find(MissingInputs, Key); });
+	}
+	else
+	{
+		InputResolver->ResolveInputData(Action.Get(), Owner,
+			[this](FBuildInputDataResolvedParams&& Params) { EndResolveInputData(MoveTemp(Params)); },
+			[this](FStringView Key) { return !!Algo::Find(MissingInputs, Key); });
+	}
 }
 
 void FBuildJob::EndResolveInputData(FBuildInputDataResolvedParams&& Params)
@@ -835,7 +652,8 @@ void FBuildJob::EndResolveInputData(FBuildInputDataResolvedParams&& Params)
 
 void FBuildJob::EnterExecuteRemote()
 {
-	if (!EnumHasAnyFlags(BuildPolicy, EBuildPolicy::BuildRemote) || bIsExportingBuild)
+	if (!EnumHasAnyFlags(BuildPolicy, EBuildPolicy::BuildRemote) ||
+		EnumHasAnyFlags(BuildStatus, EBuildStatus::BuildTryExport))
 	{
 		return AdvanceToState(EBuildJobState::ResolveInputData);
 	}
@@ -860,19 +678,16 @@ void FBuildJob::BeginExecuteRemote()
 	TRACE_CPUPROFILER_EVENT_SCOPE(FBuildJob::ExecuteRemote);
 	checkf(Worker && WorkerExecutor, TEXT("Job requires a worker in state %s for build of '%s' by %s."),
 		LexToString(State), *Name, *FunctionName);
-	bTriedRemoteExecution = true;
-	return ExecuteAsync([this]
-	{
-		return WorkerExecutor->BuildAction(Action.Get(), Inputs, *Worker, BuildPolicy, Priority,
-			[this](FBuildWorkerActionCompleteParams&& Params) { EndExecuteRemote(MoveTemp(Params)); });
-	});
+	EnumAddFlags(BuildStatus, EBuildStatus::BuildTryRemote);
+	WorkerExecutor->BuildAction(Action.Get(), Inputs, *Worker, BuildSystem, BuildPolicy, Owner,
+		[this](FBuildWorkerActionCompleteParams&& Params) { EndExecuteRemote(MoveTemp(Params)); });
 }
 
 void FBuildJob::EndExecuteRemote(FBuildWorkerActionCompleteParams&& Params)
 {
 	if (Params.Output)
 	{
-		BuildStatus |= EBuildStatus::BuildRemote;
+		EnumAddFlags(BuildStatus, EBuildStatus::BuildRemote);
 		return SetOutputNoCheck(MoveTemp(Params.Output).Get());
 	}
 	else
@@ -913,9 +728,10 @@ void FBuildJob::SkipExecuteRemote()
 
 void FBuildJob::EnterExecuteLocal()
 {
-	if (bIsExportingBuild)
+	if (EnumHasAnyFlags(BuildStatus, EBuildStatus::BuildTryExport))
 	{
 		ExportBuild();
+		EnumAddFlags(BuildStatus, EBuildStatus::BuildExport);
 	}
 	if (!EnumHasAnyFlags(BuildPolicy, EBuildPolicy::BuildLocal))
 	{
@@ -923,7 +739,7 @@ void FBuildJob::EnterExecuteLocal()
 		{
 			return CompleteWithError(TEXT("Failed because build policy does not allow local or remote execution."_SV));
 		}
-		else if (bTriedRemoteExecution)
+		else if (EnumHasAnyFlags(BuildStatus, EBuildStatus::BuildTryRemote))
 		{
 			return CompleteWithError(TEXT("Failed because build policy does not allow local execution, ")
 				TEXT("and remote execution failed to build."_SV));
@@ -960,20 +776,12 @@ void FBuildJob::EnterExecuteLocal()
 void FBuildJob::BeginExecuteLocal()
 {
 	TRACE_CPUPROFILER_EVENT_SCOPE(FBuildJob::ExecuteLocal);
-	return ExecuteAsync([this]
-	{
-		Context->GetFunction().Build(*Context);
-		if (!Context->IsAsyncBuild())
-		{
-			EndExecuteLocal();
-		}
-		return FRequest();
-	});
+	Context->BeginBuild(Owner, [this] { EndExecuteLocal(); });
 }
 
 void FBuildJob::EndExecuteLocal()
 {
-	BuildStatus |= EBuildStatus::BuildLocal;
+	EnumAddFlags(BuildStatus, EBuildStatus::BuildLocal);
 	return SetOutputNoCheck(OutputBuilder.Build());
 }
 
@@ -1055,25 +863,25 @@ void FBuildJob::SetOutput(const FBuildOutput& InOutput)
 	return SetOutputNoCheck(FBuildOutput(InOutput));
 }
 
-void FBuildJob::SetOutputNoCheck(FBuildOutput&& InOutput)
+void FBuildJob::SetOutputNoCheck(FBuildOutput&& InOutput, EBuildJobState NewState)
 {
 	checkf(Output.IsNull(), TEXT("Job already has an output for build of '%s' by %s."), *Name, *FunctionName);
 	Output = MoveTemp(InOutput);
 
-	if (!bIsCanceled)
+	if (!Owner.IsCanceled())
 	{
-		Scheduler->SetJobOutput(this, SchedulerParams, Output.Get());
+		Scheduler.SetJobOutput(this, SchedulerParams, Output.Get());
 	}
 
 	if (OnComplete)
 	{
 		const FCacheKey& CacheKey = Context ? Context->GetCacheKey() : FCacheKey::Empty;
-		const EStatus Status = bIsCanceled ? EStatus::Canceled : Output.Get().HasError() ? EStatus::Error : EStatus::Ok;
+		const EStatus Status = Owner.IsCanceled() ? EStatus::Canceled : Output.Get().HasError() ? EStatus::Error : EStatus::Ok;
 		OnComplete({*this, CacheKey, FBuildOutput(Output.Get()), BuildStatus, Status});
 		OnComplete = nullptr;
 	}
 
-	return AdvanceToState(EBuildJobState::CacheStore);
+	return AdvanceToState(NewState);
 }
 
 ///////////////////////////////////////////////////////////////////////////////////////////////////////////////////////
@@ -1085,30 +893,7 @@ void FBuildJob::CompleteWithError(FStringView Error)
 		return;
 	}
 	OutputBuilder.AddError(TEXT("LogDerivedDataBuild"_SV), Error);
-	return SetOutputNoCheck(OutputBuilder.Build());
-}
-
-void FBuildJob::ExecuteAsync(TFunctionRef<FRequest ()> Operation)
-{
-	// An operation may complete and advance to another state before its request has been stored.
-	// Hold a reference to allow access to this after starting the operation.
-	const TRequest Self(this);
-	const EBuildJobState LocalState = State;
-	if (FRequest LocalRequest = Operation())
-	{
-		if (FWriteScopeLock WriteLock(Lock); LocalState == State && !bIsCanceled)
-		{
-			// The state flow through the job requires any previous WaitRequest to be complete by
-			// the time that ExecuteAsync is called because the completion of the async operation
-			// will request the state transition that leads to another new async operation. It is
-			// not possible to poll the previous request to validate that it is complete, because
-			// execution likely reached this point from the completion callback, and a request is
-			// not complete until its completion callback has returned.
-			WaitRequest = MoveTemp(LocalRequest);
-		}
-		// Cancel the request if it was not moved above due to cancellation or change of state.
-		LocalRequest.Cancel();
-	}
+	return SetOutputNoCheck(OutputBuilder.Build(), EBuildJobState::Complete);
 }
 
 void FBuildJob::AdvanceToState(EBuildJobState NewState)
@@ -1122,7 +907,7 @@ void FBuildJob::AdvanceToState(EBuildJobState NewState)
 		//	TEXT("Job in state %s is requesting an invalid transition from %s to %s for build of '%s' by %s."),
 		//	LexToString(State), LexToString(NextState), LexToString(NewState), *Name, *FunctionName);
 
-		if (bIsCanceled)
+		if (Owner.IsCanceled())
 		{
 			NewState = EBuildJobState::Complete;
 		}
@@ -1234,15 +1019,15 @@ void FBuildJob::ExecuteState(EBuildJobState NewState)
 {
 	switch (NewState)
 	{
-	case EBuildJobState::ResolveKey:                 return Scheduler->DispatchResolveKey(this);
-	case EBuildJobState::ResolveInputMeta:           return Scheduler->DispatchResolveInputMeta(this);
-	case EBuildJobState::CacheQuery:                 return Scheduler->DispatchCacheQuery(this, SchedulerParams);
-	case EBuildJobState::ExecuteRemote:              return Scheduler->DispatchExecuteRemote(this, SchedulerParams);
-	case EBuildJobState::ResolveRemoteInputData:     return Scheduler->DispatchResolveInputData(this, SchedulerParams);
-	case EBuildJobState::ExecuteRemoteRetry:         return Scheduler->DispatchExecuteRemote(this, SchedulerParams);
-	case EBuildJobState::ResolveInputData:           return Scheduler->DispatchResolveInputData(this, SchedulerParams);
-	case EBuildJobState::ExecuteLocal:               return Scheduler->DispatchExecuteLocal(this, SchedulerParams);
-	case EBuildJobState::CacheStore:                 return Scheduler->DispatchCacheStore(this, SchedulerParams);
+	case EBuildJobState::ResolveKey:                 return Scheduler.DispatchResolveKey(this, Owner);
+	case EBuildJobState::ResolveInputMeta:           return Scheduler.DispatchResolveInputMeta(this, Owner);
+	case EBuildJobState::CacheQuery:                 return Scheduler.DispatchCacheQuery(this, Owner, SchedulerParams);
+	case EBuildJobState::ExecuteRemote:              return Scheduler.DispatchExecuteRemote(this, Owner, SchedulerParams);
+	case EBuildJobState::ResolveRemoteInputData:     return Scheduler.DispatchResolveInputData(this, Owner, SchedulerParams);
+	case EBuildJobState::ExecuteRemoteRetry:         return Scheduler.DispatchExecuteRemote(this, Owner, SchedulerParams);
+	case EBuildJobState::ResolveInputData:           return Scheduler.DispatchResolveInputData(this, Owner, SchedulerParams);
+	case EBuildJobState::ExecuteLocal:               return Scheduler.DispatchExecuteLocal(this, Owner, SchedulerParams);
+	case EBuildJobState::CacheStore:                 return Scheduler.DispatchCacheStore(this, Owner, SchedulerParams);
 	}
 }
 
@@ -1346,19 +1131,19 @@ TArray<FName> FBuildJob::ParseExportBuildTypes(bool& bOutExportAll)
 
 ///////////////////////////////////////////////////////////////////////////////////////////////////////////////////////
 
-TRequest<IBuildJob> CreateBuildJob(ICache& Cache, IBuild& BuildSystem, IBuildInputResolver* InputResolver, const FBuildKey& Key)
+void CreateBuildJob(const FBuildJobCreateParams& Params, const FBuildKey& Key, FOnBuildJobComplete&& OnComplete)
 {
-	return TRequest(new FBuildJob(Cache, BuildSystem, InputResolver, Key));
+	new FBuildJob(Params, Key, MoveTemp(OnComplete));
 }
 
-TRequest<IBuildJob> CreateBuildJob(ICache& Cache, IBuild& BuildSystem, IBuildInputResolver* InputResolver, const FBuildDefinition& Definition)
+void CreateBuildJob(const FBuildJobCreateParams& Params, const FBuildDefinition& Definition, FOnBuildJobComplete&& OnComplete)
 {
-	return TRequest(new FBuildJob(Cache, BuildSystem, InputResolver, Definition));
+	new FBuildJob(Params, Definition, MoveTemp(OnComplete));
 }
 
-TRequest<IBuildJob> CreateBuildJob(ICache& Cache, IBuild& BuildSystem, IBuildInputResolver* InputResolver, const FBuildAction& Action, const FOptionalBuildInputs& Inputs)
+void CreateBuildJob(const FBuildJobCreateParams& Params, const FBuildAction& Action, const FOptionalBuildInputs& Inputs, FOnBuildJobComplete&& OnComplete)
 {
-	return TRequest(new FBuildJob(Cache, BuildSystem, InputResolver, Action, Inputs));
+	new FBuildJob(Params, Action, Inputs, MoveTemp(OnComplete));
 }
 
 } // UE::DerivedData::Private
diff --git a/Engine/Source/Developer/DerivedDataCache/Private/DerivedDataBuildJobContext.cpp b/Engine/Source/Developer/DerivedDataCache/Private/DerivedDataBuildJobContext.cpp
--- a/Engine/Source/Developer/DerivedDataCache/Private/DerivedDataBuildJobContext.cpp
+++ b/Engine/Source/Developer/DerivedDataCache/Private/DerivedDataBuildJobContext.cpp
@@ -20,13 +20,11 @@ FBuildJobContext::FBuildJobContext(
 	const FCacheKey& InCacheKey,
 	const IBuildFunction& InFunction,
 	FBuildOutputBuilder& InOutputBuilder,
-	EBuildPolicy InBuildPolicy,
-	TUniqueFunction<void ()>&& InOnEndAsyncBuild)
+	EBuildPolicy InBuildPolicy)
 	: Job(InJob)
 	, CacheKey(InCacheKey)
 	, Function(InFunction)
 	, OutputBuilder(InOutputBuilder)
-	, OnEndAsyncBuild(MoveTemp(InOnEndAsyncBuild))
 	, CachePolicy(ECachePolicy::Default)
 	, BuildPolicy(InBuildPolicy)
 {
@@ -112,11 +110,30 @@ void FBuildJobContext::AddPayload(const FPayloadId& Id, const FCbObject& Object)
 	AddPayload(FPayload(Id, FCompressedBuffer::Compress(Buffer, GDefaultCompressor, GDefaultCompressionLevel)));
 }
 
+void FBuildJobContext::BeginBuild(IRequestOwner& InOwner, TUniqueFunction<void ()>&& InOnEndBuild)
+{
+	Owner = &InOwner;
+	OnEndBuild = MoveTemp(InOnEndBuild);
+	Function.Build(*this);
+	if (!bIsAsyncBuild)
+	{
+		EndBuild();
+	}
+}
+
+void FBuildJobContext::EndBuild()
+{
+	OnEndBuild();
+	BuildCompleteEvent->Trigger();
+	Owner = nullptr;
+}
+
 void FBuildJobContext::BeginAsyncBuild()
 {
 	checkf(!bIsAsyncBuild, TEXT("BeginAsyncBuild may only be called once for build of '%s' by %s."),
 		*WriteToString<128>(Job.GetName()), *WriteToString<32>(Job.GetFunction()));
 	bIsAsyncBuild = true;
+	Owner->Begin(this);
 }
 
 void FBuildJobContext::EndAsyncBuild()
@@ -126,7 +143,7 @@ void FBuildJobContext::EndAsyncBuild()
 	checkf(!bIsAsyncBuildComplete, TEXT("EndAsyncBuild may only be called once for build of '%s' by %s."),
 		*WriteToString<128>(Job.GetName()), *WriteToString<32>(Job.GetFunction()));
 	bIsAsyncBuildComplete = true;
-	OnEndAsyncBuild();
+	Owner->End(this, [this] { EndBuild(); });
 }
 
 FCacheBucket FBuildJobContext::CreateCacheBucket(FStringView Name) const
@@ -164,4 +181,20 @@ void FBuildJobContext::SetBuildPolicy(EBuildPolicy Policy)
 	BuildPolicy = Policy;
 }
 
+void FBuildJobContext::SetPriority(EPriority Priority)
+{
+}
+
+void FBuildJobContext::Cancel()
+{
+	checkf(bIsAsyncBuild, TEXT("Cancel may only be called after BeginAsyncBuild for build of '%s' by %s."),
+		*WriteToString<128>(Job.GetName()), *WriteToString<32>(Job.GetFunction()));
+	Function.CancelAsyncBuild(*this);
+}
+
+void FBuildJobContext::Wait()
+{
+	BuildCompleteEvent->Wait();
+}
+
 } // UE::DerivedData::Private
diff --git a/Engine/Source/Developer/DerivedDataCache/Private/DerivedDataBuildJobContext.h b/Engine/Source/Developer/DerivedDataCache/Private/DerivedDataBuildJobContext.h
--- a/Engine/Source/Developer/DerivedDataCache/Private/DerivedDataBuildJobContext.h
+++ b/Engine/Source/Developer/DerivedDataCache/Private/DerivedDataBuildJobContext.h
@@ -8,10 +8,11 @@
 #include "Containers/UnrealString.h"
 #include "DerivedDataBuildFunction.h"
 #include "DerivedDataCacheKey.h"
+#include "DerivedDataRequest.h"
+#include "HAL/Event.h"
 #include "Memory/MemoryFwd.h"
 #include "Serialization/CompactBinary.h"
 #include "Templates/Function.h"
-#include <atomic>
 
 namespace UE::DerivedData { class FBuildOutputBuilder; }
 namespace UE::DerivedData { class IBuildJob; }
@@ -21,26 +22,25 @@ namespace UE::DerivedData { enum class ECachePolicy : uint8; }
 namespace UE::DerivedData::Private
 {
 
-class FBuildJobContext final : public FBuildContext, public FBuildConfigContext
+class FBuildJobContext final : public FBuildContext, public FBuildConfigContext, public FRequestBase
 {
 public:
 	FBuildJobContext(
-		IBuildJob& InJob,
-		const FCacheKey& InCacheKey,
-		const IBuildFunction& InFunction,
-		FBuildOutputBuilder& InOutputBuilder,
-		EBuildPolicy InBuildPolicy,
-		TUniqueFunction<void ()>&& InOnEndAsyncBuild);
+		IBuildJob& Job,
+		const FCacheKey& CacheKey,
+		const IBuildFunction& Function,
+		FBuildOutputBuilder& OutputBuilder,
+		EBuildPolicy BuildPolicy);
+
+	void BeginBuild(IRequestOwner& Owner, TUniqueFunction<void ()>&& OnEndBuild);
 
 	FStringView GetName() const final;
 
 	inline const FCacheKey& GetCacheKey() const { return CacheKey; }
-	inline const IBuildFunction& GetFunction() const { return Function; }
 
 	inline ECachePolicy GetCachePolicy() const final { return CachePolicy; }
 	inline EBuildPolicy GetBuildPolicy() const final { return BuildPolicy; }
 
-	inline bool IsAsyncBuild() const { return bIsAsyncBuild; }
 	inline bool ShouldCheckDeterministicOutput() const { return bDeterministicOutputCheck; }
 
 	void AddConstant(FStringView Key, FCbObject&& Value);
@@ -58,6 +58,7 @@ private:
 	void AddPayload(const FPayloadId& Id, const FSharedBuffer& Buffer) final;
 	void AddPayload(const FPayloadId& Id, const FCbObject& Object) final;
 
+	void EndBuild();
 	void BeginAsyncBuild() final;
 	void EndAsyncBuild() final;
 
@@ -67,19 +68,9 @@ private:
 	void SetBuildPolicy(EBuildPolicy Policy) final;
 	void SkipDeterministicOutputCheck() { bDeterministicOutputCheck = false; }
 
-public:
-	void AddRef() const
-	{
-		ReferenceCount.fetch_add(1, std::memory_order_relaxed);
-	}
-
-	void Release() const
-	{
-		if (ReferenceCount.fetch_sub(1, std::memory_order_acq_rel) == 1)
-		{
-			delete this;
-		}
-	}
+	void SetPriority(EPriority Priority) final;
+	void Cancel() final;
+	void Wait() final;
 
 private:
 	IBuildJob& Job;
@@ -88,13 +79,14 @@ private:
 	FBuildOutputBuilder& OutputBuilder;
 	TMap<FString, FCbObject> Constants;
 	TMap<FString, FCompressedBuffer> Inputs;
-	TUniqueFunction<void ()> OnEndAsyncBuild;
+	FEventRef BuildCompleteEvent{EEventMode::ManualReset};
+	TUniqueFunction<void ()> OnEndBuild;
+	IRequestOwner* Owner{};
 	ECachePolicy CachePolicy;
 	EBuildPolicy BuildPolicy;
 	bool bIsAsyncBuild{false};
 	bool bIsAsyncBuildComplete{false};
 	bool bDeterministicOutputCheck{true};
-	mutable std::atomic<uint32> ReferenceCount{0};
 };
 
 } // UE::DerivedData::Private
diff --git a/Engine/Source/Developer/DerivedDataCache/Private/DerivedDataBuildPrivate.h b/Engine/Source/Developer/DerivedDataCache/Private/DerivedDataBuildPrivate.h
--- a/Engine/Source/Developer/DerivedDataCache/Private/DerivedDataBuildPrivate.h
+++ b/Engine/Source/Developer/DerivedDataCache/Private/DerivedDataBuildPrivate.h
@@ -9,6 +9,8 @@
 class FCbObject;
 struct FGuid;
 
+template <typename FuncType> class TUniqueFunction;
+
 namespace UE::DerivedData { class FBuildAction; }
 namespace UE::DerivedData { class FBuildActionBuilder; }
 namespace UE::DerivedData { class FBuildDefinition; }
@@ -24,12 +26,14 @@ namespace UE::DerivedData { class FOptionalBuildOutput; }
 namespace UE::DerivedData { class IBuild; }
 namespace UE::DerivedData { class IBuildFunctionRegistry; }
 namespace UE::DerivedData { class IBuildInputResolver; }
-namespace UE::DerivedData { class IBuildJob; }
 namespace UE::DerivedData { class IBuildScheduler; }
 namespace UE::DerivedData { class IBuildWorkerRegistry; }
 namespace UE::DerivedData { class ICache; }
+namespace UE::DerivedData { class IRequestOwner; }
+namespace UE::DerivedData { struct FBuildJobCompleteParams; }
 namespace UE::DerivedData { struct FBuildKey; }
-namespace UE::DerivedData { template <typename RequestType> class TRequest; }
+namespace UE::DerivedData { enum class EBuildPolicy : uint8; }
+namespace UE::DerivedData { using FOnBuildJobComplete = TUniqueFunction<void (FBuildJobCompleteParams&& Params)>; }
 
 namespace UE::DerivedData::Private
 {
@@ -73,9 +77,18 @@ FBuildSession CreateBuildSession(
 	IBuildInputResolver* InputResolver);
 
 // Implemented in DerivedDataBuildJob.cpp
-TRequest<IBuildJob> CreateBuildJob(ICache& Cache, IBuild& BuildSystem, IBuildInputResolver* InputResolver, const FBuildKey& Key);
-TRequest<IBuildJob> CreateBuildJob(ICache& Cache, IBuild& BuildSystem, IBuildInputResolver* InputResolver, const FBuildDefinition& Definition);
-TRequest<IBuildJob> CreateBuildJob(ICache& Cache, IBuild& BuildSystem, IBuildInputResolver* InputResolver, const FBuildAction& Action, const FOptionalBuildInputs& Inputs);
+struct FBuildJobCreateParams
+{
+	ICache& Cache;
+	IBuild& BuildSystem;
+	IBuildScheduler& Scheduler;
+	IBuildInputResolver* InputResolver{};
+	IRequestOwner& Owner;
+	EBuildPolicy Policy{};
+};
+void CreateBuildJob(const FBuildJobCreateParams& Params, const FBuildKey& Key, FOnBuildJobComplete&& OnComplete);
+void CreateBuildJob(const FBuildJobCreateParams& Params, const FBuildDefinition& Definition, FOnBuildJobComplete&& OnComplete);
+void CreateBuildJob(const FBuildJobCreateParams& Params, const FBuildAction& Action, const FOptionalBuildInputs& Inputs, FOnBuildJobComplete&& OnComplete);
 
 ///////////////////////////////////////////////////////////////////////////////////////////////////////////////////////
 
diff --git a/Engine/Source/Developer/DerivedDataCache/Private/DerivedDataBuildScheduler.cpp b/Engine/Source/Developer/DerivedDataCache/Private/DerivedDataBuildScheduler.cpp
--- a/Engine/Source/Developer/DerivedDataCache/Private/DerivedDataBuildScheduler.cpp
+++ b/Engine/Source/Developer/DerivedDataCache/Private/DerivedDataBuildScheduler.cpp
@@ -11,95 +11,103 @@
 namespace UE::DerivedData::Private
 {
 
+class FBuildSchedulerRequest : public FRequestBase
+{
+public:
+	FBuildSchedulerRequest(IBuildJob* InJob, IRequestOwner& InOwner, const TCHAR* DebugName)
+		: Job(InJob)
+		, Owner(InOwner)
+	{
+		Tasks::FTaskEvent TaskEvent(TEXT("FBuildSchedulerRequest"));
+		Task = Tasks::Launch(DebugName, [this] { Schedule(); }, TaskEvent, Tasks::ETaskPriority::BackgroundNormal);
+		Owner.Begin(this);
+		TaskEvent.Trigger();
+	}
+
+	void Schedule()
+	{
+		Owner.End(this, [this] { Job->Schedule(); });
+	}
+
+	void SetPriority(EPriority Priority) final
+	{
+	}
+
+	void Cancel() final
+	{
+		Task.Wait();
+	}
+
+	void Wait() final
+	{
+		Task.Wait();
+	}
+
+private:
+	IBuildJob* Job;
+	IRequestOwner& Owner;
+	Tasks::FTask Task;
+};
+
 class FBuildScheduler final : public IBuildScheduler
 {
 public:
-	void DispatchCacheQuery(IBuildJob* Job, const FBuildSchedulerParams& Params) final;
-	void DispatchCacheStore(IBuildJob* Job, const FBuildSchedulerParams& Params) final;
-	void DispatchResolveKey(IBuildJob* Job) final;
-	void DispatchResolveInputMeta(IBuildJob* Job) final;
-	void DispatchResolveInputData(IBuildJob* Job, const FBuildSchedulerParams& Params) final;
-	void DispatchExecuteRemote(IBuildJob* Job, const FBuildSchedulerParams& Params) final;
-	void DispatchExecuteLocal(IBuildJob* Job, const FBuildSchedulerParams& Params) final;
+	void DispatchCacheQuery(IBuildJob* Job, IRequestOwner& Owner, const FBuildSchedulerParams& Params) final;
+	void DispatchCacheStore(IBuildJob* Job, IRequestOwner& Owner, const FBuildSchedulerParams& Params) final;
+	void DispatchResolveKey(IBuildJob* Job, IRequestOwner& Owner) final;
+	void DispatchResolveInputMeta(IBuildJob* Job, IRequestOwner& Owner) final;
+	void DispatchResolveInputData(IBuildJob* Job, IRequestOwner& Owner, const FBuildSchedulerParams& Params) final;
+	void DispatchExecuteRemote(IBuildJob* Job, IRequestOwner& Owner, const FBuildSchedulerParams& Params) final;
+	void DispatchExecuteLocal(IBuildJob* Job, IRequestOwner& Owner, const FBuildSchedulerParams& Params) final;
+
+private:
+	void Dispatch(IBuildJob* Job, IRequestOwner& Owner, const TCHAR* DebugName);
 };
 
-void FBuildScheduler::DispatchCacheQuery(IBuildJob* Job, const FBuildSchedulerParams& Params)
+void FBuildScheduler::DispatchCacheQuery(IBuildJob* Job, IRequestOwner& Owner, const FBuildSchedulerParams& Params)
 {
 	Job->Schedule();
 }
 
-void FBuildScheduler::DispatchCacheStore(IBuildJob* Job, const FBuildSchedulerParams& Params)
+void FBuildScheduler::DispatchCacheStore(IBuildJob* Job, IRequestOwner& Owner, const FBuildSchedulerParams& Params)
 {
 	Job->Schedule();
 }
 
-void FBuildScheduler::DispatchResolveKey(IBuildJob* Job)
+void FBuildScheduler::DispatchResolveKey(IBuildJob* Job, IRequestOwner& Owner)
 {
-	if (Job->GetPriority() == EPriority::Blocking)
-	{
-		Job->Schedule();
-	}
-	else
-	{
-		Tasks::Launch(TEXT("FBuildScheduler::DispatchResolveKey"),
-			[Job = TRequest(Job)] { Job->Schedule(); },
-			LowLevelTasks::ETaskPriority::BackgroundNormal);
-	}
+	Dispatch(Job, Owner, TEXT("FBuildScheduler::DispatchResolveKey"));
 }
 
-void FBuildScheduler::DispatchResolveInputMeta(IBuildJob* Job)
+void FBuildScheduler::DispatchResolveInputMeta(IBuildJob* Job, IRequestOwner& Owner)
 {
-	if (Job->GetPriority() == EPriority::Blocking)
-	{
-		Job->Schedule();
-	}
-	else
-	{
-		Tasks::Launch(TEXT("FBuildScheduler::DispatchResolveInputMeta"),
-			[Job = TRequest(Job)] { Job->Schedule(); },
-			LowLevelTasks::ETaskPriority::BackgroundNormal);
-	}
+	Dispatch(Job, Owner, TEXT("FBuildScheduler::DispatchResolveInputMeta"));
 }
 
-void FBuildScheduler::DispatchResolveInputData(IBuildJob* Job, const FBuildSchedulerParams& Params)
+void FBuildScheduler::DispatchResolveInputData(IBuildJob* Job, IRequestOwner& Owner, const FBuildSchedulerParams& Params)
 {
-	if (Job->GetPriority() == EPriority::Blocking)
-	{
-		Job->Schedule();
-	}
-	else
-	{
-		Tasks::Launch(TEXT("FBuildScheduler::DispatchResolveInputData"),
-			[Job = TRequest(Job)] { Job->Schedule(); },
-			LowLevelTasks::ETaskPriority::BackgroundNormal);
-	}
+	Dispatch(Job, Owner, TEXT("FBuildScheduler::DispatchResolveInputData"));
 }
 
-void FBuildScheduler::DispatchExecuteRemote(IBuildJob* Job, const FBuildSchedulerParams& Params)
+void FBuildScheduler::DispatchExecuteRemote(IBuildJob* Job, IRequestOwner& Owner, const FBuildSchedulerParams& Params)
 {
-	if (Job->GetPriority() == EPriority::Blocking)
-	{
-		Job->Schedule();
-	}
-	else
-	{
-		Tasks::Launch(TEXT("FBuildScheduler::DispatchExecuteRemote"),
-			[Job = TRequest(Job)] { Job->Schedule(); },
-			LowLevelTasks::ETaskPriority::BackgroundNormal);
-	}
+	Dispatch(Job, Owner, TEXT("FBuildScheduler::DispatchExecuteRemote"));
+}
+
+void FBuildScheduler::DispatchExecuteLocal(IBuildJob* Job, IRequestOwner& Owner, const FBuildSchedulerParams& Params)
+{
+	Dispatch(Job, Owner, TEXT("FBuildScheduler::DispatchExecuteLocal"));
 }
 
-void FBuildScheduler::DispatchExecuteLocal(IBuildJob* Job, const FBuildSchedulerParams& Params)
+void FBuildScheduler::Dispatch(IBuildJob* Job, IRequestOwner& Owner, const TCHAR* DebugName)
 {
-	if (Job->GetPriority() == EPriority::Blocking)
+	if (Owner.GetPriority() == EPriority::Blocking)
 	{
 		Job->Schedule();
 	}
 	else
 	{
-		Tasks::Launch(TEXT("FBuildScheduler::DispatchExecuteLocal"),
-			[Job = TRequest(Job)] { Job->Schedule(); },
-			LowLevelTasks::ETaskPriority::BackgroundNormal);
+		new FBuildSchedulerRequest(Job, Owner, DebugName);
 	}
 }
 
diff --git a/Engine/Source/Developer/DerivedDataCache/Private/DerivedDataBuildSession.cpp b/Engine/Source/Developer/DerivedDataCache/Private/DerivedDataBuildSession.cpp
--- a/Engine/Source/Developer/DerivedDataCache/Private/DerivedDataBuildSession.cpp
+++ b/Engine/Source/Developer/DerivedDataCache/Private/DerivedDataBuildSession.cpp
@@ -31,23 +31,23 @@ public:
 
 	FStringView GetName() const final { return Name; }
 
-	FRequest Build(
+	void Build(
 		const FBuildDefinition& Definition,
 		EBuildPolicy Policy,
-		EPriority Priority,
+		IRequestOwner& Owner,
 		FOnBuildComplete&& OnComplete) final;
 
-	FRequest BuildAction(
+	void BuildAction(
 		const FBuildAction& Action,
 		const FOptionalBuildInputs& Inputs,
 		EBuildPolicy Policy,
-		EPriority Priority,
+		IRequestOwner& Owner,
 		FOnBuildActionComplete&& OnComplete) final;
 
-	FRequest BuildPayload(
+	void BuildPayload(
 		const FBuildPayloadKey& Payload,
 		EBuildPolicy Policy,
-		EPriority Priority,
+		IRequestOwner& Owner,
 		FOnBuildPayloadComplete&& OnComplete) final;
 
 	FString Name;
@@ -59,13 +59,12 @@ public:
 
 ///////////////////////////////////////////////////////////////////////////////////////////////////////////////////////
 
-FRequest FBuildSessionInternal::Build(
+void FBuildSessionInternal::Build(
 	const FBuildDefinition& Definition,
 	EBuildPolicy Policy,
-	EPriority Priority,
+	IRequestOwner& Owner,
 	FOnBuildComplete&& OnComplete)
 {
-	TRequest Job(CreateBuildJob(Cache, BuildSystem, InputResolver, Definition));
 	FOnBuildJobComplete OnJobComplete;
 	if (OnComplete)
 	{
@@ -74,18 +73,16 @@ FRequest FBuildSessionInternal::Build(
 			OnComplete({Definition.GetKey(), Params.CacheKey, MoveTemp(Params.Output), Params.BuildStatus, Params.Status});
 		};
 	}
-	Job->Start(Scheduler, Policy, Priority, MoveTemp(OnJobComplete));
-	return Job;
+	CreateBuildJob({Cache, BuildSystem, Scheduler, InputResolver, Owner, Policy}, Definition, MoveTemp(OnJobComplete));
 }
 
-FRequest FBuildSessionInternal::BuildAction(
+void FBuildSessionInternal::BuildAction(
 	const FBuildAction& Action,
 	const FOptionalBuildInputs& Inputs,
 	EBuildPolicy Policy,
-	EPriority Priority,
+	IRequestOwner& Owner,
 	FOnBuildActionComplete&& OnComplete)
 {
-	TRequest Job(CreateBuildJob(Cache, BuildSystem, InputResolver, Action, Inputs));
 	FOnBuildJobComplete OnJobComplete;
 	if (OnComplete)
 	{
@@ -94,18 +91,16 @@ FRequest FBuildSessionInternal::BuildAction(
 			OnComplete({Action.GetKey(), Params.CacheKey, MoveTemp(Params.Output), Params.BuildStatus, Params.Status});
 		};
 	}
-	Job->Start(Scheduler, Policy, Priority, MoveTemp(OnJobComplete));
-	return Job;
+	CreateBuildJob({Cache, BuildSystem, Scheduler, InputResolver, Owner, Policy}, Action, Inputs, MoveTemp(OnJobComplete));
 }
 
-FRequest FBuildSessionInternal::BuildPayload(
+void FBuildSessionInternal::BuildPayload(
 	const FBuildPayloadKey& PayloadKey,
 	EBuildPolicy Policy,
-	EPriority Priority,
+	IRequestOwner& Owner,
 	FOnBuildPayloadComplete&& OnComplete)
 {
 	// This requests the entire output to get one payload. It will be optimized later to request only one payload.
-	TRequest Job(CreateBuildJob(Cache, BuildSystem, InputResolver, PayloadKey.BuildKey));
 	FOnBuildJobComplete OnJobComplete;
 	if (OnComplete)
 	{
@@ -125,8 +120,7 @@ FRequest FBuildSessionInternal::BuildPayload(
 			OnComplete({PayloadKey.BuildKey, MoveTemp(Payload), Status});
 		};
 	}
-	Job->Start(Scheduler, Policy, Priority, MoveTemp(OnJobComplete));
-	return Job;
+	CreateBuildJob({Cache, BuildSystem, Scheduler, InputResolver, Owner, Policy}, PayloadKey.BuildKey, MoveTemp(OnJobComplete));
 }
 
 ///////////////////////////////////////////////////////////////////////////////////////////////////////////////////////
diff --git a/Engine/Source/Developer/DerivedDataCache/Private/DerivedDataBuildWorkerRegistry.cpp b/Engine/Source/Developer/DerivedDataCache/Private/DerivedDataBuildWorkerRegistry.cpp
--- a/Engine/Source/Developer/DerivedDataCache/Private/DerivedDataBuildWorkerRegistry.cpp
+++ b/Engine/Source/Developer/DerivedDataCache/Private/DerivedDataBuildWorkerRegistry.cpp
@@ -35,7 +35,7 @@ public:
 	inline FStringView GetHostPlatform() const final { return HostPlatform; }
 	inline FGuid GetBuildSystemVersion() const final { return BuildSystemVersion; }
 
-	FRequest FindFileData(TConstArrayView<FIoHash> RawHashes, EPriority Priority, FOnBuildWorkerFileDataComplete&& OnComplete) const final;
+	void FindFileData(TConstArrayView<FIoHash> RawHashes, IRequestOwner& Owner, FOnBuildWorkerFileDataComplete&& OnComplete) const final;
 
 	void IterateFunctions(TFunctionRef<void (FStringView WorkerName, const FGuid& Version)> Visitor) const final;
 	void IterateFiles(TFunctionRef<void (FStringView WorkerPath, const FIoHash& RawHash, uint64 RawSize)> Visitor) const final;
@@ -73,9 +73,9 @@ void FBuildWorkerInternal::Build()
 	Environment.Sort();
 }
 
-FRequest FBuildWorkerInternal::FindFileData(TConstArrayView<FIoHash> RawHashes, EPriority Priority, FOnBuildWorkerFileDataComplete&& OnComplete) const
+void FBuildWorkerInternal::FindFileData(TConstArrayView<FIoHash> RawHashes, IRequestOwner& Owner, FOnBuildWorkerFileDataComplete&& OnComplete) const
 {
-	return Factory->FindFileData(RawHashes, Priority, MoveTemp(OnComplete));
+	return Factory->FindFileData(RawHashes, Owner, MoveTemp(OnComplete));
 }
 
 void FBuildWorkerInternal::IterateFunctions(TFunctionRef<void (FStringView Name, const FGuid& Version)> Visitor) const
diff --git a/Engine/Source/Developer/DerivedDataCache/Private/DerivedDataCache.cpp b/Engine/Source/Developer/DerivedDataCache/Private/DerivedDataCache.cpp
--- a/Engine/Source/Developer/DerivedDataCache/Private/DerivedDataCache.cpp
+++ b/Engine/Source/Developer/DerivedDataCache/Private/DerivedDataCache.cpp
@@ -13,6 +13,7 @@
 #include "DerivedDataCacheRecord.h"
 #include "DerivedDataCacheUsageStats.h"
 #include "DerivedDataPluginInterface.h"
+#include "DerivedDataRequest.h"
 #include "HAL/ThreadSafeCounter.h"
 #include "Misc/CoreMisc.h"
 #include "Misc/CommandLine.h"
@@ -682,6 +683,8 @@ private:
 public:
 	// ICache Interface
 
+	FRequestGroup CreateGroup(EPriority Priority) final { return CreateRequestGroup(Priority); }
+
 	FCacheBucket CreateBucket(FStringView Name) final { return CreateCacheBucket(Name); }
 
 	FCacheRecordBuilder CreateRecord(const FCacheKey& Key) final { return CreateCacheRecordBuilder(Key); }
@@ -696,34 +699,34 @@ public:
 		return LoadCacheRecord(Package);
 	}
 
-	FRequest Put(
+	void Put(
 		TConstArrayView<FCacheRecord> Records,
 		FStringView Context,
 		ECachePolicy Policy,
-		EPriority Priority,
+		IRequestOwner& Owner,
 		FOnCachePutComplete&& OnComplete) final
 	{
-		return FDerivedDataBackend::Get().GetRoot().Put(Records, Context, Policy, Priority, MoveTemp(OnComplete));
+		return FDerivedDataBackend::Get().GetRoot().Put(Records, Context, Policy, Owner, MoveTemp(OnComplete));
 	}
 
-	FRequest Get(
+	void  Get(
 		TConstArrayView<FCacheKey> Keys,
 		FStringView Context,
 		ECachePolicy Policy,
-		EPriority Priority,
+		IRequestOwner& Owner,
 		FOnCacheGetComplete&& OnComplete) final
 	{
-		return FDerivedDataBackend::Get().GetRoot().Get(Keys, Context, Policy, Priority, MoveTemp(OnComplete));
+		return FDerivedDataBackend::Get().GetRoot().Get(Keys, Context, Policy, Owner, MoveTemp(OnComplete));
 	}
 
-	FRequest GetPayload(
+	void GetPayload(
 		TConstArrayView<FCachePayloadKey> Keys,
 		FStringView Context,
 		ECachePolicy Policy,
-		EPriority Priority,
+		IRequestOwner& Owner,
 		FOnCacheGetPayloadComplete&& OnComplete) final
 	{
-		return FDerivedDataBackend::Get().GetRoot().GetPayload(Keys, Context, Policy, Priority, MoveTemp(OnComplete));
+		return FDerivedDataBackend::Get().GetRoot().GetPayload(Keys, Context, Policy, Owner, MoveTemp(OnComplete));
 	}
 
 	void CancelAll() final
diff --git a/Engine/Source/Developer/DerivedDataCache/Private/DerivedDataCacheRecord.cpp b/Engine/Source/Developer/DerivedDataCache/Private/DerivedDataCacheRecord.cpp
--- a/Engine/Source/Developer/DerivedDataCache/Private/DerivedDataCacheRecord.cpp
+++ b/Engine/Source/Developer/DerivedDataCache/Private/DerivedDataCacheRecord.cpp
@@ -41,7 +41,7 @@ public:
 	FPayloadId AddAttachment(const FPayload& Payload) final;
 
 	FCacheRecord Build() final;
-	FRequest BuildAsync(FOnCacheRecordComplete&& OnComplete, EPriority Priority) final;
+	void BuildAsync(IRequestOwner& Owner, FOnCacheRecordComplete&& OnComplete) final;
 
 	FCacheKey Key;
 	FCbObject Meta;
@@ -258,13 +258,12 @@ FCacheRecord FCacheRecordBuilderInternal::Build()
 	return CreateCacheRecord(new FCacheRecordInternal(MoveTemp(*this)));
 }
 
-FRequest FCacheRecordBuilderInternal::BuildAsync(FOnCacheRecordComplete&& OnComplete, EPriority Priority)
+void FCacheRecordBuilderInternal::BuildAsync(IRequestOwner& Owner, FOnCacheRecordComplete&& OnComplete)
 {
 	ON_SCOPE_EXIT { delete this; };
 	checkf(OnComplete, TEXT("Failed to build cache record for %s because the completion callback is null."),
 		*WriteToString<96>(Key));
 	OnComplete(Build());
-	return FRequest();
 }
 
 ///////////////////////////////////////////////////////////////////////////////////////////////////////////////////////
diff --git a/Engine/Source/Developer/DerivedDataCache/Private/DerivedDataLimitKeyLengthWrapper.h b/Engine/Source/Developer/DerivedDataCache/Private/DerivedDataLimitKeyLengthWrapper.h
--- a/Engine/Source/Developer/DerivedDataCache/Private/DerivedDataLimitKeyLengthWrapper.h
+++ b/Engine/Source/Developer/DerivedDataCache/Private/DerivedDataLimitKeyLengthWrapper.h
@@ -242,34 +242,34 @@ public:
 		return Usage;
 	}
 
-	virtual FRequest Put(
+	virtual void Put(
 		TConstArrayView<FCacheRecord> Records,
 		FStringView Context,
 		ECachePolicy Policy,
-		EPriority Priority,
+		IRequestOwner& Owner,
 		FOnCachePutComplete&& OnComplete) override
 	{
-		return InnerBackend->Put(Records, Context, Policy, Priority, MoveTemp(OnComplete));
+		return InnerBackend->Put(Records, Context, Policy, Owner, MoveTemp(OnComplete));
 	}
 
-	virtual FRequest Get(
+	virtual void Get(
 		TConstArrayView<FCacheKey> Keys,
 		FStringView Context,
 		ECachePolicy Policy,
-		EPriority Priority,
+		IRequestOwner& Owner,
 		FOnCacheGetComplete&& OnComplete) override
 	{
-		return InnerBackend->Get(Keys, Context, Policy, Priority, MoveTemp(OnComplete));
+		return InnerBackend->Get(Keys, Context, Policy, Owner, MoveTemp(OnComplete));
 	}
 
-	virtual FRequest GetPayload(
+	virtual void GetPayload(
 		TConstArrayView<FCachePayloadKey> Keys,
 		FStringView Context,
 		ECachePolicy Policy,
-		EPriority Priority,
+		IRequestOwner& Owner,
 		FOnCacheGetPayloadComplete&& OnComplete) override
 	{
-		return InnerBackend->GetPayload(Keys, Context, Policy, Priority, MoveTemp(OnComplete));
+		return InnerBackend->GetPayload(Keys, Context, Policy, Owner, MoveTemp(OnComplete));
 	}
 
 	virtual void CancelAll() override
diff --git a/Engine/Source/Developer/DerivedDataCache/Private/DerivedDataRequest.cpp b/Engine/Source/Developer/DerivedDataCache/Private/DerivedDataRequest.cpp
new file mode 100644
--- /dev/null
+++ b/Engine/Source/Developer/DerivedDataCache/Private/DerivedDataRequest.cpp
@@ -0,0 +1,256 @@
+// Copyright Epic Games, Inc. All Rights Reserved.
+
+#include "DerivedDataRequest.h"
+
+#include "Containers/Array.h"
+#include "HAL/CriticalSection.h"
+#include "HAL/Event.h"
+#include "Misc/ScopeExit.h"
+#include "Misc/ScopeRWLock.h"
+#include "Templates/RefCounting.h"
+#include <atomic>
+
+namespace UE::DerivedData::Private
+{
+
+/** A one-time-use event to work around the lack of condition variables. */
+struct FRequestBarrierEvent : public FRefCountBase
+{
+	FEventRef Event{EEventMode::ManualReset};
+};
+
+class FRequestGroupOwner final : public FRequestGroupBase
+{
+public:
+	explicit FRequestGroupOwner(EPriority Priority);
+	~FRequestGroupOwner() final = default;
+
+	void Begin(IRequest* Request) final;
+	TRefCountPtr<IRequest> End(IRequest* Request) final;
+
+	void BeginBarrier() final;
+	void EndBarrier() final;
+
+	inline EPriority GetPriority() const final { return Priority; }
+	inline bool IsCanceled() const final { return bIsCanceled; }
+
+	void SetPriority(EPriority Priority) final;
+	void Cancel() final;
+	void Wait() final;
+	bool Poll() const final;
+
+	void KeepAlive() final;
+	void Destroy() final;
+
+private:
+	mutable FRWLock Lock;
+	TArray<TRefCountPtr<IRequest>, TInlineAllocator<1>> Requests;
+	TRefCountPtr<FRequestBarrierEvent> BarrierEvent;
+	uint32 BarrierCount{0};
+	EPriority Priority{EPriority::Normal};
+	uint8 bPriorityChangedInBarrier : 1;
+	uint8 bBeginExecuted : 1;
+	bool bIsCanceled{false};
+	bool bKeepAlive{false};
+};
+
+FRequestGroupOwner::FRequestGroupOwner(EPriority NewPriority)
+	: Priority(NewPriority)
+	, bPriorityChangedInBarrier(false)
+	, bBeginExecuted(false)
+{
+	AddRef(); // Release is called by Destroy.
+}
+
+void FRequestGroupOwner::Begin(IRequest* Request)
+{
+	AddRef();
+	EPriority NewPriority;
+	{
+		FWriteScopeLock WriteLock(Lock);
+		checkf(BarrierCount > 0 || !bBeginExecuted,
+			TEXT("At least one FRequestBarrier must be in scope when beginning a request after the first request. ")
+			TEXT("The overload of End that invokes a callback handles this automatically for most use cases."));
+		check(Request);
+		Requests.Add(Request);
+		bBeginExecuted = true;
+		if (!bPriorityChangedInBarrier)
+		{
+			return;
+		}
+		NewPriority = Priority;
+	}
+	// Loop until priority is stable. Another thread may be changing the priority concurrently.
+	for (EPriority CheckPriority; ; NewPriority = CheckPriority)
+	{
+		Request->SetPriority(NewPriority);
+		CheckPriority = (FReadScopeLock(Lock), Priority);
+		if (CheckPriority == NewPriority)
+		{
+			break;
+		}
+	}
+}
+
+TRefCountPtr<IRequest> FRequestGroupOwner::End(IRequest* Request)
+{
+	ON_SCOPE_EXIT { Release(); };
+	FWriteScopeLock WriteLock(Lock);
+	check(Request);
+	TRefCountPtr<IRequest>* RequestPtr = Requests.FindByKey(Request);
+	check(RequestPtr);
+	TRefCountPtr<IRequest> RequestRef = MoveTemp(*RequestPtr);
+	Requests.RemoveAtSwap(UE_PTRDIFF_TO_INT32(RequestPtr - Requests.GetData()), 1, /*bAllowShrinking*/ false);
+	return RequestRef;
+}
+
+void FRequestGroupOwner::BeginBarrier()
+{
+	AddRef();
+	FWriteScopeLock WriteLock(Lock);
+	++BarrierCount;
+}
+
+void FRequestGroupOwner::EndBarrier()
+{
+	ON_SCOPE_EXIT { Release(); };
+	TRefCountPtr<FRequestBarrierEvent> LocalBarrierEvent;
+	if (FWriteScopeLock WriteLock(Lock); --BarrierCount == 0)
+	{
+		bPriorityChangedInBarrier = false;
+		LocalBarrierEvent = MoveTemp(BarrierEvent);
+		BarrierEvent = nullptr;
+	}
+	if (LocalBarrierEvent)
+	{
+		LocalBarrierEvent->Event->Trigger();
+	}
+}
+
+void FRequestGroupOwner::SetPriority(EPriority NewPriority)
+{
+	TArray<TRefCountPtr<IRequest>, TInlineAllocator<16>> LocalRequests;
+	if (FWriteScopeLock WriteLock(Lock); Priority == NewPriority)
+	{
+		return;
+	}
+	else
+	{
+		Priority = NewPriority;
+		LocalRequests = Requests;
+		bPriorityChangedInBarrier = (BarrierCount > 0);
+	}
+
+	for (IRequest* Request : LocalRequests)
+	{
+		Request->SetPriority(NewPriority);
+	}
+}
+
+void FRequestGroupOwner::Cancel()
+{
+	for (;;)
+	{
+		TRefCountPtr<IRequest> LocalRequest;
+		TRefCountPtr<FRequestBarrierEvent> LocalBarrierEvent;
+
+		{
+			FWriteScopeLock WriteLock(Lock);
+			bIsCanceled = true;
+
+			if (!Requests.IsEmpty())
+			{
+				LocalRequest = Requests.Last();
+			}
+			else if (BarrierCount > 0)
+			{
+				LocalBarrierEvent = BarrierEvent;
+				if (!LocalBarrierEvent)
+				{
+					LocalBarrierEvent = BarrierEvent = new FRequestBarrierEvent;
+				}
+			}
+			else
+			{
+				return;
+			}
+		}
+
+		if (LocalRequest)
+		{
+			LocalRequest->Cancel();
+		}
+		else
+		{
+			LocalBarrierEvent->Event->Wait();
+		}
+	}
+}
+
+void FRequestGroupOwner::Wait()
+{
+	for (;;)
+	{
+		TRefCountPtr<IRequest> LocalRequest;
+		TRefCountPtr<FRequestBarrierEvent> LocalBarrierEvent;
+
+		{
+			FWriteScopeLock WriteLock(Lock);
+
+			if (!Requests.IsEmpty())
+			{
+				LocalRequest = Requests.Last();
+			}
+			else if (BarrierCount > 0)
+			{
+				LocalBarrierEvent = BarrierEvent;
+				if (!LocalBarrierEvent)
+				{
+					LocalBarrierEvent = BarrierEvent = new FRequestBarrierEvent;
+				}
+			}
+			else
+			{
+				return;
+			}
+		}
+
+		if (LocalRequest)
+		{
+			LocalRequest->Wait();
+		}
+		else
+		{
+			LocalBarrierEvent->Event->Wait();
+		}
+	}
+}
+
+bool FRequestGroupOwner::Poll() const
+{
+	FReadScopeLock ReadLock(Lock);
+	return Requests.IsEmpty() && BarrierCount == 0;
+}
+
+void FRequestGroupOwner::KeepAlive()
+{
+	FWriteScopeLock WriteLock(Lock);
+	bKeepAlive = true;
+}
+
+void FRequestGroupOwner::Destroy()
+{
+	const bool bLocalKeepAlive = (FWriteScopeLock(Lock), bKeepAlive);
+	if (!bLocalKeepAlive)
+	{
+		Cancel();
+	}
+	Release();
+}
+
+FRequestGroup CreateRequestGroup(EPriority Priority)
+{
+	return FRequestGroup(new FRequestGroupOwner(Priority));
+}
+
+} // UE::DerivedData::Private
diff --git a/Engine/Source/Developer/DerivedDataCache/Private/FileSystemDerivedDataBackend.cpp b/Engine/Source/Developer/DerivedDataCache/Private/FileSystemDerivedDataBackend.cpp
--- a/Engine/Source/Developer/DerivedDataCache/Private/FileSystemDerivedDataBackend.cpp
+++ b/Engine/Source/Developer/DerivedDataCache/Private/FileSystemDerivedDataBackend.cpp
@@ -880,11 +880,11 @@ public:
 		return true;
 	}
 
-	virtual FRequest Put(
+	virtual void Put(
 		TConstArrayView<FCacheRecord> Records,
 		FStringView Context,
 		ECachePolicy Policy,
-		EPriority Priority,
+		IRequestOwner& Owner,
 		FOnCachePutComplete&& OnComplete) override
 	{
 		for (const FCacheRecord& Record : Records)
@@ -913,14 +913,13 @@ public:
 				}
 			}
 		}
-		return FRequest();
 	}
 
-	virtual FRequest Get(
+	virtual void Get(
 		TConstArrayView<FCacheKey> Keys,
 		FStringView Context,
 		ECachePolicy Policy,
-		EPriority Priority,
+		IRequestOwner& Owner,
 		FOnCacheGetComplete&& OnComplete) override
 	{
 		for (const FCacheKey& Key : Keys)
@@ -948,14 +947,13 @@ public:
 				}
 			}
 		}
-		return FRequest();
 	}
 
-	virtual FRequest GetPayload(
+	virtual void GetPayload(
 		TConstArrayView<FCachePayloadKey> Keys,
 		FStringView Context,
 		ECachePolicy Policy,
-		EPriority Priority,
+		IRequestOwner& Owner,
 		FOnCacheGetPayloadComplete&& OnComplete) override
 	{
 		TArray<FCachePayloadKey, TInlineAllocator<16>> SortedKeys(Keys);
@@ -991,7 +989,6 @@ public:
 				}
 			}
 		}
-		return FRequest();
 	}
 
 	virtual void CancelAll() override
diff --git a/Engine/Source/Developer/DerivedDataCache/Private/HierarchicalDerivedDataBackend.h b/Engine/Source/Developer/DerivedDataCache/Private/HierarchicalDerivedDataBackend.h
--- a/Engine/Source/Developer/DerivedDataCache/Private/HierarchicalDerivedDataBackend.h
+++ b/Engine/Source/Developer/DerivedDataCache/Private/HierarchicalDerivedDataBackend.h
@@ -466,13 +466,17 @@ public:
 		return Usage;
 	}
 
-	virtual FRequest Put(
+	virtual void Put(
 		TConstArrayView<FCacheRecord> Records,
 		FStringView Context,
 		ECachePolicy Policy,
-		EPriority Priority,
+		IRequestOwner& Owner,
 		FOnCachePutComplete&& OnComplete) override
 	{
+		FRequestGroup BlockingGroup = Factory.CreateGroup(EPriority::Blocking);
+		FRequestGroup AsyncGroup = Factory.CreateGroup(Owner.GetPriority());
+		AsyncGroup.KeepAlive();
+
 		TSet<FCacheKey> RecordsOk;
 
 		{
@@ -485,7 +489,7 @@ public:
 					// Every record must put synchronously before switching to async calls.
 					if (RecordsOk.Num() < Records.Num())
 					{
-						InnerBackends[PutCacheIndex]->Put(Records, Context, Policy, Priority,
+						InnerBackends[PutCacheIndex]->Put(Records, Context, Policy, BlockingGroup,
 							[&OnComplete, &RecordsOk](FCachePutCompleteParams&& Params)
 							{
 								if (Params.Status == EStatus::Ok)
@@ -497,11 +501,12 @@ public:
 										OnComplete(MoveTemp(Params));
 									}
 								}
-							}).Wait();
+							});
+						BlockingGroup.Wait();
 					}
 					else
 					{
-						AsyncPutInnerBackends[PutCacheIndex]->Put(Records, Context, Policy, Priority);
+						AsyncPutInnerBackends[PutCacheIndex]->Put(Records, Context, Policy, AsyncGroup);
 					}
 				}
 			}
@@ -517,15 +522,13 @@ public:
 				}
 			}
 		}
-
-		return FRequest();
 	}
 
-	virtual FRequest Get(
+	virtual void Get(
 		TConstArrayView<FCacheKey> Keys,
 		FStringView Context,
 		ECachePolicy Policy,
-		EPriority Priority,
+		IRequestOwner& Owner,
 		FOnCacheGetComplete&& OnComplete) override
 	{
 		const bool bQueryLocal = bHasLocalBackends && EnumHasAnyFlags(Policy, ECachePolicy::QueryLocal);
@@ -534,6 +537,7 @@ public:
 		const bool bStoreRemote = bHasWritableRemoteBackends && EnumHasAnyFlags(Policy, ECachePolicy::StoreRemote);
 		const bool bStoreLocalCopy = bStoreLocal && bHasMultipleLocalBackends && !EnumHasAnyFlags(Policy, ECachePolicy::SkipLocalCopy);
 
+		FRequestGroup BlockingGroup = Factory.CreateGroup(EPriority::Blocking);
 		TArray<FCacheKey, TInlineAllocator<16>> RemainingKeys(Keys);
 
 		TSet<FCacheKey> KeysOk;
@@ -553,13 +557,15 @@ public:
 				const ECachePolicy FillPolicy = bFill ? (Policy & ~ECachePolicy::SkipData) : Policy;
 
 				// Block on this because backends in this hierarchy are not expected to be asynchronous.
-				InnerBackends[GetCacheIndex]->Get(RemainingKeys, Context, Policy, Priority,
-					[this, Context, GetCacheIndex, bStoreLocal, bStoreRemote, bStoreLocalCopy, bIsLocalGet, bFill, &OnComplete, &KeysOk](FCacheGetCompleteParams&& Params)
+				InnerBackends[GetCacheIndex]->Get(RemainingKeys, Context, Policy, BlockingGroup,
+					[this, Context, GetCacheIndex, bStoreLocal, bStoreRemote, bStoreLocalCopy, bIsLocalGet, bFill, &Owner, &OnComplete, &KeysOk](FCacheGetCompleteParams&& Params)
 					{
 						if (Params.Status == EStatus::Ok)
 						{
 							if (bFill)
 							{
+								FRequestGroup AsyncGroup = Factory.CreateGroup(Owner.GetPriority());
+								AsyncGroup.KeepAlive();
 								for (int32 FillCacheIndex = 0; FillCacheIndex < InnerBackends.Num(); ++FillCacheIndex)
 								{
 									if (GetCacheIndex != FillCacheIndex)
@@ -567,7 +573,7 @@ public:
 										const bool bIsLocalFill = InnerBackends[FillCacheIndex]->GetSpeedClass() == ESpeedClass::Local;
 										if ((bIsLocalFill && bStoreLocal && bIsLocalGet <= bStoreLocalCopy) || (!bIsLocalFill && bStoreRemote))
 										{
-											AsyncPutInnerBackends[FillCacheIndex]->Put({Params.Record}, Context);
+											AsyncPutInnerBackends[FillCacheIndex]->Put({Params.Record}, Context, ECachePolicy::Default, AsyncGroup);
 										}
 									}
 								}
@@ -579,7 +585,8 @@ public:
 								OnComplete(MoveTemp(Params));
 							}
 						}
-					}).Wait();
+					});
+				BlockingGroup.Wait();
 
 				RemainingKeys.RemoveAll([&KeysOk](const FCacheKey& Key) { return KeysOk.Contains(Key); });
 			}
@@ -592,17 +599,16 @@ public:
 				OnComplete({Factory.CreateRecord(Key).Build(), EStatus::Error});
 			}
 		}
-
-		return FRequest();
 	}
 
-	virtual FRequest GetPayload(
+	virtual void GetPayload(
 		TConstArrayView<FCachePayloadKey> Keys,
 		FStringView Context,
 		ECachePolicy Policy,
-		EPriority Priority,
+		IRequestOwner& Owner,
 		FOnCacheGetPayloadComplete&& OnComplete) override
 	{
+		FRequestGroup BlockingGroup = Factory.CreateGroup(EPriority::Blocking);
 		TArray<FCachePayloadKey, TInlineAllocator<16>> RemainingKeys(Keys);
 
 		TSet<FCachePayloadKey> KeysOk;
@@ -612,7 +618,7 @@ public:
 
 			for (int32 GetCacheIndex = 0; GetCacheIndex < InnerBackends.Num() && !RemainingKeys.IsEmpty(); ++GetCacheIndex)
 			{
-				InnerBackends[GetCacheIndex]->GetPayload(RemainingKeys, Context, Policy, Priority,
+				InnerBackends[GetCacheIndex]->GetPayload(RemainingKeys, Context, Policy, BlockingGroup,
 					[&OnComplete, &KeysOk](FCacheGetPayloadCompleteParams&& Params)
 					{
 						if (Params.Status == EStatus::Ok)
@@ -623,7 +629,8 @@ public:
 								OnComplete(MoveTemp(Params));
 							}
 						}
-					}).Wait();
+					});
+				BlockingGroup.Wait();
 
 				RemainingKeys.RemoveAll([&KeysOk](const FCachePayloadKey& Key) { return KeysOk.Contains(Key); });
 			}
@@ -636,8 +643,6 @@ public:
 				OnComplete({Key.CacheKey, FPayload(Key.Id), EStatus::Error});
 			}
 		}
-
-		return FRequest();
 	}
 
 	virtual void CancelAll() override
diff --git a/Engine/Source/Developer/DerivedDataCache/Private/HttpDerivedDataBackend.cpp b/Engine/Source/Developer/DerivedDataCache/Private/HttpDerivedDataBackend.cpp
--- a/Engine/Source/Developer/DerivedDataCache/Private/HttpDerivedDataBackend.cpp
+++ b/Engine/Source/Developer/DerivedDataCache/Private/HttpDerivedDataBackend.cpp
@@ -2431,11 +2431,11 @@ TSharedRef<FDerivedDataCacheStatsNode> FHttpDerivedDataBackend::GatherUsageStats
 	return Usage;
 }
 
-FRequest FHttpDerivedDataBackend::Put(
+void FHttpDerivedDataBackend::Put(
 	TConstArrayView<FCacheRecord> Records,
 	FStringView Context,
 	ECachePolicy Policy,
-	EPriority Priority,
+	IRequestOwner& Owner,
 	FOnCachePutComplete&& OnComplete)
 {
 	if (OnComplete)
@@ -2445,14 +2445,13 @@ FRequest FHttpDerivedDataBackend::Put(
 			OnComplete({Record.GetKey(), EStatus::Error});
 		}
 	}
-	return FRequest();
 }
 
-FRequest FHttpDerivedDataBackend::Get(
+void FHttpDerivedDataBackend::Get(
 	TConstArrayView<FCacheKey> Keys,
 	FStringView Context,
 	ECachePolicy Policy,
-	EPriority Priority,
+	IRequestOwner& Owner,
 	FOnCacheGetComplete&& OnComplete)
 {
 	if (OnComplete)
@@ -2462,14 +2461,13 @@ FRequest FHttpDerivedDataBackend::Get(
 			OnComplete({Factory.CreateRecord(Key).Build(), EStatus::Error});
 		}
 	}
-	return FRequest();
 }
 
-FRequest FHttpDerivedDataBackend::GetPayload(
+void FHttpDerivedDataBackend::GetPayload(
 	TConstArrayView<FCachePayloadKey> Keys,
 	FStringView Context,
 	ECachePolicy Policy,
-	EPriority Priority,
+	IRequestOwner& Owner,
 	FOnCacheGetPayloadComplete&& OnComplete)
 {
 	if (OnComplete)
@@ -2479,7 +2477,6 @@ FRequest FHttpDerivedDataBackend::GetPayload(
 			OnComplete({Key.CacheKey, FPayload(Key.Id), EStatus::Error});
 		}
 	}
-	return FRequest();
 }
 
 } // UE::DerivedData::Backends
diff --git a/Engine/Source/Developer/DerivedDataCache/Private/HttpDerivedDataBackend.h b/Engine/Source/Developer/DerivedDataCache/Private/HttpDerivedDataBackend.h
--- a/Engine/Source/Developer/DerivedDataCache/Private/HttpDerivedDataBackend.h
+++ b/Engine/Source/Developer/DerivedDataCache/Private/HttpDerivedDataBackend.h
@@ -72,25 +72,25 @@ public:
 
 	void SetSpeedClass(ESpeedClass InSpeedClass) { SpeedClass = InSpeedClass; }
 
-	virtual FRequest Put(
+	virtual void Put(
 		TConstArrayView<FCacheRecord> Records,
 		FStringView Context,
 		ECachePolicy Policy,
-		EPriority Priority,
+		IRequestOwner& Owner,
 		FOnCachePutComplete&& OnComplete) override;
 
-	virtual FRequest Get(
+	virtual void Get(
 		TConstArrayView<FCacheKey> Keys,
 		FStringView Context,
 		ECachePolicy Policy,
-		EPriority Priority,
+		IRequestOwner& Owner,
 		FOnCacheGetComplete&& OnComplete) override;
 
-	virtual FRequest GetPayload(
+	virtual void GetPayload(
 		TConstArrayView<FCachePayloadKey> Keys,
 		FStringView Context,
 		ECachePolicy Policy,
-		EPriority Priority,
+		IRequestOwner& Owner,
 		FOnCacheGetPayloadComplete&& OnComplete) override;
 
 	virtual void CancelAll() override
diff --git a/Engine/Source/Developer/DerivedDataCache/Private/MemoryDerivedDataBackend.cpp b/Engine/Source/Developer/DerivedDataCache/Private/MemoryDerivedDataBackend.cpp
--- a/Engine/Source/Developer/DerivedDataCache/Private/MemoryDerivedDataBackend.cpp
+++ b/Engine/Source/Developer/DerivedDataCache/Private/MemoryDerivedDataBackend.cpp
@@ -405,11 +405,11 @@ int64 FMemoryDerivedDataBackend::CalcSerializedCacheRecordSize(const FCacheRecor
 	return int64(TotalSize);
 }
 
-FRequest FMemoryDerivedDataBackend::Put(
+void FMemoryDerivedDataBackend::Put(
 	TConstArrayView<FCacheRecord> Records,
 	FStringView Context,
 	ECachePolicy Policy,
-	EPriority Priority,
+	IRequestOwner& Owner,
 	FOnCachePutComplete&& OnComplete)
 {
 	for (const FCacheRecord& Record : Records)
@@ -470,14 +470,13 @@ FRequest FMemoryDerivedDataBackend::Put(
 			Status = EStatus::Ok;
 		}
 	}
-	return FRequest();
 }
 
-FRequest FMemoryDerivedDataBackend::Get(
+void FMemoryDerivedDataBackend::Get(
 	TConstArrayView<FCacheKey> Keys,
 	FStringView Context,
 	ECachePolicy Policy,
-	EPriority Priority,
+	IRequestOwner& Owner,
 	FOnCacheGetComplete&& OnComplete)
 {
 	for (const FCacheKey& Key : Keys)
@@ -504,14 +503,13 @@ FRequest FMemoryDerivedDataBackend::Get(
 			}
 		}
 	}
-	return FRequest();
 }
 
-FRequest FMemoryDerivedDataBackend::GetPayload(
+void FMemoryDerivedDataBackend::GetPayload(
 	TConstArrayView<FCachePayloadKey> Keys,
 	FStringView Context,
 	ECachePolicy Policy,
-	EPriority Priority,
+	IRequestOwner& Owner,
 	FOnCacheGetPayloadComplete&& OnComplete)
 {
 	for (const FCachePayloadKey& PayloadKey : Keys)
@@ -538,7 +536,6 @@ FRequest FMemoryDerivedDataBackend::GetPayload(
 			}
 		}
 	}
-	return FRequest();
 }
 
 } // UE::DerivedData::Backends
diff --git a/Engine/Source/Developer/DerivedDataCache/Private/MemoryDerivedDataBackend.h b/Engine/Source/Developer/DerivedDataCache/Private/MemoryDerivedDataBackend.h
--- a/Engine/Source/Developer/DerivedDataCache/Private/MemoryDerivedDataBackend.h
+++ b/Engine/Source/Developer/DerivedDataCache/Private/MemoryDerivedDataBackend.h
@@ -94,25 +94,25 @@ public:
 	 */
 	bool ApplyDebugOptions(FBackendDebugOptions& InOptions) override;
 
-	virtual FRequest Put(
+	virtual void Put(
 		TConstArrayView<FCacheRecord> Records,
 		FStringView Context,
 		ECachePolicy Policy,
-		EPriority Priority,
+		IRequestOwner& Owner,
 		FOnCachePutComplete&& OnComplete) override;
 
-	virtual FRequest Get(
+	virtual void Get(
 		TConstArrayView<FCacheKey> Keys,
 		FStringView Context,
 		ECachePolicy Policy,
-		EPriority Priority,
+		IRequestOwner& Owner,
 		FOnCacheGetComplete&& OnComplete) override;
 
-	virtual FRequest GetPayload(
+	virtual void GetPayload(
 		TConstArrayView<FCachePayloadKey> Keys,
 		FStringView Context,
 		ECachePolicy Policy,
-		EPriority Priority,
+		IRequestOwner& Owner,
 		FOnCacheGetPayloadComplete&& OnComplete) override;
 
 	virtual void CancelAll() override
diff --git a/Engine/Source/Developer/DerivedDataCache/Private/PakFileDerivedDataBackend.cpp b/Engine/Source/Developer/DerivedDataCache/Private/PakFileDerivedDataBackend.cpp
--- a/Engine/Source/Developer/DerivedDataCache/Private/PakFileDerivedDataBackend.cpp
+++ b/Engine/Source/Developer/DerivedDataCache/Private/PakFileDerivedDataBackend.cpp
@@ -441,11 +441,11 @@ TSharedRef<FDerivedDataCacheStatsNode> FPakFileDerivedDataBackend::GatherUsageSt
 	return Usage;
 }
 
-FRequest FPakFileDerivedDataBackend::Put(
+void FPakFileDerivedDataBackend::Put(
 	TConstArrayView<FCacheRecord> Records,
 	FStringView Context,
 	ECachePolicy Policy,
-	EPriority Priority,
+	IRequestOwner& Owner,
 	FOnCachePutComplete&& OnComplete)
 {
 	if (OnComplete)
@@ -455,14 +455,13 @@ FRequest FPakFileDerivedDataBackend::Put(
 			OnComplete({Record.GetKey(), EStatus::Error});
 		}
 	}
-	return FRequest();
 }
 
-FRequest FPakFileDerivedDataBackend::Get(
+void FPakFileDerivedDataBackend::Get(
 	TConstArrayView<FCacheKey> Keys,
 	FStringView Context,
 	ECachePolicy Policy,
-	EPriority Priority,
+	IRequestOwner& Owner,
 	FOnCacheGetComplete&& OnComplete)
 {
 	if (OnComplete)
@@ -472,14 +471,13 @@ FRequest FPakFileDerivedDataBackend::Get(
 			OnComplete({Factory.CreateRecord(Key).Build(), EStatus::Error});
 		}
 	}
-	return FRequest();
 }
 
-FRequest FPakFileDerivedDataBackend::GetPayload(
+void FPakFileDerivedDataBackend::GetPayload(
 	TConstArrayView<FCachePayloadKey> Keys,
 	FStringView Context,
 	ECachePolicy Policy,
-	EPriority Priority,
+	IRequestOwner& Owner,
 	FOnCacheGetPayloadComplete&& OnComplete)
 {
 	if (OnComplete)
@@ -489,7 +487,6 @@ FRequest FPakFileDerivedDataBackend::GetPayload(
 			OnComplete({Key.CacheKey, FPayload(Key.Id), EStatus::Error});
 		}
 	}
-	return FRequest();
 }
 
 FCompressedPakFileDerivedDataBackend::FCompressedPakFileDerivedDataBackend(ICacheFactory& InFactory, const TCHAR* InFilename, bool bInWriting)
diff --git a/Engine/Source/Developer/DerivedDataCache/Private/PakFileDerivedDataBackend.h b/Engine/Source/Developer/DerivedDataCache/Private/PakFileDerivedDataBackend.h
--- a/Engine/Source/Developer/DerivedDataCache/Private/PakFileDerivedDataBackend.h
+++ b/Engine/Source/Developer/DerivedDataCache/Private/PakFileDerivedDataBackend.h
@@ -104,25 +104,25 @@ public:
 
 	virtual bool ApplyDebugOptions(FBackendDebugOptions& InOptions) override { return false; }
 
-	virtual FRequest Put(
+	virtual void Put(
 		TConstArrayView<FCacheRecord> Records,
 		FStringView Context,
 		ECachePolicy Policy,
-		EPriority Priority,
+		IRequestOwner& Owner,
 		FOnCachePutComplete&& OnComplete) override;
 
-	virtual FRequest Get(
+	virtual void Get(
 		TConstArrayView<FCacheKey> Keys,
 		FStringView Context,
 		ECachePolicy Policy,
-		EPriority Priority,
+		IRequestOwner& Owner,
 		FOnCacheGetComplete&& OnComplete) override;
 
-	virtual FRequest GetPayload(
+	virtual void GetPayload(
 		TConstArrayView<FCachePayloadKey> Keys,
 		FStringView Context,
 		ECachePolicy Policy,
-		EPriority Priority,
+		IRequestOwner& Owner,
 		FOnCacheGetPayloadComplete&& OnComplete) override;
 
 	virtual void CancelAll() override
diff --git a/Engine/Source/Developer/DerivedDataCache/Private/S3DerivedDataBackend.cpp b/Engine/Source/Developer/DerivedDataCache/Private/S3DerivedDataBackend.cpp
--- a/Engine/Source/Developer/DerivedDataCache/Private/S3DerivedDataBackend.cpp
+++ b/Engine/Source/Developer/DerivedDataCache/Private/S3DerivedDataBackend.cpp
@@ -1177,11 +1177,11 @@ bool FS3DerivedDataBackend::ShouldSimulateMiss(const TCHAR* InKey)
 	return false;
 }
 
-FRequest FS3DerivedDataBackend::Put(
+void FS3DerivedDataBackend::Put(
 	TConstArrayView<FCacheRecord> Records,
 	FStringView Context,
 	ECachePolicy Policy,
-	EPriority Priority,
+	IRequestOwner& Owner,
 	FOnCachePutComplete&& OnComplete)
 {
 	if (OnComplete)
@@ -1191,14 +1191,13 @@ FRequest FS3DerivedDataBackend::Put(
 			OnComplete({Record.GetKey(), EStatus::Error});
 		}
 	}
-	return FRequest();
 }
 
-FRequest FS3DerivedDataBackend::Get(
+void FS3DerivedDataBackend::Get(
 	TConstArrayView<FCacheKey> Keys,
 	FStringView Context,
 	ECachePolicy Policy,
-	EPriority Priority,
+	IRequestOwner& Owner,
 	FOnCacheGetComplete&& OnComplete)
 {
 	if (OnComplete)
@@ -1208,14 +1207,13 @@ FRequest FS3DerivedDataBackend::Get(
 			OnComplete({Factory.CreateRecord(Key).Build(), EStatus::Error});
 		}
 	}
-	return FRequest();
 }
 
-FRequest FS3DerivedDataBackend::GetPayload(
+void FS3DerivedDataBackend::GetPayload(
 	TConstArrayView<FCachePayloadKey> Keys,
 	FStringView Context,
 	ECachePolicy Policy,
-	EPriority Priority,
+	IRequestOwner& Owner,
 	FOnCacheGetPayloadComplete&& OnComplete)
 {
 	if (OnComplete)
@@ -1225,7 +1223,6 @@ FRequest FS3DerivedDataBackend::GetPayload(
 			OnComplete({Key.CacheKey, FPayload(Key.Id), EStatus::Error});
 		}
 	}
-	return FRequest();
 }
 
 } // UE::DerivedData::Backends
diff --git a/Engine/Source/Developer/DerivedDataCache/Private/S3DerivedDataBackend.h b/Engine/Source/Developer/DerivedDataCache/Private/S3DerivedDataBackend.h
--- a/Engine/Source/Developer/DerivedDataCache/Private/S3DerivedDataBackend.h
+++ b/Engine/Source/Developer/DerivedDataCache/Private/S3DerivedDataBackend.h
@@ -59,25 +59,25 @@ public:
 
 	bool ApplyDebugOptions(FBackendDebugOptions& InOptions) override;
 
-	virtual FRequest Put(
+	virtual void Put(
 		TConstArrayView<FCacheRecord> Records,
 		FStringView Context,
 		ECachePolicy Policy,
-		EPriority Priority,
+		IRequestOwner& Owner,
 		FOnCachePutComplete&& OnComplete) override;
 
-	virtual FRequest Get(
+	virtual void Get(
 		TConstArrayView<FCacheKey> Keys,
 		FStringView Context,
 		ECachePolicy Policy,
-		EPriority Priority,
+		IRequestOwner& Owner,
 		FOnCacheGetComplete&& OnComplete) override;
 
-	virtual FRequest GetPayload(
+	virtual void GetPayload(
 		TConstArrayView<FCachePayloadKey> Keys,
 		FStringView Context,
 		ECachePolicy Policy,
-		EPriority Priority,
+		IRequestOwner& Owner,
 		FOnCacheGetPayloadComplete&& OnComplete) override;
 
 	virtual void CancelAll() override
diff --git a/Engine/Source/Developer/DerivedDataCache/Private/ZenDerivedDataBackend.cpp b/Engine/Source/Developer/DerivedDataCache/Private/ZenDerivedDataBackend.cpp
--- a/Engine/Source/Developer/DerivedDataCache/Private/ZenDerivedDataBackend.cpp
+++ b/Engine/Source/Developer/DerivedDataCache/Private/ZenDerivedDataBackend.cpp
@@ -24,7 +24,7 @@ TRACE_DECLARE_INT_COUNTER(ZenDDC_PutHit,		TEXT("ZenDDC Put Hit"));
 TRACE_DECLARE_INT_COUNTER(ZenDDC_BytesReceived, TEXT("ZenDDC Bytes Received"));
 TRACE_DECLARE_INT_COUNTER(ZenDDC_BytesSent,		TEXT("ZenDDC Bytes Sent"));
 
-namespace UE::DerivedData {
+namespace UE::DerivedData::Backends {
 
 //----------------------------------------------------------------------------------------------------------
 // FZenDerivedDataBackend
@@ -485,11 +485,11 @@ bool FZenDerivedDataBackend::ShouldSimulateMiss(const TCHAR* InKey)
 	return false;
 }
 
-FRequest FZenDerivedDataBackend::Put(
-	TConstArrayView<UE::DerivedData::FCacheRecord> Records,
+void FZenDerivedDataBackend::Put(
+	TConstArrayView<FCacheRecord> Records,
 	FStringView Context,
 	ECachePolicy Policy,
-	EPriority Priority,
+	IRequestOwner& Owner,
 	FOnCachePutComplete&& OnComplete)
 {
 	for (const FCacheRecord& Record : Records)
@@ -523,17 +523,16 @@ FRequest FZenDerivedDataBackend::Put(
 			}
 		}
 	}
-	return FRequest();
 }
 
-FRequest FZenDerivedDataBackend::Get(
-	TConstArrayView<UE::DerivedData::FCacheKey> Keys,
+void FZenDerivedDataBackend::Get(
+	TConstArrayView<FCacheKey> Keys,
 	FStringView Context,
-	UE::DerivedData::ECachePolicy Policy,
-	UE::DerivedData::EPriority Priority,
-	UE::DerivedData::FOnCacheGetComplete&& OnComplete)
+	ECachePolicy Policy,
+	IRequestOwner& Owner,
+	FOnCacheGetComplete&& OnComplete)
 {
-	for (const UE::DerivedData::FCacheKey& Key : Keys)
+	for (const FCacheKey& Key : Keys)
 	{
 		TRACE_CPUPROFILER_EVENT_SCOPE(ZenDDC_Get);
 		TRACE_COUNTER_ADD(ZenDDC_Get, int64(1));
@@ -570,15 +569,14 @@ FRequest FZenDerivedDataBackend::Get(
 			}
 		}
 	}
-	return FRequest();
 }
 
-FRequest FZenDerivedDataBackend::GetPayload(
-	TConstArrayView<UE::DerivedData::FCachePayloadKey> Keys,
+void FZenDerivedDataBackend::GetPayload(
+	TConstArrayView<FCachePayloadKey> Keys,
 	FStringView Context,
-	UE::DerivedData::ECachePolicy Policy,
-	UE::DerivedData::EPriority Priority,
-	UE::DerivedData::FOnCacheGetPayloadComplete&& OnComplete)
+	ECachePolicy Policy,
+	IRequestOwner& Owner,
+	FOnCacheGetPayloadComplete&& OnComplete)
 {
 	if (bCacheRecordEndpointEnabled)
 	{
@@ -663,7 +661,6 @@ FRequest FZenDerivedDataBackend::GetPayload(
 			}
 		}
 	}
-	return FRequest();
 }
 
 void FZenDerivedDataBackend::CancelAll()
@@ -987,6 +984,6 @@ void FZenDerivedDataBackend::LegacyMakePayloadKey(const FCacheKey& CacheKey, con
 }
 
 
-} // namespace UE::DerivedData
+} // namespace UE::DerivedData::Backends
 
 #endif //WITH_HTTP_DDC_BACKEND
diff --git a/Engine/Source/Developer/DerivedDataCache/Private/ZenDerivedDataBackend.h b/Engine/Source/Developer/DerivedDataCache/Private/ZenDerivedDataBackend.h
--- a/Engine/Source/Developer/DerivedDataCache/Private/ZenDerivedDataBackend.h
+++ b/Engine/Source/Developer/DerivedDataCache/Private/ZenDerivedDataBackend.h
@@ -40,7 +40,7 @@ namespace UE::DerivedData {
 	struct FPayloadId;
 }
 
-namespace UE::DerivedData {
+namespace UE::DerivedData::Backends {
 
 /**
  * Backend for a HTTP based caching service (Zen)
@@ -89,25 +89,25 @@ public:
 
 	// ICacheStore
 
-	virtual FRequest Put(
+	virtual void Put(
 		TConstArrayView<FCacheRecord> Records,
 		FStringView Context,
-		ECachePolicy Policy = ECachePolicy::Default,
-		EPriority Priority = EPriority::Normal,
+		ECachePolicy Policy,
+		IRequestOwner& Owner,
 		FOnCachePutComplete&& OnComplete = FOnCachePutComplete()) override;
 
-	virtual FRequest Get(
+	virtual void Get(
 		TConstArrayView<FCacheKey> Keys,
 		FStringView Context,
 		ECachePolicy Policy,
-		EPriority Priority,
+		IRequestOwner& Owner,
 		FOnCacheGetComplete&& OnComplete) override;
 
-	virtual FRequest GetPayload(
+	virtual void GetPayload(
 		TConstArrayView<FCachePayloadKey> Keys,
 		FStringView Context,
 		ECachePolicy Policy,
-		EPriority Priority,
+		IRequestOwner& Owner,
 		FOnCacheGetPayloadComplete&& OnComplete) override;
 
 	virtual void CancelAll() override;
@@ -175,6 +175,6 @@ private:
 	TSet<FName> DebugMissedKeys;
 };
 
-} // namespace UE::DerivedData
+} // namespace UE::DerivedData::Backends
 
 #endif // WITH_ZEN_DDC_BACKEND
diff --git a/Engine/Source/Developer/DerivedDataCache/Public/DerivedDataBuild.h b/Engine/Source/Developer/DerivedDataCache/Public/DerivedDataBuild.h
--- a/Engine/Source/Developer/DerivedDataCache/Public/DerivedDataBuild.h
+++ b/Engine/Source/Developer/DerivedDataCache/Public/DerivedDataBuild.h
@@ -16,10 +16,12 @@ namespace UE::DerivedData { class FCacheRecord; }
 namespace UE::DerivedData { class FOptionalBuildAction; }
 namespace UE::DerivedData { class FOptionalBuildDefinition; }
 namespace UE::DerivedData { class FOptionalBuildOutput; }
+namespace UE::DerivedData { class FRequestGroup; }
 namespace UE::DerivedData { class IBuildFunctionRegistry; }
 namespace UE::DerivedData { class IBuildInputResolver; }
 namespace UE::DerivedData { class IBuildScheduler; }
 namespace UE::DerivedData { class IBuildWorkerRegistry; }
+namespace UE::DerivedData { enum class EPriority : uint8; }
 
 namespace UE::DerivedData
 {
@@ -46,6 +48,8 @@ class IBuild
 public:
 	virtual ~IBuild() = default;
 
+	virtual FRequestGroup CreateGroup(EPriority Priority) = 0;
+
 	/**
 	 * Create a build definition builder.
 	 *
diff --git a/Engine/Source/Developer/DerivedDataCache/Public/DerivedDataBuildInputResolver.h b/Engine/Source/Developer/DerivedDataCache/Public/DerivedDataBuildInputResolver.h
--- a/Engine/Source/Developer/DerivedDataCache/Public/DerivedDataBuildInputResolver.h
+++ b/Engine/Source/Developer/DerivedDataCache/Public/DerivedDataBuildInputResolver.h
@@ -89,66 +89,64 @@ public:
 	 * Asynchronous request to resolve a definition from a key.
 	 *
 	 * @param Key          The key of the definition to resolve.
+	 * @param Owner        The owner to execute the resolve request within.
 	 * @param OnResolved   A required callback invoked when resolving completes or is canceled.
 	 */
-	virtual FRequest ResolveKey(
+	virtual void ResolveKey(
 		const FBuildKey& Key,
+		IRequestOwner& Owner,
 		FOnBuildKeyResolved&& OnResolved)
 	{
 		OnResolved({Key, {}, EStatus::Error});
-		return FRequest();
 	}
 
 	/**
 	 * Asynchronous request to resolve metadata for the inputs from the definition.
 	 *
 	 * @param Definition   The definition to resolve input metadata for.
-	 * @param Priority     A priority to consider when scheduling the request. See EPriority.
+	 * @param Owner        The owner to execute the resolve request within.
 	 * @param OnResolved   A required callback invoked when resolving completes or is canceled.
 	 */
-	virtual FRequest ResolveInputMeta(
+	virtual void ResolveInputMeta(
 		const FBuildDefinition& Definition,
-		EPriority Priority,
+		IRequestOwner& Owner,
 		FOnBuildInputMetaResolved&& OnResolved)
 	{
 		OnResolved({{}, EStatus::Error});
-		return FRequest();
 	}
 
 	/**
 	 * Asynchronous request to resolve data for the inputs from the definition.
 	 *
 	 * @param Definition   The definition to resolve input data for.
-	 * @param Priority     A priority to consider when scheduling the request. See EPriority.
+	 * @param Owner        The owner to execute the resolve request within.
 	 * @param OnResolved   A required callback invoked when resolving completes or is canceled.
 	 * @param Filter       An optional predicate to filter which input keys have data resolved.
 	 */
-	virtual FRequest ResolveInputData(
+	virtual void ResolveInputData(
 		const FBuildDefinition& Definition,
-		EPriority Priority,
+		IRequestOwner& Owner,
 		FOnBuildInputDataResolved&& OnResolved,
 		FBuildInputFilter&& Filter = FBuildInputFilter())
 	{
 		OnResolved({{}, EStatus::Error});
-		return FRequest();
 	}
 
 	/**
 	 * Asynchronous request to resolve data for the inputs from the action.
 	 *
 	 * @param Action       The action to resolve input data for.
-	 * @param Priority     A priority to consider when scheduling the request. See EPriority.
+	 * @param Owner        The owner to execute the resolve request within.
 	 * @param OnResolved   A required callback invoked when resolving completes or is canceled.
 	 * @param Filter       An optional predicate to filter which input keys have data resolved.
 	 */
-	virtual FRequest ResolveInputData(
+	virtual void ResolveInputData(
 		const FBuildAction& Action,
-		EPriority Priority,
+		IRequestOwner& Owner,
 		FOnBuildInputDataResolved&& OnResolved,
 		FBuildInputFilter&& Filter = FBuildInputFilter())
 	{
 		OnResolved({{}, EStatus::Error});
-		return FRequest();
 	}
 };
 
diff --git a/Engine/Source/Developer/DerivedDataCache/Public/DerivedDataBuildJob.h b/Engine/Source/Developer/DerivedDataCache/Public/DerivedDataBuildJob.h
--- a/Engine/Source/Developer/DerivedDataCache/Public/DerivedDataBuildJob.h
+++ b/Engine/Source/Developer/DerivedDataCache/Public/DerivedDataBuildJob.h
@@ -6,8 +6,6 @@
 #include "Containers/StringFwd.h"
 #include "DerivedDataRequest.h"
 
-class FEvent;
-
 template <typename FuncType> class TUniqueFunction;
 
 namespace UE::DerivedData { class FBuildOutput; }
@@ -36,26 +34,19 @@ using FOnBuildJobComplete = TUniqueFunction<void (FBuildJobCompleteParams&& Para
  * The job depends on the build scheduler to move it through its states. That relationship allows
  * the scheduler more control over resources such as: memory, compute, storage, network.
  */
-class IBuildJob : public IRequest
+class IBuildJob
 {
 public:
 	/** Returns the name by which to identify this job for logging and profiling. */
 	virtual FStringView GetName() const = 0;
 	/** Returns the name of the function to build with, or "Unknown" if not resolved yet. */
 	virtual FStringView GetFunction() const = 0;
-	/** Returns the build policy of this job. */
-	virtual EBuildPolicy GetPolicy() const = 0;
-	/** Returns the priority of this job. */
-	virtual EPriority GetPriority() const = 0;
 
 	/** Returns the cache associated with this job. */
 	virtual ICache& GetCache() const = 0;
 	/** Returns the build system associated with this job. */
 	virtual IBuild& GetBuild() const = 0;
 
-	/** Start the build by dispatching it to the scheduler. */
-	virtual void Start(IBuildScheduler& Scheduler, EBuildPolicy Policy, EPriority Priority, FOnBuildJobComplete&& OnComplete) = 0;
-
 	/** Called by the scheduler to continue this job on the calling thread. */
 	virtual void Schedule() = 0;
 
diff --git a/Engine/Source/Developer/DerivedDataCache/Public/DerivedDataBuildPolicy.h b/Engine/Source/Developer/DerivedDataCache/Public/DerivedDataBuildPolicy.h
--- a/Engine/Source/Developer/DerivedDataCache/Public/DerivedDataBuildPolicy.h
+++ b/Engine/Source/Developer/DerivedDataCache/Public/DerivedDataBuildPolicy.h
@@ -59,11 +59,23 @@ enum class EBuildStatus : uint32
 	BuildLocal      = 1 << 0,
 	/** The build function was executed remotely. */
 	BuildRemote     = 1 << 1,
+	/** The build action and inputs were exported. */
+	BuildExport     = 1 << 2,
+
+	/** An attempt was made to execute the build function remotely. */
+	BuildTryRemote  = 1 << 3,
+	/** An attempt was made to export the build action and inputs. */
+	BuildTryExport  = 1 << 4,
 
 	/** The build made a cache query request. */
-	CacheQuery      = 1 << 2,
+	CacheQuery      = 1 << 5,
 	/** Valid build output was found in the cache. */
-	CacheQueryHit   = 1 << 3,
+	CacheQueryHit   = 1 << 6,
+
+	/** The build made a cache store request. */
+	CacheStore      = 1 << 7,
+	/** Valid build output was stored in the cache. */
+	CacheStoreHit   = 1 << 8,
 };
 
 ENUM_CLASS_FLAGS(EBuildStatus);
diff --git a/Engine/Source/Developer/DerivedDataCache/Public/DerivedDataBuildScheduler.h b/Engine/Source/Developer/DerivedDataCache/Public/DerivedDataBuildScheduler.h
--- a/Engine/Source/Developer/DerivedDataCache/Public/DerivedDataBuildScheduler.h
+++ b/Engine/Source/Developer/DerivedDataCache/Public/DerivedDataBuildScheduler.h
@@ -5,6 +5,7 @@
 #include "DerivedDataBuildJob.h"
 #include "DerivedDataBuildKey.h"
 
+namespace UE::DerivedData { class IRequestOwner; }
 namespace UE::DerivedData { struct FBuildSchedulerParams; }
 
 namespace UE::DerivedData
@@ -36,29 +37,29 @@ public:
 	virtual void UpdateJobPriority(IBuildJob* Job) {}
 
 	/** Dispatch by calling Schedule or SetOutput, either now or later. */
-	virtual void DispatchCacheQuery(IBuildJob* Job, const FBuildSchedulerParams& Params) { Job->Schedule(); }
+	virtual void DispatchCacheQuery(IBuildJob* Job, IRequestOwner& Owner, const FBuildSchedulerParams& Params) { Job->Schedule(); }
 
 	/** Dispatch by calling Schedule, either now or later. */
-	virtual void DispatchCacheStore(IBuildJob* Job, const FBuildSchedulerParams& Params) { Job->Schedule(); }
+	virtual void DispatchCacheStore(IBuildJob* Job, IRequestOwner& Owner, const FBuildSchedulerParams& Params) { Job->Schedule(); }
 
 	/** Dispatch by calling Schedule, either now or later. */
-	virtual void DispatchResolveKey(IBuildJob* Job) { Job->Schedule(); }
+	virtual void DispatchResolveKey(IBuildJob* Job, IRequestOwner& Owner) { Job->Schedule(); }
 
 	/** Dispatch by calling Schedule, either now or later. */
-	virtual void DispatchResolveInputMeta(IBuildJob* Job) { Job->Schedule(); }
+	virtual void DispatchResolveInputMeta(IBuildJob* Job, IRequestOwner& Owner) { Job->Schedule(); }
 
 	/**
 	 * Dispatch by calling Schedule, SetOutput, or SkipExecuteRemote, either now or later.
 	 *
 	 * SkipExecuteRemote is only valid to call when MissingRemoteInputsSize is non-zero.
 	 */
-	virtual void DispatchResolveInputData(IBuildJob* Job, const FBuildSchedulerParams& Params) { Job->Schedule(); }
+	virtual void DispatchResolveInputData(IBuildJob* Job, IRequestOwner& Owner, const FBuildSchedulerParams& Params) { Job->Schedule(); }
 
 	/** Dispatch by calling Schedule, SetOutput, or SkipExecuteRemote, either now or later. */
-	virtual void DispatchExecuteRemote(IBuildJob* Job, const FBuildSchedulerParams& Params) { Job->Schedule(); }
+	virtual void DispatchExecuteRemote(IBuildJob* Job, IRequestOwner& Owner, const FBuildSchedulerParams& Params) { Job->Schedule(); }
 
 	/** Dispatch by calling Schedule or SetOutput, either now or later. */
-	virtual void DispatchExecuteLocal(IBuildJob* Job, const FBuildSchedulerParams& Params) { Job->Schedule(); }
+	virtual void DispatchExecuteLocal(IBuildJob* Job, IRequestOwner& Owner, const FBuildSchedulerParams& Params) { Job->Schedule(); }
 
 	/** Set the output of the job. Always called once between BeginJob and EndJob unless canceled. */
 	virtual void SetJobOutput(IBuildJob* Job, const FBuildSchedulerParams& Params, const FBuildOutput& Output) {}
diff --git a/Engine/Source/Developer/DerivedDataCache/Public/DerivedDataBuildSession.h b/Engine/Source/Developer/DerivedDataCache/Public/DerivedDataBuildSession.h
--- a/Engine/Source/Developer/DerivedDataCache/Public/DerivedDataBuildSession.h
+++ b/Engine/Source/Developer/DerivedDataCache/Public/DerivedDataBuildSession.h
@@ -38,21 +38,21 @@ class IBuildSessionInternal
 public:
 	virtual ~IBuildSessionInternal() = default;
 	virtual FStringView GetName() const = 0;
-	virtual FRequest Build(
+	virtual void Build(
 		const FBuildDefinition& Definition,
 		EBuildPolicy Policy,
-		EPriority Priority,
+		IRequestOwner& Owner,
 		FOnBuildComplete&& OnComplete) = 0;
-	virtual FRequest BuildAction(
+	virtual void BuildAction(
 		const FBuildAction& Action,
 		const FOptionalBuildInputs& Inputs,
 		EBuildPolicy Policy,
-		EPriority Priority,
+		IRequestOwner& Owner,
 		FOnBuildActionComplete&& OnComplete) = 0;
-	virtual FRequest BuildPayload(
+	virtual void BuildPayload(
 		const FBuildPayloadKey& Payload,
 		EBuildPolicy Policy,
-		EPriority Priority,
+		IRequestOwner& Owner,
 		FOnBuildPayloadComplete&& OnComplete) = 0;
 };
 
@@ -86,16 +86,16 @@ public:
 	 *
 	 * @param Definition   The build function to execute and references to its inputs.
 	 * @param Policy       Flags to control the behavior of the request. See EBuildPolicy.
-	 * @param Priority     A priority to consider when scheduling the request. See EPriority.
+	 * @param Owner        The owner to execute the build within.
 	 * @param OnComplete   A callback invoked when the build completes or is canceled.
 	 */
-	inline FRequest Build(
+	inline void Build(
 		const FBuildDefinition& Definition,
 		EBuildPolicy Policy,
-		EPriority Priority,
+		IRequestOwner& Owner,
 		FOnBuildComplete&& OnComplete)
 	{
-		return Session->Build(Definition, Policy, Priority, MoveTemp(OnComplete));
+		Session->Build(Definition, Policy, Owner, MoveTemp(OnComplete));
 	}
 
 	/**
@@ -106,17 +106,17 @@ public:
 	 * @param Action       The build function to execute and references to its inputs.
 	 * @param Inputs       The build inputs referenced by the action, if it has any.
 	 * @param Policy       Flags to control the behavior of the request. See EBuildPolicy.
-	 * @param Priority     A priority to consider when scheduling the request. See EPriority.
+	 * @param Owner        The owner to execute the build within.
 	 * @param OnComplete   A callback invoked when the build completes or is canceled.
 	 */
-	inline FRequest BuildAction(
+	inline void BuildAction(
 		const FBuildAction& Action,
 		const FOptionalBuildInputs& Inputs,
 		EBuildPolicy Policy,
-		EPriority Priority,
+		IRequestOwner& Owner,
 		FOnBuildActionComplete&& OnComplete)
 	{
-		return Session->BuildAction(Action, Inputs, Policy, Priority, MoveTemp(OnComplete));
+		Session->BuildAction(Action, Inputs, Policy, Owner, MoveTemp(OnComplete));
 	}
 
 	/**
@@ -126,16 +126,16 @@ public:
 	 *
 	 * @param Payload      The key identifying the build definition and the payload to return.
 	 * @param Policy       Flags to control the behavior of the request. See EBuildPolicy.
-	 * @param Priority     A priority to consider when scheduling the request. See EPriority.
+	 * @param Owner        The owner to execute the build within.
 	 * @param OnComplete   A callback invoked when the build completes or is canceled.
 	 */
-	inline FRequest BuildPayload(
+	inline void BuildPayload(
 		const FBuildPayloadKey& Payload,
 		EBuildPolicy Policy,
-		EPriority Priority,
+		IRequestOwner& Owner,
 		FOnBuildPayloadComplete&& OnComplete)
 	{
-		return Session->BuildPayload(Payload, Policy, Priority, MoveTemp(OnComplete));
+		Session->BuildPayload(Payload, Policy, Owner, MoveTemp(OnComplete));
 	}
 
 private:
diff --git a/Engine/Source/Developer/DerivedDataCache/Public/DerivedDataBuildWorker.h b/Engine/Source/Developer/DerivedDataCache/Public/DerivedDataBuildWorker.h
--- a/Engine/Source/Developer/DerivedDataCache/Public/DerivedDataBuildWorker.h
+++ b/Engine/Source/Developer/DerivedDataCache/Public/DerivedDataBuildWorker.h
@@ -20,6 +20,7 @@ template <typename FuncType> class TUniqueFunction;
 namespace UE::DerivedData { class FBuildAction; }
 namespace UE::DerivedData { class FOptionalBuildInputs; }
 namespace UE::DerivedData { class FOptionalBuildOutput; }
+namespace UE::DerivedData { class IBuild; }
 namespace UE::DerivedData { struct FBuildWorkerActionCompleteParams; }
 namespace UE::DerivedData { struct FBuildWorkerFileDataCompleteParams; }
 namespace UE::DerivedData { enum class EBuildPolicy : uint8; }
@@ -38,7 +39,7 @@ public:
 	virtual FStringView GetPath() const = 0;
 	virtual FStringView GetHostPlatform() const = 0;
 	virtual FGuid GetBuildSystemVersion() const = 0;
-	virtual FRequest FindFileData(TConstArrayView<FIoHash> RawHashes, EPriority Priority, FOnBuildWorkerFileDataComplete&& OnComplete) const = 0;
+	virtual void FindFileData(TConstArrayView<FIoHash> RawHashes, IRequestOwner& Owner, FOnBuildWorkerFileDataComplete&& OnComplete) const = 0;
 	virtual void IterateFunctions(TFunctionRef<void (FStringView Name, const FGuid& Version)> Visitor) const = 0;
 	virtual void IterateFiles(TFunctionRef<void (FStringView Path, const FIoHash& RawHash, uint64 RawSize)> Visitor) const = 0;
 	virtual void IterateExecutables(TFunctionRef<void (FStringView Path, const FIoHash& RawHash, uint64 RawSize)> Visitor) const = 0;
@@ -65,7 +66,7 @@ public:
 	virtual ~IBuildWorkerFactory() = default;
 
 	virtual void Build(FBuildWorkerBuilder& Builder) = 0;
-	virtual FRequest FindFileData(TConstArrayView<FIoHash> RawHashes, EPriority Priority, FOnBuildWorkerFileDataComplete&& OnComplete) = 0;
+	virtual void FindFileData(TConstArrayView<FIoHash> RawHashes, IRequestOwner& Owner, FOnBuildWorkerFileDataComplete&& OnComplete) = 0;
 
 	/** Returns the name of the build worker factory modular feature. */
 	static FName GetFeatureName()
@@ -79,12 +80,13 @@ class IBuildWorkerExecutor : public IModularFeature
 public:
 	virtual ~IBuildWorkerExecutor() = default;
 
-	virtual FRequest BuildAction(
+	virtual void BuildAction(
 		const FBuildAction& Action,
 		const FOptionalBuildInputs& Inputs,
 		const FBuildWorker& Worker,
+		IBuild& BuildSystem,
 		EBuildPolicy Policy,
-		EPriority Priority,
+		IRequestOwner& Owner,
 		FOnBuildWorkerActionComplete&& OnComplete) = 0;
 
 	virtual TConstArrayView<FStringView> GetHostPlatforms() const = 0;
diff --git a/Engine/Source/Developer/DerivedDataCache/Public/DerivedDataCache.h b/Engine/Source/Developer/DerivedDataCache/Public/DerivedDataCache.h
--- a/Engine/Source/Developer/DerivedDataCache/Public/DerivedDataCache.h
+++ b/Engine/Source/Developer/DerivedDataCache/Public/DerivedDataCache.h
@@ -18,6 +18,7 @@ class FCacheRecord;
 class FCacheRecordBuilder;
 class FOptionalCacheRecord;
 class FPayload;
+class FRequestGroup;
 struct FCacheGetCompleteParams;
 struct FCacheGetPayloadCompleteParams;
 struct FCachePutCompleteParams;
@@ -102,6 +103,8 @@ class ICacheFactory
 public:
 	virtual ~ICacheFactory() = default;
 
+	virtual FRequestGroup CreateGroup(EPriority Priority) = 0;
+
 	/**
 	 * Create a cache bucket from a name.
 	 *
@@ -155,14 +158,14 @@ public:
 	 * @param Records      The cache records to store. Must have a key.
 	 * @param Context      A description of the request. An object path is typically sufficient.
 	 * @param Policy       Flags to control the behavior of the request. See ECachePolicy.
-	 * @param Priority     A priority to consider when scheduling the request. See EPriority.
+	 * @param Owner        The owner to execute the build within.
 	 * @param OnComplete   A callback invoked for every record as it completes or is canceled.
 	 */
-	virtual FRequest Put(
+	virtual void Put(
 		TConstArrayView<FCacheRecord> Records,
 		FStringView Context,
-		ECachePolicy Policy = ECachePolicy::Default,
-		EPriority Priority = EPriority::Normal,
+		ECachePolicy Policy,
+		IRequestOwner& Owner,
 		FOnCachePutComplete&& OnComplete = FOnCachePutComplete()) = 0;
 
 	/**
@@ -174,14 +177,14 @@ public:
 	 * @param Keys         The keys identifying the cache records to query.
 	 * @param Context      A description of the request. An object path is typically sufficient.
 	 * @param Policy       Flags to control the behavior of the request. See ECachePolicy.
-	 * @param Priority     A priority to consider when scheduling the request. See EPriority.
+	 * @param Owner        The owner to execute the build within.
 	 * @param OnComplete   A callback invoked for every key as it completes or is canceled.
 	 */
-	virtual FRequest Get(
+	virtual void Get(
 		TConstArrayView<FCacheKey> Keys,
 		FStringView Context,
 		ECachePolicy Policy,
-		EPriority Priority,
+		IRequestOwner& Owner,
 		FOnCacheGetComplete&& OnComplete) = 0;
 
 	/**
@@ -193,14 +196,14 @@ public:
 	 * @param Keys         The keys identifying the cache record payloads to query.
 	 * @param Context      A description of the request. An object path is typically sufficient.
 	 * @param Policy       Flags to control the behavior of the request. See ECachePolicy.
-	 * @param Priority     A priority to consider when scheduling the request. See EPriority.
+	 * @param Owner        The owner to execute the build within.
 	 * @param OnComplete   A callback invoked for every key as it completes or is canceled.
 	 */
-	virtual FRequest GetPayload(
+	virtual void GetPayload(
 		TConstArrayView<FCachePayloadKey> Keys,
 		FStringView Context,
 		ECachePolicy Policy,
-		EPriority Priority,
+		IRequestOwner& Owner,
 		FOnCacheGetPayloadComplete&& OnComplete) = 0;
 
 	/**
diff --git a/Engine/Source/Developer/DerivedDataCache/Public/DerivedDataCacheRecord.h b/Engine/Source/Developer/DerivedDataCache/Public/DerivedDataCacheRecord.h
--- a/Engine/Source/Developer/DerivedDataCache/Public/DerivedDataCacheRecord.h
+++ b/Engine/Source/Developer/DerivedDataCache/Public/DerivedDataCacheRecord.h
@@ -51,7 +51,7 @@ public:
 	virtual FPayloadId AddAttachment(const FSharedBuffer& Buffer, const FPayloadId& Id) = 0;
 	virtual FPayloadId AddAttachment(const FPayload& Payload) = 0;
 	virtual FCacheRecord Build() = 0;
-	virtual FRequest BuildAsync(FOnCacheRecordComplete&& OnComplete, EPriority Priority) = 0;
+	virtual void BuildAsync(IRequestOwner& Owner, FOnCacheRecordComplete&& OnComplete) = 0;
 };
 
 FCacheRecordBuilder CreateCacheRecordBuilder(ICacheRecordBuilderInternal* RecordBuilder);
@@ -215,9 +215,9 @@ public:
 	 * Prefer Build() when the value and attachments are added by payload, as compression is already
 	 * complete and BuildAsync() will complete immediately in that case.
 	 */
-	inline FRequest BuildAsync(FOnCacheRecordComplete&& OnComplete, EPriority Priority)
+	inline void BuildAsync(IRequestOwner& Owner, FOnCacheRecordComplete&& OnComplete)
 	{
-		return RecordBuilder.Release()->BuildAsync(MoveTemp(OnComplete), Priority);
+		return RecordBuilder.Release()->BuildAsync(Owner, MoveTemp(OnComplete));
 	}
 
 private:
diff --git a/Engine/Source/Developer/DerivedDataCache/Public/DerivedDataRequest.h b/Engine/Source/Developer/DerivedDataCache/Public/DerivedDataRequest.h
--- a/Engine/Source/Developer/DerivedDataCache/Public/DerivedDataRequest.h
+++ b/Engine/Source/Developer/DerivedDataCache/Public/DerivedDataRequest.h
@@ -4,14 +4,18 @@
 
 #include "CoreTypes.h"
 #include "Templates/RefCounting.h"
+#include "Templates/UniquePtr.h"
 #include "Templates/UnrealTemplate.h"
-#include <type_traits>
+#include <atomic>
+
+namespace UE::DerivedData { class FRequestGroup; }
+namespace UE::DerivedData { enum class EPriority : uint8; }
+namespace UE::DerivedData::Private { class FRequestGroupBase; }
+namespace UE::DerivedData::Private { FRequestGroup CreateRequestGroup(EPriority Priority); }
 
 namespace UE::DerivedData
 {
 
-///////////////////////////////////////////////////////////////////////////////////////////////////////////////////////
-
 /** Priority for scheduling a request. */
 enum class EPriority : uint8
 {
@@ -48,8 +52,6 @@ enum class EPriority : uint8
 	Blocking,
 };
 
-///////////////////////////////////////////////////////////////////////////////////////////////////////////////////////
-
 /** Status of a request that has completed. */
 enum class EStatus : uint8
 {
@@ -61,9 +63,19 @@ enum class EStatus : uint8
 	Canceled,
 };
 
-///////////////////////////////////////////////////////////////////////////////////////////////////////////////////////
-
-/** Interface to a request. Use via FRequest or TRequest to automate reference counting. */
+/**
+ * Interface to an asynchronous request that can be prioritized or canceled.
+ *
+ * Use IRequestOwner, typically FRequestGroup, to reference requests between its Begin and End.
+ *
+ * Requests typically invoke a callback on completion, and must not return from Cancel or Wait on
+ * any thread until any associated callbacks have finished executing. This property is crucial to
+ * allowing requests to be chained or nested by creating new requests from within callbacks.
+ *
+ * @see FRequestBase
+ * @see FRequestGroup
+ * @see IRequestOwner
+ */
 class IRequest
 {
 public:
@@ -79,24 +91,19 @@ public:
 	/**
 	 * Cancel the request and invoke any associated callback.
 	 *
-	 * Does not return until any callback for the request has finished executing.
+	 * Must not return until any callback for the request has finished executing, and the request
+	 * has been removed from its owner by calling IRequestOwner::End.
 	 */
 	virtual void Cancel() = 0;
 
 	/**
 	 * Block the calling thread until the request is complete.
 	 *
-	 * Does not return until any callback for the request has finished executing.
+	 * Must not return until any callback for the request has finished executing, and the request
+	 * has been removed from its owner by calling IRequestOwner::End.
 	 */
 	virtual void Wait() = 0;
 
-	/**
-	 * Poll the request and return true if complete and false otherwise.
-	 *
-	 * Does not return true until any callback for the request has finished executing.
-	 */
-	virtual bool Poll() = 0;
-
 	/** Add a reference to the request. */
 	virtual void AddRef() const = 0;
 
@@ -104,120 +111,205 @@ public:
 	virtual void Release() const = 0;
 };
 
-///////////////////////////////////////////////////////////////////////////////////////////////////////////////////////
-
-/** Handle to a request that manages its lifetime with reference counting. */
-template <typename RequestType>
-class TRequest
+/** An asynchronous request base type that provides reference counting. */
+class FRequestBase : public IRequest
 {
 public:
-	/** Construct a null request. */
-	TRequest() = default;
-
-	/** Construct a request from a raw pointer. */
-	inline explicit TRequest(RequestType* InRequest) : Request(InRequest) {}
-
-	/** Construct a request from a request of a compatible type. */
-	template <typename OtherRequestType, decltype(ImplicitConv<RequestType*>(DeclVal<OtherRequestType*>()))* = nullptr>
-	inline TRequest(const TRequest<OtherRequestType>& InRequest) : Request(InRequest.Request) {}
-	template <typename OtherRequestType, decltype(ImplicitConv<RequestType*>(DeclVal<OtherRequestType*>()))* = nullptr>
-	inline TRequest(TRequest<OtherRequestType>&& InRequest) : Request(MoveTemp(InRequest.Request)) {}
-
-	/** Assign from a request of a compatible type. */
-	template <typename OtherRequestType, decltype(ImplicitConv<RequestType*>(DeclVal<OtherRequestType*>()))* = nullptr>
-	inline TRequest& operator=(const TRequest<OtherRequestType>& InRequest) { Request = InRequest.Request; return *this; }
-	template <typename OtherRequestType, decltype(ImplicitConv<RequestType*>(DeclVal<OtherRequestType*>()))* = nullptr>
-	inline TRequest& operator=(TRequest<OtherRequestType>&& InRequest) { Request = MoveTemp(InRequest.Request); return *this; }
-
-	/** Reset this to null. */
-	inline void Reset() { Request.SafeRelease(); }
-
-	/** Whether the request is null. */
-	inline bool IsNull() const { return !Request; }
-	/** Whether the request is not null. */
-	inline bool IsValid() const { return !IsNull(); }
-	/** Whether the request is not null. */
-	inline explicit operator bool() const { return IsValid(); }
+	FRequestBase() = default;
+	~FRequestBase() override = default;
+	FRequestBase(const FRequestBase&) = delete;
+	FRequestBase& operator=(const FRequestBase&) = delete;
 
-	/** Access the referenced request. */
-	inline RequestType* Get() const { return Request; }
-	inline std::conditional_t<!std::is_same_v<IRequest, RequestType>, RequestType*, void> operator->() const { return Request; }
-	inline std::conditional_t<!std::is_same_v<IRequest, RequestType>, RequestType&, void> operator*() const { return *Request; }
+	inline void AddRef() const final
+	{
+		ReferenceCount.fetch_add(1, std::memory_order_relaxed);
+	}
 
-	/**
-	 * Set the priority of the request.
-	 *
-	 * @param Priority   The priority, which may be higher or lower than the existing priority.
-	 */
-	inline void SetPriority(EPriority Priority)
+	inline void Release() const final
 	{
-		if (IRequest* LocalRequest = Request)
+		if (ReferenceCount.fetch_sub(1, std::memory_order_acq_rel) == 1)
 		{
-			LocalRequest->SetPriority(Priority);
+			delete this;
 		}
 	}
 
+private:
+	mutable std::atomic<uint32> ReferenceCount{0};
+};
+
+/**
+ * A request owner manages requests throughout their execution.
+ *
+ * Requests are expected to call Begin when created, and End when finished, with their completion
+ * callback and arguments passed to End when possible. If a request needs to invoke more than one
+ * callback, have a request barrier in scope when executing callbacks to ensure that new requests
+ * are prioritized properly when the priority is being changed from another thread.
+ *
+ * The owner is responsible for keeping itself alive while it has active requests or barriers.
+ *
+ * @see FRequestBarrier
+ * @see FRequestGroup
+ */
+class IRequestOwner
+{
+public:
 	/**
-	 * Cancel the request and invoke any associated callback.
+	 * Begin tracking for the request.
 	 *
-	 * Does not return until any callback for the request has finished executing.
+	 * The owner will hold a reference to the request until End is called, and forward any cancel
+	 * operation or priority change to the request. Begin must return before End is called.
 	 */
-	inline void Cancel() const
-	{
-		if (IRequest* LocalRequest = Request)
-		{
-			LocalRequest->Cancel();
-		}
-	}
+	virtual void Begin(IRequest* Request) = 0;
 
 	/**
-	 * Block the calling thread until the request is complete.
+	 * End tracking for the request.
+	 *
+	 * Requests with a single completion callback should use the callback overload of End.
 	 *
-	 * Avoid waiting if possible and use callbacks to manage control flow instead.
-	 * Does not return until any callback for the request has finished executing.
+	 * @return The reference to the request that was held by the owner.
 	 */
-	inline void Wait() const
-	{
-		if (IRequest* LocalRequest = Request)
-		{
-			LocalRequest->Wait();
-		}
-	}
+	virtual TRefCountPtr<IRequest> End(IRequest* Request) = 0;
 
 	/**
-	 * Poll the request and return true if complete and false otherwise.
+	 * End tracking of the request.
 	 *
-	 * Avoid polling if possible and use callbacks to manage control flow instead.
-	 * Does not return true until any callback for the request has finished executing.
+	 * Begins a barrier, ends the request, invokes the callback, then ends the barrier. Keeps its
+	 * reference to the request until the callback has returned.
+	 *
+	 * @return The reference to the request that was held by the owner.
 	 */
-	inline bool Poll() const
+	template <typename CallbackType, typename... CallbackArgTypes>
+	inline TRefCountPtr<IRequest> End(IRequest* Request, CallbackType&& Callback, CallbackArgTypes&&... CallbackArgs);
+
+	/** See FRequestBarrier. */
+	virtual void BeginBarrier() = 0;
+	virtual void EndBarrier() = 0;
+
+	/** Returns the priority that new requests are expected to inherit. */
+	virtual EPriority GetPriority() const = 0;
+	/** Returns whether the owner has been canceled, which new requests are expected to check. */
+	virtual bool IsCanceled() const = 0;
+};
+
+/**
+ * A request barrier is expected to be used when an owner may have new requests added to it.
+ *
+ * An owner may not consider its execution to be complete in the presence of a barrier, and needs
+ * to take note of priority changes that occur while within a barrier, since newly-added requests
+ * may have been created with the previous priority value.
+ */
+class FRequestBarrier
+{
+public:
+	inline explicit FRequestBarrier(IRequestOwner& InOwner)
+		: Owner(InOwner)
 	{
-		if (IRequest* LocalRequest = Request)
-		{
-			return LocalRequest->Poll();
-		}
-		return true;
+		Owner.BeginBarrier();
+	}
+
+	inline ~FRequestBarrier()
+	{
+		Owner.EndBarrier();
 	}
 
-	template <typename OtherRequestType, decltype(ImplicitConv<IRequest*>(DeclVal<OtherRequestType*>()))* = nullptr>
-	inline bool operator==(const TRequest<OtherRequestType>& Other) const { return Request == Other.Request; }
-	template <typename OtherRequestType, decltype(ImplicitConv<IRequest*>(DeclVal<OtherRequestType*>()))* = nullptr>
-	inline bool operator!=(const TRequest<OtherRequestType>& Other) const { return Request != Other.Request; }
-	template <typename OtherRequestType, decltype(ImplicitConv<IRequest*>(DeclVal<OtherRequestType*>()))* = nullptr>
-	inline bool operator<(const TRequest<OtherRequestType>& Other) const { return Request < Other.Request; }
+private:
+	IRequestOwner& Owner;
+};
 
-	friend inline uint32 GetTypeHash(const TRequest& Other) { return ::GetTypeHash(Other.Request); }
+/** Private base for a request group, which is both an owner and a request. */
+class Private::FRequestGroupBase : public IRequestOwner, public FRequestBase
+{
+public:
+	virtual bool Poll() const = 0;
+	virtual void KeepAlive() = 0;
+	virtual void Destroy() = 0;
+};
 
-	template <typename OtherRequestType>
-	friend class TRequest;
+/**
+ * A request group is both a request owner and a request.
+ *
+ * Request groups may be moved but not copied, and cancel any outstanding requests on destruction
+ * unless KeepAlive has been called.
+ *
+ * Create with ICache::CreateGroup or IBuild::CreateGroup.
+ *
+ * @see IRequestOwner
+ */
+class FRequestGroup
+{
+public:
+	FRequestGroup(FRequestGroup&&) = default;
+	FRequestGroup& operator=(FRequestGroup&&) = default;
+
+	FRequestGroup(const FRequestGroup&) = delete;
+	FRequestGroup& operator=(const FRequestGroup&) = delete;
+
+	/** Keep requests in the group alive until complete, even after destruction of the group. */
+	inline void KeepAlive() { Owner->KeepAlive(); }
+
+	/** Returns the priority that new requests are expected to inherit. */
+	inline EPriority GetPriority() const { return Owner->GetPriority(); }
+	/** Set the priority of active and future requests in the group. */
+	inline void SetPriority(EPriority Priority) { Owner->SetPriority(Priority); }
+	/** Cancel any active and future requests in the group. */
+	inline void Cancel() { Owner->Cancel(); }
+	/** Wait for any active and future requests or barriers in the group. */
+	inline void Wait() { Owner->Wait(); }
+	/** Poll whether the group has any active requests or barriers. */
+	inline bool Poll() const { return Owner->Poll(); }
+
+	/** Access the group as a request owner. */
+	inline operator IRequestOwner&() { return *Owner; }
+	/** Access the group as a request. */
+	inline operator IRequest*() { return Owner.Get(); }
 
 private:
-	TRefCountPtr<RequestType> Request;
+	friend class FOptionalRequestGroup;
+	friend FRequestGroup Private::CreateRequestGroup(EPriority Priority);
+
+	/** Construct a request group. Use CreateGroup() on ICache or IBuild. */
+	inline explicit FRequestGroup(Private::FRequestGroupBase* InOwner)
+		: Owner(InOwner)
+	{
+	}
+
+	struct FDeleteOwner { void operator()(Private::FRequestGroupBase* O) const { if (O) { O->Destroy(); } } };
+	TUniquePtr<Private::FRequestGroupBase, FDeleteOwner> Owner;
 };
 
-/** Handle to a request that manages its lifetime with reference counting. */
-using FRequest = TRequest<IRequest>;
+/**
+ * A request group that can be null.
+ *
+ * @see FRequestGroup
+ */
+class FOptionalRequestGroup : private FRequestGroup
+{
+public:
+	inline FOptionalRequestGroup() : FRequestGroup(nullptr) {}
+
+	inline FOptionalRequestGroup(FRequestGroup&& InDefinition) : FRequestGroup(MoveTemp(InDefinition)) {}
+	inline FOptionalRequestGroup& operator=(FRequestGroup&& InDefinition) { FRequestGroup::operator=(MoveTemp(InDefinition)); return *this; }
+
+	/** Returns the request group. The caller must check for null before using this accessor. */
+	inline FRequestGroup& Get() & { return *this; }
+	inline FRequestGroup&& Get() && { return MoveTemp(*this); }
+	inline const FRequestGroup& Get() const & { return *this; }
+
+	inline bool IsNull() const { return !IsValid(); }
+	inline bool IsValid() const { return Owner.IsValid(); }
+	inline explicit operator bool() { return IsValid(); }
+	inline explicit operator bool() const { return IsValid(); }
+
+	inline void Reset() { *this = FOptionalRequestGroup(); }
+};
 
-///////////////////////////////////////////////////////////////////////////////////////////////////////////////////////
+template <typename CallbackType, typename... CallbackArgTypes>
+TRefCountPtr<IRequest> IRequestOwner::End(IRequest* Request, CallbackType&& Callback, CallbackArgTypes&&... CallbackArgs)
+{
+	FRequestBarrier Barrier(*this);
+	TRefCountPtr<IRequest> RequestRef = End(Request);
+	Invoke(Forward<CallbackType>(Callback), Forward<CallbackArgTypes>(CallbackArgs)...);
+	return RequestRef;
+}
 
 } // UE::DerivedData
diff --git a/Engine/Source/Developer/DesktopPlatform/Private/TargetReceiptBuildWorker.cpp b/Engine/Source/Developer/DesktopPlatform/Private/TargetReceiptBuildWorker.cpp
--- a/Engine/Source/Developer/DesktopPlatform/Private/TargetReceiptBuildWorker.cpp
+++ b/Engine/Source/Developer/DesktopPlatform/Private/TargetReceiptBuildWorker.cpp
@@ -57,7 +57,7 @@ private:
 	FTargetReceiptBuildWorker* Worker;
 };
 
-class FTargetReceiptBuildWorkerFactory : public UE::DerivedData::IBuildWorkerFactory
+class FTargetReceiptBuildWorkerFactory final : public UE::DerivedData::IBuildWorkerFactory
 {
 public:
 	const FTargetReceiptBuildWorker* GetWorker()
@@ -65,7 +65,7 @@ public:
 		return (FTargetReceiptBuildWorker*)((uintptr_t)this - offsetof(class FTargetReceiptBuildWorker, InternalFactory));
 	}
 
-	void Build(UE::DerivedData::FBuildWorkerBuilder& Builder)
+	void Build(UE::DerivedData::FBuildWorkerBuilder& Builder) final
 	{
 		const FTargetReceiptBuildWorker* Worker = GetWorker();
 		Builder.SetName(Worker->Name);
@@ -90,7 +90,7 @@ public:
 		}
 	}
 
-	UE::DerivedData::FRequest FindFileData(TConstArrayView<FIoHash> RawHashes, UE::DerivedData::EPriority Priority, UE::DerivedData::FOnBuildWorkerFileDataComplete&& OnComplete)
+	void FindFileData(TConstArrayView<FIoHash> RawHashes, UE::DerivedData::IRequestOwner& Owner, UE::DerivedData::FOnBuildWorkerFileDataComplete&& OnComplete) final
 	{
 		if (OnComplete)
 		{
@@ -134,8 +134,6 @@ public:
 				OnComplete(MoveTemp(CompleteParams));
 			}
 		}
-
-		return UE::DerivedData::FRequest();
 	}
 };
 
diff --git a/Engine/Source/Editor/UnrealEd/Private/DerivedDataBuildRemoteExecutor.cpp b/Engine/Source/Editor/UnrealEd/Private/DerivedDataBuildRemoteExecutor.cpp
--- a/Engine/Source/Editor/UnrealEd/Private/DerivedDataBuildRemoteExecutor.cpp
+++ b/Engine/Source/Editor/UnrealEd/Private/DerivedDataBuildRemoteExecutor.cpp
@@ -40,7 +40,7 @@ DEFINE_LOG_CATEGORY_STATIC(LogDerivedDataBuildRemoteExecutor, Log, All);
 
 class FRemoteBuildWorkerExecutor;
 
-class FRemoteBuildExecutionRequest final : public IRequest, public FThreadSafeRefCountedObject
+class FRemoteBuildExecutionRequest final : public FRequestBase
 {
 public:
 	FRemoteBuildExecutionRequest(
@@ -48,37 +48,27 @@ public:
 		const FBuildAction& Action,
 		const FOptionalBuildInputs& Inputs,
 		const FBuildWorker& Worker,
+		IBuild& BuildSystem,
 		EBuildPolicy Policy,
-		EPriority Priority,
+		IRequestOwner& Owner,
 		FOnBuildWorkerActionComplete&& OnComplete);
 
-	virtual ~FRemoteBuildExecutionRequest();
+	~FRemoteBuildExecutionRequest() final;
 
 	// IRequest interface
-	virtual void SetPriority(EPriority Priority) override
+	void SetPriority(EPriority Priority) final
 	{
 	}
 
-	virtual void Cancel() override
+	void Cancel() final
 	{
 		bCancelPending.store(true, std::memory_order_relaxed);
 		Wait();
 	}
-	virtual void Wait() override
-	{
-		CompletionEvent->Wait(MAX_uint32);
-	}
-	virtual bool Poll() override
-	{
-		return CompletionEvent->Wait(0);
-	}
-	virtual void AddRef() const override
-	{
-		FThreadSafeRefCountedObject::AddRef();
-	}
-	virtual void Release() const override
+
+	void Wait() final
 	{
-		FThreadSafeRefCountedObject::Release();
+		CompletionEvent->Wait();
 	}
 
 private:
@@ -144,8 +134,9 @@ private:
 		const FBuildAction& BuildAction;
 		const FOptionalBuildInputs& BuildInputs;
 		const FBuildWorker& BuildWorker;
+		IBuild& BuildSystem;
+		IRequestOwner& Owner;
 		EBuildPolicy BuildPolicy;
-		EPriority BuildPriority;
 
 		// Unordered arrays that are indexed into
 		TArray<FMerkleTreeDirectoryBuilder> Directories;
@@ -262,12 +253,13 @@ public:
 		}
 	}
 
-	FRequest BuildAction(
+	void BuildAction(
 		const FBuildAction& Action,
 		const FOptionalBuildInputs& Inputs,
 		const FBuildWorker& Worker,
+		IBuild& BuildSystem,
 		EBuildPolicy Policy,
-		EPriority Priority,
+		IRequestOwner& Owner,
 		FOnBuildWorkerActionComplete&& OnComplete) final
 	{
 		{
@@ -296,17 +288,17 @@ public:
 			if (!LimitingHeuristics.PassesPreResolveRequirements(TotalInputSize, TotalMissingInputSize))
 			{
 				OnComplete({Action.GetKey(), {}, {}, EStatus::Error});
-				return FRequest();
+				return;
 			}
 
 			if (!MissingInputViews.IsEmpty())
 			{
 				OnComplete({Action.GetKey(), {}, MissingInputViews, EStatus::Ok});
-				return FRequest();
+				return;
 			}
 		}
 
-		return FRequest(new FRemoteBuildExecutionRequest(*this, Action, Inputs, Worker, Policy, Priority, MoveTemp(OnComplete)));
+		new FRemoteBuildExecutionRequest(*this, Action, Inputs, Worker, BuildSystem, Policy, Owner, MoveTemp(OnComplete));
 	}
 
 	TConstArrayView<FStringView> GetHostPlatforms() const final
@@ -478,15 +470,17 @@ FRemoteBuildExecutionRequest::FRemoteBuildExecutionRequest(
 		const FBuildAction& Action,
 		const FOptionalBuildInputs& Inputs,
 		const FBuildWorker& Worker,
+		IBuild& BuildSystem,
 		EBuildPolicy Policy,
-		EPriority Priority,
+		IRequestOwner& Owner,
 		FOnBuildWorkerActionComplete&& OnComplete)
-: State {Action, Inputs, Worker, Policy, Priority}
+: State{Action, Inputs, Worker, BuildSystem, Owner, Policy}
 , CompletionCallback(MoveTemp(OnComplete))
 , Executor(InExecutor)
 , bCancelPending(false)
 , bHeuristicBuildStarted(false)
 {
+	Owner.Begin(this);
 	DetermineMissingBlobsAsync()
 		.Next([this] (TPair<FStatus, FFindMissingBlobsResponse>&& Result) { OnMissingBlobsDetermined(MoveTemp(Result.Value)); });
 }
@@ -594,7 +588,8 @@ void FRemoteBuildExecutionRequest::BuildMerkleTreeNodes()
 			WorkerFileMeta.Emplace(Path, false);
 		});
 
-	State.BuildWorker.FindFileData(WorkerFileHashes, EPriority::Normal,
+	FRequestGroup BlockingGroup = State.BuildSystem.CreateGroup(EPriority::Blocking);
+	State.BuildWorker.FindFileData(WorkerFileHashes, BlockingGroup,
 		[this, &WorkerFileMeta] (FBuildWorkerFileDataCompleteParams&& Params)
 		{
 			uint32 MetaIndex = 0;
@@ -605,7 +600,8 @@ void FRemoteBuildExecutionRequest::BuildMerkleTreeNodes()
 				AddMerkleTreeFile(Meta.Key, FIoHash::HashBuffer(DecompressedComposite), DecompressedComposite.GetSize(), Meta.Value, EFileType::Worker, Buffer.DecompressToComposite());
 				++MetaIndex;
 			}
-		}).Wait();
+		});
+	BlockingGroup.Wait();
 
 	State.BuildAction.IterateInputs([this] (FStringView Key, const FIoHash& RawHash, uint64 RawSize)
 		{
@@ -754,7 +750,7 @@ FOptionalBuildOutput FRemoteBuildExecutionRequest::ComposeBuildOutput(FBatchRead
 				return FOptionalBuildOutput();
 			}
 
-			RemoteBuildOutput = GetDerivedDataBuildRef().LoadOutput(State.BuildAction.GetName(), State.BuildAction.GetFunction(), FCbObject(BuildOutputBuffer));
+			RemoteBuildOutput = State.BuildSystem.LoadOutput(State.BuildAction.GetName(), State.BuildAction.GetFunction(), FCbObject(BuildOutputBuffer));
 		}
 		else
 		{
@@ -769,7 +765,7 @@ FOptionalBuildOutput FRemoteBuildExecutionRequest::ComposeBuildOutput(FBatchRead
 		return FOptionalBuildOutput();
 	}
 
-	FBuildOutputBuilder OutputBuilder = GetDerivedDataBuildRef().CreateOutput(State.BuildAction.GetName(), State.BuildAction.GetFunction());
+	FBuildOutputBuilder OutputBuilder = State.BuildSystem.CreateOutput(State.BuildAction.GetName(), State.BuildAction.GetFunction());
 	
 	RemoteBuildOutput.Get().IterateDiagnostics( [&OutputBuilder](const FBuildDiagnostic& Diagnostic)
 	{
@@ -803,9 +799,11 @@ bool FRemoteBuildExecutionRequest::ProcessCancellation()
 {
 	if (bCancelPending.load(std::memory_order_relaxed))
 	{
-		TRefCountPtr<FRemoteBuildExecutionRequest> SelfRef(this);
-		CompletionCallback({State.BuildAction.GetKey(), {}, {}, EStatus::Canceled});
-		CompletionEvent->Trigger();
+		State.Owner.End(this, [this]
+		{
+			CompletionCallback({State.BuildAction.GetKey(), {}, {}, EStatus::Canceled});
+			CompletionEvent->Trigger();
+		});
 		return true;
 	}
 	return false;
@@ -848,7 +846,8 @@ void FRemoteBuildExecutionRequest::LoadMissingWorkerFileBlobsAsync()
 		}
 	}
 
-	State.BuildWorker.FindFileData(WorkerFileHashes, EPriority::Normal,
+	FRequestGroup BlockingGroup = State.BuildSystem.CreateGroup(EPriority::Blocking);
+	State.BuildWorker.FindFileData(WorkerFileHashes, BlockingGroup,
 		[this, &WorkerFileMapping] (FBuildWorkerFileDataCompleteParams&& Params)
 		{
 			uint32 MetaIndex = 0;
@@ -863,7 +862,8 @@ void FRemoteBuildExecutionRequest::LoadMissingWorkerFileBlobsAsync()
 					State.Files[FileIndex].ContentBytes = UncompressedWorkerFile;
 				}
 			}
-		}).Wait();
+		});
+	BlockingGroup.Wait();
 }
 
 TFuture<TPair<FStatus, FBatchUpdateBlobsResponse>> FRemoteBuildExecutionRequest::UploadMissingBlobsAsync()
@@ -965,17 +965,21 @@ void FRemoteBuildExecutionRequest::OnMissingBlobsDetermined(FFindMissingBlobsRes
 	GatherMissingInputFileBlobs(MissingInputViews);
 	if (!MissingInputViews.IsEmpty())
 	{
-		TRefCountPtr<FRemoteBuildExecutionRequest> SelfRef(this);
-		CompletionCallback({State.BuildAction.GetKey(), {}, MissingInputViews, EStatus::Ok});
-		CompletionEvent->Trigger();
+		State.Owner.End(this, [this, &MissingInputViews]
+		{
+			CompletionCallback({State.BuildAction.GetKey(), {}, MissingInputViews, EStatus::Ok});
+			CompletionEvent->Trigger();
+		});
 		return;
 	}
 
 	if (!Executor.LimitingHeuristics.TryStartNewBuild(Executor.Stats))
 	{
-		TRefCountPtr<FRemoteBuildExecutionRequest> SelfRef(this);
-		CompletionCallback({State.BuildAction.GetKey(), {}, {}, EStatus::Error});
-		CompletionEvent->Trigger();
+		State.Owner.End(this, [this]
+		{
+			CompletionCallback({State.BuildAction.GetKey(), {}, {}, EStatus::Error});
+			CompletionEvent->Trigger();
+		});
 		return;
 	}
 
@@ -1006,9 +1010,11 @@ void FRemoteBuildExecutionRequest::OnMissingBlobsUploaded(const FBatchUpdateBlob
 	UE_LOG(LogDerivedDataBuildRemoteExecutor, Verbose, TEXT("Uploaded %d data blobs for remote execution."), State.BatchUpdateBlobsRequest.Requests.Num());
 	if (!ValidateUploadSuccess(Result))
 	{
-		TRefCountPtr<FRemoteBuildExecutionRequest> SelfRef(this);
-		CompletionCallback({State.BuildAction.GetKey(), {}, {}, EStatus::Error});
-		CompletionEvent->Trigger();
+		State.Owner.End(this, [this]
+		{
+			CompletionCallback({State.BuildAction.GetKey(), {}, {}, EStatus::Error});
+			CompletionEvent->Trigger();
+		});
 		return;
 	}
 
@@ -1035,10 +1041,12 @@ void FRemoteBuildExecutionRequest::OnExecutionCompleted(FExecuteResponse&& Resul
 	}
 	else
 	{
-		TRefCountPtr<FRemoteBuildExecutionRequest> SelfRef(this);
 		UE_LOG(LogDerivedDataBuildRemoteExecutor, Warning, TEXT("Remote execution system error: Failed to execute build operation!"));
-		CompletionCallback({State.BuildAction.GetKey(), {}, {}, EStatus::Error});
-		CompletionEvent->Trigger();
+		State.Owner.End(this, [this]
+		{
+			CompletionCallback({State.BuildAction.GetKey(), {}, {}, EStatus::Error});
+			CompletionEvent->Trigger();
+		});
 	}
 
 }
@@ -1057,9 +1065,11 @@ void FRemoteBuildExecutionRequest::OnOutputBlobsDownloaded(FBatchReadBlobsRespon
 		Executor.Stats.TotalSuccessfulRemoteBuilds.fetch_add(1, std::memory_order_relaxed);
 	}
 
-	TRefCountPtr<FRemoteBuildExecutionRequest> SelfRef(this);
-	CompletionCallback({State.BuildAction.GetKey(), MoveTemp(BuildOutput), {}, BuildStatus});
-	CompletionEvent->Trigger();
+	State.Owner.End(this, [this, &BuildOutput, BuildStatus]() mutable
+	{
+		CompletionCallback({State.BuildAction.GetKey(), MoveTemp(BuildOutput), {}, BuildStatus});
+		CompletionEvent->Trigger();
+	});
 }
 
 } // namespace UE::DerivedData
diff --git a/Engine/Source/Editor/UnrealEd/Private/EditorDomain/BulkDataRegistryEditorDomain.cpp b/Engine/Source/Editor/UnrealEd/Private/EditorDomain/BulkDataRegistryEditorDomain.cpp
--- a/Engine/Source/Editor/UnrealEd/Private/EditorDomain/BulkDataRegistryEditorDomain.cpp
+++ b/Engine/Source/Editor/UnrealEd/Private/EditorDomain/BulkDataRegistryEditorDomain.cpp
@@ -4,6 +4,7 @@
 
 #include "Serialization/BulkDataRegistry.h"
 #include "Compression/CompressedBuffer.h"
+#include "DerivedDataCacheInterface.h"
 #include "DerivedDataCacheRecord.h"
 #include "EditorBuildInputResolver.h"
 #include "EditorDomain/EditorDomainUtils.h"
@@ -380,23 +381,21 @@ void FBulkDataRegistryEditorDomain::WritePayloadIdToCache(FName PackageName, con
 }
 
 void FBulkDataRegistryEditorDomain::ReadPayloadIdsFromCache(FName PackageName, TArray<TRefCountPtr<FPendingPayloadId>>&& OldPendings,
-	TArray<TPair<TRefCountPtr<FPendingPayloadId>, UE::DerivedData::FRequest>>&& NewPendings)
+	TArray<TRefCountPtr<FPendingPayloadId>>&& NewPendings)
 {
 	// Cancel any old requests for the Guids in NewPendings; we are about to overwrite them
-	// This cancelation has to occur outside of any lock, since the task may be in progress and
+	// This cancellation has to occur outside of any lock, since the task may be in progress and
 	// and to enter the lock and Cancel will wait on it
 	for (TRefCountPtr<FPendingPayloadId>& OldPending : OldPendings)
 	{
-		OldPending->GetRequest().Cancel();
+		OldPending->Cancel();
 	}
 	OldPendings.Empty();
-	for (TPair<TRefCountPtr<FPendingPayloadId>, UE::DerivedData::FRequest>& NewPair : NewPendings)
+	for (TRefCountPtr<FPendingPayloadId>& NewPending : NewPendings)
 	{
-		TRefCountPtr<FPendingPayloadId>& NewPending = NewPair.Get<0>();
 		// Creation of the Request has to occur outside of any lock, because the request
 		// may execute immediately on this thread and need to enter the lock; our locks are non-reentrant
-		UE::DerivedData::FRequest& Request = NewPair.Get<1>();
-		Request = UE::EditorDomain::GetBulkDataPayloadId(PackageName, NewPending->GetBulkDataId(),
+		UE::EditorDomain::GetBulkDataPayloadId(PackageName, NewPending->GetBulkDataId(), NewPending->GetRequestGroup(),
 			[this, PackageName, NewPending](FSharedBuffer Buffer)
 		{
 			if (Buffer.IsNull())
@@ -458,47 +457,40 @@ void FBulkDataRegistryEditorDomain::ReadPayloadIdsFromCache(FName PackageName, T
 		FWriteScopeLock RegistryScopeLock(RegistryLock);
 		if (!bActive)
 		{
-			for (TPair<TRefCountPtr<FPendingPayloadId>, UE::DerivedData::FRequest>& NewPair : NewPendings)
+			for (TRefCountPtr<FPendingPayloadId>& NewPending : NewPendings)
 			{
-				OldPendings.Add(MoveTemp(NewPair.Get<0>()));
+				OldPendings.Add(MoveTemp(NewPending));
 			}
 		}
 		else
 		{
-			for (TPair<TRefCountPtr<FPendingPayloadId>, UE::DerivedData::FRequest>& NewPair : NewPendings)
+			for (TRefCountPtr<FPendingPayloadId>& NewPending : NewPendings)
 			{
-				TRefCountPtr<FPendingPayloadId>& NewPending = NewPair.Get<0>();
-				UE::DerivedData::FRequest& Request = NewPair.Get<1>();
-
 				TRefCountPtr<FPendingPayloadId>* ExistingPending = PendingPayloadIds.Find(NewPending->GetBulkDataId());
 				if (!ExistingPending || *ExistingPending != NewPending)
 				{
 					OldPendings.Add(MoveTemp(NewPending));
 				}
-				else
-				{
-					check(NewPending->GetRequest().IsNull());
-					NewPending->GetRequest() = Request;
-				}
 			}
 		}
 	}
 	NewPendings.Empty();
 	for (TRefCountPtr<FPendingPayloadId>& OldPending : OldPendings)
 	{
-		OldPending->GetRequest().Cancel();
+		OldPending->Cancel();
 	}
 	OldPendings.Empty();
 }
 
 FPendingPackage::FPendingPackage(FName InPackageName, FBulkDataRegistryEditorDomain* InOwner)
 	: PackageName(InPackageName)
+	, BulkDataListCacheRequest(GetDerivedDataCacheRef().CreateGroup(UE::DerivedData::EPriority::Low))
 	, Owner(InOwner)
 {
 	PendingOperations = Flag_EndLoad | Flag_BulkDataListResults;
 
-	BulkDataListCacheRequest = UE::EditorDomain::GetBulkDataList(PackageName, [this](FSharedBuffer Buffer)
-		{ OnBulkDataListResults(Buffer); });
+	UE::EditorDomain::GetBulkDataList(PackageName, BulkDataListCacheRequest,
+		[this](FSharedBuffer Buffer) { OnBulkDataListResults(Buffer); });
 }
 
 void FPendingPackage::Cancel()
@@ -545,6 +537,10 @@ void FPendingPackage::OnBulkDataListResults(FSharedBuffer Buffer)
 	{
 		WriteCache();
 
+		// Removing *this from its owner will delete it, which will attempt to Cancel this request.
+		// Direct the group to keep requests alive to avoid a deadlock.
+		BulkDataListCacheRequest.KeepAlive();
+
 		FScopeLock PendingPackageScopeLock(&Owner->PendingPackageLock);
 		Owner->PendingPackages.Remove(PackageName);
 		// *this has been deleted and can no longer be accessed
@@ -559,7 +555,7 @@ void FPendingPackage::ReadCache()
 	}
 
 	TArray<TRefCountPtr<FPendingPayloadId>> OldPendings;
-	TArray<TPair<TRefCountPtr<FPendingPayloadId>, UE::DerivedData::FRequest>> NewPendings;
+	TArray<TRefCountPtr<FPendingPayloadId>> NewPendings;
 
 	// Add each CachedBulkData to the Registry, updating RawHash if it is missing.
 	// For every BulkData in this package in the Registry after the CachedBulkData has been added, 
@@ -597,22 +593,18 @@ void FPendingPackage::ReadCache()
 
 			if (bCachedLocationMatches && TargetBulkData.HasPlaceholderPayloadId())
 			{
-				NewPendings.Emplace(MakeTuple(new FPendingPayloadId(BulkDataId), UE::DerivedData::FRequest()));
+				NewPendings.Emplace(new FPendingPayloadId(BulkDataId));
 			}
 		}
 
-		if (NewPendings.Num())
+		for (TRefCountPtr<FPendingPayloadId>& NewPending : NewPendings)
 		{
-			for (TPair<TRefCountPtr<FPendingPayloadId>, UE::DerivedData::FRequest>& NewPair : NewPendings)
+			TRefCountPtr<FPendingPayloadId>& ExistingPending = Owner->PendingPayloadIds.FindOrAdd(NewPending->GetBulkDataId());
+			if (ExistingPending.IsValid())
 			{
-				TRefCountPtr<FPendingPayloadId>& NewPending = NewPair.Get<0>();
-				TRefCountPtr<FPendingPayloadId>& ExistingPending = Owner->PendingPayloadIds.FindOrAdd(NewPending->GetBulkDataId());
-				if (ExistingPending.IsValid())
-				{
-					OldPendings.Add(MoveTemp(ExistingPending));
-				}
-				ExistingPending = NewPending;
+				OldPendings.Add(MoveTemp(ExistingPending));
 			}
+			ExistingPending = NewPending;
 		}
 	}
 
@@ -738,6 +730,12 @@ void FUpdatePayloadWorker::DoWork()
 	}
 }
 
+FPendingPayloadId::FPendingPayloadId(const FGuid& InBulkDataId)
+	: BulkDataId(InBulkDataId)
+	, Request(GetDerivedDataCacheRef().CreateGroup(UE::DerivedData::EPriority::Low))
+{
+}
+
 void FPendingPayloadId::Cancel()
 {
 	Request.Cancel();
diff --git a/Engine/Source/Editor/UnrealEd/Private/EditorDomain/BulkDataRegistryEditorDomain.h b/Engine/Source/Editor/UnrealEd/Private/EditorDomain/BulkDataRegistryEditorDomain.h
--- a/Engine/Source/Editor/UnrealEd/Private/EditorDomain/BulkDataRegistryEditorDomain.h
+++ b/Engine/Source/Editor/UnrealEd/Private/EditorDomain/BulkDataRegistryEditorDomain.h
@@ -74,7 +74,7 @@ private:
 	FName PackageName;
 	TArray<UE::Virtualization::FVirtualizedUntypedBulkData> BulkDatas;
 	TArray<UE::Virtualization::FVirtualizedUntypedBulkData> CachedBulkDatas;
-	UE::DerivedData::FRequest BulkDataListCacheRequest;
+	UE::DerivedData::FRequestGroup BulkDataListCacheRequest;
 	FBulkDataRegistryEditorDomain* Owner;
 	/**
 	 * When PendingOperations reaches zero, we can remove the FPendingPackage.
@@ -136,16 +136,14 @@ struct FUpdatingPayload
 class FPendingPayloadId : public FThreadSafeRefCountedObject
 {
 public:
-	FPendingPayloadId(const FGuid& InBulkDataId)
-		:BulkDataId(InBulkDataId)
-	{
-	}
+	explicit FPendingPayloadId(const FGuid& InBulkDataId);
+
 	FPendingPayloadId(FPendingPayloadId&& Other) = delete;
 	FPendingPayloadId(const FPendingPayloadId& Other) = delete;
 
 	void Cancel();
 
-	UE::DerivedData::FRequest& GetRequest()
+	UE::DerivedData::FRequestGroup& GetRequestGroup()
 	{
 		return Request;
 	}
@@ -156,7 +154,7 @@ public:
 
 private:
 	FGuid BulkDataId;
-	UE::DerivedData::FRequest Request;
+	UE::DerivedData::FRequestGroup Request;
 };
 
 /** Implementation of a BulkDataRegistry that stores its persistent data in a DDC bucket. */
@@ -186,7 +184,7 @@ private:
 	void OnEndLoadPackage(TConstArrayView<UPackage*> LoadedPackages);
 	void WritePayloadIdToCache(FName PackageName, const UE::Virtualization::FVirtualizedUntypedBulkData& BulkData) const;
 	void ReadPayloadIdsFromCache(FName PackageName, TArray<TRefCountPtr<FPendingPayloadId>>&& OldPendings,
-		TArray<TPair<TRefCountPtr<FPendingPayloadId>, UE::DerivedData::FRequest>>&& NewPendings);
+		TArray<TRefCountPtr<FPendingPayloadId>>&& NewPendings);
 
 	friend class FPendingPackage;
 	friend class FUpdatePayloadWorker;
diff --git a/Engine/Source/Editor/UnrealEd/Private/EditorDomain/EditorBuildInputResolver.cpp b/Engine/Source/Editor/UnrealEd/Private/EditorDomain/EditorBuildInputResolver.cpp
--- a/Engine/Source/Editor/UnrealEd/Private/EditorDomain/EditorBuildInputResolver.cpp
+++ b/Engine/Source/Editor/UnrealEd/Private/EditorDomain/EditorBuildInputResolver.cpp
@@ -5,7 +5,9 @@
 #include "Async/Future.h"
 #include <atomic>
 #include "Containers/ArrayView.h"
+#include "DerivedDataBuildAction.h"
 #include "HAL/Event.h"
+#include "Misc/StringBuilder.h"
 #include "Serialization/BulkDataRegistry.h"
 #include "Templates/Tuple.h"
 
@@ -17,12 +19,12 @@ namespace UE::DerivedData
  * and calls the Resolve callback after all the futures complete
  */
 template <class BulkPayloadType, class CallbackPayloadType, class CallbackType>
-class TResolveInputRequest : public IRequest, public FThreadSafeRefCountedObject
+class TResolveInputRequest final : public FRequestBase
 {
 public:
 	typedef TResolveInputRequest<BulkPayloadType, CallbackPayloadType, CallbackType> ThisType;
 
-	TResolveInputRequest(const FBuildDefinition& InDefinition, CallbackType&& InOnResolved);
+	TResolveInputRequest(const FBuildDefinition& InDefinition, IRequestOwner& InOwner, CallbackType&& InOnResolved);
 
 	// IRequest interface
 	virtual void SetPriority(EPriority Priority) override
@@ -30,15 +32,6 @@ public:
 	}
 	virtual void Cancel() override;
 	virtual void Wait() override;
-	virtual bool Poll() override;
-	virtual void AddRef() const override
-	{
-		FThreadSafeRefCountedObject::AddRef();
-	}
-	virtual void Release() const override
-	{
-		FThreadSafeRefCountedObject::Release();
-	}
 
 	/** Stores the full set of Payload Futures and completes the creation of this IRequest */
 	void SetPayloads(TArrayView<TTuple<FString, TFuture<BulkPayloadType>>> BulkPayloads, EStatus OtherStatus);
@@ -74,6 +67,7 @@ private:
 
 	// Data written during creation that is read-only during Cancel and CallOnResolved
 	FBuildDefinition Definition;
+	IRequestOwner& Owner;
 #if DO_CHECK
 	bool bCreationComplete = false;
 #endif
@@ -116,14 +110,7 @@ FEditorBuildInputResolver& FEditorBuildInputResolver::Get()
 	return Singleton;
 }
 
-FRequest FEditorBuildInputResolver::ResolveKey(const FBuildKey& Key, FOnBuildKeyResolved&& OnResolved)
-{
-	// Not yet implemented
-	OnResolved({ Key, {}, EStatus::Error });
-	return FRequest();
-}
-
-FRequest FEditorBuildInputResolver::ResolveInputMeta(const FBuildDefinition& Definition, EPriority Priority,
+void FEditorBuildInputResolver::ResolveInputMeta(const FBuildDefinition& Definition, IRequestOwner& Owner,
 	FOnBuildInputMetaResolved&& OnResolved)
 {
 	EStatus OtherStatus = EStatus::Ok;
@@ -155,16 +142,11 @@ FRequest FEditorBuildInputResolver::ResolveInputMeta(const FBuildDefinition& Def
 			OtherStatus = EStatus::Error;
 		});
 
-	TRequest Request(new FResolveInputMetaRequest(Definition, MoveTemp(OnResolved)));
+	TRefCountPtr Request(new FResolveInputMetaRequest(Definition, Owner, MoveTemp(OnResolved)));
 	Request->SetPayloads(BulkPayloads, OtherStatus);
-	if (Priority == EPriority::Blocking)
-	{
-		Request->Wait();
-	}
-	return Request;
 }
 
-FRequest FEditorBuildInputResolver::ResolveInputData(const FBuildDefinition& Definition, EPriority Priority,
+void FEditorBuildInputResolver::ResolveInputData(const FBuildDefinition& Definition, IRequestOwner& Owner,
 	FOnBuildInputDataResolved&& OnResolved, FBuildInputFilter&& Filter)
 {
 	EStatus OtherStatus = EStatus::Ok;
@@ -199,29 +181,26 @@ FRequest FEditorBuildInputResolver::ResolveInputData(const FBuildDefinition& Def
 			OtherStatus = EStatus::Error;
 		});
 
-	TRequest Request(new FResolveInputDataRequest(Definition, MoveTemp(OnResolved)));
+	TRefCountPtr Request(new FResolveInputDataRequest(Definition, Owner, MoveTemp(OnResolved)));
 	Request->SetPayloads(BulkPayloads, OtherStatus);
-	if (Priority == EPriority::Blocking)
-	{
-		Request->Wait();
-	}
-	return Request;
 }
 
-FRequest FEditorBuildInputResolver::ResolveInputData(const FBuildAction& Action, EPriority Priority,
+void FEditorBuildInputResolver::ResolveInputData(const FBuildAction& Action, IRequestOwner& Owner,
 	FOnBuildInputDataResolved&& OnResolved, FBuildInputFilter&& Filter)
 {
-	// Not yet implemented
-	OnResolved({ {}, EStatus::Error });
-	return FRequest();
+	UE_LOG(LogCore, Error, TEXT("FEditorBuildInputResolver does not implement ResolveInputData from FBuildAction. ")
+		TEXT("Failed to resolve input data for build of '%s' by %s."),
+		*WriteToString<128>(Action.GetName()), *WriteToString<32>(Action.GetFunction()));
+	OnResolved({{}, EStatus::Error});
 }
 
 
 template <class BulkPayloadType, class CallbackPayloadType, class CallbackType>
 TResolveInputRequest<BulkPayloadType, CallbackPayloadType, CallbackType>::TResolveInputRequest(
-	const FBuildDefinition& InDefinition, CallbackType&& InOnResolved)
+	const FBuildDefinition& InDefinition, IRequestOwner& InOwner, CallbackType&& InOnResolved)
 	: OnResolved(MoveTemp(InOnResolved))
 	, Definition(InDefinition)
+	, Owner(InOwner)
 	, CallbackComplete(EEventMode::ManualReset)
 {
 }
@@ -232,9 +211,11 @@ void TResolveInputRequest<BulkPayloadType, CallbackPayloadType, CallbackType>::C
 	check(bCreationComplete);
 	if (bCallbackCalled.exchange(true) == false)
 	{
-		TArray<CallbackPayloadType> EmptyPayloads;
-		OnResolved({ EmptyPayloads, EStatus::Canceled });
-		CallbackComplete->Trigger();
+		Owner.End(this, [this]
+		{
+			OnResolved({{}, EStatus::Canceled});
+			CallbackComplete->Trigger();
+		});
 	}
 	else
 	{
@@ -256,23 +237,6 @@ void TResolveInputRequest<BulkPayloadType, CallbackPayloadType, CallbackType>::W
 	CallbackComplete->Wait();
 }
 
-template <class BulkPayloadType, class CallbackPayloadType, class CallbackType>
-bool TResolveInputRequest<BulkPayloadType, CallbackPayloadType, CallbackType>::Poll()
-{
-	check(bCreationComplete);
-	if (!bCallbackCalled)
-	{
-		for (FPayloadData& PayloadData : PayloadDatas)
-		{
-			if (!PayloadData.Future.IsReady())
-			{
-				return false;
-			}
-		}
-	}
-	return CallbackComplete->Wait(0, /* bIgnoreThreadIdleStats */ true);
-}
-
 template <class BulkPayloadType, class CallbackPayloadType, class CallbackType>
 void TResolveInputRequest<BulkPayloadType, CallbackPayloadType, CallbackType>::SetPayloads(
 	TArrayView<TTuple<FString, TFuture<BulkPayloadType>>> InBulkPayloads, EStatus OtherStatus)
@@ -320,10 +284,15 @@ void TResolveInputRequest<BulkPayloadType, CallbackPayloadType, CallbackType>::S
 	check(!bCreationComplete);
 	bCreationComplete = true;
 #endif
+	Owner.Begin(this);
 	if (--RemainingInputCount == 0)
 	{
 		CallOnResolved();
 	}
+	else if (Owner.GetPriority() == EPriority::Blocking)
+	{
+		Wait();
+	}
 }
 
 template <class BulkPayloadType, class CallbackPayloadType, class CallbackType>
@@ -331,16 +300,11 @@ void TResolveInputRequest<BulkPayloadType, CallbackPayloadType, CallbackType>::C
 {
 	if (bCallbackCalled.exchange(true) == false)
 	{
-		if (OnResolved)
+		Owner.End(this, [this]
 		{
-			OnResolved({ Payloads, EStatus::Ok });
-		}
-		else
-		{
-			TArray<CallbackPayloadType> EmptyPayloads;
-			OnResolved({ EmptyPayloads, EStatus::Error });
-		}
-		CallbackComplete->Trigger();
+			OnResolved({Payloads, EStatus::Ok});
+			CallbackComplete->Trigger();
+		});
 	}
 }
 
diff --git a/Engine/Source/Editor/UnrealEd/Private/EditorDomain/EditorBuildInputResolver.h b/Engine/Source/Editor/UnrealEd/Private/EditorDomain/EditorBuildInputResolver.h
--- a/Engine/Source/Editor/UnrealEd/Private/EditorDomain/EditorBuildInputResolver.h
+++ b/Engine/Source/Editor/UnrealEd/Private/EditorDomain/EditorBuildInputResolver.h
@@ -10,17 +10,16 @@ namespace UE::DerivedData
 {
 
 /** A BuildInputResolver that looks up the process-global inputs registered from package loads or a cache of the package loads. */
-class FEditorBuildInputResolver : public UE::DerivedData::IBuildInputResolver
+class FEditorBuildInputResolver : public IBuildInputResolver
 {
 public:
 	UNREALED_API static FEditorBuildInputResolver& Get();
 
-	virtual FRequest ResolveKey(const FBuildKey& Key, FOnBuildKeyResolved&& OnResolved) override;
-	virtual FRequest ResolveInputMeta(const FBuildDefinition& Definition, EPriority Priority,
+	virtual void ResolveInputMeta(const FBuildDefinition& Definition, IRequestOwner& Owner,
 		FOnBuildInputMetaResolved&& OnResolved) override;
-	virtual FRequest ResolveInputData(const FBuildDefinition& Definition, EPriority Priority,
+	virtual void ResolveInputData(const FBuildDefinition& Definition, IRequestOwner& Owner,
 		FOnBuildInputDataResolved&& OnResolved, FBuildInputFilter&& Filter = FBuildInputFilter()) override;
-	virtual FRequest ResolveInputData(const FBuildAction& Action, EPriority Priority,
+	virtual void ResolveInputData(const FBuildAction& Action, IRequestOwner& Owner,
 		FOnBuildInputDataResolved&& OnResolved, FBuildInputFilter&& Filter = FBuildInputFilter()) override;
 };
 
diff --git a/Engine/Source/Editor/UnrealEd/Private/EditorDomain/EditorDomain.cpp b/Engine/Source/Editor/UnrealEd/Private/EditorDomain/EditorDomain.cpp
--- a/Engine/Source/Editor/UnrealEd/Private/EditorDomain/EditorDomain.cpp
+++ b/Engine/Source/Editor/UnrealEd/Private/EditorDomain/EditorDomain.cpp
@@ -234,7 +234,7 @@ int64 FEditorDomain::FileSize(const FPackagePath& PackagePath, EPackageSegment P
 		return Workspace->FileSize(PackagePath, PackageSegment, OutUpdatedPath);
 	}
 
-	UE::DerivedData::FRequest Request;
+	UE::DerivedData::FOptionalRequestGroup RequestGroup;
 	int64 FileSize = -1;
 	{
 		FScopeLock ScopeLock(&Locks->Lock);
@@ -244,7 +244,6 @@ int64 FEditorDomain::FileSize(const FPackagePath& PackagePath, EPackageSegment P
 			return Workspace->FileSize(PackagePath, PackageSegment, OutUpdatedPath);
 		}
 
-		UE::DerivedData::ICache& Cache = GetDerivedDataCacheRef();
 		auto MetaDataGetComplete =
 			[&FileSize, &PackageSource, &PackagePath, PackageSegment, Locks=this->Locks, OutUpdatedPath]
 			(UE::DerivedData::FCacheGetCompleteParams&& Params)
@@ -278,10 +277,10 @@ int64 FEditorDomain::FileSize(const FPackagePath& PackagePath, EPackageSegment P
 		};
 		// Fetch meta-data only
 		ECachePolicy SkipFlags = ECachePolicy::SkipData & ~ECachePolicy::SkipMeta;
-		Request = RequestEditorDomainPackage(PackagePath, PackageSource->Digest, SkipFlags,
-			EPriority::Highest, MoveTemp(MetaDataGetComplete));
+		RequestGroup = GetDerivedDataCacheRef().CreateGroup(EPriority::Highest);
+		RequestEditorDomainPackage(PackagePath, PackageSource->Digest, SkipFlags, RequestGroup.Get(), MoveTemp(MetaDataGetComplete));
 	}
-	Request.Wait();
+	RequestGroup.Get().Wait();
 	return FileSize;
 }
 
@@ -311,14 +310,13 @@ FOpenPackageResult FEditorDomain::OpenReadPackage(const FPackagePath& PackagePat
 
 	// Fetch only meta-data in the initial request
 	ECachePolicy SkipFlags = ECachePolicy::SkipData & ~ECachePolicy::SkipMeta;
-	FRequest Request = RequestEditorDomainPackage(PackagePath, PackageSourceDigest,
-		SkipFlags, EPriority::Normal,
+	RequestEditorDomainPackage(PackagePath, PackageSourceDigest,
+		SkipFlags, Result->GetRequestOwner(),
 		[Result](FCacheGetCompleteParams&& Params)
 		{
 			// Note that ~FEditorDomainReadArchive waits for this callback to be called, so Result cannot dangle
 			Result->OnRecordRequestComplete(MoveTemp(Params));
 		});
-	Result->SetRequest(Request);
 
 	// Precache the exports segment
 	// EDITOR_DOMAIN_TODO: Skip doing this for OpenReadPackage calls from bulk data
@@ -354,14 +352,13 @@ IAsyncReadFileHandle* FEditorDomain::OpenAsyncReadPackage(const FPackagePath& Pa
 	// Fetch meta-data only in the initial request
 	ECachePolicy SkipFlags = ECachePolicy::SkipData & ~ECachePolicy::SkipMeta;
 	FEditorDomainAsyncReadFileHandle* Result = new FEditorDomainAsyncReadFileHandle(Locks, PackagePath, PackageSource);
-	FRequest Request = RequestEditorDomainPackage(PackagePath, PackageSource->Digest,
-		SkipFlags, EPriority::Normal,
+	RequestEditorDomainPackage(PackagePath, PackageSource->Digest,
+		SkipFlags, Result->GetRequestOwner(),
 		[Result](FCacheGetCompleteParams&& Params)
 		{
 			// Note that ~FEditorDomainAsyncReadFileHandle waits for this callback to be called, so Result cannot dangle
 			Result->OnRecordRequestComplete(MoveTemp(Params));
 		});
-	Result->SetRequest(Request);
 
 	return Result;
 }
diff --git a/Engine/Source/Editor/UnrealEd/Private/EditorDomain/EditorDomainArchive.cpp b/Engine/Source/Editor/UnrealEd/Private/EditorDomain/EditorDomainArchive.cpp
--- a/Engine/Source/Editor/UnrealEd/Private/EditorDomain/EditorDomainArchive.cpp
+++ b/Engine/Source/Editor/UnrealEd/Private/EditorDomain/EditorDomainArchive.cpp
@@ -20,6 +20,7 @@
 FEditorDomainPackageSegments::FEditorDomainPackageSegments(const TRefCountPtr<FEditorDomain::FLocks>& InLocks,
 	const FPackagePath& InPackagePath, const TRefCountPtr<FEditorDomain::FPackageSource>& InPackageSource)
 	: EditorDomainLocks(InLocks)
+	, RequestGroup(GetDerivedDataCacheRef().CreateGroup(UE::DerivedData::EPriority::Normal))
 	, PackagePath(InPackagePath)
 	, PackageSource(InPackageSource)
 	, PackageDigest(InPackageSource->Digest)
@@ -31,11 +32,6 @@ FEditorDomainPackageSegments::~FEditorDomainPackageSegments()
 	Close();
 }
 
-void FEditorDomainPackageSegments::SetRequest(const UE::DerivedData::FRequest& InRequest)
-{
-	Request = InRequest;
-}
-
 FEditorDomainPackageSegments::ESource FEditorDomainPackageSegments::GetSource() const
 {
 	return Source;
@@ -52,7 +48,7 @@ void FEditorDomainPackageSegments::WaitForReady() const
 	{
 		return;
 	}
-	Request.Wait();
+	const_cast<UE::DerivedData::FRequestGroup&>(RequestGroup).Wait();
 	Source = AsyncSource;
 }
 
@@ -81,10 +77,13 @@ void FEditorDomainPackageSegments::Close()
 {
 	// Called from Interface-only
 	AsyncSource = ESource::Closed;
-	Request.Cancel();
+	RequestGroup.Cancel();
 	for (FSegment& Segment : Segments)
 	{
-		Segment.Request.Cancel();
+		if (Segment.RequestGroup)
+		{
+			Segment.RequestGroup.Get().Cancel();
+		}
 	}
 
 	Segments.Reset();
@@ -198,7 +197,7 @@ FEditorDomainPackageSegments::FSegment* FEditorDomainPackageSegments::EnsureSegm
 	{
 		if (!Segment->bComplete)
 		{
-			if (!Segment->Request.IsValid())
+			if (!Segment->RequestGroup.IsValid())
 			{
 				SendSegmentRequest(*Segment);
 			}
@@ -221,7 +220,10 @@ FEditorDomainPackageSegments::FSegment* FEditorDomainPackageSegments::EnsureSegm
 		{
 			if (!Segment->bComplete)
 			{
-				Segment->Request.Wait();
+				if (Segment->RequestGroup)
+				{
+					Segment->RequestGroup.Get().Wait();
+				}
 				check(Segment->bAsyncComplete);
 				Segment->bComplete = true;
 			}
@@ -374,9 +376,10 @@ void FEditorDomainPackageSegments::SendSegmentRequest(FSegment& Segment)
 	// Called from Callback-only until AsyncSource is set, then from Interface-only
 	ICache& Cache = GetDerivedDataCacheRef();
 
-	// Note that Segment.Request is Interface-only and so we can write it outside the lock
-	Segment.Request = Cache.GetPayload({ FCachePayloadKey{ UE::EditorDomain::GetEditorDomainPackageKey(PackageDigest), Segment.PayloadId} },
-		PackagePath.GetDebugName(), ECachePolicy::Local, EPriority::Normal,
+	// Note that Segment.RequestGroup is Interface-only and so we can write it outside the lock
+	Segment.RequestGroup = Cache.CreateGroup(UE::DerivedData::EPriority::Normal);
+	Cache.GetPayload({ FCachePayloadKey{ UE::EditorDomain::GetEditorDomainPackageKey(PackageDigest), Segment.PayloadId} },
+		PackagePath.GetDebugName(), ECachePolicy::Local, Segment.RequestGroup.Get(),
 		[this, &Segment](FCacheGetPayloadCompleteParams&& Params)
 		{
 			if (AsyncSource == ESource::Closed)
@@ -416,11 +419,6 @@ FEditorDomainReadArchive::FEditorDomainReadArchive(const TRefCountPtr<FEditorDom
 	this->SetIsPersistent(true);
 }
 
-void FEditorDomainReadArchive::SetRequest(const UE::DerivedData::FRequest& InRequest)
-{
-	Segments.SetRequest(InRequest);
-}
-
 void FEditorDomainReadArchive::WaitForReady() const
 {
 	if (Segments.GetSource() != FEditorDomainPackageSegments::ESource::Uninitialized)
@@ -713,11 +711,6 @@ FEditorDomainAsyncReadFileHandle::FEditorDomainAsyncReadFileHandle(const TRefCou
 {
 }
 
-void FEditorDomainAsyncReadFileHandle::SetRequest(const UE::DerivedData::FRequest& InRequest)
-{
-	Segments.SetRequest(InRequest);
-}
-
 void FEditorDomainAsyncReadFileHandle::OnRecordRequestComplete(UE::DerivedData::FCacheGetCompleteParams&& Params)
 {
 	Segments.OnRecordRequestComplete(MoveTemp(Params),
diff --git a/Engine/Source/Editor/UnrealEd/Private/EditorDomain/EditorDomainArchive.h b/Engine/Source/Editor/UnrealEd/Private/EditorDomain/EditorDomainArchive.h
--- a/Engine/Source/Editor/UnrealEd/Private/EditorDomain/EditorDomainArchive.h
+++ b/Engine/Source/Editor/UnrealEd/Private/EditorDomain/EditorDomainArchive.h
@@ -43,8 +43,8 @@ public:
 		const FPackagePath& InPackagePath, const TRefCountPtr<FEditorDomain::FPackageSource>& InPackageSource);
 	~FEditorDomainPackageSegments();
 
-	/** Set the CacheRequest handle that will feed this archive. */
-	void SetRequest(const UE::DerivedData::FRequest& InRequest);
+	/** Get the request owner for requests that will feed this archive. */
+	UE::DerivedData::IRequestOwner& GetRequestOwner() { return RequestGroup; }
 	/** Callback from the request of the Record; set whether we're reading from EditorDomain bytes or WorkspaceDomain archive. */
 	void OnRecordRequestComplete(UE::DerivedData::FCacheGetCompleteParams&& Params,
 		TUniqueFunction<void(bool bValid)>&& CreateSegmentData,
@@ -108,7 +108,7 @@ private:
 		/** Attachment ID to request from cache for this segment. Read-only. */
 		UE::DerivedData::FPayloadId PayloadId;
 		/** The request that will be (has been) sent for the bytes from cache. Interface-only. */
-		UE::DerivedData::FRequest Request;
+		UE::DerivedData::FOptionalRequestGroup RequestGroup;
 		/** Offset from the start of the entire archive. Read-only. */
 		uint64 Start;
 		/** Payload bytes. Callback-only before bAsyncComplete, Interface-only after. */
@@ -143,7 +143,7 @@ private:
 	 */
 	TRefCountPtr<FEditorDomain::FLocks> EditorDomainLocks;
 	/** Initial Request for the list of segments. Read-only. */
-	UE::DerivedData::FRequest Request;
+	UE::DerivedData::FRequestGroup RequestGroup;
 	/** PackagePath for fallback to WorkspaceDomain. Read-only. */
 	FPackagePath PackagePath;
 	/** Offset to next use in Serialize. Interface-only. */
@@ -196,8 +196,8 @@ public:
 	FEditorDomainReadArchive(const TRefCountPtr<FEditorDomain::FLocks>& InLocks,
 		const FPackagePath& InPackagePath, const TRefCountPtr<FEditorDomain::FPackageSource>& InPackageSource);
 
-	/** Set the CacheRequest handle that will feed this archive. */
-	void SetRequest(const UE::DerivedData::FRequest& InRequest);
+	/** Get the request owner for requests that will feed this archive. */
+	UE::DerivedData::IRequestOwner& GetRequestOwner() { return Segments.GetRequestOwner(); }
 	/** Callback from the request of the Record; set whether we're reading from EditorDomain bytes or WorkspaceDomain archive. */
 	void OnRecordRequestComplete(UE::DerivedData::FCacheGetCompleteParams&& Params);
 	/** Get the PackageFormat, which depends on the domain the data is read from. */
@@ -235,8 +235,8 @@ public:
 	FEditorDomainAsyncReadFileHandle(const TRefCountPtr<FEditorDomain::FLocks>& InLocks,
 		const FPackagePath& InPackagePath, const TRefCountPtr<FEditorDomain::FPackageSource>& InPackageSource);
 
-	/** Set the CacheRequest handle that will feed this archive. */
-	void SetRequest(const UE::DerivedData::FRequest& InRequest);
+	/** Get the request owner for requests that will feed this archive. */
+	UE::DerivedData::IRequestOwner& GetRequestOwner() { return Segments.GetRequestOwner(); }
 	/** Callback from the request of the Record; set whether we're reading from EditorDomain bytes or WorkspaceDomain archive. */
 	void OnRecordRequestComplete(UE::DerivedData::FCacheGetCompleteParams&& Params);
 
diff --git a/Engine/Source/Editor/UnrealEd/Private/EditorDomain/EditorDomainUtils.cpp b/Engine/Source/Editor/UnrealEd/Private/EditorDomain/EditorDomainUtils.cpp
--- a/Engine/Source/Editor/UnrealEd/Private/EditorDomain/EditorDomainUtils.cpp
+++ b/Engine/Source/Editor/UnrealEd/Private/EditorDomain/EditorDomainUtils.cpp
@@ -223,16 +223,16 @@ UE::DerivedData::FCacheKey GetBulkDataPayloadIdKey(const FIoHash& PackageAndGuid
 	return UE::DerivedData::FCacheKey{ EditorDomainBulkDataPayloadIdBucket, PackageAndGuidDigest };
 }
 
-UE::DerivedData::FRequest RequestEditorDomainPackage(const FPackagePath& PackagePath,
-	const FPackageDigest& PackageDigest, UE::DerivedData::ECachePolicy SkipFlags, UE::DerivedData::EPriority CachePriority,
+void RequestEditorDomainPackage(const FPackagePath& PackagePath,
+	const FPackageDigest& PackageDigest, UE::DerivedData::ECachePolicy SkipFlags, UE::DerivedData::IRequestOwner& Owner,
 	UE::DerivedData::FOnCacheGetComplete&& Callback)
 {
 	UE::DerivedData::ICache& Cache = GetDerivedDataCacheRef();
 	checkf((SkipFlags & (~UE::DerivedData::ECachePolicy::SkipData)) == UE::DerivedData::ECachePolicy::None,
 		TEXT("SkipFlags should only contain ECachePolicy::Skip* flags"));
-	return Cache.Get({ GetEditorDomainPackageKey(PackageDigest) },
+	Cache.Get({ GetEditorDomainPackageKey(PackageDigest) },
 		PackagePath.GetDebugName(),
-		SkipFlags | UE::DerivedData::ECachePolicy::QueryLocal, CachePriority, MoveTemp(Callback));
+		SkipFlags | UE::DerivedData::ECachePolicy::QueryLocal, Owner, MoveTemp(Callback));
 }
 
 /** Stores data from SavePackage in accessible fields */
@@ -473,15 +473,14 @@ bool TrySavePackage(UPackage* Package)
 	MetaData.EndObject();
 
 	RecordBuilder.SetMeta(MetaData.Save().AsObject());
-	FCacheRecord Record = RecordBuilder.Build();
-	Cache.Put(MakeArrayView(&Record, 1), Package->GetName(), ECachePolicy::Local);
+	FRequestGroup Group = Cache.CreateGroup(EPriority::Normal);
+	Cache.Put({RecordBuilder.Build()}, Package->GetName(), ECachePolicy::Local, Group);
+	Group.KeepAlive();
 	return true;
 }
 
-UE::DerivedData::FRequest GetBulkDataList(FName PackageName, TUniqueFunction<void(FSharedBuffer Buffer)>&& Callback)
+void GetBulkDataList(FName PackageName, UE::DerivedData::IRequestOwner& Owner, TUniqueFunction<void(FSharedBuffer Buffer)>&& Callback)
 {
-	UE::DerivedData::ICache& Cache = GetDerivedDataCacheRef();
-
 	FString ErrorMessage;
 	FPackageDigest PackageDigest;
 	EPackageDigestResult FindHashResult = GetPackageDigest(*IAssetRegistry::Get(), PackageName, PackageDigest,
@@ -493,23 +492,23 @@ UE::DerivedData::FRequest GetBulkDataList(FName PackageName, TUniqueFunction<voi
 	default:
 	{
 		Callback(FSharedBuffer());
-		return UE::DerivedData::FRequest();
+		return;
 	}
 	}
 
-	return Cache.Get({ GetBulkDataListKey(PackageDigest) },
-		WriteToString<128>(PackageName), UE::DerivedData::ECachePolicy::Default, UE::DerivedData::EPriority::Low,
-		[InnerCallback = MoveTemp(Callback)](UE::DerivedData::FCacheGetCompleteParams&& Params)
+	using namespace UE::DerivedData;
+	ICache& Cache = GetDerivedDataCacheRef();
+	Cache.Get({ GetBulkDataListKey(PackageDigest) },
+		WriteToString<128>(PackageName), ECachePolicy::Default, Owner,
+		[InnerCallback = MoveTemp(Callback)](FCacheGetCompleteParams&& Params)
 		{
-			bool bOk = Params.Status == UE::DerivedData::EStatus::Ok;
+			bool bOk = Params.Status == EStatus::Ok;
 			InnerCallback(bOk ? Params.Record.GetValue() : FSharedBuffer());
 		});
 }
 
 void PutBulkDataList(FName PackageName, FSharedBuffer Buffer)
 {
-	UE::DerivedData::ICache& Cache = GetDerivedDataCacheRef();
-
 	FString ErrorMessage;
 	FPackageDigest PackageDigest;
 	EPackageDigestResult FindHashResult = GetPackageDigest(*IAssetRegistry::Get(), PackageName, PackageDigest,
@@ -524,10 +523,13 @@ void PutBulkDataList(FName PackageName, FSharedBuffer Buffer)
 	}
 	}
 
-	UE::DerivedData::FCacheRecordBuilder RecordBuilder = Cache.CreateRecord(GetBulkDataListKey(PackageDigest));
+	using namespace UE::DerivedData;
+	ICache& Cache = GetDerivedDataCacheRef();
+	FRequestGroup Group = Cache.CreateGroup(EPriority::Normal);
+	FCacheRecordBuilder RecordBuilder = Cache.CreateRecord(GetBulkDataListKey(PackageDigest));
 	RecordBuilder.SetValue(Buffer);
-	UE::DerivedData::FCacheRecord Record = RecordBuilder.Build();
-	Cache.Put(MakeArrayView(&Record, 1), WriteToString<128>(PackageName));
+	Cache.Put({RecordBuilder.Build()}, WriteToString<128>(PackageName), ECachePolicy::Default, Group);
+	Group.KeepAlive();
 }
 
 FIoHash GetPackageAndGuidDigest(FCbWriter& Builder, const FGuid& BulkDataId)
@@ -536,11 +538,9 @@ FIoHash GetPackageAndGuidDigest(FCbWriter& Builder, const FGuid& BulkDataId)
 	return Builder.Save().GetRangeHash();
 }
 
-UE::DerivedData::FRequest GetBulkDataPayloadId(FName PackageName, const FGuid& BulkDataId,
+void GetBulkDataPayloadId(FName PackageName, const FGuid& BulkDataId, UE::DerivedData::IRequestOwner& Owner,
 	TUniqueFunction<void(FSharedBuffer Buffer)>&& Callback)
 {
-	UE::DerivedData::ICache& Cache = GetDerivedDataCacheRef();
-
 	FString ErrorMessage;
 	FCbWriter Builder;
 	EPackageDigestResult FindHashResult = AppendPackageDigest(*IAssetRegistry::Get(), PackageName, Builder, ErrorMessage);
@@ -551,24 +551,24 @@ UE::DerivedData::FRequest GetBulkDataPayloadId(FName PackageName, const FGuid& B
 	default:
 	{
 		Callback(FSharedBuffer());
-		return UE::DerivedData::FRequest();
+		return;
 	}
 	}
 	FIoHash PackageAndGuidDigest = GetPackageAndGuidDigest(Builder, BulkDataId);
 
-	return Cache.Get({ GetBulkDataPayloadIdKey(PackageAndGuidDigest) },
-		WriteToString<192>(PackageName, TEXT("/"), BulkDataId), UE::DerivedData::ECachePolicy::Default, UE::DerivedData::EPriority::Low,
-		[InnerCallback = MoveTemp(Callback)](UE::DerivedData::FCacheGetCompleteParams&& Params)
+	using namespace UE::DerivedData;
+	ICache& Cache = GetDerivedDataCacheRef();
+	Cache.Get({ GetBulkDataPayloadIdKey(PackageAndGuidDigest) },
+		WriteToString<192>(PackageName, TEXT("/"), BulkDataId), ECachePolicy::Default, Owner,
+		[InnerCallback = MoveTemp(Callback)](FCacheGetCompleteParams&& Params)
 	{
-		bool bOk = Params.Status == UE::DerivedData::EStatus::Ok;
+		bool bOk = Params.Status == EStatus::Ok;
 		InnerCallback(bOk ? Params.Record.GetValue() : FSharedBuffer());
 	});
 }
 
 void PutBulkDataPayloadId(FName PackageName, const FGuid& BulkDataId, FSharedBuffer Buffer)
 {
-	UE::DerivedData::ICache& Cache = GetDerivedDataCacheRef();
-
 	FString ErrorMessage;
 	FCbWriter Builder;
 	EPackageDigestResult FindHashResult = AppendPackageDigest(*IAssetRegistry::Get(), PackageName, Builder, ErrorMessage);
@@ -583,10 +583,12 @@ void PutBulkDataPayloadId(FName PackageName, const FGuid& BulkDataId, FSharedBuf
 	}
 	FIoHash PackageAndGuidDigest = GetPackageAndGuidDigest(Builder, BulkDataId);
 
-	UE::DerivedData::FCacheRecordBuilder RecordBuilder = Cache.CreateRecord(GetBulkDataPayloadIdKey(PackageAndGuidDigest));
+	using namespace UE::DerivedData;
+	ICache& Cache = GetDerivedDataCacheRef();
+	FRequestGroup Group = Cache.CreateGroup(EPriority::Normal);
+	FCacheRecordBuilder RecordBuilder = Cache.CreateRecord(GetBulkDataPayloadIdKey(PackageAndGuidDigest));
 	RecordBuilder.SetValue(Buffer);
-	UE::DerivedData::FCacheRecord Record = RecordBuilder.Build();
-	Cache.Put(MakeArrayView(&Record, 1), WriteToString<128>(PackageName));
+	Cache.Put({RecordBuilder.Build()}, WriteToString<128>(PackageName), ECachePolicy::Default, Group);
 }
 
 
diff --git a/Engine/Source/Editor/UnrealEd/Private/EditorDomain/EditorDomainUtils.h b/Engine/Source/Editor/UnrealEd/Private/EditorDomain/EditorDomainUtils.h
--- a/Engine/Source/Editor/UnrealEd/Private/EditorDomain/EditorDomainUtils.h
+++ b/Engine/Source/Editor/UnrealEd/Private/EditorDomain/EditorDomainUtils.h
@@ -60,21 +60,22 @@ EPackageDigestResult AppendPackageDigest(IAssetRegistry& AssetRegistry, FName Pa
 void PrecacheClassDigests(TConstArrayView<FName> ClassNames);
 
 /** Get the CacheRequest for the given package from the EditorDomain cache bucket. */
-UE::DerivedData::FRequest RequestEditorDomainPackage(const FPackagePath& PackagePath,
+void RequestEditorDomainPackage(const FPackagePath& PackagePath,
 	const FPackageDigest& PackageDigest, UE::DerivedData::ECachePolicy SkipFlags,
-	UE::DerivedData::EPriority CachePriority, UE::DerivedData::FOnCacheGetComplete&& Callback);
+	UE::DerivedData::IRequestOwner& Owner, UE::DerivedData::FOnCacheGetComplete&& Callback);
 UE::DerivedData::FCacheKey GetEditorDomainPackageKey(const FPackageDigest& PackageDigest);
 
 /** Save the given package into the EditorDomain. */
 bool TrySavePackage(UPackage* Package);
 
 /** Get the CacheRequest for the BulkDataList of the given package. */
-UE::DerivedData::FRequest GetBulkDataList(FName PackageName, TUniqueFunction<void(FSharedBuffer Buffer)>&& Callback);
+void GetBulkDataList(FName PackageName, UE::DerivedData::IRequestOwner& Owner,
+	TUniqueFunction<void(FSharedBuffer Buffer)>&& Callback);
 
 /** Write the data for the BulkDataList of the given package to the cache. */
 void PutBulkDataList(FName PackageName, FSharedBuffer Buffer);
 
-UE::DerivedData::FRequest GetBulkDataPayloadId(FName PackageName, const FGuid& BulkDataId,
+void GetBulkDataPayloadId(FName PackageName, const FGuid& BulkDataId, UE::DerivedData::IRequestOwner& Owner,
 	TUniqueFunction<void(FSharedBuffer Buffer)> && Callback);
 
 void PutBulkDataPayloadId(FName PackageName, const FGuid& BulkDataId, FSharedBuffer Buffer);
diff --git a/Engine/Source/Programs/DerivedDataBuildWorker/Private/DerivedDataBuildWorker.cpp b/Engine/Source/Programs/DerivedDataBuildWorker/Private/DerivedDataBuildWorker.cpp
--- a/Engine/Source/Programs/DerivedDataBuildWorker/Private/DerivedDataBuildWorker.cpp
+++ b/Engine/Source/Programs/DerivedDataBuildWorker/Private/DerivedDataBuildWorker.cpp
@@ -45,7 +45,7 @@ public:
 private:
 	void BuildComplete(FBuildActionCompleteParams&& Params) const;
 
-	FRequest ResolveInputData(const FBuildAction& Action, EPriority Priority, FOnBuildInputDataResolved&& OnResolved, FBuildInputFilter&& Filter) final;
+	void ResolveInputData(const FBuildAction& Action, IRequestOwner& Owner, FOnBuildInputDataResolved&& OnResolved, FBuildInputFilter&& Filter) final;
 
 	TUniquePtr<FArchive> OpenInput(FStringView ActionPath, const FIoHash& RawHash) const;
 	TUniquePtr<FArchive> OpenOutput(FStringView ActionPath, const FIoHash& RawHash) const;
@@ -204,16 +204,19 @@ bool FBuildWorkerProgram::Build()
 
 	IBuild& BuildSystem = GetDerivedDataBuildRef();
 	FBuildSession Session = BuildSystem.CreateSession(TEXT("BuildWorker"_SV), this);
-	TArray<FRequest> Builds;
+	FRequestGroup Group = BuildSystem.CreateGroup(EPriority::Normal);
 
 	ON_SCOPE_EXIT
 	{
-		for (FRequest& Build : Builds)
-		{
-			Build.Wait();
-		}
+		Group.Wait();
 	};
 
+	if (ActionPaths.IsEmpty())
+	{
+		UE_LOG(LogDerivedDataBuildWorker, Error, TEXT("No build actions to operate on."));
+		return false;
+	}
+
 	for (const FString& ActionPath : ActionPaths)
 	{
 		UE_LOG(LogDerivedDataBuildWorker, Log, TEXT("Loading build action '%s'"), *ActionPath);
@@ -234,17 +237,11 @@ bool FBuildWorkerProgram::Build()
 		}
 		else
 		{
-			Builds.Add(Session.BuildAction(Action.Get(), {}, EBuildPolicy::BuildLocal, EPriority::Normal,
-				[this](FBuildActionCompleteParams&& Params) { BuildComplete(MoveTemp(Params)); }));
+			Session.BuildAction(Action.Get(), {}, EBuildPolicy::BuildLocal, Group,
+				[this](FBuildActionCompleteParams&& Params) { BuildComplete(MoveTemp(Params)); });
 		}
 	}
 
-	if (Builds.IsEmpty())
-	{
-		UE_LOG(LogDerivedDataBuildWorker, Error, TEXT("No build actions to operate on."));
-		return false;
-	}
-
 	return true;
 }
 
@@ -305,7 +302,7 @@ void FBuildWorkerProgram::BuildComplete(FBuildActionCompleteParams&& Params) con
 	}
 }
 
-FRequest FBuildWorkerProgram::ResolveInputData(const FBuildAction& Action, EPriority Priority, FOnBuildInputDataResolved&& OnResolved, FBuildInputFilter&& Filter)
+void FBuildWorkerProgram::ResolveInputData(const FBuildAction& Action, IRequestOwner& Owner, FOnBuildInputDataResolved&& OnResolved, FBuildInputFilter&& Filter)
 {
 	EStatus Status = EStatus::Ok;
 	TArray<FString> InputKeys;
@@ -330,7 +327,6 @@ FRequest FBuildWorkerProgram::ResolveInputData(const FBuildAction& Action, EPrio
 		}
 	});
 	OnResolved({Inputs, Status});
-	return FRequest();
 }
 
 TUniquePtr<FArchive> FBuildWorkerProgram::OpenInput(FStringView ActionPath, const FIoHash& RawHash) const
diff --git a/Engine/Source/Runtime/Engine/Private/TextureDerivedDataTask.cpp b/Engine/Source/Runtime/Engine/Private/TextureDerivedDataTask.cpp
--- a/Engine/Source/Runtime/Engine/Private/TextureDerivedDataTask.cpp
+++ b/Engine/Source/Runtime/Engine/Private/TextureDerivedDataTask.cpp
@@ -294,9 +294,9 @@ public:
 		return Buffer;
 	}
 
-	FRequest ResolveInputMeta(
+	void ResolveInputMeta(
 		const FBuildDefinition& Definition,
-		EPriority Priority,
+		IRequestOwner& Owner,
 		FOnBuildInputMetaResolved&& OnResolved) final
 	{
 		EStatus Status = EStatus::Ok;
@@ -318,12 +318,11 @@ public:
 			}
 		});
 		OnResolved({Inputs, Status});
-		return FRequest();
 	}
 
-	FRequest ResolveInputData(
+	void ResolveInputData(
 		const FBuildDefinition& Definition,
-		EPriority Priority,
+		IRequestOwner& Owner,
 		FOnBuildInputDataResolved&& OnResolved,
 		FBuildInputFilter&& Filter) final
 	{
@@ -349,7 +348,6 @@ public:
 			}
 		});
 		OnResolved({Inputs, Status});
-		return FRequest();
 	}
 
 private:
@@ -564,8 +562,9 @@ void FTextureCacheDerivedDataWorker::BuildTexture(bool bReplaceExistingDDC)
 				PlaceholderResolver.Emplace(Texture);
 				InputResolver = &PlaceholderResolver.GetValue();
 			}
+			FRequestGroup Group = Build.CreateGroup(EPriority::Blocking);
 			FBuildSession Session = Build.CreateSession(TexturePath, InputResolver);
-			Session.Build(DefinitionBuilder.Build(), EBuildPolicy::Default, EPriority::Blocking,
+			Session.Build(DefinitionBuilder.Build(), EBuildPolicy::Default, Group,
 				[this, &TexturePath, bReplaceExistingDDC] (FBuildCompleteParams&& Params)
 				{
 					#if !NO_LOGGING
@@ -583,7 +582,8 @@ void FTextureCacheDerivedDataWorker::BuildTexture(bool bReplaceExistingDDC)
 					{
 						ConsumeBuildFunctionOutput(Params.Output, TexturePath, bReplaceExistingDDC);
 					}
-				}).Wait();
+				});
+			Group.Wait();
 		}
 		else
 		{
@@ -996,11 +996,10 @@ public:
 		FStringView FunctionName,
 		FTexturePlatformData& InDerivedData,
 		const FTextureBuildSettings& Settings,
-		EQueuedWorkPriority BasePriority,
+		EQueuedWorkPriority InPriority,
 		ETextureCacheFlags Flags)
-		: ActiveRequest(&Request)
-		, DerivedData(InDerivedData)
-		, Priority(EnumHasAnyFlags(Flags, ETextureCacheFlags::Async) ? ConvertPriority(BasePriority) : UE::DerivedData::EPriority::Blocking)
+		: DerivedData(InDerivedData)
+		, Priority(InPriority)
 		, bCacheHit(false)
 		, bInlineMips(EnumHasAnyFlags(Flags, ETextureCacheFlags::InlineMips))
 		, FirstMipToLoad(Settings.LODBiasWithCinematicMips)
@@ -1020,15 +1019,17 @@ public:
 		IBuildInputResolver* GlobalResolver = GetGlobalBuildInputResolver();
 		BuildSession = Build.CreateSession(TexturePath, GlobalResolver ? GlobalResolver : &InputResolver);
 
+		EPriority GroupPriority = EnumHasAnyFlags(Flags, ETextureCacheFlags::Async) ? ConvertPriority(Priority) : UE::DerivedData::EPriority::Blocking;
+		Group = Build.CreateGroup(GroupPriority);
+
 		FBuildDefinition Definition = CreateDefinition(Build, Texture, TexturePath, FunctionName, KeySuffix, Settings);
 
 		if (!EnumHasAnyFlags(Flags, ETextureCacheFlags::ForceRebuild) && Settings.bHasEditorOnlyData)
 		{
-			ActiveRequest.store(&ShippingRequest, std::memory_order_relaxed);
 			FTextureBuildSettings ShippingSettings = Settings;
 			ShippingSettings.bHasEditorOnlyData = false;
 			FBuildDefinition ShippingDefinition = CreateDefinition(Build, Texture, TexturePath, FunctionName, KeySuffix, ShippingSettings);
-			ShippingRequest = BuildSession.Get().Build(ShippingDefinition, EBuildPolicy::Cache, Priority,
+			BuildSession.Get().Build(ShippingDefinition, EBuildPolicy::Cache, Group.Get(),
 				[this, Definition = MoveTemp(Definition), Flags](FBuildCompleteParams&& Params)
 				{
 					switch (Params.Status)
@@ -1043,7 +1044,6 @@ public:
 		}
 		else
 		{
-			ActiveRequest.store(&Request, std::memory_order_relaxed);
 			BeginBuild(Definition, Flags);
 		}
 	}
@@ -1075,10 +1075,8 @@ public:
 		{
 			BuildPolicy &= ~EBuildPolicy::CacheQuery;
 		}
-		FRequest* Active = ActiveRequest.load(std::memory_order_relaxed);
-		Request = BuildSession.Get().Build(Definition, BuildPolicy, Priority,
+		BuildSession.Get().Build(Definition, BuildPolicy, Group.Get(),
 			[this](FBuildCompleteParams&& Params) { EndBuild(MoveTemp(Params.Output), Params.BuildStatus); });
-		ActiveRequest.compare_exchange_strong(Active, &Request, std::memory_order_release, std::memory_order_relaxed);
 	}
 
 	void EndBuild(UE::DerivedData::FBuildOutput&& Output, UE::DerivedData::EBuildStatus Status)
@@ -1088,7 +1086,6 @@ public:
 		BuildOutputSize = Algo::TransformAccumulate(Output.GetPayloads(),
 			[](const FPayload& Payload) { return Payload.GetData().GetRawSize(); }, uint64(0));
 		WriteDerivedData(MoveTemp(Output));
-		ActiveRequest.store(nullptr, std::memory_order_relaxed);
 	}
 
 	void Finalize(bool& bOutFoundInCache, uint64& OutProcessedByteCount) final
@@ -1099,41 +1096,25 @@ public:
 
 	EQueuedWorkPriority GetPriority() const final
 	{
-		return ConvertPriority(Priority);
+		return Priority;
 	}
 
 	bool SetPriority(EQueuedWorkPriority QueuedWorkPriority) final
 	{
-		using namespace UE::DerivedData;
-		const EPriority NewPriority = ConvertPriority(QueuedWorkPriority);
-		if (Priority < EPriority::Blocking)
-		{
-			Priority = NewPriority;
-			if (FRequest* Active = ActiveRequest.load(std::memory_order_acquire))
-			{
-				Active->SetPriority(NewPriority);
-			}
-		}
+		Priority = QueuedWorkPriority;
+		Group.Get().SetPriority(ConvertPriority(QueuedWorkPriority));
 		return true;
 	}
 
 	bool Cancel() final
 	{
-		using namespace UE::DerivedData;
-		while (FRequest* Active = ActiveRequest.load(std::memory_order_acquire))
-		{
-			Active->Cancel();
-		}
+		Group.Get().Cancel();
 		return true;
 	}
 
 	void Wait() final
 	{
-		using namespace UE::DerivedData;
-		while (FRequest* Active = ActiveRequest.load(std::memory_order_acquire))
-		{
-			Active->Wait();
-		}
+		Group.Get().Wait();
 	}
 
 	bool WaitWithTimeout(float TimeLimitSeconds) final
@@ -1157,15 +1138,7 @@ public:
 
 	bool Poll() const final
 	{
-		using namespace UE::DerivedData;
-		while (FRequest* Active = ActiveRequest.load(std::memory_order_acquire))
-		{
-			if (!Active->Poll())
-			{
-				return false;
-			}
-		}
-		return true;
+		return Group.Get().Poll();
 	}
 
 private:
@@ -1276,12 +1249,10 @@ private:
 		return true;
 	}
 
-	std::atomic<UE::DerivedData::FRequest*> ActiveRequest;
 	FTexturePlatformData& DerivedData;
-	UE::DerivedData::FRequest Request;
-	UE::DerivedData::FRequest ShippingRequest;
+	UE::DerivedData::FOptionalRequestGroup Group;
 	UE::DerivedData::FOptionalBuildSession BuildSession;
-	UE::DerivedData::EPriority Priority;
+	EQueuedWorkPriority Priority;
 	bool bCacheHit;
 	bool bInlineMips;
 	int32 FirstMipToLoad;
