Zen HTTP interface refactored into separate header/source files

diff --git a/Engine/Source/Developer/DerivedDataCache/Internal/ZenServerHttp.h b/Engine/Source/Developer/DerivedDataCache/Internal/ZenServerHttp.h
new file mode 100644
--- /dev/null
+++ b/Engine/Source/Developer/DerivedDataCache/Internal/ZenServerHttp.h
@@ -0,0 +1,235 @@
+// Copyright Epic Games, Inc. All Rights Reserved.
+
+#pragma once
+
+#include "ZenServerInterface.h"
+#include "Containers/StaticArray.h"
+
+class FCompositeBuffer;
+class FCbPackage;
+
+#if UE_WITH_ZEN
+
+namespace UE::Zen {
+
+static bool IsSuccessCode(int ResponseCode)
+{
+	return 200 <= ResponseCode && ResponseCode < 300;
+}
+
+enum class EContentType
+{
+	Binary,
+	CompactBinary,
+	CompactBinaryPackage
+};
+
+/** Minimal HTTP request type wrapping CURL without the need for managers. This request
+  * is written to allow reuse of request objects, in order to allow connections to be reused.
+	
+  * CURL has a global library initialization (curl_global_init). We rely on this happening in
+  * the Online/HTTP library which is a dependency of this module
+  */
+
+class FZenHttpRequest
+{
+public:
+	FZenHttpRequest(const TCHAR* InDomain, bool bInLogErrors);
+	~FZenHttpRequest();
+
+	/**
+		* Resets all options on the request except those that should always be set.
+		*/
+	void Reset();
+
+	/** Returns the HTTP response code.*/
+	inline const int GetResponseCode() const
+	{
+		return int(ResponseCode);
+	}
+
+	inline const bool GetResponseFormatValid() const
+	{
+		return bResponseFormatValid;
+	}
+
+	/** Returns the number of bytes sent during this request (headers withstanding). */
+	inline const size_t  GetBytesSent() const
+	{
+		return BytesSent;
+	}
+
+	/**
+		* Convenience result type interpreted from HTTP response code.
+		*/
+	enum class Result
+	{
+		Success,
+		Failed
+	};
+
+	/**
+		* Upload buffer using the request, using PUT verb
+		* @param Uri Url to use.
+		* @param Buffer Data to upload
+		* @return Result of the request
+		*/
+	Result PerformBlockingPut(const TCHAR* Uri, const FCompositeBuffer& Buffer, EContentType ContentType);
+
+	/**
+		* Download an url into a buffer using the request.
+		* @param Uri Url to use.
+		* @param Buffer Optional buffer where data should be downloaded to. If this is null then 
+		* downloaded data will be stored in an internal buffer and accessed via GetResponseAsString
+		* @return Result of the request
+		*/
+	Result PerformBlockingDownload(const TCHAR* Uri, TArray<uint8>* Buffer);
+
+	/**
+		* Download an url into a buffer using the request.
+		* @param Uri Url to use.
+		* @param OutPackage Package instance which will receive the data
+		* @result Request success/failure status
+		*/
+	Result PerformBlockingDownload(const TCHAR* Uri, FCbPackage& OutPackage);
+
+	/**
+		* Query an url using the request. Queries can use either "Head" or "Delete" verbs.
+		* @param Uri Url to use.
+		* @return Result of the request
+		*/
+	Result PerformBlockingHead(const TCHAR* Uri);
+
+	/**
+		* Query an url using the request. Queries can use either "Head" or "Delete" verbs.
+		* @param Uri Url to use.
+		* @return Result of the request
+		*/
+	Result PerformBlockingDelete(const TCHAR* Uri);
+
+	/**
+		* Returns the response buffer as a string. Note that is the request is performed
+		* with an external buffer as target buffer this string will be empty.
+		*/
+	inline FString GetResponseAsString() const
+	{
+		return GetAnsiBufferAsString(ResponseBuffer);
+	}
+
+private:
+	void* /* CURL */		Curl = nullptr;
+	long /* CURLCode */		CurlResult;
+	long					ResponseCode = 0;
+	size_t					BytesSent = 0;
+	size_t					BytesReceived = 0;
+	bool					bLogErrors = false;
+	bool					bResponseFormatValid = false;
+
+	const FCompositeBuffer*	ReadDataView = nullptr;
+	TArray<uint8>*			WriteDataBufferPtr = nullptr;
+	FCbPackage*				WriteDataPackage = nullptr;
+	TArray<uint8>*			WriteHeaderBufferPtr = nullptr;
+
+	TArray<uint8>			ResponseHeader;
+	TArray<uint8>			ResponseBuffer;
+	TArray<FString>			Headers;
+	FString					Domain;
+
+	/**
+	  * Supported request verb
+	  */
+	enum class RequestVerb
+	{
+		Get,
+		Put,
+		Post,
+		Delete,
+		Head
+	};
+
+	void LogResult(long /*CURLcode*/ Result, const TCHAR* Uri, RequestVerb Verb) const;
+
+	/**
+		* Performs the request, blocking until finished.
+		* @param Uri Address on the domain to query
+		* @param Verb HTTP verb to use
+		* @param Buffer Optional buffer to directly receive the result of the request.
+		* If unset the response body will be stored in the request.
+		*/
+	Result PerformBlocking(const TCHAR* Uri, RequestVerb Verb, uint32 ContentLength);
+
+	static FString GetAnsiBufferAsString(const TArray<uint8>& Buffer);
+
+	struct FStatics;
+};
+
+/**
+  * Pool which manages a fixed set of requests. Users are required to release requests that have been
+  * acquired. 
+  * 
+  * Intended to be used with \ref FScopedRequestPtr which handles lifetime management transparently
+  */
+struct FRequestPool
+{
+	explicit FRequestPool(const TCHAR* InServiceUrl);
+	~FRequestPool();
+
+	/** Block until a request is free. Once a request has been returned it is
+	  * "owned by the caller and need to release it to the pool when work has been completed.
+	  * @return Usable request instance.
+	  */
+	FZenHttpRequest* WaitForFreeRequest();
+
+	/** Release request to the pool.
+	  * @param Request Request that should be freed. Note that any buffer owned by the request can now be reset.
+	  */
+	void ReleaseRequestToPool(FZenHttpRequest* Request);
+
+private:
+	struct FEntry
+	{
+		TAtomic<uint8> Usage;
+		FZenHttpRequest* Request;
+	};
+
+	TStaticArray<FEntry, 16> Pool;
+};
+
+/**
+  * Utility class to manage requesting and releasing requests from the \ref FRequestPool.
+  */
+struct FScopedRequestPtr
+{
+public:
+	FScopedRequestPtr(FRequestPool* InPool)
+	:	Request(InPool->WaitForFreeRequest())
+	,	Pool(InPool)
+	{
+	}
+
+	~FScopedRequestPtr()
+	{
+		Pool->ReleaseRequestToPool(Request);
+	}
+
+	inline bool IsValid() const
+	{
+		return Request != nullptr;
+	}
+
+	inline operator bool() const { return IsValid(); }
+
+	FZenHttpRequest* operator->()
+	{
+		check(IsValid());
+		return Request;
+	}
+
+private:
+	FZenHttpRequest*	Request;
+	FRequestPool*		Pool;
+};
+
+} // namespace UE::Zen
+
+#endif // UE_WITH_ZEN
diff --git a/Engine/Source/Developer/DerivedDataCache/Internal/ZenServerInterface.h b/Engine/Source/Developer/DerivedDataCache/Internal/ZenServerInterface.h
--- a/Engine/Source/Developer/DerivedDataCache/Internal/ZenServerInterface.h
+++ b/Engine/Source/Developer/DerivedDataCache/Internal/ZenServerInterface.h
@@ -31,224 +31,6 @@ public:
 private:
 };
 
-static bool IsSuccessCode(int ResponseCode)
-{
-	return 200 <= ResponseCode && ResponseCode < 300;
-}
-
-enum class EContentType
-{
-	Binary,
-	CompactBinary,
-	CompactBinaryPackage
-};
-
-/** Minimal HTTP request type wrapping CURL without the need for managers. This request
-  * is written to allow reuse of request objects, in order to allow connections to be reused.
-	
-  * CURL has a global library initialization (curl_global_init). We rely on this happening in
-  * the Online/HTTP library which is a dependency of this module
-  */
-
-class FZenHttpRequest
-{
-public:
-	FZenHttpRequest(const TCHAR* InDomain, bool bInLogErrors);
-	~FZenHttpRequest();
-
-	/**
-		* Resets all options on the request except those that should always be set.
-		*/
-	void Reset();
-
-	/** Returns the HTTP response code.*/
-	inline const int GetResponseCode() const
-	{
-		return int(ResponseCode);
-	}
-
-	inline const bool GetResponseFormatValid() const
-	{
-		return bResponseFormatValid;
-	}
-
-	/** Returns the number of bytes sent during this request (headers withstanding). */
-	inline const size_t  GetBytesSent() const
-	{
-		return BytesSent;
-	}
-
-	/**
-		* Convenience result type interpreted from HTTP response code.
-		*/
-	enum class Result
-	{
-		Success,
-		Failed
-	};
-
-	/**
-		* Upload buffer using the request, using PUT verb
-		* @param Uri Url to use.
-		* @param Buffer Data to upload
-		* @return Result of the request
-		*/
-	Result PerformBlockingPut(const TCHAR* Uri, const FCompositeBuffer& Buffer, EContentType ContentType);
-
-	/**
-		* Download an url into a buffer using the request.
-		* @param Uri Url to use.
-		* @param Buffer Optional buffer where data should be downloaded to. If this is null then 
-		* downloaded data will be stored in an internal buffer and accessed via GetResponseAsString
-		* @return Result of the request
-		*/
-	Result PerformBlockingDownload(const TCHAR* Uri, TArray<uint8>* Buffer);
-
-	/**
-		* Download an url into a buffer using the request.
-		* @param Uri Url to use.
-		* @param OutPackage Package instance which will receive the data
-		* @result Request success/failure status
-		*/
-	Result PerformBlockingDownload(const TCHAR* Uri, FCbPackage& OutPackage);
-
-	/**
-		* Query an url using the request. Queries can use either "Head" or "Delete" verbs.
-		* @param Uri Url to use.
-		* @return Result of the request
-		*/
-	Result PerformBlockingHead(const TCHAR* Uri);
-
-	/**
-		* Query an url using the request. Queries can use either "Head" or "Delete" verbs.
-		* @param Uri Url to use.
-		* @return Result of the request
-		*/
-	Result PerformBlockingDelete(const TCHAR* Uri);
-
-	/**
-		* Returns the response buffer as a string. Note that is the request is performed
-		* with an external buffer as target buffer this string will be empty.
-		*/
-	inline FString GetResponseAsString() const
-	{
-		return GetAnsiBufferAsString(ResponseBuffer);
-	}
-
-private:
-	void* /* CURL */		Curl = nullptr;
-	long /* CURLCode */		CurlResult;
-	long					ResponseCode = 0;
-	size_t					BytesSent = 0;
-	size_t					BytesReceived = 0;
-	bool					bLogErrors = false;
-	bool					bResponseFormatValid = false;
-
-	const FCompositeBuffer*	ReadDataView = nullptr;
-	TArray<uint8>*			WriteDataBufferPtr = nullptr;
-	FCbPackage*				WriteDataPackage = nullptr;
-	TArray<uint8>*			WriteHeaderBufferPtr = nullptr;
-
-	TArray<uint8>			ResponseHeader;
-	TArray<uint8>			ResponseBuffer;
-	TArray<FString>			Headers;
-	FString					Domain;
-
-	/**
-	  * Supported request verb
-	  */
-	enum class RequestVerb
-	{
-		Get,
-		Put,
-		Post,
-		Delete,
-		Head
-	};
-
-	void LogResult(long /*CURLcode*/ Result, const TCHAR* Uri, RequestVerb Verb) const;
-
-	/**
-		* Performs the request, blocking until finished.
-		* @param Uri Address on the domain to query
-		* @param Verb HTTP verb to use
-		* @param Buffer Optional buffer to directly receive the result of the request.
-		* If unset the response body will be stored in the request.
-		*/
-	Result PerformBlocking(const TCHAR* Uri, RequestVerb Verb, uint32 ContentLength);
-
-	static FString GetAnsiBufferAsString(const TArray<uint8>& Buffer);
-
-	struct FStatics;
-};
-
-/**
-  * Pool which manages a fixed set of requests. Users are required to release requests that have been
-  * acquired. 
-  * 
-  * Intended to be used with \ref FScopedRequestPtr which handles lifetime management transparently
-  */
-struct FRequestPool
-{
-	explicit FRequestPool(const TCHAR* InServiceUrl);
-	~FRequestPool();
-
-	/** Block until a request is free. Once a request has been returned it is
-	  * "owned by the caller and need to release it to the pool when work has been completed.
-	  * @return Usable request instance.
-	  */
-	FZenHttpRequest* WaitForFreeRequest();
-
-	/** Release request to the pool.
-	  * @param Request Request that should be freed. Note that any buffer owned by the request can now be reset.
-	  */
-	void ReleaseRequestToPool(FZenHttpRequest* Request);
-
-private:
-	struct FEntry
-	{
-		TAtomic<uint8> Usage;
-		FZenHttpRequest* Request;
-	};
-
-	TStaticArray<FEntry, 16> Pool;
-};
-
-/**
-  * Utility class to manage requesting and releasing requests from the \ref FRequestPool.
-  */
-struct FScopedRequestPtr
-{
-public:
-	FScopedRequestPtr(FRequestPool* InPool)
-	:	Request(InPool->WaitForFreeRequest())
-	,	Pool(InPool)
-	{
-	}
-
-	~FScopedRequestPtr()
-	{
-		Pool->ReleaseRequestToPool(Request);
-	}
-
-	inline bool IsValid() const
-	{
-		return Request != nullptr;
-	}
-
-	inline operator bool() const { return IsValid(); }
-
-	FZenHttpRequest* operator->()
-	{
-		check(IsValid());
-		return Request;
-	}
-
-private:
-	FZenHttpRequest*	Request;
-	FRequestPool*		Pool;
-};
-
 } // namespace UE::Zen
 
 #endif // UE_WITH_ZEN
diff --git a/Engine/Source/Developer/DerivedDataCache/Private/ZenDerivedDataBackend.cpp b/Engine/Source/Developer/DerivedDataCache/Private/ZenDerivedDataBackend.cpp
--- a/Engine/Source/Developer/DerivedDataCache/Private/ZenDerivedDataBackend.cpp
+++ b/Engine/Source/Developer/DerivedDataCache/Private/ZenDerivedDataBackend.cpp
@@ -13,7 +13,7 @@
 #include "Serialization/CompactBinaryValidation.h"
 #include "Serialization/CompactBinaryWriter.h"
 #include "ZenBackendUtils.h"
-#include "ZenServerInterface.h"
+#include "ZenServerHttp.h"
 
 TRACE_DECLARE_INT_COUNTER(ZenDDC_Exist,			TEXT("ZenDDC Exist"));
 TRACE_DECLARE_INT_COUNTER(ZenDDC_ExistHit,		TEXT("ZenDDC Exist Hit"));
diff --git a/Engine/Source/Developer/DerivedDataCache/Private/ZenServerHttp.cpp b/Engine/Source/Developer/DerivedDataCache/Private/ZenServerHttp.cpp
new file mode 100644
--- /dev/null
+++ b/Engine/Source/Developer/DerivedDataCache/Private/ZenServerHttp.cpp
@@ -0,0 +1,515 @@
+// Copyright Epic Games, Inc. All Rights Reserved.
+
+#include "ZenServerHttp.h"
+
+#if UE_WITH_ZEN
+
+#include "DerivedDataBackendInterface.h"	// For logger definition
+#include "ZenBackendUtils.h"
+
+#if PLATFORM_WINDOWS || PLATFORM_HOLOLENS
+#	include "Windows/WindowsHWrapper.h"
+#	include "Windows/AllowWindowsPlatformTypes.h"
+#endif
+
+#include "curl/curl.h"
+
+#if PLATFORM_WINDOWS || PLATFORM_HOLOLENS
+#	include "Windows/HideWindowsPlatformTypes.h"
+#endif
+
+#include "Memory/CompositeBuffer.h"
+#include "Serialization/CompactBinary.h"
+#include "Serialization/CompactBinaryPackage.h"
+#include "Serialization/CompactBinaryValidation.h"
+#include "Serialization/MemoryReader.h"
+
+namespace UE::Zen {
+
+#define UE_ZENDDC_BACKEND_WAIT_INTERVAL			0.01f
+#define UE_ZENDDC_HTTP_REQUEST_TIMEOUT_SECONDS	30L
+#define UE_ZENDDC_HTTP_DEBUG					0
+
+	struct FZenHttpRequest::FStatics
+	{
+		static size_t StaticDebugCallback(CURL* Handle, curl_infotype DebugInfoType, char* DebugInfo, size_t DebugInfoSize, void* UserData);
+		static size_t StaticReadFn(void* Ptr, size_t SizeInBlocks, size_t BlockSizeInBytes, void* UserData);
+		static size_t StaticWriteHeaderFn(void* Ptr, size_t SizeInBlocks, size_t BlockSizeInBytes, void* UserData);
+		static size_t StaticWriteBodyFn(void* Ptr, size_t SizeInBlocks, size_t BlockSizeInBytes, void* UserData);
+		static size_t StaticSeekFn(void* UserData, curl_off_t Offset, int Origin);
+	};
+
+	FZenHttpRequest::FZenHttpRequest(const TCHAR* InDomain, bool bInLogErrors)
+		: bLogErrors(bInLogErrors)
+		, Domain(InDomain)
+	{
+		Curl = curl_easy_init();
+		Reset();
+	}
+
+	FZenHttpRequest::~FZenHttpRequest()
+	{
+		curl_easy_cleanup(Curl);
+	}
+
+	/**
+	 * Resets all options on the request except those that should always be set.
+	 */
+	void FZenHttpRequest::Reset()
+	{
+		Headers.Reset();
+		ResponseHeader.Reset();
+		ResponseBuffer.Reset();
+		ResponseCode = 0;
+		bResponseFormatValid = true;
+		ReadDataView = nullptr;
+		WriteDataBufferPtr = nullptr;
+		WriteHeaderBufferPtr = nullptr;
+		BytesSent = 0;
+		BytesReceived = 0;
+		CurlResult = CURL_LAST;
+
+		curl_easy_reset(Curl);
+
+		// Options that are always set for all connections.
+		curl_easy_setopt(Curl, CURLOPT_CONNECTTIMEOUT, UE_ZENDDC_HTTP_REQUEST_TIMEOUT_SECONDS);
+		curl_easy_setopt(Curl, CURLOPT_FOLLOWLOCATION, 1L);
+		curl_easy_setopt(Curl, CURLOPT_NOSIGNAL, 1L);
+		curl_easy_setopt(Curl, CURLOPT_BUFFERSIZE, 256 * 1024L);
+		//curl_easy_setopt(Curl, CURLOPT_UPLOAD_BUFFERSIZE, 256 * 1024L);
+		// Response functions
+		curl_easy_setopt(Curl, CURLOPT_HEADERDATA, this);
+		curl_easy_setopt(Curl, CURLOPT_HEADERFUNCTION, &FZenHttpRequest::FStatics::StaticWriteHeaderFn);
+		curl_easy_setopt(Curl, CURLOPT_WRITEDATA, this);
+		curl_easy_setopt(Curl, CURLOPT_WRITEFUNCTION, &FZenHttpRequest::FStatics::StaticWriteBodyFn);
+		// Rewind method, handle special error case where request need to rewind data stream
+		curl_easy_setopt(Curl, CURLOPT_SEEKFUNCTION, &FZenHttpRequest::FStatics::StaticSeekFn);
+		curl_easy_setopt(Curl, CURLOPT_SEEKDATA, this);
+		// Debug hooks
+#if UE_ZENDDC_HTTP_DEBUG
+		curl_easy_setopt(Curl, CURLOPT_DEBUGDATA, this);
+		curl_easy_setopt(Curl, CURLOPT_DEBUGFUNCTION, StaticDebugCallback);
+		curl_easy_setopt(Curl, CURLOPT_VERBOSE, 1L);
+#endif
+	}
+
+	FZenHttpRequest::Result FZenHttpRequest::PerformBlockingPut(const TCHAR* Uri, const FCompositeBuffer& Buffer, EContentType ContentType)
+	{
+		uint32 ContentLength = 0u;
+
+		ContentLength = Buffer.GetSize();
+		curl_easy_setopt(Curl, CURLOPT_UPLOAD, 1L);
+		curl_easy_setopt(Curl, CURLOPT_INFILESIZE, ContentLength);
+		curl_easy_setopt(Curl, CURLOPT_READDATA, this);
+		curl_easy_setopt(Curl, CURLOPT_READFUNCTION, &FZenHttpRequest::FStatics::StaticReadFn);
+
+		switch (ContentType)
+		{
+		case EContentType::Binary:
+			Headers.Add(FString(TEXT("Content-Type: application/octet-stream")));
+			break;
+		case EContentType::CompactBinary:
+			Headers.Add(FString(TEXT("Content-Type: application/x-ue-cb")));
+			break;
+		case EContentType::CompactBinaryPackage:
+			Headers.Add(FString(TEXT("Content-Type: application/x-ue-cbpkg")));
+			break;
+		default:
+			checkNoEntry();
+			break;
+		}
+		ReadDataView = &Buffer;
+
+		return PerformBlocking(Uri, RequestVerb::Put, ContentLength);
+	}
+
+	FZenHttpRequest::Result FZenHttpRequest::PerformBlockingDownload(const TCHAR* Uri, TArray<uint8>* Buffer)
+	{
+		curl_easy_setopt(Curl, CURLOPT_HTTPGET, 1L);
+		WriteDataBufferPtr = Buffer;
+
+		return PerformBlocking(Uri, RequestVerb::Get, 0u);
+	}
+
+	FZenHttpRequest::Result FZenHttpRequest::PerformBlockingDownload(const TCHAR* Uri, FCbPackage& OutPackage)
+	{
+		curl_easy_setopt(Curl, CURLOPT_HTTPGET, 1L);
+		OutPackage.Reset();
+		// TODO: When PackageBytes can be written in segments directly, set the WritePtr to the OutPackage and use that
+		TArray<uint8> PackageBytes;
+		WriteDataBufferPtr = &PackageBytes;
+		Result LocalResult = PerformBlocking(Uri, RequestVerb::Get, 0u);
+		if (IsSuccessCode(ResponseCode))
+		{
+			if (ValidateCompactBinaryPackage(FMemoryView(PackageBytes.GetData(), PackageBytes.Num()), ECbValidateMode::All) != ECbValidateError::None)
+			{
+				bResponseFormatValid = false;
+			}
+			else
+			{
+				FMemoryReader PackageLoader(PackageBytes);
+				bResponseFormatValid = OutPackage.TryLoad(PackageLoader);
+			}
+		}
+		return LocalResult;
+	}
+
+	FZenHttpRequest::Result FZenHttpRequest::PerformBlockingHead(const TCHAR* Uri)
+	{
+		curl_easy_setopt(Curl, CURLOPT_NOBODY, 1L);
+
+		return PerformBlocking(Uri, RequestVerb::Head, 0u);
+	}
+
+	FZenHttpRequest::Result FZenHttpRequest::PerformBlockingDelete(const TCHAR* Uri)
+	{
+		curl_easy_setopt(Curl, CURLOPT_POST, 1L);
+		curl_easy_setopt(Curl, CURLOPT_CUSTOMREQUEST, "DELETE");
+
+		return PerformBlocking(Uri, RequestVerb::Delete, 0u);
+	}
+
+	FZenHttpRequest::Result FZenHttpRequest::PerformBlocking(const TCHAR* Uri, RequestVerb Verb, uint32 ContentLength)
+	{
+		static const char* CommonHeaders[] = {
+			"User-Agent: UE",
+			nullptr
+		};
+
+		TRACE_CPUPROFILER_EVENT_SCOPE(ZenDDC_CurlPerform);
+
+		// Setup request options
+		FString Url = FString::Printf(TEXT("%s/%s"), *Domain, Uri);
+		curl_easy_setopt(Curl, CURLOPT_URL, TCHAR_TO_ANSI(*Url));
+
+		// Setup response header buffer. If caller has not setup a response data buffer, use internal.
+		WriteHeaderBufferPtr = &ResponseHeader;
+		if (WriteDataBufferPtr == nullptr)
+		{
+			WriteDataBufferPtr = &ResponseBuffer;
+		}
+
+		// Content-Length should always be set
+		Headers.Add(FString::Printf(TEXT("Content-Length: %d"), ContentLength));
+
+		// Build headers list
+		curl_slist* CurlHeaders = nullptr;
+		// Add common headers
+		for (uint8 i = 0; CommonHeaders[i] != nullptr; ++i)
+		{
+			CurlHeaders = curl_slist_append(CurlHeaders, CommonHeaders[i]);
+		}
+		// Setup added headers
+		for (const FString& Header : Headers)
+		{
+			CurlHeaders = curl_slist_append(CurlHeaders, TCHAR_TO_ANSI(*Header));
+		}
+		curl_easy_setopt(Curl, CURLOPT_HTTPHEADER, CurlHeaders);
+
+		// Shots fired!
+		CurlResult = curl_easy_perform(Curl);
+
+		// Get response code
+		bool bRedirected = false;
+		if (CURLE_OK == curl_easy_getinfo(Curl, CURLINFO_RESPONSE_CODE, &ResponseCode))
+		{
+			bRedirected = (ResponseCode >= 300 && ResponseCode < 400);
+		}
+
+		LogResult(CurlResult, Uri, Verb);
+
+		// Clean up
+		curl_slist_free_all(CurlHeaders);
+
+		return CurlResult == CURLE_OK ? Result::Success : Result::Failed;
+	}
+
+	void FZenHttpRequest::LogResult(long InResult, const TCHAR* Uri, RequestVerb Verb) const
+	{
+		CURLcode Result = (CURLcode) InResult;
+		if (Result == CURLE_OK)
+		{
+			bool bSuccess = false;
+			const TCHAR* VerbStr = nullptr;
+			FString AdditionalInfo;
+
+			const bool Is400 = (ResponseCode >= 400) && (ResponseCode <= 499);
+			const bool Is200 = (ResponseCode >= 200) && (ResponseCode <= 299);
+
+			switch (Verb)
+			{
+			case RequestVerb::Head:
+				bSuccess = Is400 || Is200;
+				VerbStr = TEXT("querying");
+				break;
+			case RequestVerb::Get:
+				bSuccess = Is400 || Is200;
+				VerbStr = TEXT("fetching");
+				AdditionalInfo = FString::Printf(TEXT("Received: %d bytes."), BytesReceived);
+				break;
+			case RequestVerb::Put:
+				bSuccess = Is200;
+				VerbStr = TEXT("updating");
+				AdditionalInfo = FString::Printf(TEXT("Sent: %d bytes."), BytesSent);
+				break;
+			case RequestVerb::Post:
+				bSuccess = Is200;
+				VerbStr = TEXT("posting");
+				break;
+			case RequestVerb::Delete:
+				bSuccess = Is200;
+				VerbStr = TEXT("deleting");
+				break;
+			}
+
+			if (bSuccess)
+			{
+				UE_LOG(
+					LogDerivedDataCache,
+					Verbose,
+					TEXT("Finished %s HTTP cache entry (response %d) from %s. %s"),
+					VerbStr,
+					ResponseCode,
+					Uri,
+					*AdditionalInfo
+				);
+			}
+			else if (bLogErrors)
+			{
+				// Print the response body if we got one, otherwise print header.
+				FString Response = GetAnsiBufferAsString(ResponseBuffer.Num() > 0 ? ResponseBuffer : ResponseHeader);
+				Response.ReplaceCharInline('\n', ' ');
+				Response.ReplaceCharInline('\r', ' ');
+				UE_LOG(
+					LogDerivedDataCache,
+					Error,
+					TEXT("Failed %s HTTP cache entry (response %d) from %s. Response: %s"),
+					VerbStr,
+					ResponseCode,
+					Uri,
+					*Response
+				);
+			}
+		}
+		else if (bLogErrors)
+		{
+			UE_LOG(
+				LogDerivedDataCache,
+				Error,
+				TEXT("Error while connecting to %s: %s"),
+				*Domain,
+				ANSI_TO_TCHAR(curl_easy_strerror(Result))
+			);
+		}
+	}
+
+	FString FZenHttpRequest::GetAnsiBufferAsString(const TArray<uint8>& Buffer)
+	{
+		// Content is NOT null-terminated; we need to specify lengths here
+		FUTF8ToTCHAR TCHARData(reinterpret_cast<const ANSICHAR*>(Buffer.GetData()), Buffer.Num());
+		return FString(TCHARData.Length(), TCHARData.Get());
+	}
+
+	size_t FZenHttpRequest::FStatics::StaticDebugCallback(CURL* Handle, curl_infotype DebugInfoType, char* DebugInfo, size_t DebugInfoSize, void* UserData)
+	{
+		FZenHttpRequest* Request = static_cast<FZenHttpRequest*>(UserData);
+
+		switch (DebugInfoType)
+		{
+		case CURLINFO_TEXT:
+		{
+			// Truncate at 1023 characters. This is just an arbitrary number based on a buffer size seen in
+			// the libcurl code.
+			DebugInfoSize = FMath::Min(DebugInfoSize, (size_t)1023);
+
+			// Calculate the actual length of the string due to incorrect use of snprintf() in lib/vtls/openssl.c.
+			char* FoundNulPtr = (char*)memchr(DebugInfo, 0, DebugInfoSize);
+			int CalculatedSize = FoundNulPtr != nullptr ? FoundNulPtr - DebugInfo : DebugInfoSize;
+
+			auto ConvertedString = StringCast<TCHAR>(static_cast<const ANSICHAR*>(DebugInfo), CalculatedSize);
+			FString DebugText(ConvertedString.Length(), ConvertedString.Get());
+			DebugText.ReplaceInline(TEXT("\n"), TEXT(""), ESearchCase::CaseSensitive);
+			DebugText.ReplaceInline(TEXT("\r"), TEXT(""), ESearchCase::CaseSensitive);
+			UE_LOG(LogDerivedDataCache, VeryVerbose, TEXT("%p: '%s'"), Request, *DebugText);
+		}
+		break;
+
+		case CURLINFO_HEADER_IN:
+			UE_LOG(LogDerivedDataCache, VeryVerbose, TEXT("%p: Received header (%d bytes)"), Request, DebugInfoSize);
+			break;
+
+		case CURLINFO_DATA_IN:
+			UE_LOG(LogDerivedDataCache, VeryVerbose, TEXT("%p: Received data (%d bytes)"), Request, DebugInfoSize);
+			break;
+
+		case CURLINFO_DATA_OUT:
+			UE_LOG(LogDerivedDataCache, VeryVerbose, TEXT("%p: Sent data (%d bytes)"), Request, DebugInfoSize);
+			break;
+		}
+
+		return 0;
+	}
+
+	size_t FZenHttpRequest::FStatics::StaticReadFn(void* Ptr, size_t SizeInBlocks, size_t BlockSizeInBytes, void* UserData)
+	{
+		FZenHttpRequest* Request = static_cast<FZenHttpRequest*>(UserData);
+		check(Request->ReadDataView);
+		const FCompositeBuffer& ReadDataView = *Request->ReadDataView;
+
+		const size_t Offset = Request->BytesSent;
+		const size_t ReadSize = FMath::Min((size_t)ReadDataView.GetSize() - Offset, SizeInBlocks * BlockSizeInBytes);
+		check(ReadDataView.GetSize() >= Offset + ReadSize);
+
+		Memcpy(Ptr, ReadDataView, Offset, ReadSize);
+		Request->BytesSent += ReadSize;
+		return ReadSize;
+
+		return 0;
+	}
+
+	size_t FZenHttpRequest::FStatics::StaticWriteHeaderFn(void* Ptr, size_t SizeInBlocks, size_t BlockSizeInBytes, void* UserData)
+	{
+		FZenHttpRequest* Request = static_cast<FZenHttpRequest*>(UserData);
+		const size_t WriteSize = SizeInBlocks * BlockSizeInBytes;
+		TArray<uint8>* WriteHeaderBufferPtr = Request->WriteHeaderBufferPtr;
+		if (WriteHeaderBufferPtr && WriteSize > 0)
+		{
+			const size_t CurrentBufferLength = WriteHeaderBufferPtr->Num();
+			if (CurrentBufferLength > 0)
+			{
+				// Remove the previous zero termination
+				(*WriteHeaderBufferPtr)[CurrentBufferLength - 1] = ' ';
+			}
+
+			// Write the header
+			WriteHeaderBufferPtr->Append((const uint8*)Ptr, WriteSize + 1);
+			(*WriteHeaderBufferPtr)[WriteHeaderBufferPtr->Num() - 1] = 0; // Zero terminate string
+			return WriteSize;
+		}
+		return 0;
+	}
+
+	size_t FZenHttpRequest::FStatics::StaticWriteBodyFn(void* Ptr, size_t SizeInBlocks, size_t BlockSizeInBytes, void* UserData)
+	{
+		FZenHttpRequest* Request = static_cast<FZenHttpRequest*>(UserData);
+		const size_t WriteSize = SizeInBlocks * BlockSizeInBytes;
+		TArray<uint8>* WriteDataBufferPtr = Request->WriteDataBufferPtr;
+
+		if (WriteDataBufferPtr && WriteSize > 0)
+		{
+			// If this is the first part of the body being received, try to reserve 
+			// memory if content length is defined in the header.
+			if (Request->BytesReceived == 0 && Request->WriteHeaderBufferPtr)
+			{
+				static const ANSICHAR* ContentLengthHeaderStr = "Content-Length: ";
+				const ANSICHAR* Header = (const ANSICHAR*)Request->WriteHeaderBufferPtr->GetData();
+
+				if (const ANSICHAR* ContentLengthHeader = FCStringAnsi::Strstr(Header, ContentLengthHeaderStr))
+				{
+					size_t ContentLength = (size_t)FCStringAnsi::Atoi64(ContentLengthHeader + strlen(ContentLengthHeaderStr));
+					if (ContentLength > 0)
+					{
+						WriteDataBufferPtr->Reserve(ContentLength);
+					}
+				}
+			}
+
+			// Write to the target buffer
+			WriteDataBufferPtr->Append((const uint8*)Ptr, WriteSize);
+			Request->BytesReceived += WriteSize;
+			return WriteSize;
+		}
+
+		return 0;
+	}
+
+	size_t FZenHttpRequest::FStatics::StaticSeekFn(void* UserData, curl_off_t Offset, int Origin)
+	{
+		FZenHttpRequest* Request = static_cast<FZenHttpRequest*>(UserData);
+		size_t NewPosition = 0;
+		uint64 ReadDataSize = Request->ReadDataView ? Request->ReadDataView->GetSize() : 0;
+
+		switch (Origin)
+		{
+		case SEEK_SET: NewPosition = Offset; break;
+		case SEEK_CUR: NewPosition = Request->BytesSent + Offset; break;
+		case SEEK_END: NewPosition = ReadDataSize + Offset; break;
+		}
+
+		// Make sure we don't seek outside of the buffer
+		if (NewPosition < 0 || NewPosition >= ReadDataSize)
+		{
+			return CURL_SEEKFUNC_FAIL;
+		}
+
+		// Update the used offset
+		Request->BytesSent = NewPosition;
+		return CURL_SEEKFUNC_OK;
+	}
+
+	//////////////////////////////////////////////////////////////////////////
+
+	FRequestPool::FRequestPool(const TCHAR* InServiceUrl)
+	{
+		for (uint8 i = 0; i < Pool.Num(); ++i)
+		{
+			Pool[i].Usage = 0u;
+			Pool[i].Request = new FZenHttpRequest(InServiceUrl, true);
+		}
+	}
+
+	FRequestPool::~FRequestPool()
+	{
+		for (uint8 i = 0; i < Pool.Num(); ++i)
+		{
+			// No requests should be in use by now.
+			check(Pool[i].Usage.Load(EMemoryOrder::Relaxed) == 0u);
+			delete Pool[i].Request;
+		}
+	}
+
+	/**
+		* Block until a request is free. Once a request has been returned it is
+		* "owned by the caller and need to release it to the pool when work has been completed.
+		* @return Usable request instance.
+		*/
+	FZenHttpRequest* FRequestPool::WaitForFreeRequest()
+	{
+		TRACE_CPUPROFILER_EVENT_SCOPE(ZenDDC_WaitForConnPool);
+		while (true)
+		{
+			for (uint8 i = 0; i < Pool.Num(); ++i)
+			{
+				if (!Pool[i].Usage.Load(EMemoryOrder::Relaxed))
+				{
+					uint8 Expected = 0u;
+					if (Pool[i].Usage.CompareExchange(Expected, 1u))
+					{
+						return Pool[i].Request;
+					}
+				}
+			}
+			FPlatformProcess::Sleep(UE_ZENDDC_BACKEND_WAIT_INTERVAL);
+		}
+	}
+
+	/**
+		* Release request to the pool.
+		* @param Request Request that should be freed. Note that any buffer owened by the request can now be reset.
+		*/
+	void FRequestPool::ReleaseRequestToPool(FZenHttpRequest* Request)
+	{
+		for (uint8 i = 0; i < Pool.Num(); ++i)
+		{
+			if (Pool[i].Request == Request)
+			{
+				Request->Reset();
+				uint8 Expected = 1u;
+				Pool[i].Usage.CompareExchange(Expected, 0u);
+				return;
+			}
+		}
+		check(false);
+	}
+}
+#endif // UE_WITH_ZEN
diff --git a/Engine/Source/Developer/DerivedDataCache/Private/ZenServerInterface.cpp b/Engine/Source/Developer/DerivedDataCache/Private/ZenServerInterface.cpp
--- a/Engine/Source/Developer/DerivedDataCache/Private/ZenServerInterface.cpp
+++ b/Engine/Source/Developer/DerivedDataCache/Private/ZenServerInterface.cpp
@@ -6,17 +6,6 @@
 #include "DerivedDataBackendInterface.h"
 #include "ZenBackendUtils.h"
 
-#if PLATFORM_WINDOWS || PLATFORM_HOLOLENS
-#	include "Windows/WindowsHWrapper.h"
-#	include "Windows/AllowWindowsPlatformTypes.h"
-#endif
-
-#include "curl/curl.h"
-
-#if PLATFORM_WINDOWS || PLATFORM_HOLOLENS
-#	include "Windows/HideWindowsPlatformTypes.h"
-#endif
-
 #include "Memory/CompositeBuffer.h"
 #include "Serialization/CompactBinary.h"
 #include "Serialization/CompactBinaryPackage.h"
@@ -84,490 +73,5 @@ FZenServiceInstance::IsServiceReady()
 	return false;
 }
 
-#define UE_ZENDDC_BACKEND_WAIT_INTERVAL			0.01f
-#define UE_ZENDDC_HTTP_REQUEST_TIMEOUT_SECONDS	30L
-#define UE_ZENDDC_HTTP_DEBUG					0
-
-	struct FZenHttpRequest::FStatics
-	{
-		static size_t StaticDebugCallback(CURL* Handle, curl_infotype DebugInfoType, char* DebugInfo, size_t DebugInfoSize, void* UserData);
-		static size_t StaticReadFn(void* Ptr, size_t SizeInBlocks, size_t BlockSizeInBytes, void* UserData);
-		static size_t StaticWriteHeaderFn(void* Ptr, size_t SizeInBlocks, size_t BlockSizeInBytes, void* UserData);
-		static size_t StaticWriteBodyFn(void* Ptr, size_t SizeInBlocks, size_t BlockSizeInBytes, void* UserData);
-		static size_t StaticSeekFn(void* UserData, curl_off_t Offset, int Origin);
-	};
-
-	FZenHttpRequest::FZenHttpRequest(const TCHAR* InDomain, bool bInLogErrors)
-		: bLogErrors(bInLogErrors)
-		, Domain(InDomain)
-	{
-		Curl = curl_easy_init();
-		Reset();
-	}
-
-	FZenHttpRequest::~FZenHttpRequest()
-	{
-		curl_easy_cleanup(Curl);
-	}
-
-	/**
-	 * Resets all options on the request except those that should always be set.
-	 */
-	void FZenHttpRequest::Reset()
-	{
-		Headers.Reset();
-		ResponseHeader.Reset();
-		ResponseBuffer.Reset();
-		ResponseCode = 0;
-		bResponseFormatValid = true;
-		ReadDataView = nullptr;
-		WriteDataBufferPtr = nullptr;
-		WriteHeaderBufferPtr = nullptr;
-		BytesSent = 0;
-		BytesReceived = 0;
-		CurlResult = CURL_LAST;
-
-		curl_easy_reset(Curl);
-
-		// Options that are always set for all connections.
-		curl_easy_setopt(Curl, CURLOPT_CONNECTTIMEOUT, UE_ZENDDC_HTTP_REQUEST_TIMEOUT_SECONDS);
-		curl_easy_setopt(Curl, CURLOPT_FOLLOWLOCATION, 1L);
-		curl_easy_setopt(Curl, CURLOPT_NOSIGNAL, 1L);
-		curl_easy_setopt(Curl, CURLOPT_BUFFERSIZE, 256 * 1024L);
-		//curl_easy_setopt(Curl, CURLOPT_UPLOAD_BUFFERSIZE, 256 * 1024L);
-		// Response functions
-		curl_easy_setopt(Curl, CURLOPT_HEADERDATA, this);
-		curl_easy_setopt(Curl, CURLOPT_HEADERFUNCTION, &FZenHttpRequest::FStatics::StaticWriteHeaderFn);
-		curl_easy_setopt(Curl, CURLOPT_WRITEDATA, this);
-		curl_easy_setopt(Curl, CURLOPT_WRITEFUNCTION, &FZenHttpRequest::FStatics::StaticWriteBodyFn);
-		// Rewind method, handle special error case where request need to rewind data stream
-		curl_easy_setopt(Curl, CURLOPT_SEEKFUNCTION, &FZenHttpRequest::FStatics::StaticSeekFn);
-		curl_easy_setopt(Curl, CURLOPT_SEEKDATA, this);
-		// Debug hooks
-#if UE_ZENDDC_HTTP_DEBUG
-		curl_easy_setopt(Curl, CURLOPT_DEBUGDATA, this);
-		curl_easy_setopt(Curl, CURLOPT_DEBUGFUNCTION, StaticDebugCallback);
-		curl_easy_setopt(Curl, CURLOPT_VERBOSE, 1L);
-#endif
-	}
-
-	FZenHttpRequest::Result FZenHttpRequest::PerformBlockingPut(const TCHAR* Uri, const FCompositeBuffer& Buffer, EContentType ContentType)
-	{
-		uint32 ContentLength = 0u;
-
-		ContentLength = Buffer.GetSize();
-		curl_easy_setopt(Curl, CURLOPT_UPLOAD, 1L);
-		curl_easy_setopt(Curl, CURLOPT_INFILESIZE, ContentLength);
-		curl_easy_setopt(Curl, CURLOPT_READDATA, this);
-		curl_easy_setopt(Curl, CURLOPT_READFUNCTION, &FZenHttpRequest::FStatics::StaticReadFn);
-
-		switch (ContentType)
-		{
-		case EContentType::Binary:
-			Headers.Add(FString(TEXT("Content-Type: application/octet-stream")));
-			break;
-		case EContentType::CompactBinary:
-			Headers.Add(FString(TEXT("Content-Type: application/x-ue-cb")));
-			break;
-		case EContentType::CompactBinaryPackage:
-			Headers.Add(FString(TEXT("Content-Type: application/x-ue-cbpkg")));
-			break;
-		default:
-			checkNoEntry();
-			break;
-		}
-		ReadDataView = &Buffer;
-
-		return PerformBlocking(Uri, RequestVerb::Put, ContentLength);
-	}
-
-	FZenHttpRequest::Result FZenHttpRequest::PerformBlockingDownload(const TCHAR* Uri, TArray<uint8>* Buffer)
-	{
-		curl_easy_setopt(Curl, CURLOPT_HTTPGET, 1L);
-		WriteDataBufferPtr = Buffer;
-
-		return PerformBlocking(Uri, RequestVerb::Get, 0u);
-	}
-
-	FZenHttpRequest::Result FZenHttpRequest::PerformBlockingDownload(const TCHAR* Uri, FCbPackage& OutPackage)
-	{
-		curl_easy_setopt(Curl, CURLOPT_HTTPGET, 1L);
-		OutPackage.Reset();
-		// TODO: When PackageBytes can be written in segments directly, set the WritePtr to the OutPackage and use that
-		TArray<uint8> PackageBytes;
-		WriteDataBufferPtr = &PackageBytes;
-		Result LocalResult = PerformBlocking(Uri, RequestVerb::Get, 0u);
-		if (IsSuccessCode(ResponseCode))
-		{
-			if (ValidateCompactBinaryPackage(FMemoryView(PackageBytes.GetData(), PackageBytes.Num()), ECbValidateMode::All) != ECbValidateError::None)
-			{
-				bResponseFormatValid = false;
-			}
-			else
-			{
-				FMemoryReader PackageLoader(PackageBytes);
-				bResponseFormatValid = OutPackage.TryLoad(PackageLoader);
-			}
-		}
-		return LocalResult;
-	}
-
-	FZenHttpRequest::Result FZenHttpRequest::PerformBlockingHead(const TCHAR* Uri)
-	{
-		curl_easy_setopt(Curl, CURLOPT_NOBODY, 1L);
-
-		return PerformBlocking(Uri, RequestVerb::Head, 0u);
-	}
-
-	FZenHttpRequest::Result FZenHttpRequest::PerformBlockingDelete(const TCHAR* Uri)
-	{
-		curl_easy_setopt(Curl, CURLOPT_POST, 1L);
-		curl_easy_setopt(Curl, CURLOPT_CUSTOMREQUEST, "DELETE");
-
-		return PerformBlocking(Uri, RequestVerb::Delete, 0u);
-	}
-
-	FZenHttpRequest::Result FZenHttpRequest::PerformBlocking(const TCHAR* Uri, RequestVerb Verb, uint32 ContentLength)
-	{
-		static const char* CommonHeaders[] = {
-			"User-Agent: UE",
-			nullptr
-		};
-
-		TRACE_CPUPROFILER_EVENT_SCOPE(ZenDDC_CurlPerform);
-
-		// Setup request options
-		FString Url = FString::Printf(TEXT("%s/%s"), *Domain, Uri);
-		curl_easy_setopt(Curl, CURLOPT_URL, TCHAR_TO_ANSI(*Url));
-
-		// Setup response header buffer. If caller has not setup a response data buffer, use internal.
-		WriteHeaderBufferPtr = &ResponseHeader;
-		if (WriteDataBufferPtr == nullptr)
-		{
-			WriteDataBufferPtr = &ResponseBuffer;
-		}
-
-		// Content-Length should always be set
-		Headers.Add(FString::Printf(TEXT("Content-Length: %d"), ContentLength));
-
-		// Build headers list
-		curl_slist* CurlHeaders = nullptr;
-		// Add common headers
-		for (uint8 i = 0; CommonHeaders[i] != nullptr; ++i)
-		{
-			CurlHeaders = curl_slist_append(CurlHeaders, CommonHeaders[i]);
-		}
-		// Setup added headers
-		for (const FString& Header : Headers)
-		{
-			CurlHeaders = curl_slist_append(CurlHeaders, TCHAR_TO_ANSI(*Header));
-		}
-		curl_easy_setopt(Curl, CURLOPT_HTTPHEADER, CurlHeaders);
-
-		// Shots fired!
-		CurlResult = curl_easy_perform(Curl);
-
-		// Get response code
-		bool bRedirected = false;
-		if (CURLE_OK == curl_easy_getinfo(Curl, CURLINFO_RESPONSE_CODE, &ResponseCode))
-		{
-			bRedirected = (ResponseCode >= 300 && ResponseCode < 400);
-		}
-
-		LogResult(CurlResult, Uri, Verb);
-
-		// Clean up
-		curl_slist_free_all(CurlHeaders);
-
-		return CurlResult == CURLE_OK ? Result::Success : Result::Failed;
-	}
-
-	void FZenHttpRequest::LogResult(long InResult, const TCHAR* Uri, RequestVerb Verb) const
-	{
-		CURLcode Result = (CURLcode) InResult;
-		if (Result == CURLE_OK)
-		{
-			bool bSuccess = false;
-			const TCHAR* VerbStr = nullptr;
-			FString AdditionalInfo;
-
-			const bool Is400 = (ResponseCode >= 400) && (ResponseCode <= 499);
-			const bool Is200 = (ResponseCode >= 200) && (ResponseCode <= 299);
-
-			switch (Verb)
-			{
-			case RequestVerb::Head:
-				bSuccess = Is400 || Is200;
-				VerbStr = TEXT("querying");
-				break;
-			case RequestVerb::Get:
-				bSuccess = Is400 || Is200;
-				VerbStr = TEXT("fetching");
-				AdditionalInfo = FString::Printf(TEXT("Received: %d bytes."), BytesReceived);
-				break;
-			case RequestVerb::Put:
-				bSuccess = Is200;
-				VerbStr = TEXT("updating");
-				AdditionalInfo = FString::Printf(TEXT("Sent: %d bytes."), BytesSent);
-				break;
-			case RequestVerb::Post:
-				bSuccess = Is200;
-				VerbStr = TEXT("posting");
-				break;
-			case RequestVerb::Delete:
-				bSuccess = Is200;
-				VerbStr = TEXT("deleting");
-				break;
-			}
-
-			if (bSuccess)
-			{
-				UE_LOG(
-					LogDerivedDataCache,
-					Verbose,
-					TEXT("Finished %s HTTP cache entry (response %d) from %s. %s"),
-					VerbStr,
-					ResponseCode,
-					Uri,
-					*AdditionalInfo
-				);
-			}
-			else if (bLogErrors)
-			{
-				// Print the response body if we got one, otherwise print header.
-				FString Response = GetAnsiBufferAsString(ResponseBuffer.Num() > 0 ? ResponseBuffer : ResponseHeader);
-				Response.ReplaceCharInline('\n', ' ');
-				Response.ReplaceCharInline('\r', ' ');
-				UE_LOG(
-					LogDerivedDataCache,
-					Error,
-					TEXT("Failed %s HTTP cache entry (response %d) from %s. Response: %s"),
-					VerbStr,
-					ResponseCode,
-					Uri,
-					*Response
-				);
-			}
-		}
-		else if (bLogErrors)
-		{
-			UE_LOG(
-				LogDerivedDataCache,
-				Error,
-				TEXT("Error while connecting to %s: %s"),
-				*Domain,
-				ANSI_TO_TCHAR(curl_easy_strerror(Result))
-			);
-		}
-	}
-
-	FString FZenHttpRequest::GetAnsiBufferAsString(const TArray<uint8>& Buffer)
-	{
-		// Content is NOT null-terminated; we need to specify lengths here
-		FUTF8ToTCHAR TCHARData(reinterpret_cast<const ANSICHAR*>(Buffer.GetData()), Buffer.Num());
-		return FString(TCHARData.Length(), TCHARData.Get());
-	}
-
-	size_t FZenHttpRequest::FStatics::StaticDebugCallback(CURL* Handle, curl_infotype DebugInfoType, char* DebugInfo, size_t DebugInfoSize, void* UserData)
-	{
-		FZenHttpRequest* Request = static_cast<FZenHttpRequest*>(UserData);
-
-		switch (DebugInfoType)
-		{
-		case CURLINFO_TEXT:
-		{
-			// Truncate at 1023 characters. This is just an arbitrary number based on a buffer size seen in
-			// the libcurl code.
-			DebugInfoSize = FMath::Min(DebugInfoSize, (size_t)1023);
-
-			// Calculate the actual length of the string due to incorrect use of snprintf() in lib/vtls/openssl.c.
-			char* FoundNulPtr = (char*)memchr(DebugInfo, 0, DebugInfoSize);
-			int CalculatedSize = FoundNulPtr != nullptr ? FoundNulPtr - DebugInfo : DebugInfoSize;
-
-			auto ConvertedString = StringCast<TCHAR>(static_cast<const ANSICHAR*>(DebugInfo), CalculatedSize);
-			FString DebugText(ConvertedString.Length(), ConvertedString.Get());
-			DebugText.ReplaceInline(TEXT("\n"), TEXT(""), ESearchCase::CaseSensitive);
-			DebugText.ReplaceInline(TEXT("\r"), TEXT(""), ESearchCase::CaseSensitive);
-			UE_LOG(LogDerivedDataCache, VeryVerbose, TEXT("%p: '%s'"), Request, *DebugText);
-		}
-		break;
-
-		case CURLINFO_HEADER_IN:
-			UE_LOG(LogDerivedDataCache, VeryVerbose, TEXT("%p: Received header (%d bytes)"), Request, DebugInfoSize);
-			break;
-
-		case CURLINFO_DATA_IN:
-			UE_LOG(LogDerivedDataCache, VeryVerbose, TEXT("%p: Received data (%d bytes)"), Request, DebugInfoSize);
-			break;
-
-		case CURLINFO_DATA_OUT:
-			UE_LOG(LogDerivedDataCache, VeryVerbose, TEXT("%p: Sent data (%d bytes)"), Request, DebugInfoSize);
-			break;
-		}
-
-		return 0;
-	}
-
-	size_t FZenHttpRequest::FStatics::StaticReadFn(void* Ptr, size_t SizeInBlocks, size_t BlockSizeInBytes, void* UserData)
-	{
-		FZenHttpRequest* Request = static_cast<FZenHttpRequest*>(UserData);
-		check(Request->ReadDataView);
-		const FCompositeBuffer& ReadDataView = *Request->ReadDataView;
-
-		const size_t Offset = Request->BytesSent;
-		const size_t ReadSize = FMath::Min((size_t)ReadDataView.GetSize() - Offset, SizeInBlocks * BlockSizeInBytes);
-		check(ReadDataView.GetSize() >= Offset + ReadSize);
-
-		Memcpy(Ptr, ReadDataView, Offset, ReadSize);
-		Request->BytesSent += ReadSize;
-		return ReadSize;
-
-		return 0;
-	}
-
-	size_t FZenHttpRequest::FStatics::StaticWriteHeaderFn(void* Ptr, size_t SizeInBlocks, size_t BlockSizeInBytes, void* UserData)
-	{
-		FZenHttpRequest* Request = static_cast<FZenHttpRequest*>(UserData);
-		const size_t WriteSize = SizeInBlocks * BlockSizeInBytes;
-		TArray<uint8>* WriteHeaderBufferPtr = Request->WriteHeaderBufferPtr;
-		if (WriteHeaderBufferPtr && WriteSize > 0)
-		{
-			const size_t CurrentBufferLength = WriteHeaderBufferPtr->Num();
-			if (CurrentBufferLength > 0)
-			{
-				// Remove the previous zero termination
-				(*WriteHeaderBufferPtr)[CurrentBufferLength - 1] = ' ';
-			}
-
-			// Write the header
-			WriteHeaderBufferPtr->Append((const uint8*)Ptr, WriteSize + 1);
-			(*WriteHeaderBufferPtr)[WriteHeaderBufferPtr->Num() - 1] = 0; // Zero terminate string
-			return WriteSize;
-		}
-		return 0;
-	}
-
-	size_t FZenHttpRequest::FStatics::StaticWriteBodyFn(void* Ptr, size_t SizeInBlocks, size_t BlockSizeInBytes, void* UserData)
-	{
-		FZenHttpRequest* Request = static_cast<FZenHttpRequest*>(UserData);
-		const size_t WriteSize = SizeInBlocks * BlockSizeInBytes;
-		TArray<uint8>* WriteDataBufferPtr = Request->WriteDataBufferPtr;
-
-		if (WriteDataBufferPtr && WriteSize > 0)
-		{
-			// If this is the first part of the body being received, try to reserve 
-			// memory if content length is defined in the header.
-			if (Request->BytesReceived == 0 && Request->WriteHeaderBufferPtr)
-			{
-				static const ANSICHAR* ContentLengthHeaderStr = "Content-Length: ";
-				const ANSICHAR* Header = (const ANSICHAR*)Request->WriteHeaderBufferPtr->GetData();
-
-				if (const ANSICHAR* ContentLengthHeader = FCStringAnsi::Strstr(Header, ContentLengthHeaderStr))
-				{
-					size_t ContentLength = (size_t)FCStringAnsi::Atoi64(ContentLengthHeader + strlen(ContentLengthHeaderStr));
-					if (ContentLength > 0)
-					{
-						WriteDataBufferPtr->Reserve(ContentLength);
-					}
-				}
-			}
-
-			// Write to the target buffer
-			WriteDataBufferPtr->Append((const uint8*)Ptr, WriteSize);
-			Request->BytesReceived += WriteSize;
-			return WriteSize;
-		}
-
-		return 0;
-	}
-
-	size_t FZenHttpRequest::FStatics::StaticSeekFn(void* UserData, curl_off_t Offset, int Origin)
-	{
-		FZenHttpRequest* Request = static_cast<FZenHttpRequest*>(UserData);
-		size_t NewPosition = 0;
-		uint64 ReadDataSize = Request->ReadDataView ? Request->ReadDataView->GetSize() : 0;
-
-		switch (Origin)
-		{
-		case SEEK_SET: NewPosition = Offset; break;
-		case SEEK_CUR: NewPosition = Request->BytesSent + Offset; break;
-		case SEEK_END: NewPosition = ReadDataSize + Offset; break;
-		}
-
-		// Make sure we don't seek outside of the buffer
-		if (NewPosition < 0 || NewPosition >= ReadDataSize)
-		{
-			return CURL_SEEKFUNC_FAIL;
-		}
-
-		// Update the used offset
-		Request->BytesSent = NewPosition;
-		return CURL_SEEKFUNC_OK;
-	}
-
-	//////////////////////////////////////////////////////////////////////////
-
-	FRequestPool::FRequestPool(const TCHAR* InServiceUrl)
-	{
-		for (uint8 i = 0; i < Pool.Num(); ++i)
-		{
-			Pool[i].Usage = 0u;
-			Pool[i].Request = new FZenHttpRequest(InServiceUrl, true);
-		}
-	}
-
-	FRequestPool::~FRequestPool()
-	{
-		for (uint8 i = 0; i < Pool.Num(); ++i)
-		{
-			// No requests should be in use by now.
-			check(Pool[i].Usage.Load(EMemoryOrder::Relaxed) == 0u);
-			delete Pool[i].Request;
-		}
-	}
-
-	/**
-		* Block until a request is free. Once a request has been returned it is
-		* "owned by the caller and need to release it to the pool when work has been completed.
-		* @return Usable request instance.
-		*/
-	FZenHttpRequest* FRequestPool::WaitForFreeRequest()
-	{
-		TRACE_CPUPROFILER_EVENT_SCOPE(ZenDDC_WaitForConnPool);
-		while (true)
-		{
-			for (uint8 i = 0; i < Pool.Num(); ++i)
-			{
-				if (!Pool[i].Usage.Load(EMemoryOrder::Relaxed))
-				{
-					uint8 Expected = 0u;
-					if (Pool[i].Usage.CompareExchange(Expected, 1u))
-					{
-						return Pool[i].Request;
-					}
-				}
-			}
-			FPlatformProcess::Sleep(UE_ZENDDC_BACKEND_WAIT_INTERVAL);
-		}
-	}
-
-	/**
-		* Release request to the pool.
-		* @param Request Request that should be freed. Note that any buffer owened by the request can now be reset.
-		*/
-	void FRequestPool::ReleaseRequestToPool(FZenHttpRequest* Request)
-	{
-		for (uint8 i = 0; i < Pool.Num(); ++i)
-		{
-			if (Pool[i].Request == Request)
-			{
-				Request->Reset();
-				uint8 Expected = 1u;
-				Pool[i].Usage.CompareExchange(Expected, 0u);
-				return;
-			}
-		}
-		check(false);
-	}
 }
 #endif // UE_WITH_ZEN
