Another NonUnity build fix (ZenServerHttp.h)

diff --git a/Engine/Source/Developer/DerivedDataCache/Internal/ZenServerHttp.h b/Engine/Source/Developer/DerivedDataCache/Internal/ZenServerHttp.h
--- a/Engine/Source/Developer/DerivedDataCache/Internal/ZenServerHttp.h
+++ b/Engine/Source/Developer/DerivedDataCache/Internal/ZenServerHttp.h
@@ -3,7 +3,9 @@
 #pragma once
 
 #include "ZenServerInterface.h"
+
 #include "Containers/StaticArray.h"
+#include "Containers/UnrealString.h"
 
 class FCompositeBuffer;
 class FCbPackage;
