DDC: Added conversions between FPayloadId and FCbObjectId

#rnx

diff --git a/Engine/Source/Developer/DerivedDataCache/Private/DerivedDataBuildDefinition.cpp b/Engine/Source/Developer/DerivedDataCache/Private/DerivedDataBuildDefinition.cpp
--- a/Engine/Source/Developer/DerivedDataCache/Private/DerivedDataBuildDefinition.cpp
+++ b/Engine/Source/Developer/DerivedDataCache/Private/DerivedDataBuildDefinition.cpp
@@ -195,7 +195,7 @@ FBuildDefinitionInternal::FBuildDefinitionInternal(FBuildDefinitionBuilderIntern
 				const FBuildPayloadKey& PayloadKey = Pair.Value.Get<FBuildPayloadKey>();
 				Writer.BeginObject(FTCHARToUTF8(Pair.Key));
 				Writer.AddHash("Build"_ASV, PayloadKey.BuildKey.Hash);
-				Writer.AddObjectId("Payload"_ASV, FCbObjectId(PayloadKey.Id));
+				Writer.AddObjectId("Payload"_ASV, PayloadKey.Id);
 				Writer.EndObject();
 			}
 		}
@@ -300,7 +300,7 @@ void FBuildDefinitionInternal::IterateInputBuilds(TFunctionRef<void (FStringView
 	for (FCbFieldView Field : Definition.AsView()["Inputs"_ASV]["Builds"_ASV])
 	{
 		const FBuildKey BuildKey{Field["Build"_ASV].AsHash()};
-		const FPayloadId Id = FPayloadId(Field["Payload"_ASV].AsObjectId());
+		const FPayloadId Id = Field["Payload"_ASV].AsObjectId();
 		Visitor(FUTF8ToTCHAR(Field.GetName()), FBuildPayloadKey{BuildKey, Id});
 	}
 }
diff --git a/Engine/Source/Developer/DerivedDataCache/Private/DerivedDataBuildOutput.cpp b/Engine/Source/Developer/DerivedDataCache/Private/DerivedDataBuildOutput.cpp
--- a/Engine/Source/Developer/DerivedDataCache/Private/DerivedDataBuildOutput.cpp
+++ b/Engine/Source/Developer/DerivedDataCache/Private/DerivedDataBuildOutput.cpp
@@ -146,7 +146,7 @@ FBuildOutputInternal::FBuildOutputInternal(FStringView InName, FStringView InFun
 	Meta.MakeOwned();
 	for (FCbFieldView Payload : InOutput["Payloads"_ASV])
 	{
-		const FPayloadId Id(Payload["Id"_ASV].AsObjectId());
+		const FPayloadId Id = Payload["Id"_ASV].AsObjectId();
 		const FIoHash& RawHash = Payload["RawHash"_ASV].AsAttachment();
 		if (Id.IsNull() || RawHash.IsZero() || !Payload["RawSize"_ASV].IsInteger())
 		{
@@ -216,7 +216,7 @@ void FBuildOutputInternal::Save(FCbWriter& Writer) const
 		for (const FPayload& Payload : Payloads)
 		{
 			Writer.BeginObject();
-			Writer.AddObjectId("Id"_ASV, FCbObjectId(Payload.GetId()));
+			Writer.AddObjectId("Id"_ASV, Payload.GetId());
 			Writer.AddBinaryAttachment("RawHash"_ASV, Payload.GetRawHash());
 			Writer.AddInteger("RawSize"_ASV, Payload.GetRawSize());
 			Writer.EndObject();
diff --git a/Engine/Source/Developer/DerivedDataCache/Private/DerivedDataCacheRecord.cpp b/Engine/Source/Developer/DerivedDataCache/Private/DerivedDataCacheRecord.cpp
--- a/Engine/Source/Developer/DerivedDataCache/Private/DerivedDataCacheRecord.cpp
+++ b/Engine/Source/Developer/DerivedDataCache/Private/DerivedDataCacheRecord.cpp
@@ -305,7 +305,7 @@ FCbPackage SaveCacheRecord(const FCacheRecord& Record)
 	auto SavePayload = [&Package, &Writer](const FPayload& Payload)
 	{
 		Writer.BeginObject();
-		Writer.AddObjectId("Id"_ASV, FCbObjectId(Payload.GetId()));
+		Writer.AddObjectId("Id"_ASV, Payload.GetId());
 		Writer.AddHash("RawHash"_ASV, Payload.GetRawHash());
 		Writer.AddInteger("RawSize"_ASV, Payload.GetRawSize());
 		const FCompositeBuffer CompressedBuffer = Payload.GetData().GetCompressed();
@@ -362,7 +362,7 @@ FOptionalCacheRecord LoadCacheRecord(const FCbPackage& Package)
 
 	auto LoadPayload = [&Package](const FCbObjectView& PayloadObject)
 	{
-		const FPayloadId PayloadId(PayloadObject["Id"].AsObjectId());
+		const FPayloadId PayloadId = PayloadObject["Id"].AsObjectId();
 		if (PayloadId.IsNull())
 		{
 			return FPayload();
diff --git a/Engine/Source/Developer/DerivedDataCache/Private/DerivedDataPayload.cpp b/Engine/Source/Developer/DerivedDataCache/Private/DerivedDataPayload.cpp
new file mode 100644
--- /dev/null
+++ b/Engine/Source/Developer/DerivedDataCache/Private/DerivedDataPayload.cpp
@@ -0,0 +1,49 @@
+// Copyright Epic Games, Inc. All Rights Reserved.
+
+#include "DerivedDataPayload.h"
+
+#include "Containers/StringConv.h"
+#include "Hash/xxhash.h"
+#include "Serialization/CompactBinary.h"
+#include "Templates/UnrealTemplate.h"
+
+namespace UE::DerivedData
+{
+
+FPayloadId::FPayloadId(const FMemoryView Id)
+{
+	checkf(Id.GetSize() == sizeof(ByteArray),
+		TEXT("FPayloadId cannot be constructed from a view of %" UINT64_FMT " bytes."), Id.GetSize());
+	FMemory::Memcpy(Bytes, Id.GetData(), sizeof(ByteArray));
+}
+
+FPayloadId::FPayloadId(const FCbObjectId& Id)
+	: FPayloadId(ImplicitConv<const ByteArray&>(Id))
+{
+}
+
+FPayloadId::operator FCbObjectId() const
+{
+	return FCbObjectId(ImplicitConv<const ByteArray&>(*this));
+}
+
+FPayloadId FPayloadId::FromHash(const FIoHash& Hash)
+{
+	checkf(!Hash.IsZero(), TEXT("FPayloadId requires a non-zero hash."));
+	return FPayloadId(MakeMemoryView(Hash.GetBytes()).Left(sizeof(ByteArray)));
+}
+
+FPayloadId FPayloadId::FromName(const FAnsiStringView Name)
+{
+	checkf(!Name.IsEmpty(), TEXT("FPayloadId requires a non-empty name."));
+	uint8 HashBytes[16];
+	FXxHash128::HashBuffer(Name.GetData(), Name.Len()).ToByteArray(HashBytes);
+	return FPayloadId(MakeMemoryView(HashBytes, sizeof(ByteArray)));
+}
+
+FPayloadId FPayloadId::FromName(const FWideStringView Name)
+{
+	return FPayloadId::FromName(FTCHARToUTF8(Name));
+}
+
+} // UE::DerivedData
diff --git a/Engine/Source/Developer/DerivedDataCache/Private/FileSystemDerivedDataBackend.cpp b/Engine/Source/Developer/DerivedDataCache/Private/FileSystemDerivedDataBackend.cpp
--- a/Engine/Source/Developer/DerivedDataCache/Private/FileSystemDerivedDataBackend.cpp
+++ b/Engine/Source/Developer/DerivedDataCache/Private/FileSystemDerivedDataBackend.cpp
@@ -1208,7 +1208,7 @@ private:
 		}
 
 		Writer.BeginObject();
-		Writer.AddObjectId("Id"_ASV, FCbObjectId(Payload.GetId()));
+		Writer.AddObjectId("Id"_ASV, Payload.GetId());
 		Writer.AddInteger("RawSize"_ASV, Payload.GetRawSize());
 		if (bStoreInline)
 		{
@@ -1227,7 +1227,7 @@ private:
 
 	FPayload GetCachePayload(const FCacheKey& Key, FStringView Context, ECachePolicy Policy, const FCbObject& Object, bool bAlwaysLoadInlineData = false) const
 	{
-		const FPayloadId Id(Object.FindView("Id"_ASV).AsObjectId());
+		const FPayloadId Id = Object.FindView("Id"_ASV).AsObjectId();
 		const uint64 RawSize = Object.FindView("RawSize"_ASV).AsUInt64(MAX_uint64);
 		const FIoHash RawHash = Object.FindView("RawHash"_ASV).AsHash();
 		FIoHash CompressedHash = Object.FindView("CompressedHash"_ASV).AsHash();
diff --git a/Engine/Source/Developer/DerivedDataCache/Private/ZenDerivedDataBackend.cpp b/Engine/Source/Developer/DerivedDataCache/Private/ZenDerivedDataBackend.cpp
--- a/Engine/Source/Developer/DerivedDataCache/Private/ZenDerivedDataBackend.cpp
+++ b/Engine/Source/Developer/DerivedDataCache/Private/ZenDerivedDataBackend.cpp
@@ -781,7 +781,7 @@ bool FZenDerivedDataBackend::LegacyPutCachePayload(const FCacheKey& Key, FString
 	}
 
 	Writer.BeginObject();
-	Writer.AddObjectId("Id"_ASV, FCbObjectId(Payload.GetId().GetView()));
+	Writer.AddObjectId("Id"_ASV, Payload.GetId());
 	Writer.AddInteger("RawSize"_ASV, Payload.GetRawSize());
 	if (bStoreInline)
 	{
@@ -921,7 +921,7 @@ FPayload FZenDerivedDataBackend::LegacyGetCachePayload(const FCacheKey& Key, FSt
 	const FCbObject& Object, bool bAlwaysLoadInlineData) const
 {
 	// TODO: Temporary function copied from FFileSystemDerivedDataBackend, helper for CreateRecord
-	const FPayloadId Id(Object.FindView("Id"_ASV).AsObjectId().GetView());
+	const FPayloadId Id = Object.FindView("Id"_ASV).AsObjectId();
 	const uint64 RawSize = Object.FindView("RawSize"_ASV).AsUInt64(MAX_uint64);
 	const FIoHash RawHash = Object.FindView("RawHash"_ASV).AsHash();
 	FIoHash CompressedHash = Object.FindView("CompressedHash"_ASV).AsHash();
diff --git a/Engine/Source/Developer/DerivedDataCache/Public/DerivedDataPayload.h b/Engine/Source/Developer/DerivedDataCache/Public/DerivedDataPayload.h
--- a/Engine/Source/Developer/DerivedDataCache/Public/DerivedDataPayload.h
+++ b/Engine/Source/Developer/DerivedDataCache/Public/DerivedDataPayload.h
@@ -4,14 +4,16 @@
 
 #include "CoreTypes.h"
 #include "Compression/CompressedBuffer.h"
-#include "Containers/StringConv.h"
 #include "Containers/StringFwd.h"
 #include "Containers/StringView.h"
-#include "Hash/xxhash.h"
 #include "IO/IoHash.h"
 #include "Memory/MemoryView.h"
 #include "String/BytesToHex.h"
 
+#define UE_API DERIVEDDATACACHE_API
+
+class FCbObjectId;
+
 namespace UE::DerivedData
 {
 
@@ -30,14 +32,20 @@ public:
 	inline explicit FPayloadId(const ByteArray& Id);
 
 	/** Construct an ID from a view of 12 bytes. */
-	inline explicit FPayloadId(FMemoryView Id);
+	UE_API explicit FPayloadId(FMemoryView Id);
+
+	/** Construct an ID from a Compact Binary Object ID. */
+	UE_API FPayloadId(const FCbObjectId& Id);
+
+	/** Returns the ID as a Compact Binary Object ID. */
+	UE_API operator FCbObjectId() const;
 
 	/** Construct an ID from a non-zero hash. */
-	[[nodiscard]] static inline FPayloadId FromHash(const FIoHash& Hash);
+	[[nodiscard]] UE_API static FPayloadId FromHash(const FIoHash& Hash);
 
 	/** Construct an ID from a non-empty name. */
-	[[nodiscard]] static inline FPayloadId FromName(FAnsiStringView Name);
-	[[nodiscard]] static inline FPayloadId FromName(FWideStringView Name);
+	[[nodiscard]] UE_API static FPayloadId FromName(FAnsiStringView Name);
+	[[nodiscard]] UE_API static FPayloadId FromName(FWideStringView Name);
 
 	/** Returns a reference to the raw byte array for the ID. */
 	inline const ByteArray& GetBytes() const { return Bytes; }
@@ -131,32 +139,6 @@ inline FPayloadId::FPayloadId(const ByteArray& Id)
 	FMemory::Memcpy(Bytes, Id, sizeof(ByteArray));
 }
 
-inline FPayloadId::FPayloadId(const FMemoryView Id)
-{
-	checkf(Id.GetSize() == sizeof(ByteArray),
-		TEXT("FPayloadId cannot be constructed from a view of %" UINT64_FMT " bytes."), Id.GetSize());
-	FMemory::Memcpy(Bytes, Id.GetData(), sizeof(ByteArray));
-}
-
-inline FPayloadId FPayloadId::FromHash(const FIoHash& Hash)
-{
-	checkf(!Hash.IsZero(), TEXT("FPayloadId requires a non-zero hash."));
-	return FPayloadId(MakeMemoryView(Hash.GetBytes()).Left(sizeof(ByteArray)));
-}
-
-inline FPayloadId FPayloadId::FromName(const FAnsiStringView Name)
-{
-	checkf(!Name.IsEmpty(), TEXT("FPayloadId requires a non-empty name."));
-	uint8 HashBytes[16];
-	FXxHash128::HashBuffer(Name.GetData(), Name.Len()).ToByteArray(HashBytes);
-	return FPayloadId(MakeMemoryView(HashBytes, sizeof(ByteArray)));
-}
-
-inline FPayloadId FPayloadId::FromName(const FWideStringView Name)
-{
-	return FPayloadId::FromName(FTCHARToUTF8(Name));
-}
-
 ///////////////////////////////////////////////////////////////////////////////////////////////////////////////////////
 
 inline FPayload::FPayload(const FPayloadId& InId)
@@ -247,3 +229,5 @@ inline bool operator<(const FPayload& A, const FPayload& B)
 ///////////////////////////////////////////////////////////////////////////////////////////////////////////////////////
 
 } // UE::DerivedData
+
+#undef UE_API
diff --git a/Engine/Source/Editor/UnrealEd/Private/EditorDomain/EditorDomainUtils.cpp b/Engine/Source/Editor/UnrealEd/Private/EditorDomain/EditorDomainUtils.cpp
--- a/Engine/Source/Editor/UnrealEd/Private/EditorDomain/EditorDomainUtils.cpp
+++ b/Engine/Source/Editor/UnrealEd/Private/EditorDomain/EditorDomainUtils.cpp
@@ -445,10 +445,9 @@ bool TrySavePackage(UPackage* Package)
 	// the exports with offsets that expect all attachment segments to exist in the segmented archive.
 	auto IntToPayloadId = [](int32 Value)
 	{
-		FPayloadId::ByteArray Bytes;
+		alignas(int32) FPayloadId::ByteArray Bytes{};
 		static_assert(sizeof(Bytes) >= sizeof(int32), "We are storing an int32 counter in the Bytes array");
-		FMemory::Memset(&Bytes, 0, sizeof(Bytes));
-		int32* IntView = reinterpret_cast<int32*>(&Bytes);
+		int32* IntView = reinterpret_cast<int32*>(Bytes);
 		IntView[0] = Value;
 		return FPayloadId(Bytes);
 	};
