DDC: Rename FRequestGroup to FRequestOwner and export its constructor

#rnx

diff --git a/Engine/Source/Developer/DerivedDataCache/Private/DerivedDataBackendAsyncPutWrapper.cpp b/Engine/Source/Developer/DerivedDataCache/Private/DerivedDataBackendAsyncPutWrapper.cpp
--- a/Engine/Source/Developer/DerivedDataCache/Private/DerivedDataBackendAsyncPutWrapper.cpp
+++ b/Engine/Source/Developer/DerivedDataCache/Private/DerivedDataBackendAsyncPutWrapper.cpp
@@ -479,9 +479,9 @@ void FDerivedDataBackendAsyncPutWrapper::Put(
 			{
 				if (!bCancel)
 				{
-					FRequestGroup BlockingGroup = Factory.CreateGroup(EPriority::Blocking);
-					InnerBackend->Put(Records, Context, Policy, BlockingGroup, MoveTemp(OnComplete));
-					BlockingGroup.Wait();
+					FRequestOwner BlockingOwner(EPriority::Blocking);
+					InnerBackend->Put(Records, Context, Policy, BlockingOwner, MoveTemp(OnComplete));
+					BlockingOwner.Wait();
 				}
 				else if (OnComplete)
 				{
@@ -513,9 +513,9 @@ void FDerivedDataBackendAsyncPutWrapper::Get(
 			{
 				if (!bCancel)
 				{
-					FRequestGroup BlockingGroup = Factory.CreateGroup(EPriority::Blocking);
-					InnerBackend->Get(Keys, Context, Policy, BlockingGroup, MoveTemp(OnComplete));
-					BlockingGroup.Wait();
+					FRequestOwner BlockingOwner(EPriority::Blocking);
+					InnerBackend->Get(Keys, Context, Policy, BlockingOwner, MoveTemp(OnComplete));
+					BlockingOwner.Wait();
 				}
 				else if (OnComplete)
 				{
@@ -547,9 +547,9 @@ void FDerivedDataBackendAsyncPutWrapper::GetPayload(
 			{
 				if (!bCancel)
 				{
-					FRequestGroup BlockingGroup = Factory.CreateGroup(EPriority::Blocking);
-					InnerBackend->GetPayload(Keys, Context, Policy, BlockingGroup, MoveTemp(OnComplete));
-					BlockingGroup.Wait();
+					FRequestOwner BlockingOwner(EPriority::Blocking);
+					InnerBackend->GetPayload(Keys, Context, Policy, BlockingOwner, MoveTemp(OnComplete));
+					BlockingOwner.Wait();
 				}
 				else if (OnComplete)
 				{
diff --git a/Engine/Source/Developer/DerivedDataCache/Private/DerivedDataBackends.cpp b/Engine/Source/Developer/DerivedDataCache/Private/DerivedDataBackends.cpp
--- a/Engine/Source/Developer/DerivedDataCache/Private/DerivedDataBackends.cpp
+++ b/Engine/Source/Developer/DerivedDataCache/Private/DerivedDataBackends.cpp
@@ -23,7 +23,6 @@
 #include "DerivedDataBackendCorruptionWrapper.h"
 #include "DerivedDataBackendThrottleWrapper.h"
 #include "DerivedDataBackendVerifyWrapper.h"
-#include "DerivedDataRequest.h"
 #include "DerivedDataUtilsInterface.h"
 #include "Misc/EngineBuildSettings.h"
 #include "Modules/ModuleManager.h"
@@ -1176,11 +1175,6 @@ public:
 		return MakeShared<FDerivedDataCacheStatsNode>(nullptr, TEXT(""));
 	}
 
-	virtual FRequestGroup CreateGroup(EPriority Priority) override
-	{
-		return Private::CreateRequestGroup(Priority);
-	}
-
 	virtual FCacheBucket CreateBucket(FStringView Name) override
 	{
 		return Private::CreateCacheBucket(Name);
diff --git a/Engine/Source/Developer/DerivedDataCache/Private/DerivedDataBuild.cpp b/Engine/Source/Developer/DerivedDataCache/Private/DerivedDataBuild.cpp
--- a/Engine/Source/Developer/DerivedDataCache/Private/DerivedDataBuild.cpp
+++ b/Engine/Source/Developer/DerivedDataCache/Private/DerivedDataBuild.cpp
@@ -12,7 +12,6 @@
 #include "DerivedDataBuildSession.h"
 #include "DerivedDataBuildWorkerRegistry.h"
 #include "DerivedDataCache.h"
-#include "DerivedDataRequest.h"
 #include "Misc/Guid.h"
 
 namespace UE::DerivedData::Private
@@ -58,7 +57,7 @@ DEFINE_LOG_CATEGORY(LogDerivedDataBuild);
  * - Resolves FBuildDefinition to Key->(RawHash, RawSize, OptionalBuffer)
  * IBuildJob:
  * - From FBuildKey or FBuildDefinition or FBuildAction+FBuildInputs via FBuildSession
- * - Represents an executing build request and implements IRequest
+ * - Represents an executing build request
  * IBuildScheduler:
  * - Schedules execution of the operations on IBuildJob
  * FBuildSession:
@@ -73,11 +72,6 @@ public:
 	{
 	}
 
-	FRequestGroup CreateGroup(EPriority Priority) final
-	{
-		return CreateRequestGroup(Priority);
-	}
-
 	FBuildDefinitionBuilder CreateDefinition(FStringView Name, FStringView Function) final
 	{
 		return CreateBuildDefinition(Name, Function);
diff --git a/Engine/Source/Developer/DerivedDataCache/Private/DerivedDataCache.cpp b/Engine/Source/Developer/DerivedDataCache/Private/DerivedDataCache.cpp
--- a/Engine/Source/Developer/DerivedDataCache/Private/DerivedDataCache.cpp
+++ b/Engine/Source/Developer/DerivedDataCache/Private/DerivedDataCache.cpp
@@ -14,7 +14,6 @@
 #include "DerivedDataCacheRecord.h"
 #include "DerivedDataCacheUsageStats.h"
 #include "DerivedDataPluginInterface.h"
-#include "DerivedDataRequest.h"
 #include "HAL/ThreadSafeCounter.h"
 #include "Misc/CoreMisc.h"
 #include "Misc/CommandLine.h"
@@ -684,8 +683,6 @@ private:
 public:
 	// ICache Interface
 
-	FRequestGroup CreateGroup(EPriority Priority) final { return CreateRequestGroup(Priority); }
-
 	FCacheBucket CreateBucket(FStringView Name) final { return CreateCacheBucket(Name); }
 
 	FCacheRecordBuilder CreateRecord(const FCacheKey& Key) final { return CreateCacheRecordBuilder(Key); }
diff --git a/Engine/Source/Developer/DerivedDataCache/Private/DerivedDataRequest.cpp b/Engine/Source/Developer/DerivedDataCache/Private/DerivedDataRequest.cpp
--- a/Engine/Source/Developer/DerivedDataCache/Private/DerivedDataRequest.cpp
+++ b/Engine/Source/Developer/DerivedDataCache/Private/DerivedDataRequest.cpp
@@ -19,11 +19,11 @@ struct FRequestBarrierEvent : public FRefCountBase
 	FEventRef Event{EEventMode::ManualReset};
 };
 
-class FRequestGroupOwner final : public FRequestGroupBase
+class FRequestOwnerShared final : public FRequestOwnerSharedBase
 {
 public:
-	explicit FRequestGroupOwner(EPriority Priority);
-	~FRequestGroupOwner() final = default;
+	explicit FRequestOwnerShared(EPriority Priority);
+	~FRequestOwnerShared() final = default;
 
 	void Begin(IRequest* Request) final;
 	TRefCountPtr<IRequest> End(IRequest* Request) final;
@@ -54,7 +54,7 @@ private:
 	bool bKeepAlive{false};
 };
 
-FRequestGroupOwner::FRequestGroupOwner(EPriority NewPriority)
+FRequestOwnerShared::FRequestOwnerShared(EPriority NewPriority)
 	: Priority(NewPriority)
 	, bPriorityChangedInBarrier(false)
 	, bBeginExecuted(false)
@@ -62,7 +62,7 @@ FRequestGroupOwner::FRequestGroupOwner(EPriority NewPriority)
 	AddRef(); // Release is called by Destroy.
 }
 
-void FRequestGroupOwner::Begin(IRequest* Request)
+void FRequestOwnerShared::Begin(IRequest* Request)
 {
 	AddRef();
 	EPriority NewPriority;
@@ -92,7 +92,7 @@ void FRequestGroupOwner::Begin(IRequest* Request)
 	}
 }
 
-TRefCountPtr<IRequest> FRequestGroupOwner::End(IRequest* Request)
+TRefCountPtr<IRequest> FRequestOwnerShared::End(IRequest* Request)
 {
 	ON_SCOPE_EXIT { Release(); };
 	FWriteScopeLock WriteLock(Lock);
@@ -104,14 +104,14 @@ TRefCountPtr<IRequest> FRequestGroupOwner::End(IRequest* Request)
 	return RequestRef;
 }
 
-void FRequestGroupOwner::BeginBarrier()
+void FRequestOwnerShared::BeginBarrier()
 {
 	AddRef();
 	FWriteScopeLock WriteLock(Lock);
 	++BarrierCount;
 }
 
-void FRequestGroupOwner::EndBarrier()
+void FRequestOwnerShared::EndBarrier()
 {
 	ON_SCOPE_EXIT { Release(); };
 	TRefCountPtr<FRequestBarrierEvent> LocalBarrierEvent;
@@ -127,7 +127,7 @@ void FRequestGroupOwner::EndBarrier()
 	}
 }
 
-void FRequestGroupOwner::SetPriority(EPriority NewPriority)
+void FRequestOwnerShared::SetPriority(EPriority NewPriority)
 {
 	TArray<TRefCountPtr<IRequest>, TInlineAllocator<16>> LocalRequests;
 	if (FWriteScopeLock WriteLock(Lock); Priority == NewPriority)
@@ -147,7 +147,7 @@ void FRequestGroupOwner::SetPriority(EPriority NewPriority)
 	}
 }
 
-void FRequestGroupOwner::Cancel()
+void FRequestOwnerShared::Cancel()
 {
 	for (;;)
 	{
@@ -187,7 +187,7 @@ void FRequestGroupOwner::Cancel()
 	}
 }
 
-void FRequestGroupOwner::Wait()
+void FRequestOwnerShared::Wait()
 {
 	for (;;)
 	{
@@ -226,19 +226,19 @@ void FRequestGroupOwner::Wait()
 	}
 }
 
-bool FRequestGroupOwner::Poll() const
+bool FRequestOwnerShared::Poll() const
 {
 	FReadScopeLock ReadLock(Lock);
 	return Requests.IsEmpty() && BarrierCount == 0;
 }
 
-void FRequestGroupOwner::KeepAlive()
+void FRequestOwnerShared::KeepAlive()
 {
 	FWriteScopeLock WriteLock(Lock);
 	bKeepAlive = true;
 }
 
-void FRequestGroupOwner::Destroy()
+void FRequestOwnerShared::Destroy()
 {
 	const bool bLocalKeepAlive = (FWriteScopeLock(Lock), bKeepAlive);
 	if (!bLocalKeepAlive)
@@ -248,9 +248,14 @@ void FRequestGroupOwner::Destroy()
 	Release();
 }
 
-FRequestGroup CreateRequestGroup(EPriority Priority)
+} // UE::DerivedData::Private
+
+namespace UE::DerivedData
+{
+
+FRequestOwner::FRequestOwner(EPriority Priority)
+	: Owner(new Private::FRequestOwnerShared(Priority))
 {
-	return FRequestGroup(new FRequestGroupOwner(Priority));
 }
 
-} // UE::DerivedData::Private
+} // UE::DerivedData
diff --git a/Engine/Source/Developer/DerivedDataCache/Private/HierarchicalDerivedDataBackend.h b/Engine/Source/Developer/DerivedDataCache/Private/HierarchicalDerivedDataBackend.h
--- a/Engine/Source/Developer/DerivedDataCache/Private/HierarchicalDerivedDataBackend.h
+++ b/Engine/Source/Developer/DerivedDataCache/Private/HierarchicalDerivedDataBackend.h
@@ -473,9 +473,9 @@ public:
 		IRequestOwner& Owner,
 		FOnCachePutComplete&& OnComplete) override
 	{
-		FRequestGroup BlockingGroup = Factory.CreateGroup(EPriority::Blocking);
-		FRequestGroup AsyncGroup = Factory.CreateGroup(Owner.GetPriority());
-		AsyncGroup.KeepAlive();
+		FRequestOwner AsyncOwner(Owner.GetPriority());
+		FRequestBarrier AsyncBarrier(AsyncOwner);
+		AsyncOwner.KeepAlive();
 
 		TSet<FCacheKey> RecordsOk;
 
@@ -489,7 +489,8 @@ public:
 					// Every record must put synchronously before switching to async calls.
 					if (RecordsOk.Num() < Records.Num())
 					{
-						InnerBackends[PutCacheIndex]->Put(Records, Context, Policy, BlockingGroup,
+						FRequestOwner BlockingOwner(EPriority::Blocking);
+						InnerBackends[PutCacheIndex]->Put(Records, Context, Policy, BlockingOwner,
 							[&OnComplete, &RecordsOk](FCachePutCompleteParams&& Params)
 							{
 								if (Params.Status == EStatus::Ok)
@@ -502,11 +503,11 @@ public:
 									}
 								}
 							});
-						BlockingGroup.Wait();
+						BlockingOwner.Wait();
 					}
 					else
 					{
-						AsyncPutInnerBackends[PutCacheIndex]->Put(Records, Context, Policy, AsyncGroup);
+						AsyncPutInnerBackends[PutCacheIndex]->Put(Records, Context, Policy, AsyncOwner);
 					}
 				}
 			}
@@ -537,7 +538,6 @@ public:
 		const bool bStoreRemote = bHasWritableRemoteBackends && EnumHasAnyFlags(Policy, ECachePolicy::StoreRemote);
 		const bool bStoreLocalCopy = bStoreLocal && bHasMultipleLocalBackends && !EnumHasAnyFlags(Policy, ECachePolicy::SkipLocalCopy);
 
-		FRequestGroup BlockingGroup = Factory.CreateGroup(EPriority::Blocking);
 		TArray<FCacheKey, TInlineAllocator<16>> RemainingKeys(Keys);
 
 		TSet<FCacheKey> KeysOk;
@@ -557,15 +557,16 @@ public:
 				const ECachePolicy FillPolicy = bFill ? (Policy & ~ECachePolicy::SkipData) : Policy;
 
 				// Block on this because backends in this hierarchy are not expected to be asynchronous.
-				InnerBackends[GetCacheIndex]->Get(RemainingKeys, Context, Policy, BlockingGroup,
+				FRequestOwner BlockingOwner(EPriority::Blocking);
+				InnerBackends[GetCacheIndex]->Get(RemainingKeys, Context, Policy, BlockingOwner,
 					[this, Context, GetCacheIndex, bStoreLocal, bStoreRemote, bStoreLocalCopy, bIsLocalGet, bFill, &Owner, &OnComplete, &KeysOk](FCacheGetCompleteParams&& Params)
 					{
 						if (Params.Status == EStatus::Ok)
 						{
 							if (bFill)
 							{
-								FRequestGroup AsyncGroup = Factory.CreateGroup(Owner.GetPriority());
-								AsyncGroup.KeepAlive();
+								FRequestOwner AsyncOwner(Owner.GetPriority());
+								AsyncOwner.KeepAlive();
 								for (int32 FillCacheIndex = 0; FillCacheIndex < InnerBackends.Num(); ++FillCacheIndex)
 								{
 									if (GetCacheIndex != FillCacheIndex)
@@ -573,7 +574,7 @@ public:
 										const bool bIsLocalFill = InnerBackends[FillCacheIndex]->GetSpeedClass() == ESpeedClass::Local;
 										if ((bIsLocalFill && bStoreLocal && bIsLocalGet <= bStoreLocalCopy) || (!bIsLocalFill && bStoreRemote))
 										{
-											AsyncPutInnerBackends[FillCacheIndex]->Put({Params.Record}, Context, ECachePolicy::Default, AsyncGroup);
+											AsyncPutInnerBackends[FillCacheIndex]->Put({Params.Record}, Context, ECachePolicy::Default, AsyncOwner);
 										}
 									}
 								}
@@ -586,7 +587,7 @@ public:
 							}
 						}
 					});
-				BlockingGroup.Wait();
+				BlockingOwner.Wait();
 
 				RemainingKeys.RemoveAll([&KeysOk](const FCacheKey& Key) { return KeysOk.Contains(Key); });
 			}
@@ -608,7 +609,6 @@ public:
 		IRequestOwner& Owner,
 		FOnCacheGetPayloadComplete&& OnComplete) override
 	{
-		FRequestGroup BlockingGroup = Factory.CreateGroup(EPriority::Blocking);
 		TArray<FCachePayloadKey, TInlineAllocator<16>> RemainingKeys(Keys);
 
 		TSet<FCachePayloadKey> KeysOk;
@@ -618,7 +618,8 @@ public:
 
 			for (int32 GetCacheIndex = 0; GetCacheIndex < InnerBackends.Num() && !RemainingKeys.IsEmpty(); ++GetCacheIndex)
 			{
-				InnerBackends[GetCacheIndex]->GetPayload(RemainingKeys, Context, Policy, BlockingGroup,
+				FRequestOwner BlockingOwner(EPriority::Blocking);
+				InnerBackends[GetCacheIndex]->GetPayload(RemainingKeys, Context, Policy, BlockingOwner,
 					[&OnComplete, &KeysOk](FCacheGetPayloadCompleteParams&& Params)
 					{
 						if (Params.Status == EStatus::Ok)
@@ -630,7 +631,7 @@ public:
 							}
 						}
 					});
-				BlockingGroup.Wait();
+				BlockingOwner.Wait();
 
 				RemainingKeys.RemoveAll([&KeysOk](const FCachePayloadKey& Key) { return KeysOk.Contains(Key); });
 			}
diff --git a/Engine/Source/Developer/DerivedDataCache/Public/DerivedDataBuild.h b/Engine/Source/Developer/DerivedDataCache/Public/DerivedDataBuild.h
--- a/Engine/Source/Developer/DerivedDataCache/Public/DerivedDataBuild.h
+++ b/Engine/Source/Developer/DerivedDataCache/Public/DerivedDataBuild.h
@@ -18,7 +18,6 @@ namespace UE::DerivedData { class FCacheRecord; }
 namespace UE::DerivedData { class FOptionalBuildAction; }
 namespace UE::DerivedData { class FOptionalBuildDefinition; }
 namespace UE::DerivedData { class FOptionalBuildOutput; }
-namespace UE::DerivedData { class FRequestGroup; }
 namespace UE::DerivedData { class IBuildFunctionRegistry; }
 namespace UE::DerivedData { class IBuildInputResolver; }
 namespace UE::DerivedData { class IBuildScheduler; }
@@ -50,8 +49,6 @@ class IBuild
 public:
 	virtual ~IBuild() = default;
 
-	virtual FRequestGroup CreateGroup(EPriority Priority) = 0;
-
 	/**
 	 * Create a build definition builder.
 	 *
diff --git a/Engine/Source/Developer/DerivedDataCache/Public/DerivedDataCache.h b/Engine/Source/Developer/DerivedDataCache/Public/DerivedDataCache.h
--- a/Engine/Source/Developer/DerivedDataCache/Public/DerivedDataCache.h
+++ b/Engine/Source/Developer/DerivedDataCache/Public/DerivedDataCache.h
@@ -20,7 +20,6 @@ class FCacheRecord;
 class FCacheRecordBuilder;
 class FOptionalCacheRecord;
 class FPayload;
-class FRequestGroup;
 struct FCacheGetCompleteParams;
 struct FCacheGetPayloadCompleteParams;
 struct FCachePutCompleteParams;
@@ -105,8 +104,6 @@ class ICacheFactory
 public:
 	virtual ~ICacheFactory() = default;
 
-	virtual FRequestGroup CreateGroup(EPriority Priority) = 0;
-
 	/**
 	 * Create a cache bucket from a name.
 	 *
diff --git a/Engine/Source/Developer/DerivedDataCache/Public/DerivedDataRequest.h b/Engine/Source/Developer/DerivedDataCache/Public/DerivedDataRequest.h
--- a/Engine/Source/Developer/DerivedDataCache/Public/DerivedDataRequest.h
+++ b/Engine/Source/Developer/DerivedDataCache/Public/DerivedDataRequest.h
@@ -3,15 +3,15 @@
 #pragma once
 
 #include "CoreTypes.h"
+#include "Templates/Invoke.h"
 #include "Templates/RefCounting.h"
 #include "Templates/UniquePtr.h"
 #include "Templates/UnrealTemplate.h"
 #include <atomic>
 
-namespace UE::DerivedData { class FRequestGroup; }
-namespace UE::DerivedData { enum class EPriority : uint8; }
-namespace UE::DerivedData::Private { class FRequestGroupBase; }
-namespace UE::DerivedData::Private { FRequestGroup CreateRequestGroup(EPriority Priority); }
+#define UE_API DERIVEDDATACACHE_API
+
+namespace UE::DerivedData::Private { class FRequestOwnerSharedBase; }
 
 namespace UE::DerivedData
 {
@@ -66,14 +66,14 @@ enum class EStatus : uint8
 /**
  * Interface to an asynchronous request that can be prioritized or canceled.
  *
- * Use IRequestOwner, typically FRequestGroup, to reference requests between its Begin and End.
+ * Use IRequestOwner, typically FRequestOwner, to reference requests between its Begin and End.
  *
  * Requests typically invoke a callback on completion, and must not return from Cancel or Wait on
  * any thread until any associated callbacks have finished executing. This property is crucial to
  * allowing requests to be chained or nested by creating new requests from within callbacks.
  *
  * @see FRequestBase
- * @see FRequestGroup
+ * @see FRequestOwner
  * @see IRequestOwner
  */
 class IRequest
@@ -148,7 +148,7 @@ private:
  * The owner is responsible for keeping itself alive while it has active requests or barriers.
  *
  * @see FRequestBarrier
- * @see FRequestGroup
+ * @see FRequestOwner
  */
 class IRequestOwner
 {
@@ -216,8 +216,8 @@ private:
 	IRequestOwner& Owner;
 };
 
-/** Private base for a request group, which is both an owner and a request. */
-class Private::FRequestGroupBase : public IRequestOwner, public FRequestBase
+/** Private base for a shared request owner, which is both an owner and a request. */
+class Private::FRequestOwnerSharedBase : public IRequestOwner, public FRequestBase
 {
 public:
 	virtual bool Poll() const = 0;
@@ -226,81 +226,47 @@ public:
 };
 
 /**
- * A request group is both a request owner and a request.
+ * A concrete request owner that also presents as a request.
  *
- * Request groups may be moved but not copied, and cancel any outstanding requests on destruction
+ * Request owners may be moved but not copied, and cancel any outstanding requests on destruction
  * unless KeepAlive has been called.
  *
- * Create with ICache::CreateGroup or IBuild::CreateGroup.
- *
  * @see IRequestOwner
  */
-class FRequestGroup
+class FRequestOwner
 {
 public:
-	FRequestGroup(FRequestGroup&&) = default;
-	FRequestGroup& operator=(FRequestGroup&&) = default;
+	/** Construct a request owner with the given priority. */
+	UE_API explicit FRequestOwner(EPriority Priority);
+
+	FRequestOwner(FRequestOwner&&) = default;
+	FRequestOwner& operator=(FRequestOwner&&) = default;
 
-	FRequestGroup(const FRequestGroup&) = delete;
-	FRequestGroup& operator=(const FRequestGroup&) = delete;
+	FRequestOwner(const FRequestOwner&) = delete;
+	FRequestOwner& operator=(const FRequestOwner&) = delete;
 
-	/** Keep requests in the group alive until complete, even after destruction of the group. */
+	/** Keep requests in the owner alive until complete, even after destruction of the owner. */
 	inline void KeepAlive() { Owner->KeepAlive(); }
 
 	/** Returns the priority that new requests are expected to inherit. */
 	inline EPriority GetPriority() const { return Owner->GetPriority(); }
-	/** Set the priority of active and future requests in the group. */
+	/** Set the priority of active and future requests in the owner. */
 	inline void SetPriority(EPriority Priority) { Owner->SetPriority(Priority); }
-	/** Cancel any active and future requests in the group. */
+	/** Cancel any active and future requests in the owner. */
 	inline void Cancel() { Owner->Cancel(); }
-	/** Wait for any active and future requests or barriers in the group. */
+	/** Wait for any active and future requests or barriers in the owner. */
 	inline void Wait() { Owner->Wait(); }
-	/** Poll whether the group has any active requests or barriers. */
+	/** Poll whether the owner has any active requests or barriers. */
 	inline bool Poll() const { return Owner->Poll(); }
 
-	/** Access the group as a request owner. */
+	/** Access as a request owner. */
 	inline operator IRequestOwner&() { return *Owner; }
-	/** Access the group as a request. */
+	/** Access as a request. */
 	inline operator IRequest*() { return Owner.Get(); }
 
 private:
-	friend class FOptionalRequestGroup;
-	friend FRequestGroup Private::CreateRequestGroup(EPriority Priority);
-
-	/** Construct a request group. Use CreateGroup() on ICache or IBuild. */
-	inline explicit FRequestGroup(Private::FRequestGroupBase* InOwner)
-		: Owner(InOwner)
-	{
-	}
-
-	struct FDeleteOwner { void operator()(Private::FRequestGroupBase* O) const { if (O) { O->Destroy(); } } };
-	TUniquePtr<Private::FRequestGroupBase, FDeleteOwner> Owner;
-};
-
-/**
- * A request group that can be null.
- *
- * @see FRequestGroup
- */
-class FOptionalRequestGroup : private FRequestGroup
-{
-public:
-	inline FOptionalRequestGroup() : FRequestGroup(nullptr) {}
-
-	inline FOptionalRequestGroup(FRequestGroup&& InDefinition) : FRequestGroup(MoveTemp(InDefinition)) {}
-	inline FOptionalRequestGroup& operator=(FRequestGroup&& InDefinition) { FRequestGroup::operator=(MoveTemp(InDefinition)); return *this; }
-
-	/** Returns the request group. The caller must check for null before using this accessor. */
-	inline FRequestGroup& Get() & { return *this; }
-	inline FRequestGroup&& Get() && { return MoveTemp(*this); }
-	inline const FRequestGroup& Get() const & { return *this; }
-
-	inline bool IsNull() const { return !IsValid(); }
-	inline bool IsValid() const { return Owner.IsValid(); }
-	inline explicit operator bool() { return IsValid(); }
-	inline explicit operator bool() const { return IsValid(); }
-
-	inline void Reset() { *this = FOptionalRequestGroup(); }
+	struct FDeleteOwner { void operator()(Private::FRequestOwnerSharedBase* O) const { if (O) { O->Destroy(); } } };
+	TUniquePtr<Private::FRequestOwnerSharedBase, FDeleteOwner> Owner;
 };
 
 template <typename CallbackType, typename... CallbackArgTypes>
@@ -313,3 +279,5 @@ TRefCountPtr<IRequest> IRequestOwner::End(IRequest* Request, CallbackType&& Call
 }
 
 } // UE::DerivedData
+
+#undef UE_API
diff --git a/Engine/Source/Editor/UnrealEd/Private/DerivedDataBuildRemoteExecutor.cpp b/Engine/Source/Editor/UnrealEd/Private/DerivedDataBuildRemoteExecutor.cpp
--- a/Engine/Source/Editor/UnrealEd/Private/DerivedDataBuildRemoteExecutor.cpp
+++ b/Engine/Source/Editor/UnrealEd/Private/DerivedDataBuildRemoteExecutor.cpp
@@ -588,8 +588,8 @@ void FRemoteBuildExecutionRequest::BuildMerkleTreeNodes()
 			WorkerFileMeta.Emplace(Path, false);
 		});
 
-	FRequestGroup BlockingGroup = State.BuildSystem.CreateGroup(EPriority::Blocking);
-	State.BuildWorker.FindFileData(WorkerFileHashes, BlockingGroup,
+	FRequestOwner BlockingOwner(EPriority::Blocking);
+	State.BuildWorker.FindFileData(WorkerFileHashes, BlockingOwner,
 		[this, &WorkerFileMeta] (FBuildWorkerFileDataCompleteParams&& Params)
 		{
 			uint32 MetaIndex = 0;
@@ -601,7 +601,7 @@ void FRemoteBuildExecutionRequest::BuildMerkleTreeNodes()
 				++MetaIndex;
 			}
 		});
-	BlockingGroup.Wait();
+	BlockingOwner.Wait();
 
 	State.BuildAction.IterateInputs([this] (FStringView Key, const FIoHash& RawHash, uint64 RawSize)
 		{
@@ -846,8 +846,8 @@ void FRemoteBuildExecutionRequest::LoadMissingWorkerFileBlobsAsync()
 		}
 	}
 
-	FRequestGroup BlockingGroup = State.BuildSystem.CreateGroup(EPriority::Blocking);
-	State.BuildWorker.FindFileData(WorkerFileHashes, BlockingGroup,
+	FRequestOwner BlockingOwner(EPriority::Blocking);
+	State.BuildWorker.FindFileData(WorkerFileHashes, BlockingOwner,
 		[this, &WorkerFileMapping] (FBuildWorkerFileDataCompleteParams&& Params)
 		{
 			uint32 MetaIndex = 0;
@@ -863,7 +863,7 @@ void FRemoteBuildExecutionRequest::LoadMissingWorkerFileBlobsAsync()
 				}
 			}
 		});
-	BlockingGroup.Wait();
+	BlockingOwner.Wait();
 }
 
 TFuture<TPair<FStatus, FBatchUpdateBlobsResponse>> FRemoteBuildExecutionRequest::UploadMissingBlobsAsync()
diff --git a/Engine/Source/Editor/UnrealEd/Private/EditorDomain/BulkDataRegistryEditorDomain.cpp b/Engine/Source/Editor/UnrealEd/Private/EditorDomain/BulkDataRegistryEditorDomain.cpp
--- a/Engine/Source/Editor/UnrealEd/Private/EditorDomain/BulkDataRegistryEditorDomain.cpp
+++ b/Engine/Source/Editor/UnrealEd/Private/EditorDomain/BulkDataRegistryEditorDomain.cpp
@@ -395,7 +395,8 @@ void FBulkDataRegistryEditorDomain::ReadPayloadIdsFromCache(FName PackageName, T
 	{
 		// Creation of the Request has to occur outside of any lock, because the request
 		// may execute immediately on this thread and need to enter the lock; our locks are non-reentrant
-		UE::EditorDomain::GetBulkDataPayloadId(PackageName, NewPending->GetBulkDataId(), NewPending->GetRequestGroup(),
+		UE::DerivedData::FRequestBarrier Barrier(NewPending->GetRequestOwner());
+		UE::EditorDomain::GetBulkDataPayloadId(PackageName, NewPending->GetBulkDataId(), NewPending->GetRequestOwner(),
 			[this, PackageName, NewPending](FSharedBuffer Buffer)
 		{
 			if (Buffer.IsNull())
@@ -484,7 +485,7 @@ void FBulkDataRegistryEditorDomain::ReadPayloadIdsFromCache(FName PackageName, T
 
 FPendingPackage::FPendingPackage(FName InPackageName, FBulkDataRegistryEditorDomain* InOwner)
 	: PackageName(InPackageName)
-	, BulkDataListCacheRequest(UE::DerivedData::GetCache().CreateGroup(UE::DerivedData::EPriority::Low))
+	, BulkDataListCacheRequest(UE::DerivedData::EPriority::Low)
 	, Owner(InOwner)
 {
 	PendingOperations = Flag_EndLoad | Flag_BulkDataListResults;
@@ -538,7 +539,7 @@ void FPendingPackage::OnBulkDataListResults(FSharedBuffer Buffer)
 		WriteCache();
 
 		// Removing *this from its owner will delete it, which will attempt to Cancel this request.
-		// Direct the group to keep requests alive to avoid a deadlock.
+		// Direct the owner to keep requests alive to avoid a deadlock.
 		BulkDataListCacheRequest.KeepAlive();
 
 		FScopeLock PendingPackageScopeLock(&Owner->PendingPackageLock);
@@ -732,7 +733,7 @@ void FUpdatePayloadWorker::DoWork()
 
 FPendingPayloadId::FPendingPayloadId(const FGuid& InBulkDataId)
 	: BulkDataId(InBulkDataId)
-	, Request(UE::DerivedData::GetCache().CreateGroup(UE::DerivedData::EPriority::Low))
+	, Request(UE::DerivedData::EPriority::Low)
 {
 }
 
diff --git a/Engine/Source/Editor/UnrealEd/Private/EditorDomain/BulkDataRegistryEditorDomain.h b/Engine/Source/Editor/UnrealEd/Private/EditorDomain/BulkDataRegistryEditorDomain.h
--- a/Engine/Source/Editor/UnrealEd/Private/EditorDomain/BulkDataRegistryEditorDomain.h
+++ b/Engine/Source/Editor/UnrealEd/Private/EditorDomain/BulkDataRegistryEditorDomain.h
@@ -74,7 +74,7 @@ private:
 	FName PackageName;
 	TArray<UE::Virtualization::FVirtualizedUntypedBulkData> BulkDatas;
 	TArray<UE::Virtualization::FVirtualizedUntypedBulkData> CachedBulkDatas;
-	UE::DerivedData::FRequestGroup BulkDataListCacheRequest;
+	UE::DerivedData::FRequestOwner BulkDataListCacheRequest;
 	FBulkDataRegistryEditorDomain* Owner;
 	/**
 	 * When PendingOperations reaches zero, we can remove the FPendingPackage.
@@ -143,7 +143,7 @@ public:
 
 	void Cancel();
 
-	UE::DerivedData::FRequestGroup& GetRequestGroup()
+	UE::DerivedData::FRequestOwner& GetRequestOwner()
 	{
 		return Request;
 	}
@@ -154,7 +154,7 @@ public:
 
 private:
 	FGuid BulkDataId;
-	UE::DerivedData::FRequestGroup Request;
+	UE::DerivedData::FRequestOwner Request;
 };
 
 /** Implementation of a BulkDataRegistry that stores its persistent data in a DDC bucket. */
diff --git a/Engine/Source/Editor/UnrealEd/Private/EditorDomain/EditorDomain.cpp b/Engine/Source/Editor/UnrealEd/Private/EditorDomain/EditorDomain.cpp
--- a/Engine/Source/Editor/UnrealEd/Private/EditorDomain/EditorDomain.cpp
+++ b/Engine/Source/Editor/UnrealEd/Private/EditorDomain/EditorDomain.cpp
@@ -233,7 +233,7 @@ int64 FEditorDomain::FileSize(const FPackagePath& PackagePath, EPackageSegment P
 		return Workspace->FileSize(PackagePath, PackageSegment, OutUpdatedPath);
 	}
 
-	UE::DerivedData::FOptionalRequestGroup RequestGroup;
+	TOptional<UE::DerivedData::FRequestOwner> Owner;
 	int64 FileSize = -1;
 	{
 		FScopeLock ScopeLock(&Locks->Lock);
@@ -276,10 +276,10 @@ int64 FEditorDomain::FileSize(const FPackagePath& PackagePath, EPackageSegment P
 		};
 		// Fetch meta-data only
 		ECachePolicy SkipFlags = ECachePolicy::SkipData & ~ECachePolicy::SkipMeta;
-		RequestGroup = GetCache().CreateGroup(EPriority::Highest);
-		RequestEditorDomainPackage(PackagePath, PackageSource->Digest, SkipFlags, RequestGroup.Get(), MoveTemp(MetaDataGetComplete));
+		Owner.Emplace(EPriority::Highest);
+		RequestEditorDomainPackage(PackagePath, PackageSource->Digest, SkipFlags, *Owner, MoveTemp(MetaDataGetComplete));
 	}
-	RequestGroup.Get().Wait();
+	Owner->Wait();
 	return FileSize;
 }
 
diff --git a/Engine/Source/Editor/UnrealEd/Private/EditorDomain/EditorDomainArchive.cpp b/Engine/Source/Editor/UnrealEd/Private/EditorDomain/EditorDomainArchive.cpp
--- a/Engine/Source/Editor/UnrealEd/Private/EditorDomain/EditorDomainArchive.cpp
+++ b/Engine/Source/Editor/UnrealEd/Private/EditorDomain/EditorDomainArchive.cpp
@@ -19,7 +19,7 @@
 FEditorDomainPackageSegments::FEditorDomainPackageSegments(const TRefCountPtr<FEditorDomain::FLocks>& InLocks,
 	const FPackagePath& InPackagePath, const TRefCountPtr<FEditorDomain::FPackageSource>& InPackageSource)
 	: EditorDomainLocks(InLocks)
-	, RequestGroup(UE::DerivedData::GetCache().CreateGroup(UE::DerivedData::EPriority::Normal))
+	, RequestOwner(UE::DerivedData::EPriority::Normal)
 	, PackagePath(InPackagePath)
 	, PackageSource(InPackageSource)
 	, PackageDigest(InPackageSource->Digest)
@@ -47,7 +47,7 @@ void FEditorDomainPackageSegments::WaitForReady() const
 	{
 		return;
 	}
-	const_cast<UE::DerivedData::FRequestGroup&>(RequestGroup).Wait();
+	const_cast<UE::DerivedData::FRequestOwner&>(RequestOwner).Wait();
 	Source = AsyncSource;
 }
 
@@ -76,12 +76,12 @@ void FEditorDomainPackageSegments::Close()
 {
 	// Called from Interface-only
 	AsyncSource = ESource::Closed;
-	RequestGroup.Cancel();
+	RequestOwner.Cancel();
 	for (FSegment& Segment : Segments)
 	{
-		if (Segment.RequestGroup)
+		if (Segment.RequestOwner)
 		{
-			Segment.RequestGroup.Get().Cancel();
+			Segment.RequestOwner->Cancel();
 		}
 	}
 
@@ -196,7 +196,7 @@ FEditorDomainPackageSegments::FSegment* FEditorDomainPackageSegments::EnsureSegm
 	{
 		if (!Segment->bComplete)
 		{
-			if (!Segment->RequestGroup.IsValid())
+			if (!Segment->RequestOwner)
 			{
 				SendSegmentRequest(*Segment);
 			}
@@ -219,9 +219,9 @@ FEditorDomainPackageSegments::FSegment* FEditorDomainPackageSegments::EnsureSegm
 		{
 			if (!Segment->bComplete)
 			{
-				if (Segment->RequestGroup)
+				if (Segment->RequestOwner)
 				{
-					Segment->RequestGroup.Get().Wait();
+					Segment->RequestOwner->Wait();
 				}
 				check(Segment->bAsyncComplete);
 				Segment->bComplete = true;
@@ -375,10 +375,10 @@ void FEditorDomainPackageSegments::SendSegmentRequest(FSegment& Segment)
 	// Called from Callback-only until AsyncSource is set, then from Interface-only
 	ICache& Cache = GetCache();
 
-	// Note that Segment.RequestGroup is Interface-only and so we can write it outside the lock
-	Segment.RequestGroup = Cache.CreateGroup(UE::DerivedData::EPriority::Normal);
+	// Note that Segment.RequestOwner is Interface-only and so we can write it outside the lock
+	Segment.RequestOwner.Emplace(UE::DerivedData::EPriority::Normal);
 	Cache.GetPayload({ FCachePayloadKey{ UE::EditorDomain::GetEditorDomainPackageKey(PackageDigest), Segment.PayloadId} },
-		PackagePath.GetDebugName(), ECachePolicy::Local, Segment.RequestGroup.Get(),
+		PackagePath.GetDebugName(), ECachePolicy::Local, *Segment.RequestOwner,
 		[this, &Segment](FCacheGetPayloadCompleteParams&& Params)
 		{
 			if (AsyncSource == ESource::Closed)
diff --git a/Engine/Source/Editor/UnrealEd/Private/EditorDomain/EditorDomainArchive.h b/Engine/Source/Editor/UnrealEd/Private/EditorDomain/EditorDomainArchive.h
--- a/Engine/Source/Editor/UnrealEd/Private/EditorDomain/EditorDomainArchive.h
+++ b/Engine/Source/Editor/UnrealEd/Private/EditorDomain/EditorDomainArchive.h
@@ -44,7 +44,7 @@ public:
 	~FEditorDomainPackageSegments();
 
 	/** Get the request owner for requests that will feed this archive. */
-	UE::DerivedData::IRequestOwner& GetRequestOwner() { return RequestGroup; }
+	UE::DerivedData::IRequestOwner& GetRequestOwner() { return RequestOwner; }
 	/** Callback from the request of the Record; set whether we're reading from EditorDomain bytes or WorkspaceDomain archive. */
 	void OnRecordRequestComplete(UE::DerivedData::FCacheGetCompleteParams&& Params,
 		TUniqueFunction<void(bool bValid)>&& CreateSegmentData,
@@ -102,13 +102,14 @@ private:
 	struct FSegment
 	{
 		FSegment(const UE::DerivedData::FPayloadId& InPayloadId, uint64 InStart)
-			: PayloadId(InPayloadId), Start(InStart)
+			: PayloadId(InPayloadId)
+			, Start(InStart)
 		{}
 
 		/** Attachment ID to request from cache for this segment. Read-only. */
 		UE::DerivedData::FPayloadId PayloadId;
 		/** The request that will be (has been) sent for the bytes from cache. Interface-only. */
-		UE::DerivedData::FOptionalRequestGroup RequestGroup;
+		TOptional<UE::DerivedData::FRequestOwner> RequestOwner;
 		/** Offset from the start of the entire archive. Read-only. */
 		uint64 Start;
 		/** Payload bytes. Callback-only before bAsyncComplete, Interface-only after. */
@@ -143,7 +144,7 @@ private:
 	 */
 	TRefCountPtr<FEditorDomain::FLocks> EditorDomainLocks;
 	/** Initial Request for the list of segments. Read-only. */
-	UE::DerivedData::FRequestGroup RequestGroup;
+	UE::DerivedData::FRequestOwner RequestOwner;
 	/** PackagePath for fallback to WorkspaceDomain. Read-only. */
 	FPackagePath PackagePath;
 	/** Offset to next use in Serialize. Interface-only. */
diff --git a/Engine/Source/Editor/UnrealEd/Private/EditorDomain/EditorDomainUtils.cpp b/Engine/Source/Editor/UnrealEd/Private/EditorDomain/EditorDomainUtils.cpp
--- a/Engine/Source/Editor/UnrealEd/Private/EditorDomain/EditorDomainUtils.cpp
+++ b/Engine/Source/Editor/UnrealEd/Private/EditorDomain/EditorDomainUtils.cpp
@@ -471,9 +471,9 @@ bool TrySavePackage(UPackage* Package)
 	MetaData.EndObject();
 
 	RecordBuilder.SetMeta(MetaData.Save().AsObject());
-	FRequestGroup Group = Cache.CreateGroup(EPriority::Normal);
-	Cache.Put({RecordBuilder.Build()}, Package->GetName(), ECachePolicy::Local, Group);
-	Group.KeepAlive();
+	FRequestOwner Owner(EPriority::Normal);
+	Cache.Put({RecordBuilder.Build()}, Package->GetName(), ECachePolicy::Local, Owner);
+	Owner.KeepAlive();
 	return true;
 }
 
@@ -523,11 +523,11 @@ void PutBulkDataList(FName PackageName, FSharedBuffer Buffer)
 
 	using namespace UE::DerivedData;
 	ICache& Cache = GetCache();
-	FRequestGroup Group = Cache.CreateGroup(EPriority::Normal);
+	FRequestOwner Owner(EPriority::Normal);
 	FCacheRecordBuilder RecordBuilder = Cache.CreateRecord(GetBulkDataListKey(PackageDigest));
 	RecordBuilder.SetValue(Buffer);
-	Cache.Put({RecordBuilder.Build()}, WriteToString<128>(PackageName), ECachePolicy::Default, Group);
-	Group.KeepAlive();
+	Cache.Put({RecordBuilder.Build()}, WriteToString<128>(PackageName), ECachePolicy::Default, Owner);
+	Owner.KeepAlive();
 }
 
 FIoHash GetPackageAndGuidDigest(FCbWriter& Builder, const FGuid& BulkDataId)
@@ -583,10 +583,11 @@ void PutBulkDataPayloadId(FName PackageName, const FGuid& BulkDataId, FSharedBuf
 
 	using namespace UE::DerivedData;
 	ICache& Cache = GetCache();
-	FRequestGroup Group = Cache.CreateGroup(EPriority::Normal);
+	FRequestOwner Owner(EPriority::Normal);
 	FCacheRecordBuilder RecordBuilder = Cache.CreateRecord(GetBulkDataPayloadIdKey(PackageAndGuidDigest));
 	RecordBuilder.SetValue(Buffer);
-	Cache.Put({RecordBuilder.Build()}, WriteToString<128>(PackageName), ECachePolicy::Default, Group);
+	Cache.Put({RecordBuilder.Build()}, WriteToString<128>(PackageName), ECachePolicy::Default, Owner);
+	Owner.KeepAlive();
 }
 
 
diff --git a/Engine/Source/Programs/DerivedDataBuildWorker/Private/DerivedDataBuildWorker.cpp b/Engine/Source/Programs/DerivedDataBuildWorker/Private/DerivedDataBuildWorker.cpp
--- a/Engine/Source/Programs/DerivedDataBuildWorker/Private/DerivedDataBuildWorker.cpp
+++ b/Engine/Source/Programs/DerivedDataBuildWorker/Private/DerivedDataBuildWorker.cpp
@@ -204,11 +204,11 @@ bool FBuildWorkerProgram::Build()
 
 	IBuild& BuildSystem = GetBuild();
 	FBuildSession Session = BuildSystem.CreateSession(TEXT("BuildWorker"_SV), this);
-	FRequestGroup Group = BuildSystem.CreateGroup(EPriority::Normal);
+	FRequestOwner Owner(EPriority::Normal);
 
 	ON_SCOPE_EXIT
 	{
-		Group.Wait();
+		Owner.Wait();
 	};
 
 	if (ActionPaths.IsEmpty())
@@ -237,7 +237,8 @@ bool FBuildWorkerProgram::Build()
 		}
 		else
 		{
-			Session.BuildAction(Action.Get(), {}, EBuildPolicy::BuildLocal, Group,
+			FRequestBarrier Barrier(Owner);
+			Session.BuildAction(Action.Get(), {}, EBuildPolicy::BuildLocal, Owner,
 				[this](FBuildActionCompleteParams&& Params) { BuildComplete(MoveTemp(Params)); });
 		}
 	}
diff --git a/Engine/Source/Runtime/Engine/Private/TextureDerivedDataTask.cpp b/Engine/Source/Runtime/Engine/Private/TextureDerivedDataTask.cpp
--- a/Engine/Source/Runtime/Engine/Private/TextureDerivedDataTask.cpp
+++ b/Engine/Source/Runtime/Engine/Private/TextureDerivedDataTask.cpp
@@ -562,9 +562,9 @@ void FTextureCacheDerivedDataWorker::BuildTexture(bool bReplaceExistingDDC)
 				PlaceholderResolver.Emplace(Texture);
 				InputResolver = &PlaceholderResolver.GetValue();
 			}
-			FRequestGroup Group = Build.CreateGroup(EPriority::Blocking);
+			FRequestOwner Owner(EPriority::Blocking);
 			FBuildSession Session = Build.CreateSession(TexturePath, InputResolver);
-			Session.Build(DefinitionBuilder.Build(), EBuildPolicy::Default, Group,
+			Session.Build(DefinitionBuilder.Build(), EBuildPolicy::Default, Owner,
 				[this, &TexturePath, bReplaceExistingDDC] (FBuildCompleteParams&& Params)
 				{
 					#if !NO_LOGGING
@@ -583,7 +583,7 @@ void FTextureCacheDerivedDataWorker::BuildTexture(bool bReplaceExistingDDC)
 						ConsumeBuildFunctionOutput(Params.Output, TexturePath, bReplaceExistingDDC);
 					}
 				});
-			Group.Wait();
+			Owner.Wait();
 		}
 		else
 		{
@@ -1019,8 +1019,7 @@ public:
 		IBuildInputResolver* GlobalResolver = GetGlobalBuildInputResolver();
 		BuildSession = Build.CreateSession(TexturePath, GlobalResolver ? GlobalResolver : &InputResolver);
 
-		EPriority GroupPriority = EnumHasAnyFlags(Flags, ETextureCacheFlags::Async) ? ConvertPriority(Priority) : UE::DerivedData::EPriority::Blocking;
-		Group = Build.CreateGroup(GroupPriority);
+		Owner.Emplace(EnumHasAnyFlags(Flags, ETextureCacheFlags::Async) ? ConvertPriority(Priority) : UE::DerivedData::EPriority::Blocking);
 
 		FBuildDefinition Definition = CreateDefinition(Build, Texture, TexturePath, FunctionName, KeySuffix, Settings);
 
@@ -1029,7 +1028,7 @@ public:
 			FTextureBuildSettings ShippingSettings = Settings;
 			ShippingSettings.bHasEditorOnlyData = false;
 			FBuildDefinition ShippingDefinition = CreateDefinition(Build, Texture, TexturePath, FunctionName, KeySuffix, ShippingSettings);
-			BuildSession.Get().Build(ShippingDefinition, EBuildPolicy::Cache, Group.Get(),
+			BuildSession.Get().Build(ShippingDefinition, EBuildPolicy::Cache, *Owner,
 				[this, Definition = MoveTemp(Definition), Flags](FBuildCompleteParams&& Params)
 				{
 					switch (Params.Status)
@@ -1075,7 +1074,7 @@ public:
 		{
 			BuildPolicy &= ~EBuildPolicy::CacheQuery;
 		}
-		BuildSession.Get().Build(Definition, BuildPolicy, Group.Get(),
+		BuildSession.Get().Build(Definition, BuildPolicy, *Owner,
 			[this](FBuildCompleteParams&& Params) { EndBuild(MoveTemp(Params.Output), Params.BuildStatus); });
 	}
 
@@ -1102,19 +1101,19 @@ public:
 	bool SetPriority(EQueuedWorkPriority QueuedWorkPriority) final
 	{
 		Priority = QueuedWorkPriority;
-		Group.Get().SetPriority(ConvertPriority(QueuedWorkPriority));
+		Owner->SetPriority(ConvertPriority(QueuedWorkPriority));
 		return true;
 	}
 
 	bool Cancel() final
 	{
-		Group.Get().Cancel();
+		Owner->Cancel();
 		return true;
 	}
 
 	void Wait() final
 	{
-		Group.Get().Wait();
+		Owner->Wait();
 	}
 
 	bool WaitWithTimeout(float TimeLimitSeconds) final
@@ -1138,7 +1137,7 @@ public:
 
 	bool Poll() const final
 	{
-		return Group.Get().Poll();
+		return Owner->Poll();
 	}
 
 private:
@@ -1250,7 +1249,7 @@ private:
 	}
 
 	FTexturePlatformData& DerivedData;
-	UE::DerivedData::FOptionalRequestGroup Group;
+	TOptional<UE::DerivedData::FRequestOwner> Owner;
 	UE::DerivedData::FOptionalBuildSession BuildSession;
 	EQueuedWorkPriority Priority;
 	bool bCacheHit;
