DDC: Fixed non-unity compile errors

#rnx

diff --git a/Engine/Source/Developer/DerivedDataCache/Private/DerivedDataCacheKey.cpp b/Engine/Source/Developer/DerivedDataCache/Private/DerivedDataCacheKey.cpp
--- a/Engine/Source/Developer/DerivedDataCache/Private/DerivedDataCacheKey.cpp
+++ b/Engine/Source/Developer/DerivedDataCache/Private/DerivedDataCacheKey.cpp
@@ -4,6 +4,7 @@
 
 #include "Algo/AllOf.h"
 #include "Containers/Set.h"
+#include "Containers/StringConv.h"
 #include "DerivedDataCachePrivate.h"
 #include "Math/UnrealMathUtility.h"
 #include "Misc/ScopeRWLock.h"
diff --git a/Engine/Source/Developer/DerivedDataCache/Private/ZenDerivedDataBackend.cpp b/Engine/Source/Developer/DerivedDataCache/Private/ZenDerivedDataBackend.cpp
--- a/Engine/Source/Developer/DerivedDataCache/Private/ZenDerivedDataBackend.cpp
+++ b/Engine/Source/Developer/DerivedDataCache/Private/ZenDerivedDataBackend.cpp
@@ -5,6 +5,7 @@
 
 #include "Algo/Accumulate.h"
 #include "DerivedDataCacheRecord.h"
+#include "Misc/ConfigCacheIni.h"
 #include "Misc/ScopeLock.h"
 #include "ProfilingDebugging/CountersTrace.h"
 #include "ProfilingDebugging/CpuProfilerTrace.h"
diff --git a/Engine/Source/Developer/DerivedDataCache/Private/ZenServerInterface.cpp b/Engine/Source/Developer/DerivedDataCache/Private/ZenServerInterface.cpp
--- a/Engine/Source/Developer/DerivedDataCache/Private/ZenServerInterface.cpp
+++ b/Engine/Source/Developer/DerivedDataCache/Private/ZenServerInterface.cpp
@@ -7,6 +7,7 @@
 #include "ZenBackendUtils.h"
 
 #include "Memory/CompositeBuffer.h"
+#include "Misc/Paths.h"
 #include "Serialization/CompactBinary.h"
 #include "Serialization/CompactBinaryPackage.h"
 #include "Serialization/CompactBinaryValidation.h"
