DDC: Removed IBuild::Load[Type] in favor of exported functions

#rnx

diff --git a/Engine/Source/Developer/DerivedDataCache/Private/DerivedDataBuild.cpp b/Engine/Source/Developer/DerivedDataCache/Private/DerivedDataBuild.cpp
--- a/Engine/Source/Developer/DerivedDataCache/Private/DerivedDataBuild.cpp
+++ b/Engine/Source/Developer/DerivedDataCache/Private/DerivedDataBuild.cpp
@@ -77,21 +77,11 @@ public:
 		return CreateBuildDefinition(Name, Function);
 	}
 
-	FOptionalBuildDefinition LoadDefinition(FStringView Name, FCbObject&& Definition) final
-	{
-		return LoadBuildDefinition(Name, MoveTemp(Definition));
-	}
-
 	FBuildActionBuilder CreateAction(FStringView Name, FStringView Function) final
 	{
 		return CreateBuildAction(Name, Function, FunctionRegistry->FindFunctionVersion(Function), Version);
 	}
 
-	FOptionalBuildAction LoadAction(FStringView Name, FCbObject&& Action) final
-	{
-		return LoadBuildAction(Name, MoveTemp(Action));
-	}
-
 	FBuildInputsBuilder CreateInputs(FStringView Name) final
 	{
 		return CreateBuildInputs(Name);
@@ -102,16 +92,6 @@ public:
 		return CreateBuildOutput(Name, Function);
 	}
 
-	FOptionalBuildOutput LoadOutput(FStringView Name, FStringView Function, const FCbObject& Output) final
-	{
-		return LoadBuildOutput(Name, Function, Output);
-	}
-
-	FOptionalBuildOutput LoadOutput(FStringView Name, FStringView Function, const FCacheRecord& Output) final
-	{
-		return LoadBuildOutput(Name, Function, Output);
-	}
-
 	FBuildSession CreateSession(FStringView Name, IBuildInputResolver* InputResolver, IBuildScheduler* Scheduler) final
 	{
 		return CreateBuildSession(Name, Cache, *this, Scheduler ? *Scheduler : *DefaultScheduler, InputResolver);
diff --git a/Engine/Source/Developer/DerivedDataCache/Private/DerivedDataBuildAction.cpp b/Engine/Source/Developer/DerivedDataCache/Private/DerivedDataBuildAction.cpp
--- a/Engine/Source/Developer/DerivedDataCache/Private/DerivedDataBuildAction.cpp
+++ b/Engine/Source/Developer/DerivedDataCache/Private/DerivedDataBuildAction.cpp
@@ -264,10 +264,16 @@ FBuildActionBuilder CreateBuildAction(FStringView Name, FStringView Function, co
 	return CreateBuildActionBuilder(new FBuildActionBuilderInternal(Name, Function, FunctionVersion, BuildSystemVersion));
 }
 
-FOptionalBuildAction LoadBuildAction(FStringView Name, FCbObject&& Action)
+} // UE::DerivedData::Private
+
+namespace UE::DerivedData
+{
+
+FOptionalBuildAction FBuildAction::Load(FStringView Name, FCbObject&& Action)
 {
 	bool bIsValid = false;
-	FOptionalBuildAction Out = CreateBuildAction(new FBuildActionInternal(Name, MoveTemp(Action), bIsValid));
+	FOptionalBuildAction Out = Private::CreateBuildAction(
+		new Private::FBuildActionInternal(Name, MoveTemp(Action), bIsValid));
 	if (!bIsValid)
 	{
 		Out.Reset();
@@ -275,4 +281,4 @@ FOptionalBuildAction LoadBuildAction(FStringView Name, FCbObject&& Action)
 	return Out;
 }
 
-} // UE::DerivedData::Private
+} // UE::DerivedData
diff --git a/Engine/Source/Developer/DerivedDataCache/Private/DerivedDataBuildDefinition.cpp b/Engine/Source/Developer/DerivedDataCache/Private/DerivedDataBuildDefinition.cpp
--- a/Engine/Source/Developer/DerivedDataCache/Private/DerivedDataBuildDefinition.cpp
+++ b/Engine/Source/Developer/DerivedDataCache/Private/DerivedDataBuildDefinition.cpp
@@ -358,10 +358,16 @@ FBuildDefinitionBuilder CreateBuildDefinition(FStringView Name, FStringView Func
 	return CreateBuildDefinitionBuilder(new FBuildDefinitionBuilderInternal(Name, Function));
 }
 
-FOptionalBuildDefinition LoadBuildDefinition(FStringView Name, FCbObject&& Definition)
+} // UE::DerivedData::Private
+
+namespace UE::DerivedData
+{
+
+FOptionalBuildDefinition FBuildDefinition::Load(FStringView Name, FCbObject&& Definition)
 {
 	bool bIsValid = false;
-	FOptionalBuildDefinition Out = CreateBuildDefinition(new FBuildDefinitionInternal(Name, MoveTemp(Definition), bIsValid));
+	FOptionalBuildDefinition Out = Private::CreateBuildDefinition(
+		new Private::FBuildDefinitionInternal(Name, MoveTemp(Definition), bIsValid));
 	if (!bIsValid)
 	{
 		Out.Reset();
@@ -369,4 +375,4 @@ FOptionalBuildDefinition LoadBuildDefinition(FStringView Name, FCbObject&& Defin
 	return Out;
 }
 
-} // UE::DerivedData::Private
+} // UE::DerivedData
diff --git a/Engine/Source/Developer/DerivedDataCache/Private/DerivedDataBuildJob.cpp b/Engine/Source/Developer/DerivedDataCache/Private/DerivedDataBuildJob.cpp
--- a/Engine/Source/Developer/DerivedDataCache/Private/DerivedDataBuildJob.cpp
+++ b/Engine/Source/Developer/DerivedDataCache/Private/DerivedDataBuildJob.cpp
@@ -447,7 +447,7 @@ void FBuildJob::EndCacheQuery(FCacheGetCompleteParams&& Params)
 {
 	if (Params.Status == EStatus::Ok)
 	{
-		if (FOptionalBuildOutput CacheOutput = BuildSystem.LoadOutput(Name, FunctionName, Params.Record))
+		if (FOptionalBuildOutput CacheOutput = FBuildOutput::Load(Name, FunctionName, Params.Record))
 		{
 			EnumAddFlags(BuildStatus, EBuildStatus::CacheQueryHit);
 			return SetOutputNoCheck(MoveTemp(CacheOutput).Get());
diff --git a/Engine/Source/Developer/DerivedDataCache/Private/DerivedDataBuildOutput.cpp b/Engine/Source/Developer/DerivedDataCache/Private/DerivedDataBuildOutput.cpp
--- a/Engine/Source/Developer/DerivedDataCache/Private/DerivedDataBuildOutput.cpp
+++ b/Engine/Source/Developer/DerivedDataCache/Private/DerivedDataBuildOutput.cpp
@@ -316,10 +316,16 @@ FBuildOutputBuilder CreateBuildOutput(FStringView Name, FStringView Function)
 	return CreateBuildOutputBuilder(new FBuildOutputBuilderInternal(Name, Function));
 }
 
-FOptionalBuildOutput LoadBuildOutput(FStringView Name, FStringView Function, const FCbObject& Output)
+} // UE::DerivedData::Private
+
+namespace UE::DerivedData
+{
+
+FOptionalBuildOutput FBuildOutput::Load(FStringView Name, FStringView Function, const FCbObject& Output)
 {
 	bool bIsValid = false;
-	FOptionalBuildOutput Out = CreateBuildOutput(new FBuildOutputInternal(Name, Function, Output, bIsValid));
+	FOptionalBuildOutput Out = Private::CreateBuildOutput(
+		new Private::FBuildOutputInternal(Name, Function, Output, bIsValid));
 	if (!bIsValid)
 	{
 		Out.Reset();
@@ -327,10 +333,11 @@ FOptionalBuildOutput LoadBuildOutput(FStringView Name, FStringView Function, con
 	return Out;
 }
 
-FOptionalBuildOutput LoadBuildOutput(FStringView Name, FStringView Function, const FCacheRecord& Output)
+FOptionalBuildOutput FBuildOutput::Load(FStringView Name, FStringView Function, const FCacheRecord& Output)
 {
 	bool bIsValid = false;
-	FOptionalBuildOutput Out = CreateBuildOutput(new FBuildOutputInternal(Name, Function, Output, bIsValid));
+	FOptionalBuildOutput Out = Private::CreateBuildOutput(
+		new Private::FBuildOutputInternal(Name, Function, Output, bIsValid));
 	if (!bIsValid)
 	{
 		Out.Reset();
@@ -338,4 +345,4 @@ FOptionalBuildOutput LoadBuildOutput(FStringView Name, FStringView Function, con
 	return Out;
 }
 
-} // UE::DerivedData::Private
+} // UE::DerivedData
diff --git a/Engine/Source/Developer/DerivedDataCache/Private/DerivedDataBuildPrivate.h b/Engine/Source/Developer/DerivedDataCache/Private/DerivedDataBuildPrivate.h
--- a/Engine/Source/Developer/DerivedDataCache/Private/DerivedDataBuildPrivate.h
+++ b/Engine/Source/Developer/DerivedDataCache/Private/DerivedDataBuildPrivate.h
@@ -6,7 +6,6 @@
 #include "Logging/LogMacros.h"
 #include "Misc/AsciiSet.h"
 
-class FCbObject;
 struct FGuid;
 
 template <typename FuncType> class TUniqueFunction;
@@ -18,11 +17,7 @@ namespace UE::DerivedData { class FBuildDefinitionBuilder; }
 namespace UE::DerivedData { class FBuildInputsBuilder; }
 namespace UE::DerivedData { class FBuildOutputBuilder; }
 namespace UE::DerivedData { class FBuildSession; }
-namespace UE::DerivedData { class FCacheRecord; }
-namespace UE::DerivedData { class FOptionalBuildAction; }
-namespace UE::DerivedData { class FOptionalBuildDefinition; }
 namespace UE::DerivedData { class FOptionalBuildInputs; }
-namespace UE::DerivedData { class FOptionalBuildOutput; }
 namespace UE::DerivedData { class IBuild; }
 namespace UE::DerivedData { class IBuildFunctionRegistry; }
 namespace UE::DerivedData { class IBuildInputResolver; }
@@ -54,19 +49,15 @@ IBuildScheduler* CreateBuildScheduler();
 
 // Implemented in DerivedDataBuildDefinition.cpp
 FBuildDefinitionBuilder CreateBuildDefinition(FStringView Name, FStringView Function);
-FOptionalBuildDefinition LoadBuildDefinition(FStringView Name, FCbObject&& Definition);
 
 // Implemented in DerivedDataBuildAction.cpp
 FBuildActionBuilder CreateBuildAction(FStringView Name, FStringView Function, const FGuid& FunctionVersion, const FGuid& BuildSystemVersion);
-FOptionalBuildAction LoadBuildAction(FStringView Name, FCbObject&& Action);
 
 // Implemented in DerivedDataBuildInput.cpp
 FBuildInputsBuilder CreateBuildInputs(FStringView Name);
 
 // Implemented in DerivedDataBuildOutput.cpp
 FBuildOutputBuilder CreateBuildOutput(FStringView Name, FStringView Function);
-FOptionalBuildOutput LoadBuildOutput(FStringView Name, FStringView Function, const FCbObject& Output);
-FOptionalBuildOutput LoadBuildOutput(FStringView Name, FStringView Function, const FCacheRecord& Output);
 
 // Implemented in DerivedDataBuildSession.cpp
 FBuildSession CreateBuildSession(
diff --git a/Engine/Source/Developer/DerivedDataCache/Public/DerivedDataBuild.h b/Engine/Source/Developer/DerivedDataCache/Public/DerivedDataBuild.h
--- a/Engine/Source/Developer/DerivedDataCache/Public/DerivedDataBuild.h
+++ b/Engine/Source/Developer/DerivedDataCache/Public/DerivedDataBuild.h
@@ -6,7 +6,6 @@
 
 #define UE_API DERIVEDDATACACHE_API
 
-class FCbObject;
 struct FGuid;
 
 namespace UE::DerivedData { class FBuildActionBuilder; }
@@ -14,10 +13,6 @@ namespace UE::DerivedData { class FBuildDefinitionBuilder; }
 namespace UE::DerivedData { class FBuildInputsBuilder; }
 namespace UE::DerivedData { class FBuildOutputBuilder; }
 namespace UE::DerivedData { class FBuildSession; }
-namespace UE::DerivedData { class FCacheRecord; }
-namespace UE::DerivedData { class FOptionalBuildAction; }
-namespace UE::DerivedData { class FOptionalBuildDefinition; }
-namespace UE::DerivedData { class FOptionalBuildOutput; }
 namespace UE::DerivedData { class IBuildFunctionRegistry; }
 namespace UE::DerivedData { class IBuildInputResolver; }
 namespace UE::DerivedData { class IBuildScheduler; }
@@ -57,15 +52,6 @@ public:
 	 */
 	virtual FBuildDefinitionBuilder CreateDefinition(FStringView Name, FStringView Function) = 0;
 
-	/**
-	 * Load a build definition from compact binary.
-	 *
-	 * @param Name         The name by which to identify this definition for logging and profiling.
-	 * @param Definition   An object saved from a build definition. Holds a reference and is cloned if not owned.
-	 * @return A valid build definition, or null on error.
-	 */
-	virtual FOptionalBuildDefinition LoadDefinition(FStringView Name, FCbObject&& Definition) = 0;
-
 	/**
 	 * Create a build action builder.
 	 *
@@ -74,15 +60,6 @@ public:
 	 */
 	virtual FBuildActionBuilder CreateAction(FStringView Name, FStringView Function) = 0;
 
-	/**
-	 * Load a build action from compact binary.
-	 *
-	 * @param Name     The name by which to identify this action for logging and profiling.
-	 * @param Action   The saved action to load.
-	 * @return A valid build action, or null on error.
-	 */
-	virtual FOptionalBuildAction LoadAction(FStringView Name, FCbObject&& Action) = 0;
-
 	/**
 	 * Create a build inputs builder.
 	 *
@@ -98,17 +75,6 @@ public:
 	 */
 	virtual FBuildOutputBuilder CreateOutput(FStringView Name, FStringView Function) = 0;
 
-	/**
-	 * Load a build output.
-	 *
-	 * @param Name       The name by which to identify this output for logging and profiling.
-	 * @param Function   The name of the build function that produced this output.
-	 * @param Output     The saved output to load.
-	 * @return A valid build output, or null on error.
-	 */
-	virtual FOptionalBuildOutput LoadOutput(FStringView Name, FStringView Function, const FCbObject& Output) = 0;
-	virtual FOptionalBuildOutput LoadOutput(FStringView Name, FStringView Function, const FCacheRecord& Output) = 0;
-
 	/**
 	 * Create a build session.
 	 *
@@ -144,6 +110,7 @@ public:
 	virtual IBuildWorkerRegistry& GetWorkerRegistry() const = 0;
 };
 
+/** Returns a reference to the build system. Asserts if not available. */
 UE_API IBuild& GetBuild();
 
 } // UE::DerivedData
diff --git a/Engine/Source/Developer/DerivedDataCache/Public/DerivedDataBuildAction.h b/Engine/Source/Developer/DerivedDataCache/Public/DerivedDataBuildAction.h
--- a/Engine/Source/Developer/DerivedDataCache/Public/DerivedDataBuildAction.h
+++ b/Engine/Source/Developer/DerivedDataCache/Public/DerivedDataBuildAction.h
@@ -9,6 +9,8 @@
 #include "Templates/RefCounting.h"
 #include "Templates/UniquePtr.h"
 
+#define UE_API DERIVEDDATACACHE_API
+
 class FCbObject;
 class FCbWriter;
 struct FGuid;
@@ -16,6 +18,7 @@ struct FIoHash;
 
 namespace UE::DerivedData { class FBuildAction; }
 namespace UE::DerivedData { class FBuildActionBuilder; }
+namespace UE::DerivedData { class FOptionalBuildAction; }
 namespace UE::DerivedData { struct FBuildActionKey; }
 
 namespace UE::DerivedData::Private
@@ -113,6 +116,15 @@ public:
 		return Action->Save(Writer);
 	}
 
+	/**
+	 * Load a build action from compact binary.
+	 *
+	 * @param Name     The name by which to identify this action for logging and profiling.
+	 * @param Action   The saved action to load.
+	 * @return A valid build action, or null on error.
+	 */
+	UE_API static FOptionalBuildAction Load(FStringView Name, FCbObject&& Action);
+
 private:
 	friend class FOptionalBuildAction;
 	friend FBuildAction Private::CreateBuildAction(Private::IBuildActionInternal* Action);
@@ -194,3 +206,5 @@ public:
 };
 
 } // UE::DerivedData
+
+#undef UE_API
diff --git a/Engine/Source/Developer/DerivedDataCache/Public/DerivedDataBuildDefinition.h b/Engine/Source/Developer/DerivedDataCache/Public/DerivedDataBuildDefinition.h
--- a/Engine/Source/Developer/DerivedDataCache/Public/DerivedDataBuildDefinition.h
+++ b/Engine/Source/Developer/DerivedDataCache/Public/DerivedDataBuildDefinition.h
@@ -9,6 +9,8 @@
 #include "Templates/RefCounting.h"
 #include "Templates/UniquePtr.h"
 
+#define UE_API DERIVEDDATACACHE_API
+
 class FCbObject;
 class FCbWriter;
 struct FGuid;
@@ -16,6 +18,7 @@ struct FIoHash;
 
 namespace UE::DerivedData { class FBuildDefinition; }
 namespace UE::DerivedData { class FBuildDefinitionBuilder; }
+namespace UE::DerivedData { class FOptionalBuildDefinition; }
 namespace UE::DerivedData { struct FBuildKey; }
 namespace UE::DerivedData { struct FBuildPayloadKey; }
 
@@ -135,6 +138,15 @@ public:
 		return Definition->Save(Writer);
 	}
 
+	/**
+	 * Load a build definition from compact binary.
+	 *
+	 * @param Name         The name by which to identify this definition for logging and profiling.
+	 * @param Definition   An object saved from a build definition. Holds a reference and is cloned if not owned.
+	 * @return A valid build definition, or null on error.
+	 */
+	UE_API static FOptionalBuildDefinition Load(FStringView Name, FCbObject&& Definition);
+
 private:
 	friend class FOptionalBuildDefinition;
 	friend FBuildDefinition Private::CreateBuildDefinition(Private::IBuildDefinitionInternal* Definition);
@@ -248,3 +260,5 @@ public:
 };
 
 } // UE::DerivedData
+
+#undef UE_API
diff --git a/Engine/Source/Developer/DerivedDataCache/Public/DerivedDataBuildOutput.h b/Engine/Source/Developer/DerivedDataCache/Public/DerivedDataBuildOutput.h
--- a/Engine/Source/Developer/DerivedDataCache/Public/DerivedDataBuildOutput.h
+++ b/Engine/Source/Developer/DerivedDataCache/Public/DerivedDataBuildOutput.h
@@ -10,12 +10,16 @@
 #include "Templates/RefCounting.h"
 #include "Templates/UniquePtr.h"
 
+#define UE_API DERIVEDDATACACHE_API
+
 class FCbObject;
 class FCbWriter;
 
 namespace UE::DerivedData { class FBuildOutput; }
 namespace UE::DerivedData { class FBuildOutputBuilder; }
+namespace UE::DerivedData { class FCacheRecord; }
 namespace UE::DerivedData { class FCacheRecordBuilder; }
+namespace UE::DerivedData { class FOptionalBuildOutput; }
 namespace UE::DerivedData { class FPayload; }
 namespace UE::DerivedData { struct FBuildDiagnostic; }
 namespace UE::DerivedData { struct FPayloadId; }
@@ -126,6 +130,17 @@ public:
 		Output->Save(RecordBuilder);
 	}
 
+	/**
+	 * Load a build output.
+	 *
+	 * @param Name       The name by which to identify this output for logging and profiling.
+	 * @param Function   The name of the build function that produced this output.
+	 * @param Output     The saved output to load.
+	 * @return A valid build output, or null on error.
+	 */
+	UE_API static FOptionalBuildOutput Load(FStringView Name, FStringView Function, const FCbObject& Output);
+	UE_API static FOptionalBuildOutput Load(FStringView Name, FStringView Function, const FCacheRecord& Output);
+
 private:
 	friend class FOptionalBuildOutput;
 	friend FBuildOutput Private::CreateBuildOutput(Private::IBuildOutputInternal* Output);
@@ -224,3 +239,5 @@ public:
 };
 
 } // UE::DerivedData
+
+#undef UE_API
diff --git a/Engine/Source/Editor/UnrealEd/Private/DerivedDataBuildRemoteExecutor.cpp b/Engine/Source/Editor/UnrealEd/Private/DerivedDataBuildRemoteExecutor.cpp
--- a/Engine/Source/Editor/UnrealEd/Private/DerivedDataBuildRemoteExecutor.cpp
+++ b/Engine/Source/Editor/UnrealEd/Private/DerivedDataBuildRemoteExecutor.cpp
@@ -750,7 +750,7 @@ FOptionalBuildOutput FRemoteBuildExecutionRequest::ComposeBuildOutput(FBatchRead
 				return FOptionalBuildOutput();
 			}
 
-			RemoteBuildOutput = State.BuildSystem.LoadOutput(State.BuildAction.GetName(), State.BuildAction.GetFunction(), FCbObject(BuildOutputBuffer));
+			RemoteBuildOutput = FBuildOutput::Load(State.BuildAction.GetName(), State.BuildAction.GetFunction(), FCbObject(BuildOutputBuffer));
 		}
 		else
 		{
diff --git a/Engine/Source/Programs/DerivedDataBuildWorker/Private/DerivedDataBuildWorker.cpp b/Engine/Source/Programs/DerivedDataBuildWorker/Private/DerivedDataBuildWorker.cpp
--- a/Engine/Source/Programs/DerivedDataBuildWorker/Private/DerivedDataBuildWorker.cpp
+++ b/Engine/Source/Programs/DerivedDataBuildWorker/Private/DerivedDataBuildWorker.cpp
@@ -230,7 +230,7 @@ bool FBuildWorkerProgram::Build()
 			UE_LOG(LogDerivedDataBuildWorker, Error, TEXT("Missing build action '%s'"), *ActionPath);
 			return false;
 		}
-		if (FOptionalBuildAction Action = BuildSystem.LoadAction(ActionPath, MoveTemp(ActionObject)); Action.IsNull())
+		if (FOptionalBuildAction Action = FBuildAction::Load(ActionPath, MoveTemp(ActionObject)); Action.IsNull())
 		{
 			UE_LOG(LogDerivedDataBuildWorker, Error, TEXT("Invalid build action '%s'"), *ActionPath);
 			return false;
