DDC: Removed ICacheFactory in favor of exported functions

FCacheBucket is now always exposed as FAnsiStringView.

#rnx

diff --git a/Engine/Source/Developer/DerivedDataCache/Private/DerivedDataBackendAsyncPutWrapper.cpp b/Engine/Source/Developer/DerivedDataCache/Private/DerivedDataBackendAsyncPutWrapper.cpp
--- a/Engine/Source/Developer/DerivedDataCache/Private/DerivedDataBackendAsyncPutWrapper.cpp
+++ b/Engine/Source/Developer/DerivedDataCache/Private/DerivedDataBackendAsyncPutWrapper.cpp
@@ -163,10 +163,9 @@ public:
 	}
 };
 
-FDerivedDataBackendAsyncPutWrapper::FDerivedDataBackendAsyncPutWrapper(ICacheFactory& InFactory, FDerivedDataBackendInterface* InInnerBackend, bool bCacheInFlightPuts)
-	: Factory(InFactory)
-	, InnerBackend(InInnerBackend)
-	, InflightCache(bCacheInFlightPuts ? (new FMemoryDerivedDataBackend(Factory, TEXT("InflightMemoryCache"))) : NULL)
+FDerivedDataBackendAsyncPutWrapper::FDerivedDataBackendAsyncPutWrapper(FDerivedDataBackendInterface* InInnerBackend, bool bCacheInFlightPuts)
+	: InnerBackend(InInnerBackend)
+	, InflightCache(bCacheInFlightPuts ? (new FMemoryDerivedDataBackend(TEXT("InflightMemoryCache"))) : NULL)
 {
 	check(InnerBackend);
 }
@@ -521,7 +520,7 @@ void FDerivedDataBackendAsyncPutWrapper::Get(
 				{
 					for (const FCacheKey& Key : Keys)
 					{
-						OnComplete({Factory.CreateRecord(Key).Build(), EStatus::Canceled});
+						OnComplete({FCacheRecordBuilder(Key).Build(), EStatus::Canceled});
 					}
 				}
 			});
diff --git a/Engine/Source/Developer/DerivedDataCache/Private/DerivedDataBackendAsyncPutWrapper.h b/Engine/Source/Developer/DerivedDataCache/Private/DerivedDataBackendAsyncPutWrapper.h
--- a/Engine/Source/Developer/DerivedDataCache/Private/DerivedDataBackendAsyncPutWrapper.h
+++ b/Engine/Source/Developer/DerivedDataCache/Private/DerivedDataBackendAsyncPutWrapper.h
@@ -64,7 +64,7 @@ public:
 	 * @param	InInnerBackend		Backend to use for storage, my responsibilities are about async puts
 	 * @param	bCacheInFlightPuts	if true, cache in-flight puts in a memory cache so that they hit immediately
 	 */
-	FDerivedDataBackendAsyncPutWrapper(ICacheFactory& InFactory, FDerivedDataBackendInterface* InInnerBackend, bool bCacheInFlightPuts);
+	FDerivedDataBackendAsyncPutWrapper(FDerivedDataBackendInterface* InInnerBackend, bool bCacheInFlightPuts);
 
 	/** Return a name for this interface */
 	virtual FString GetName() const override
@@ -162,7 +162,6 @@ private:
 	FDerivedDataCacheUsageStats UsageStats;
 	FDerivedDataCacheUsageStats PutSyncUsageStats;
 
-	ICacheFactory& Factory;
 	/** Backend to use for storage, my responsibilities are about async puts **/
 	FDerivedDataBackendInterface*					InnerBackend;
 	/** Memory based cache to deal with gets that happen while an async put is still in flight **/
diff --git a/Engine/Source/Developer/DerivedDataCache/Private/DerivedDataBackends.cpp b/Engine/Source/Developer/DerivedDataCache/Private/DerivedDataBackends.cpp
--- a/Engine/Source/Developer/DerivedDataCache/Private/DerivedDataBackends.cpp
+++ b/Engine/Source/Developer/DerivedDataCache/Private/DerivedDataBackends.cpp
@@ -10,7 +10,6 @@
 #include "HAL/IConsoleManager.h"
 #include "Misc/App.h"
 #include "DerivedDataBackendInterface.h"
-#include "DerivedDataCachePrivate.h"
 #include "DerivedDataCacheUsageStats.h"
 #include "MemoryDerivedDataBackend.h"
 #include "HttpDerivedDataBackend.h"
@@ -44,13 +43,13 @@ DEFINE_LOG_CATEGORY(LogDerivedDataCache);
 namespace UE::DerivedData::Backends
 {
 
-FDerivedDataBackendInterface* CreateFileSystemDerivedDataBackend(ICacheFactory& Factory, const TCHAR* CacheDirectory, const TCHAR* InParams, const TCHAR* InAccessLogFileName = nullptr);
+FDerivedDataBackendInterface* CreateFileSystemDerivedDataBackend(const TCHAR* CacheDirectory, const TCHAR* InParams, const TCHAR* InAccessLogFileName = nullptr);
 
 /**
   * This class is used to create a singleton that represents the derived data cache hierarchy and all of the wrappers necessary
   * ideally this would be data driven and the backends would be plugins...
 **/
-class FDerivedDataBackendGraph : public FDerivedDataBackend, public ICacheFactory
+class FDerivedDataBackendGraph : public FDerivedDataBackend
 {
 public:
 	/**
@@ -122,7 +121,7 @@ public:
 		// Make sure AsyncPutWrapper and KeyLengthWrapper are created
 		if( !AsyncPutWrapper )
 		{
-			AsyncPutWrapper = new FDerivedDataBackendAsyncPutWrapper( *this, RootCache, true );
+			AsyncPutWrapper = new FDerivedDataBackendAsyncPutWrapper( RootCache, true );
 			check(AsyncPutWrapper);
 			CreatedBackends.Add( AsyncPutWrapper );
 			RootCache = AsyncPutWrapper;
@@ -309,7 +308,7 @@ public:
 				FGuid Temp = FGuid::NewGuid();
 				ReadPakFilename = PakFilename;
 				WritePakFilename = PakFilename + TEXT(".") + Temp.ToString();
-				WritePakCache = bCompressed? new FCompressedPakFileDerivedDataBackend( *this, *WritePakFilename, true ) : new FPakFileDerivedDataBackend( *this, *WritePakFilename, true );
+				WritePakCache = bCompressed? new FCompressedPakFileDerivedDataBackend( *WritePakFilename, true ) : new FPakFileDerivedDataBackend( *WritePakFilename, true );
 				PakNode = WritePakCache;
 			}
 			else
@@ -317,7 +316,7 @@ public:
 				bool bReadPak = FPlatformFileManager::Get().GetPlatformFile().FileExists( *PakFilename );
 				if( bReadPak )
 				{
-					FPakFileDerivedDataBackend* ReadPak = bCompressed? new FCompressedPakFileDerivedDataBackend( *this, *PakFilename, false ) : new FPakFileDerivedDataBackend( *this, *PakFilename, false );
+					FPakFileDerivedDataBackend* ReadPak = bCompressed? new FCompressedPakFileDerivedDataBackend( *PakFilename, false ) : new FPakFileDerivedDataBackend( *PakFilename, false );
 					ReadPakFilename = PakFilename;
 					PakNode = ReadPak;
 					ReadPakCache.Add(ReadPak);
@@ -406,7 +405,7 @@ public:
 		FDerivedDataBackendInterface* AsyncNode = NULL;
 		if( InnerNode )
 		{
-			AsyncNode = new FDerivedDataBackendAsyncPutWrapper( *this, InnerNode, true );
+			AsyncNode = new FDerivedDataBackendAsyncPutWrapper( InnerNode, true );
 		}
 		else
 		{
@@ -507,7 +506,7 @@ public:
 		FDerivedDataBackendInterface* Hierarchy = NULL;
 		if( InnerNodes.Num() > 1 )
 		{
-			FHierarchicalDerivedDataBackend* HierarchyBackend = new FHierarchicalDerivedDataBackend( *this, InnerNodes );
+			FHierarchicalDerivedDataBackend* HierarchyBackend = new FHierarchicalDerivedDataBackend( InnerNodes );
 			Hierarchy = HierarchyBackend;
 			if (HierarchicalWrapper == NULL)
 			{
@@ -645,7 +644,7 @@ public:
 
 				if (!bShared || IFileManager::Get().DirectoryExists(*Path))
 				{
-					InnerFileSystem = CreateFileSystemDerivedDataBackend(*this, *Path, Entry, *WriteAccessLog);
+					InnerFileSystem = CreateFileSystemDerivedDataBackend(*Path, Entry, *WriteAccessLog);
 				}
 
 				if (InnerFileSystem)
@@ -739,7 +738,7 @@ public:
 		}
 
 		// Insert the backend corruption wrapper. Since the filesystem already uses this, and we're recycling the data with the trailer intact, we need to use it for the S3 cache too.
-		FS3DerivedDataBackend* Backend = new FS3DerivedDataBackend(*this, *ManifestPath, *BaseUrl, *Region, *CanaryObjectKey, *CachePath);
+		FS3DerivedDataBackend* Backend = new FS3DerivedDataBackend(*ManifestPath, *BaseUrl, *Region, *CanaryObjectKey, *CachePath);
 		return new FDerivedDataBackendCorruptionWrapper(Backend);
 #else
 		UE_LOG(LogDerivedDataCache, Log, TEXT("S3 backend is not supported on the current platform."));
@@ -839,7 +838,7 @@ public:
 
 		const bool bReadOnly = GetParsedBool(Entry, TEXT("ReadOnly="));
 
-		FHttpDerivedDataBackend* backend = new FHttpDerivedDataBackend(*this, *ServiceUrl, *Namespace, *OAuthProvider, *OAuthClientId, *OAuthSecret, bReadOnly);
+		FHttpDerivedDataBackend* backend = new FHttpDerivedDataBackend(*ServiceUrl, *Namespace, *OAuthProvider, *OAuthClientId, *OAuthSecret, bReadOnly);
 
 		if (ForceSpeedClass != FDerivedDataBackendInterface::ESpeedClass::Unknown)
 		{
@@ -879,7 +878,7 @@ public:
 			UE_LOG(LogDerivedDataCache, Warning, TEXT("Node %s does not specify 'Namespace', falling back to '%s'"), NodeName, *Namespace);
 		}
 
-		FZenDerivedDataBackend* backend = new FZenDerivedDataBackend(*this, *ServiceUrl, *Namespace);
+		FZenDerivedDataBackend* backend = new FZenDerivedDataBackend(*ServiceUrl, *Namespace);
 		if (!backend->IsUsable())
 		{
 			UE_LOG(LogDerivedDataCache, Warning, TEXT("%s could not contact the service (%s), will not use it."), NodeName, *ServiceUrl);
@@ -923,7 +922,7 @@ public:
 			MaxCacheSize = FMath::Min(MaxCacheSize, MaxSupportedCacheSize);
 
 			UE_LOG( LogDerivedDataCache, Display, TEXT("Max Cache Size: %d MB"), MaxCacheSize);
-			Cache = new FMemoryDerivedDataBackend(*this, TEXT("Boot"), MaxCacheSize * 1024 * 1024, true /* bCanBeDisabled */);
+			Cache = new FMemoryDerivedDataBackend(TEXT("Boot"), MaxCacheSize * 1024 * 1024, true /* bCanBeDisabled */);
 
 			if( Cache && Filename.Len() )
 			{
@@ -965,7 +964,7 @@ public:
 		FString Filename;
 
 		FParse::Value( Entry, TEXT("Filename="), Filename );
-		Cache = new FMemoryDerivedDataBackend(*this, NodeName);
+		Cache = new FMemoryDerivedDataBackend(NodeName);
 		if( Cache && Filename.Len() )
 		{
 			if( Cache->LoadCache( *Filename ) )
@@ -1036,7 +1035,7 @@ public:
 
 				for(const FString& MergePakName : MergePakList)
 				{
-					FPakFileDerivedDataBackend ReadPak(*this, *FPaths::Combine(*FPaths::GetPath(WritePakFilename), *MergePakName), false);
+					FPakFileDerivedDataBackend ReadPak(*FPaths::Combine(*FPaths::GetPath(WritePakFilename), *MergePakName), false);
 					WritePakCache->MergeCache(&ReadPak);
 				}
 			}
@@ -1061,7 +1060,7 @@ public:
 							UE_LOG(LogDerivedDataCache, Error, TEXT("Could not delete the pak file %s to overwrite it with a new one."), *ReadPakFilename);
 						}
 					}
-					if (!FPakFileDerivedDataBackend::SortAndCopy(*this, WritePakFilename, ReadPakFilename))
+					if (!FPakFileDerivedDataBackend::SortAndCopy(WritePakFilename, ReadPakFilename))
 					{
 						UE_LOG(LogDerivedDataCache, Error, TEXT("Couldn't sort pak file (%s)"), *WritePakFilename);
 					}
@@ -1128,7 +1127,7 @@ public:
 		FPakFileDerivedDataBackend* ReadPak = NULL;
 		if (HierarchicalWrapper && FPlatformFileManager::Get().GetPlatformFile().FileExists(PakFilename))
 		{
-			ReadPak = new FPakFileDerivedDataBackend(*this, PakFilename, false);
+			ReadPak = new FPakFileDerivedDataBackend(PakFilename, false);
 
 			HierarchicalWrapper->AddInnerBackend(ReadPak);
 			CreatedBackends.Add(ReadPak);
@@ -1175,26 +1174,6 @@ public:
 		return MakeShared<FDerivedDataCacheStatsNode>(nullptr, TEXT(""));
 	}
 
-	virtual FCacheBucket CreateBucket(FStringView Name) override
-	{
-		return Private::CreateCacheBucket(Name);
-	}
-
-	virtual FCacheRecordBuilder CreateRecord(const FCacheKey& Key) override
-	{
-		return Private::CreateCacheRecordBuilder(Key);
-	}
-
-	virtual FCbPackage SaveRecord(const FCacheRecord& Record) override
-	{
-		return Private::SaveCacheRecord(Record);
-	}
-
-	virtual FOptionalCacheRecord LoadRecord(const FCbPackage& Package) override
-	{
-		return Private::LoadCacheRecord(Package);
-	}
-
 private:
 
 	/** Delete all created backends in the reversed order they were created. */
diff --git a/Engine/Source/Developer/DerivedDataCache/Private/DerivedDataBuildJob.cpp b/Engine/Source/Developer/DerivedDataCache/Private/DerivedDataBuildJob.cpp
--- a/Engine/Source/Developer/DerivedDataCache/Private/DerivedDataBuildJob.cpp
+++ b/Engine/Source/Developer/DerivedDataCache/Private/DerivedDataBuildJob.cpp
@@ -390,7 +390,7 @@ void FBuildJob::CreateContext()
 	}
 	else
 	{
-		const FCacheKey CacheKey{Cache.CreateBucket(FunctionName), Action.Get().GetKey().Hash};
+		const FCacheKey CacheKey{FCacheBucket(FunctionName), Action.Get().GetKey().Hash};
 		Context = new FBuildJobContext(*this, CacheKey, *Function, OutputBuilder, BuildPolicy);
 		Function->Configure(*Context);
 		BuildPolicy = Context->GetBuildPolicy();
@@ -473,7 +473,7 @@ void FBuildJob::EnterCacheStore()
 void FBuildJob::BeginCacheStore()
 {
 	TRACE_CPUPROFILER_EVENT_SCOPE(FBuildJob::CacheStore);
-	FCacheRecordBuilder RecordBuilder = Cache.CreateRecord(Context->GetCacheKey());
+	FCacheRecordBuilder RecordBuilder(Context->GetCacheKey());
 	Output.Get().Save(RecordBuilder);
 	EnumAddFlags(BuildStatus, EBuildStatus::CacheStore);
 	Cache.Put({RecordBuilder.Build()}, Name, Context->GetCachePolicy(), Owner,
@@ -1037,7 +1037,7 @@ void FBuildJob::ExportBuild() const
 	TStringBuilder<256> ExportPath;
 	const FCacheKey& Key = Context->GetCacheKey();
 	FPathViews::Append(ExportPath, FPaths::ProjectSavedDir(), TEXT("DerivedDataBuildExport"), Key.Bucket);
-	if (FunctionName != Key.Bucket.ToString<TCHAR>())
+	if (!Key.Bucket.ToString().Equals(FunctionName))
 	{
 		FPathViews::Append(ExportPath, FunctionName);
 	}
@@ -1110,7 +1110,7 @@ bool FBuildJob::ShouldExportBuild() const
 	static TArray<FName> ExportTypes = ParseExportBuildTypes(bExportAll);
 	return bExportAll
 		|| Algo::BinarySearch(ExportTypes, FName(FunctionName), FNameFastLess()) != INDEX_NONE
-		|| Algo::BinarySearch(ExportTypes, FName(Context->GetCacheKey().Bucket.ToString<ANSICHAR>()), FNameFastLess()) != INDEX_NONE;
+		|| Algo::BinarySearch(ExportTypes, FName(Context->GetCacheKey().Bucket.ToString()), FNameFastLess()) != INDEX_NONE;
 }
 
 TArray<FName> FBuildJob::ParseExportBuildTypes(bool& bOutExportAll)
diff --git a/Engine/Source/Developer/DerivedDataCache/Private/DerivedDataBuildJobContext.cpp b/Engine/Source/Developer/DerivedDataCache/Private/DerivedDataBuildJobContext.cpp
--- a/Engine/Source/Developer/DerivedDataCache/Private/DerivedDataBuildJobContext.cpp
+++ b/Engine/Source/Developer/DerivedDataCache/Private/DerivedDataBuildJobContext.cpp
@@ -146,11 +146,6 @@ void FBuildJobContext::EndAsyncBuild()
 	Owner->End(this, [this] { EndBuild(); });
 }
 
-FCacheBucket FBuildJobContext::CreateCacheBucket(FStringView Name) const
-{
-	return Job.GetCache().CreateBucket(Name);
-}
-
 void FBuildJobContext::SetCacheBucket(FCacheBucket Bucket)
 {
 	checkf(!Bucket.IsNull(), TEXT("Null cache bucket not allowed for build of '%s' by %s. ")
diff --git a/Engine/Source/Developer/DerivedDataCache/Private/DerivedDataBuildJobContext.h b/Engine/Source/Developer/DerivedDataCache/Private/DerivedDataBuildJobContext.h
--- a/Engine/Source/Developer/DerivedDataCache/Private/DerivedDataBuildJobContext.h
+++ b/Engine/Source/Developer/DerivedDataCache/Private/DerivedDataBuildJobContext.h
@@ -62,7 +62,6 @@ private:
 	void BeginAsyncBuild() final;
 	void EndAsyncBuild() final;
 
-	FCacheBucket CreateCacheBucket(FStringView Name) const final;
 	void SetCacheBucket(FCacheBucket Bucket) final;
 	void SetCachePolicy(ECachePolicy Policy) final;
 	void SetBuildPolicy(EBuildPolicy Policy) final;
diff --git a/Engine/Source/Developer/DerivedDataCache/Private/DerivedDataCache.cpp b/Engine/Source/Developer/DerivedDataCache/Private/DerivedDataCache.cpp
--- a/Engine/Source/Developer/DerivedDataCache/Private/DerivedDataCache.cpp
+++ b/Engine/Source/Developer/DerivedDataCache/Private/DerivedDataCache.cpp
@@ -11,7 +11,6 @@
 #include "DerivedDataBackendInterface.h"
 #include "DerivedDataCache.h"
 #include "DerivedDataCachePrivate.h"
-#include "DerivedDataCacheRecord.h"
 #include "DerivedDataCacheUsageStats.h"
 #include "DerivedDataPluginInterface.h"
 #include "HAL/ThreadSafeCounter.h"
@@ -19,7 +18,6 @@
 #include "Misc/CommandLine.h"
 #include "Misc/ScopeLock.h"
 #include "ProfilingDebugging/CookStats.h"
-#include "Serialization/CompactBinaryPackage.h"
 #include "Stats/Stats.h"
 #include "Stats/StatsMisc.h"
 #include <atomic>
@@ -683,20 +681,6 @@ private:
 public:
 	// ICache Interface
 
-	FCacheBucket CreateBucket(FStringView Name) final { return CreateCacheBucket(Name); }
-
-	FCacheRecordBuilder CreateRecord(const FCacheKey& Key) final { return CreateCacheRecordBuilder(Key); }
-
-	FCbPackage SaveRecord(const FCacheRecord& Record) final
-	{
-		return SaveCacheRecord(Record);
-	}
-
-	FOptionalCacheRecord LoadRecord(const FCbPackage& Package) final
-	{
-		return LoadCacheRecord(Package);
-	}
-
 	void Put(
 		TConstArrayView<FCacheRecord> Records,
 		FStringView Context,
diff --git a/Engine/Source/Developer/DerivedDataCache/Private/DerivedDataCacheKey.cpp b/Engine/Source/Developer/DerivedDataCache/Private/DerivedDataCacheKey.cpp
--- a/Engine/Source/Developer/DerivedDataCache/Private/DerivedDataCacheKey.cpp
+++ b/Engine/Source/Developer/DerivedDataCache/Private/DerivedDataCacheKey.cpp
@@ -2,22 +2,28 @@
 
 #include "DerivedDataCacheKey.h"
 
-#include "Algo/AllOf.h"
 #include "Containers/Set.h"
 #include "Containers/StringConv.h"
 #include "DerivedDataCachePrivate.h"
 #include "Math/UnrealMathUtility.h"
 #include "Misc/ScopeRWLock.h"
+#include "Misc/StringBuilder.h"
 
 namespace UE::DerivedData::Private
 {
 
+template <typename CharType>
+static inline void AssertValidCacheBucketName(TStringView<CharType> Name)
+{
+	checkf(IsValidCacheBucketName(Name),
+		TEXT("A cache bucket name must be alphanumeric, non-empty, and contain fewer than 256 code units. ")
+		TEXT("Name: '%s'"), *WriteToString<256>(Name));
+}
+
 class FCacheBucketOwner : public FCacheBucket
 {
 public:
-	static const TCHAR* AllocName(FStringView Bucket);
-
-	inline explicit FCacheBucketOwner(FStringView Bucket);
+	inline explicit FCacheBucketOwner(FAnsiStringView Bucket);
 	inline FCacheBucketOwner(FCacheBucketOwner&& Bucket);
 	inline ~FCacheBucketOwner();
 
@@ -26,44 +32,26 @@ public:
 
 	using FCacheBucket::operator==;
 
-	inline bool operator==(FStringView Bucket) const { return ToString<TCHAR>().Equals(Bucket, ESearchCase::IgnoreCase); }
+	inline bool operator==(FAnsiStringView Bucket) const { return ToString().Equals(Bucket, ESearchCase::IgnoreCase); }
 
 	friend inline uint32 GetTypeHash(const FCacheBucketOwner& Bucket)
 	{
-		return ::GetTypeHash(Bucket.ToString<TCHAR>());
+		return ::GetTypeHash(Bucket.ToString());
 	}
 };
 
-inline const TCHAR* FCacheBucketOwner::AllocName(FStringView Bucket)
+inline FCacheBucketOwner::FCacheBucketOwner(FAnsiStringView Bucket)
 {
-	FTCHARToUTF8 Bucket8Conversion(Bucket);
-	FAnsiStringView Bucket8 = Bucket8Conversion;
-
 	const int32 BucketLen = Bucket.Len();
-	const int32 Bucket8Len = Bucket8.Len();
-	const int32 PrefixSize = FMath::DivideAndRoundUp<int32>(2, sizeof(TCHAR));
-	TCHAR* Buffer = new TCHAR[PrefixSize + BucketLen + 1 + FMath::DivideAndRoundUp<int32>(Bucket8Len + 1, sizeof(TCHAR))];
+	const int32 PrefixSize = FMath::DivideAndRoundUp<int32>(1, sizeof(ANSICHAR));
+	ANSICHAR* Buffer = new ANSICHAR[PrefixSize + BucketLen + 1];
 	Buffer += PrefixSize;
-	const TCHAR* const Name = Buffer;
+	Name = Buffer;
 
 	reinterpret_cast<uint8*>(Buffer)[LengthOffset] = static_cast<uint8>(BucketLen);
-	reinterpret_cast<uint8*>(Buffer)[LengthOffset8] = static_cast<uint8>(Bucket8Len);
-
 	Bucket.CopyString(Buffer, BucketLen);
 	Buffer += BucketLen;
-	*Buffer++ = TEXT('\0');
-
-	ANSICHAR* Buffer8 = reinterpret_cast<ANSICHAR*>(Buffer);
-	Bucket8.CopyString(Buffer8, Bucket8Len);
-	Buffer8 += Bucket8Len;
-	*Buffer8++ = '\0';
-
-	return Name;
-}
-
-inline FCacheBucketOwner::FCacheBucketOwner(FStringView Bucket)
-	: FCacheBucket(AllocName(Bucket))
-{
+	*Buffer = '\0';
 }
 
 inline FCacheBucketOwner::FCacheBucketOwner(FCacheBucketOwner&& Bucket)
@@ -74,42 +62,63 @@ inline FCacheBucketOwner::FCacheBucketOwner(FCacheBucketOwner&& Bucket)
 
 inline FCacheBucketOwner::~FCacheBucketOwner()
 {
-	if (IsValid())
+	if (Name)
 	{
-		const int32 PrefixSize = FMath::DivideAndRoundUp<int32>(2, sizeof(TCHAR));
-		delete[] (ToString<TCHAR>().GetData() - PrefixSize);
+		const int32 PrefixSize = FMath::DivideAndRoundUp<int32>(1, sizeof(ANSICHAR));
+		delete[] (Name - PrefixSize);
 	}
 }
 
 class FCacheBuckets
 {
 public:
-	inline FCacheBucket FindOrAdd(FStringView Name);
+	template <typename CharType>
+	inline FCacheBucket FindOrAdd(TStringView<CharType> Name);
 
 private:
 	FRWLock Lock;
 	TSet<FCacheBucketOwner> Buckets;
 };
 
-inline FCacheBucket FCacheBuckets::FindOrAdd(FStringView Name)
+template <typename CharType>
+inline FCacheBucket FCacheBuckets::FindOrAdd(TStringView<CharType> Name)
 {
-	const uint32 Hash = GetTypeHash(Name);
-	if (FReadScopeLock ReadLock(Lock); const FCacheBucketOwner* Bucket = Buckets.FindByHash(Hash, Name))
+	const auto AnsiCast = StringCast<ANSICHAR>(Name.GetData(), Name.Len());
+	const FAnsiStringView AnsiName = AnsiCast;
+
+	const uint32 Hash = GetTypeHash(AnsiName);
+	if (FReadScopeLock ReadLock(Lock); const FCacheBucketOwner* Bucket = Buckets.FindByHash(Hash, AnsiName))
 	{
 		return *Bucket;
 	}
 
+	// It is valid to assert after because the "bogus char" '?' is not valid in a bucket name.
 	AssertValidCacheBucketName(Name);
 
-	FCacheBucketOwner LocalBucket(Name);
+	FCacheBucketOwner LocalBucket(AnsiName);
 	FWriteScopeLock WriteLock(Lock);
 	return Buckets.FindOrAddByHash(Hash, MoveTemp(LocalBucket));
 }
 
-FCacheBucket CreateCacheBucket(FStringView Name)
+static FCacheBuckets& GetCacheBuckets()
 {
 	static FCacheBuckets Buckets;
-	return Buckets.FindOrAdd(Name);
+	return Buckets;
 }
 
 } // UE::DerivedData::Private
+
+namespace UE::DerivedData
+{
+
+FCacheBucket::FCacheBucket(FUtf8StringView InName)
+	: FCacheBucket(Private::GetCacheBuckets().FindOrAdd(InName))
+{
+}
+
+FCacheBucket::FCacheBucket(FWideStringView InName)
+	: FCacheBucket(Private::GetCacheBuckets().FindOrAdd(InName))
+{
+}
+
+} // UE::DerivedData
diff --git a/Engine/Source/Developer/DerivedDataCache/Private/DerivedDataCacheModule.cpp b/Engine/Source/Developer/DerivedDataCache/Private/DerivedDataCacheModule.cpp
--- a/Engine/Source/Developer/DerivedDataCache/Private/DerivedDataCacheModule.cpp
+++ b/Engine/Source/Developer/DerivedDataCache/Private/DerivedDataCacheModule.cpp
@@ -6,8 +6,8 @@
 #include "DerivedDataBuild.h"
 #include "DerivedDataBuildPrivate.h"
 #include "DerivedDataCache.h"
-#include "DerivedDataCacheInterface.h"
 #include "DerivedDataCachePrivate.h"
+#include "HAL/CriticalSection.h"
 #include "Misc/ScopeLock.h"
 #include "Modules/ModuleManager.h"
 
@@ -61,8 +61,8 @@ public:
 		delete GDerivedDataBuild;
 		GDerivedDataBuild = nullptr;
 		delete GDerivedDataCache;
-		GDerivedDataLegacyCache = nullptr;
 		GDerivedDataCache = nullptr;
+		GDerivedDataLegacyCache = nullptr;
 	}
 
 private:
@@ -75,7 +75,6 @@ static FDerivedDataCacheModule* GetModule()
 {
 	if (!FPlatformProperties::RequiresCookedData())
 	{
-		check(IsInGameThread());
 		if (IDerivedDataCacheModule* Module = FModuleManager::LoadModulePtr<IDerivedDataCacheModule>("DerivedDataCache"))
 		{
 			return static_cast<FDerivedDataCacheModule*>(Module);
@@ -95,12 +94,13 @@ ICache& GetCache()
 	{
 		return *Cache;
 	}
+	checkf(IsInGameThread(), TEXT("The derived data cache must be created on the main thread."));
 	if (Private::FDerivedDataCacheModule* Module = Private::GetModule())
 	{
 		Module->CreateCacheOnce();
 	}
 	ICache* Cache = Private::GDerivedDataCache;
-	check(Cache);
+	checkf(Cache, TEXT("Failed to create derived data cache."));
 	return *Cache;
 }
 
@@ -110,12 +110,13 @@ IBuild& GetBuild()
 	{
 		return *Build;
 	}
+	checkf(IsInGameThread(), TEXT("The derived data build system must be created on the main thread."));
 	if (Private::FDerivedDataCacheModule* Module = Private::GetModule())
 	{
 		Module->CreateBuildOnce();
 	}
 	IBuild* Build = Private::GDerivedDataBuild;
-	check(Build);
+	checkf(Build, TEXT("Failed to create derived data build system."));
 	return *Build;
 }
 
diff --git a/Engine/Source/Developer/DerivedDataCache/Private/DerivedDataCachePrivate.h b/Engine/Source/Developer/DerivedDataCache/Private/DerivedDataCachePrivate.h
--- a/Engine/Source/Developer/DerivedDataCache/Private/DerivedDataCachePrivate.h
+++ b/Engine/Source/Developer/DerivedDataCache/Private/DerivedDataCachePrivate.h
@@ -5,14 +5,8 @@
 #include "Containers/StringFwd.h"
 #include "Misc/AsciiSet.h"
 
-class FCbPackage;
 class FDerivedDataCacheInterface;
-enum class ECachePolicy : uint8;
 
-namespace UE::DerivedData { class FCacheBucket; }
-namespace UE::DerivedData { class FCacheRecord; }
-namespace UE::DerivedData { class FCacheRecordBuilder; }
-namespace UE::DerivedData { class FOptionalCacheRecord; }
 namespace UE::DerivedData { class ICache; }
 namespace UE::DerivedData { struct FCacheKey; }
 
@@ -22,27 +16,13 @@ namespace UE::DerivedData::Private
 // Implemented in DerivedDataCache.cpp
 ICache* CreateCache(FDerivedDataCacheInterface** OutLegacyCache);
 
-// Implemented in DerivedDataCacheKey.cpp
-FCacheBucket CreateCacheBucket(FStringView Name);
-
-// Implemented in DerivedDataCacheRecord.cpp
-FCacheRecordBuilder CreateCacheRecordBuilder(const FCacheKey& Key);
-FCbPackage SaveCacheRecord(const FCacheRecord& Record);
-FOptionalCacheRecord LoadCacheRecord(const FCbPackage& Package);
-
 ///////////////////////////////////////////////////////////////////////////////////////////////////////////////////////
 
-inline bool IsValidCacheBucketName(FStringView Name)
+template <typename CharType>
+inline bool IsValidCacheBucketName(TStringView<CharType> Name)
 {
 	constexpr FAsciiSet Valid("ABCDEFGHIJKLMNOPQRSTUVWXYZabcdefghijklmnopqrstuvwxyz0123456789");
 	return !Name.IsEmpty() && Name.Len() < 256 && FAsciiSet::HasOnly(Name, Valid);
 }
 
-inline void AssertValidCacheBucketName(FStringView Name)
-{
-	checkf(IsValidCacheBucketName(Name),
-		TEXT("A cache bucket name must be alphanumeric, non-empty, and contain fewer than 256 code units. ")
-		TEXT("Name: '%.*s'"), Name.Len(), Name.GetData());
-}
-
 } // UE::DerivedData::Private
diff --git a/Engine/Source/Developer/DerivedDataCache/Private/DerivedDataCacheRecord.cpp b/Engine/Source/Developer/DerivedDataCache/Private/DerivedDataCacheRecord.cpp
--- a/Engine/Source/Developer/DerivedDataCache/Private/DerivedDataCacheRecord.cpp
+++ b/Engine/Source/Developer/DerivedDataCache/Private/DerivedDataCacheRecord.cpp
@@ -2,10 +2,6 @@
 
 #include "DerivedDataCacheRecord.h"
 
-#include "Algo/AllOf.h"
-#include "Containers/StringConv.h"
-#include "Containers/UnrealString.h"
-#include "DerivedDataCache.h"
 #include "DerivedDataCacheKey.h"
 #include "DerivedDataCachePrivate.h"
 #include "DerivedDataPayloadPrivate.h"
@@ -16,7 +12,6 @@
 #include "Serialization/CompactBinary.h"
 #include "Serialization/CompactBinaryPackage.h"
 #include "Serialization/CompactBinaryWriter.h"
-#include "UObject/NameTypes.h"
 #include <atomic>
 
 namespace UE::DerivedData::Private
@@ -273,32 +268,27 @@ FCacheRecord CreateCacheRecord(ICacheRecordInternal* Record)
 	return FCacheRecord(Record);
 }
 
-FCacheRecordBuilder CreateCacheRecordBuilder(ICacheRecordBuilderInternal* RecordBuilder)
-{
-	return FCacheRecordBuilder(RecordBuilder);
-}
+} // UE::DerivedData::Private
 
 ///////////////////////////////////////////////////////////////////////////////////////////////////////////////////////
 
-FCacheRecordBuilder CreateCacheRecordBuilder(const FCacheKey& Key)
+namespace UE::DerivedData
 {
-	return CreateCacheRecordBuilder(new FCacheRecordBuilderInternal(Key));
-}
 
-FCbPackage SaveCacheRecord(const FCacheRecord& Record)
+FCbPackage FCacheRecord::Save() const
 {
 	FCbPackage Package;
 	FCbWriter Writer;
 	Writer.BeginObject();
 	{
-		const FCacheKey& Key = Record.GetKey();
+		const FCacheKey& Key = GetKey();
 		Writer.BeginObject("Key"_ASV);
-		Writer.AddString("Bucket"_ASV, Key.Bucket.ToString<ANSICHAR>());
+		Writer.AddString("Bucket"_ASV, Key.Bucket.ToString());
 		Writer.AddHash("Hash"_ASV, Key.Hash);
 		Writer.EndObject();
 	}
 
-	if (const FCbObject& Meta = Record.GetMeta())
+	if (const FCbObject& Meta = GetMeta())
 	{
 		Writer.AddObject("Meta"_ASV, Meta);
 	}
@@ -310,16 +300,16 @@ FCbPackage SaveCacheRecord(const FCacheRecord& Record)
 		Writer.AddInteger("RawSize"_ASV, Payload.GetRawSize());
 		const FCompositeBuffer CompressedBuffer = Payload.GetData().GetCompressed();
 		FCbAttachment AttachedCompressed(CompressedBuffer.Flatten());
-		Writer.AddAttachment("CompressedHash", AttachedCompressed);
+		Writer.AddAttachment("CompressedHash"_ASV, AttachedCompressed);
 		Package.AddAttachment(AttachedCompressed);
 		Writer.EndObject();
 	};
-	if (const FPayload& Value = Record.GetValuePayload())
+	if (const FPayload& Value = GetValuePayload())
 	{
 		Writer.SetName("Value"_ASV);
 		SavePayload(Value);
 	}
-	TConstArrayView<FPayload> Attachments = Record.GetAttachmentPayloads();
+	TConstArrayView<FPayload> Attachments = GetAttachmentPayloads();
 	if (!Attachments.IsEmpty())
 	{
 		Writer.BeginArray("Attachments"_ASV);
@@ -335,28 +325,28 @@ FCbPackage SaveCacheRecord(const FCacheRecord& Record)
 	return Package;
 }
 
-FOptionalCacheRecord LoadCacheRecord(const FCbPackage& Package)
+FOptionalCacheRecord FCacheRecord::Load(const FCbPackage& Package)
 {
 	FCbObjectView RecordObject = Package.GetObject();
-	
+
 	FCacheKey Key;
 	FCbObjectView KeyObject = RecordObject["Key"_ASV].AsObjectView();
-	auto TrySetBucketName = [](FStringView Name, FCacheKey& Key)
+	auto TrySetBucketName = [](FUtf8StringView Name, FCacheKey& Key)
 	{
-		if (IsValidCacheBucketName(Name))
+		if (Private::IsValidCacheBucketName(Name))
 		{
-			Key.Bucket = CreateCacheBucket(Name);
+			Key.Bucket = FCacheBucket(Name);
 			return true;
 		}
 		return false;
 	};
-	if (!TrySetBucketName(FUTF8ToTCHAR(KeyObject["Bucket"_ASV].AsString()), Key))
+	if (!TrySetBucketName(KeyObject["Bucket"_ASV].AsString(), Key))
 	{
 		return FOptionalCacheRecord();
 	}
 	Key.Hash = KeyObject["Hash"_ASV].AsHash();
 
-	FCacheRecordBuilder Builder = CreateCacheRecordBuilder(Key);
+	FCacheRecordBuilder Builder(Key);
 
 	Builder.SetMeta(Package.GetObject()["Meta"_ASV].AsObject());
 
@@ -410,6 +400,9 @@ FOptionalCacheRecord LoadCacheRecord(const FCbPackage& Package)
 	return Builder.Build();
 }
 
-///////////////////////////////////////////////////////////////////////////////////////////////////////////////////////
+FCacheRecordBuilder::FCacheRecordBuilder(const FCacheKey& Key)
+	: RecordBuilder(new Private::FCacheRecordBuilderInternal(Key))
+{
+}
 
-} // UE::DerivedData::Private
+} // UE::DerivedData
diff --git a/Engine/Source/Developer/DerivedDataCache/Private/FileSystemDerivedDataBackend.cpp b/Engine/Source/Developer/DerivedDataCache/Private/FileSystemDerivedDataBackend.cpp
--- a/Engine/Source/Developer/DerivedDataCache/Private/FileSystemDerivedDataBackend.cpp
+++ b/Engine/Source/Developer/DerivedDataCache/Private/FileSystemDerivedDataBackend.cpp
@@ -108,9 +108,8 @@ public:
 	 * @param InCacheDirectory	directory to store the cache in
 	 * @param bForceReadOnly	if true, do not attempt to write to this cache
 	*/
-	FFileSystemDerivedDataBackend(ICacheFactory& InFactory, const TCHAR* InCacheDirectory, const TCHAR* InParams, const TCHAR* InAccessLogFileName)
-		: Factory(InFactory)
-		, CachePath(InCacheDirectory)
+	FFileSystemDerivedDataBackend(const TCHAR* InCacheDirectory, const TCHAR* InParams, const TCHAR* InAccessLogFileName)
+		: CachePath(InCacheDirectory)
 		, SpeedClass(ESpeedClass::Unknown)
 		, bReadOnly(false)
 		, bTouch(false)
@@ -943,7 +942,7 @@ public:
 			{
 				if (OnComplete)
 				{
-					OnComplete({Factory.CreateRecord(Key).Build(), EStatus::Error});
+					OnComplete({FCacheRecordBuilder(Key).Build(), EStatus::Error});
 				}
 			}
 		}
@@ -1122,7 +1121,7 @@ private:
 		}
 
 		const FCbObject RecordObject(MoveTemp(Buffer));
-		FCacheRecordBuilder RecordBuilder = Factory.CreateRecord(Key);
+		FCacheRecordBuilder RecordBuilder(Key);
 
 		if (!EnumHasAnyFlags(Policy, ECachePolicy::SkipMeta))
 		{
@@ -1463,7 +1462,6 @@ private:
 		return CachePath / BuildPathForCacheKey(CacheKey);
 	}
 
-	ICacheFactory& Factory;
 	/** Base path we are storing the cache files in. **/
 	FString	CachePath;
 	/** Class of this cache */
@@ -1533,9 +1531,9 @@ private:
 
 };
 
-FDerivedDataBackendInterface* CreateFileSystemDerivedDataBackend(ICacheFactory& Factory, const TCHAR* CacheDirectory, const TCHAR* InParams, const TCHAR* InAccessLogFileName /*= nullptr*/)
+FDerivedDataBackendInterface* CreateFileSystemDerivedDataBackend(const TCHAR* CacheDirectory, const TCHAR* InParams, const TCHAR* InAccessLogFileName /*= nullptr*/)
 {
-	FFileSystemDerivedDataBackend* FileDDB = new FFileSystemDerivedDataBackend(Factory, CacheDirectory, InParams, InAccessLogFileName);
+	FFileSystemDerivedDataBackend* FileDDB = new FFileSystemDerivedDataBackend(CacheDirectory, InParams, InAccessLogFileName);
 
 	if (!FileDDB->IsUsable())
 	{
diff --git a/Engine/Source/Developer/DerivedDataCache/Private/HierarchicalDerivedDataBackend.h b/Engine/Source/Developer/DerivedDataCache/Private/HierarchicalDerivedDataBackend.h
--- a/Engine/Source/Developer/DerivedDataCache/Private/HierarchicalDerivedDataBackend.h
+++ b/Engine/Source/Developer/DerivedDataCache/Private/HierarchicalDerivedDataBackend.h
@@ -28,9 +28,8 @@ public:
 	 * Constructor
 	 * @param	InInnerBackends Backends to call into for actual storage of the cache, first item is the "fastest cache"
 	 */
-	FHierarchicalDerivedDataBackend(ICacheFactory& InFactory, const TArray<FDerivedDataBackendInterface*>& InInnerBackends)
-		: Factory(InFactory)
-		, InnerBackends(InInnerBackends)
+	FHierarchicalDerivedDataBackend(const TArray<FDerivedDataBackendInterface*>& InInnerBackends)
+		: InnerBackends(InInnerBackends)
 		, bIsWritable(false)
 		, bHasLocalBackends(false)
 		, bHasRemoteBackends(false)
@@ -88,7 +87,7 @@ private:
 			// async puts to allow us to fill all levels without holding up the engine
 			// we need to cache inflight puts to avoid having inconsistent miss and redownload on lower cache levels while puts are still async
 			const bool bCacheInFlightPuts = true;
-			AsyncPutInnerBackends.Emplace(new FDerivedDataBackendAsyncPutWrapper(Factory, InnerBackends[CacheIndex], bCacheInFlightPuts));
+			AsyncPutInnerBackends.Emplace(new FDerivedDataBackendAsyncPutWrapper(InnerBackends[CacheIndex], bCacheInFlightPuts));
 		}
 	}
 
@@ -597,7 +596,7 @@ public:
 		{
 			for (const FCacheKey& Key : RemainingKeys)
 			{
-				OnComplete({Factory.CreateRecord(Key).Build(), EStatus::Error});
+				OnComplete({FCacheRecordBuilder(Key).Build(), EStatus::Error});
 			}
 		}
 	}
@@ -656,7 +655,6 @@ private:
 	/** To avoid race while manipulating the arrays */
 	mutable FRWLock Lock;
 
-	ICacheFactory& Factory;
 	/** Array of backends forming the hierarchical cache...the first element is the fastest cache. **/
 	TArray<FDerivedDataBackendInterface*> InnerBackends;
 	/** Each of the backends wrapped with an async put **/
diff --git a/Engine/Source/Developer/DerivedDataCache/Private/HttpDerivedDataBackend.cpp b/Engine/Source/Developer/DerivedDataCache/Private/HttpDerivedDataBackend.cpp
--- a/Engine/Source/Developer/DerivedDataCache/Private/HttpDerivedDataBackend.cpp
+++ b/Engine/Source/Developer/DerivedDataCache/Private/HttpDerivedDataBackend.cpp
@@ -1910,15 +1910,13 @@ uint32 FHttpAccessToken::GetSerial() const
 //----------------------------------------------------------------------------------------------------------
 
 FHttpDerivedDataBackend::FHttpDerivedDataBackend(
-	ICacheFactory& InFactory,
 	const TCHAR* InServiceUrl, 
 	const TCHAR* InNamespace, 
 	const TCHAR* InOAuthProvider,
 	const TCHAR* InOAuthClientId,
 	const TCHAR* InOAuthSecret,
 	const bool bInReadOnly)
-	: Factory(InFactory)
-	, Domain(InServiceUrl)
+	: Domain(InServiceUrl)
 	, Namespace(InNamespace)
 	, DefaultBucket(TEXT("default"))
 	, OAuthProvider(InOAuthProvider)
@@ -2458,7 +2456,7 @@ void FHttpDerivedDataBackend::Get(
 	{
 		for (const FCacheKey& Key : Keys)
 		{
-			OnComplete({Factory.CreateRecord(Key).Build(), EStatus::Error});
+			OnComplete({FCacheRecordBuilder(Key).Build(), EStatus::Error});
 		}
 	}
 }
diff --git a/Engine/Source/Developer/DerivedDataCache/Private/HttpDerivedDataBackend.h b/Engine/Source/Developer/DerivedDataCache/Private/HttpDerivedDataBackend.h
--- a/Engine/Source/Developer/DerivedDataCache/Private/HttpDerivedDataBackend.h
+++ b/Engine/Source/Developer/DerivedDataCache/Private/HttpDerivedDataBackend.h
@@ -36,7 +36,8 @@ public:
 	 * @param OAuthClientId	OAuth client identifier.
 	 * @param OAuthData		OAuth form data to send to login service. Can either be the raw form data or a Windows network file address (starting with "\\").
 	 */
-	FHttpDerivedDataBackend(ICacheFactory& Factory, const TCHAR* ServiceUrl, 
+	FHttpDerivedDataBackend(
+		const TCHAR* ServiceUrl, 
 		const TCHAR* Namespace, 
 		const TCHAR* OAuthProvider, 
 		const TCHAR* OAuthClientId, 
@@ -103,7 +104,6 @@ public:
 	}
 
 private:
-	ICacheFactory& Factory;
 	FString Domain;
 	FString Namespace;
 	FString DefaultBucket;
diff --git a/Engine/Source/Developer/DerivedDataCache/Private/MemoryDerivedDataBackend.cpp b/Engine/Source/Developer/DerivedDataCache/Private/MemoryDerivedDataBackend.cpp
--- a/Engine/Source/Developer/DerivedDataCache/Private/MemoryDerivedDataBackend.cpp
+++ b/Engine/Source/Developer/DerivedDataCache/Private/MemoryDerivedDataBackend.cpp
@@ -11,9 +11,8 @@
 namespace UE::DerivedData::Backends
 {
 
-FMemoryDerivedDataBackend::FMemoryDerivedDataBackend(ICacheFactory& InFactory, const TCHAR* InName, int64 InMaxCacheSize, bool bInCanBeDisabled)
-	: Factory(InFactory)
-	, Name(InName)
+FMemoryDerivedDataBackend::FMemoryDerivedDataBackend(const TCHAR* InName, int64 InMaxCacheSize, bool bInCanBeDisabled)
+	: Name(InName)
 	, MaxCacheSize(InMaxCacheSize)
 	, bDisabled( false )
 	, CurrentCacheSize( SerializationSpecificDataSize )
@@ -394,7 +393,7 @@ int64 FMemoryDerivedDataBackend::CalcSerializedCacheRecordSize(const FCacheRecor
 {
 	// Estimate the serialized size of the cache record.
 	uint64 TotalSize = 20;
-	TotalSize += Record.GetKey().Bucket.ToString<TCHAR>().Len();
+	TotalSize += Record.GetKey().Bucket.ToString().Len();
 	TotalSize += Record.GetMeta().GetSize();
 	const auto CalcCachePayloadSize = [](const FPayload& Payload)
 	{
@@ -499,7 +498,7 @@ void FMemoryDerivedDataBackend::Get(
 		{
 			if (OnComplete)
 			{
-				OnComplete({Factory.CreateRecord(Key).Build(), EStatus::Error});
+				OnComplete({FCacheRecordBuilder(Key).Build(), EStatus::Error});
 			}
 		}
 	}
diff --git a/Engine/Source/Developer/DerivedDataCache/Private/MemoryDerivedDataBackend.h b/Engine/Source/Developer/DerivedDataCache/Private/MemoryDerivedDataBackend.h
--- a/Engine/Source/Developer/DerivedDataCache/Private/MemoryDerivedDataBackend.h
+++ b/Engine/Source/Developer/DerivedDataCache/Private/MemoryDerivedDataBackend.h
@@ -21,7 +21,7 @@ namespace UE::DerivedData::Backends
 class FMemoryDerivedDataBackend : public FFileBackedDerivedDataBackend
 {
 public:
-	explicit FMemoryDerivedDataBackend(ICacheFactory& InFactory, const TCHAR* InName, int64 InMaxCacheSize = -1, bool bCanBeDisabled = false);
+	explicit FMemoryDerivedDataBackend(const TCHAR* InName, int64 InMaxCacheSize = -1, bool bCanBeDisabled = false);
 	~FMemoryDerivedDataBackend();
 
 	/** Return a name for this interface */
@@ -143,9 +143,6 @@ private:
 	int64 CalcRawCacheRecordSize(const FCacheRecord& Record) const;
 	int64 CalcSerializedCacheRecordSize(const FCacheRecord& Record) const;
 
-	/** Factory used to construct cache records. */
-	ICacheFactory& Factory;
-
 	/** Name of this cache (used for debugging) */
 	FString Name;
 
diff --git a/Engine/Source/Developer/DerivedDataCache/Private/PakFileDerivedDataBackend.cpp b/Engine/Source/Developer/DerivedDataCache/Private/PakFileDerivedDataBackend.cpp
--- a/Engine/Source/Developer/DerivedDataCache/Private/PakFileDerivedDataBackend.cpp
+++ b/Engine/Source/Developer/DerivedDataCache/Private/PakFileDerivedDataBackend.cpp
@@ -8,9 +8,8 @@
 namespace UE::DerivedData::Backends
 {
 
-FPakFileDerivedDataBackend::FPakFileDerivedDataBackend(ICacheFactory& InFactory, const TCHAR* InFilename, bool bInWriting)
-	: Factory(InFactory)
-	, bWriting(bInWriting)
+FPakFileDerivedDataBackend::FPakFileDerivedDataBackend(const TCHAR* InFilename, bool bInWriting)
+	: bWriting(bInWriting)
 	, bClosed(false)
 	, Filename(InFilename)
 {
@@ -397,13 +396,13 @@ void FPakFileDerivedDataBackend::MergeCache(FPakFileDerivedDataBackend* OtherPak
 	}
 }
 
-bool FPakFileDerivedDataBackend::SortAndCopy(ICacheFactory& InFactory, const FString &InputFilename, const FString &OutputFilename)
+bool FPakFileDerivedDataBackend::SortAndCopy(const FString &InputFilename, const FString &OutputFilename)
 {
 	// Open the input and output files
-	FPakFileDerivedDataBackend InputPak(InFactory, *InputFilename, false);
+	FPakFileDerivedDataBackend InputPak(*InputFilename, false);
 	if (InputPak.bClosed) return false;
 
-	FPakFileDerivedDataBackend OutputPak(InFactory, *OutputFilename, true);
+	FPakFileDerivedDataBackend OutputPak(*OutputFilename, true);
 	if (OutputPak.bClosed) return false;
 
 	// Sort the key names
@@ -468,7 +467,7 @@ void FPakFileDerivedDataBackend::Get(
 	{
 		for (const FCacheKey& Key : Keys)
 		{
-			OnComplete({Factory.CreateRecord(Key).Build(), EStatus::Error});
+			OnComplete({FCacheRecordBuilder(Key).Build(), EStatus::Error});
 		}
 	}
 }
@@ -489,8 +488,8 @@ void FPakFileDerivedDataBackend::GetPayload(
 	}
 }
 
-FCompressedPakFileDerivedDataBackend::FCompressedPakFileDerivedDataBackend(ICacheFactory& InFactory, const TCHAR* InFilename, bool bInWriting)
-	: FPakFileDerivedDataBackend(InFactory, InFilename, bInWriting)
+FCompressedPakFileDerivedDataBackend::FCompressedPakFileDerivedDataBackend(const TCHAR* InFilename, bool bInWriting)
+	: FPakFileDerivedDataBackend(InFilename, bInWriting)
 {
 }
 
diff --git a/Engine/Source/Developer/DerivedDataCache/Private/PakFileDerivedDataBackend.h b/Engine/Source/Developer/DerivedDataCache/Private/PakFileDerivedDataBackend.h
--- a/Engine/Source/Developer/DerivedDataCache/Private/PakFileDerivedDataBackend.h
+++ b/Engine/Source/Developer/DerivedDataCache/Private/PakFileDerivedDataBackend.h
@@ -25,7 +25,7 @@ namespace UE::DerivedData::Backends
 class FPakFileDerivedDataBackend : public FDerivedDataBackendInterface
 {
 public:
-	FPakFileDerivedDataBackend(ICacheFactory& InFactory, const TCHAR* InFilename, bool bInWriting);
+	FPakFileDerivedDataBackend(const TCHAR* InFilename, bool bInWriting);
 	~FPakFileDerivedDataBackend();
 
 	void Close();
@@ -91,7 +91,7 @@ public:
 		return Filename;
 	}
 
-	static bool SortAndCopy(ICacheFactory& InFactory, const FString &InputFilename, const FString &OutputFilename);
+	static bool SortAndCopy(const FString &InputFilename, const FString &OutputFilename);
 
 	virtual TSharedRef<FDerivedDataCacheStatsNode> GatherUsageStats() const override;
 
@@ -145,7 +145,6 @@ private:
 		}
 	};
 
-	ICacheFactory& Factory;
 	/** When set to true, we are a pak writer (we don't do reads). */
 	bool bWriting;
 	/** When set to true, we are a pak writer and we saved, so we shouldn't be used anymore. Also, a read cache that failed to open. */
@@ -168,7 +167,7 @@ private:
 class FCompressedPakFileDerivedDataBackend : public FPakFileDerivedDataBackend
 {
 public:
-	FCompressedPakFileDerivedDataBackend(ICacheFactory& InFactory, const TCHAR* InFilename, bool bInWriting);
+	FCompressedPakFileDerivedDataBackend(const TCHAR* InFilename, bool bInWriting);
 
 	virtual EPutStatus PutCachedData(const TCHAR* CacheKey, TArrayView<const uint8> InData, bool bPutEvenIfExists) override;
 	virtual bool GetCachedData(const TCHAR* CacheKey, TArray<uint8>& OutData) override;
diff --git a/Engine/Source/Developer/DerivedDataCache/Private/S3DerivedDataBackend.cpp b/Engine/Source/Developer/DerivedDataCache/Private/S3DerivedDataBackend.cpp
--- a/Engine/Source/Developer/DerivedDataCache/Private/S3DerivedDataBackend.cpp
+++ b/Engine/Source/Developer/DerivedDataCache/Private/S3DerivedDataBackend.cpp
@@ -762,9 +762,8 @@ struct FS3DerivedDataBackend::FRootManifest
 // FS3DerivedDataBackend
 //----------------------------------------------------------------------------------------------------------
 
-FS3DerivedDataBackend::FS3DerivedDataBackend(ICacheFactory& InFactory, const TCHAR* InRootManifestPath, const TCHAR* InBaseUrl, const TCHAR* InRegion, const TCHAR* InCanaryObjectKey, const TCHAR* InCachePath)
-	: Factory(InFactory)
-	, RootManifestPath(InRootManifestPath)
+FS3DerivedDataBackend::FS3DerivedDataBackend(const TCHAR* InRootManifestPath, const TCHAR* InBaseUrl, const TCHAR* InRegion, const TCHAR* InCanaryObjectKey, const TCHAR* InCachePath)
+	: RootManifestPath(InRootManifestPath)
 	, BaseUrl(InBaseUrl)
 	, Region(InRegion)
 	, CanaryObjectKey(InCanaryObjectKey)
@@ -1204,7 +1203,7 @@ void FS3DerivedDataBackend::Get(
 	{
 		for (const FCacheKey& Key : Keys)
 		{
-			OnComplete({Factory.CreateRecord(Key).Build(), EStatus::Error});
+			OnComplete({FCacheRecordBuilder(Key).Build(), EStatus::Error});
 		}
 	}
 }
diff --git a/Engine/Source/Developer/DerivedDataCache/Private/S3DerivedDataBackend.h b/Engine/Source/Developer/DerivedDataCache/Private/S3DerivedDataBackend.h
--- a/Engine/Source/Developer/DerivedDataCache/Private/S3DerivedDataBackend.h
+++ b/Engine/Source/Developer/DerivedDataCache/Private/S3DerivedDataBackend.h
@@ -31,7 +31,7 @@ public:
 	 * @param  InCanaryObjectKey    Key for a canary object used to test whether this backend is usable
 	 * @param  InCachePath          Path to cache the DDC files
 	 */
-	FS3DerivedDataBackend(ICacheFactory& Factory, const TCHAR* InRootManifestPath, const TCHAR* InBaseUrl, const TCHAR* InRegion, const TCHAR* InCanaryObjectKey, const TCHAR* InCachePath);
+	FS3DerivedDataBackend(const TCHAR* InRootManifestPath, const TCHAR* InBaseUrl, const TCHAR* InRegion, const TCHAR* InCanaryObjectKey, const TCHAR* InCachePath);
 	~FS3DerivedDataBackend();
 
 	/**
@@ -94,7 +94,6 @@ private:
 	class FHttpRequest;
 	class FRequestPool;
 
-	ICacheFactory& Factory;
 	FString RootManifestPath;
 	FString BaseUrl;
 	FString Region;
diff --git a/Engine/Source/Developer/DerivedDataCache/Private/ZenDerivedDataBackend.cpp b/Engine/Source/Developer/DerivedDataCache/Private/ZenDerivedDataBackend.cpp
--- a/Engine/Source/Developer/DerivedDataCache/Private/ZenDerivedDataBackend.cpp
+++ b/Engine/Source/Developer/DerivedDataCache/Private/ZenDerivedDataBackend.cpp
@@ -32,11 +32,9 @@ namespace UE::DerivedData::Backends {
 //----------------------------------------------------------------------------------------------------------
 
 FZenDerivedDataBackend::FZenDerivedDataBackend(
-	ICacheFactory& InFactory, 
 	const TCHAR* InServiceUrl,
 	const TCHAR* InNamespace)
-	: Factory(InFactory)
-	, Domain(InServiceUrl)
+	: Domain(InServiceUrl)
 	, Namespace(InNamespace)
 {
 	if (IsServiceReady())
@@ -312,16 +310,13 @@ FString FZenDerivedDataBackend::MakeLegacyZenKey(const TCHAR* CacheKey)
 
 void FZenDerivedDataBackend::AppendZenUri(const FCacheKey& CacheKey, FStringBuilderBase& Out)
 {
-	FStringView BucketStr = CacheKey.Bucket.ToString<TCHAR>();
-	Out.Appendf(TEXT("/z$/%.*s/"), BucketStr.Len(), BucketStr.GetData());
-	Out << CacheKey.Hash;
+	Out << TEXT("/z$/") << CacheKey.Bucket << TEXT('/') << CacheKey.Hash;
 }
 
 void FZenDerivedDataBackend::AppendZenUri(const FCacheKey& CacheKey, const FPayloadId& PayloadId, FStringBuilderBase& Out)
 {
 	AppendZenUri(CacheKey, Out);
-	Out << TEXT("/");
-	UE::String::BytesToHexLower(PayloadId.GetBytes(), Out);
+	Out << TEXT('/') << PayloadId;
 }
 
 void FZenDerivedDataBackend::AppendPolicyQueryString(ECachePolicy Policy, FStringBuilderBase& Uri)
@@ -566,7 +561,7 @@ void FZenDerivedDataBackend::Get(
 		{
 			if (OnComplete)
 			{
-				OnComplete({ Factory.CreateRecord(Key).Build(), EStatus::Error });
+				OnComplete({ FCacheRecordBuilder(Key).Build(), EStatus::Error });
 			}
 		}
 	}
@@ -671,7 +666,7 @@ void FZenDerivedDataBackend::CancelAll()
 bool FZenDerivedDataBackend::PutCacheRecord(const FCacheRecord& Record, FStringView Context, ECachePolicy Policy)
 {
 	const FCacheKey& Key = Record.GetKey();
-	FCbPackage Package = Factory.SaveRecord(Record);
+	FCbPackage Package = Record.Save();
 	FCbWriter PackageWriter;
 	Package.Save(PackageWriter);
 	FUniqueBuffer Buffer = FUniqueBuffer::Alloc(PackageWriter.GetSaveSize());
@@ -708,7 +703,7 @@ FOptionalCacheRecord FZenDerivedDataBackend::GetCacheRecord(const FCacheKey& Key
 		}
 		return FOptionalCacheRecord();
 	}
-	FOptionalCacheRecord Record = Factory.LoadRecord(Package);
+	FOptionalCacheRecord Record = FCacheRecord::Load(Package);
 	if (!Record)
 	{
 		UE_LOG(LogDerivedDataCache, Verbose, TEXT("%s: Cache miss with corrupted record (corrupted record) for %s from '%.*s'"),
@@ -873,7 +868,7 @@ FOptionalCacheRecord FZenDerivedDataBackend::LegacyCreateRecord(FSharedBuffer&&
 	// Eventually We will instead send the Policy to zen sever along with the request for the FCacheKey,
 	// and it will end back all payloads in the response rather than requiring separate requests.
 	const FCbObject RecordObject(MoveTemp(RecordBytes));
-	FCacheRecordBuilder RecordBuilder = Factory.CreateRecord(Key);
+	FCacheRecordBuilder RecordBuilder(Key);
 
 	if (!EnumHasAnyFlags(Policy, ECachePolicy::SkipMeta))
 	{
@@ -973,15 +968,13 @@ FPayload FZenDerivedDataBackend::LegacyValidateCachePayload(const FCacheKey& Key
 void FZenDerivedDataBackend::LegacyMakeZenKey(const FCacheKey& CacheKey, FStringBuilderBase& Out) const
 {
 	Out.Reset();
-	FStringView Bucket = CacheKey.Bucket.ToString<TCHAR>();
-	Out.Appendf(TEXT("/z$/%.*s/%s"), Bucket.Len(), Bucket.GetData(), *LexToString(CacheKey.Hash));
+	Out << TEXT("/z$/") << CacheKey;
 }
 
 void FZenDerivedDataBackend::LegacyMakePayloadKey(const FCacheKey& CacheKey, const FIoHash& RawHash, FStringBuilderBase& Out) const
 {
 	Out.Reset();
-	FStringView Bucket = CacheKey.Bucket.ToString<TCHAR>();
-	Out.Appendf(TEXT("/z$/%.*s/%s/%s"), Bucket.Len(), Bucket.GetData(), *LexToString(CacheKey.Hash), *LexToString(RawHash));
+	Out << TEXT("/z$/") << CacheKey << TEXT('/') << RawHash;
 }
 
 
diff --git a/Engine/Source/Developer/DerivedDataCache/Private/ZenDerivedDataBackend.h b/Engine/Source/Developer/DerivedDataCache/Private/ZenDerivedDataBackend.h
--- a/Engine/Source/Developer/DerivedDataCache/Private/ZenDerivedDataBackend.h
+++ b/Engine/Source/Developer/DerivedDataCache/Private/ZenDerivedDataBackend.h
@@ -55,7 +55,7 @@ public:
 	 * @param ServiceUrl	Base url to the service including schema.
 	 * @param Namespace		Namespace to use.
 	 */
-	FZenDerivedDataBackend(ICacheFactory& Factory, const TCHAR* ServiceUrl, const TCHAR* Namespace);
+	FZenDerivedDataBackend(const TCHAR* ServiceUrl, const TCHAR* Namespace);
 
 	~FZenDerivedDataBackend();
 
@@ -157,7 +157,6 @@ private:
 	bool DidSimulateMiss(const TCHAR* InKey);
 	bool ShouldSimulateMiss(const TCHAR* InKey);
 private:
-	ICacheFactory& Factory;
 	FString Domain;
 	FString Namespace;
 	mutable FDerivedDataCacheUsageStats UsageStats;
diff --git a/Engine/Source/Developer/DerivedDataCache/Public/DerivedDataBuildFunction.h b/Engine/Source/Developer/DerivedDataCache/Public/DerivedDataBuildFunction.h
--- a/Engine/Source/Developer/DerivedDataCache/Public/DerivedDataBuildFunction.h
+++ b/Engine/Source/Developer/DerivedDataCache/Public/DerivedDataBuildFunction.h
@@ -134,9 +134,6 @@ public:
 	/** Returns the constant with the matching key, or an object with no fields if not found. */
 	virtual FCbObject FindConstant(FStringView Key) const = 0;
 
-	/** Create a cache bucket from a name. See ICacheFactory::CreateBucket. */
-	virtual FCacheBucket CreateCacheBucket(FStringView Name) const = 0;
-
 	/** Overrides the cache bucket used when reading or writing this build in the cache. */
 	virtual void SetCacheBucket(FCacheBucket Bucket) = 0;
 
diff --git a/Engine/Source/Developer/DerivedDataCache/Public/DerivedDataCache.h b/Engine/Source/Developer/DerivedDataCache/Public/DerivedDataCache.h
--- a/Engine/Source/Developer/DerivedDataCache/Public/DerivedDataCache.h
+++ b/Engine/Source/Developer/DerivedDataCache/Public/DerivedDataCache.h
@@ -8,17 +8,12 @@
 #include "Misc/EnumClassFlags.h"
 #include "Templates/Function.h"
 
-class FCbPackage;
-
 #define UE_API DERIVEDDATACACHE_API
 
 namespace UE::DerivedData
 {
 
-class FCacheBucket;
 class FCacheRecord;
-class FCacheRecordBuilder;
-class FOptionalCacheRecord;
 class FPayload;
 struct FCacheGetCompleteParams;
 struct FCacheGetPayloadCompleteParams;
@@ -96,42 +91,6 @@ ENUM_CLASS_FLAGS(ECachePolicy);
 
 ///////////////////////////////////////////////////////////////////////////////////////////////////////////////////////
 
-/**
- * Interface to construct types related to the cache.
- */
-class ICacheFactory
-{
-public:
-	virtual ~ICacheFactory() = default;
-
-	/**
-	 * Create a cache bucket from a name.
-	 *
-	 * A cache bucket name must be alphanumeric, non-empty, and contain fewer than 256 code units.
-	 */
-	virtual FCacheBucket CreateBucket(FStringView Name) = 0;
-
-	/**
-	 * Create a cache record builder from a cache key.
-	 */
-	virtual FCacheRecordBuilder CreateRecord(const FCacheKey& Key) = 0;
-
-	/**
-	 * Save the cache record to a compact binary package.
-	 */
-	virtual FCbPackage SaveRecord(const FCacheRecord& Record) = 0;
-
-	/**
-	 * Load a cache record from compact binary.
-	 *
-	 * @param Package   A package saved from a cache record.
-	 * @return A valid cache record, or null on error.
-	 */
-	virtual FOptionalCacheRecord LoadRecord(const FCbPackage& Package) = 0;
-};
-
-///////////////////////////////////////////////////////////////////////////////////////////////////////////////////////
-
 /**
  * Interface to a store of cache records.
  *
@@ -222,9 +181,8 @@ public:
  * Interface to the cache.
  *
  * @see ICacheStore for cache record storage.
- * @see ICacheFactory for cache type constructors.
  */
-class ICache : public ICacheStore, public ICacheFactory
+class ICache : public ICacheStore
 {
 public:
 	virtual ~ICache() = default;
@@ -280,6 +238,7 @@ struct FCacheGetPayloadCompleteParams
 
 ///////////////////////////////////////////////////////////////////////////////////////////////////////////////////////
 
+/** Returns a reference to the cache. Asserts if not available. */
 UE_API ICache& GetCache();
 
 ///////////////////////////////////////////////////////////////////////////////////////////////////////////////////////
diff --git a/Engine/Source/Developer/DerivedDataCache/Public/DerivedDataCacheKey.h b/Engine/Source/Developer/DerivedDataCache/Public/DerivedDataCacheKey.h
--- a/Engine/Source/Developer/DerivedDataCache/Public/DerivedDataCacheKey.h
+++ b/Engine/Source/Developer/DerivedDataCache/Public/DerivedDataCacheKey.h
@@ -9,6 +9,8 @@
 #include "IO/IoHash.h"
 #include "Templates/TypeHash.h"
 
+#define UE_API DERIVEDDATACACHE_API
+
 namespace UE::DerivedData
 {
 
@@ -17,7 +19,7 @@ namespace UE::DerivedData
 /**
  * An alphanumeric identifier that groups related cache records.
  *
- * Create using ICache::CreateBucket.
+ * A cache bucket name must be alphanumeric, non-empty, and contain fewer than 256 code units.
  */
 class FCacheBucket
 {
@@ -25,6 +27,10 @@ public:
 	/** Construct a null cache bucket. */
 	FCacheBucket() = default;
 
+	/** Create a cache bucket from a name. */
+	UE_API explicit FCacheBucket(FUtf8StringView Name);
+	UE_API explicit FCacheBucket(FWideStringView Name);
+
 	/** Whether this is null. */
 	inline bool IsNull() const { return !Name; }
 	/** Whether this is not null. */
@@ -38,7 +44,7 @@ public:
 
 	inline bool operator<(FCacheBucket Other) const
 	{
-		return Name != Other.Name && ToString<TCHAR>().Compare(Other.ToString<TCHAR>(), ESearchCase::IgnoreCase) < 0;
+		return Name != Other.Name && ToString().Compare(Other.ToString(), ESearchCase::IgnoreCase) < 0;
 	}
 
 	friend inline uint32 GetTypeHash(FCacheBucket Bucket)
@@ -47,30 +53,13 @@ public:
 	}
 
 	/** Get the name of the cache bucket as a string. */
-	template <typename CharType>
-	inline TStringView<CharType> ToString() const;
-	template <>
-	inline FStringView ToString<TCHAR>() const;
-	template <>
-	inline FAnsiStringView ToString<ANSICHAR>() const;
+	inline FAnsiStringView ToString() const;
 
 protected:
-	explicit FCacheBucket(const TCHAR* InName)
-		: Name(InName)
-	{
-	}
-
 	static constexpr int32 LengthOffset = -1;
-	static constexpr int32 LengthOffset8 = -2;
-
-private:
-	/**
-	 * Name as TCHAR followed by the name as UTF-8, both null-terminated.
-	 *
-	 * The byte preceding the TCHAR name is the number of TCHAR code units excluding the terminator.
-	 * The byte preceding the TCHAR name length is the number of UTF-8 code units.
-	 */
-	const TCHAR* Name = nullptr;
+
+	/** Name stored as a null-terminated string preceded by one byte containing its length. */
+	const ANSICHAR* Name = nullptr;
 };
 
 ///////////////////////////////////////////////////////////////////////////////////////////////////////////////////////
@@ -101,24 +90,15 @@ inline const FCachePayloadKey FCachePayloadKey::Empty;
 
 ///////////////////////////////////////////////////////////////////////////////////////////////////////////////////////
 
-template <>
-inline FStringView FCacheBucket::ToString<TCHAR>() const
+inline FAnsiStringView FCacheBucket::ToString() const
 {
-	return Name ? FStringView(Name, reinterpret_cast<const uint8*>(Name)[LengthOffset]) : FStringView();
-}
-
-template <>
-inline FAnsiStringView FCacheBucket::ToString<ANSICHAR>() const
-{
-	return Name ? FAnsiStringView(
-		reinterpret_cast<const ANSICHAR*>(Name + reinterpret_cast<const uint8*>(Name)[LengthOffset] + 1),
-		reinterpret_cast<const uint8*>(Name)[LengthOffset8]) : FAnsiStringView();
+	return Name ? FAnsiStringView(Name, reinterpret_cast<const uint8*>(Name)[LengthOffset]) : FAnsiStringView();
 }
 
 template <typename CharType>
 inline TStringBuilderBase<CharType>& operator<<(TStringBuilderBase<CharType>& Builder, const FCacheBucket& Bucket)
 {
-	return Builder << Bucket.ToString<CharType>();
+	return Builder << Bucket.ToString();
 }
 
 ///////////////////////////////////////////////////////////////////////////////////////////////////////////////////////
@@ -184,3 +164,5 @@ inline TStringBuilderBase<CharType>& operator<<(TStringBuilderBase<CharType>& Bu
 ///////////////////////////////////////////////////////////////////////////////////////////////////////////////////////
 
 } // UE::DerivedData
+
+#undef UE_API
diff --git a/Engine/Source/Developer/DerivedDataCache/Public/DerivedDataCacheRecord.h b/Engine/Source/Developer/DerivedDataCache/Public/DerivedDataCacheRecord.h
--- a/Engine/Source/Developer/DerivedDataCache/Public/DerivedDataCacheRecord.h
+++ b/Engine/Source/Developer/DerivedDataCache/Public/DerivedDataCacheRecord.h
@@ -10,10 +10,13 @@
 #include "Templates/RefCounting.h"
 #include "Templates/UniquePtr.h"
 
+#define UE_API DERIVEDDATACACHE_API
+
 class FCbObject;
+class FCbPackage;
 
 namespace UE::DerivedData { class FCacheRecord; }
-namespace UE::DerivedData { class FCacheRecordBuilder; }
+namespace UE::DerivedData { class FOptionalCacheRecord; }
 namespace UE::DerivedData { using FOnCacheRecordComplete = TUniqueFunction<void (FCacheRecord&& Record)>; }
 
 ///////////////////////////////////////////////////////////////////////////////////////////////////////////////////////
@@ -54,8 +57,6 @@ public:
 	virtual void BuildAsync(IRequestOwner& Owner, FOnCacheRecordComplete&& OnComplete) = 0;
 };
 
-FCacheRecordBuilder CreateCacheRecordBuilder(ICacheRecordBuilderInternal* RecordBuilder);
-
 } // UE::DerivedData::Private
 
 ///////////////////////////////////////////////////////////////////////////////////////////////////////////////////////
@@ -105,11 +106,17 @@ public:
 	/** Returns the payload matching the ID, whether value or attachment. Null if no match. Data is null if skipped. */
 	inline const FPayload& GetPayload(const FPayloadId& Id) const { return Record->GetPayload(Id); }
 
+	/** Save the cache record to a compact binary package. */
+	UE_API FCbPackage Save() const;
+
+	/** Load a cache record from a compact binary package. Null on error. */
+	UE_API static FOptionalCacheRecord Load(const FCbPackage& Package);
+
 private:
 	friend class FOptionalCacheRecord;
 	friend FCacheRecord Private::CreateCacheRecord(Private::ICacheRecordInternal* Record);
 
-	/** Construct a cache record. Use Build() or BuildAsync() on a builder from ICache::CreateRecord(). */
+	/** Construct a cache record. Use Build() or BuildAsync() on FCacheRecordBuilder. */
 	inline explicit FCacheRecord(Private::ICacheRecordInternal* InRecord)
 		: Record(InRecord)
 	{
@@ -121,8 +128,8 @@ private:
 /**
  * A cache record builder is used to construct a cache record.
  *
- * Create using ICache::CreateRecord() which must be given a key that uniquely corresponds to the
- * value and attachments for the cache record. Metadata may vary between records of the same key.
+ * Create using a key that uniquely corresponds to the value and attachments for the cache record.
+ * Metadata may vary between records of the same key.
  *
  * The value and attachments can be provided as buffers, which will be compressed, or as payloads
  * which were previously compressed and have an identifier assigned.
@@ -132,6 +139,11 @@ private:
 class FCacheRecordBuilder
 {
 public:
+	/**
+	 * Create a cache record builder from a cache key.
+	 */
+	UE_API explicit FCacheRecordBuilder(const FCacheKey& Key);
+
 	/**
 	 * Set the metadata for the cache record.
 	 *
@@ -221,14 +233,6 @@ public:
 	}
 
 private:
-	friend FCacheRecordBuilder Private::CreateCacheRecordBuilder(Private::ICacheRecordBuilderInternal* RecordBuilder);
-
-	/** Construct a cache record builder. Use ICache::CreateRecord(). */
-	inline explicit FCacheRecordBuilder(Private::ICacheRecordBuilderInternal* InRecordBuilder)
-		: RecordBuilder(InRecordBuilder)
-	{
-	}
-
 	TUniquePtr<Private::ICacheRecordBuilderInternal> RecordBuilder;
 };
 
@@ -280,3 +284,5 @@ struct FCacheRecordKeyFuncs
 };
 
 } // UE::DerivedData
+
+#undef UE_API
diff --git a/Engine/Source/Developer/TextureBuild/Private/TextureBuildFunction.cpp b/Engine/Source/Developer/TextureBuild/Private/TextureBuildFunction.cpp
--- a/Engine/Source/Developer/TextureBuild/Private/TextureBuildFunction.cpp
+++ b/Engine/Source/Developer/TextureBuild/Private/TextureBuildFunction.cpp
@@ -270,7 +270,7 @@ FGuid FTextureBuildFunction::GetVersion() const
 
 void FTextureBuildFunction::Configure(UE::DerivedData::FBuildConfigContext& Context) const
 {
-	Context.SetCacheBucket(Context.CreateCacheBucket(TEXT("Texture"_SV)));
+	Context.SetCacheBucket(UE::DerivedData::FCacheBucket("Texture"_ASV));
 }
 
 void FTextureBuildFunction::Build(UE::DerivedData::FBuildContext& Context) const
diff --git a/Engine/Source/Editor/UnrealEd/Private/EditorDomain/EditorDomainUtils.cpp b/Engine/Source/Editor/UnrealEd/Private/EditorDomain/EditorDomainUtils.cpp
--- a/Engine/Source/Editor/UnrealEd/Private/EditorDomain/EditorDomainUtils.cpp
+++ b/Engine/Source/Editor/UnrealEd/Private/EditorDomain/EditorDomainUtils.cpp
@@ -202,23 +202,19 @@ EPackageDigestResult AppendPackageDigest(IAssetRegistry& AssetRegistry, FName Pa
 
 UE::DerivedData::FCacheKey GetEditorDomainPackageKey(const FPackageDigest& PackageDigest)
 {
-	static UE::DerivedData::FCacheBucket EditorDomainPackageCacheBucket =
-		UE::DerivedData::GetCache().CreateBucket(EditorDomainPackageBucketName);
+	static UE::DerivedData::FCacheBucket EditorDomainPackageCacheBucket(EditorDomainPackageBucketName);
 	return UE::DerivedData::FCacheKey{EditorDomainPackageCacheBucket, PackageDigest};
 }
 
 UE::DerivedData::FCacheKey GetBulkDataListKey(const FPackageDigest& PackageDigest)
 {
-	static UE::DerivedData::FCacheBucket EditorDomainBulkDataListBucket =
-		UE::DerivedData::GetCache().CreateBucket(EditorDomainBulkDataListBucketName);
+	static UE::DerivedData::FCacheBucket EditorDomainBulkDataListBucket(EditorDomainBulkDataListBucketName);
 	return UE::DerivedData::FCacheKey{ EditorDomainBulkDataListBucket, PackageDigest };
 }
 
 UE::DerivedData::FCacheKey GetBulkDataPayloadIdKey(const FIoHash& PackageAndGuidDigest)
 {
-	static UE::DerivedData::FCacheBucket EditorDomainBulkDataPayloadIdBucket =
-		UE::DerivedData::GetCache().CreateBucket(EditorDomainBulkDataPayloadIdBucketName);
-
+	static UE::DerivedData::FCacheBucket EditorDomainBulkDataPayloadIdBucket(EditorDomainBulkDataPayloadIdBucketName);
 	return UE::DerivedData::FCacheKey{ EditorDomainBulkDataPayloadIdBucket, PackageAndGuidDigest };
 }
 
@@ -437,7 +433,7 @@ bool TrySavePackage(UPackage* Package)
 	}
 
 	ICache& Cache = GetCache();
-	FCacheRecordBuilder RecordBuilder = Cache.CreateRecord(GetEditorDomainPackageKey(PackageDigest));
+	FCacheRecordBuilder RecordBuilder(GetEditorDomainPackageKey(PackageDigest));
 
 	// We use a counter for PayloadIds rather than hashes of the Attachments. We do this because
 	// some attachments may be identical, and Attachments are not allowed to have identical PayloadIds.
@@ -524,7 +520,7 @@ void PutBulkDataList(FName PackageName, FSharedBuffer Buffer)
 	using namespace UE::DerivedData;
 	ICache& Cache = GetCache();
 	FRequestOwner Owner(EPriority::Normal);
-	FCacheRecordBuilder RecordBuilder = Cache.CreateRecord(GetBulkDataListKey(PackageDigest));
+	FCacheRecordBuilder RecordBuilder(GetBulkDataListKey(PackageDigest));
 	RecordBuilder.SetValue(Buffer);
 	Cache.Put({RecordBuilder.Build()}, WriteToString<128>(PackageName), ECachePolicy::Default, Owner);
 	Owner.KeepAlive();
@@ -584,7 +580,7 @@ void PutBulkDataPayloadId(FName PackageName, const FGuid& BulkDataId, FSharedBuf
 	using namespace UE::DerivedData;
 	ICache& Cache = GetCache();
 	FRequestOwner Owner(EPriority::Normal);
-	FCacheRecordBuilder RecordBuilder = Cache.CreateRecord(GetBulkDataPayloadIdKey(PackageAndGuidDigest));
+	FCacheRecordBuilder RecordBuilder(GetBulkDataPayloadIdKey(PackageAndGuidDigest));
 	RecordBuilder.SetValue(Buffer);
 	Cache.Put({RecordBuilder.Build()}, WriteToString<128>(PackageName), ECachePolicy::Default, Owner);
 	Owner.KeepAlive();
