DDC: Split DerivedDataRequest.h and rename DerivedDataBuildPolicy.h

#rnx

diff --git a/Engine/Source/Developer/DerivedDataCache/Private/DerivedDataBackendAsyncPutWrapper.cpp b/Engine/Source/Developer/DerivedDataCache/Private/DerivedDataBackendAsyncPutWrapper.cpp
--- a/Engine/Source/Developer/DerivedDataCache/Private/DerivedDataBackendAsyncPutWrapper.cpp
+++ b/Engine/Source/Developer/DerivedDataCache/Private/DerivedDataBackendAsyncPutWrapper.cpp
@@ -1,6 +1,9 @@
 // Copyright Epic Games, Inc. All Rights Reserved.
 
 #include "DerivedDataBackendAsyncPutWrapper.h"
+
+#include "DerivedDataRequest.h"
+#include "DerivedDataRequestOwner.h"
 #include "MemoryDerivedDataBackend.h"
 #include "Tasks/Task.h"
 #include <atomic>
diff --git a/Engine/Source/Developer/DerivedDataCache/Private/DerivedDataBuildJob.cpp b/Engine/Source/Developer/DerivedDataCache/Private/DerivedDataBuildJob.cpp
--- a/Engine/Source/Developer/DerivedDataCache/Private/DerivedDataBuildJob.cpp
+++ b/Engine/Source/Developer/DerivedDataCache/Private/DerivedDataBuildJob.cpp
@@ -14,14 +14,15 @@
 #include "DerivedDataBuildJobContext.h"
 #include "DerivedDataBuildKey.h"
 #include "DerivedDataBuildOutput.h"
-#include "DerivedDataBuildPolicy.h"
 #include "DerivedDataBuildPrivate.h"
 #include "DerivedDataBuildScheduler.h"
+#include "DerivedDataBuildTypes.h"
 #include "DerivedDataBuildWorker.h"
 #include "DerivedDataBuildWorkerRegistry.h"
 #include "DerivedDataCache.h"
 #include "DerivedDataCacheRecord.h"
 #include "DerivedDataPayload.h"
+#include "DerivedDataRequestOwner.h"
 #include "HAL/CriticalSection.h"
 #include "HAL/Event.h"
 #include "HAL/FileManager.h"
diff --git a/Engine/Source/Developer/DerivedDataCache/Private/DerivedDataBuildJobContext.cpp b/Engine/Source/Developer/DerivedDataCache/Private/DerivedDataBuildJobContext.cpp
--- a/Engine/Source/Developer/DerivedDataCache/Private/DerivedDataBuildJobContext.cpp
+++ b/Engine/Source/Developer/DerivedDataCache/Private/DerivedDataBuildJobContext.cpp
@@ -4,10 +4,11 @@
 
 #include "DerivedDataBuildJob.h"
 #include "DerivedDataBuildOutput.h"
-#include "DerivedDataBuildPolicy.h"
 #include "DerivedDataBuildPrivate.h"
+#include "DerivedDataBuildTypes.h"
 #include "DerivedDataCache.h"
 #include "DerivedDataPayloadPrivate.h"
+#include "DerivedDataRequestOwner.h"
 #include "Hash/Blake3.h"
 #include "Misc/StringBuilder.h"
 #include "UObject/NameTypes.h"
diff --git a/Engine/Source/Developer/DerivedDataCache/Private/DerivedDataBuildJobContext.h b/Engine/Source/Developer/DerivedDataCache/Private/DerivedDataBuildJobContext.h
--- a/Engine/Source/Developer/DerivedDataCache/Private/DerivedDataBuildJobContext.h
+++ b/Engine/Source/Developer/DerivedDataCache/Private/DerivedDataBuildJobContext.h
@@ -16,8 +16,10 @@
 
 namespace UE::DerivedData { class FBuildOutputBuilder; }
 namespace UE::DerivedData { class IBuildJob; }
+namespace UE::DerivedData { class IRequestOwner; }
 namespace UE::DerivedData { enum class EBuildPolicy : uint8; }
 namespace UE::DerivedData { enum class ECachePolicy : uint8; }
+namespace UE::DerivedData { enum class EPriority : uint8; }
 
 namespace UE::DerivedData::Private
 {
diff --git a/Engine/Source/Developer/DerivedDataCache/Private/DerivedDataBuildScheduler.cpp b/Engine/Source/Developer/DerivedDataCache/Private/DerivedDataBuildScheduler.cpp
--- a/Engine/Source/Developer/DerivedDataCache/Private/DerivedDataBuildScheduler.cpp
+++ b/Engine/Source/Developer/DerivedDataCache/Private/DerivedDataBuildScheduler.cpp
@@ -5,6 +5,8 @@
 #include "DerivedDataBuildFunction.h"
 #include "DerivedDataBuildFunctionFactory.h"
 #include "DerivedDataBuildFunctionRegistry.h"
+#include "DerivedDataRequest.h"
+#include "DerivedDataRequestOwner.h"
 #include "Misc/Guid.h"
 #include "Tasks/Task.h"
 
diff --git a/Engine/Source/Developer/DerivedDataCache/Private/DerivedDataBuildWorkerRegistry.cpp b/Engine/Source/Developer/DerivedDataCache/Private/DerivedDataBuildWorkerRegistry.cpp
--- a/Engine/Source/Developer/DerivedDataCache/Private/DerivedDataBuildWorkerRegistry.cpp
+++ b/Engine/Source/Developer/DerivedDataCache/Private/DerivedDataBuildWorkerRegistry.cpp
@@ -8,7 +8,6 @@
 #include "DerivedDataBuildKey.h"
 #include "DerivedDataBuildPrivate.h"
 #include "DerivedDataBuildWorker.h"
-#include "DerivedDataRequest.h"
 #include "Features/IModularFeatures.h"
 #include "HAL/CriticalSection.h"
 #include "IO/IoHash.h"
diff --git a/Engine/Source/Developer/DerivedDataCache/Private/DerivedDataRequest.cpp b/Engine/Source/Developer/DerivedDataCache/Private/DerivedDataRequest.cpp
deleted file mode 100644
diff --git a/Engine/Source/Developer/DerivedDataCache/Private/DerivedDataRequestOwner.cpp b/Engine/Source/Developer/DerivedDataCache/Private/DerivedDataRequestOwner.cpp
new file mode 100644
--- /dev/null
+++ b/Engine/Source/Developer/DerivedDataCache/Private/DerivedDataRequestOwner.cpp
@@ -0,0 +1,269 @@
+// Copyright Epic Games, Inc. All Rights Reserved.
+
+#include "DerivedDataRequestOwner.h"
+
+#include "Containers/Array.h"
+#include "DerivedDataRequest.h"
+#include "DerivedDataRequestTypes.h"
+#include "HAL/CriticalSection.h"
+#include "HAL/Event.h"
+#include "Misc/ScopeExit.h"
+#include "Misc/ScopeRWLock.h"
+#include "Templates/RefCounting.h"
+#include <atomic>
+
+namespace UE::DerivedData::Private
+{
+
+/** A one-time-use event to work around the lack of condition variables. */
+struct FRequestBarrierEvent : public FRefCountBase
+{
+	FEventRef Event{EEventMode::ManualReset};
+};
+
+class FRequestOwnerShared final : public IRequestOwner, public FRequestBase
+{
+public:
+	explicit FRequestOwnerShared(EPriority Priority);
+	~FRequestOwnerShared() final = default;
+
+	void Begin(IRequest* Request) final;
+	TRefCountPtr<IRequest> End(IRequest* Request) final;
+
+	void BeginBarrier() final;
+	void EndBarrier() final;
+
+	inline EPriority GetPriority() const final { return Priority; }
+	inline bool IsCanceled() const final { return bIsCanceled; }
+
+	inline void SetPriority(EPriority Priority) final;
+	inline void Cancel() final;
+	inline void Wait() final;
+	inline bool Poll() const;
+
+	inline void KeepAlive();
+	inline void Destroy();
+
+private:
+	mutable FRWLock Lock;
+	TArray<TRefCountPtr<IRequest>, TInlineAllocator<1>> Requests;
+	TRefCountPtr<FRequestBarrierEvent> BarrierEvent;
+	uint32 BarrierCount{0};
+	EPriority Priority{EPriority::Normal};
+	uint8 bPriorityChangedInBarrier : 1;
+	uint8 bBeginExecuted : 1;
+	bool bIsCanceled{false};
+	bool bKeepAlive{false};
+};
+
+FRequestOwnerShared::FRequestOwnerShared(EPriority NewPriority)
+	: Priority(NewPriority)
+	, bPriorityChangedInBarrier(false)
+	, bBeginExecuted(false)
+{
+	AddRef(); // Release is called by Destroy.
+}
+
+void FRequestOwnerShared::Begin(IRequest* Request)
+{
+	AddRef();
+	EPriority NewPriority;
+	{
+		FWriteScopeLock WriteLock(Lock);
+		checkf(BarrierCount > 0 || !bBeginExecuted,
+			TEXT("At least one FRequestBarrier must be in scope when beginning a request after the first request. ")
+			TEXT("The overload of End that invokes a callback handles this automatically for most use cases."));
+		check(Request);
+		Requests.Add(Request);
+		bBeginExecuted = true;
+		if (!bPriorityChangedInBarrier)
+		{
+			return;
+		}
+		NewPriority = Priority;
+	}
+	// Loop until priority is stable. Another thread may be changing the priority concurrently.
+	for (EPriority CheckPriority; ; NewPriority = CheckPriority)
+	{
+		Request->SetPriority(NewPriority);
+		CheckPriority = (FReadScopeLock(Lock), Priority);
+		if (CheckPriority == NewPriority)
+		{
+			break;
+		}
+	}
+}
+
+TRefCountPtr<IRequest> FRequestOwnerShared::End(IRequest* Request)
+{
+	ON_SCOPE_EXIT { Release(); };
+	FWriteScopeLock WriteLock(Lock);
+	check(Request);
+	TRefCountPtr<IRequest>* RequestPtr = Requests.FindByKey(Request);
+	check(RequestPtr);
+	TRefCountPtr<IRequest> RequestRef = MoveTemp(*RequestPtr);
+	Requests.RemoveAtSwap(UE_PTRDIFF_TO_INT32(RequestPtr - Requests.GetData()), 1, /*bAllowShrinking*/ false);
+	return RequestRef;
+}
+
+void FRequestOwnerShared::BeginBarrier()
+{
+	AddRef();
+	FWriteScopeLock WriteLock(Lock);
+	++BarrierCount;
+}
+
+void FRequestOwnerShared::EndBarrier()
+{
+	ON_SCOPE_EXIT { Release(); };
+	TRefCountPtr<FRequestBarrierEvent> LocalBarrierEvent;
+	if (FWriteScopeLock WriteLock(Lock); --BarrierCount == 0)
+	{
+		bPriorityChangedInBarrier = false;
+		LocalBarrierEvent = MoveTemp(BarrierEvent);
+		BarrierEvent = nullptr;
+	}
+	if (LocalBarrierEvent)
+	{
+		LocalBarrierEvent->Event->Trigger();
+	}
+}
+
+void FRequestOwnerShared::SetPriority(EPriority NewPriority)
+{
+	TArray<TRefCountPtr<IRequest>, TInlineAllocator<16>> LocalRequests;
+	if (FWriteScopeLock WriteLock(Lock); Priority == NewPriority)
+	{
+		return;
+	}
+	else
+	{
+		Priority = NewPriority;
+		LocalRequests = Requests;
+		bPriorityChangedInBarrier = (BarrierCount > 0);
+	}
+
+	for (IRequest* Request : LocalRequests)
+	{
+		Request->SetPriority(NewPriority);
+	}
+}
+
+void FRequestOwnerShared::Cancel()
+{
+	for (;;)
+	{
+		TRefCountPtr<IRequest> LocalRequest;
+		TRefCountPtr<FRequestBarrierEvent> LocalBarrierEvent;
+
+		{
+			FWriteScopeLock WriteLock(Lock);
+			bIsCanceled = true;
+
+			if (!Requests.IsEmpty())
+			{
+				LocalRequest = Requests.Last();
+			}
+			else if (BarrierCount > 0)
+			{
+				LocalBarrierEvent = BarrierEvent;
+				if (!LocalBarrierEvent)
+				{
+					LocalBarrierEvent = BarrierEvent = new FRequestBarrierEvent;
+				}
+			}
+			else
+			{
+				return;
+			}
+		}
+
+		if (LocalRequest)
+		{
+			LocalRequest->Cancel();
+		}
+		else
+		{
+			LocalBarrierEvent->Event->Wait();
+		}
+	}
+}
+
+void FRequestOwnerShared::Wait()
+{
+	for (;;)
+	{
+		TRefCountPtr<IRequest> LocalRequest;
+		TRefCountPtr<FRequestBarrierEvent> LocalBarrierEvent;
+
+		{
+			FWriteScopeLock WriteLock(Lock);
+
+			if (!Requests.IsEmpty())
+			{
+				LocalRequest = Requests.Last();
+			}
+			else if (BarrierCount > 0)
+			{
+				LocalBarrierEvent = BarrierEvent;
+				if (!LocalBarrierEvent)
+				{
+					LocalBarrierEvent = BarrierEvent = new FRequestBarrierEvent;
+				}
+			}
+			else
+			{
+				return;
+			}
+		}
+
+		if (LocalRequest)
+		{
+			LocalRequest->Wait();
+		}
+		else
+		{
+			LocalBarrierEvent->Event->Wait();
+		}
+	}
+}
+
+bool FRequestOwnerShared::Poll() const
+{
+	FReadScopeLock ReadLock(Lock);
+	return Requests.IsEmpty() && BarrierCount == 0;
+}
+
+void FRequestOwnerShared::KeepAlive()
+{
+	FWriteScopeLock WriteLock(Lock);
+	bKeepAlive = true;
+}
+
+void FRequestOwnerShared::Destroy()
+{
+	const bool bLocalKeepAlive = (FWriteScopeLock(Lock), bKeepAlive);
+	if (!bLocalKeepAlive)
+	{
+		Cancel();
+	}
+	Release();
+}
+
+} // UE::DerivedData::Private
+
+namespace UE::DerivedData
+{
+
+FRequestOwner::FRequestOwner(EPriority Priority) : Owner(new Private::FRequestOwnerShared(Priority)) {}
+static Private::FRequestOwnerShared* ToSharedOwner(IRequestOwner* Owner) { return static_cast<Private::FRequestOwnerShared*>(Owner); }
+void FRequestOwner::KeepAlive()                         { return ToSharedOwner(Owner.Get())->KeepAlive(); }
+EPriority FRequestOwner::GetPriority() const            { return ToSharedOwner(Owner.Get())->GetPriority(); }
+void FRequestOwner::SetPriority(EPriority Priority)     { return ToSharedOwner(Owner.Get())->SetPriority(Priority); }
+void FRequestOwner::Cancel()                            { return ToSharedOwner(Owner.Get())->Cancel(); }
+void FRequestOwner::Wait()                              { return ToSharedOwner(Owner.Get())->Wait(); }
+bool FRequestOwner::Poll() const                        { return ToSharedOwner(Owner.Get())->Poll(); }
+FRequestOwner::operator IRequest*()                     { return ToSharedOwner(Owner.Get()); }
+void FRequestOwner::Destroy(IRequestOwner& SharedOwner) { return ToSharedOwner(&SharedOwner)->Destroy(); }
+
+} // UE::DerivedData
diff --git a/Engine/Source/Developer/DerivedDataCache/Private/HierarchicalDerivedDataBackend.h b/Engine/Source/Developer/DerivedDataCache/Private/HierarchicalDerivedDataBackend.h
--- a/Engine/Source/Developer/DerivedDataCache/Private/HierarchicalDerivedDataBackend.h
+++ b/Engine/Source/Developer/DerivedDataCache/Private/HierarchicalDerivedDataBackend.h
@@ -8,6 +8,7 @@
 #include "DerivedDataCacheRecord.h"
 #include "DerivedDataCacheUsageStats.h"
 #include "DerivedDataBackendAsyncPutWrapper.h"
+#include "DerivedDataRequestOwner.h"
 #include "Misc/ScopeRWLock.h"
 #include "Templates/UniquePtr.h"
 #include "Containers/ArrayView.h"
diff --git a/Engine/Source/Developer/DerivedDataCache/Private/ZenDerivedDataBackend.h b/Engine/Source/Developer/DerivedDataCache/Private/ZenDerivedDataBackend.h
--- a/Engine/Source/Developer/DerivedDataCache/Private/ZenDerivedDataBackend.h
+++ b/Engine/Source/Developer/DerivedDataCache/Private/ZenDerivedDataBackend.h
@@ -19,7 +19,6 @@
 #include "Containers/StringFwd.h"
 #include "DerivedDataBackendInterface.h"
 #include "DerivedDataCacheUsageStats.h"
-#include "DerivedDataRequest.h"
 #include "HAL/CriticalSection.h"
 
 class FCbObject;
diff --git a/Engine/Source/Developer/DerivedDataCache/Public/DerivedDataBuildInputResolver.h b/Engine/Source/Developer/DerivedDataCache/Public/DerivedDataBuildInputResolver.h
--- a/Engine/Source/Developer/DerivedDataCache/Public/DerivedDataBuildInputResolver.h
+++ b/Engine/Source/Developer/DerivedDataCache/Public/DerivedDataBuildInputResolver.h
@@ -7,13 +7,14 @@
 #include "Containers/StringView.h"
 #include "DerivedDataBuildDefinition.h"
 #include "DerivedDataBuildKey.h"
-#include "DerivedDataRequest.h"
+#include "DerivedDataRequestTypes.h"
 #include "IO/IoHash.h"
 #include "Templates/Function.h"
 
 struct FGuid;
 
 namespace UE::DerivedData { class FBuildAction; }
+namespace UE::DerivedData { class IRequestOwner; }
 namespace UE::DerivedData { struct FBuildInputDataResolvedParams; }
 namespace UE::DerivedData { struct FBuildInputMetaResolvedParams; }
 namespace UE::DerivedData { struct FBuildKeyResolvedParams; }
diff --git a/Engine/Source/Developer/DerivedDataCache/Public/DerivedDataBuildJob.h b/Engine/Source/Developer/DerivedDataCache/Public/DerivedDataBuildJob.h
--- a/Engine/Source/Developer/DerivedDataCache/Public/DerivedDataBuildJob.h
+++ b/Engine/Source/Developer/DerivedDataCache/Public/DerivedDataBuildJob.h
@@ -4,7 +4,7 @@
 
 #include "CoreTypes.h"
 #include "Containers/StringFwd.h"
-#include "DerivedDataRequest.h"
+#include "DerivedDataRequestTypes.h"
 
 template <typename FuncType> class TUniqueFunction;
 
diff --git a/Engine/Source/Developer/DerivedDataCache/Public/DerivedDataBuildPolicy.h b/Engine/Source/Developer/DerivedDataCache/Public/DerivedDataBuildPolicy.h
deleted file mode 100644
diff --git a/Engine/Source/Developer/DerivedDataCache/Public/DerivedDataBuildSession.h b/Engine/Source/Developer/DerivedDataCache/Public/DerivedDataBuildSession.h
--- a/Engine/Source/Developer/DerivedDataCache/Public/DerivedDataBuildSession.h
+++ b/Engine/Source/Developer/DerivedDataCache/Public/DerivedDataBuildSession.h
@@ -5,8 +5,8 @@
 #include "CoreTypes.h"
 #include "Containers/StringView.h"
 #include "DerivedDataBuildKey.h"
-#include "DerivedDataBuildPolicy.h"
-#include "DerivedDataRequest.h"
+#include "DerivedDataBuildTypes.h"
+#include "DerivedDataRequestTypes.h"
 #include "Templates/Function.h"
 #include "Templates/UniquePtr.h"
 
@@ -16,6 +16,7 @@ namespace UE::DerivedData { class FBuildOutput; }
 namespace UE::DerivedData { class FBuildSession; }
 namespace UE::DerivedData { class FOptionalBuildInputs; }
 namespace UE::DerivedData { class FPayload; }
+namespace UE::DerivedData { class IRequestOwner; }
 namespace UE::DerivedData { struct FBuildActionCompleteParams; }
 namespace UE::DerivedData { struct FBuildCompleteParams; }
 namespace UE::DerivedData { struct FBuildPayloadCompleteParams; }
diff --git a/Engine/Source/Developer/DerivedDataCache/Public/DerivedDataBuildTypes.h b/Engine/Source/Developer/DerivedDataCache/Public/DerivedDataBuildTypes.h
new file mode 100644
--- /dev/null
+++ b/Engine/Source/Developer/DerivedDataCache/Public/DerivedDataBuildTypes.h
@@ -0,0 +1,83 @@
+// Copyright Epic Games, Inc. All Rights Reserved.
+
+#pragma once
+
+#include "CoreTypes.h"
+#include "Misc/EnumClassFlags.h"
+
+namespace UE::DerivedData
+{
+
+/**
+ * Flags to control the behavior of build requests.
+ *
+ * The build policy flags can be combined to support a variety of usage patterns. Examples:
+ *
+ * Build(Default): allow cache get; build if missing; return every payload.
+ * Build(Build | CacheStore): never query the cache; always build; return every payload.
+ * Build(Cache | CacheSkipData): allow cache get; never build; skip every payload.
+ *
+ * BuildPayload(Default): allow cache get; build if missing; return one payload.
+ * BuildPayload(Build | CacheStore): never get from the cache; always build; return one payload.
+ * BuildPayload(Cache | CacheSkipData): allow cache get; never build; skip every payload.
+ */
+enum class EBuildPolicy : uint8
+{
+	/** A value without any flags set. */
+	None            = 0,
+
+	/** Allow local execution of the build function. */
+	BuildLocal      = 1 << 0,
+	/** Allow remote execution of the build function if it has a registered build worker. */
+	BuildRemote     = 1 << 1,
+	/** Allow local and remote execution of the build function. */
+	Build           = BuildLocal | BuildRemote,
+
+	/** Allow a cache query to avoid having to build. */
+	CacheQuery      = 1 << 2,
+	/** Allow a cache store to persist the build output. */
+	CacheStore      = 1 << 3,
+	/** Allow a cache query and a cache store for the build. */
+	Cache           = CacheQuery | CacheStore,
+
+	/** Skip fetching the payload data from the cache. */
+	SkipData        = 1 << 4,
+
+	/** Allow cache query+store, allow local+remote build when missed or skipped, and fetch the payload(s). */
+	Default         = Build | Cache,
+};
+
+ENUM_CLASS_FLAGS(EBuildPolicy);
+
+/** Flags for build request completion callbacks. */
+enum class EBuildStatus : uint32
+{
+	/** A value without any flags set. */
+	None            = 0,
+
+	/** The build function was executed locally. */
+	BuildLocal      = 1 << 0,
+	/** The build function was executed remotely. */
+	BuildRemote     = 1 << 1,
+	/** The build action and inputs were exported. */
+	BuildExport     = 1 << 2,
+
+	/** An attempt was made to execute the build function remotely. */
+	BuildTryRemote  = 1 << 3,
+	/** An attempt was made to export the build action and inputs. */
+	BuildTryExport  = 1 << 4,
+
+	/** The build made a cache query request. */
+	CacheQuery      = 1 << 5,
+	/** Valid build output was found in the cache. */
+	CacheQueryHit   = 1 << 6,
+
+	/** The build made a cache store request. */
+	CacheStore      = 1 << 7,
+	/** Valid build output was stored in the cache. */
+	CacheStoreHit   = 1 << 8,
+};
+
+ENUM_CLASS_FLAGS(EBuildStatus);
+
+} // UE::DerivedData
diff --git a/Engine/Source/Developer/DerivedDataCache/Public/DerivedDataBuildWorker.h b/Engine/Source/Developer/DerivedDataCache/Public/DerivedDataBuildWorker.h
--- a/Engine/Source/Developer/DerivedDataCache/Public/DerivedDataBuildWorker.h
+++ b/Engine/Source/Developer/DerivedDataCache/Public/DerivedDataBuildWorker.h
@@ -6,7 +6,7 @@
 #include "Containers/ArrayView.h"
 #include "Containers/StringFwd.h"
 #include "DerivedDataBuildKey.h"
-#include "DerivedDataRequest.h"
+#include "DerivedDataRequestTypes.h"
 #include "Features/IModularFeature.h"
 #include "UObject/NameTypes.h"
 
@@ -21,6 +21,7 @@ namespace UE::DerivedData { class FBuildAction; }
 namespace UE::DerivedData { class FOptionalBuildInputs; }
 namespace UE::DerivedData { class FOptionalBuildOutput; }
 namespace UE::DerivedData { class IBuild; }
+namespace UE::DerivedData { class IRequestOwner; }
 namespace UE::DerivedData { struct FBuildWorkerActionCompleteParams; }
 namespace UE::DerivedData { struct FBuildWorkerFileDataCompleteParams; }
 namespace UE::DerivedData { enum class EBuildPolicy : uint8; }
diff --git a/Engine/Source/Developer/DerivedDataCache/Public/DerivedDataCache.h b/Engine/Source/Developer/DerivedDataCache/Public/DerivedDataCache.h
--- a/Engine/Source/Developer/DerivedDataCache/Public/DerivedDataCache.h
+++ b/Engine/Source/Developer/DerivedDataCache/Public/DerivedDataCache.h
@@ -4,21 +4,22 @@
 
 #include "CoreTypes.h"
 #include "DerivedDataCacheKey.h"
-#include "DerivedDataRequest.h"
+#include "DerivedDataRequestTypes.h"
 #include "Misc/EnumClassFlags.h"
 #include "Templates/Function.h"
 
 #define UE_API DERIVEDDATACACHE_API
 
+namespace UE::DerivedData { class FCacheRecord; }
+namespace UE::DerivedData { class FPayload; }
+namespace UE::DerivedData { class IRequestOwner; }
+namespace UE::DerivedData { struct FCacheGetCompleteParams; }
+namespace UE::DerivedData { struct FCacheGetPayloadCompleteParams; }
+namespace UE::DerivedData { struct FCachePutCompleteParams; }
+
 namespace UE::DerivedData
 {
 
-class FCacheRecord;
-class FPayload;
-struct FCacheGetCompleteParams;
-struct FCacheGetPayloadCompleteParams;
-struct FCachePutCompleteParams;
-
 using FOnCachePutComplete = TUniqueFunction<void (FCachePutCompleteParams&& Params)>;
 using FOnCacheGetComplete = TUniqueFunction<void (FCacheGetCompleteParams&& Params)>;
 using FOnCacheGetPayloadComplete = TUniqueFunction<void (FCacheGetPayloadCompleteParams&& Params)>;
diff --git a/Engine/Source/Developer/DerivedDataCache/Public/DerivedDataCacheRecord.h b/Engine/Source/Developer/DerivedDataCache/Public/DerivedDataCacheRecord.h
--- a/Engine/Source/Developer/DerivedDataCache/Public/DerivedDataCacheRecord.h
+++ b/Engine/Source/Developer/DerivedDataCache/Public/DerivedDataCacheRecord.h
@@ -4,7 +4,6 @@
 
 #include "DerivedDataCacheKey.h"
 #include "DerivedDataPayload.h"
-#include "DerivedDataRequest.h"
 #include "Memory/SharedBuffer.h"
 #include "Templates/Function.h"
 #include "Templates/RefCounting.h"
@@ -17,6 +16,7 @@ class FCbPackage;
 
 namespace UE::DerivedData { class FCacheRecord; }
 namespace UE::DerivedData { class FOptionalCacheRecord; }
+namespace UE::DerivedData { class IRequestOwner; }
 namespace UE::DerivedData { using FOnCacheRecordComplete = TUniqueFunction<void (FCacheRecord&& Record)>; }
 
 ///////////////////////////////////////////////////////////////////////////////////////////////////////////////////////
diff --git a/Engine/Source/Developer/DerivedDataCache/Public/DerivedDataRequest.h b/Engine/Source/Developer/DerivedDataCache/Public/DerivedDataRequest.h
--- a/Engine/Source/Developer/DerivedDataCache/Public/DerivedDataRequest.h
+++ b/Engine/Source/Developer/DerivedDataCache/Public/DerivedDataRequest.h
@@ -3,66 +3,13 @@
 #pragma once
 
 #include "CoreTypes.h"
-#include "Templates/Invoke.h"
-#include "Templates/RefCounting.h"
-#include "Templates/UniquePtr.h"
-#include "Templates/UnrealTemplate.h"
 #include <atomic>
 
-#define UE_API DERIVEDDATACACHE_API
-
-namespace UE::DerivedData::Private { class FRequestOwnerSharedBase; }
+namespace UE::DerivedData { enum class EPriority : uint8; }
 
 namespace UE::DerivedData
 {
 
-/** Priority for scheduling a request. */
-enum class EPriority : uint8
-{
-	/**
-	 * Lowest is the minimum priority for asynchronous requests, and primarily for requests that are
-	 * speculative in nature, while minimizing impact on other requests.
-	 */
-	Lowest,
-	/**
-	 * Low is intended for requests that are below the default priority, but are used for operations
-	 * that the program will execute now rather than at an unknown time in the future.
-	 */
-	Low,
-	/**
-	 * Normal is intended as the default request priority.
-	 */
-	Normal,
-	/**
-	 * High is intended for requests that are on the critical path, but are not required to maintain
-	 * interactivity of the program.
-	 */
-	High,
-	/**
-	 * Highest is the maximum priority for asynchronous requests, and intended for requests that are
-	 * required to maintain interactivity of the program.
-	 */
-	Highest,
-	/**
-	 * Blocking is to be used only when the thread making the request will wait on completion of the
-	 * request before doing any other work. Requests at this priority level will be processed before
-	 * any request at a lower priority level. This priority permits a request to be processed on the
-	 * thread making the request. Waiting on a request may increase its priority to this level.
-	 */
-	Blocking,
-};
-
-/** Status of a request that has completed. */
-enum class EStatus : uint8
-{
-	/** The request completed successfully. Any requested data is available. */
-	Ok,
-	/** The request completed unsuccessfully. Any requested data is not available. */
-	Error,
-	/** The request was canceled before it completed. Any requested data is not available. */
-	Canceled,
-};
-
 /**
  * Interface to an asynchronous request that can be prioritized or canceled.
  *
@@ -137,147 +84,4 @@ private:
 	mutable std::atomic<uint32> ReferenceCount{0};
 };
 
-/**
- * A request owner manages requests throughout their execution.
- *
- * Requests are expected to call Begin when created, and End when finished, with their completion
- * callback and arguments passed to End when possible. If a request needs to invoke more than one
- * callback, have a request barrier in scope when executing callbacks to ensure that new requests
- * are prioritized properly when the priority is being changed from another thread.
- *
- * The owner is responsible for keeping itself alive while it has active requests or barriers.
- *
- * @see FRequestBarrier
- * @see FRequestOwner
- */
-class IRequestOwner
-{
-public:
-	/**
-	 * Begin tracking for the request.
-	 *
-	 * The owner will hold a reference to the request until End is called, and forward any cancel
-	 * operation or priority change to the request. Begin must return before End is called.
-	 */
-	virtual void Begin(IRequest* Request) = 0;
-
-	/**
-	 * End tracking for the request.
-	 *
-	 * Requests with a single completion callback should use the callback overload of End.
-	 *
-	 * @return The reference to the request that was held by the owner.
-	 */
-	virtual TRefCountPtr<IRequest> End(IRequest* Request) = 0;
-
-	/**
-	 * End tracking of the request.
-	 *
-	 * Begins a barrier, ends the request, invokes the callback, then ends the barrier. Keeps its
-	 * reference to the request until the callback has returned.
-	 *
-	 * @return The reference to the request that was held by the owner.
-	 */
-	template <typename CallbackType, typename... CallbackArgTypes>
-	inline TRefCountPtr<IRequest> End(IRequest* Request, CallbackType&& Callback, CallbackArgTypes&&... CallbackArgs);
-
-	/** See FRequestBarrier. */
-	virtual void BeginBarrier() = 0;
-	virtual void EndBarrier() = 0;
-
-	/** Returns the priority that new requests are expected to inherit. */
-	virtual EPriority GetPriority() const = 0;
-	/** Returns whether the owner has been canceled, which new requests are expected to check. */
-	virtual bool IsCanceled() const = 0;
-};
-
-/**
- * A request barrier is expected to be used when an owner may have new requests added to it.
- *
- * An owner may not consider its execution to be complete in the presence of a barrier, and needs
- * to take note of priority changes that occur while within a barrier, since newly-added requests
- * may have been created with the previous priority value.
- */
-class FRequestBarrier
-{
-public:
-	inline explicit FRequestBarrier(IRequestOwner& InOwner)
-		: Owner(InOwner)
-	{
-		Owner.BeginBarrier();
-	}
-
-	inline ~FRequestBarrier()
-	{
-		Owner.EndBarrier();
-	}
-
-private:
-	IRequestOwner& Owner;
-};
-
-/** Private base for a shared request owner, which is both an owner and a request. */
-class Private::FRequestOwnerSharedBase : public IRequestOwner, public FRequestBase
-{
-public:
-	virtual bool Poll() const = 0;
-	virtual void KeepAlive() = 0;
-	virtual void Destroy() = 0;
-};
-
-/**
- * A concrete request owner that also presents as a request.
- *
- * Request owners may be moved but not copied, and cancel any outstanding requests on destruction
- * unless KeepAlive has been called.
- *
- * @see IRequestOwner
- */
-class FRequestOwner
-{
-public:
-	/** Construct a request owner with the given priority. */
-	UE_API explicit FRequestOwner(EPriority Priority);
-
-	FRequestOwner(FRequestOwner&&) = default;
-	FRequestOwner& operator=(FRequestOwner&&) = default;
-
-	FRequestOwner(const FRequestOwner&) = delete;
-	FRequestOwner& operator=(const FRequestOwner&) = delete;
-
-	/** Keep requests in the owner alive until complete, even after destruction of the owner. */
-	inline void KeepAlive() { Owner->KeepAlive(); }
-
-	/** Returns the priority that new requests are expected to inherit. */
-	inline EPriority GetPriority() const { return Owner->GetPriority(); }
-	/** Set the priority of active and future requests in the owner. */
-	inline void SetPriority(EPriority Priority) { Owner->SetPriority(Priority); }
-	/** Cancel any active and future requests in the owner. */
-	inline void Cancel() { Owner->Cancel(); }
-	/** Wait for any active and future requests or barriers in the owner. */
-	inline void Wait() { Owner->Wait(); }
-	/** Poll whether the owner has any active requests or barriers. */
-	inline bool Poll() const { return Owner->Poll(); }
-
-	/** Access as a request owner. */
-	inline operator IRequestOwner&() { return *Owner; }
-	/** Access as a request. */
-	inline operator IRequest*() { return Owner.Get(); }
-
-private:
-	struct FDeleteOwner { void operator()(Private::FRequestOwnerSharedBase* O) const { if (O) { O->Destroy(); } } };
-	TUniquePtr<Private::FRequestOwnerSharedBase, FDeleteOwner> Owner;
-};
-
-template <typename CallbackType, typename... CallbackArgTypes>
-TRefCountPtr<IRequest> IRequestOwner::End(IRequest* Request, CallbackType&& Callback, CallbackArgTypes&&... CallbackArgs)
-{
-	FRequestBarrier Barrier(*this);
-	TRefCountPtr<IRequest> RequestRef = End(Request);
-	Invoke(Forward<CallbackType>(Callback), Forward<CallbackArgTypes>(CallbackArgs)...);
-	return RequestRef;
-}
-
 } // UE::DerivedData
-
-#undef UE_API
diff --git a/Engine/Source/Developer/DerivedDataCache/Public/DerivedDataRequestOwner.h b/Engine/Source/Developer/DerivedDataCache/Public/DerivedDataRequestOwner.h
new file mode 100644
--- /dev/null
+++ b/Engine/Source/Developer/DerivedDataCache/Public/DerivedDataRequestOwner.h
@@ -0,0 +1,158 @@
+// Copyright Epic Games, Inc. All Rights Reserved.
+
+#pragma once
+
+#include "CoreTypes.h"
+#include "Templates/Invoke.h"
+#include "Templates/RefCounting.h"
+#include "Templates/UniquePtr.h"
+#include "Templates/UnrealTemplate.h"
+
+#define UE_API DERIVEDDATACACHE_API
+
+namespace UE::DerivedData { class IRequest; }
+namespace UE::DerivedData { enum class EPriority : uint8; }
+
+namespace UE::DerivedData
+{
+
+/**
+ * A request owner manages requests throughout their execution.
+ *
+ * Requests are expected to call Begin when created, and End when finished, with their completion
+ * callback and arguments passed to End when possible. If a request needs to invoke more than one
+ * callback, have a request barrier in scope when executing callbacks to ensure that new requests
+ * are prioritized properly when the priority is being changed from another thread.
+ *
+ * The owner is responsible for keeping itself alive while it has active requests or barriers.
+ *
+ * @see FRequestBarrier
+ * @see FRequestOwner
+ */
+class IRequestOwner
+{
+public:
+	/**
+	 * Begin tracking for the request.
+	 *
+	 * The owner will hold a reference to the request until End is called, and forward any cancel
+	 * operation or priority change to the request. Begin must return before End is called.
+	 */
+	virtual void Begin(IRequest* Request) = 0;
+
+	/**
+	 * End tracking for the request.
+	 *
+	 * Requests with a single completion callback should use the callback overload of End.
+	 *
+	 * @return The reference to the request that was held by the owner.
+	 */
+	virtual TRefCountPtr<IRequest> End(IRequest* Request) = 0;
+
+	/**
+	 * End tracking of the request.
+	 *
+	 * Begins a barrier, ends the request, invokes the callback, then ends the barrier. Keeps its
+	 * reference to the request until the callback has returned.
+	 *
+	 * @return The reference to the request that was held by the owner.
+	 */
+	template <typename CallbackType, typename... CallbackArgTypes>
+	inline TRefCountPtr<IRequest> End(IRequest* Request, CallbackType&& Callback, CallbackArgTypes&&... CallbackArgs);
+
+	/** See FRequestBarrier. */
+	virtual void BeginBarrier() = 0;
+	virtual void EndBarrier() = 0;
+
+	/** Returns the priority that new requests are expected to inherit. */
+	virtual EPriority GetPriority() const = 0;
+	/** Returns whether the owner has been canceled, which new requests are expected to check. */
+	virtual bool IsCanceled() const = 0;
+};
+
+/**
+ * A concrete request owner that also presents as a request.
+ *
+ * Request owners may be moved but not copied, and cancel any outstanding requests on destruction
+ * unless KeepAlive has been called.
+ *
+ * @see IRequestOwner
+ */
+class FRequestOwner
+{
+public:
+	/** Construct a request owner with the given priority. */
+	UE_API explicit FRequestOwner(EPriority Priority);
+
+	FRequestOwner(FRequestOwner&&) = default;
+	FRequestOwner& operator=(FRequestOwner&&) = default;
+
+	FRequestOwner(const FRequestOwner&) = delete;
+	FRequestOwner& operator=(const FRequestOwner&) = delete;
+
+	/** Keep requests in the owner alive until complete, even after destruction of the owner. */
+	UE_API void KeepAlive();
+
+	/** Returns the priority that new requests are expected to inherit. */
+	UE_API EPriority GetPriority() const;
+	/** Set the priority of active and future requests in the owner. */
+	UE_API void SetPriority(EPriority Priority);
+	/** Cancel any active and future requests in the owner. */
+	UE_API void Cancel();
+	/** Wait for any active and future requests or barriers in the owner. */
+	UE_API void Wait();
+	/** Poll whether the owner has any active requests or barriers. */
+	UE_API bool Poll() const;
+
+	/** Access as a request owner. */
+	inline operator IRequestOwner&() { return *Owner; }
+	/** Access as a request. */
+	UE_API operator IRequest*();
+
+private:
+	UE_API static void Destroy(IRequestOwner& SharedOwner);
+
+	struct FDeleteOwner { inline void operator()(IRequestOwner* O) const { if (O) { Destroy(*O); } } };
+	TUniquePtr<IRequestOwner, FDeleteOwner> Owner;
+};
+
+/**
+ * A request barrier is expected to be used when an owner may have new requests added to it.
+ *
+ * An owner may not consider its execution to be complete in the presence of a barrier, and needs
+ * to take note of priority changes that occur while within a barrier, since newly-added requests
+ * may have been created with the previous priority value.
+ */
+class FRequestBarrier
+{
+public:
+	inline explicit FRequestBarrier(IRequestOwner& InOwner)
+		: Owner(InOwner)
+	{
+		Owner.BeginBarrier();
+	}
+
+	inline ~FRequestBarrier()
+	{
+		Owner.EndBarrier();
+	}
+
+	FRequestBarrier(const FRequestBarrier&) = delete;
+	FRequestBarrier& operator=(const FRequestBarrier&) = delete;
+
+private:
+	IRequestOwner& Owner;
+};
+
+template <typename CallbackType, typename... CallbackArgTypes>
+TRefCountPtr<IRequest> IRequestOwner::End(IRequest* Request, CallbackType&& Callback, CallbackArgTypes&&... CallbackArgs)
+{
+	FRequestBarrier Barrier(*this);
+	TRefCountPtr<IRequest> RequestRef = End(Request);
+	Invoke(Forward<CallbackType>(Callback), Forward<CallbackArgTypes>(CallbackArgs)...);
+	return RequestRef;
+}
+
+} // UE::DerivedData
+
+#undef UE_API
diff --git a/Engine/Source/Developer/DerivedDataCache/Public/DerivedDataRequestTypes.h b/Engine/Source/Developer/DerivedDataCache/Public/DerivedDataRequestTypes.h
new file mode 100644
--- /dev/null
+++ b/Engine/Source/Developer/DerivedDataCache/Public/DerivedDataRequestTypes.h
@@ -0,0 +1,57 @@
+// Copyright Epic Games, Inc. All Rights Reserved.
+
+#pragma once
+
+#include "CoreTypes.h"
+
+namespace UE::DerivedData
+{
+
+/** Priority for scheduling a request. */
+enum class EPriority : uint8
+{
+	/**
+	 * Lowest is the minimum priority for asynchronous requests, and primarily for requests that are
+	 * speculative in nature, while minimizing impact on other requests.
+	 */
+	Lowest,
+	/**
+	 * Low is intended for requests that are below the default priority, but are used for operations
+	 * that the program will execute now rather than at an unknown time in the future.
+	 */
+	Low,
+	/**
+	 * Normal is intended as the default request priority.
+	 */
+	Normal,
+	/**
+	 * High is intended for requests that are on the critical path, but are not required to maintain
+	 * interactivity of the program.
+	 */
+	High,
+	/**
+	 * Highest is the maximum priority for asynchronous requests, and intended for requests that are
+	 * required to maintain interactivity of the program.
+	 */
+	Highest,
+	/**
+	 * Blocking is to be used only when the thread making the request will wait on completion of the
+	 * request before doing any other work. Requests at this priority level will be processed before
+	 * any request at a lower priority level. This priority permits a request to be processed on the
+	 * thread making the request. Waiting on a request may increase its priority to this level.
+	 */
+	Blocking,
+};
+
+/** Status of a request that has completed. */
+enum class EStatus : uint8
+{
+	/** The request completed successfully. Any requested data is available. */
+	Ok,
+	/** The request completed unsuccessfully. Any requested data is not available. */
+	Error,
+	/** The request was canceled before it completed. Any requested data is not available. */
+	Canceled,
+};
+
+} // UE::DerivedData
diff --git a/Engine/Source/Editor/UnrealEd/Private/DerivedDataBuildRemoteExecutor.cpp b/Engine/Source/Editor/UnrealEd/Private/DerivedDataBuildRemoteExecutor.cpp
--- a/Engine/Source/Editor/UnrealEd/Private/DerivedDataBuildRemoteExecutor.cpp
+++ b/Engine/Source/Editor/UnrealEd/Private/DerivedDataBuildRemoteExecutor.cpp
@@ -8,6 +8,8 @@
 #include "DerivedDataBuildInputs.h"
 #include "DerivedDataBuildOutput.h"
 #include "DerivedDataBuildWorker.h"
+#include "DerivedDataRequest.h"
+#include "DerivedDataRequestOwner.h"
 #include "Features/IModularFeatures.h"
 #include "HAL/Event.h"
 #include "HAL/PlatformProcess.h"
diff --git a/Engine/Source/Editor/UnrealEd/Private/EditorDomain/BulkDataRegistryEditorDomain.h b/Engine/Source/Editor/UnrealEd/Private/EditorDomain/BulkDataRegistryEditorDomain.h
--- a/Engine/Source/Editor/UnrealEd/Private/EditorDomain/BulkDataRegistryEditorDomain.h
+++ b/Engine/Source/Editor/UnrealEd/Private/EditorDomain/BulkDataRegistryEditorDomain.h
@@ -8,7 +8,7 @@
 #include "Containers/Array.h"
 #include "Containers/Map.h"
 #include "Containers/RingBuffer.h"
-#include "DerivedDataRequest.h"
+#include "DerivedDataRequestOwner.h"
 #include "HAL/CriticalSection.h"
 #include "Serialization/BulkDataRegistry.h"
 #include "Templates/RefCounting.h"
diff --git a/Engine/Source/Editor/UnrealEd/Private/EditorDomain/EditorBuildInputResolver.cpp b/Engine/Source/Editor/UnrealEd/Private/EditorDomain/EditorBuildInputResolver.cpp
--- a/Engine/Source/Editor/UnrealEd/Private/EditorDomain/EditorBuildInputResolver.cpp
+++ b/Engine/Source/Editor/UnrealEd/Private/EditorDomain/EditorBuildInputResolver.cpp
@@ -6,6 +6,8 @@
 #include <atomic>
 #include "Containers/ArrayView.h"
 #include "DerivedDataBuildAction.h"
+#include "DerivedDataRequest.h"
+#include "DerivedDataRequestOwner.h"
 #include "HAL/Event.h"
 #include "Misc/StringBuilder.h"
 #include "Serialization/BulkDataRegistry.h"
diff --git a/Engine/Source/Editor/UnrealEd/Private/EditorDomain/EditorDomainArchive.h b/Engine/Source/Editor/UnrealEd/Private/EditorDomain/EditorDomainArchive.h
--- a/Engine/Source/Editor/UnrealEd/Private/EditorDomain/EditorDomainArchive.h
+++ b/Engine/Source/Editor/UnrealEd/Private/EditorDomain/EditorDomainArchive.h
@@ -6,7 +6,7 @@
 #include "Containers/Array.h"
 #include "DerivedDataCache.h"
 #include "DerivedDataPayload.h"
-#include "DerivedDataRequest.h"
+#include "DerivedDataRequestOwner.h"
 #include "EditorDomain/EditorDomain.h"
 #include "HAL/CriticalSection.h"
 #include "HAL/Platform.h"
diff --git a/Engine/Source/Editor/UnrealEd/Private/EditorDomain/EditorDomainUtils.cpp b/Engine/Source/Editor/UnrealEd/Private/EditorDomain/EditorDomainUtils.cpp
--- a/Engine/Source/Editor/UnrealEd/Private/EditorDomain/EditorDomainUtils.cpp
+++ b/Engine/Source/Editor/UnrealEd/Private/EditorDomain/EditorDomainUtils.cpp
@@ -9,6 +9,7 @@
 #include "DerivedDataCache.h"
 #include "DerivedDataCacheKey.h"
 #include "DerivedDataCacheRecord.h"
+#include "DerivedDataRequestOwner.h"
 #include "Editor.h"
 #include "HAL/FileManager.h"
 #include "IO/PackageStoreWriter.h"
diff --git a/Engine/Source/Programs/DerivedDataBuildWorker/Private/DerivedDataBuildWorker.cpp b/Engine/Source/Programs/DerivedDataBuildWorker/Private/DerivedDataBuildWorker.cpp
--- a/Engine/Source/Programs/DerivedDataBuildWorker/Private/DerivedDataBuildWorker.cpp
+++ b/Engine/Source/Programs/DerivedDataBuildWorker/Private/DerivedDataBuildWorker.cpp
@@ -11,6 +11,7 @@
 #include "DerivedDataBuildOutput.h"
 #include "DerivedDataBuildSession.h"
 #include "DerivedDataPayload.h"
+#include "DerivedDataRequestOwner.h"
 #include "HAL/FileManager.h"
 #include "Memory/SharedBuffer.h"
 #include "Misc/CommandLine.h"
diff --git a/Engine/Source/Runtime/Engine/Private/TextureDerivedDataTask.cpp b/Engine/Source/Runtime/Engine/Private/TextureDerivedDataTask.cpp
--- a/Engine/Source/Runtime/Engine/Private/TextureDerivedDataTask.cpp
+++ b/Engine/Source/Runtime/Engine/Private/TextureDerivedDataTask.cpp
@@ -32,6 +32,7 @@
 #include "DerivedDataBuildSession.h"
 #include "DerivedDataCacheInterface.h"
 #include "DerivedDataPayload.h"
+#include "DerivedDataRequestOwner.h"
 #include "Engine/TextureCube.h"
 #include "Engine/VolumeTexture.h"
 #include "GenericPlatform/GenericPlatformMath.h"
