DDC: Added flags to control request barrier behavior

#rnx

diff --git a/Engine/Source/Developer/DerivedDataCache/Private/DerivedDataRequestOwner.cpp b/Engine/Source/Developer/DerivedDataCache/Private/DerivedDataRequestOwner.cpp
--- a/Engine/Source/Developer/DerivedDataCache/Private/DerivedDataRequestOwner.cpp
+++ b/Engine/Source/Developer/DerivedDataCache/Private/DerivedDataRequestOwner.cpp
@@ -30,8 +30,8 @@ public:
 	void Begin(IRequest* Request) final;
 	TRefCountPtr<IRequest> End(IRequest* Request) final;
 
-	void BeginBarrier() final;
-	void EndBarrier() final;
+	void BeginBarrier(ERequestBarrierFlags Flags) final;
+	void EndBarrier(ERequestBarrierFlags Flags) final;
 
 	inline EPriority GetPriority() const final { return Priority; }
 	inline bool IsCanceled() const final { return bIsCanceled; }
@@ -49,6 +49,7 @@ private:
 	TArray<TRefCountPtr<IRequest>, TInlineAllocator<1>> Requests;
 	TRefCountPtr<FRequestBarrierEvent> BarrierEvent;
 	uint32 BarrierCount{0};
+	uint32 PriorityBarrierCount{0};
 	EPriority Priority{EPriority::Normal};
 	uint8 bPriorityChangedInBarrier : 1;
 	uint8 bBeginExecuted : 1;
@@ -106,22 +107,32 @@ TRefCountPtr<IRequest> FRequestOwnerShared::End(IRequest* Request)
 	return RequestRef;
 }
 
-void FRequestOwnerShared::BeginBarrier()
+void FRequestOwnerShared::BeginBarrier(ERequestBarrierFlags Flags)
 {
 	AddRef();
 	FWriteScopeLock WriteLock(Lock);
 	++BarrierCount;
+	if (EnumHasAnyFlags(Flags, ERequestBarrierFlags::Priority))
+	{
+		++PriorityBarrierCount;
+	}
 }
 
-void FRequestOwnerShared::EndBarrier()
+void FRequestOwnerShared::EndBarrier(ERequestBarrierFlags Flags)
 {
 	ON_SCOPE_EXIT { Release(); };
 	TRefCountPtr<FRequestBarrierEvent> LocalBarrierEvent;
-	if (FWriteScopeLock WriteLock(Lock); --BarrierCount == 0)
 	{
-		bPriorityChangedInBarrier = false;
-		LocalBarrierEvent = MoveTemp(BarrierEvent);
-		BarrierEvent = nullptr;
+		FWriteScopeLock WriteLock(Lock);
+		if (EnumHasAnyFlags(Flags, ERequestBarrierFlags::Priority) && --PriorityBarrierCount == 0)
+		{
+			bPriorityChangedInBarrier = false;
+		}
+		if (--BarrierCount == 0)
+		{
+			LocalBarrierEvent = MoveTemp(BarrierEvent);
+			BarrierEvent = nullptr;
+		}
 	}
 	if (LocalBarrierEvent)
 	{
@@ -140,7 +151,7 @@ void FRequestOwnerShared::SetPriority(EPriority NewPriority)
 	{
 		Priority = NewPriority;
 		LocalRequests = Requests;
-		bPriorityChangedInBarrier = (BarrierCount > 0);
+		bPriorityChangedInBarrier = (PriorityBarrierCount > 0);
 	}
 
 	for (IRequest* Request : LocalRequests)
diff --git a/Engine/Source/Developer/DerivedDataCache/Public/DerivedDataRequestOwner.h b/Engine/Source/Developer/DerivedDataCache/Public/DerivedDataRequestOwner.h
--- a/Engine/Source/Developer/DerivedDataCache/Public/DerivedDataRequestOwner.h
+++ b/Engine/Source/Developer/DerivedDataCache/Public/DerivedDataRequestOwner.h
@@ -3,6 +3,7 @@
 #pragma once
 
 #include "CoreTypes.h"
+#include "Misc/EnumClassFlags.h"
 #include "Templates/Invoke.h"
 #include "Templates/RefCounting.h"
 #include "Templates/UniquePtr.h"
@@ -16,6 +17,17 @@ namespace UE::DerivedData { enum class EPriority : uint8; }
 namespace UE::DerivedData
 {
 
+/** Flags to control the behavior of request barriers. */
+enum class ERequestBarrierFlags : uint32
+{
+	/** A value without any flags set. */
+	None            = 0,
+	/** Within the barrier, calls to SetPriority are replayed to requests in calls to Begin. */
+	Priority        = 1 << 0,
+};
+
+ENUM_CLASS_FLAGS(ERequestBarrierFlags);
+
 /**
  * A request owner manages requests throughout their execution.
  *
@@ -61,8 +73,8 @@ public:
 	inline TRefCountPtr<IRequest> End(IRequest* Request, CallbackType&& Callback, CallbackArgTypes&&... CallbackArgs);
 
 	/** See FRequestBarrier. */
-	virtual void BeginBarrier() = 0;
-	virtual void EndBarrier() = 0;
+	virtual void BeginBarrier(ERequestBarrierFlags Flags) = 0;
+	virtual void EndBarrier(ERequestBarrierFlags Flags) = 0;
 
 	/** Returns the priority that new requests are expected to inherit. */
 	virtual EPriority GetPriority() const = 0;
@@ -126,15 +138,16 @@ private:
 class FRequestBarrier
 {
 public:
-	inline explicit FRequestBarrier(IRequestOwner& InOwner)
+	inline explicit FRequestBarrier(IRequestOwner& InOwner, ERequestBarrierFlags InFlags = ERequestBarrierFlags::None)
 		: Owner(InOwner)
+		, Flags(InFlags)
 	{
-		Owner.BeginBarrier();
+		Owner.BeginBarrier(Flags);
 	}
 
 	inline ~FRequestBarrier()
 	{
-		Owner.EndBarrier();
+		Owner.EndBarrier(Flags);
 	}
 
 	FRequestBarrier(const FRequestBarrier&) = delete;
@@ -142,12 +155,13 @@ public:
 
 private:
 	IRequestOwner& Owner;
+	ERequestBarrierFlags Flags;
 };
 
 template <typename CallbackType, typename... CallbackArgTypes>
 TRefCountPtr<IRequest> IRequestOwner::End(IRequest* Request, CallbackType&& Callback, CallbackArgTypes&&... CallbackArgs)
 {
-	FRequestBarrier Barrier(*this);
+	FRequestBarrier Barrier(*this, ERequestBarrierFlags::Priority);
 	TRefCountPtr<IRequest> RequestRef = End(Request);
 	Invoke(Forward<CallbackType>(Callback), Forward<CallbackArgTypes>(CallbackArgs)...);
 	return RequestRef;
diff --git a/Engine/Source/Programs/DerivedDataBuildWorker/Private/DerivedDataBuildWorker.cpp b/Engine/Source/Programs/DerivedDataBuildWorker/Private/DerivedDataBuildWorker.cpp
--- a/Engine/Source/Programs/DerivedDataBuildWorker/Private/DerivedDataBuildWorker.cpp
+++ b/Engine/Source/Programs/DerivedDataBuildWorker/Private/DerivedDataBuildWorker.cpp
@@ -217,30 +217,32 @@ bool FBuildWorkerProgram::Build()
 		UE_LOG(LogDerivedDataBuildWorker, Error, TEXT("No build actions to operate on."));
 		return false;
 	}
-
-	for (const FString& ActionPath : ActionPaths)
+	else
 	{
-		UE_LOG(LogDerivedDataBuildWorker, Log, TEXT("Loading build action '%s'"), *ActionPath);
-		FCbObject ActionObject;
-		if (TUniquePtr<FArchive> Ar{IFileManager::Get().CreateFileReader(*ActionPath, FILEREAD_Silent)})
-		{
-			*Ar << ActionObject;
-		}
-		else
-		{
-			UE_LOG(LogDerivedDataBuildWorker, Error, TEXT("Missing build action '%s'"), *ActionPath);
-			return false;
-		}
-		if (FOptionalBuildAction Action = FBuildAction::Load(ActionPath, MoveTemp(ActionObject)); Action.IsNull())
-		{
-			UE_LOG(LogDerivedDataBuildWorker, Error, TEXT("Invalid build action '%s'"), *ActionPath);
-			return false;
-		}
-		else
+		FRequestBarrier Barrier(Owner);
+		for (const FString& ActionPath : ActionPaths)
 		{
-			FRequestBarrier Barrier(Owner);
-			Session.BuildAction(Action.Get(), {}, EBuildPolicy::BuildLocal, Owner,
-				[this](FBuildActionCompleteParams&& Params) { BuildComplete(MoveTemp(Params)); });
+			UE_LOG(LogDerivedDataBuildWorker, Log, TEXT("Loading build action '%s'"), *ActionPath);
+			FCbObject ActionObject;
+			if (TUniquePtr<FArchive> Ar{IFileManager::Get().CreateFileReader(*ActionPath, FILEREAD_Silent)})
+			{
+				*Ar << ActionObject;
+			}
+			else
+			{
+				UE_LOG(LogDerivedDataBuildWorker, Error, TEXT("Missing build action '%s'"), *ActionPath);
+				return false;
+			}
+			if (FOptionalBuildAction Action = FBuildAction::Load(ActionPath, MoveTemp(ActionObject)); Action.IsNull())
+			{
+				UE_LOG(LogDerivedDataBuildWorker, Error, TEXT("Invalid build action '%s'"), *ActionPath);
+				return false;
+			}
+			else
+			{
+				Session.BuildAction(Action.Get(), {}, EBuildPolicy::BuildLocal, Owner,
+					[this](FBuildActionCompleteParams&& Params) { BuildComplete(MoveTemp(Params)); });
+			}
 		}
 	}
 
