DDC: Add FBuildConfigContext::SetRequiredMemory() and forward the estimate to the build scheduler

#rnx

diff --git a/Engine/Source/Developer/DerivedDataCache/Private/DerivedDataBuildJob.cpp b/Engine/Source/Developer/DerivedDataCache/Private/DerivedDataBuildJob.cpp
--- a/Engine/Source/Developer/DerivedDataCache/Private/DerivedDataBuildJob.cpp
+++ b/Engine/Source/Developer/DerivedDataCache/Private/DerivedDataBuildJob.cpp
@@ -417,6 +417,7 @@ void FBuildJob::CreateContext()
 			SchedulerParams.ResolvedInputsSize += Buffer.GetRawSize();
 		});
 	}
+	SchedulerParams.TotalRequiredMemory = SchedulerParams.TotalInputsSize + Context->GetRequiredMemory();
 }
 
 ///////////////////////////////////////////////////////////////////////////////////////////////////////////////////////
diff --git a/Engine/Source/Developer/DerivedDataCache/Private/DerivedDataBuildJobContext.h b/Engine/Source/Developer/DerivedDataCache/Private/DerivedDataBuildJobContext.h
--- a/Engine/Source/Developer/DerivedDataCache/Private/DerivedDataBuildJobContext.h
+++ b/Engine/Source/Developer/DerivedDataCache/Private/DerivedDataBuildJobContext.h
@@ -42,7 +42,7 @@ public:
 
 	inline ECachePolicy GetCachePolicy() const final { return CachePolicy; }
 	inline EBuildPolicy GetBuildPolicy() const final { return BuildPolicy; }
-
+	inline uint64 GetRequiredMemory() const { return RequiredMemory; }
 	inline bool ShouldCheckDeterministicOutput() const { return bDeterministicOutputCheck; }
 
 	void AddConstant(FStringView Key, FCbObject&& Value);
@@ -67,7 +67,8 @@ private:
 	void SetCacheBucket(FCacheBucket Bucket) final;
 	void SetCachePolicy(ECachePolicy Policy) final;
 	void SetBuildPolicy(EBuildPolicy Policy) final;
-	void SkipDeterministicOutputCheck() { bDeterministicOutputCheck = false; }
+	void SetRequiredMemory(uint64 InRequiredMemory) final { RequiredMemory = InRequiredMemory; }
+	void SkipDeterministicOutputCheck() final { bDeterministicOutputCheck = false; }
 
 	void SetPriority(EPriority Priority) final;
 	void Cancel() final;
@@ -83,6 +84,7 @@ private:
 	FEventRef BuildCompleteEvent{EEventMode::ManualReset};
 	TUniqueFunction<void ()> OnEndBuild;
 	IRequestOwner* Owner{};
+	uint64 RequiredMemory{};
 	ECachePolicy CachePolicy;
 	EBuildPolicy BuildPolicy;
 	bool bIsAsyncBuild{false};
diff --git a/Engine/Source/Developer/DerivedDataCache/Public/DerivedDataBuildFunction.h b/Engine/Source/Developer/DerivedDataCache/Public/DerivedDataBuildFunction.h
--- a/Engine/Source/Developer/DerivedDataCache/Public/DerivedDataBuildFunction.h
+++ b/Engine/Source/Developer/DerivedDataCache/Public/DerivedDataBuildFunction.h
@@ -149,6 +149,9 @@ public:
 	/** Overrides the build policy used when executing this build. */
 	virtual void SetBuildPolicy(EBuildPolicy Policy) = 0;
 
+	/** Sets the estimated peak memory required to execute the build, excluding input size. */
+	virtual void SetRequiredMemory(uint64 RequiredMemory) = 0;
+
 	/** Skips verification that this function has deterministic output. */
 	virtual void SkipDeterministicOutputCheck() = 0;
 };
diff --git a/Engine/Source/Developer/DerivedDataCache/Public/DerivedDataBuildScheduler.h b/Engine/Source/Developer/DerivedDataCache/Public/DerivedDataBuildScheduler.h
--- a/Engine/Source/Developer/DerivedDataCache/Public/DerivedDataBuildScheduler.h
+++ b/Engine/Source/Developer/DerivedDataCache/Public/DerivedDataBuildScheduler.h
@@ -79,6 +79,9 @@ struct FBuildSchedulerParams
 	uint64 MissingLocalInputsSize = 0;
 	/** Total size of inputs that need to be resolved for remote execution. Available in ResolveInputData. */
 	uint64 MissingRemoteInputsSize = 0;
+
+	/** Estimate of the peak memory required to execute the build, including constants and inputs. */
+	uint64 TotalRequiredMemory = 0;
 };
 
 } // UE::DerivedData
