DDC: Fixed a build job assert caused by canceling at the wrong time

#rnx

diff --git a/Engine/Source/Developer/DerivedDataCache/Private/DerivedDataBuildJob.cpp b/Engine/Source/Developer/DerivedDataCache/Private/DerivedDataBuildJob.cpp
--- a/Engine/Source/Developer/DerivedDataCache/Private/DerivedDataBuildJob.cpp
+++ b/Engine/Source/Developer/DerivedDataCache/Private/DerivedDataBuildJob.cpp
@@ -950,7 +950,10 @@ void FBuildJob::AdvanceToState(EBuildJobState NewState)
 		}
 	}
 
-	ExecuteState(NewState);
+	if (OldState < NewState)
+	{
+		ExecuteState(NewState);
+	}
 }
 
 void FBuildJob::ExecuteTransition(EBuildJobState OldState, EBuildJobState NewState)
@@ -985,7 +988,7 @@ void FBuildJob::ExecuteTransition(EBuildJobState OldState, EBuildJobState NewSta
 	}
 	if (OldState <= EBuildJobState::ExecuteLocalWait && EBuildJobState::ExecuteLocalWait < NewState && !Output)
 	{
-		SetOutputNoCheck(OutputBuilder.Build());
+		SetOutputNoCheck(OutputBuilder.Build(), NewState);
 	}
 	if (NewState == EBuildJobState::Complete)
 	{
