DLL-exported functions for Zen HTTP interfaces to allow them to be consumed from other modules, and renamed them to make it easier to search for them (previous names were derived from the Jupiter HTTP back-end)

diff --git a/Engine/Source/Developer/DerivedDataCache/Internal/ZenServerHttp.h b/Engine/Source/Developer/DerivedDataCache/Internal/ZenServerHttp.h
--- a/Engine/Source/Developer/DerivedDataCache/Internal/ZenServerHttp.h
+++ b/Engine/Source/Developer/DerivedDataCache/Internal/ZenServerHttp.h
@@ -36,8 +36,8 @@ enum class EContentType
 class FZenHttpRequest
 {
 public:
-	FZenHttpRequest(const TCHAR* InDomain, bool bInLogErrors);
-	~FZenHttpRequest();
+	DERIVEDDATACACHE_API FZenHttpRequest(const TCHAR* InDomain, bool bInLogErrors);
+	DERIVEDDATACACHE_API ~FZenHttpRequest();
 
 	/**
 		* Resets all options on the request except those that should always be set.
@@ -76,7 +76,7 @@ public:
 		* @param Buffer Data to upload
 		* @return Result of the request
 		*/
-	Result PerformBlockingPut(const TCHAR* Uri, const FCompositeBuffer& Buffer, EContentType ContentType);
+	DERIVEDDATACACHE_API Result PerformBlockingPut(const TCHAR* Uri, const FCompositeBuffer& Buffer, EContentType ContentType);
 
 	/**
 		* Download an url into a buffer using the request.
@@ -85,7 +85,7 @@ public:
 		* downloaded data will be stored in an internal buffer and accessed via GetResponseAsString
 		* @return Result of the request
 		*/
-	Result PerformBlockingDownload(const TCHAR* Uri, TArray<uint8>* Buffer);
+	DERIVEDDATACACHE_API Result PerformBlockingDownload(const TCHAR* Uri, TArray<uint8>* Buffer);
 
 	/**
 		* Download an url into a buffer using the request.
@@ -93,21 +93,21 @@ public:
 		* @param OutPackage Package instance which will receive the data
 		* @result Request success/failure status
 		*/
-	Result PerformBlockingDownload(const TCHAR* Uri, FCbPackage& OutPackage);
+	DERIVEDDATACACHE_API Result PerformBlockingDownload(const TCHAR* Uri, FCbPackage& OutPackage);
 
 	/**
 		* Query an url using the request. Queries can use either "Head" or "Delete" verbs.
 		* @param Uri Url to use.
 		* @return Result of the request
 		*/
-	Result PerformBlockingHead(const TCHAR* Uri);
+	DERIVEDDATACACHE_API Result PerformBlockingHead(const TCHAR* Uri);
 
 	/**
 		* Query an url using the request. Queries can use either "Head" or "Delete" verbs.
 		* @param Uri Url to use.
 		* @return Result of the request
 		*/
-	Result PerformBlockingDelete(const TCHAR* Uri);
+	DERIVEDDATACACHE_API Result PerformBlockingDelete(const TCHAR* Uri);
 
 	/**
 		* Returns the response buffer as a string. Note that is the request is performed
@@ -171,21 +171,21 @@ private:
   * 
   * Intended to be used with \ref FScopedRequestPtr which handles lifetime management transparently
   */
-struct FRequestPool
+struct FZenHttpRequestPool
 {
-	explicit FRequestPool(const TCHAR* InServiceUrl);
-	~FRequestPool();
+	DERIVEDDATACACHE_API explicit FZenHttpRequestPool(const TCHAR* InServiceUrl);
+	DERIVEDDATACACHE_API ~FZenHttpRequestPool();
 
 	/** Block until a request is free. Once a request has been returned it is
 	  * "owned by the caller and need to release it to the pool when work has been completed.
 	  * @return Usable request instance.
 	  */
-	FZenHttpRequest* WaitForFreeRequest();
+	DERIVEDDATACACHE_API FZenHttpRequest* WaitForFreeRequest();
 
 	/** Release request to the pool.
 	  * @param Request Request that should be freed. Note that any buffer owned by the request can now be reset.
 	  */
-	void ReleaseRequestToPool(FZenHttpRequest* Request);
+	DERIVEDDATACACHE_API void ReleaseRequestToPool(FZenHttpRequest* Request);
 
 private:
 	struct FEntry
@@ -200,16 +200,16 @@ private:
 /**
   * Utility class to manage requesting and releasing requests from the \ref FRequestPool.
   */
-struct FScopedRequestPtr
+struct FZenScopedRequestPtr
 {
 public:
-	FScopedRequestPtr(FRequestPool* InPool)
+	FZenScopedRequestPtr(FZenHttpRequestPool* InPool)
 	:	Request(InPool->WaitForFreeRequest())
 	,	Pool(InPool)
 	{
 	}
 
-	~FScopedRequestPtr()
+	~FZenScopedRequestPtr()
 	{
 		Pool->ReleaseRequestToPool(Request);
 	}
@@ -228,8 +228,8 @@ public:
 	}
 
 private:
-	FZenHttpRequest*	Request;
-	FRequestPool*		Pool;
+	FZenHttpRequest*		Request;
+	FZenHttpRequestPool*	Pool;
 };
 
 } // namespace UE::Zen
diff --git a/Engine/Source/Developer/DerivedDataCache/Private/ZenDerivedDataBackend.cpp b/Engine/Source/Developer/DerivedDataCache/Private/ZenDerivedDataBackend.cpp
--- a/Engine/Source/Developer/DerivedDataCache/Private/ZenDerivedDataBackend.cpp
+++ b/Engine/Source/Developer/DerivedDataCache/Private/ZenDerivedDataBackend.cpp
@@ -39,7 +39,7 @@ FZenDerivedDataBackend::FZenDerivedDataBackend(
 {
 	if (IsServiceReady())
 	{
-		RequestPool = MakeUnique<Zen::FRequestPool>(InServiceUrl);
+		RequestPool = MakeUnique<Zen::FZenHttpRequestPool>(InServiceUrl);
 		bIsUsable = true;
 	}
 	bCacheRecordEndpointEnabled = false;
@@ -55,7 +55,6 @@ FString FZenDerivedDataBackend::GetName() const
 	return Domain;
 }
 
-
 bool FZenDerivedDataBackend::IsServiceReady()
 {
 	Zen::FZenHttpRequest Request(*Domain, false);
@@ -109,7 +108,7 @@ bool FZenDerivedDataBackend::CachedDataProbablyExists(const TCHAR* CacheKey)
 	// Retry request until we get an accepted response or exhaust allowed number of attempts.
 	while (ResponseCode == 0 && ++Attempts < MaxAttempts)
 	{
-		Zen::FScopedRequestPtr Request(RequestPool.Get());
+		Zen::FZenScopedRequestPtr Request(RequestPool.Get());
 		Zen::FZenHttpRequest::Result Result = Request->PerformBlockingHead(*Uri);
 		ResponseCode = Request->GetResponseCode();
 
@@ -181,7 +180,7 @@ FZenDerivedDataBackend::GetZenData(const TCHAR* Uri, TArray<uint8>* OutData) con
 	EGetResult GetResult = EGetResult::NotFound;
 	for (uint32 Attempts = 0; Attempts < MaxAttempts; ++Attempts)
 	{
-		Zen::FScopedRequestPtr Request(RequestPool.Get());
+		Zen::FZenScopedRequestPtr Request(RequestPool.Get());
 		if (Request.IsValid())
 		{
 			Zen::FZenHttpRequest::Result Result = Request->PerformBlockingDownload(Uri, OutData);
@@ -204,6 +203,7 @@ FZenDerivedDataBackend::GetZenData(const TCHAR* Uri, TArray<uint8>* OutData) con
 	{
 		OutData->Reset();
 	}
+
 	return GetResult;
 }
 
@@ -218,7 +218,7 @@ FZenDerivedDataBackend::GetZenData(const FCacheKey& CacheKey, ECachePolicy Cache
 	EGetResult GetResult = EGetResult::NotFound;
 	for (uint32 Attempts = 0; Attempts < MaxAttempts; ++Attempts)
 	{
-		Zen::FScopedRequestPtr Request(RequestPool.Get());
+		Zen::FZenScopedRequestPtr Request(RequestPool.Get());
 		if (Request.IsValid())
 		{
 			Zen::FZenHttpRequest::Result Result = Request->PerformBlockingDownload(QueryUri.ToString(), OutPackage);
@@ -277,7 +277,7 @@ FZenDerivedDataBackend::PutZenData(const TCHAR* Uri, const FCompositeBuffer& InD
 	// Retry request until we get an accepted response or exhaust allowed number of attempts.
 	while (ResponseCode == 0 && ++Attempts < MaxAttempts)
 	{
-		Zen::FScopedRequestPtr Request(RequestPool.Get());
+		Zen::FZenScopedRequestPtr Request(RequestPool.Get());
 		if (Request.IsValid())
 		{
 			Zen::FZenHttpRequest::Result Result = Request->PerformBlockingPut(Uri, InData, ContentType);
@@ -399,7 +399,7 @@ void FZenDerivedDataBackend::RemoveCachedData(const TCHAR* CacheKey, bool bTrans
 	// Retry request until we get an accepted response or exhaust allowed number of attempts.
 	while (ResponseCode == 0 && ++Attempts < MaxAttempts)
 	{
-		Zen::FScopedRequestPtr Request(RequestPool.Get());
+		Zen::FZenScopedRequestPtr Request(RequestPool.Get());
 		if (Request)
 		{
 			Zen::FZenHttpRequest::Result Result = Request->PerformBlockingDelete(*Uri);
diff --git a/Engine/Source/Developer/DerivedDataCache/Private/ZenDerivedDataBackend.h b/Engine/Source/Developer/DerivedDataCache/Private/ZenDerivedDataBackend.h
--- a/Engine/Source/Developer/DerivedDataCache/Private/ZenDerivedDataBackend.h
+++ b/Engine/Source/Developer/DerivedDataCache/Private/ZenDerivedDataBackend.h
@@ -29,7 +29,7 @@ struct FIoHash;
 
 namespace UE::Zen {
 	enum class EContentType;
-	struct FRequestPool;
+	struct FZenHttpRequestPool;
 }
 
 namespace UE::DerivedData {
@@ -159,7 +159,7 @@ private:
 	FString Domain;
 	FString Namespace;
 	mutable FDerivedDataCacheUsageStats UsageStats;
-	TUniquePtr<UE::Zen::FRequestPool> RequestPool;
+	TUniquePtr<UE::Zen::FZenHttpRequestPool> RequestPool;
 	bool bIsUsable = false;
 	uint32 FailedLoginAttempts = 0;
 	uint32 MaxAttempts = 4;
diff --git a/Engine/Source/Developer/DerivedDataCache/Private/ZenServerHttp.cpp b/Engine/Source/Developer/DerivedDataCache/Private/ZenServerHttp.cpp
--- a/Engine/Source/Developer/DerivedDataCache/Private/ZenServerHttp.cpp
+++ b/Engine/Source/Developer/DerivedDataCache/Private/ZenServerHttp.cpp
@@ -449,7 +449,7 @@ namespace UE::Zen {
 
 	//////////////////////////////////////////////////////////////////////////
 
-	FRequestPool::FRequestPool(const TCHAR* InServiceUrl)
+	FZenHttpRequestPool::FZenHttpRequestPool(const TCHAR* InServiceUrl)
 	{
 		for (uint8 i = 0; i < Pool.Num(); ++i)
 		{
@@ -458,7 +458,7 @@ namespace UE::Zen {
 		}
 	}
 
-	FRequestPool::~FRequestPool()
+	FZenHttpRequestPool::~FZenHttpRequestPool()
 	{
 		for (uint8 i = 0; i < Pool.Num(); ++i)
 		{
@@ -473,7 +473,7 @@ namespace UE::Zen {
 		* "owned by the caller and need to release it to the pool when work has been completed.
 		* @return Usable request instance.
 		*/
-	FZenHttpRequest* FRequestPool::WaitForFreeRequest()
+	FZenHttpRequest* FZenHttpRequestPool::WaitForFreeRequest()
 	{
 		TRACE_CPUPROFILER_EVENT_SCOPE(ZenDDC_WaitForConnPool);
 		while (true)
@@ -497,7 +497,7 @@ namespace UE::Zen {
 		* Release request to the pool.
 		* @param Request Request that should be freed. Note that any buffer owened by the request can now be reset.
 		*/
-	void FRequestPool::ReleaseRequestToPool(FZenHttpRequest* Request)
+	void FZenHttpRequestPool::ReleaseRequestToPool(FZenHttpRequest* Request)
 	{
 		for (uint8 i = 0; i < Pool.Num(); ++i)
 		{
