Updated Zen HTTP interfaces to accept large payloads as a preparation for unifying the cache/store client implementations

diff --git a/Engine/Source/Developer/DerivedDataCache/Internal/ZenServerHttp.h b/Engine/Source/Developer/DerivedDataCache/Internal/ZenServerHttp.h
--- a/Engine/Source/Developer/DerivedDataCache/Internal/ZenServerHttp.h
+++ b/Engine/Source/Developer/DerivedDataCache/Internal/ZenServerHttp.h
@@ -6,8 +6,10 @@
 
 #include "Containers/StaticArray.h"
 #include "Containers/UnrealString.h"
+#include "Memory/MemoryFwd.h"
 
 class FCompositeBuffer;
+class FCbObjectView;
 class FCbPackage;
 
 #if UE_WITH_ZEN
@@ -79,13 +81,13 @@ public:
 	DERIVEDDATACACHE_API Result PerformBlockingPut(const TCHAR* Uri, const FCompositeBuffer& Buffer, EContentType ContentType);
 
 	/**
-		* Download an url into a buffer using the request.
-		* @param Uri Url to use.
-		* @param Buffer Optional buffer where data should be downloaded to. If this is null then 
-		* downloaded data will be stored in an internal buffer and accessed via GetResponseAsString
-		* @return Result of the request
-		*/
-	DERIVEDDATACACHE_API Result PerformBlockingDownload(const TCHAR* Uri, TArray<uint8>* Buffer);
+	* Download an url into a buffer using the request.
+	* @param Uri Url to use.
+	* @param Buffer Optional buffer where data should be downloaded to. If this is null then
+	* downloaded data will be stored in an internal buffer and accessed via GetResponseAsString
+	* @return Result of the request
+	*/
+	DERIVEDDATACACHE_API Result PerformBlockingDownload(FStringView Uri, TArray64<uint8>* Buffer);
 
 	/**
 		* Download an url into a buffer using the request.
@@ -128,12 +130,12 @@ private:
 	bool					bResponseFormatValid = false;
 
 	const FCompositeBuffer*	ReadDataView = nullptr;
-	TArray<uint8>*			WriteDataBufferPtr = nullptr;
+	TArray64<uint8>*		WriteDataBufferPtr = nullptr;
 	FCbPackage*				WriteDataPackage = nullptr;
-	TArray<uint8>*			WriteHeaderBufferPtr = nullptr;
+	TArray64<uint8>*		WriteHeaderBufferPtr = nullptr;
 
-	TArray<uint8>			ResponseHeader;
-	TArray<uint8>			ResponseBuffer;
+	TArray64<uint8>			ResponseHeader;
+	TArray64<uint8>			ResponseBuffer;
 	TArray<FString>			Headers;
 	FString					Domain;
 
@@ -158,9 +160,9 @@ private:
 		* @param Buffer Optional buffer to directly receive the result of the request.
 		* If unset the response body will be stored in the request.
 		*/
-	Result PerformBlocking(const TCHAR* Uri, RequestVerb Verb, uint32 ContentLength);
+	Result PerformBlocking(FStringView Uri, RequestVerb Verb, uint32 ContentLength);
 
-	static FString GetAnsiBufferAsString(const TArray<uint8>& Buffer);
+	static FString GetAnsiBufferAsString(const TArray64<uint8>& Buffer);
 
 	struct FStatics;
 };
diff --git a/Engine/Source/Developer/DerivedDataCache/Private/ZenDerivedDataBackend.cpp b/Engine/Source/Developer/DerivedDataCache/Private/ZenDerivedDataBackend.cpp
--- a/Engine/Source/Developer/DerivedDataCache/Private/ZenDerivedDataBackend.cpp
+++ b/Engine/Source/Developer/DerivedDataCache/Private/ZenDerivedDataBackend.cpp
@@ -58,7 +58,7 @@ FString FZenDerivedDataBackend::GetName() const
 bool FZenDerivedDataBackend::IsServiceReady()
 {
 	Zen::FZenHttpRequest Request(*Domain, false);
-	Zen::FZenHttpRequest::Result Result = Request.PerformBlockingDownload(TEXT("health/ready"), nullptr);
+	Zen::FZenHttpRequest::Result Result = Request.PerformBlockingDownload(TEXT("health/ready"_SV), nullptr);
 	
 	if (Result == Zen::FZenHttpRequest::Result::Success && Zen::IsSuccessCode(Request.GetResponseCode()))
 	{
@@ -147,7 +147,7 @@ bool FZenDerivedDataBackend::GetCachedData(const TCHAR* CacheKey, TArray<uint8>&
 
 	double StartTime = FPlatformTime::Seconds();
 
-	EGetResult Result = GetZenData(*MakeLegacyZenKey(CacheKey), &OutData);
+	EGetResult Result = GetZenData(MakeLegacyZenKey(CacheKey), &OutData);
 	if (Result != EGetResult::Success)
 	{
 		switch (Result)
@@ -174,7 +174,27 @@ bool FZenDerivedDataBackend::GetCachedData(const TCHAR* CacheKey, TArray<uint8>&
 }
 
 FZenDerivedDataBackend::EGetResult
-FZenDerivedDataBackend::GetZenData(const TCHAR* Uri, TArray<uint8>* OutData) const
+FZenDerivedDataBackend::GetZenData(FStringView Uri, TArray<uint8>* OutData) const
+{
+	TArray64<uint8> TempArray;
+
+	EGetResult Result = GetZenData(Uri, TempArray);
+
+	if (Result == EGetResult::Success)
+	{
+		if (OutData)
+		{
+			check(TempArray.Num() <= UINT32_MAX);
+
+			*OutData = TempArray;
+		}
+	}
+
+	return Result;
+}
+
+FZenDerivedDataBackend::EGetResult
+FZenDerivedDataBackend::GetZenData(FStringView Uri, TArray64<uint8>& OutData) const
 {
 	// Retry request until we get an accepted response or exhaust allowed number of attempts.
 	EGetResult GetResult = EGetResult::NotFound;
@@ -183,7 +203,7 @@ FZenDerivedDataBackend::GetZenData(const TCHAR* Uri, TArray<uint8>* OutData) con
 		Zen::FZenScopedRequestPtr Request(RequestPool.Get());
 		if (Request.IsValid())
 		{
-			Zen::FZenHttpRequest::Result Result = Request->PerformBlockingDownload(Uri, OutData);
+			Zen::FZenHttpRequest::Result Result = Request->PerformBlockingDownload(Uri, &OutData);
 			int64 ResponseCode = Request->GetResponseCode();
 
 			// Request was successful, make sure we got all the expected data.
@@ -199,10 +219,7 @@ FZenDerivedDataBackend::GetZenData(const TCHAR* Uri, TArray<uint8>* OutData) con
 		}
 	}
 
-	if (OutData)
-	{
-		OutData->Reset();
-	}
+	OutData.Reset();
 
 	return GetResult;
 }
@@ -586,8 +603,8 @@ void FZenDerivedDataBackend::GetPayload(
 			AppendZenUri(Key.CacheKey, Key.Id, QueryUri);
 			AppendPolicyQueryString(Policy, QueryUri);
 
-			TArray<uint8> PayloadData;
-			EGetResult GetResult = GetZenData(QueryUri.ToString(), &PayloadData);
+			TArray64<uint8> PayloadData;
+			EGetResult GetResult = GetZenData(QueryUri, PayloadData);
 			FCompressedBuffer CompressedBuffer;
 			if (GetResult == EGetResult::Success)
 			{
diff --git a/Engine/Source/Developer/DerivedDataCache/Private/ZenDerivedDataBackend.h b/Engine/Source/Developer/DerivedDataCache/Private/ZenDerivedDataBackend.h
--- a/Engine/Source/Developer/DerivedDataCache/Private/ZenDerivedDataBackend.h
+++ b/Engine/Source/Developer/DerivedDataCache/Private/ZenDerivedDataBackend.h
@@ -118,7 +118,8 @@ private:
 		NotFound,
 		Corrupted
 	};
-	EGetResult GetZenData(const TCHAR* Uri, TArray<uint8>* OutData) const;
+	EGetResult GetZenData(FStringView Uri, TArray64<uint8>& OutData) const;
+	EGetResult GetZenData(FStringView Uri, TArray<uint8>* OutData) const;
 
 	// TODO: need ability to specify content type
 	FDerivedDataBackendInterface::EPutStatus PutZenData(const TCHAR* Uri, const FCompositeBuffer& InData, Zen::EContentType ContentType);
diff --git a/Engine/Source/Developer/DerivedDataCache/Private/ZenServerHttp.cpp b/Engine/Source/Developer/DerivedDataCache/Private/ZenServerHttp.cpp
--- a/Engine/Source/Developer/DerivedDataCache/Private/ZenServerHttp.cpp
+++ b/Engine/Source/Developer/DerivedDataCache/Private/ZenServerHttp.cpp
@@ -22,7 +22,7 @@
 #include "Serialization/CompactBinary.h"
 #include "Serialization/CompactBinaryPackage.h"
 #include "Serialization/CompactBinaryValidation.h"
-#include "Serialization/MemoryReader.h"
+#include "Serialization/LargeMemoryReader.h"
 
 namespace UE::Zen {
 
@@ -123,7 +123,7 @@ namespace UE::Zen {
 		return PerformBlocking(Uri, RequestVerb::Put, ContentLength);
 	}
 
-	FZenHttpRequest::Result FZenHttpRequest::PerformBlockingDownload(const TCHAR* Uri, TArray<uint8>* Buffer)
+	FZenHttpRequest::Result FZenHttpRequest::PerformBlockingDownload(FStringView Uri, TArray64<uint8>* Buffer)
 	{
 		curl_easy_setopt(Curl, CURLOPT_HTTPGET, 1L);
 		WriteDataBufferPtr = Buffer;
@@ -136,7 +136,7 @@ namespace UE::Zen {
 		curl_easy_setopt(Curl, CURLOPT_HTTPGET, 1L);
 		OutPackage.Reset();
 		// TODO: When PackageBytes can be written in segments directly, set the WritePtr to the OutPackage and use that
-		TArray<uint8> PackageBytes;
+		TArray64<uint8> PackageBytes;
 		WriteDataBufferPtr = &PackageBytes;
 		Result LocalResult = PerformBlocking(Uri, RequestVerb::Get, 0u);
 		if (IsSuccessCode(ResponseCode))
@@ -147,7 +147,7 @@ namespace UE::Zen {
 			}
 			else
 			{
-				FMemoryReader PackageLoader(PackageBytes);
+				FLargeMemoryReader PackageLoader(PackageBytes.GetData(), PackageBytes.Num());
 				bResponseFormatValid = OutPackage.TryLoad(PackageLoader);
 			}
 		}
@@ -169,7 +169,7 @@ namespace UE::Zen {
 		return PerformBlocking(Uri, RequestVerb::Delete, 0u);
 	}
 
-	FZenHttpRequest::Result FZenHttpRequest::PerformBlocking(const TCHAR* Uri, RequestVerb Verb, uint32 ContentLength)
+	FZenHttpRequest::Result FZenHttpRequest::PerformBlocking(FStringView Uri, RequestVerb Verb, uint32 ContentLength)
 	{
 		static const char* CommonHeaders[] = {
 			"User-Agent: UE",
@@ -179,7 +179,7 @@ namespace UE::Zen {
 		TRACE_CPUPROFILER_EVENT_SCOPE(ZenDDC_CurlPerform);
 
 		// Setup request options
-		FString Url = FString::Printf(TEXT("%s/%s"), *Domain, Uri);
+		FString Url = FString::Printf(TEXT("%s/%s"), *Domain, *FString(Uri));
 		curl_easy_setopt(Curl, CURLOPT_URL, TCHAR_TO_ANSI(*Url));
 
 		// Setup response header buffer. If caller has not setup a response data buffer, use internal.
@@ -216,7 +216,7 @@ namespace UE::Zen {
 			bRedirected = (ResponseCode >= 300 && ResponseCode < 400);
 		}
 
-		LogResult(CurlResult, Uri, Verb);
+		LogResult(CurlResult, *FString(Uri), Verb);
 
 		// Clean up
 		curl_slist_free_all(CurlHeaders);
@@ -303,7 +303,7 @@ namespace UE::Zen {
 		}
 	}
 
-	FString FZenHttpRequest::GetAnsiBufferAsString(const TArray<uint8>& Buffer)
+	FString FZenHttpRequest::GetAnsiBufferAsString(const TArray64<uint8>& Buffer)
 	{
 		// Content is NOT null-terminated; we need to specify lengths here
 		FUTF8ToTCHAR TCHARData(reinterpret_cast<const ANSICHAR*>(Buffer.GetData()), Buffer.Num());
@@ -371,7 +371,7 @@ namespace UE::Zen {
 	{
 		FZenHttpRequest* Request = static_cast<FZenHttpRequest*>(UserData);
 		const size_t WriteSize = SizeInBlocks * BlockSizeInBytes;
-		TArray<uint8>* WriteHeaderBufferPtr = Request->WriteHeaderBufferPtr;
+		TArray64<uint8>* WriteHeaderBufferPtr = Request->WriteHeaderBufferPtr;
 		if (WriteHeaderBufferPtr && WriteSize > 0)
 		{
 			const size_t CurrentBufferLength = WriteHeaderBufferPtr->Num();
@@ -393,7 +393,7 @@ namespace UE::Zen {
 	{
 		FZenHttpRequest* Request = static_cast<FZenHttpRequest*>(UserData);
 		const size_t WriteSize = SizeInBlocks * BlockSizeInBytes;
-		TArray<uint8>* WriteDataBufferPtr = Request->WriteDataBufferPtr;
+		TArray64<uint8>* WriteDataBufferPtr = Request->WriteDataBufferPtr;
 
 		if (WriteDataBufferPtr && WriteSize > 0)
 		{
@@ -468,11 +468,13 @@ namespace UE::Zen {
 		}
 	}
 
-	/**
-		* Block until a request is free. Once a request has been returned it is
-		* "owned by the caller and need to release it to the pool when work has been completed.
-		* @return Usable request instance.
-		*/
+	/** Block until a request is free
+	  * 
+	  * Once a request has been returned it is owned by the caller and 
+	  * it needs to release it to the pool when work has been completed
+	  * 
+	  * @return Usable request instance.
+	  */
 	FZenHttpRequest* FZenHttpRequestPool::WaitForFreeRequest()
 	{
 		TRACE_CPUPROFILER_EVENT_SCOPE(ZenDDC_WaitForConnPool);
@@ -494,9 +496,9 @@ namespace UE::Zen {
 	}
 
 	/**
-		* Release request to the pool.
-		* @param Request Request that should be freed. Note that any buffer owened by the request can now be reset.
-		*/
+	  * Release request to the pool.
+	  * @param Request Request that should be freed. Note that any buffer owned by the request can now be reset.
+	  */
 	void FZenHttpRequestPool::ReleaseRequestToPool(FZenHttpRequest* Request)
 	{
 		for (uint8 i = 0; i < Pool.Num(); ++i)
