DDC: Added EBuildStatus::CacheKey and return EStatus::Ok from a build when no build or cache was requested

#rnx

diff --git a/Engine/Source/Developer/DerivedDataCache/Private/DerivedDataBuildJob.cpp b/Engine/Source/Developer/DerivedDataCache/Private/DerivedDataBuildJob.cpp
--- a/Engine/Source/Developer/DerivedDataCache/Private/DerivedDataBuildJob.cpp
+++ b/Engine/Source/Developer/DerivedDataCache/Private/DerivedDataBuildJob.cpp
@@ -396,6 +396,7 @@ void FBuildJob::CreateContext()
 		Function->Configure(*Context);
 		BuildPolicy = Context->GetBuildPolicy();
 		EnumAddFlags(BuildStatus, ShouldExportBuild() ? EBuildStatus::BuildTryExport : EBuildStatus::None);
+		EnumAddFlags(BuildStatus, EBuildStatus::CacheKey);
 	}
 
 	// Populate the scheduler params with the information that is available now.
@@ -739,7 +740,7 @@ void FBuildJob::EnterExecuteLocal()
 	{
 		if (!EnumHasAnyFlags(BuildPolicy, EBuildPolicy::BuildRemote))
 		{
-			return CompleteWithError(TEXT("Failed because build policy does not allow local or remote execution."_SV));
+			return AdvanceToState(EBuildJobState::Complete);
 		}
 		else if (EnumHasAnyFlags(BuildStatus, EBuildStatus::BuildTryRemote))
 		{
diff --git a/Engine/Source/Developer/DerivedDataCache/Public/DerivedDataBuildTypes.h b/Engine/Source/Developer/DerivedDataCache/Public/DerivedDataBuildTypes.h
--- a/Engine/Source/Developer/DerivedDataCache/Public/DerivedDataBuildTypes.h
+++ b/Engine/Source/Developer/DerivedDataCache/Public/DerivedDataBuildTypes.h
@@ -76,6 +76,9 @@ enum class EBuildStatus : uint32
 	CacheStore      = 1 << 7,
 	/** Valid build output was stored in the cache. */
 	CacheStoreHit   = 1 << 8,
+
+	/** The cache key was calculated. */
+	CacheKey        = 1 << 9,
 };
 
 ENUM_CLASS_FLAGS(EBuildStatus);
