DDC: Reverted the exclusion of DDCUtils from non-editor targets

#rnx

diff --git a/Engine/Source/Developer/DerivedDataCache/DerivedDataCache.Build.cs b/Engine/Source/Developer/DerivedDataCache/DerivedDataCache.Build.cs
--- a/Engine/Source/Developer/DerivedDataCache/DerivedDataCache.Build.cs
+++ b/Engine/Source/Developer/DerivedDataCache/DerivedDataCache.Build.cs
@@ -17,7 +17,7 @@ public class DerivedDataCache : ModuleRules
 
 		// Internal (NotForLicensees) module
 		var DDCUtilsModule = Path.Combine(EngineDirectory, "Restricted/NotForLicensees/Source/Developer/DDCUtils/DDCUtils.Build.cs");
-		if (Target.Type == TargetType.Editor &&  File.Exists(DDCUtilsModule))
+		if (File.Exists(DDCUtilsModule))
 		{
 			DynamicallyLoadedModuleNames.Add("DDCUtils");
 		}
