DDC: Fixed build callback status when the policy allows querying the cache but not building

#rnx

diff --git a/Engine/Source/Developer/DerivedDataCache/Private/DerivedDataBuildJob.cpp b/Engine/Source/Developer/DerivedDataCache/Private/DerivedDataBuildJob.cpp
--- a/Engine/Source/Developer/DerivedDataCache/Private/DerivedDataBuildJob.cpp
+++ b/Engine/Source/Developer/DerivedDataCache/Private/DerivedDataBuildJob.cpp
@@ -456,6 +456,10 @@ void FBuildJob::EndCacheQuery(FCacheGetCompleteParams&& Params)
 			return SetOutputNoCheck(MoveTemp(CacheOutput).Get());
 		}
 	}
+	if (!EnumHasAnyFlags(BuildPolicy, EBuildPolicy::Build))
+	{
+		return CompleteWithError(TEXT("Failed to fetch from the cache and build policy does not allow execution."));
+	}
 	return AdvanceToState(EBuildJobState::ExecuteRemote);
 }
 
