DDC: Added request barriers to the build job constructors

The barriers are needed to handle re-entrant AdvanceToState calls within the constructor of jobs with Blocking priority.

#rnx

diff --git a/Engine/Source/Developer/DerivedDataCache/Private/DerivedDataBuildJob.cpp b/Engine/Source/Developer/DerivedDataCache/Private/DerivedDataBuildJob.cpp
--- a/Engine/Source/Developer/DerivedDataCache/Private/DerivedDataBuildJob.cpp
+++ b/Engine/Source/Developer/DerivedDataCache/Private/DerivedDataBuildJob.cpp
@@ -291,6 +291,7 @@ FBuildJob::FBuildJob(
 	: FBuildJob(Params, WriteToString<64>(TEXT("Resolve: "_SV), InKey), TEXT("Unknown"_SV), MoveTemp(InOnComplete))
 {
 	DefinitionKey = InKey;
+	FRequestBarrier Barrier(Owner, ERequestBarrierFlags::Priority);
 	AdvanceToState(EBuildJobState::ResolveKey);
 }
 
@@ -302,6 +303,7 @@ FBuildJob::FBuildJob(
 {
 	DefinitionKey = InDefinition.GetKey();
 	Definition = InDefinition;
+	FRequestBarrier Barrier(Owner, ERequestBarrierFlags::Priority);
 	AdvanceToState(EBuildJobState::ResolveKey);
 }
 
@@ -314,6 +316,7 @@ FBuildJob::FBuildJob(
 {
 	Action = InAction;
 	Inputs = InInputs;
+	FRequestBarrier Barrier(Owner, ERequestBarrierFlags::Priority);
 	AdvanceToState(EBuildJobState::ResolveKey);
 }
 
