DDC: Implemented existing cache debug options for the new cache stores

#rnx

diff --git a/Engine/Source/Developer/DerivedDataCache/Private/DerivedDataBackends.cpp b/Engine/Source/Developer/DerivedDataCache/Private/DerivedDataBackends.cpp
--- a/Engine/Source/Developer/DerivedDataCache/Private/DerivedDataBackends.cpp
+++ b/Engine/Source/Developer/DerivedDataCache/Private/DerivedDataBackends.cpp
@@ -1307,25 +1307,33 @@ bool FDerivedDataBackendInterface::FBackendDebugOptions::ParseFromTokens(FDerive
 /* Convenience function for backends that check if the key should be missed or not */
 bool FDerivedDataBackendInterface::FBackendDebugOptions::ShouldSimulateMiss(const TCHAR* InCacheKey)
 {
-	bool bDoMiss = false;
-
-	if (SimulateMissTypes.Num() > 0)
+	if (!SimulateMissTypes.IsEmpty())
 	{
-		FString TypeStr = FString(InCacheKey);
-		TypeStr = TypeStr.Left(TypeStr.Find(TEXT("_")));
-
+		FStringView TypeStr(InCacheKey);
+		TypeStr.LeftInline(TypeStr.Find(TEXT("_")));
 		if (SimulateMissTypes.Contains(TypeStr))
 		{
-			bDoMiss = true;
+			return true;
 		}
 	}
 
-	if (!bDoMiss && RandomMissRate > 0)
+	return RandomMissRate > 0 && FMath::RandHelper(100) < RandomMissRate;
+}
+
+/* Convenience function for backends that check if the key should be missed or not */
+bool FDerivedDataBackendInterface::FBackendDebugOptions::ShouldSimulateMiss(const UE::DerivedData::FCacheKey& InCacheKey)
+{
+	if (!SimulateMissTypes.IsEmpty())
 	{
-		bDoMiss = FMath::RandHelper(100) < RandomMissRate;
+		TStringBuilder<256> Type;
+		Type << InCacheKey.Bucket;
+		if (SimulateMissTypes.Contains(Type.ToView()))
+		{
+			return true;
+		}
 	}
 
-	return bDoMiss;
+	return RandomMissRate > 0 && FMath::RandHelper(100) < RandomMissRate;
 }
 
 #undef LOCTEXT_NAMESPACE
diff --git a/Engine/Source/Developer/DerivedDataCache/Private/FileSystemDerivedDataBackend.cpp b/Engine/Source/Developer/DerivedDataCache/Private/FileSystemDerivedDataBackend.cpp
--- a/Engine/Source/Developer/DerivedDataCache/Private/FileSystemDerivedDataBackend.cpp
+++ b/Engine/Source/Developer/DerivedDataCache/Private/FileSystemDerivedDataBackend.cpp
@@ -1003,7 +1003,7 @@ private:
 				[](const FPayload& Payload) { return Payload.GetData().GetRawSize(); }, uint64(0));
 	}
 
-	bool PutCacheRecord(const FCacheRecord& Record, FStringView Context, ECachePolicy Policy) const
+	bool PutCacheRecord(const FCacheRecord& Record, FStringView Context, ECachePolicy Policy)
 	{
 		if (!IsWritable())
 		{
@@ -1022,6 +1022,13 @@ private:
 			return false;
 		}
 
+		if (ShouldSimulateMiss(Key))
+		{
+			UE_LOG(LogDerivedDataCache, Verbose, TEXT("%s: Simulated miss for put of %s from '%.*s'"),
+				*CachePath, *WriteToString<96>(Key), Context.Len(), Context.GetData());
+			return false;
+		}
+
 		// Save the payloads and build the cache object.
 		const TConstArrayView<FPayload> Attachments = Record.GetAttachmentPayloads();
 		TCbWriter<512> Writer;
@@ -1074,7 +1081,7 @@ private:
 		return true;
 	}
 
-	FOptionalCacheRecord GetCacheRecord(const FCacheKey& Key, FStringView Context, ECachePolicy Policy, bool bAlwaysLoadInlineData = false) const
+	FOptionalCacheRecord GetCacheRecord(const FCacheKey& Key, FStringView Context, ECachePolicy Policy, bool bAlwaysLoadInlineData = false)
 	{
 		if (!IsUsable())
 		{
@@ -1091,6 +1098,13 @@ private:
 			return FOptionalCacheRecord();
 		}
 
+		if (ShouldSimulateMiss(Key))
+		{
+			UE_LOG(LogDerivedDataCache, Verbose, TEXT("%s: Simulated miss for get of %s from '%.*s'"),
+				*CachePath, *WriteToString<96>(Key), Context.Len(), Context.GetData());
+			return FOptionalCacheRecord();
+		}
+
 		// Request the record from storage.
 		TStringBuilder<256> Path;
 		BuildCacheObjectPath(Key, Path);
@@ -1499,21 +1513,19 @@ private:
 	/** Keys we ignored due to miss rate settings */
 	FCriticalSection MissedKeysCS;
 	TSet<FName> DebugMissedKeys;
+	TSet<FCacheKey> DebugMissedCacheKeys;
 
-	bool DidSimulateMiss(const TCHAR* InKey)
+	bool ShouldSimulateMiss(const TCHAR* InKey)
 	{
-		if (DebugOptions.RandomMissRate == 0 || DebugOptions.SimulateMissTypes.Num() == 0)
+		if (DebugOptions.RandomMissRate == 0 && DebugOptions.SimulateMissTypes.IsEmpty())
 		{
 			return false;
 		}
-		FScopeLock Lock(&MissedKeysCS);
-		return DebugMissedKeys.Contains(FName(InKey));
-	}
 
-	bool ShouldSimulateMiss(const TCHAR* InKey)
-	{
-		// once missed, always missed
-		if (DidSimulateMiss(InKey))
+		const FName Key(InKey);
+		const uint32 Hash = GetTypeHash(Key);
+
+		if (FScopeLock Lock(&MissedKeysCS); DebugMissedKeys.ContainsByHash(Hash, Key))
 		{
 			return true;
 		}
@@ -1522,7 +1534,31 @@ private:
 		{
 			FScopeLock Lock(&MissedKeysCS);
 			UE_LOG(LogDerivedDataCache, Verbose, TEXT("Simulating miss in %s for %s"), *GetName(), InKey);
-			DebugMissedKeys.Add(FName(InKey));
+			DebugMissedKeys.AddByHash(Hash, Key);
+			return true;
+		}
+
+		return false;
+	}
+
+	bool ShouldSimulateMiss(const FCacheKey& Key)
+	{
+		if (DebugOptions.RandomMissRate == 0 && DebugOptions.SimulateMissTypes.IsEmpty())
+		{
+			return false;
+		}
+
+		const uint32 Hash = GetTypeHash(Key);
+
+		if (FScopeLock Lock(&MissedKeysCS); DebugMissedCacheKeys.ContainsByHash(Hash, Key))
+		{
+			return true;
+		}
+
+		if (DebugOptions.ShouldSimulateMiss(Key))
+		{
+			FScopeLock Lock(&MissedKeysCS);
+			DebugMissedCacheKeys.AddByHash(Hash, Key);
 			return true;
 		}
 
diff --git a/Engine/Source/Developer/DerivedDataCache/Private/MemoryDerivedDataBackend.cpp b/Engine/Source/Developer/DerivedDataCache/Private/MemoryDerivedDataBackend.cpp
--- a/Engine/Source/Developer/DerivedDataCache/Private/MemoryDerivedDataBackend.cpp
+++ b/Engine/Source/Developer/DerivedDataCache/Private/MemoryDerivedDataBackend.cpp
@@ -114,7 +114,7 @@ FDerivedDataBackendInterface::EPutStatus FMemoryDerivedDataBackend::PutCachedDat
 	COOK_STAT(auto Timer = UsageStats.TimePut());
 	FScopeLock ScopeLock(&SynchronizationObject);
 
-	if (DidSimulateMiss(CacheKey))
+	if (ShouldSimulateMiss(CacheKey))
 	{
 		return EPutStatus::Skipped;
 	}
@@ -354,20 +354,17 @@ bool FMemoryDerivedDataBackend::ApplyDebugOptions(FBackendDebugOptions& InOption
 	return true;
 }
 
-bool FMemoryDerivedDataBackend::DidSimulateMiss(const TCHAR* InKey)
+bool FMemoryDerivedDataBackend::ShouldSimulateMiss(const TCHAR* InKey)
 {
-	if (DebugOptions.RandomMissRate == 0 || DebugOptions.SimulateMissTypes.Num() == 0)
+	if (DebugOptions.RandomMissRate == 0 && DebugOptions.SimulateMissTypes.IsEmpty())
 	{
 		return false;
 	}
-	FScopeLock Lock(&MissedKeysCS);
-	return DebugMissedKeys.Contains(FName(InKey));
-}
 
-bool FMemoryDerivedDataBackend::ShouldSimulateMiss(const TCHAR* InKey)
-{
-	// once missed, always missed
-	if (DidSimulateMiss(InKey))
+	const FName Key(InKey);
+	const uint32 Hash = GetTypeHash(Key);
+
+	if (FScopeLock Lock(&MissedKeysCS); DebugMissedKeys.ContainsByHash(Hash, Key))
 	{
 		return true;
 	}
@@ -376,7 +373,31 @@ bool FMemoryDerivedDataBackend::ShouldSimulateMiss(const TCHAR* InKey)
 	{
 		FScopeLock Lock(&MissedKeysCS);
 		UE_LOG(LogDerivedDataCache, Verbose, TEXT("Simulating miss in %s for %s"), *GetName(), InKey);
-		DebugMissedKeys.Add(FName(InKey));
+		DebugMissedKeys.AddByHash(Hash, Key);
+		return true;
+	}
+
+	return false;
+}
+
+bool FMemoryDerivedDataBackend::ShouldSimulateMiss(const FCacheKey& Key)
+{
+	if (DebugOptions.RandomMissRate == 0 && DebugOptions.SimulateMissTypes.IsEmpty())
+	{
+		return false;
+	}
+
+	const uint32 Hash = GetTypeHash(Key);
+
+	if (FScopeLock Lock(&MissedKeysCS); DebugMissedCacheKeys.ContainsByHash(Hash, Key))
+	{
+		return true;
+	}
+
+	if (DebugOptions.ShouldSimulateMiss(Key))
+	{
+		FScopeLock Lock(&MissedKeysCS);
+		DebugMissedCacheKeys.AddByHash(Hash, Key);
 		return true;
 	}
 
@@ -423,6 +444,13 @@ void FMemoryDerivedDataBackend::Put(
 			}
 		};
 
+		if (ShouldSimulateMiss(Key))
+		{
+			UE_LOG(LogDerivedDataCache, Verbose, TEXT("%s: Simulated miss for put of %s from '%.*s'"),
+				*GetName(), *WriteToString<96>(Key), Context.Len(), Context.GetData());
+			continue;
+		}
+
 		const FPayload& Value = Record.GetValuePayload();
 		const TConstArrayView<FPayload> Attachments = Record.GetAttachmentPayloads();
 
@@ -482,7 +510,12 @@ void FMemoryDerivedDataBackend::Get(
 	{
 		COOK_STAT(auto Timer = UsageStats.TimeGet());
 		FOptionalCacheRecord Record;
-		if (FScopeLock ScopeLock(&SynchronizationObject); const FCacheRecord* CacheRecord = CacheRecords.Find(Key))
+		if (ShouldSimulateMiss(Key))
+		{
+			UE_LOG(LogDerivedDataCache, Verbose, TEXT("%s: Simulated miss for get of %s from '%.*s'"),
+				*GetName(), *WriteToString<96>(Key), Context.Len(), Context.GetData());
+		}
+		else if (FScopeLock ScopeLock(&SynchronizationObject); const FCacheRecord* CacheRecord = CacheRecords.Find(Key))
 		{
 			Record = *CacheRecord;
 		}
@@ -515,7 +548,12 @@ void FMemoryDerivedDataBackend::GetPayload(
 	{
 		COOK_STAT(auto Timer = UsageStats.TimeGet());
 		FPayload Payload;
-		if (FScopeLock ScopeLock(&SynchronizationObject); const FCacheRecord* Record = CacheRecords.Find(PayloadKey.CacheKey))
+		if (ShouldSimulateMiss(PayloadKey.CacheKey))
+		{
+			UE_LOG(LogDerivedDataCache, Verbose, TEXT("%s: Simulated miss for get of %s from '%.*s'"),
+				*GetName(), *WriteToString<96>(PayloadKey), Context.Len(), Context.GetData());
+		}
+		else if (FScopeLock ScopeLock(&SynchronizationObject); const FCacheRecord* Record = CacheRecords.Find(PayloadKey.CacheKey))
 		{
 			Payload = Record->GetAttachmentPayload(PayloadKey.Id);
 		}
diff --git a/Engine/Source/Developer/DerivedDataCache/Private/MemoryDerivedDataBackend.h b/Engine/Source/Developer/DerivedDataCache/Private/MemoryDerivedDataBackend.h
--- a/Engine/Source/Developer/DerivedDataCache/Private/MemoryDerivedDataBackend.h
+++ b/Engine/Source/Developer/DerivedDataCache/Private/MemoryDerivedDataBackend.h
@@ -198,8 +198,8 @@ private:
 protected:
 
 	/* Debug helpers */
-	bool DidSimulateMiss(const TCHAR* InKey);
 	bool ShouldSimulateMiss(const TCHAR* InKey);
+	bool ShouldSimulateMiss(const FCacheKey& InKey);
 
 	/** Debug Options */
 	FBackendDebugOptions DebugOptions;
@@ -207,6 +207,7 @@ protected:
 	/** Keys we ignored due to miss rate settings */
 	FCriticalSection MissedKeysCS;
 	TSet<FName> DebugMissedKeys;
+	TSet<FCacheKey> DebugMissedCacheKeys;
 };
 
 } // UE::DerivedData::Backends
diff --git a/Engine/Source/Developer/DerivedDataCache/Private/S3DerivedDataBackend.cpp b/Engine/Source/Developer/DerivedDataCache/Private/S3DerivedDataBackend.cpp
--- a/Engine/Source/Developer/DerivedDataCache/Private/S3DerivedDataBackend.cpp
+++ b/Engine/Source/Developer/DerivedDataCache/Private/S3DerivedDataBackend.cpp
@@ -1149,7 +1149,7 @@ bool FS3DerivedDataBackend::ApplyDebugOptions(FBackendDebugOptions& InOptions)
 
 bool FS3DerivedDataBackend::DidSimulateMiss(const TCHAR* InKey)
 {
-	if (DebugOptions.RandomMissRate == 0 || DebugOptions.SimulateMissTypes.Num() == 0)
+	if (DebugOptions.RandomMissRate == 0 && DebugOptions.SimulateMissTypes.IsEmpty())
 	{
 		return false;
 	}
diff --git a/Engine/Source/Developer/DerivedDataCache/Private/ZenDerivedDataBackend.cpp b/Engine/Source/Developer/DerivedDataCache/Private/ZenDerivedDataBackend.cpp
--- a/Engine/Source/Developer/DerivedDataCache/Private/ZenDerivedDataBackend.cpp
+++ b/Engine/Source/Developer/DerivedDataCache/Private/ZenDerivedDataBackend.cpp
@@ -273,7 +273,7 @@ FZenDerivedDataBackend::GetZenData(const FCacheKey& CacheKey, ECachePolicy Cache
 FDerivedDataBackendInterface::EPutStatus
 FZenDerivedDataBackend::PutCachedData(const TCHAR* CacheKey, TArrayView<const uint8> InData, bool bPutEvenIfExists)
 {
-	if (DidSimulateMiss(CacheKey))
+	if (ShouldSimulateMiss(CacheKey))
 	{
 		return EPutStatus::NotCached;
 	}
@@ -469,20 +469,17 @@ bool FZenDerivedDataBackend::ApplyDebugOptions(FBackendDebugOptions& InOptions)
 	return true;
 }
 
-bool FZenDerivedDataBackend::DidSimulateMiss(const TCHAR* InKey)
+bool FZenDerivedDataBackend::ShouldSimulateMiss(const TCHAR* InKey)
 {
-	if (DebugOptions.RandomMissRate == 0 || DebugOptions.SimulateMissTypes.Num() == 0)
+	if (DebugOptions.RandomMissRate == 0 && DebugOptions.SimulateMissTypes.IsEmpty())
 	{
 		return false;
 	}
-	FScopeLock Lock(&MissedKeysCS);
-	return DebugMissedKeys.Contains(FName(InKey));
-}
 
-bool FZenDerivedDataBackend::ShouldSimulateMiss(const TCHAR* InKey)
-{
-	// once missed, always missed
-	if (DidSimulateMiss(InKey))
+	const FName Key(InKey);
+	const uint32 Hash = GetTypeHash(Key);
+
+	if (FScopeLock Lock(&MissedKeysCS); DebugMissedKeys.ContainsByHash(Hash, Key))
 	{
 		return true;
 	}
@@ -491,7 +488,31 @@ bool FZenDerivedDataBackend::ShouldSimulateMiss(const TCHAR* InKey)
 	{
 		FScopeLock Lock(&MissedKeysCS);
 		UE_LOG(LogDerivedDataCache, Verbose, TEXT("Simulating miss in %s for %s"), *GetName(), InKey);
-		DebugMissedKeys.Add(FName(InKey));
+		DebugMissedKeys.AddByHash(Hash, Key);
+		return true;
+	}
+
+	return false;
+}
+
+bool FZenDerivedDataBackend::ShouldSimulateMiss(const FCacheKey& Key)
+{
+	if (DebugOptions.RandomMissRate == 0 && DebugOptions.SimulateMissTypes.IsEmpty())
+	{
+		return false;
+	}
+
+	const uint32 Hash = GetTypeHash(Key);
+
+	if (FScopeLock Lock(&MissedKeysCS); DebugMissedCacheKeys.ContainsByHash(Hash, Key))
+	{
+		return true;
+	}
+
+	if (DebugOptions.ShouldSimulateMiss(Key))
+	{
+		FScopeLock Lock(&MissedKeysCS);
+		DebugMissedCacheKeys.AddByHash(Hash, Key);
 		return true;
 	}
 
@@ -509,7 +530,13 @@ void FZenDerivedDataBackend::Put(
 	{
 		COOK_STAT(auto Timer = UsageStats.TimePut());
 		bool bResult;
-		if (bCacheRecordEndpointEnabled)
+		if (ShouldSimulateMiss(Record.GetKey()))
+		{
+			UE_LOG(LogDerivedDataCache, Verbose, TEXT("%s: Simulated miss for put of %s from '%.*s'"),
+				*GetName(), *WriteToString<96>(Record.GetKey()), Context.Len(), Context.GetData());
+			bResult = false;
+		}
+		else if (bCacheRecordEndpointEnabled)
 		{
 			bResult = PutCacheRecord(Record, Context, Policy);
 		}
@@ -552,7 +579,12 @@ void FZenDerivedDataBackend::Get(
 		COOK_STAT(auto Timer = UsageStats.TimeGet());
 
 		FOptionalCacheRecord Record;
-		if (bCacheRecordEndpointEnabled)
+		if (ShouldSimulateMiss(Key))
+		{
+			UE_LOG(LogDerivedDataCache, Verbose, TEXT("%s: Simulated miss for get of %s from '%.*s'"),
+				*GetName(), *WriteToString<96>(Key), Context.Len(), Context.GetData());
+		}
+		else if (bCacheRecordEndpointEnabled)
 		{
 			Record = GetCacheRecord(Key, Context, Policy);
 		}
@@ -603,12 +635,22 @@ void FZenDerivedDataBackend::GetPayload(
 			AppendZenUri(Key.CacheKey, Key.Id, QueryUri);
 			AppendPolicyQueryString(Policy, QueryUri);
 
-			TArray64<uint8> PayloadData;
-			EGetResult GetResult = GetZenData(QueryUri, PayloadData);
+			EGetResult GetResult;
 			FCompressedBuffer CompressedBuffer;
-			if (GetResult == EGetResult::Success)
+			if (ShouldSimulateMiss(Key.CacheKey))
 			{
-				CompressedBuffer = FCompressedBuffer::FromCompressed(MakeSharedBufferFromArray(MoveTemp(PayloadData)));
+				UE_LOG(LogDerivedDataCache, Verbose, TEXT("%s: Simulated miss for get of %s from '%.*s'"),
+					*GetName(), *WriteToString<96>(Key), Context.Len(), Context.GetData());
+				GetResult = EGetResult::NotFound;
+			}
+			else
+			{
+				TArray64<uint8> PayloadData;
+				GetResult = GetZenData(QueryUri, PayloadData);
+				if (GetResult == EGetResult::Success)
+				{
+					CompressedBuffer = FCompressedBuffer::FromCompressed(MakeSharedBufferFromArray(MoveTemp(PayloadData)));
+				}
 			}
 			if (CompressedBuffer)
 			{
diff --git a/Engine/Source/Developer/DerivedDataCache/Private/ZenDerivedDataBackend.h b/Engine/Source/Developer/DerivedDataCache/Private/ZenDerivedDataBackend.h
--- a/Engine/Source/Developer/DerivedDataCache/Private/ZenDerivedDataBackend.h
+++ b/Engine/Source/Developer/DerivedDataCache/Private/ZenDerivedDataBackend.h
@@ -155,8 +155,8 @@ private:
 		const FIoHash& CompressedHash, FSharedBuffer&& CompressedData) const;
 
 	/* Debug helpers */
-	bool DidSimulateMiss(const TCHAR* InKey);
 	bool ShouldSimulateMiss(const TCHAR* InKey);
+	bool ShouldSimulateMiss(const FCacheKey& InKey);
 private:
 	FString Domain;
 	FString Namespace;
@@ -173,6 +173,7 @@ private:
 	/** Keys we ignored due to miss rate settings */
 	FCriticalSection MissedKeysCS;
 	TSet<FName> DebugMissedKeys;
+	TSet<FCacheKey> DebugMissedCacheKeys;
 };
 
 } // namespace UE::DerivedData::Backends
diff --git a/Engine/Source/Developer/DerivedDataCache/Public/DerivedDataBackendInterface.h b/Engine/Source/Developer/DerivedDataCache/Public/DerivedDataBackendInterface.h
--- a/Engine/Source/Developer/DerivedDataCache/Public/DerivedDataBackendInterface.h
+++ b/Engine/Source/Developer/DerivedDataCache/Public/DerivedDataBackendInterface.h
@@ -11,6 +11,8 @@
 class FDerivedDataCacheUsageStats;
 class FDerivedDataCacheStatsNode;
 
+namespace UE::DerivedData { struct FCacheKey; }
+
 DECLARE_LOG_CATEGORY_EXTERN(LogDerivedDataCache, Log, All);
 
 DECLARE_DWORD_ACCUMULATOR_STAT_EXTERN(TEXT("Num Gets"),STAT_DDC_NumGets,STATGROUP_DDC, );
@@ -85,6 +87,7 @@ public:
 			each backend.
 		*/
 		bool ShouldSimulateMiss(const TCHAR* InCacheKey);
+		bool ShouldSimulateMiss(const UE::DerivedData::FCacheKey& InCacheKey);
 	};
 
 	virtual ~FDerivedDataBackendInterface()
