DDC: Stopped loading build inputs for local execution when execution is disabled

#rnx

diff --git a/Engine/Source/Developer/DerivedDataCache/Private/DerivedDataBuildJob.cpp b/Engine/Source/Developer/DerivedDataCache/Private/DerivedDataBuildJob.cpp
--- a/Engine/Source/Developer/DerivedDataCache/Private/DerivedDataBuildJob.cpp
+++ b/Engine/Source/Developer/DerivedDataCache/Private/DerivedDataBuildJob.cpp
@@ -594,6 +594,12 @@ void FBuildJob::EnterResolveInputData()
 		{
 			return AdvanceToState(EBuildJobState::ExecuteLocal);
 		}
+		// Skip resolving input data if it will not be used anyway.
+		if (!EnumHasAnyFlags(BuildPolicy, EBuildPolicy::BuildLocal) &&
+			!EnumHasAnyFlags(BuildStatus, EBuildStatus::BuildTryExport))
+		{
+			return AdvanceToState(EBuildJobState::ExecuteLocal);
+		}
 	}
 	else
 	{
