DDC: Changed the heuristic for local cache backends to identify them more accurately

#rnx

diff --git a/Engine/Source/Developer/DerivedDataCache/Private/HierarchicalDerivedDataBackend.h b/Engine/Source/Developer/DerivedDataCache/Private/HierarchicalDerivedDataBackend.h
--- a/Engine/Source/Developer/DerivedDataCache/Private/HierarchicalDerivedDataBackend.h
+++ b/Engine/Source/Developer/DerivedDataCache/Private/HierarchicalDerivedDataBackend.h
@@ -548,13 +548,12 @@ public:
 			for (int32 GetCacheIndex = 0; GetCacheIndex < InnerBackends.Num() && !RemainingKeys.IsEmpty(); ++GetCacheIndex)
 			{
 				// Remove SkipData flags when possibly filling other backends because only complete records can be written.
-				const bool bIsLocalGet = InnerBackends[GetCacheIndex]->GetSpeedClass() == ESpeedClass::Local;
+				const bool bIsLocalGet = InnerBackends[GetCacheIndex]->GetSpeedClass() >= ESpeedClass::Fast;
 				const bool bFill =
 					(bIsLocalGet && bStoreLocalCopy) ||
 					(bIsLocalGet && bStoreRemote) ||
 					(!bIsLocalGet && bStoreLocal) ||
 					(!bIsLocalGet && bStoreRemote && bHasMultipleRemoteBackends);
-				const ECachePolicy FillPolicy = bFill ? (Policy & ~ECachePolicy::SkipData) : Policy;
 
 				// Block on this because backends in this hierarchy are not expected to be asynchronous.
 				FRequestOwner BlockingOwner(EPriority::Blocking);
@@ -571,7 +570,7 @@ public:
 								{
 									if (GetCacheIndex != FillCacheIndex)
 									{
-										const bool bIsLocalFill = InnerBackends[FillCacheIndex]->GetSpeedClass() == ESpeedClass::Local;
+										const bool bIsLocalFill = InnerBackends[FillCacheIndex]->GetSpeedClass() >= ESpeedClass::Fast;
 										if ((bIsLocalFill && bStoreLocal && bIsLocalGet <= bStoreLocalCopy) || (!bIsLocalFill && bStoreRemote))
 										{
 											AsyncPutInnerBackends[FillCacheIndex]->Put({Params.Record}, Context, ECachePolicy::Default, AsyncOwner);
