Fix error in ZenDerivedDataBackend::Put. It needs to make payloads with missing data references to existing data in the CacheStore, and reject the put if any don't exist. This error was causing population of editordomain packages in a Zen store from a shareddatacache to fail.
Change GetZenData to only take TArray64, and to use PerformBlockingHead instead of PerformBlockingDownload if outputdata is null.
#rnx

diff --git a/Engine/Source/Developer/DerivedDataCache/Private/ZenDerivedDataBackend.cpp b/Engine/Source/Developer/DerivedDataCache/Private/ZenDerivedDataBackend.cpp
--- a/Engine/Source/Developer/DerivedDataCache/Private/ZenDerivedDataBackend.cpp
+++ b/Engine/Source/Developer/DerivedDataCache/Private/ZenDerivedDataBackend.cpp
@@ -147,7 +147,10 @@ bool FZenDerivedDataBackend::GetCachedData(const TCHAR* CacheKey, TArray<uint8>&
 
 	double StartTime = FPlatformTime::Seconds();
 
-	EGetResult Result = GetZenData(MakeLegacyZenKey(CacheKey), &OutData);
+	TArray64<uint8> ArrayBuffer;
+	EGetResult Result = GetZenData(MakeLegacyZenKey(CacheKey), &ArrayBuffer);
+	check(ArrayBuffer.Num() <= UINT32_MAX);
+	OutData = MoveTemp(ArrayBuffer);
 	if (Result != EGetResult::Success)
 	{
 		switch (Result)
@@ -174,27 +177,7 @@ bool FZenDerivedDataBackend::GetCachedData(const TCHAR* CacheKey, TArray<uint8>&
 }
 
 FZenDerivedDataBackend::EGetResult
-FZenDerivedDataBackend::GetZenData(FStringView Uri, TArray<uint8>* OutData) const
-{
-	TArray64<uint8> TempArray;
-
-	EGetResult Result = GetZenData(Uri, TempArray);
-
-	if (Result == EGetResult::Success)
-	{
-		if (OutData)
-		{
-			check(TempArray.Num() <= UINT32_MAX);
-
-			*OutData = TempArray;
-		}
-	}
-
-	return Result;
-}
-
-FZenDerivedDataBackend::EGetResult
-FZenDerivedDataBackend::GetZenData(FStringView Uri, TArray64<uint8>& OutData) const
+FZenDerivedDataBackend::GetZenData(FStringView Uri, TArray64<uint8>* OutData) const
 {
 	// Retry request until we get an accepted response or exhaust allowed number of attempts.
 	EGetResult GetResult = EGetResult::NotFound;
@@ -203,7 +186,15 @@ FZenDerivedDataBackend::GetZenData(FStringView Uri, TArray64<uint8>& OutData) co
 		Zen::FZenScopedRequestPtr Request(RequestPool.Get());
 		if (Request.IsValid())
 		{
-			Zen::FZenHttpRequest::Result Result = Request->PerformBlockingDownload(Uri, &OutData);
+			Zen::FZenHttpRequest::Result Result;
+			if (OutData)
+			{
+				Result = Request->PerformBlockingDownload(Uri, OutData);
+			}
+			else
+			{
+				Result = Request->PerformBlockingHead(Uri);
+			}
 			int64 ResponseCode = Request->GetResponseCode();
 
 			// Request was successful, make sure we got all the expected data.
@@ -219,7 +210,10 @@ FZenDerivedDataBackend::GetZenData(FStringView Uri, TArray64<uint8>& OutData) co
 		}
 	}
 
-	OutData.Reset();
+	if (OutData)
+	{
+		OutData->Reset();
+	}
 
 	return GetResult;
 }
@@ -646,7 +640,7 @@ void FZenDerivedDataBackend::GetPayload(
 			else
 			{
 				TArray64<uint8> PayloadData;
-				GetResult = GetZenData(QueryUri, PayloadData);
+				GetResult = GetZenData(QueryUri, &PayloadData);
 				if (GetResult == EGetResult::Success)
 				{
 					CompressedBuffer = FCompressedBuffer::FromCompressed(MakeSharedBufferFromArray(MoveTemp(PayloadData)));
@@ -824,21 +818,32 @@ bool FZenDerivedDataBackend::LegacyPutCachePayload(const FCacheKey& Key, FString
 	const FIoHash& RawHash = Payload.GetRawHash();
 	const FCompositeBuffer CompressedBuffer = Payload.GetData().GetCompressed();
 
-	const bool bStoreInline = !Payload.HasData();
-	if (!bStoreInline)
+	uint64 RawSize = Payload.GetRawSize();
+	if (RawSize != 0)
 	{
 		TStringBuilder<256> Uri;
 		LegacyMakePayloadKey(Key, RawHash, Uri);
-		if (PutZenData(Uri.ToString(), CompressedBuffer, Zen::EContentType::Binary) != FDerivedDataBackendInterface::EPutStatus::Cached)
+		if (Payload.HasData())
 		{
-			return false;
+			if (PutZenData(Uri.ToString(), CompressedBuffer, Zen::EContentType::Binary) != FDerivedDataBackendInterface::EPutStatus::Cached)
+			{
+				return false;
+			}
+		}
+		else
+		{
+			// Put of a record with payloads-with-no-data should succeed iff the payload data is available in the CacheStore
+			if (GetZenData(Uri.ToString(), nullptr) != EGetResult::Success)
+			{
+				return false;
+			}
 		}
 	}
 
 	Writer.BeginObject();
 	Writer.AddObjectId("Id"_ASV, Payload.GetId());
-	Writer.AddInteger("RawSize"_ASV, Payload.GetRawSize());
-	if (bStoreInline)
+	Writer.AddInteger("RawSize"_ASV, RawSize);
+	if (RawSize == 0)
 	{
 		Writer.AddHash("RawHash"_ASV, RawHash);
 		Writer.AddHash("CompressedHash"_ASV, FIoHash::HashBuffer(CompressedBuffer));
@@ -857,7 +862,7 @@ FOptionalCacheRecord FZenDerivedDataBackend::LegacyGetCacheRecord(const FCacheKe
 {
 	TStringBuilder<256> Uri;
 	LegacyMakeZenKey(Key, Uri);
-	TArray<uint8> Buffer;
+	TArray64<uint8> Buffer;
 	EGetResult Result = GetZenData(Uri.ToString(), &Buffer);
 	if (Result != EGetResult::Success)
 	{
@@ -887,7 +892,7 @@ FPayload FZenDerivedDataBackend::LegacyGetCachePayload(const FCacheKey& Key, FSt
 
 	bool bSkipData = EnumHasAllFlags(Policy, ECachePolicy::SkipData);
 	EGetResult GetResult;
-	TArray<uint8> ArrayBuffer;
+	TArray64<uint8> ArrayBuffer;
 	GetResult = GetZenData(Uri.ToString(), bSkipData ? nullptr : &ArrayBuffer);
 	if (GetResult != EGetResult::Success)
 	{
diff --git a/Engine/Source/Developer/DerivedDataCache/Private/ZenDerivedDataBackend.h b/Engine/Source/Developer/DerivedDataCache/Private/ZenDerivedDataBackend.h
--- a/Engine/Source/Developer/DerivedDataCache/Private/ZenDerivedDataBackend.h
+++ b/Engine/Source/Developer/DerivedDataCache/Private/ZenDerivedDataBackend.h
@@ -119,8 +119,7 @@ private:
 		NotFound,
 		Corrupted
 	};
-	EGetResult GetZenData(FStringView Uri, TArray64<uint8>& OutData) const;
-	EGetResult GetZenData(FStringView Uri, TArray<uint8>* OutData) const;
+	EGetResult GetZenData(FStringView Uri, TArray64<uint8>* OutData) const;
 
 	// TODO: need ability to specify content type
 	FDerivedDataBackendInterface::EPutStatus PutZenData(const TCHAR* Uri, const FCompositeBuffer& InData, Zen::EContentType ContentType);
