DDC: Moved FPayloadId into a dedicated header

#rnx

diff --git a/Engine/Source/Developer/DerivedDataCache/Private/DerivedDataBackendAsyncPutWrapper.cpp b/Engine/Source/Developer/DerivedDataCache/Private/DerivedDataBackendAsyncPutWrapper.cpp
--- a/Engine/Source/Developer/DerivedDataCache/Private/DerivedDataBackendAsyncPutWrapper.cpp
+++ b/Engine/Source/Developer/DerivedDataCache/Private/DerivedDataBackendAsyncPutWrapper.cpp
@@ -2,6 +2,7 @@
 
 #include "DerivedDataBackendAsyncPutWrapper.h"
 
+#include "DerivedDataPayload.h"
 #include "DerivedDataRequest.h"
 #include "DerivedDataRequestOwner.h"
 #include "MemoryDerivedDataBackend.h"
diff --git a/Engine/Source/Developer/DerivedDataCache/Private/DerivedDataCacheRecord.cpp b/Engine/Source/Developer/DerivedDataCache/Private/DerivedDataCacheRecord.cpp
--- a/Engine/Source/Developer/DerivedDataCache/Private/DerivedDataCacheRecord.cpp
+++ b/Engine/Source/Developer/DerivedDataCache/Private/DerivedDataCacheRecord.cpp
@@ -4,6 +4,7 @@
 
 #include "DerivedDataCacheKey.h"
 #include "DerivedDataCachePrivate.h"
+#include "DerivedDataPayload.h"
 #include "DerivedDataPayloadPrivate.h"
 #include "HAL/CriticalSection.h"
 #include "Misc/ScopeExit.h"
diff --git a/Engine/Source/Developer/DerivedDataCache/Private/HierarchicalDerivedDataBackend.h b/Engine/Source/Developer/DerivedDataCache/Private/HierarchicalDerivedDataBackend.h
--- a/Engine/Source/Developer/DerivedDataCache/Private/HierarchicalDerivedDataBackend.h
+++ b/Engine/Source/Developer/DerivedDataCache/Private/HierarchicalDerivedDataBackend.h
@@ -2,16 +2,17 @@
 
 #pragma once
 
-#include "CoreMinimal.h"
+#include "CoreTypes.h"
+#include "Containers/ArrayView.h"
+#include "DerivedDataBackendAsyncPutWrapper.h"
 #include "DerivedDataBackendInterface.h"
-#include "ProfilingDebugging/CookStats.h"
 #include "DerivedDataCacheRecord.h"
 #include "DerivedDataCacheUsageStats.h"
-#include "DerivedDataBackendAsyncPutWrapper.h"
+#include "DerivedDataPayload.h"
 #include "DerivedDataRequestOwner.h"
 #include "Misc/ScopeRWLock.h"
+#include "ProfilingDebugging/CookStats.h"
 #include "Templates/UniquePtr.h"
-#include "Containers/ArrayView.h"
 
 extern bool GVerifyDDC;
 
diff --git a/Engine/Source/Developer/DerivedDataCache/Public/DerivedDataBuildKey.h b/Engine/Source/Developer/DerivedDataCache/Public/DerivedDataBuildKey.h
--- a/Engine/Source/Developer/DerivedDataCache/Public/DerivedDataBuildKey.h
+++ b/Engine/Source/Developer/DerivedDataCache/Public/DerivedDataBuildKey.h
@@ -4,7 +4,7 @@
 
 #include "CoreTypes.h"
 #include "Containers/StringFwd.h"
-#include "DerivedDataPayload.h"
+#include "DerivedDataPayloadId.h"
 #include "IO/IoHash.h"
 #include "Templates/TypeHash.h"
 
diff --git a/Engine/Source/Developer/DerivedDataCache/Public/DerivedDataCacheKey.h b/Engine/Source/Developer/DerivedDataCache/Public/DerivedDataCacheKey.h
--- a/Engine/Source/Developer/DerivedDataCache/Public/DerivedDataCacheKey.h
+++ b/Engine/Source/Developer/DerivedDataCache/Public/DerivedDataCacheKey.h
@@ -5,7 +5,7 @@
 #include "CoreTypes.h"
 #include "Containers/StringFwd.h"
 #include "Containers/StringView.h"
-#include "DerivedDataPayload.h"
+#include "DerivedDataPayloadId.h"
 #include "IO/IoHash.h"
 #include "Templates/TypeHash.h"
 
diff --git a/Engine/Source/Developer/DerivedDataCache/Public/DerivedDataCacheRecord.h b/Engine/Source/Developer/DerivedDataCache/Public/DerivedDataCacheRecord.h
--- a/Engine/Source/Developer/DerivedDataCache/Public/DerivedDataCacheRecord.h
+++ b/Engine/Source/Developer/DerivedDataCache/Public/DerivedDataCacheRecord.h
@@ -3,7 +3,7 @@
 #pragma once
 
 #include "DerivedDataCacheKey.h"
-#include "DerivedDataPayload.h"
+#include "DerivedDataPayloadId.h"
 #include "Memory/SharedBuffer.h"
 #include "Templates/Function.h"
 #include "Templates/RefCounting.h"
@@ -16,6 +16,7 @@ class FCbPackage;
 
 namespace UE::DerivedData { class FCacheRecord; }
 namespace UE::DerivedData { class FOptionalCacheRecord; }
+namespace UE::DerivedData { class FPayload; }
 namespace UE::DerivedData { class IRequestOwner; }
 namespace UE::DerivedData { using FOnCacheRecordComplete = TUniqueFunction<void (FCacheRecord&& Record)>; }
 
diff --git a/Engine/Source/Developer/DerivedDataCache/Public/DerivedDataPayload.h b/Engine/Source/Developer/DerivedDataCache/Public/DerivedDataPayload.h
--- a/Engine/Source/Developer/DerivedDataCache/Public/DerivedDataPayload.h
+++ b/Engine/Source/Developer/DerivedDataCache/Public/DerivedDataPayload.h
@@ -4,75 +4,12 @@
 
 #include "CoreTypes.h"
 #include "Compression/CompressedBuffer.h"
-#include "Containers/StringFwd.h"
-#include "Containers/StringView.h"
+#include "DerivedDataPayloadId.h"
 #include "IO/IoHash.h"
-#include "Memory/MemoryView.h"
-#include "String/BytesToHex.h"
-
-#define UE_API DERIVEDDATACACHE_API
-
-class FCbObjectId;
 
 namespace UE::DerivedData
 {
 
-///////////////////////////////////////////////////////////////////////////////////////////////////////////////////////
-
-/** A 12-byte value that uniquely identifies a payload in the context that it was created. */
-struct FPayloadId
-{
-public:
-	using ByteArray = uint8[12];
-
-	/** Construct a null ID. */
-	FPayloadId() = default;
-
-	/** Construct an ID from an array of 12 bytes. */
-	inline explicit FPayloadId(const ByteArray& Id);
-
-	/** Construct an ID from a view of 12 bytes. */
-	UE_API explicit FPayloadId(FMemoryView Id);
-
-	/** Construct an ID from a Compact Binary Object ID. */
-	UE_API FPayloadId(const FCbObjectId& Id);
-
-	/** Returns the ID as a Compact Binary Object ID. */
-	UE_API operator FCbObjectId() const;
-
-	/** Construct an ID from a non-zero hash. */
-	[[nodiscard]] UE_API static FPayloadId FromHash(const FIoHash& Hash);
-
-	/** Construct an ID from a non-empty name. */
-	[[nodiscard]] UE_API static FPayloadId FromName(FAnsiStringView Name);
-	[[nodiscard]] UE_API static FPayloadId FromName(FWideStringView Name);
-
-	/** Returns a reference to the raw byte array for the ID. */
-	inline const ByteArray& GetBytes() const { return Bytes; }
-	inline operator const ByteArray&() const { return Bytes; }
-
-	/** Returns a view of the raw byte array for the ID. */
-	inline FMemoryView GetView() const { return MakeMemoryView(Bytes); }
-
-	/** Whether this is null. */
-	inline bool IsNull() const;
-	/** Whether this is not null. */
-	inline bool IsValid() const { return !IsNull(); }
-
-	/** Reset this to null. */
-	inline void Reset() { *this = FPayloadId(); }
-
-	/** A null ID. */
-	static const FPayloadId Null;
-
-private:
-	alignas(uint32) ByteArray Bytes{};
-};
-
-inline const FPayloadId FPayloadId::Null;
-
-///////////////////////////////////////////////////////////////////////////////////////////////////////////////////////
-
 /**
  * A payload is described by an ID and by the hash and size of its raw buffer (uncompressed).
  *
@@ -132,15 +69,6 @@ private:
 
 inline const FPayload FPayload::Null;
 
-///////////////////////////////////////////////////////////////////////////////////////////////////////////////////////
-
-inline FPayloadId::FPayloadId(const ByteArray& Id)
-{
-	FMemory::Memcpy(Bytes, Id, sizeof(ByteArray));
-}
-
-///////////////////////////////////////////////////////////////////////////////////////////////////////////////////////
-
 inline FPayload::FPayload(const FPayloadId& InId)
 	: Id(InId)
 {
@@ -173,43 +101,6 @@ inline FPayload::FPayload(const FPayloadId& InId, FCompressedBuffer&& InData)
 	checkf(Id.IsValid(), TEXT("A valid ID is required to construct a payload."));
 }
 
-///////////////////////////////////////////////////////////////////////////////////////////////////////////////////////
-
-inline bool operator==(const FPayloadId& A, const FPayloadId& B)
-{
-	return A.GetView().EqualBytes(B.GetView());
-}
-
-inline bool operator!=(const FPayloadId& A, const FPayloadId& B)
-{
-	return !A.GetView().EqualBytes(B.GetView());
-}
-
-inline bool operator<(const FPayloadId& A, const FPayloadId& B)
-{
-	return A.GetView().CompareBytes(B.GetView()) < 0;
-}
-
-inline uint32 GetTypeHash(const FPayloadId& Id)
-{
-	return *reinterpret_cast<const uint32*>(Id.GetView().GetData());
-}
-
-/** Convert the ID to a 24-character hex string. */
-template <typename CharType>
-inline TStringBuilderBase<CharType>& operator<<(TStringBuilderBase<CharType>& Builder, const FPayloadId& Id)
-{
-	UE::String::BytesToHexLower(Id.GetBytes(), Builder);
-	return Builder;
-}
-
-inline bool FPayloadId::IsNull() const
-{
-	return *this == FPayloadId();
-}
-
-///////////////////////////////////////////////////////////////////////////////////////////////////////////////////////
-
 /** Compare payloads by their ID and the hash and size of their raw buffer. */
 inline bool operator==(const FPayload& A, const FPayload& B)
 {
@@ -226,8 +117,4 @@ inline bool operator<(const FPayload& A, const FPayload& B)
 	return !(IdA == IdB) ? IdA < IdB : !(HashA == HashB) ? HashA < HashB : A.GetRawSize() < B.GetRawSize();
 }
 
-///////////////////////////////////////////////////////////////////////////////////////////////////////////////////////
-
 } // UE::DerivedData
-
-#undef UE_API
diff --git a/Engine/Source/Developer/DerivedDataCache/Public/DerivedDataPayloadId.h b/Engine/Source/Developer/DerivedDataCache/Public/DerivedDataPayloadId.h
new file mode 100644
--- /dev/null
+++ b/Engine/Source/Developer/DerivedDataCache/Public/DerivedDataPayloadId.h
@@ -0,0 +1,110 @@
+// Copyright Epic Games, Inc. All Rights Reserved.
+
+#pragma once
+
+#include "CoreTypes.h"
+#include "Containers/StringFwd.h"
+#include "Memory/MemoryView.h"
+#include "String/BytesToHex.h"
+
+#define UE_API DERIVEDDATACACHE_API
+
+class FCbObjectId;
+struct FIoHash;
+
+namespace UE::DerivedData
+{
+
+/** A 12-byte value that uniquely identifies a payload in the context that it was created. */
+struct FPayloadId
+{
+public:
+	using ByteArray = uint8[12];
+
+	/** Construct a null ID. */
+	FPayloadId() = default;
+
+	/** Construct an ID from an array of 12 bytes. */
+	inline explicit FPayloadId(const ByteArray& Id);
+
+	/** Construct an ID from a view of 12 bytes. */
+	UE_API explicit FPayloadId(FMemoryView Id);
+
+	/** Construct an ID from a Compact Binary Object ID. */
+	UE_API FPayloadId(const FCbObjectId& Id);
+
+	/** Returns the ID as a Compact Binary Object ID. */
+	UE_API operator FCbObjectId() const;
+
+	/** Construct an ID from a non-zero hash. */
+	[[nodiscard]] UE_API static FPayloadId FromHash(const FIoHash& Hash);
+
+	/** Construct an ID from a non-empty name. */
+	[[nodiscard]] UE_API static FPayloadId FromName(FAnsiStringView Name);
+	[[nodiscard]] UE_API static FPayloadId FromName(FWideStringView Name);
+
+	/** Returns a reference to the raw byte array for the ID. */
+	inline const ByteArray& GetBytes() const { return Bytes; }
+	inline operator const ByteArray&() const { return Bytes; }
+
+	/** Returns a view of the raw byte array for the ID. */
+	inline FMemoryView GetView() const { return MakeMemoryView(Bytes); }
+
+	/** Whether this is null. */
+	inline bool IsNull() const;
+	/** Whether this is not null. */
+	inline bool IsValid() const { return !IsNull(); }
+
+	/** Reset this to null. */
+	inline void Reset() { *this = FPayloadId(); }
+
+	/** A null ID. */
+	static const FPayloadId Null;
+
+private:
+	alignas(uint32) ByteArray Bytes{};
+};
+
+inline const FPayloadId FPayloadId::Null;
+
+inline FPayloadId::FPayloadId(const ByteArray& Id)
+{
+	FMemory::Memcpy(Bytes, Id, sizeof(ByteArray));
+}
+
+inline bool operator==(const FPayloadId& A, const FPayloadId& B)
+{
+	return A.GetView().EqualBytes(B.GetView());
+}
+
+inline bool operator!=(const FPayloadId& A, const FPayloadId& B)
+{
+	return !A.GetView().EqualBytes(B.GetView());
+}
+
+inline bool operator<(const FPayloadId& A, const FPayloadId& B)
+{
+	return A.GetView().CompareBytes(B.GetView()) < 0;
+}
+
+inline uint32 GetTypeHash(const FPayloadId& Id)
+{
+	return *reinterpret_cast<const uint32*>(Id.GetView().GetData());
+}
+
+/** Convert the ID to a 24-character hex string. */
+template <typename CharType>
+inline TStringBuilderBase<CharType>& operator<<(TStringBuilderBase<CharType>& Builder, const FPayloadId& Id)
+{
+	UE::String::BytesToHexLower(Id.GetBytes(), Builder);
+	return Builder;
+}
+
+inline bool FPayloadId::IsNull() const
+{
+	return *this == FPayloadId();
+}
+
+} // UE::DerivedData
+
+#undef UE_API
diff --git a/Engine/Source/Developer/DesktopPlatform/Private/TargetReceiptBuildWorker.cpp b/Engine/Source/Developer/DesktopPlatform/Private/TargetReceiptBuildWorker.cpp
--- a/Engine/Source/Developer/DesktopPlatform/Private/TargetReceiptBuildWorker.cpp
+++ b/Engine/Source/Developer/DesktopPlatform/Private/TargetReceiptBuildWorker.cpp
@@ -2,6 +2,7 @@
 
 #include "TargetReceiptBuildWorker.h"
 
+#include "Compression/CompressedBuffer.h"
 #include "Compression/OodleDataCompression.h"
 #include "DerivedDataBuildWorker.h"
 #include "Features/IModularFeatures.h"
diff --git a/Engine/Source/Developer/TextureBuild/Private/TextureBuildFunction.cpp b/Engine/Source/Developer/TextureBuild/Private/TextureBuildFunction.cpp
--- a/Engine/Source/Developer/TextureBuild/Private/TextureBuildFunction.cpp
+++ b/Engine/Source/Developer/TextureBuild/Private/TextureBuildFunction.cpp
@@ -3,7 +3,7 @@
 #include "TextureBuildFunction.h"
 
 #include "DerivedDataCache.h"
-#include "DerivedDataPayload.h"
+#include "DerivedDataPayloadId.h"
 #include "Engine/TextureDefines.h"
 #include "IImageWrapper.h"
 #include "IImageWrapperModule.h"
diff --git a/Engine/Source/Editor/UnrealEd/Private/DerivedDataBuildLocalExecutor.cpp b/Engine/Source/Editor/UnrealEd/Private/DerivedDataBuildLocalExecutor.cpp
--- a/Engine/Source/Editor/UnrealEd/Private/DerivedDataBuildLocalExecutor.cpp
+++ b/Engine/Source/Editor/UnrealEd/Private/DerivedDataBuildLocalExecutor.cpp
@@ -7,6 +7,7 @@
 #include "DerivedDataBuildInputs.h"
 #include "DerivedDataBuildOutput.h"
 #include "DerivedDataBuildWorker.h"
+#include "DerivedDataPayload.h"
 #include "DerivedDataRequest.h"
 #include "DerivedDataRequestOwner.h"
 #include "Features/IModularFeatures.h"
diff --git a/Engine/Source/Editor/UnrealEd/Private/DerivedDataBuildRemoteExecutor.cpp b/Engine/Source/Editor/UnrealEd/Private/DerivedDataBuildRemoteExecutor.cpp
--- a/Engine/Source/Editor/UnrealEd/Private/DerivedDataBuildRemoteExecutor.cpp
+++ b/Engine/Source/Editor/UnrealEd/Private/DerivedDataBuildRemoteExecutor.cpp
@@ -8,6 +8,7 @@
 #include "DerivedDataBuildInputs.h"
 #include "DerivedDataBuildOutput.h"
 #include "DerivedDataBuildWorker.h"
+#include "DerivedDataPayload.h"
 #include "DerivedDataRequest.h"
 #include "DerivedDataRequestOwner.h"
 #include "Features/IModularFeatures.h"
diff --git a/Engine/Source/Editor/UnrealEd/Private/EditorDomain/EditorDomainArchive.cpp b/Engine/Source/Editor/UnrealEd/Private/EditorDomain/EditorDomainArchive.cpp
--- a/Engine/Source/Editor/UnrealEd/Private/EditorDomain/EditorDomainArchive.cpp
+++ b/Engine/Source/Editor/UnrealEd/Private/EditorDomain/EditorDomainArchive.cpp
@@ -7,6 +7,7 @@
 #include "DerivedDataCache.h"
 #include "DerivedDataCacheKey.h"
 #include "DerivedDataCacheRecord.h"
+#include "DerivedDataPayload.h"
 #include "EditorDomain/EditorDomainSave.h"
 #include "EditorDomain/EditorDomainUtils.h"
 #include "HAL/UnrealMemory.h"
diff --git a/Engine/Source/Editor/UnrealEd/Private/EditorDomain/EditorDomainArchive.h b/Engine/Source/Editor/UnrealEd/Private/EditorDomain/EditorDomainArchive.h
--- a/Engine/Source/Editor/UnrealEd/Private/EditorDomain/EditorDomainArchive.h
+++ b/Engine/Source/Editor/UnrealEd/Private/EditorDomain/EditorDomainArchive.h
@@ -5,7 +5,7 @@
 #include "Async/AsyncFileHandle.h"
 #include "Containers/Array.h"
 #include "DerivedDataCache.h"
-#include "DerivedDataPayload.h"
+#include "DerivedDataPayloadId.h"
 #include "DerivedDataRequestOwner.h"
 #include "EditorDomain/EditorDomain.h"
 #include "HAL/CriticalSection.h"
diff --git a/Engine/Source/Runtime/Engine/Private/Streaming/Texture2DMipDataProvider_DDC.cpp b/Engine/Source/Runtime/Engine/Private/Streaming/Texture2DMipDataProvider_DDC.cpp
--- a/Engine/Source/Runtime/Engine/Private/Streaming/Texture2DMipDataProvider_DDC.cpp
+++ b/Engine/Source/Runtime/Engine/Private/Streaming/Texture2DMipDataProvider_DDC.cpp
@@ -16,6 +16,7 @@ Texture2DMipDataProvider_DDC.cpp : Implementation of FTextureMipDataProvider usi
 #include "DerivedDataCache.h"
 #include "DerivedDataCacheInterface.h"
 #include "DerivedDataCacheKey.h"
+#include "DerivedDataPayload.h"
 
 FTexture2DMipDataProvider_DDC::FTexture2DMipDataProvider_DDC(const UTexture* Texture)
 	: FTextureMipDataProvider(Texture, ETickState::Init, ETickThread::Async)
diff --git a/Engine/Source/Runtime/Engine/Private/Streaming/Texture2DStreamIn_DDC.cpp b/Engine/Source/Runtime/Engine/Private/Streaming/Texture2DStreamIn_DDC.cpp
--- a/Engine/Source/Runtime/Engine/Private/Streaming/Texture2DStreamIn_DDC.cpp
+++ b/Engine/Source/Runtime/Engine/Private/Streaming/Texture2DStreamIn_DDC.cpp
@@ -14,6 +14,7 @@ Texture2DStreamIn_DDC.cpp: Stream in helper for 2D textures loading DDC files.
 #include "DerivedDataCache.h"
 #include "DerivedDataCacheInterface.h"
 #include "DerivedDataCacheKey.h"
+#include "DerivedDataPayload.h"
 #include "DerivedDataRequestOwner.h"
 
 int32 GStreamingUseAsyncRequestsForDDC = 1;
diff --git a/Engine/Source/Runtime/Engine/Private/TextureDerivedData.cpp b/Engine/Source/Runtime/Engine/Private/TextureDerivedData.cpp
--- a/Engine/Source/Runtime/Engine/Private/TextureDerivedData.cpp
+++ b/Engine/Source/Runtime/Engine/Private/TextureDerivedData.cpp
@@ -35,6 +35,7 @@
 #include "DerivedDataCache.h"
 #include "DerivedDataCacheInterface.h"
 #include "DerivedDataCacheKey.h"
+#include "DerivedDataPayload.h"
 #include "DerivedDataRequestOwner.h"
 #include "DerivedDataRequestTypes.h"
 #include "Interfaces/ITargetPlatform.h"
