Add DDC2 memory scheduler to limit peak memory consumption

Also:
* Implemented rough memory usage estimate for texture cooking
* Fixed build action -> constant copying bug

diff --git a/Engine/Source/Developer/DerivedDataCache/Private/DerivedDataBuildJob.cpp b/Engine/Source/Developer/DerivedDataCache/Private/DerivedDataBuildJob.cpp
--- a/Engine/Source/Developer/DerivedDataCache/Private/DerivedDataBuildJob.cpp
+++ b/Engine/Source/Developer/DerivedDataCache/Private/DerivedDataBuildJob.cpp
@@ -111,7 +111,7 @@ public:
 	inline ICache& GetCache() const final { return Cache; }
 	inline IBuild& GetBuild() const final { return BuildSystem; }
 
-	void Schedule() final;
+	void StepExecution() final;
 
 private:
 	FBuildJob(const FBuildJobCreateParams& Params, FStringView Name, FStringView FunctionName, FOnBuildJobComplete&& OnComplete);
@@ -210,7 +210,7 @@ private:
 	/** Available in [CacheQuery, CacheStoreWait). Context used by IBuildFunction. */
 	TRefCountPtr<FBuildJobContext> Context;
 
-	FBuildSchedulerParams SchedulerParams;
+	TUniquePtr<IBuildJobSchedule> Schedule;
 
 	/** Lock to synchronize writes to State, NextState, bInAdvanceToState. */
 	FRWLock Lock;
@@ -330,7 +330,7 @@ FBuildJob::~FBuildJob()
 		LexToString(State), *Name, *FunctionName);
 }
 
-void FBuildJob::Schedule()
+void FBuildJob::StepExecution()
 {
 	switch (State)
 	{
@@ -354,12 +354,13 @@ void FBuildJob::Schedule()
 
 void FBuildJob::BeginJob()
 {
-	Scheduler.BeginJob(this);
+	check(!Schedule);
+	Schedule = Scheduler.BeginJob(*this, Owner);
 }
 
 void FBuildJob::EndJob()
 {
-	Scheduler.EndJob(this);
+	Schedule->EndJob();
 	delete this;
 }
 
@@ -396,6 +397,10 @@ void FBuildJob::CreateContext()
 	{
 		const FCacheKey CacheKey{FCacheBucket(FunctionName), Action.Get().GetKey().Hash};
 		Context = new FBuildJobContext(*this, CacheKey, *Function, OutputBuilder, BuildPolicy);
+		Action.Get().IterateConstants([this](FStringView Key, FCbObject&& Value)
+		{
+			Context->AddConstant(Key, MoveTemp(Value));
+		});
 		Function->Configure(*Context);
 		BuildPolicy = Context->GetBuildPolicy();
 		EnumAddFlags(BuildStatus, ShouldExportBuild() ? EBuildStatus::BuildTryExport : EBuildStatus::None);
@@ -403,20 +408,21 @@ void FBuildJob::CreateContext()
 	}
 
 	// Populate the scheduler params with the information that is available now.
+	FBuildSchedulerParams& SchedulerParams = Schedule->EditParameters();
 	SchedulerParams.Key = Action.Get().GetKey();
-	Action.Get().IterateConstants([this](FStringView Key, FCbObject&& Value)
+	Action.Get().IterateConstants([&SchedulerParams](FStringView Key, FCbObject&& Value)
 	{
 		const uint64 ValueSize = Value.GetSize();
 		SchedulerParams.TotalInputsSize += ValueSize;
 		SchedulerParams.ResolvedInputsSize += ValueSize;
 	});
-	Action.Get().IterateInputs([this](FStringView Key, const FIoHash& RawHash, uint64 RawSize)
+	Action.Get().IterateInputs([&SchedulerParams](FStringView Key, const FIoHash& RawHash, uint64 RawSize)
 	{
 		SchedulerParams.TotalInputsSize += RawSize;
 	});
 	if (Inputs)
 	{
-		Inputs.Get().IterateInputs([this](FStringView Key, const FCompressedBuffer& Buffer)
+		Inputs.Get().IterateInputs([&SchedulerParams](FStringView Key, const FCompressedBuffer& Buffer)
 		{
 			SchedulerParams.ResolvedInputsSize += Buffer.GetRawSize();
 		});
@@ -577,12 +583,13 @@ void FBuildJob::EndResolveInputMeta(FBuildInputMetaResolvedParams&& Params)
 
 void FBuildJob::EnterResolveInputData()
 {
+	FBuildSchedulerParams& SchedulerParams = Schedule->EditParameters();
 	SchedulerParams.MissingLocalInputsSize = 0;
 	SchedulerParams.MissingRemoteInputsSize = 0;
 	if (State == EBuildJobState::ResolveInputData)
 	{
 		MissingInputs.Reset();
-		Action.Get().IterateInputs([this](FStringView Key, const FIoHash& RawHash, uint64 RawSize)
+		Action.Get().IterateInputs([this, &SchedulerParams](FStringView Key, const FIoHash& RawHash, uint64 RawSize)
 		{
 			if (!Inputs || Inputs.Get().FindInput(Key).IsNull())
 			{
@@ -603,7 +610,7 @@ void FBuildJob::EnterResolveInputData()
 	}
 	else
 	{
-		Action.Get().IterateInputs([this](FStringView Key, const FIoHash& RawHash, uint64 RawSize)
+		Action.Get().IterateInputs([this, &SchedulerParams](FStringView Key, const FIoHash& RawHash, uint64 RawSize)
 		{
 			if (MissingInputs.FindByKey(Key))
 			{
@@ -643,6 +650,7 @@ void FBuildJob::EndResolveInputData(FBuildInputDataResolvedParams&& Params)
 {
 	if (Params.Status == EStatus::Ok)
 	{
+		FBuildSchedulerParams& SchedulerParams = Schedule->EditParameters();
 		FBuildInputsBuilder Builder = BuildSystem.CreateInputs(Name);
 		for (const FBuildInputDataByKey& Input : Params.Inputs)
 		{
@@ -766,10 +774,7 @@ void FBuildJob::EnterExecuteLocal()
 				TEXT("and remote execution was not available."_SV));
 		}
 	}
-	Action.Get().IterateConstants([this](FStringView Key, FCbObject&& Value)
-	{
-		Context->AddConstant(Key, MoveTemp(Value));
-	});
+
 	Action.Get().IterateInputs([this](FStringView Key, const FIoHash& RawHash, uint64 RawSize)
 	{
 		const FCompressedBuffer& Buffer = Inputs.Get().FindInput(Key);
@@ -786,7 +791,7 @@ void FBuildJob::EnterExecuteLocal()
 	});
 	Action.Reset();
 	Inputs.Reset();
-	SchedulerParams.ResolvedInputsSize = 0;
+	Schedule->EditParameters().ResolvedInputsSize = 0;
 }
 
 void FBuildJob::BeginExecuteLocal()
@@ -863,6 +868,7 @@ void FBuildJob::SetInputs(FBuildInputs&& InInputs)
 		return;
 	}
 	Inputs = MoveTemp(InInputs);
+	FBuildSchedulerParams& SchedulerParams = Schedule->EditParameters();
 	SchedulerParams.MissingLocalInputsSize = 0;
 	SchedulerParams.MissingRemoteInputsSize = 0;
 	return AdvanceToState(ExecuteState);
@@ -884,11 +890,6 @@ void FBuildJob::SetOutputNoCheck(FBuildOutput&& InOutput, EBuildJobState NewStat
 	checkf(Output.IsNull(), TEXT("Job already has an output for build of '%s' by %s."), *Name, *FunctionName);
 	Output = MoveTemp(InOutput);
 
-	if (!Owner.IsCanceled())
-	{
-		Scheduler.SetJobOutput(this, SchedulerParams, Output.Get());
-	}
-
 	if (OnComplete)
 	{
 		const FCacheKey& CacheKey = Context ? Context->GetCacheKey() : FCacheKey::Empty;
@@ -998,7 +999,7 @@ void FBuildJob::ExecuteTransition(EBuildJobState OldState, EBuildJobState NewSta
 		{
 			Context->ResetInputs();
 		}
-		SchedulerParams.ResolvedInputsSize = 0;
+		Schedule->EditParameters().ResolvedInputsSize = 0;
 	}
 	if (OldState <= EBuildJobState::ExecuteLocalWait && EBuildJobState::ExecuteLocalWait < NewState && !Output)
 	{
@@ -1038,15 +1039,15 @@ void FBuildJob::ExecuteState(EBuildJobState NewState)
 {
 	switch (NewState)
 	{
-	case EBuildJobState::ResolveKey:                 return Scheduler.DispatchResolveKey(this, Owner);
-	case EBuildJobState::ResolveInputMeta:           return Scheduler.DispatchResolveInputMeta(this, Owner);
-	case EBuildJobState::CacheQuery:                 return Scheduler.DispatchCacheQuery(this, Owner, SchedulerParams);
-	case EBuildJobState::ExecuteRemote:              return Scheduler.DispatchExecuteRemote(this, Owner, SchedulerParams);
-	case EBuildJobState::ResolveRemoteInputData:     return Scheduler.DispatchResolveInputData(this, Owner, SchedulerParams);
-	case EBuildJobState::ExecuteRemoteRetry:         return Scheduler.DispatchExecuteRemote(this, Owner, SchedulerParams);
-	case EBuildJobState::ResolveInputData:           return Scheduler.DispatchResolveInputData(this, Owner, SchedulerParams);
-	case EBuildJobState::ExecuteLocal:               return Scheduler.DispatchExecuteLocal(this, Owner, SchedulerParams);
-	case EBuildJobState::CacheStore:                 return Scheduler.DispatchCacheStore(this, Owner, SchedulerParams);
+	case EBuildJobState::ResolveKey:                 return Schedule->DispatchResolveKey();
+	case EBuildJobState::ResolveInputMeta:           return Schedule->DispatchResolveInputMeta();
+	case EBuildJobState::CacheQuery:                 return Schedule->DispatchCacheQuery();
+	case EBuildJobState::ExecuteRemote:              return Schedule->DispatchExecuteRemote();
+	case EBuildJobState::ResolveRemoteInputData:     return Schedule->DispatchResolveInputData();
+	case EBuildJobState::ExecuteRemoteRetry:         return Schedule->DispatchExecuteRemote();
+	case EBuildJobState::ResolveInputData:           return Schedule->DispatchResolveInputData();
+	case EBuildJobState::ExecuteLocal:               return Schedule->DispatchExecuteLocal();
+	case EBuildJobState::CacheStore:                 return Schedule->DispatchCacheStore();
 	}
 }
 
diff --git a/Engine/Source/Developer/DerivedDataCache/Private/DerivedDataBuildScheduler.cpp b/Engine/Source/Developer/DerivedDataCache/Private/DerivedDataBuildScheduler.cpp
--- a/Engine/Source/Developer/DerivedDataCache/Private/DerivedDataBuildScheduler.cpp
+++ b/Engine/Source/Developer/DerivedDataCache/Private/DerivedDataBuildScheduler.cpp
@@ -1,121 +1,367 @@
 // Copyright Epic Games, Inc. All Rights Reserved.
 
 #include "DerivedDataBuildScheduler.h"
-
+#include "Containers/RingBuffer.h"
 #include "DerivedDataBuildFunction.h"
 #include "DerivedDataBuildFunctionFactory.h"
 #include "DerivedDataBuildFunctionRegistry.h"
 #include "DerivedDataRequest.h"
 #include "DerivedDataRequestOwner.h"
+#include "HAL/Event.h"
 #include "Misc/Guid.h"
+#include "Misc/ScopeLock.h"
 #include "Tasks/Task.h"
 
 namespace UE::DerivedData::Private
 {
 
-class FBuildSchedulerRequest : public FRequestBase
+////////////////////////////////////////////////////////////////////////////////
+
+static void ScheduleAsyncStep(IBuildJob& Job, IRequestOwner& Owner, const TCHAR* DebugName)
 {
-public:
-	FBuildSchedulerRequest(IBuildJob* InJob, IRequestOwner& InOwner, const TCHAR* DebugName)
-		: Job(InJob)
-		, Owner(InOwner)
+	class FStepAsyncRequest final : public FRequestBase
 	{
-		Tasks::FTaskEvent TaskEvent(TEXT("FBuildSchedulerRequest"));
-		Task = Tasks::Launch(DebugName, [this] { Schedule(); }, TaskEvent, Tasks::ETaskPriority::BackgroundNormal);
-		Owner.Begin(this);
-		TaskEvent.Trigger();
-	}
+	public:
+		Tasks::FTask Task;
+
+		void SetPriority(EPriority Priority) override {}
+		void Cancel() override { Task.Wait(); }
+		void Wait() override { Task.Wait(); }
+	};
+	
+	FStepAsyncRequest* Request = new FStepAsyncRequest;
+
+	Tasks::FTaskEvent TaskEvent(TEXT("ScheduleAsyncStep"));
+	Request->Task = Tasks::Launch(DebugName, 
+								[&, Request] { Owner.End(Request, [&] { Job.StepExecution(); }); }, 
+								TaskEvent, Tasks::ETaskPriority::BackgroundNormal);
+	Owner.Begin(Request);
+	TaskEvent.Trigger();
+}
+
+////////////////////////////////////////////////////////////////////////////////
 
-	void Schedule()
+static FEvent* GPretriggeredEvent;
+
+/** On-demand allocated event to reduce total event allocation */
+class FLazyManualResetEvent
+{
+public:
+	static void Init()
 	{
-		Owner.End(this, [this] { Job->Schedule(); });
+		static FEventRef PretriggeredEvent(EEventMode::ManualReset);
+		PretriggeredEvent->Trigger();
+		GPretriggeredEvent = PretriggeredEvent.Get();
 	}
 
-	void SetPriority(EPriority Priority) final
+	~FLazyManualResetEvent()
 	{
+		// Sequential consistency might not be needed here in case performance matters
+		FEvent* Event = AtomicEvent.load();
+		if ((Event != nullptr) & (Event != GPretriggeredEvent))
+		{
+			FPlatformProcess::ReturnSynchEventToPool(Event);
+		}
 	}
 
-	void Cancel() final
+	void Wait()
 	{
-		Task.Wait();
-	}
+		FEvent* Event = AtomicEvent.load(std::memory_order_relaxed);
+
+		if (Event == nullptr)
+		{
+			FEvent* NewEvent = FPlatformProcess::GetSynchEventFromPool(/* manual reset */ true);
+			if (AtomicEvent.compare_exchange_strong(/* out */ Event, NewEvent))
+			{
+				Event = NewEvent;
+			}
+			else
+			{
+				FPlatformProcess::ReturnSynchEventToPool(NewEvent);
+			}
+		}
 
-	void Wait() final
+		Event->Wait();
+	}
+	
+	void Trigger()
 	{
-		Task.Wait();
+		FEvent* Event = AtomicEvent.load(std::memory_order_relaxed);
+		
+		if (Event == nullptr)
+		{
+			if (AtomicEvent.compare_exchange_strong(/* out */ Event, GPretriggeredEvent))
+			{
+				return;
+			}
+		}
+
+		if (Event != GPretriggeredEvent)
+		{
+			Event->Trigger();
+		}
 	}
 
 private:
-	IBuildJob* Job;
-	IRequestOwner& Owner;
-	Tasks::FTask Task;
+	std::atomic<FEvent*> AtomicEvent {nullptr};
 };
 
-class FBuildScheduler final : public IBuildScheduler
+////////////////////////////////////////////////////////////////////////////////
+
+/** Limits simultaneous build jobs to reduce peak memory usage */
+class FMemoryScheduler
 {
 public:
-	void DispatchCacheQuery(IBuildJob* Job, IRequestOwner& Owner, const FBuildSchedulerParams& Params) final;
-	void DispatchCacheStore(IBuildJob* Job, IRequestOwner& Owner, const FBuildSchedulerParams& Params) final;
-	void DispatchResolveKey(IBuildJob* Job, IRequestOwner& Owner) final;
-	void DispatchResolveInputMeta(IBuildJob* Job, IRequestOwner& Owner) final;
-	void DispatchResolveInputData(IBuildJob* Job, IRequestOwner& Owner, const FBuildSchedulerParams& Params) final;
-	void DispatchExecuteRemote(IBuildJob* Job, IRequestOwner& Owner, const FBuildSchedulerParams& Params) final;
-	void DispatchExecuteLocal(IBuildJob* Job, IRequestOwner& Owner, const FBuildSchedulerParams& Params) final;
+	FMemoryScheduler();
+	~FMemoryScheduler();
+
+	void RegisterRunningJob(uint64 MemoryEstimate);
+	void StepAsyncOrQueue(uint64 MemoryEstimate, IBuildJob& Job, IRequestOwner& Owner, const TCHAR* DebugName);
+	void RegisterEndedJob(uint64 MemoryEstimate);
 
 private:
-	void Dispatch(IBuildJob* Job, IRequestOwner& Owner, const TCHAR* DebugName);
+	/** Handles waiting and cancellation while a job is queued up */
+	class FRequest final : public FRequestBase
+	{
+	public:
+		FRequest(FMemoryScheduler& InScheduler, IBuildJob& InJob, IRequestOwner& InOwner, uint64 InMemoryEstimate);
+		~FRequest() { ensure(!TryClaimEnd()); }
+		
+		void SetPriority(EPriority Priority) override {}
+		void Wait() override { Event.Wait(); }
+		void Cancel() override;
+
+		uint64 GetMemoryEstimate() const { return MemoryEstimate; }
+		bool TryClaimEnd() { return !bClaimed.test_and_set(); }
+		void End(const TCHAR* DebugName);
+
+	private:
+		FMemoryScheduler& Scheduler;
+		IBuildJob& Job;
+		IRequestOwner& Owner;
+		const uint64 MemoryEstimate;
+		FLazyManualResetEvent Event;
+		std::atomic_flag bClaimed = ATOMIC_FLAG_INIT;
+	};
+	
+	const uint64 TotalPhysical;
+	const uint64 AvailablePhysicalAtStartup;
+	const uint64 MaxMemoryUsage;
+	
+	FCriticalSection CriticalSection;
+	TRingBuffer<TRefCountPtr<FRequest>> Queue;
+	uint64 TotalScheduledMemory = 0;
+	uint64 TotalScheduledWatermark = 0;
+
+	// @pre CriticalSection locked
+	bool CanRunNow(uint64 MemoryEstimate) const
+	{
+		return TotalScheduledMemory == 0 || TotalScheduledMemory + MemoryEstimate < MaxMemoryUsage;
+	}
 };
 
-void FBuildScheduler::DispatchCacheQuery(IBuildJob* Job, IRequestOwner& Owner, const FBuildSchedulerParams& Params)
+FMemoryScheduler::FRequest::FRequest(FMemoryScheduler& InScheduler, IBuildJob& InJob, IRequestOwner& InOwner, uint64 InMemoryEstimate)
+: Scheduler(InScheduler)
+, Job(InJob)
+, Owner(InOwner)
+, MemoryEstimate(InMemoryEstimate)
 {
-	Job->Schedule();
+	Owner.Begin(this);
 }
 
-void FBuildScheduler::DispatchCacheStore(IBuildJob* Job, IRequestOwner& Owner, const FBuildSchedulerParams& Params)
+void FMemoryScheduler::FRequest::End(const TCHAR* DebugName)
 {
-	Job->Schedule();
+	Owner.End(this, [this, DebugName]()
+		{ 
+			ScheduleAsyncStep(Job, Owner, DebugName);
+			Event.Trigger();
+		});
+	
 }
 
-void FBuildScheduler::DispatchResolveKey(IBuildJob* Job, IRequestOwner& Owner)
+void FMemoryScheduler::FRequest::Cancel()
 {
-	Dispatch(Job, Owner, TEXT("FBuildScheduler::DispatchResolveKey"));
+	if (TryClaimEnd())
+	{
+		// Add estimated memory to simplify implementation, even though memory won't be allocated.
+		// FBuildJobSchedule::EndJob() will restore the scheduler's available memory.
+		// Might require optimization if lots of queued jobs are cancelled at the same time.
+		Scheduler.RegisterRunningJob(MemoryEstimate);
+		End(TEXT("MemoryQueueCancel"));
+	}
 }
 
-void FBuildScheduler::DispatchResolveInputMeta(IBuildJob* Job, IRequestOwner& Owner)
+FMemoryScheduler::FMemoryScheduler()
+: TotalPhysical(FPlatformMemory::GetStats().TotalPhysical)
+, AvailablePhysicalAtStartup(FPlatformMemory::GetStats().AvailablePhysical)
+, MaxMemoryUsage(TotalPhysical / 8 + AvailablePhysicalAtStartup / 2)
 {
-	Dispatch(Job, Owner, TEXT("FBuildScheduler::DispatchResolveInputMeta"));
+	Queue.Reserve(128);
 }
 
-void FBuildScheduler::DispatchResolveInputData(IBuildJob* Job, IRequestOwner& Owner, const FBuildSchedulerParams& Params)
+FMemoryScheduler::~FMemoryScheduler()
 {
-	Dispatch(Job, Owner, TEXT("FBuildScheduler::DispatchResolveInputData"));
+	ensure(Queue.IsEmpty());
+	ensure(TotalScheduledMemory == 0);
 }
 
-void FBuildScheduler::DispatchExecuteRemote(IBuildJob* Job, IRequestOwner& Owner, const FBuildSchedulerParams& Params)
+void FMemoryScheduler::RegisterRunningJob(uint64 MemoryEstimate)
 {
-	Dispatch(Job, Owner, TEXT("FBuildScheduler::DispatchExecuteRemote"));
+	check(MemoryEstimate);
+
+	FScopeLock Lock(&CriticalSection);
+	TotalScheduledMemory += MemoryEstimate;
+	TotalScheduledWatermark = FMath::Max(TotalScheduledWatermark, TotalScheduledMemory);
 }
 
-void FBuildScheduler::DispatchExecuteLocal(IBuildJob* Job, IRequestOwner& Owner, const FBuildSchedulerParams& Params)
+void FMemoryScheduler::StepAsyncOrQueue(uint64 MemoryEstimate, IBuildJob& Job, IRequestOwner& Owner, const TCHAR* DebugName)
 {
-	Dispatch(Job, Owner, TEXT("FBuildScheduler::DispatchExecuteLocal"));
+	check(MemoryEstimate);
+	{
+		FScopeLock Lock(&CriticalSection);
+
+		if (!CanRunNow(MemoryEstimate))
+		{
+			Queue.Add(TRefCountPtr(new FRequest(*this, Job, Owner, MemoryEstimate)));
+			return;
+		}
+
+		TotalScheduledMemory += MemoryEstimate;
+		TotalScheduledWatermark = FMath::Max(TotalScheduledWatermark, TotalScheduledMemory);
+	}
+
+	ScheduleAsyncStep(Job, Owner, DebugName);
 }
 
-void FBuildScheduler::Dispatch(IBuildJob* Job, IRequestOwner& Owner, const TCHAR* DebugName)
+void FMemoryScheduler::RegisterEndedJob(uint64 DoneEstimate)
 {
-	if (Owner.GetPriority() == EPriority::Blocking)
+	if (DoneEstimate)
 	{
-		Job->Schedule();
+		TArray<TRefCountPtr<FRequest>, TInlineAllocator<16>> Continuations;
+	
+		{
+			FScopeLock Lock(&CriticalSection);
+
+			TotalScheduledMemory -= DoneEstimate;
+
+			if (Queue.IsEmpty())
+			{
+				return;
+			}
+		
+			while (Queue.Num() && CanRunNow(Queue.First()->GetMemoryEstimate()))
+			{
+				if (Queue.First()->TryClaimEnd())
+				{
+					TotalScheduledMemory += Queue.First()->GetMemoryEstimate();
+					Continuations.Add(Queue.First());
+				}
+				Queue.PopFront();
+			}
+
+			TotalScheduledWatermark = FMath::Max(TotalScheduledWatermark, TotalScheduledMemory);
+		}
+
+		for (const TRefCountPtr<FRequest>& Request : Continuations)
+		{
+			Request->End(TEXT("MemoryQueueContinue"));
+		}
+	}
+}
+
+////////////////////////////////////////////////////////////////////////////////
+
+class FBuildJobSchedule final : public IBuildJobSchedule
+{
+public:
+	FBuildJobSchedule(IBuildJob& InJob, IRequestOwner& InOwner, FMemoryScheduler& InMemoryLimiter)
+		: Job(InJob)
+		, Owner(InOwner)
+		, MemoryLimiter(InMemoryLimiter)
+	{}
+
+	virtual FBuildSchedulerParams& EditParameters() override { return Params; }
+
+	virtual void DispatchCacheQuery() override			{ StepSync(); }
+	virtual void DispatchCacheStore() override			{ StepSync(); }
+	virtual void DispatchResolveKey() override			{ StepAsync(TEXT("ResolveKey")); }
+	virtual void DispatchResolveInputMeta() override	{ StepAsync(TEXT("ResolveInputMeta")); }
+	virtual void DispatchResolveInputData() override	{ StepAsyncOrQueue(TEXT("ResolveInputData")); }
+	virtual void DispatchExecuteRemote() override		{ StepAsyncOrQueue(TEXT("ExecuteRemote")); }
+	virtual void DispatchExecuteLocal() override		{ StepAsyncOrQueue(TEXT("ExecuteLocal")); }
+
+	virtual void EndJob() override						{ MemoryLimiter.RegisterEndedJob(ScheduledMemoryEstimate); }
+
+private:
+	void StepSync()
+	{
+		Job.StepExecution();
 	}
-	else
+
+	void StepAsync(const TCHAR* DebugName)
 	{
-		new FBuildSchedulerRequest(Job, Owner, DebugName);
+		if (Owner.GetPriority() == EPriority::Blocking)
+		{
+			StepSync();
+		}
+		else
+		{
+			ScheduleAsyncStep(Job, Owner, DebugName);
+		}
 	}
-}
+
+	void StepAsyncOrQueue(const TCHAR* DebugName)
+	{
+		check(Params.TotalRequiredMemory >= Params.ResolvedInputsSize);
+		const uint64 CurrentMemoryEstimate = Params.TotalRequiredMemory - Params.ResolvedInputsSize;
+
+		// Only queue for memory once
+		if (ScheduledMemoryEstimate || CurrentMemoryEstimate == 0)
+		{
+			StepAsync(DebugName);
+		}
+		else
+		{
+			ScheduledMemoryEstimate = CurrentMemoryEstimate;
+
+			if (Owner.GetPriority() == EPriority::Blocking)
+			{
+				MemoryLimiter.RegisterRunningJob(ScheduledMemoryEstimate);
+				StepSync();
+			}
+			else
+			{
+				MemoryLimiter.StepAsyncOrQueue(ScheduledMemoryEstimate, Job, Owner, DebugName);
+			}
+		}
+	}
+
+private:
+	IBuildJob& Job;
+	IRequestOwner& Owner;
+	FBuildSchedulerParams Params;
+	FMemoryScheduler& MemoryLimiter;
+	uint64 ScheduledMemoryEstimate = 0;
+};
+
+////////////////////////////////////////////////////////////////////////////////
+
+class FBuildScheduler final : public IBuildScheduler
+{
+	TUniquePtr<IBuildJobSchedule> BeginJob(IBuildJob& Job, IRequestOwner& Owner)
+	{
+		return MakeUnique<FBuildJobSchedule>(Job, Owner, MemoryLimiter);
+	}
+
+	FMemoryScheduler MemoryLimiter;
+};
 
 IBuildScheduler* CreateBuildScheduler()
 {
+	FLazyManualResetEvent::Init();
 	return new FBuildScheduler();
 }
 
+////////////////////////////////////////////////////////////////////////////////
+
 } // UE::DerivedData::Private
diff --git a/Engine/Source/Developer/DerivedDataCache/Public/DerivedDataBuildJob.h b/Engine/Source/Developer/DerivedDataCache/Public/DerivedDataBuildJob.h
--- a/Engine/Source/Developer/DerivedDataCache/Public/DerivedDataBuildJob.h
+++ b/Engine/Source/Developer/DerivedDataCache/Public/DerivedDataBuildJob.h
@@ -48,7 +48,7 @@ public:
 	virtual IBuild& GetBuild() const = 0;
 
 	/** Called by the scheduler to continue this job on the calling thread. */
-	virtual void Schedule() = 0;
+	virtual void StepExecution() = 0;
 
 	/** Called by the scheduler to skip remote execution and fall back to local execution if permitted. */
 	virtual void SkipExecuteRemote() = 0;
diff --git a/Engine/Source/Developer/DerivedDataCache/Public/DerivedDataBuildScheduler.h b/Engine/Source/Developer/DerivedDataCache/Public/DerivedDataBuildScheduler.h
--- a/Engine/Source/Developer/DerivedDataCache/Public/DerivedDataBuildScheduler.h
+++ b/Engine/Source/Developer/DerivedDataCache/Public/DerivedDataBuildScheduler.h
@@ -4,9 +4,13 @@
 
 #include "DerivedDataBuildJob.h"
 #include "DerivedDataBuildKey.h"
+#include "DerivedDataRequest.h"
+#include "Templates/RefCounting.h"
+#include "Templates/UniquePtr.h"
 
-namespace UE::DerivedData { class IRequestOwner; }
 namespace UE::DerivedData { struct FBuildSchedulerParams; }
+namespace UE::DerivedData { class IBuildJobSchedule; }
+namespace UE::DerivedData { class IRequestOwner; }
 
 namespace UE::DerivedData
 {
@@ -23,46 +27,7 @@ class IBuildScheduler
 {
 public:
 	virtual ~IBuildScheduler() = default;
-
-	/** Begin processing of the job by this scheduler. Always paired with EndJob. */
-	virtual void BeginJob(IBuildJob* Job) {}
-
-	/** End processing of the job by this scheduler. Always paired with BeginJob. */
-	virtual void EndJob(IBuildJob* Job) {}
-
-	/** Dispatch the job immediately if it is queued. May be called multiple times and/or concurrently. */
-	virtual void CancelJob(IBuildJob* Job) {}
-
-	/** Update the priority of the job if it is queued. May be called multiple times and/or concurrently. */
-	virtual void UpdateJobPriority(IBuildJob* Job) {}
-
-	/** Dispatch by calling Schedule or SetOutput, either now or later. */
-	virtual void DispatchCacheQuery(IBuildJob* Job, IRequestOwner& Owner, const FBuildSchedulerParams& Params) { Job->Schedule(); }
-
-	/** Dispatch by calling Schedule, either now or later. */
-	virtual void DispatchCacheStore(IBuildJob* Job, IRequestOwner& Owner, const FBuildSchedulerParams& Params) { Job->Schedule(); }
-
-	/** Dispatch by calling Schedule, either now or later. */
-	virtual void DispatchResolveKey(IBuildJob* Job, IRequestOwner& Owner) { Job->Schedule(); }
-
-	/** Dispatch by calling Schedule, either now or later. */
-	virtual void DispatchResolveInputMeta(IBuildJob* Job, IRequestOwner& Owner) { Job->Schedule(); }
-
-	/**
-	 * Dispatch by calling Schedule, SetOutput, or SkipExecuteRemote, either now or later.
-	 *
-	 * SkipExecuteRemote is only valid to call when MissingRemoteInputsSize is non-zero.
-	 */
-	virtual void DispatchResolveInputData(IBuildJob* Job, IRequestOwner& Owner, const FBuildSchedulerParams& Params) { Job->Schedule(); }
-
-	/** Dispatch by calling Schedule, SetOutput, or SkipExecuteRemote, either now or later. */
-	virtual void DispatchExecuteRemote(IBuildJob* Job, IRequestOwner& Owner, const FBuildSchedulerParams& Params) { Job->Schedule(); }
-
-	/** Dispatch by calling Schedule or SetOutput, either now or later. */
-	virtual void DispatchExecuteLocal(IBuildJob* Job, IRequestOwner& Owner, const FBuildSchedulerParams& Params) { Job->Schedule(); }
-
-	/** Set the output of the job. Always called once between BeginJob and EndJob unless canceled. */
-	virtual void SetJobOutput(IBuildJob* Job, const FBuildSchedulerParams& Params, const FBuildOutput& Output) {}
+	virtual TUniquePtr<IBuildJobSchedule> BeginJob(IBuildJob& Job, IRequestOwner& Owner) = 0;
 };
 
 /** Parameters that describe a build job to the build scheduler. */
@@ -84,4 +49,34 @@ struct FBuildSchedulerParams
 	uint64 TotalRequiredMemory = 0;
 };
 
+/** Scheduling interface / context for an IBuildJob */
+class IBuildJobSchedule
+{
+public:
+	virtual ~IBuildJobSchedule() = default;
+
+	virtual FBuildSchedulerParams& EditParameters() = 0;
+
+	/** Calls StepExecution() now or later. */
+	virtual void DispatchCacheQuery() = 0;
+	/** Calls StepExecution() now or later. */
+	virtual void DispatchCacheStore() = 0;
+	/** Calls StepExecution() now or later. */
+	virtual void DispatchResolveKey() = 0;
+	/** Calls StepExecution() now or later. */
+	virtual void DispatchResolveInputMeta() = 0;
+	/**
+	 * Calls StepExecution() or SkipExecuteRemote() now or later.
+	 *
+	 * SkipExecuteRemote() won't be called unless MissingRemoteInputsSize is non-zero.
+	 */
+	virtual void DispatchResolveInputData() = 0;
+	/** Calls StepExecution() or SkipExecuteRemote() now or later. */
+	virtual void DispatchExecuteRemote() = 0;
+	/** Calls StepExecution() now or later. */
+	virtual void DispatchExecuteLocal()	= 0;
+
+	virtual void EndJob() = 0;
+};
+
 } // UE::DerivedData
diff --git a/Engine/Source/Developer/TextureBuild/Private/TextureBuildFunction.cpp b/Engine/Source/Developer/TextureBuild/Private/TextureBuildFunction.cpp
--- a/Engine/Source/Developer/TextureBuild/Private/TextureBuildFunction.cpp
+++ b/Engine/Source/Developer/TextureBuild/Private/TextureBuildFunction.cpp
@@ -260,6 +260,20 @@ static bool TryReadTextureSourceFromCompactBinary(FCbFieldView Source, UE::Deriv
 	return true;
 }
 
+// Estimate peak memory usage required to cook this texture, excluding raw input / source data
+// Similar to UTexture::GetBuildRequiredMemory()
+// This is a rough blackbox estimate that can be improved
+static uint64 EstimateTextureBuildMemoryUsage(const FCbObject& Settings)
+{
+	uint64 InputSize = 0;
+	for (FCbField Mip : Settings["Source"]["Mips"])
+	{
+		InputSize += Mip["Size"].AsInt64();
+	}
+
+	return static_cast<uint64>(InputSize * 7.5);
+}
+
 FGuid FTextureBuildFunction::GetVersion() const
 {
 	UE::DerivedData::FBuildVersionBuilder Builder;
@@ -284,6 +298,7 @@ FGuid FTextureBuildFunction::GetVersion() const
 void FTextureBuildFunction::Configure(UE::DerivedData::FBuildConfigContext& Context) const
 {
 	Context.SetCacheBucket(UE::DerivedData::FCacheBucket("Texture"_ASV));
+	Context.SetRequiredMemory(EstimateTextureBuildMemoryUsage(Context.FindConstant(TEXT("Settings"_SV))));
 }
 
 void FTextureBuildFunction::Build(UE::DerivedData::FBuildContext& Context) const
diff --git a/Engine/Source/Runtime/Engine/Private/Texture.cpp b/Engine/Source/Runtime/Engine/Private/Texture.cpp
--- a/Engine/Source/Runtime/Engine/Private/Texture.cpp
+++ b/Engine/Source/Runtime/Engine/Private/Texture.cpp
@@ -1892,6 +1892,8 @@ void UTexture::SetLayerFormatSettings(int32 LayerIndex, const FTextureFormatSett
 
 int64 UTexture::GetBuildRequiredMemory() const
 {
+	// If you improve this estimate, please update EstimateTextureBuildMemoryUsage() as well
+
 	// Compute the memory it should take to uncompress the bulkdata in memory
 	int64 MemoryEstimate = 0;
 
