DDC: Renamed Dispatch[Op] to Schedule[Op] and tweaked formatting for consistency

#rnx

diff --git a/Engine/Source/Developer/DerivedDataCache/Private/DerivedDataBuildJob.cpp b/Engine/Source/Developer/DerivedDataCache/Private/DerivedDataBuildJob.cpp
--- a/Engine/Source/Developer/DerivedDataCache/Private/DerivedDataBuildJob.cpp
+++ b/Engine/Source/Developer/DerivedDataCache/Private/DerivedDataBuildJob.cpp
@@ -1039,15 +1039,15 @@ void FBuildJob::ExecuteState(EBuildJobState NewState)
 {
 	switch (NewState)
 	{
-	case EBuildJobState::ResolveKey:                 return Schedule->DispatchResolveKey();
-	case EBuildJobState::ResolveInputMeta:           return Schedule->DispatchResolveInputMeta();
-	case EBuildJobState::CacheQuery:                 return Schedule->DispatchCacheQuery();
-	case EBuildJobState::ExecuteRemote:              return Schedule->DispatchExecuteRemote();
-	case EBuildJobState::ResolveRemoteInputData:     return Schedule->DispatchResolveInputData();
-	case EBuildJobState::ExecuteRemoteRetry:         return Schedule->DispatchExecuteRemote();
-	case EBuildJobState::ResolveInputData:           return Schedule->DispatchResolveInputData();
-	case EBuildJobState::ExecuteLocal:               return Schedule->DispatchExecuteLocal();
-	case EBuildJobState::CacheStore:                 return Schedule->DispatchCacheStore();
+	case EBuildJobState::ResolveKey:                 return Schedule->ScheduleResolveKey();
+	case EBuildJobState::ResolveInputMeta:           return Schedule->ScheduleResolveInputMeta();
+	case EBuildJobState::CacheQuery:                 return Schedule->ScheduleCacheQuery();
+	case EBuildJobState::ExecuteRemote:              return Schedule->ScheduleExecuteRemote();
+	case EBuildJobState::ResolveRemoteInputData:     return Schedule->ScheduleResolveInputData();
+	case EBuildJobState::ExecuteRemoteRetry:         return Schedule->ScheduleExecuteRemote();
+	case EBuildJobState::ResolveInputData:           return Schedule->ScheduleResolveInputData();
+	case EBuildJobState::ExecuteLocal:               return Schedule->ScheduleExecuteLocal();
+	case EBuildJobState::CacheStore:                 return Schedule->ScheduleCacheStore();
 	}
 }
 
diff --git a/Engine/Source/Developer/DerivedDataCache/Private/DerivedDataBuildScheduler.cpp b/Engine/Source/Developer/DerivedDataCache/Private/DerivedDataBuildScheduler.cpp
--- a/Engine/Source/Developer/DerivedDataCache/Private/DerivedDataBuildScheduler.cpp
+++ b/Engine/Source/Developer/DerivedDataCache/Private/DerivedDataBuildScheduler.cpp
@@ -15,7 +15,7 @@
 namespace UE::DerivedData::Private
 {
 
-////////////////////////////////////////////////////////////////////////////////
+///////////////////////////////////////////////////////////////////////////////////////////////////////////////////////
 
 static void ScheduleAsyncStep(IBuildJob& Job, IRequestOwner& Owner, const TCHAR* DebugName)
 {
@@ -24,22 +24,22 @@ static void ScheduleAsyncStep(IBuildJob& Job, IRequestOwner& Owner, const TCHAR*
 	public:
 		Tasks::FTask Task;
 
-		void SetPriority(EPriority Priority) override {}
-		void Cancel() override { Task.Wait(); }
-		void Wait() override { Task.Wait(); }
+		void SetPriority(EPriority Priority) final {}
+		void Cancel() final { Task.Wait(); }
+		void Wait() final { Task.Wait(); }
 	};
-	
+
 	FStepAsyncRequest* Request = new FStepAsyncRequest;
 
 	Tasks::FTaskEvent TaskEvent(TEXT("ScheduleAsyncStep"));
-	Request->Task = Tasks::Launch(DebugName, 
-								[&, Request] { Owner.End(Request, [&] { Job.StepExecution(); }); }, 
-								TaskEvent, Tasks::ETaskPriority::BackgroundNormal);
+	Request->Task = Tasks::Launch(DebugName,
+		[&Job, &Owner, Request] { Owner.End(Request, [&Job] { Job.StepExecution(); }); },
+		TaskEvent, Tasks::ETaskPriority::BackgroundNormal);
 	Owner.Begin(Request);
 	TaskEvent.Trigger();
 }
 
-////////////////////////////////////////////////////////////////////////////////
+///////////////////////////////////////////////////////////////////////////////////////////////////////////////////////
 
 static FEvent* GPretriggeredEvent;
 
@@ -83,11 +83,11 @@ public:
 
 		Event->Wait();
 	}
-	
+
 	void Trigger()
 	{
 		FEvent* Event = AtomicEvent.load(std::memory_order_relaxed);
-		
+
 		if (Event == nullptr)
 		{
 			if (AtomicEvent.compare_exchange_strong(/* out */ Event, GPretriggeredEvent))
@@ -103,10 +103,10 @@ public:
 	}
 
 private:
-	std::atomic<FEvent*> AtomicEvent {nullptr};
+	std::atomic<FEvent*> AtomicEvent{nullptr};
 };
 
-////////////////////////////////////////////////////////////////////////////////
+///////////////////////////////////////////////////////////////////////////////////////////////////////////////////////
 
 /** Limits simultaneous build jobs to reduce peak memory usage */
 class FMemoryScheduler
@@ -126,10 +126,10 @@ private:
 	public:
 		FRequest(FMemoryScheduler& InScheduler, IBuildJob& InJob, IRequestOwner& InOwner, uint64 InMemoryEstimate);
 		~FRequest() { ensure(!TryClaimEnd()); }
-		
-		void SetPriority(EPriority Priority) override {}
-		void Wait() override { Event.Wait(); }
-		void Cancel() override;
+
+		void SetPriority(EPriority Priority) final {}
+		void Wait() final { Event.Wait(); }
+		void Cancel() final;
 
 		uint64 GetMemoryEstimate() const { return MemoryEstimate; }
 		bool TryClaimEnd() { return !bClaimed.test_and_set(); }
@@ -143,11 +143,11 @@ private:
 		FLazyManualResetEvent Event;
 		std::atomic_flag bClaimed = ATOMIC_FLAG_INIT;
 	};
-	
+
 	const uint64 TotalPhysical;
 	const uint64 AvailablePhysicalAtStartup;
 	const uint64 MaxMemoryUsage;
-	
+
 	FCriticalSection CriticalSection;
 	TRingBuffer<TRefCountPtr<FRequest>> Queue;
 	uint64 TotalScheduledMemory = 0;
@@ -160,23 +160,26 @@ private:
 	}
 };
 
-FMemoryScheduler::FRequest::FRequest(FMemoryScheduler& InScheduler, IBuildJob& InJob, IRequestOwner& InOwner, uint64 InMemoryEstimate)
-: Scheduler(InScheduler)
-, Job(InJob)
-, Owner(InOwner)
-, MemoryEstimate(InMemoryEstimate)
+FMemoryScheduler::FRequest::FRequest(
+	FMemoryScheduler& InScheduler,
+	IBuildJob& InJob,
+	IRequestOwner& InOwner,
+	uint64 InMemoryEstimate)
+	: Scheduler(InScheduler)
+	, Job(InJob)
+	, Owner(InOwner)
+	, MemoryEstimate(InMemoryEstimate)
 {
 	Owner.Begin(this);
 }
 
 void FMemoryScheduler::FRequest::End(const TCHAR* DebugName)
 {
-	Owner.End(this, [this, DebugName]()
-		{ 
-			ScheduleAsyncStep(Job, Owner, DebugName);
-			Event.Trigger();
-		});
-	
+	Owner.End(this, [this, DebugName]
+	{ 
+		ScheduleAsyncStep(Job, Owner, DebugName);
+		Event.Trigger();
+	});
 }
 
 void FMemoryScheduler::FRequest::Cancel()
@@ -192,9 +195,9 @@ void FMemoryScheduler::FRequest::Cancel()
 }
 
 FMemoryScheduler::FMemoryScheduler()
-: TotalPhysical(FPlatformMemory::GetStats().TotalPhysical)
-, AvailablePhysicalAtStartup(FPlatformMemory::GetStats().AvailablePhysical)
-, MaxMemoryUsage(TotalPhysical / 8 + AvailablePhysicalAtStartup / 2)
+	: TotalPhysical(FPlatformMemory::GetStats().TotalPhysical)
+	, AvailablePhysicalAtStartup(FPlatformMemory::GetStats().AvailablePhysical)
+	, MaxMemoryUsage(TotalPhysical / 8 + AvailablePhysicalAtStartup / 2)
 {
 	Queue.Reserve(128);
 }
@@ -214,7 +217,11 @@ void FMemoryScheduler::RegisterRunningJob(uint64 MemoryEstimate)
 	TotalScheduledWatermark = FMath::Max(TotalScheduledWatermark, TotalScheduledMemory);
 }
 
-void FMemoryScheduler::StepAsyncOrQueue(uint64 MemoryEstimate, IBuildJob& Job, IRequestOwner& Owner, const TCHAR* DebugName)
+void FMemoryScheduler::StepAsyncOrQueue(
+	uint64 MemoryEstimate,
+	IBuildJob& Job,
+	IRequestOwner& Owner,
+	const TCHAR* DebugName)
 {
 	check(MemoryEstimate);
 	{
@@ -222,7 +229,7 @@ void FMemoryScheduler::StepAsyncOrQueue(uint64 MemoryEstimate, IBuildJob& Job, I
 
 		if (!CanRunNow(MemoryEstimate))
 		{
-			Queue.Add(TRefCountPtr(new FRequest(*this, Job, Owner, MemoryEstimate)));
+			Queue.Emplace(new FRequest(*this, Job, Owner, MemoryEstimate));
 			return;
 		}
 
@@ -238,7 +245,7 @@ void FMemoryScheduler::RegisterEndedJob(uint64 DoneEstimate)
 	if (DoneEstimate)
 	{
 		TArray<TRefCountPtr<FRequest>, TInlineAllocator<16>> Continuations;
-	
+
 		{
 			FScopeLock Lock(&CriticalSection);
 
@@ -248,7 +255,7 @@ void FMemoryScheduler::RegisterEndedJob(uint64 DoneEstimate)
 			{
 				return;
 			}
-		
+
 			while (Queue.Num() && CanRunNow(Queue.First()->GetMemoryEstimate()))
 			{
 				if (Queue.First()->TryClaimEnd())
@@ -269,7 +276,7 @@ void FMemoryScheduler::RegisterEndedJob(uint64 DoneEstimate)
 	}
 }
 
-////////////////////////////////////////////////////////////////////////////////
+///////////////////////////////////////////////////////////////////////////////////////////////////////////////////////
 
 class FBuildJobSchedule final : public IBuildJobSchedule
 {
@@ -278,19 +285,20 @@ public:
 		: Job(InJob)
 		, Owner(InOwner)
 		, MemoryLimiter(InMemoryLimiter)
-	{}
+	{
+	}
 
-	virtual FBuildSchedulerParams& EditParameters() override { return Params; }
+	FBuildSchedulerParams& EditParameters() final { return Params; }
 
-	virtual void DispatchCacheQuery() override			{ StepSync(); }
-	virtual void DispatchCacheStore() override			{ StepSync(); }
-	virtual void DispatchResolveKey() override			{ StepAsync(TEXT("ResolveKey")); }
-	virtual void DispatchResolveInputMeta() override	{ StepAsync(TEXT("ResolveInputMeta")); }
-	virtual void DispatchResolveInputData() override	{ StepAsyncOrQueue(TEXT("ResolveInputData")); }
-	virtual void DispatchExecuteRemote() override		{ StepAsyncOrQueue(TEXT("ExecuteRemote")); }
-	virtual void DispatchExecuteLocal() override		{ StepAsyncOrQueue(TEXT("ExecuteLocal")); }
+	void ScheduleCacheQuery() final       { StepSync(); }
+	void ScheduleCacheStore() final       { StepSync(); }
+	void ScheduleResolveKey() final       { StepAsync(TEXT("ResolveKey")); }
+	void ScheduleResolveInputMeta() final { StepAsync(TEXT("ResolveInputMeta")); }
+	void ScheduleResolveInputData() final { StepAsyncOrQueue(TEXT("ResolveInputData")); }
+	void ScheduleExecuteRemote() final    { StepAsyncOrQueue(TEXT("ExecuteRemote")); }
+	void ScheduleExecuteLocal() final     { StepAsyncOrQueue(TEXT("ExecuteLocal")); }
 
-	virtual void EndJob() override						{ MemoryLimiter.RegisterEndedJob(ScheduledMemoryEstimate); }
+	void EndJob() final { MemoryLimiter.RegisterEndedJob(ScheduledMemoryEstimate); }
 
 private:
 	void StepSync()
@@ -344,11 +352,11 @@ private:
 	uint64 ScheduledMemoryEstimate = 0;
 };
 
-////////////////////////////////////////////////////////////////////////////////
+///////////////////////////////////////////////////////////////////////////////////////////////////////////////////////
 
 class FBuildScheduler final : public IBuildScheduler
 {
-	TUniquePtr<IBuildJobSchedule> BeginJob(IBuildJob& Job, IRequestOwner& Owner)
+	TUniquePtr<IBuildJobSchedule> BeginJob(IBuildJob& Job, IRequestOwner& Owner) final
 	{
 		return MakeUnique<FBuildJobSchedule>(Job, Owner, MemoryLimiter);
 	}
@@ -362,6 +370,6 @@ IBuildScheduler* CreateBuildScheduler()
 	return new FBuildScheduler();
 }
 
-////////////////////////////////////////////////////////////////////////////////
+///////////////////////////////////////////////////////////////////////////////////////////////////////////////////////
 
 } // UE::DerivedData::Private
diff --git a/Engine/Source/Developer/DerivedDataCache/Public/DerivedDataBuildJob.h b/Engine/Source/Developer/DerivedDataCache/Public/DerivedDataBuildJob.h
--- a/Engine/Source/Developer/DerivedDataCache/Public/DerivedDataBuildJob.h
+++ b/Engine/Source/Developer/DerivedDataCache/Public/DerivedDataBuildJob.h
@@ -10,13 +10,9 @@ template <typename FuncType> class TUniqueFunction;
 
 namespace UE::DerivedData { class FBuildOutput; }
 namespace UE::DerivedData { class IBuild; }
-namespace UE::DerivedData { class IBuildScheduler; }
 namespace UE::DerivedData { class ICache; }
 namespace UE::DerivedData { struct FBuildJobCompleteParams; }
-namespace UE::DerivedData { struct FBuildActionKey; }
-namespace UE::DerivedData { struct FBuildKey; }
 namespace UE::DerivedData { struct FCacheKey; }
-namespace UE::DerivedData { enum class EBuildPolicy : uint8; }
 namespace UE::DerivedData { enum class EBuildStatus : uint32; }
 
 namespace UE::DerivedData
diff --git a/Engine/Source/Developer/DerivedDataCache/Public/DerivedDataBuildScheduler.h b/Engine/Source/Developer/DerivedDataCache/Public/DerivedDataBuildScheduler.h
--- a/Engine/Source/Developer/DerivedDataCache/Public/DerivedDataBuildScheduler.h
+++ b/Engine/Source/Developer/DerivedDataCache/Public/DerivedDataBuildScheduler.h
@@ -27,6 +27,8 @@ class IBuildScheduler
 {
 public:
 	virtual ~IBuildScheduler() = default;
+
+	/** Begin processing of the job by this scheduler. Always paired with IBuildJobSchedule::EndJob. */
 	virtual TUniquePtr<IBuildJobSchedule> BeginJob(IBuildJob& Job, IRequestOwner& Owner) = 0;
 };
 
@@ -49,7 +51,7 @@ struct FBuildSchedulerParams
 	uint64 TotalRequiredMemory = 0;
 };
 
-/** Scheduling interface / context for an IBuildJob */
+/** Scheduling interface and context for a build job. */
 class IBuildJobSchedule
 {
 public:
@@ -58,24 +60,25 @@ public:
 	virtual FBuildSchedulerParams& EditParameters() = 0;
 
 	/** Calls StepExecution() now or later. */
-	virtual void DispatchCacheQuery() = 0;
+	virtual void ScheduleCacheQuery() = 0;
 	/** Calls StepExecution() now or later. */
-	virtual void DispatchCacheStore() = 0;
+	virtual void ScheduleCacheStore() = 0;
 	/** Calls StepExecution() now or later. */
-	virtual void DispatchResolveKey() = 0;
+	virtual void ScheduleResolveKey() = 0;
 	/** Calls StepExecution() now or later. */
-	virtual void DispatchResolveInputMeta() = 0;
+	virtual void ScheduleResolveInputMeta() = 0;
 	/**
 	 * Calls StepExecution() or SkipExecuteRemote() now or later.
 	 *
 	 * SkipExecuteRemote() won't be called unless MissingRemoteInputsSize is non-zero.
 	 */
-	virtual void DispatchResolveInputData() = 0;
+	virtual void ScheduleResolveInputData() = 0;
 	/** Calls StepExecution() or SkipExecuteRemote() now or later. */
-	virtual void DispatchExecuteRemote() = 0;
+	virtual void ScheduleExecuteRemote() = 0;
 	/** Calls StepExecution() now or later. */
-	virtual void DispatchExecuteLocal()	= 0;
+	virtual void ScheduleExecuteLocal()	= 0;
 
+	/** End processing of the job by this scheduler. Always paired with IBuildScheduler::BeginJob. */
 	virtual void EndJob() = 0;
 };
 
