Added initial revision of CbPackage POST message formatting, along with fixes to Content-Type: header handling

Also improved conditionals in logging, to ensure all 2xx results are treated as success

diff --git a/Engine/Source/Developer/DerivedDataCache/Internal/ZenServerHttp.h b/Engine/Source/Developer/DerivedDataCache/Internal/ZenServerHttp.h
--- a/Engine/Source/Developer/DerivedDataCache/Internal/ZenServerHttp.h
+++ b/Engine/Source/Developer/DerivedDataCache/Internal/ZenServerHttp.h
@@ -112,7 +112,9 @@ public:
 	DERIVEDDATACACHE_API Result PerformBlockingDelete(FStringView Uri);
 
 	DERIVEDDATACACHE_API Result PerformBlockingPost(FStringView Uri, FCbObjectView Obj);
+	DERIVEDDATACACHE_API Result PerformBlockingPostPackage(FStringView Uri, FCbPackage Package);
 	DERIVEDDATACACHE_API Result PerformBlockingPost(FStringView Uri, FMemoryView Payload);
+	DERIVEDDATACACHE_API Result PerformBlockingPost(FStringView Uri, FMemoryView Payload, EContentType ContentType);
 
 	/**
 		* Returns the response buffer as a string. Note that is the request is performed
@@ -127,11 +129,13 @@ public:
 	  * Returns the response buffer. Note that is the request is performed
 	  * with an external buffer as target buffer this will be empty.
 	  */
-	const TArray64<uint8>& GetResponseBuffer() const
+	inline const TArray64<uint8>& GetResponseBuffer() const
 	{
 		return ResponseBuffer;
 	}
 
+	DERIVEDDATACACHE_API FCbObjectView GetResponseAsObject() const;
+
 	DERIVEDDATACACHE_API bool GetHeader(const ANSICHAR* Header, FString& OutValue) const;
 
 private:
@@ -153,6 +157,8 @@ private:
 	TArray<FString>			Headers;
 	FString					Domain;
 
+	void AddHeadersForContentType(EContentType ContentType);
+
 	/**
 	  * Supported request verb
 	  */
diff --git a/Engine/Source/Developer/DerivedDataCache/Private/ZenServerHttp.cpp b/Engine/Source/Developer/DerivedDataCache/Private/ZenServerHttp.cpp
--- a/Engine/Source/Developer/DerivedDataCache/Private/ZenServerHttp.cpp
+++ b/Engine/Source/Developer/DerivedDataCache/Private/ZenServerHttp.cpp
@@ -19,10 +19,13 @@
 #	include "Windows/HideWindowsPlatformTypes.h"
 #endif
 
+#include "Compression/CompressedBuffer.h"
+#include "Compression/OodleDataCompression.h"
 #include "Memory/CompositeBuffer.h"
 #include "Serialization/CompactBinary.h"
 #include "Serialization/CompactBinaryPackage.h"
 #include "Serialization/CompactBinaryValidation.h"
+#include "Serialization/CompactBinaryWriter.h"
 #include "Serialization/LargeMemoryReader.h"
 #include "Serialization/LargeMemoryWriter.h"
 
@@ -95,16 +98,8 @@ namespace UE::Zen {
 #endif
 	}
 
-	FZenHttpRequest::Result FZenHttpRequest::PerformBlockingPut(const TCHAR* Uri, const FCompositeBuffer& Buffer, EContentType ContentType)
+	void FZenHttpRequest::AddHeadersForContentType(EContentType ContentType)
 	{
-		uint32 ContentLength = 0u;
-
-		ContentLength = Buffer.GetSize();
-		curl_easy_setopt(Curl, CURLOPT_UPLOAD, 1L);
-		curl_easy_setopt(Curl, CURLOPT_INFILESIZE, ContentLength);
-		curl_easy_setopt(Curl, CURLOPT_READDATA, this);
-		curl_easy_setopt(Curl, CURLOPT_READFUNCTION, &FZenHttpRequest::FStatics::StaticReadFn);
-
 		switch (ContentType)
 		{
 		case EContentType::Binary:
@@ -120,6 +115,20 @@ namespace UE::Zen {
 			checkNoEntry();
 			break;
 		}
+	}
+
+	FZenHttpRequest::Result FZenHttpRequest::PerformBlockingPut(const TCHAR* Uri, const FCompositeBuffer& Buffer, EContentType ContentType)
+	{
+		uint32 ContentLength = 0u;
+
+		ContentLength = Buffer.GetSize();
+		curl_easy_setopt(Curl, CURLOPT_UPLOAD, 1L);
+		curl_easy_setopt(Curl, CURLOPT_INFILESIZE, ContentLength);
+		curl_easy_setopt(Curl, CURLOPT_READDATA, this);
+		curl_easy_setopt(Curl, CURLOPT_READFUNCTION, &FZenHttpRequest::FStatics::StaticReadFn);
+
+		AddHeadersForContentType(ContentType);
+
 		ReadDataView = &Buffer;
 
 		return PerformBlocking(Uri, RequestVerb::Put, ContentLength);
@@ -130,7 +139,77 @@ namespace UE::Zen {
 		FLargeMemoryWriter Out;
 		Obj.CopyTo(Out);
 
-		return PerformBlockingPost(Uri, Out.GetView());
+		return PerformBlockingPost(Uri, Out.GetView(), EContentType::CompactBinary);
+	}
+
+	FZenHttpRequest::Result FZenHttpRequest::PerformBlockingPostPackage(FStringView Uri, FCbPackage Package)
+	{
+		struct CbPackageHeader
+		{
+			uint32	HeaderMagic;
+			uint32	AttachmentCount;
+			uint32	Reserved1;
+			uint32	Reserved2;
+		};
+
+		static const uint32 kMagic = 0xaa77aacc;
+
+		struct CbAttachmentEntry
+		{
+			uint64	AttachmentSize;
+			uint32	Reserved1;
+			FIoHash	AttachmentHash;
+		};
+
+		TConstArrayView<FCbAttachment> Attachments = Package.GetAttachments();
+		FCbObject Object = Package.GetObject();
+		FCompressedBuffer ObjectBuffer = FCompressedBuffer::Compress(Object.GetBuffer(), FOodleDataCompression::ECompressor::NotSet, FOodleDataCompression::ECompressionLevel::None);
+
+		CbPackageHeader Hdr;
+		Hdr.HeaderMagic = kMagic;
+		Hdr.AttachmentCount = Attachments.Num();
+		Hdr.Reserved1 = 0;
+		Hdr.Reserved2 = 0;
+
+		FLargeMemoryWriter Out;
+		Out.Serialize(&Hdr, sizeof Hdr);
+
+		// Root object metadata
+
+		{
+			CbAttachmentEntry Entry;
+			Entry.AttachmentHash = ObjectBuffer.GetRawHash();
+			Entry.AttachmentSize = ObjectBuffer.GetCompressedSize();
+			Entry.Reserved1 = 0;
+
+			Out.Serialize(&Entry, sizeof Entry);
+		}
+
+		// Attachment metadata
+
+		for (const FCbAttachment& Attachment : Attachments)
+		{
+			CbAttachmentEntry Entry;
+			Entry.AttachmentHash = Attachment.GetHash();
+			Entry.AttachmentSize = Attachment.AsCompressedBinary().GetCompressedSize();
+			Entry.Reserved1 = 0;
+
+			Out.Serialize(&Entry, sizeof Entry);
+		}
+
+		// Root object
+
+		Out << ObjectBuffer;
+
+		// Payloads back-to-back
+
+		for (const FCbAttachment& Attachment : Attachments)
+		{
+			FCompressedBuffer Payload = Attachment.AsCompressedBinary();
+			Out << Payload;
+		}
+
+		return PerformBlockingPost(Uri, Out.GetView(), EContentType::CompactBinaryPackage);
 	}
 
 	FZenHttpRequest::Result FZenHttpRequest::PerformBlockingPost(FStringView Uri, FMemoryView Payload)
@@ -142,7 +221,24 @@ namespace UE::Zen {
 		curl_easy_setopt(Curl, CURLOPT_READDATA, this);
 		curl_easy_setopt(Curl, CURLOPT_READFUNCTION, &FZenHttpRequest::FStatics::StaticReadFn);
 
-		// TODO: proper Content-Type: header
+		ContentLength = Payload.GetSize();
+
+		FCompositeBuffer Buffer(FSharedBuffer::MakeView(FMemoryView{ reinterpret_cast<const uint8*>(Payload.GetData()), Payload.GetSize() }));
+		ReadDataView = &Buffer;
+
+		return PerformBlocking(Uri, RequestVerb::Post, ContentLength);
+	}
+
+	FZenHttpRequest::Result FZenHttpRequest::PerformBlockingPost(FStringView Uri, FMemoryView Payload, EContentType ContentType)
+	{
+		uint64 ContentLength = 0u;
+
+		curl_easy_setopt(Curl, CURLOPT_POST, 1L);
+		curl_easy_setopt(Curl, CURLOPT_INFILESIZE, Payload.GetSize());
+		curl_easy_setopt(Curl, CURLOPT_READDATA, this);
+		curl_easy_setopt(Curl, CURLOPT_READFUNCTION, &FZenHttpRequest::FStatics::StaticReadFn);
+
+		AddHeadersForContentType(ContentType);
 
 		ContentLength = Payload.GetSize();
 
@@ -200,6 +296,12 @@ namespace UE::Zen {
 
 	FZenHttpRequest::Result FZenHttpRequest::PerformBlocking(FStringView Uri, RequestVerb Verb, uint32 ContentLength)
 	{
+		// Strip any leading slashes because we compose the prefix and the suffix with a separating slash below
+		while (Uri.StartsWith('/'))
+		{
+			Uri.RightChopInline(1);
+		}
+
 		static const char* CommonHeaders[] = {
 			"User-Agent: UE",
 			nullptr
@@ -279,7 +381,7 @@ namespace UE::Zen {
 
 	void FZenHttpRequest::LogResult(long InResult, const TCHAR* Uri, RequestVerb Verb) const
 	{
-		CURLcode Result = (CURLcode) InResult;
+		CURLcode Result = (CURLcode)InResult;
 		if (Result == CURLE_OK)
 		{
 			bool bSuccess = false;
@@ -287,30 +389,30 @@ namespace UE::Zen {
 			FString AdditionalInfo;
 
 			const bool Is404 = (ResponseCode == 404);
-			const bool Is200 = (ResponseCode >= 200) && (ResponseCode <= 299);
+			const bool Is2xx = (ResponseCode >= 200) && (ResponseCode <= 299);
 
 			switch (Verb)
 			{
 			case RequestVerb::Head:
-				bSuccess = Is200 || Is404;
+				bSuccess = Is2xx || Is404;
 				VerbStr = TEXT("querying");
 				break;
 			case RequestVerb::Get:
-				bSuccess = Is200 || Is404;
+				bSuccess = Is2xx || Is404;
 				VerbStr = TEXT("fetching");
 				AdditionalInfo = FString::Printf(TEXT("Received: %d bytes."), BytesReceived);
 				break;
 			case RequestVerb::Put:
-				bSuccess = Is200;
+				bSuccess = Is2xx;
 				VerbStr = TEXT("updating");
 				AdditionalInfo = FString::Printf(TEXT("Sent: %d bytes."), BytesSent);
 				break;
 			case RequestVerb::Post:
-				bSuccess = Is200;
+				bSuccess = Is2xx || Is404;
 				VerbStr = TEXT("posting");
 				break;
 			case RequestVerb::Delete:
-				bSuccess = Is200 || Is404;
+				bSuccess = Is2xx || Is404;
 				VerbStr = TEXT("deleting");
 				break;
 			}
@@ -363,6 +465,11 @@ namespace UE::Zen {
 		return FString(TCHARData.Length(), TCHARData.Get());
 	}
 
+	FCbObjectView FZenHttpRequest::GetResponseAsObject() const
+	{
+		return FCbObjectView(ResponseBuffer.GetData());
+	}
+
 	size_t FZenHttpRequest::FStatics::StaticDebugCallback(CURL* Handle, curl_infotype DebugInfoType, char* DebugInfo, size_t DebugInfoSize, void* UserData)
 	{
 		FZenHttpRequest* Request = static_cast<FZenHttpRequest*>(UserData);
@@ -504,6 +611,12 @@ namespace UE::Zen {
 
 	FZenHttpRequestPool::FZenHttpRequestPool(FStringView InServiceUrl, uint32 PoolEntryCount)
 	{
+		int32 LastSlash = 0;
+		while (InServiceUrl.FindLastChar('/', /* out */ LastSlash) && (LastSlash == (InServiceUrl.Len() - 1)))
+		{
+			InServiceUrl.LeftChopInline(1);
+		}
+
 		Pool.SetNum(PoolEntryCount);
 
 		for (uint8 i = 0; i < Pool.Num(); ++i)
@@ -525,10 +638,10 @@ namespace UE::Zen {
 	}
 
 	/** Block until a request is free
-	  * 
-	  * Once a request has been returned it is owned by the caller and 
+	  *
+	  * Once a request has been returned it is owned by the caller and
 	  * it needs to release it to the pool when work has been completed
-	  * 
+	  *
 	  * @return Usable request instance.
 	  */
 	FZenHttpRequest* FZenHttpRequestPool::WaitForFreeRequest()
