Added code for Zen/HTTP marshaling of CbPackages. There is some more work to do to handle package responses properly but otherwise functional

#rnx

diff --git a/Engine/Source/Developer/DerivedDataCache/Internal/ZenServerHttp.h b/Engine/Source/Developer/DerivedDataCache/Internal/ZenServerHttp.h
--- a/Engine/Source/Developer/DerivedDataCache/Internal/ZenServerHttp.h
+++ b/Engine/Source/Developer/DerivedDataCache/Internal/ZenServerHttp.h
@@ -116,17 +116,15 @@ public:
 	DERIVEDDATACACHE_API Result PerformBlockingPost(FStringView Uri, FMemoryView Payload);
 	DERIVEDDATACACHE_API Result PerformBlockingPost(FStringView Uri, FMemoryView Payload, EContentType ContentType);
 
-	/**
-		* Returns the response buffer as a string. Note that is the request is performed
-		* with an external buffer as target buffer this string will be empty.
-		*/
+	/** Returns the response buffer as a string. Note that is the request is performed
+		with an external buffer as target buffer this string will be empty.
+	 */
 	inline FString GetResponseAsString() const
 	{
 		return GetAnsiBufferAsString(ResponseBuffer);
 	}
 
-	/**
-	  * Returns the response buffer. Note that is the request is performed
+	/** Returns the response buffer. Note that is the request is performed
 	  * with an external buffer as target buffer this will be empty.
 	  */
 	inline const TArray64<uint8>& GetResponseBuffer() const
@@ -135,6 +133,7 @@ public:
 	}
 
 	DERIVEDDATACACHE_API FCbObjectView GetResponseAsObject() const;
+	DERIVEDDATACACHE_API FCbPackage GetResponseAsPackage() const;
 
 	DERIVEDDATACACHE_API bool GetHeader(const ANSICHAR* Header, FString& OutValue) const;
 
@@ -153,7 +152,7 @@ private:
 	TArray64<uint8>*		WriteHeaderBufferPtr = nullptr;
 
 	TArray64<uint8>			ResponseHeader;
-	TArray64<uint8>			ResponseBuffer;
+	TArray64<uint8>			ResponseBuffer;	// If no other response buffer is set, this is where the response payload goes
 	TArray<FString>			Headers;
 	FString					Domain;
 
diff --git a/Engine/Source/Developer/DerivedDataCache/Private/ZenServerHttp.cpp b/Engine/Source/Developer/DerivedDataCache/Private/ZenServerHttp.cpp
--- a/Engine/Source/Developer/DerivedDataCache/Private/ZenServerHttp.cpp
+++ b/Engine/Source/Developer/DerivedDataCache/Private/ZenServerHttp.cpp
@@ -142,25 +142,32 @@ namespace UE::Zen {
 		return PerformBlockingPost(Uri, Out.GetView(), EContentType::CompactBinary);
 	}
 
-	FZenHttpRequest::Result FZenHttpRequest::PerformBlockingPostPackage(FStringView Uri, FCbPackage Package)
+	struct CbPackageHeader
 	{
-		struct CbPackageHeader
-		{
-			uint32	HeaderMagic;
-			uint32	AttachmentCount;
-			uint32	Reserved1;
-			uint32	Reserved2;
-		};
+		uint32	HeaderMagic;
+		uint32	AttachmentCount;
+		uint32	Reserved1;
+		uint32	Reserved2;
+	};
 
-		static const uint32 kMagic = 0xaa77aacc;
+	static const uint32 kMagic = 0xaa77aacc;
 
-		struct CbAttachmentEntry
+	struct CbAttachmentEntry
+	{
+		uint64	AttachmentSize;
+		uint32	Flags;
+		FIoHash	AttachmentHash;
+
+		enum
 		{
-			uint64	AttachmentSize;
-			uint32	Reserved1;
-			FIoHash	AttachmentHash;
+			IsCompressed = (1u << 0),	// Is marshaled using compressed buffer storage format
+			IsObject = (1u << 1),		// Is compact binary object
 		};
 
+	};
+
+	FZenHttpRequest::Result FZenHttpRequest::PerformBlockingPostPackage(FStringView Uri, FCbPackage Package)
+	{
 		TConstArrayView<FCbAttachment> Attachments = Package.GetAttachments();
 		FCbObject Object = Package.GetObject();
 		FCompressedBuffer ObjectBuffer = FCompressedBuffer::Compress(Object.GetBuffer(), FOodleDataCompression::ECompressor::NotSet, FOodleDataCompression::ECompressionLevel::None);
@@ -180,7 +187,7 @@ namespace UE::Zen {
 			CbAttachmentEntry Entry;
 			Entry.AttachmentHash = ObjectBuffer.GetRawHash();
 			Entry.AttachmentSize = ObjectBuffer.GetCompressedSize();
-			Entry.Reserved1 = 0;
+			Entry.Flags = CbAttachmentEntry::IsObject;
 
 			Out.Serialize(&Entry, sizeof Entry);
 		}
@@ -191,8 +198,30 @@ namespace UE::Zen {
 		{
 			CbAttachmentEntry Entry;
 			Entry.AttachmentHash = Attachment.GetHash();
-			Entry.AttachmentSize = Attachment.AsCompressedBinary().GetCompressedSize();
-			Entry.Reserved1 = 0;
+			Entry.Flags = 0;
+
+			if (Attachment.IsCompressedBinary())
+			{
+				Entry.AttachmentSize = Attachment.AsCompressedBinary().GetCompressedSize();
+				Entry.Flags |= CbAttachmentEntry::IsCompressed;
+			}
+			else if (Attachment.IsBinary())
+			{
+				Entry.AttachmentSize = Attachment.AsCompositeBinary().GetSize();
+			}
+			else if (Attachment.IsNull())
+			{
+				checkNoEntry();
+			}
+			else if (Attachment.IsObject())
+			{
+				Entry.AttachmentSize = Attachment.AsObject().GetSize();
+				Entry.Flags |= CbAttachmentEntry::IsObject;
+			}
+			else
+			{
+				checkNoEntry();
+			}
 
 			Out.Serialize(&Entry, sizeof Entry);
 		}
@@ -205,13 +234,78 @@ namespace UE::Zen {
 
 		for (const FCbAttachment& Attachment : Attachments)
 		{
-			FCompressedBuffer Payload = Attachment.AsCompressedBinary();
-			Out << Payload;
+			if (Attachment.IsCompressedBinary())
+			{
+				FCompressedBuffer Payload = Attachment.AsCompressedBinary();
+				Out << Payload;
+			}
+			else if (Attachment.IsBinary())
+			{
+				const FCompositeBuffer& Buffer = Attachment.AsCompositeBinary();
+				FSharedBuffer SharedBuffer = Buffer.ToShared();
+
+				Out.Serialize((void*)SharedBuffer.GetData(), SharedBuffer.GetSize());
+			}
+			else if (Attachment.IsNull())
+			{
+				checkNoEntry();
+			}
+			else if (Attachment.IsObject())
+			{
+				checkNoEntry();
+			}
+			else
+			{
+				checkNoEntry();
+			}
 		}
 
 		return PerformBlockingPost(Uri, Out.GetView(), EContentType::CompactBinaryPackage);
 	}
 
+	FCbPackage FZenHttpRequest::GetResponseAsPackage() const
+	{
+		const TArray64<uint8>& Response = GetResponseBuffer();
+		FLargeMemoryReader Reader(Response.GetData(), Response.Num());
+
+		FCbPackage Package;
+
+		CbPackageHeader Hdr;
+		Reader.Serialize(&Hdr, sizeof Hdr);
+
+		if (Hdr.HeaderMagic != kMagic)
+		{
+			return {};
+		}
+
+		TArray<CbAttachmentEntry> AttachmentEntries;
+		AttachmentEntries.SetNum(Hdr.AttachmentCount);
+
+		Reader.Serialize(AttachmentEntries.GetData(), Hdr.AttachmentCount * sizeof(CbAttachmentEntry));
+
+		int Index = 0;
+
+		for (const CbAttachmentEntry& Entry : AttachmentEntries)
+		{
+			FUniqueBuffer AttachmentData = FUniqueBuffer::Alloc(Entry.AttachmentSize);
+			Reader.Serialize(AttachmentData.GetData(), AttachmentData.GetSize());
+
+			if (Index == 0)
+			{
+
+			}
+			else
+			{
+				FCbAttachment Attachment(AttachmentData.MoveToShared());
+				Package.AddAttachment(Attachment);
+			}
+
+			++Index;
+		}
+
+		return Package;
+	}
+
 	FZenHttpRequest::Result FZenHttpRequest::PerformBlockingPost(FStringView Uri, FMemoryView Payload)
 	{
 		uint64 ContentLength = 0u;
