DDC: Improved handling of build cancellation

#rnx

diff --git a/Engine/Source/Developer/DerivedDataCache/Private/DerivedDataBuildJob.cpp b/Engine/Source/Developer/DerivedDataCache/Private/DerivedDataBuildJob.cpp
--- a/Engine/Source/Developer/DerivedDataCache/Private/DerivedDataBuildJob.cpp
+++ b/Engine/Source/Developer/DerivedDataCache/Private/DerivedDataBuildJob.cpp
@@ -457,6 +457,10 @@ void FBuildJob::BeginCacheQuery()
 
 void FBuildJob::EndCacheQuery(FCacheGetCompleteParams&& Params)
 {
+	if (Params.Status == EStatus::Canceled)
+	{
+		return CompleteWithError(TEXT("Build was canceled."));
+	}
 	if (Params.Status == EStatus::Ok)
 	{
 		if (FOptionalBuildOutput CacheOutput = FBuildOutput::Load(Name, FunctionName, Params.Record))
@@ -536,6 +540,10 @@ void FBuildJob::BeginResolveKey()
 
 void FBuildJob::EndResolveKey(FBuildKeyResolvedParams&& Params)
 {
+	if (Params.Status == EStatus::Canceled)
+	{
+		return CompleteWithError(TEXT("Build was canceled."));
+	}
 	if (Params.Status == EStatus::Ok && Params.Definition)
 	{
 		return SetDefinition(MoveTemp(Params.Definition).Get());
@@ -569,6 +577,10 @@ void FBuildJob::BeginResolveInputMeta()
 
 void FBuildJob::EndResolveInputMeta(FBuildInputMetaResolvedParams&& Params)
 {
+	if (Params.Status == EStatus::Canceled)
+	{
+		return CompleteWithError(TEXT("Build was canceled."));
+	}
 	if (Params.Status == EStatus::Ok)
 	{
 		return CreateAction(Params.Inputs);
@@ -648,6 +660,10 @@ void FBuildJob::BeginResolveInputData()
 
 void FBuildJob::EndResolveInputData(FBuildInputDataResolvedParams&& Params)
 {
+	if (Params.Status == EStatus::Canceled)
+	{
+		return CompleteWithError(TEXT("Build was canceled."));
+	}
 	if (Params.Status == EStatus::Ok)
 	{
 		FBuildSchedulerParams& SchedulerParams = Schedule->EditParameters();
@@ -709,6 +725,10 @@ void FBuildJob::BeginExecuteRemote()
 
 void FBuildJob::EndExecuteRemote(FBuildWorkerActionCompleteParams&& Params)
 {
+	if (Params.Status == EStatus::Canceled)
+	{
+		return CompleteWithError(TEXT("Build was canceled."));
+	}
 	if (Params.Output)
 	{
 		EnumAddFlags(BuildStatus, EBuildStatus::BuildRemote);
@@ -927,6 +947,10 @@ void FBuildJob::AdvanceToState(EBuildJobState NewState)
 		if (Owner.IsCanceled())
 		{
 			NewState = EBuildJobState::Complete;
+			if (Output.IsNull())
+			{
+				OutputBuilder.AddError(TEXT("LogDerivedDataBuild"_SV), TEXT("Build was canceled."_SV));
+			}
 		}
 
 		NextState = NewState;
