ZenDDC: Remove the use of Save/LoadCbPackage

#rnx

diff --git a/Engine/Source/Developer/DerivedDataCache/Internal/ZenServerHttp.h b/Engine/Source/Developer/DerivedDataCache/Internal/ZenServerHttp.h
--- a/Engine/Source/Developer/DerivedDataCache/Internal/ZenServerHttp.h
+++ b/Engine/Source/Developer/DerivedDataCache/Internal/ZenServerHttp.h
@@ -157,6 +157,7 @@ private:
 	FString					Domain;
 
 	void AddHeadersForContentType(EContentType ContentType);
+	void AddHeadersForAcceptType(EContentType ContentType);
 
 	/**
 	  * Supported request verb
diff --git a/Engine/Source/Developer/DerivedDataCache/Private/ZenDerivedDataBackend.cpp b/Engine/Source/Developer/DerivedDataCache/Private/ZenDerivedDataBackend.cpp
--- a/Engine/Source/Developer/DerivedDataCache/Private/ZenDerivedDataBackend.cpp
+++ b/Engine/Source/Developer/DerivedDataCache/Private/ZenDerivedDataBackend.cpp
@@ -14,6 +14,7 @@
 #include "Serialization/CompactBinaryPackage.h"
 #include "Serialization/CompactBinaryValidation.h"
 #include "Serialization/CompactBinaryWriter.h"
+#include "Serialization/BufferArchive.h"
 #include "ZenBackendUtils.h"
 #include "ZenSerialization.h"
 #include "ZenServerHttp.h"
@@ -722,14 +723,13 @@ bool FZenDerivedDataBackend::PutCacheRecord(const FCacheRecord& Record, FStringV
 {
 	const FCacheKey& Key = Record.GetKey();
 	FCbPackage Package = Record.Save();
-	FCbWriter PackageWriter;
-	UE::Zen::SaveCbPackage(Package, PackageWriter);
-	FUniqueBuffer Buffer = FUniqueBuffer::Alloc(PackageWriter.GetSaveSize());
-	PackageWriter.Save(Buffer);
+	FBufferArchive Ar;
+	Package.Save(Ar);
+	FCompositeBuffer Buffer = FCompositeBuffer(FSharedBuffer::MakeView(Ar.GetData(), Ar.Num()));
 	TStringBuilder<256> Uri;
 	AppendZenUri(Record.GetKey(), Uri);
 	AppendPolicyQueryString(Policy, Uri);
-	if (PutZenData(Uri.ToString(), FCompositeBuffer(Buffer.MoveToShared()), Zen::EContentType::CompactBinaryPackage)
+	if (PutZenData(Uri.ToString(), Buffer, Zen::EContentType::CompactBinaryPackage)
 		!= FDerivedDataBackendInterface::EPutStatus::Cached)
 	{
 		return false;
diff --git a/Engine/Source/Developer/DerivedDataCache/Private/ZenServerHttp.cpp b/Engine/Source/Developer/DerivedDataCache/Private/ZenServerHttp.cpp
--- a/Engine/Source/Developer/DerivedDataCache/Private/ZenServerHttp.cpp
+++ b/Engine/Source/Developer/DerivedDataCache/Private/ZenServerHttp.cpp
@@ -117,6 +117,25 @@ namespace UE::Zen {
 		}
 	}
 
+	void FZenHttpRequest::AddHeadersForAcceptType(EContentType ContentType)
+	{
+		switch (ContentType)
+		{
+		case EContentType::Binary:
+			Headers.Add(FString(TEXT("Accept: application/octet-stream")));
+			break;
+		case EContentType::CompactBinary:
+			Headers.Add(FString(TEXT("Accept: application/x-ue-cb")));
+			break;
+		case EContentType::CompactBinaryPackage:
+			Headers.Add(FString(TEXT("Accept: application/x-ue-cbpkg")));
+			break;
+		default:
+			checkNoEntry();
+			break;
+		}
+	}
+
 	FZenHttpRequest::Result FZenHttpRequest::PerformBlockingPut(const TCHAR* Uri, const FCompositeBuffer& Buffer, EContentType ContentType)
 	{
 		uint32 ContentLength = 0u;
@@ -354,21 +373,17 @@ namespace UE::Zen {
 	{
 		curl_easy_setopt(Curl, CURLOPT_HTTPGET, 1L);
 		OutPackage.Reset();
+
+		AddHeadersForAcceptType(EContentType::CompactBinaryPackage);
+
 		// TODO: When PackageBytes can be written in segments directly, set the WritePtr to the OutPackage and use that
 		TArray64<uint8> PackageBytes;
 		WriteDataBufferPtr = &PackageBytes;
 		Result LocalResult = PerformBlocking(Uri, RequestVerb::Get, 0u);
 		if (IsSuccessCode(ResponseCode))
 		{
-			if (ValidateCompactBinaryPackage(FMemoryView(PackageBytes.GetData(), PackageBytes.Num()), ECbValidateMode::All) != ECbValidateError::None)
-			{
-				bResponseFormatValid = false;
-			}
-			else
-			{
-				FLargeMemoryReader PackageLoader(PackageBytes.GetData(), PackageBytes.Num());
-				bResponseFormatValid = UE::Zen::TryLoadCbPackage(OutPackage, PackageLoader);
-			}
+			FLargeMemoryReader Ar(PackageBytes.GetData(), PackageBytes.Num());
+			bResponseFormatValid = OutPackage.TryLoad(Ar);
 		}
 		return LocalResult;
 	}
