Package marshaling now properly flags the root object as compressed

#rnx

diff --git a/Engine/Source/Developer/DerivedDataCache/Private/ZenServerHttp.cpp b/Engine/Source/Developer/DerivedDataCache/Private/ZenServerHttp.cpp
--- a/Engine/Source/Developer/DerivedDataCache/Private/ZenServerHttp.cpp
+++ b/Engine/Source/Developer/DerivedDataCache/Private/ZenServerHttp.cpp
@@ -206,7 +206,7 @@ namespace UE::Zen {
 			CbAttachmentEntry Entry;
 			Entry.AttachmentHash = ObjectBuffer.GetRawHash();
 			Entry.AttachmentSize = ObjectBuffer.GetCompressedSize();
-			Entry.Flags = CbAttachmentEntry::IsObject;
+			Entry.Flags = CbAttachmentEntry::IsObject | CbAttachmentEntry::IsCompressed;
 
 			Out.Serialize(&Entry, sizeof Entry);
 		}
