DDC: Added ECachePolicy::KeepAlive and EBuildPolicy::CacheKeepAlive

The new flags are meant to prevent a cache store from evicting the associated cache records for the duration of the session.

#rnx

diff --git a/Engine/Source/Developer/DerivedDataCache/Private/DerivedDataBuildJob.cpp b/Engine/Source/Developer/DerivedDataCache/Private/DerivedDataBuildJob.cpp
--- a/Engine/Source/Developer/DerivedDataCache/Private/DerivedDataBuildJob.cpp
+++ b/Engine/Source/Developer/DerivedDataCache/Private/DerivedDataBuildJob.cpp
@@ -190,9 +190,6 @@ private:
 	/** Policy requested and/or configured for the job. */
 	EBuildPolicy BuildPolicy{EBuildPolicy::None};
 
-	/** True if AdvanceToState is executing. */
-	bool bInAdvanceToState{};
-
 	/** Available in [ResolveKey, Complete] for jobs created from a key or definition. */
 	FBuildKey DefinitionKey;
 	/** Available in [ResolveAction, ExecuteLocal) for jobs created from a key or definition. */
@@ -210,6 +207,7 @@ private:
 	/** Available in [CacheQuery, CacheStoreWait). Context used by IBuildFunction. */
 	TRefCountPtr<FBuildJobContext> Context;
 
+	/** Available in (NotStarted, Complete). */
 	TUniquePtr<IBuildJobSchedule> Schedule;
 
 	/** Lock to synchronize writes to State, NextState, bInAdvanceToState. */
@@ -232,6 +230,9 @@ private:
 
 	ICache& Cache;
 	IBuild& BuildSystem;
+
+	/** True if AdvanceToState is executing. */
+	bool bInAdvanceToState{};
 };
 
 ///////////////////////////////////////////////////////////////////////////////////////////////////////////////////////
diff --git a/Engine/Source/Developer/DerivedDataCache/Private/DerivedDataBuildJobContext.cpp b/Engine/Source/Developer/DerivedDataCache/Private/DerivedDataBuildJobContext.cpp
--- a/Engine/Source/Developer/DerivedDataCache/Private/DerivedDataBuildJobContext.cpp
+++ b/Engine/Source/Developer/DerivedDataCache/Private/DerivedDataBuildJobContext.cpp
@@ -30,6 +30,10 @@ FBuildJobContext::FBuildJobContext(
 	, CachePolicy(ECachePolicy::Default)
 	, BuildPolicy(InBuildPolicy)
 {
+	if (EnumHasAnyFlags(BuildPolicy, EBuildPolicy::CacheKeepAlive))
+	{
+		EnumAddFlags(CachePolicy, ECachePolicy::KeepAlive);
+	}
 }
 
 FStringView FBuildJobContext::GetName() const
@@ -158,10 +162,14 @@ void FBuildJobContext::SetCacheBucket(FCacheBucket Bucket)
 
 void FBuildJobContext::SetCachePolicy(ECachePolicy Policy)
 {
-	checkf(!EnumHasAnyFlags(Policy, ECachePolicy::SkipData),
-		TEXT("SkipData flags not allowed on the cache policy for build of '%s' by %s. ")
+	checkf(!EnumHasAnyFlags(CachePolicy ^ Policy, ECachePolicy::SkipData),
+		TEXT("SkipData flags may not be modified on the cache policy for build of '%s' by %s. ")
 		TEXT("Flags for skipping data may be set indirectly through EBuildPolicy."),
 		*WriteToString<128>(Job.GetName()), *WriteToString<32>(Job.GetFunction()));
+	checkf(!EnumHasAnyFlags(CachePolicy ^ Policy, ECachePolicy::KeepAlive),
+		TEXT("KeepAlive flag may not be modified on the cache policy for build of '%s' by %s. ")
+		TEXT("Flags for cache record lifetime may be set indirectly through EBuildPolicy."),
+		*WriteToString<128>(Job.GetName()), *WriteToString<32>(Job.GetFunction()));
 	CachePolicy = Policy;
 }
 
@@ -171,8 +179,12 @@ void FBuildJobContext::SetBuildPolicy(EBuildPolicy Policy)
 		TEXT("Cache flags may not be modified on the build policy for build of '%s' by %s. ")
 		TEXT("Flags for modifying cache operations may be set through ECachePolicy."),
 		*WriteToString<128>(Job.GetName()), *WriteToString<32>(Job.GetFunction()));
+	checkf(!EnumHasAnyFlags(BuildPolicy ^ Policy, EBuildPolicy::CacheKeepAlive),
+		TEXT("CacheKeepAlive flag may not be modified on the build policy for build of '%s' by %s. ")
+		TEXT("Flags for cache record lifetime may only be set through the build session."),
+		*WriteToString<128>(Job.GetName()), *WriteToString<32>(Job.GetFunction()));
 	checkf(!EnumHasAnyFlags(BuildPolicy ^ Policy, EBuildPolicy::SkipData),
-		TEXT("SkipData flag may not be modified on the build policy for build of '%s' by %s. ")
+		TEXT("SkipData flags may not be modified on the build policy for build of '%s' by %s. ")
 		TEXT("Flags for skipping the data may only be set through the build session."),
 		*WriteToString<128>(Job.GetName()), *WriteToString<32>(Job.GetFunction()));
 	BuildPolicy = Policy;
diff --git a/Engine/Source/Developer/DerivedDataCache/Private/DerivedDataBuildJobContext.h b/Engine/Source/Developer/DerivedDataCache/Private/DerivedDataBuildJobContext.h
--- a/Engine/Source/Developer/DerivedDataCache/Private/DerivedDataBuildJobContext.h
+++ b/Engine/Source/Developer/DerivedDataCache/Private/DerivedDataBuildJobContext.h
@@ -18,8 +18,8 @@ namespace UE::DerivedData { class FBuildOutputBuilder; }
 namespace UE::DerivedData { class FPayload; }
 namespace UE::DerivedData { class IBuildJob; }
 namespace UE::DerivedData { class IRequestOwner; }
-namespace UE::DerivedData { enum class EBuildPolicy : uint8; }
-namespace UE::DerivedData { enum class ECachePolicy : uint8; }
+namespace UE::DerivedData { enum class EBuildPolicy : uint32; }
+namespace UE::DerivedData { enum class ECachePolicy : uint32; }
 namespace UE::DerivedData { enum class EPriority : uint8; }
 
 namespace UE::DerivedData::Private
diff --git a/Engine/Source/Developer/DerivedDataCache/Private/DerivedDataBuildPrivate.h b/Engine/Source/Developer/DerivedDataCache/Private/DerivedDataBuildPrivate.h
--- a/Engine/Source/Developer/DerivedDataCache/Private/DerivedDataBuildPrivate.h
+++ b/Engine/Source/Developer/DerivedDataCache/Private/DerivedDataBuildPrivate.h
@@ -27,7 +27,7 @@ namespace UE::DerivedData { class ICache; }
 namespace UE::DerivedData { class IRequestOwner; }
 namespace UE::DerivedData { struct FBuildJobCompleteParams; }
 namespace UE::DerivedData { struct FBuildKey; }
-namespace UE::DerivedData { enum class EBuildPolicy : uint8; }
+namespace UE::DerivedData { enum class EBuildPolicy : uint32; }
 namespace UE::DerivedData { using FOnBuildJobComplete = TUniqueFunction<void (FBuildJobCompleteParams&& Params)>; }
 
 namespace UE::DerivedData::Private
diff --git a/Engine/Source/Developer/DerivedDataCache/Public/DerivedDataBuildFunction.h b/Engine/Source/Developer/DerivedDataCache/Public/DerivedDataBuildFunction.h
--- a/Engine/Source/Developer/DerivedDataCache/Public/DerivedDataBuildFunction.h
+++ b/Engine/Source/Developer/DerivedDataCache/Public/DerivedDataBuildFunction.h
@@ -14,8 +14,8 @@ namespace UE::DerivedData { class FBuildContext; }
 namespace UE::DerivedData { class FCacheBucket; }
 namespace UE::DerivedData { class FPayload; }
 namespace UE::DerivedData { struct FPayloadId; }
-namespace UE::DerivedData { enum class EBuildPolicy : uint8; }
-namespace UE::DerivedData { enum class ECachePolicy : uint8; }
+namespace UE::DerivedData { enum class EBuildPolicy : uint32; }
+namespace UE::DerivedData { enum class ECachePolicy : uint32; }
 
 namespace UE::DerivedData
 {
@@ -140,13 +140,23 @@ public:
 	/** Returns the cache policy used when reading or writing this build in the cache. */
 	virtual ECachePolicy GetCachePolicy() const = 0;
 
-	/** Overrides the cache policy used when reading or writing this build in the cache. */
+	/**
+	 * Overrides the cache policy used when reading or writing this build in the cache.
+	 *
+	 * Always set the cache policy by modifying the policy returned by GetCachePolicy(), to maintain
+	 * the state of other flags. Changes to Skip* and KeepAlive flags are not allowed.
+	 */
 	virtual void SetCachePolicy(ECachePolicy Policy) = 0;
 
 	/** Returns the build policy used when executing this build. */
 	virtual EBuildPolicy GetBuildPolicy() const = 0;
 
-	/** Overrides the build policy used when executing this build. */
+	/**
+	 * Overrides the build policy used when executing this build.
+	 *
+	 * Always set the build policy by modifying the policy returned by GetBuildPolicy(), to maintain
+	 * the state of other flags. Changes to Cache* and Skip* flags are not allowed.
+	 */
 	virtual void SetBuildPolicy(EBuildPolicy Policy) = 0;
 
 	/** Sets the estimated peak memory required to execute the build, excluding input size. */
diff --git a/Engine/Source/Developer/DerivedDataCache/Public/DerivedDataBuildTypes.h b/Engine/Source/Developer/DerivedDataCache/Public/DerivedDataBuildTypes.h
--- a/Engine/Source/Developer/DerivedDataCache/Public/DerivedDataBuildTypes.h
+++ b/Engine/Source/Developer/DerivedDataCache/Public/DerivedDataBuildTypes.h
@@ -21,7 +21,7 @@ namespace UE::DerivedData
  * BuildPayload(Build | CacheStore): never get from the cache; always build; return one payload.
  * BuildPayload(Cache | CacheSkipData): allow cache get; never build; skip every payload.
  */
-enum class EBuildPolicy : uint8
+enum class EBuildPolicy : uint32
 {
 	/** A value without any flags set. */
 	None            = 0,
@@ -40,8 +40,11 @@ enum class EBuildPolicy : uint8
 	/** Allow a cache query and a cache store for the build. */
 	Cache           = CacheQuery | CacheStore,
 
+	/** Keep records in the cache for at least the duration of the session. */
+	CacheKeepAlive  = 1 << 4,
+
 	/** Skip fetching the payload data from the cache. */
-	SkipData        = 1 << 4,
+	SkipData        = 1 << 5,
 
 	/** Allow cache query+store, allow local+remote build when missed or skipped, and fetch the payload(s). */
 	Default         = Build | Cache,
diff --git a/Engine/Source/Developer/DerivedDataCache/Public/DerivedDataBuildWorker.h b/Engine/Source/Developer/DerivedDataCache/Public/DerivedDataBuildWorker.h
--- a/Engine/Source/Developer/DerivedDataCache/Public/DerivedDataBuildWorker.h
+++ b/Engine/Source/Developer/DerivedDataCache/Public/DerivedDataBuildWorker.h
@@ -24,7 +24,7 @@ namespace UE::DerivedData { class IBuild; }
 namespace UE::DerivedData { class IRequestOwner; }
 namespace UE::DerivedData { struct FBuildWorkerActionCompleteParams; }
 namespace UE::DerivedData { struct FBuildWorkerFileDataCompleteParams; }
-namespace UE::DerivedData { enum class EBuildPolicy : uint8; }
+namespace UE::DerivedData { enum class EBuildPolicy : uint32; }
 
 namespace UE::DerivedData
 {
diff --git a/Engine/Source/Developer/DerivedDataCache/Public/DerivedDataCache.h b/Engine/Source/Developer/DerivedDataCache/Public/DerivedDataCache.h
--- a/Engine/Source/Developer/DerivedDataCache/Public/DerivedDataCache.h
+++ b/Engine/Source/Developer/DerivedDataCache/Public/DerivedDataCache.h
@@ -45,7 +45,7 @@ using FOnCacheGetPayloadComplete = TUniqueFunction<void (FCacheGetPayloadComplet
  * Put(Remote): write to every writable remote cache, skipping local caches.
  * Put(Local): write to every writable local cache, skipping remote caches.
  */
-enum class ECachePolicy : uint8
+enum class ECachePolicy : uint32
 {
 	/** A value without any flags set. */
 	None            = 0,
@@ -76,6 +76,9 @@ enum class ECachePolicy : uint8
 	/** Skip copying records to local caches for get requests if they already exist locally. */
 	SkipLocalCopy   = 1 << 7,
 
+	/** Keep records in the cache for at least the duration of the session. */
+	KeepAlive       = 1 << 8,
+
 	/** Allow cache requests to query and store records in local caches. */
 	Local           = QueryLocal | StoreLocal,
 	/** Allow cache requests to query and store records in remote caches. */
diff --git a/Engine/Source/Developer/Virtualization/Private/VirtualizationDDCBackend.h b/Engine/Source/Developer/Virtualization/Private/VirtualizationDDCBackend.h
--- a/Engine/Source/Developer/Virtualization/Private/VirtualizationDDCBackend.h
+++ b/Engine/Source/Developer/Virtualization/Private/VirtualizationDDCBackend.h
@@ -5,7 +5,7 @@
 #include "Virtualization/IVirtualizationBackend.h"
 #include "DerivedDataCacheKey.h"
 
-namespace UE::DerivedData { enum class ECachePolicy : uint8; }
+namespace UE::DerivedData { enum class ECachePolicy : uint32; }
 
 namespace UE::Virtualization
 {
