DDC2: Use compressed buffer as attachment

#rnx

diff --git a/Engine/Source/Developer/DerivedDataCache/Private/DerivedDataCacheRecord.cpp b/Engine/Source/Developer/DerivedDataCache/Private/DerivedDataCacheRecord.cpp
--- a/Engine/Source/Developer/DerivedDataCache/Private/DerivedDataCacheRecord.cpp
+++ b/Engine/Source/Developer/DerivedDataCache/Private/DerivedDataCacheRecord.cpp
@@ -297,12 +297,11 @@ FCbPackage FCacheRecord::Save() const
 	{
 		Writer.BeginObject();
 		Writer.AddObjectId("Id"_ASV, Payload.GetId());
-		Writer.AddHash("RawHash"_ASV, Payload.GetRawHash());
 		Writer.AddInteger("RawSize"_ASV, Payload.GetRawSize());
-		const FCompositeBuffer CompressedBuffer = Payload.GetData().GetCompressed();
-		FCbAttachment AttachedCompressed(CompressedBuffer.ToShared());
-		Writer.AddAttachment("CompressedHash"_ASV, AttachedCompressed);
-		Package.AddAttachment(AttachedCompressed);
+		const FCompressedBuffer& CompressedBuffer = Payload.GetData();
+		FCbAttachment Attachment(CompressedBuffer);
+		Writer.AddAttachment("RawHash"_ASV, Attachment); // Will use the uncompressed raw hash
+		Package.AddAttachment(Attachment);
 		Writer.EndObject();
 	};
 	if (const FPayload& Value = GetValuePayload())
@@ -353,21 +352,20 @@ FOptionalCacheRecord FCacheRecord::Load(const FCbPackage& Package)
 
 	auto LoadPayload = [&Package](const FCbObjectView& PayloadObject)
 	{
-		const FPayloadId PayloadId = PayloadObject["Id"].AsObjectId();
+		const FPayloadId PayloadId = PayloadObject["Id"_ASV].AsObjectId();
 		if (PayloadId.IsNull())
 		{
 			return FPayload();
 		}
-
-		if (const FCbAttachment* Attachment = Package.FindAttachment(PayloadObject["CompressedHash"].AsHash()))
+		const FIoHash RawHash = PayloadObject["RawHash"_ASV].AsHash();
+		if (const FCbAttachment* Attachment = Package.FindAttachment(RawHash))
 		{
-			if (FCompressedBuffer Compressed = FCompressedBuffer::FromCompressed(Attachment->AsBinary()))
+			if (const FCompressedBuffer& Compressed = Attachment->AsCompressedBinary())
 			{
 				return FPayload(PayloadId, Compressed);
 			}
 		}
-		FIoHash RawHash = PayloadObject["RawHash"].AsHash();
-		uint64 RawSize = PayloadObject["RawSize"].AsUInt64(MAX_uint64);
+		const uint64 RawSize = PayloadObject["RawSize"_ASV].AsUInt64(MAX_uint64);
 		if (!RawHash.IsZero() && RawSize != MAX_uint64)
 		{
 			return FPayload(PayloadId, RawHash, RawSize);
