DDC: Made CacheRecord serialized field order consistent with other types

#rnx

diff --git a/Engine/Source/Developer/DerivedDataCache/Private/DerivedDataCacheRecord.cpp b/Engine/Source/Developer/DerivedDataCache/Private/DerivedDataCacheRecord.cpp
--- a/Engine/Source/Developer/DerivedDataCache/Private/DerivedDataCacheRecord.cpp
+++ b/Engine/Source/Developer/DerivedDataCache/Private/DerivedDataCacheRecord.cpp
@@ -295,13 +295,12 @@ FCbPackage FCacheRecord::Save() const
 	}
 	auto SavePayload = [&Package, &Writer](const FPayload& Payload)
 	{
+		FCbAttachment Attachment(Payload.GetData());
+		Package.AddAttachment(Attachment);
 		Writer.BeginObject();
 		Writer.AddObjectId("Id"_ASV, Payload.GetId());
+		Writer.AddAttachment("RawHash"_ASV, Attachment);
 		Writer.AddInteger("RawSize"_ASV, Payload.GetRawSize());
-		const FCompressedBuffer& CompressedBuffer = Payload.GetData();
-		FCbAttachment Attachment(CompressedBuffer);
-		Writer.AddAttachment("RawHash"_ASV, Attachment); // Will use the uncompressed raw hash
-		Package.AddAttachment(Attachment);
 		Writer.EndObject();
 	};
 	if (const FPayload& Value = GetValuePayload())
