Allow ZenServer to be autolaunched by editor use cases that require its presence without insisting that it always be autolaunched.

diff --git a/Engine/Config/BaseEngine.ini b/Engine/Config/BaseEngine.ini
--- a/Engine/Config/BaseEngine.ini
+++ b/Engine/Config/BaseEngine.ini
@@ -1998,6 +1998,17 @@ HighDetail=(NumLODs=6,MaxNumStreamedLODs=0,bSupportLODStreaming=0,LODPercentTria
 +BucketConfigs=(Name="DefaultBucket", Size=10000000)
 PathToRAC=RuntimeAssetCache
 
+[Zen]
+AutoLaunch=true
+
+[Zen.AutoLaunch]
+Hidden=false
+DesiredPort=1337
+WorkspaceDataPath=%ENGINEVERSIONAGNOSTICUSERDIR%Zen\Workspace
+DataPath=%ENGINEVERSIONAGNOSTICUSERDIR%Zen
+DataPathEnvOverride=UE-LocalDataCachePath
+DataPathEditorOverrideSetting=LocalDerivedDataCache
+
 ; Derived Data backend graphs
 ; Each of the below backend graph sections contains a set of nodes used to create derived data backed graph.
 ; DerivedDataBackendGraph is the default graph for source builds, and InstalledDerivedDataBackendGraph is the default for installed builds.
@@ -2047,14 +2058,14 @@ EnginePak=(Type=ReadPak, Filename=%ENGINEDIR%DerivedDataCache/DDC.ddp)
 MinimumDaysToKeepFile=7
 Root=(Type=KeyLength, Length=120, Inner=AsyncPut)
 AsyncPut=(Type=AsyncPut, Inner=Local)
-Local=(Type=Zen, Host="http://localhost:1337", ReadOnly=false, Namespace="ue4.ddc")
+Local=(Type=Zen, ReadOnly=false, Namespace="ue4.ddc")
 
 [ZenDDC]
 MinimumDaysToKeepFile=7
 Root=(Type=KeyLength, Length=120, Inner=AsyncPut)
 AsyncPut=(Type=AsyncPut, Inner=Hierarchy)
 Hierarchy=(Type=Hierarchical, Inner=Local, Inner=Shared)
-Local=(Type=Zen, Host="http://localhost:1337", ReadOnly=false, Namespace="ue4.ddc")
+Local=(Type=Zen, ReadOnly=false, Namespace="ue4.ddc")
 Shared=(Type=FileSystem, ReadOnly=false, Clean=false, Flush=false, DeleteUnused=true, UnusedFileAge=10, FoldersToClean=10, MaxFileChecksPerSec=1, ConsiderSlowAt=70, PromptIfMissing=false, Path=?EpicDDC, EnvPathOverride=UE-SharedDataCachePath, EditorOverrideSetting=SharedDerivedDataCache, CommandLineOverride=SharedDataCachePath)
 
 [DerivedDataBackendGraph_Fill_Seattle]
diff --git a/Engine/Source/Developer/DerivedDataCache/Internal/ZenServerInterface.h b/Engine/Source/Developer/DerivedDataCache/Internal/ZenServerInterface.h
--- a/Engine/Source/Developer/DerivedDataCache/Internal/ZenServerInterface.h
+++ b/Engine/Source/Developer/DerivedDataCache/Internal/ZenServerInterface.h
@@ -2,11 +2,6 @@
 
 #pragma once
 
-#include "Containers/StaticArray.h"
-
-class FCompositeBuffer;
-class FCbPackage;
-
 #ifndef UE_WITH_ZEN
 #	if PLATFORM_WINDOWS
 #		define UE_WITH_ZEN 1
@@ -17,20 +12,113 @@ class FCbPackage;
 
 #if UE_WITH_ZEN
 
+#include "Containers/StringView.h"
+#include "Containers/UnrealString.h"
+#include "Templates/UniquePtr.h"
+
+#define UE_API DERIVEDDATACACHE_API
+
 namespace UE::Zen {
 
+class FZenServiceInstance;
+
+/**
+ * Type used to declare usage of a Zen server instance whether the shared default instance or a unique non-default instance.
+ * Used to help manage launch, and optionally in the future, shutdown a shared default instance.  Use the default constructor
+ * to reference the default instance (which may be launched on demand), or use the non-default constructors with a specific
+ * URL or HostName/Port pair which is required to pre-exist (will not be auto-launched).
+ */
+class FScopeZenService
+{
+public:
+	UE_API FScopeZenService()
+	: FScopeZenService(FStringView())
+	{
+	}
+	UE_API FScopeZenService(FStringView InstanceURL);
+	UE_API FScopeZenService(FStringView InstanceHostName, uint16 InstancePort);
+
+	UE_API const FZenServiceInstance& GetInstance() const { return *ServiceInstance; }
+	UE_API FZenServiceInstance& GetInstance() { return *ServiceInstance; }
+
+private:
+	FZenServiceInstance* ServiceInstance;
+	TUniquePtr<UE::Zen::FZenServiceInstance> UniqueNonDefaultInstance;
+};
+
+/**
+ * Gets the default Zen service instance.  The default instance is configured through ZenServiceInstance INI section.
+ * The default instance can (depending on configuration):
+ *  - Auto-launched on demand (optionally elevated on Windows)
+ *  - Be copied out of the workspace tree before execution
+ *  - Shared between multiple tools running concurrently (implemented by launching multiple instances and expecting them to communicate and shutdown as needed)
+ *  - Instigate self-shutdown when all processes that requested it have terminated
+ *  - Use a subdirectory of the existing local DDC cache path as Zen's data path
+ *  - Use an upstream Zen service
+ *  - Be overridden by commandline arguments to reference an existing running instance instead of auto-launching a new instance.
+ * Note that no assumptions should be made about the hostname/port of the default instance.  Calling code should instead expect
+ * that the instance is authoritative over the hostname/port that it will end up using and should query that information from the
+ * instance as needed.
+ */
+UE_API FZenServiceInstance& GetDefaultServiceInstance();
+
+/**
+ * A representation of a Zen service instance.  Generally not accessed directly, but via FScopeZenService.
+ */
 class FZenServiceInstance
 {
 public:
-	 DERIVEDDATACACHE_API FZenServiceInstance();
-	 DERIVEDDATACACHE_API ~FZenServiceInstance();
+	 UE_API FZenServiceInstance()
+	 : FZenServiceInstance(FStringView())
+	 {
+	 }
+	 UE_API FZenServiceInstance(FStringView InstanceURL);
+	 UE_API ~FZenServiceInstance();
 
-	 DERIVEDDATACACHE_API bool IsServiceRunning();
-	 DERIVEDDATACACHE_API bool IsServiceReady();
+	 inline const TCHAR* GetURL() const { return *URL; }
+	 inline const TCHAR* GetHostName() const { return *HostName; }
+	 inline uint16 GetPort() const { return Port; }
+
+	 UE_API bool IsServiceRunning();
+	 UE_API bool IsServiceReady();
 
 private:
+	struct FZenConnectSettings
+	{
+		FString HostName;
+		uint16 Port = 1337;
+	};
+
+	struct FSettings
+	{
+		bool bAutoLaunch = true;
+
+		struct FAutoLaunchSettings
+		{
+			FString WorkspaceDataPath;
+			FString DataPath;
+			FString ExtraArgs;
+			uint16 DesiredPort = 1337;
+			bool bHidden = true;
+		} AutoLaunchSettings;
+
+		FZenConnectSettings ConnectExistingSettings;
+	};
+
+	void PopulateSettings(FStringView InstanceURL);
+	void PromptUserToStopRunningServerInstance(const FString& ServerFilePath);
+	FString ConditionalUpdateLocalInstall();
+	bool AutoLaunch();
+
+	FSettings Settings;
+	FString URL;
+	FString HostName;
+	uint16 Port;
+	bool bHasLaunchedLocal = false;
 };
 
 } // namespace UE::Zen
 
+#undef UE_API
+
 #endif // UE_WITH_ZEN
diff --git a/Engine/Source/Developer/DerivedDataCache/Private/DerivedDataBackends.cpp b/Engine/Source/Developer/DerivedDataCache/Private/DerivedDataBackends.cpp
--- a/Engine/Source/Developer/DerivedDataCache/Private/DerivedDataBackends.cpp
+++ b/Engine/Source/Developer/DerivedDataCache/Private/DerivedDataBackends.cpp
@@ -870,11 +870,7 @@ public:
 	{
 #if WITH_ZEN_DDC_BACKEND
 		FString ServiceUrl;
-		if (!FParse::Value(Entry, TEXT("Host="), ServiceUrl))
-		{
-			UE_LOG(LogDerivedDataCache, Error, TEXT("Node %s does not specify 'Host'."), NodeName);
-			return nullptr;
-		}
+		FParse::Value(Entry, TEXT("Host="), ServiceUrl);
 
 		FString Namespace;
 		if (!FParse::Value(Entry, TEXT("Namespace="), Namespace))
@@ -886,7 +882,7 @@ public:
 		FZenDerivedDataBackend* backend = new FZenDerivedDataBackend(*ServiceUrl, *Namespace);
 		if (!backend->IsUsable())
 		{
-			UE_LOG(LogDerivedDataCache, Warning, TEXT("%s could not contact the service (%s), will not use it."), NodeName, *ServiceUrl);
+			UE_LOG(LogDerivedDataCache, Warning, TEXT("%s could not contact the service (%s), will not use it."), NodeName, *backend->GetName());
 			delete backend;
 			return nullptr;
 		}
diff --git a/Engine/Source/Developer/DerivedDataCache/Private/ZenDerivedDataBackend.cpp b/Engine/Source/Developer/DerivedDataCache/Private/ZenDerivedDataBackend.cpp
--- a/Engine/Source/Developer/DerivedDataCache/Private/ZenDerivedDataBackend.cpp
+++ b/Engine/Source/Developer/DerivedDataCache/Private/ZenDerivedDataBackend.cpp
@@ -37,15 +37,15 @@ namespace UE::DerivedData::Backends {
 FZenDerivedDataBackend::FZenDerivedDataBackend(
 	const TCHAR* InServiceUrl,
 	const TCHAR* InNamespace)
-	: Domain(InServiceUrl)
-	, Namespace(InNamespace)
+	: Namespace(InNamespace)
+	, ZenService(InServiceUrl)
 {
 	if (IsServiceReady())
 	{
-		RequestPool = MakeUnique<Zen::FZenHttpRequestPool>(InServiceUrl);
+		RequestPool = MakeUnique<Zen::FZenHttpRequestPool>(ZenService.GetInstance().GetURL());
 		bIsUsable = true;
 	}
-	bCacheRecordEndpointEnabled = false;
+ 	bCacheRecordEndpointEnabled = false;
 	GConfig->GetBool(TEXT("Zen"), TEXT("CacheRecordEndpointEnabled"), bCacheRecordEndpointEnabled, GEngineIni);
 }
 
@@ -55,25 +55,12 @@ FZenDerivedDataBackend::~FZenDerivedDataBackend()
 
 FString FZenDerivedDataBackend::GetName() const
 {
-	return Domain;
+	return ZenService.GetInstance().GetURL();
 }
 
 bool FZenDerivedDataBackend::IsServiceReady()
 {
-	Zen::FZenHttpRequest Request(*Domain, false);
-	Zen::FZenHttpRequest::Result Result = Request.PerformBlockingDownload(TEXT("health/ready"_SV), nullptr);
-	
-	if (Result == Zen::FZenHttpRequest::Result::Success && Zen::IsSuccessCode(Request.GetResponseCode()))
-	{
-		UE_LOG(LogDerivedDataCache, Display, TEXT("Z$ HTTP DDC service status: %s."), *Request.GetResponseAsString());
-		return true;
-	}
-	else
-	{
-		UE_LOG(LogDerivedDataCache, Warning, TEXT("Unable to reach Z$ HTTP DDC service at %s. Status: %d . Response: %s"), *Domain, Request.GetResponseCode(), *Request.GetResponseAsString());
-	}
-
-	return false;
+	return ZenService.GetInstance().IsServiceReady();
 }
 
 bool FZenDerivedDataBackend::ShouldRetryOnError(int64 ResponseCode)
diff --git a/Engine/Source/Developer/DerivedDataCache/Private/ZenDerivedDataBackend.h b/Engine/Source/Developer/DerivedDataCache/Private/ZenDerivedDataBackend.h
--- a/Engine/Source/Developer/DerivedDataCache/Private/ZenDerivedDataBackend.h
+++ b/Engine/Source/Developer/DerivedDataCache/Private/ZenDerivedDataBackend.h
@@ -21,6 +21,7 @@
 #include "DerivedDataBackendInterface.h"
 #include "DerivedDataCacheUsageStats.h"
 #include "HAL/CriticalSection.h"
+#include "ZenServerInterface.h"
 
 class FCbObject;
 class FCbPackage;
@@ -157,8 +158,8 @@ private:
 	bool ShouldSimulateMiss(const TCHAR* InKey);
 	bool ShouldSimulateMiss(const FCacheKey& InKey);
 private:
-	FString Domain;
 	FString Namespace;
+	UE::Zen::FScopeZenService ZenService;
 	mutable FDerivedDataCacheUsageStats UsageStats;
 	TUniquePtr<UE::Zen::FZenHttpRequestPool> RequestPool;
 	bool bIsUsable = false;
diff --git a/Engine/Source/Developer/DerivedDataCache/Private/ZenServerInterface.cpp b/Engine/Source/Developer/DerivedDataCache/Private/ZenServerInterface.cpp
--- a/Engine/Source/Developer/DerivedDataCache/Private/ZenServerInterface.cpp
+++ b/Engine/Source/Developer/DerivedDataCache/Private/ZenServerInterface.cpp
@@ -3,59 +3,85 @@
 
 #if UE_WITH_ZEN
 
-#include "DerivedDataBackendInterface.h"
 #include "ZenBackendUtils.h"
+#include "ZenServerHttp.h"
 
-#include "Memory/CompositeBuffer.h"
+#include "HAL/FileManager.h"
+#include "HAL/Platform.h"
+#include "HAL/PlatformMisc.h"
+#include "Logging/LogScopedCategoryAndVerbosityOverride.h"
+#include "Misc/App.h"
+#include "Misc/ConfigCacheIni.h"
 #include "Misc/Paths.h"
-#include "Serialization/CompactBinary.h"
-#include "Serialization/CompactBinaryPackage.h"
-#include "Serialization/CompactBinaryValidation.h"
-#include "Serialization/MemoryReader.h"
+#include "String/LexFromString.h"
+
+#if PLATFORM_WINDOWS
+#	include "Windows/AllowWindowsPlatformTypes.h"
+#	include <shellapi.h>
+#	include <Windows.h>
+#	include "Windows/HideWindowsPlatformTypes.h"
+#endif
 
 namespace UE::Zen {
 
-static bool bHasLaunched = false;
+DEFINE_LOG_CATEGORY_STATIC(LogZenServiceInstance, Log, All);
+
+FZenServiceInstance& GetDefaultServiceInstance()
+{
+	static FZenServiceInstance DefaultServiceInstance;
+	return DefaultServiceInstance;
+}
 
-FZenServiceInstance::FZenServiceInstance()
+FScopeZenService::FScopeZenService(FStringView InstanceURL)
 {
-	if (bHasLaunched)
+	if (!InstanceURL.IsEmpty() && !InstanceURL.Equals(TEXT("<DefaultInstance>")))
 	{
-		return;
+		UniqueNonDefaultInstance = MakeUnique<FZenServiceInstance>(InstanceURL);
+		ServiceInstance = UniqueNonDefaultInstance.Get();
+	}
+	else
+	{
+		ServiceInstance = &GetDefaultServiceInstance();
 	}
+}
 
-	FString Parms;
-	bool bLaunchDetached = false;
-	bool bLaunchHidden = false;
-	bool bLaunchReallyHidden = false;
-	uint32* OutProcessID = nullptr;
-	int32 PriorityModifier = 0;
-	const TCHAR* OptionalWorkingDirectory = nullptr;
-	void* PipeWriteChild = nullptr;
-	void* PipeReadChild = nullptr;
-
-	FString MainFilePath = FPaths::ConvertRelativePathToFull(FPaths::Combine(FPaths::EngineDir(), TEXT("Binaries"), TEXT("Win64"), TEXT("zenserver.exe")));
-
-	Parms.Appendf(TEXT("--owner-pid %d"), FPlatformProcess::GetCurrentProcessId());
-
-	FProcHandle Proc = FPlatformProcess::CreateProc(
-		*MainFilePath,
-		*Parms,
-		bLaunchDetached,
-		bLaunchHidden,
-		bLaunchReallyHidden,
-		OutProcessID,
-		PriorityModifier,
-		OptionalWorkingDirectory,
-		PipeWriteChild,
-		PipeReadChild);
+FScopeZenService::FScopeZenService(FStringView InstanceHostName, uint16 InstancePort)
+{
+	if (!InstanceHostName.IsEmpty() && !InstanceHostName.Equals(TEXT("<DefaultInstance>")))
+	{
+		TStringBuilder<64> URL;
+		URL.AppendAnsi("http://");
+		URL.Append(InstanceHostName);
+		URL.AppendAnsi(":");
+		URL << InstancePort;
 
-	if (!Proc.IsValid())
+		UniqueNonDefaultInstance = MakeUnique<FZenServiceInstance>(URL);
+		ServiceInstance = UniqueNonDefaultInstance.Get();
+	}
+	else
 	{
-		Proc = FPlatformProcess::CreateElevatedProcess(*MainFilePath, *Parms);
+		ServiceInstance = &GetDefaultServiceInstance();
 	}
+}
 
-	bHasLaunched = true;
+FZenServiceInstance::FZenServiceInstance(FStringView InstanceURL)
+{
+	Settings.AutoLaunchSettings.DataPath = FPaths::ConvertRelativePathToFull(FPaths::Combine(FPaths::EngineVersionAgnosticUserDir(), TEXT("Zen")));
+	Settings.AutoLaunchSettings.WorkspaceDataPath = Settings.AutoLaunchSettings.DataPath;
+	PopulateSettings(InstanceURL);
+
+	if (Settings.bAutoLaunch)
+	{
+		bHasLaunchedLocal = AutoLaunch();
+	}
+	else
+	{
+		HostName = Settings.ConnectExistingSettings.HostName;
+		Port = Settings.ConnectExistingSettings.Port;
+	}
+	TStringBuilder<128> URLBuilder;
+	URLBuilder << TEXT("http://") << HostName << TEXT(":") << Port << TEXT("/");
+	URL = URLBuilder.ToString();
 }
 
 FZenServiceInstance::~FZenServiceInstance()
@@ -65,14 +91,304 @@ FZenServiceInstance::~FZenServiceInstance()
 bool 
 FZenServiceInstance::IsServiceRunning()
 {
-	return false;
+	return !Settings.bAutoLaunch || bHasLaunchedLocal;
 }
 
 bool 
 FZenServiceInstance::IsServiceReady()
 {
+	if (IsServiceRunning())
+	{
+		TStringBuilder<128> ZenDomain;
+		ZenDomain << HostName << TEXT(":") << Port;
+		Zen::FZenHttpRequest Request(ZenDomain.ToString(), false);
+		Zen::FZenHttpRequest::Result Result = Request.PerformBlockingDownload(TEXT("health/ready"_SV), nullptr);
+		
+		if (Result == Zen::FZenHttpRequest::Result::Success && Zen::IsSuccessCode(Request.GetResponseCode()))
+		{
+			UE_LOG(LogZenServiceInstance, Display, TEXT("Z$ HTTP DDC service status: %s."), *Request.GetResponseAsString());
+			return true;
+		}
+		else
+		{
+			UE_LOG(LogZenServiceInstance, Warning, TEXT("Unable to reach Z$ HTTP DDC service at %s. Status: %d . Response: %s"), ZenDomain.ToString(), Request.GetResponseCode(), *Request.GetResponseAsString());
+		}
+	}
 	return false;
 }
 
+static void DetermineDataPath(const TCHAR* ConfigSection, FString& DataPath)
+{
+	GConfig->GetString(ConfigSection, TEXT("DataPath"), DataPath, GEngineIni);
+
+	bool bUsedEnvOverride = false;
+	bool bUsedEditorOverride = false;
+	// Much of the logic here is meant to mirror the way that DDC backend paths can be specified to allow
+	// us to match behavior and not behave differently in the presence of custom config
+	FString DataPathEnvOverride;
+	if (GConfig->GetString(ConfigSection, TEXT("DataPathEnvOverride"), DataPathEnvOverride, GEngineIni))
+	{
+		FString DataPathEnvOverrideValue = FPlatformMisc::GetEnvironmentVariable(*DataPathEnvOverride);
+		if(!DataPathEnvOverrideValue.IsEmpty())
+		{
+			DataPath = DataPathEnvOverrideValue;
+			bUsedEnvOverride = true;
+			UE_LOG(LogZenServiceInstance, Log, TEXT("Found environment variable %s=%s"), *DataPathEnvOverride, *DataPathEnvOverrideValue);
+		}
+
+		if (FPlatformMisc::GetStoredValue(TEXT("Epic Games"), TEXT("GlobalDataCachePath"), *DataPathEnvOverride, DataPathEnvOverrideValue))
+		{
+			if (!DataPathEnvOverrideValue.IsEmpty())
+			{
+				DataPath = DataPathEnvOverrideValue;
+				bUsedEnvOverride = true;
+				UE_LOG(LogZenServiceInstance, Log, TEXT("Found registry key GlobalDataCachePath %s=%s"), *DataPathEnvOverride, *DataPath);
+			}
+		}
+	}
+
+	FString DataPathEditorOverrideSetting;
+	if (GConfig->GetString(ConfigSection, TEXT("DataPathEditorOverrideSetting"), DataPathEditorOverrideSetting, GEngineIni))
+	{
+		FString Setting = GConfig->GetStr(TEXT("/Script/UnrealEd.EditorSettings"), *DataPathEditorOverrideSetting, GEditorSettingsIni);
+		if(!Setting.IsEmpty())
+		{
+			FString SettingPath;
+			if(FParse::Value(*Setting, TEXT("Path="), SettingPath))
+			{
+				SettingPath = SettingPath.TrimQuotes();
+				if(!SettingPath.IsEmpty())
+				{
+					DataPath = SettingPath;
+					bUsedEditorOverride = true;
+					UE_LOG(LogZenServiceInstance, Log, TEXT("Found editor setting /Script/UnrealEd.EditorSettings.Path=%s"), *DataPath);
+				}
+			}
+		}
+	}
+
+	if (bUsedEnvOverride || bUsedEditorOverride)
+	{
+		DataPath = FPaths::Combine(DataPath, TEXT("Zen"));
+	}
+
+	{
+		FString CommandLineOverrideValue;
+		if (FParse::Value(FCommandLine::Get(), TEXT("ZenDataPath="), CommandLineOverrideValue))
+		{
+			DataPath = CommandLineOverrideValue;
+			UE_LOG(LogZenServiceInstance, Log, TEXT("Found command line override ZenDataPath=%s"), *CommandLineOverrideValue);
+		}
+	}
+}
+
+static void
+ReadUInt16FromConfig(const TCHAR* Section, const TCHAR* Key, uint16& Value, const FString& ConfigFile)
+{
+	int32 ValueInt32 = Value;
+	GConfig->GetInt(Section, Key, ValueInt32, ConfigFile);
+	Value = (uint16)ValueInt32;
 }
+
+
+static void
+ParseURLToHostNameAndPort(FStringView URL, FString& OutHostName, uint16& OutPort)
+{
+	if (URL.StartsWith(TEXT("http://")))
+	{
+		URL.RightChopInline(7);
+	}
+
+	FStringView::SizeType PortDelimIndex = INDEX_NONE;
+	URL.FindChar(TEXT(':'), PortDelimIndex);
+	if (PortDelimIndex != INDEX_NONE)
+	{
+		OutHostName = URL.Left(PortDelimIndex);
+		LexFromString(OutPort, URL.RightChop(PortDelimIndex + 1));
+	}
+	else
+	{
+		OutHostName = URL;
+		OutPort = 1337;
+	}
+}
+
+void
+FZenServiceInstance::PopulateSettings(FStringView InstanceURL)
+{
+	// Allow a provided InstanceURL to override everything
+	if (!InstanceURL.IsEmpty())
+	{
+		Settings.bAutoLaunch = false;
+		ParseURLToHostNameAndPort(InstanceURL, Settings.ConnectExistingSettings.HostName, Settings.ConnectExistingSettings.Port);
+		return;
+	}
+
+	check(GConfig && GConfig->IsReadyForUse());
+	const TCHAR* ConfigSection = TEXT("Zen");
+	GConfig->GetBool(ConfigSection, TEXT("AutoLaunch"), Settings.bAutoLaunch, GEngineIni);
+
+	// AutoLaunch settings
+	{
+		const TCHAR* AutoLaunchConfigSection = TEXT("Zen.AutoLaunch");
+
+		// Workspace path is for marker files that we don't want in a shared location out of tree.
+		GConfig->GetString(AutoLaunchConfigSection, TEXT("WorkspaceDataPath"), Settings.AutoLaunchSettings.WorkspaceDataPath, GEngineIni);
+		DetermineDataPath(AutoLaunchConfigSection, Settings.AutoLaunchSettings.DataPath);
+		Settings.AutoLaunchSettings.DataPath = FPaths::ConvertRelativePathToFull(Settings.AutoLaunchSettings.DataPath);
+		GConfig->GetString(AutoLaunchConfigSection, TEXT("ExtraArgs"), Settings.AutoLaunchSettings.ExtraArgs, GEngineIni);
+
+		ReadUInt16FromConfig(AutoLaunchConfigSection, TEXT("DesiredPort"), Settings.AutoLaunchSettings.DesiredPort, GEngineIni);
+		GConfig->GetBool(AutoLaunchConfigSection, TEXT("Hidden"), Settings.AutoLaunchSettings.bHidden, GEngineIni);
+	}
+
+	// ConnectExisting settings
+	{
+		const TCHAR* ConnectExistingConfigSection = TEXT("Zen.ConnectExisting");
+		GConfig->GetString(ConnectExistingConfigSection, TEXT("HostName"), Settings.ConnectExistingSettings.HostName, GEngineIni);
+		ReadUInt16FromConfig(ConnectExistingConfigSection, TEXT("Port"), Settings.ConnectExistingSettings.Port, GEngineIni);
+	}
+}
+
+void
+FZenServiceInstance::PromptUserToStopRunningServerInstance(const FString& ServerFilePath)
+{
+	if (FApp::IsUnattended())
+	{
+		// Do not ask if there is no one to show a message
+		return;
+	}
+
+	FString UserPrompt = FString::Printf(TEXT("Please shutdown currently running Unreal editor and tools using zenserver at '%s' so that it can be upgraded to a new version."), *ServerFilePath);
+	FPlatformMisc::MessageBoxExt(EAppMsgType::Ok, *UserPrompt, TEXT("Upgrade required"));
+}
+
+FString
+FZenServiceInstance::ConditionalUpdateLocalInstall()
+{
+	FString InTreeFilePath = FPaths::ConvertRelativePathToFull(FPlatformProcess::GenerateApplicationPath(TEXT("zenserver"), EBuildConfiguration::Development));
+
+	FString InstallFilePath = FPaths::ConvertRelativePathToFull(FPaths::Combine(FPaths::EngineSavedDir(), TEXT("ZenServer"), FPaths::GetCleanFilename(InTreeFilePath)));
+
+	IFileManager& FileManager = IFileManager::Get();
+	FDateTime InTreeFileTime;
+	FDateTime InstallFileTime;
+	FileManager.GetTimeStampPair(*InTreeFilePath, *InstallFilePath, InTreeFileTime, InstallFileTime);
+	if (InTreeFileTime > InstallFileTime)
+	{
+		if (FPlatformProcess::IsApplicationRunning(*FPaths::GetCleanFilename(InstallFilePath)))
+		{
+			PromptUserToStopRunningServerInstance(InstallFilePath);
+		}
+
+		uint32 CopyResult = FileManager.Copy(*InstallFilePath, *InTreeFilePath, true, true, false);
+		checkf(CopyResult == COPY_OK, TEXT("Failed to copy zenserver to install location '%s'."), *InstallFilePath);
+
+#if PLATFORM_WINDOWS
+		FString InstallSymbolFilePath = FPaths::ChangeExtension(InstallFilePath, TEXT("pdb"));
+		CopyResult = FileManager.Copy(*InstallSymbolFilePath, *FPaths::ChangeExtension(InTreeFilePath, TEXT("pdb")), true, true, false);
+		checkf(CopyResult == COPY_OK, TEXT("Failed to copy zenserver symbols to install location '%s'."), *InstallSymbolFilePath);
+#endif
+	}
+
+	return InstallFilePath;
+}
+
+bool
+FZenServiceInstance::AutoLaunch()
+{
+	// TODO: Install and update will be delegated to be the responsibility of the launched zenserver process in the future.
+	FString MainFilePath = ConditionalUpdateLocalInstall();
+
+	FString Parms;
+
+	Parms.Appendf(TEXT("--owner-pid %d --port %d --data-dir \"%s\""),
+		FPlatformProcess::GetCurrentProcessId(),
+		Settings.AutoLaunchSettings.DesiredPort,
+		*Settings.AutoLaunchSettings.DataPath);
+
+	if (!Settings.AutoLaunchSettings.ExtraArgs.IsEmpty())
+	{
+		Parms.AppendChar(TEXT(' '));
+		Parms.Append(Settings.AutoLaunchSettings.ExtraArgs);
+	}
+	HostName = TEXT("localhost");
+	// For now assuming that we get to run on the port we want
+	Port = Settings.AutoLaunchSettings.DesiredPort;
+
+	FProcHandle Proc;
+#if PLATFORM_WINDOWS
+	{
+		// Attempt non-elevated launch
+		STARTUPINFO StartupInfo = {
+			sizeof(STARTUPINFO),
+			NULL, NULL, NULL,
+			(::DWORD)CW_USEDEFAULT,
+			(::DWORD)CW_USEDEFAULT,
+			(::DWORD)CW_USEDEFAULT,
+			(::DWORD)CW_USEDEFAULT,
+			(::DWORD)0, (::DWORD)0, (::DWORD)0,
+			(::DWORD) STARTF_USESHOWWINDOW,
+			Settings.AutoLaunchSettings.bHidden ? SW_HIDE : SW_SHOWMINNOACTIVE,
+			0, NULL,
+			HANDLE(nullptr),
+			HANDLE(nullptr),
+			HANDLE(nullptr)
+		};
+
+		FString CommandLine = FString::Printf(TEXT("\"%s\" %s"), *MainFilePath, *Parms);
+		PROCESS_INFORMATION ProcInfo;
+		if (CreateProcess(NULL, CommandLine.GetCharArray().GetData(), nullptr, nullptr, false, (::DWORD)(NORMAL_PRIORITY_CLASS | DETACHED_PROCESS), nullptr, nullptr, &StartupInfo, &ProcInfo))
+		{
+			::CloseHandle( ProcInfo.hThread );
+			Proc = FProcHandle(ProcInfo.hProcess);
+		}
+
+	}
+	if (!Proc.IsValid())
+	{
+		// Fall back to elevated launch
+		SHELLEXECUTEINFO ShellExecuteInfo;
+		ZeroMemory(&ShellExecuteInfo, sizeof(ShellExecuteInfo));
+		ShellExecuteInfo.cbSize = sizeof(ShellExecuteInfo);
+		ShellExecuteInfo.fMask = SEE_MASK_UNICODE | SEE_MASK_NOCLOSEPROCESS;
+		ShellExecuteInfo.lpFile = *MainFilePath;
+		ShellExecuteInfo.lpVerb = TEXT("runas");
+		ShellExecuteInfo.nShow = Settings.AutoLaunchSettings.bHidden ? SW_HIDE : SW_SHOWMINNOACTIVE;
+		ShellExecuteInfo.lpParameters = *Parms;
+
+		if (ShellExecuteEx(&ShellExecuteInfo))
+		{
+			Proc = FProcHandle(ShellExecuteInfo.hProcess);
+		}
+	}
+#else
+	{
+		bool bLaunchDetached = true;
+		bool bLaunchHidden = true;
+		bool bLaunchReallyHidden = Settings.AutoLaunchSettings.bHidden;
+		uint32* OutProcessID = nullptr;
+		int32 PriorityModifier = 0;
+		const TCHAR* OptionalWorkingDirectory = nullptr;
+		void* PipeWriteChild = nullptr;
+		void* PipeReadChild = nullptr;
+		Proc = FPlatformProcess::CreateProc(
+			*MainFilePath,
+			*Parms,
+			bLaunchDetached,
+			bLaunchHidden,
+			bLaunchReallyHidden,
+			OutProcessID,
+			PriorityModifier,
+			OptionalWorkingDirectory,
+			PipeWriteChild,
+			PipeReadChild);
+	}
+#endif
+
+	return Proc.IsValid();
+}
+
+}
+
 #endif // UE_WITH_ZEN
diff --git a/Engine/Source/Developer/IoStoreUtilities/Internal/ZenStoreHttpClient.h b/Engine/Source/Developer/IoStoreUtilities/Internal/ZenStoreHttpClient.h
--- a/Engine/Source/Developer/IoStoreUtilities/Internal/ZenStoreHttpClient.h
+++ b/Engine/Source/Developer/IoStoreUtilities/Internal/ZenStoreHttpClient.h
@@ -11,6 +11,7 @@
 #include "Misc/StringBuilder.h"
 #include "Serialization/CompactBinaryPackage.h"
 #include "Templates/UniquePtr.h"
+#include "ZenServerInterface.h"
 
 class FCbPackage;
 class FCbObject;
@@ -26,6 +27,7 @@ namespace UE {
 class IOSTOREUTILITIES_API FZenStoreHttpClient
 {
 public:
+	FZenStoreHttpClient();
 	FZenStoreHttpClient(const FStringView InHostName, uint16 InPort);
 	~FZenStoreHttpClient();
 
@@ -49,6 +51,14 @@ public:
 	TIoStatusOr<FIoBuffer> ReadChunk(const FIoChunkId& Id, uint64 Offset = 0, uint64 Size = ~0ull);
 	TIoStatusOr<FIoBuffer> ReadOpLogAttachment(FStringView Id, uint64 Offset = 0, uint64 Size = ~0ull);
 
+#if UE_WITH_ZEN
+	const TCHAR* GetHostName() const { return ZenService.GetInstance().GetHostName(); }
+	uint16 GetPort() const { return ZenService.GetInstance().GetPort(); }
+#else // Default to localhost:1337 for platforms where Zen wouldn't be supported yet
+	const TCHAR* GetHostName() const { return TEXT("localhost"); }
+	uint16 GetPort() const { return 1337; }
+#endif
+
 	TFuture<TIoStatusOr<FCbObject>> GetOplog();
 	TFuture<TIoStatusOr<FCbObject>> GetFiles();
 
@@ -58,8 +68,10 @@ public:
 private:
 	TIoStatusOr<FIoBuffer> ReadOpLogUri(FStringBuilderBase& ChunkUri, uint64 Offset, uint64 Size);
 
-	FString	HostName;
-	uint16 Port;
+	static const uint32 PoolEntryCount;
+#if UE_WITH_ZEN
+	UE::Zen::FScopeZenService ZenService;
+#endif
 	TUniquePtr<Zen::FZenHttpRequestPool> RequestPool;
 	FString OplogPath;
 	FString OplogNewEntryPath;
diff --git a/Engine/Source/Developer/IoStoreUtilities/Private/ZenStoreHttpClient.cpp b/Engine/Source/Developer/IoStoreUtilities/Private/ZenStoreHttpClient.cpp
--- a/Engine/Source/Developer/IoStoreUtilities/Private/ZenStoreHttpClient.cpp
+++ b/Engine/Source/Developer/IoStoreUtilities/Private/ZenStoreHttpClient.cpp
@@ -25,17 +25,17 @@ DEFINE_LOG_CATEGORY_STATIC(LogZenStore, Log, All);
 
 namespace UE { 
 
-FZenStoreHttpClient::FZenStoreHttpClient(const FStringView InHostName, uint16 InPort)
-: HostName(InHostName)
-, Port(InPort)
+const uint32 FZenStoreHttpClient::PoolEntryCount = 32;
+
+FZenStoreHttpClient::FZenStoreHttpClient()
 {
-	TStringBuilder<64> Uri;
-	Uri.AppendAnsi("http://");
-	Uri.Append(InHostName);
-	Uri.AppendAnsi(":");
-	Uri << InPort;
+	RequestPool = MakeUnique<Zen::FZenHttpRequestPool>(ZenService.GetInstance().GetURL(), PoolEntryCount);
+}
 
-	RequestPool = MakeUnique<Zen::FZenHttpRequestPool>(Uri, 32);
+FZenStoreHttpClient::FZenStoreHttpClient(const FStringView InHostName, uint16 InPort)
+: ZenService(InHostName, InPort)
+{
+	RequestPool = MakeUnique<Zen::FZenHttpRequestPool>(ZenService.GetInstance().GetURL(), PoolEntryCount);
 }
 
 FZenStoreHttpClient::~FZenStoreHttpClient()
@@ -569,6 +569,10 @@ namespace Zen {
 	};
 }
 
+FZenStoreHttpClient::FZenStoreHttpClient()
+{
+}
+
 FZenStoreHttpClient::FZenStoreHttpClient(const FStringView InHostName, uint16 InPort)
 {
 }
diff --git a/Engine/Source/Developer/IoStoreUtilities/Private/ZenStoreWriter.cpp b/Engine/Source/Developer/IoStoreUtilities/Private/ZenStoreWriter.cpp
--- a/Engine/Source/Developer/IoStoreUtilities/Private/ZenStoreWriter.cpp
+++ b/Engine/Source/Developer/IoStoreUtilities/Private/ZenStoreWriter.cpp
@@ -261,12 +261,10 @@ FZenStoreWriter::FZenStoreWriter(
 {
 	StaticInit();
 
-	FString HostName = TEXT("localhost");
-	uint16 Port = 1337;
 	FString ProjectId = FApp::GetProjectName();
 	FString OplogId = InTargetPlatform->PlatformName();
 
-	HttpClient = MakeUnique<UE::FZenStoreHttpClient>(HostName, Port);
+	HttpClient = MakeUnique<UE::FZenStoreHttpClient>();
 
 	FString RootDir = FPaths::RootDir();
 	FString EngineDir = FPaths::EngineDir();
@@ -288,8 +286,8 @@ FZenStoreWriter::FZenStoreWriter(
 	PackageStoreOptimizer->Initialize(InTargetPlatform);
 
 	FPackageStoreManifest::FZenServerInfo& ZenServerInfo = PackageStoreManifest.EditZenServerInfo();
-	ZenServerInfo.HostName = HostName;
-	ZenServerInfo.Port = Port;
+	ZenServerInfo.HostName = HttpClient->GetHostName();
+	ZenServerInfo.Port = HttpClient->GetPort();
 	ZenServerInfo.ProjectId = ProjectId;
 	ZenServerInfo.OplogId = OplogId;
 
diff --git a/Engine/Source/Editor/UnrealEd/Private/TargetDomain/TargetDomainUtils.cpp b/Engine/Source/Editor/UnrealEd/Private/TargetDomain/TargetDomainUtils.cpp
--- a/Engine/Source/Editor/UnrealEd/Private/TargetDomain/TargetDomainUtils.cpp
+++ b/Engine/Source/Editor/UnrealEd/Private/TargetDomain/TargetDomainUtils.cpp
@@ -388,7 +388,6 @@ bool IsIterativeEnabled(FName PackageName)
 TArray<const UTF8CHAR*> FEditorDomainOplog::ReservedOplogKeys;
 
 FEditorDomainOplog::FEditorDomainOplog()
-	: HttpClient(TEXT("localhost"), 1337)
 {
 	StaticInit();
 
