Avoid a compile error due to narrowing of types when compiling win64 editor with clang.

diff --git a/Engine/Source/Developer/DerivedDataCache/Private/ZenServerInterface.cpp b/Engine/Source/Developer/DerivedDataCache/Private/ZenServerInterface.cpp
--- a/Engine/Source/Developer/DerivedDataCache/Private/ZenServerInterface.cpp
+++ b/Engine/Source/Developer/DerivedDataCache/Private/ZenServerInterface.cpp
@@ -329,7 +329,7 @@ FZenServiceInstance::AutoLaunch()
 			(::DWORD)CW_USEDEFAULT,
 			(::DWORD)0, (::DWORD)0, (::DWORD)0,
 			(::DWORD) STARTF_USESHOWWINDOW,
-			Settings.AutoLaunchSettings.bHidden ? SW_HIDE : SW_SHOWMINNOACTIVE,
+			(::WORD)(Settings.AutoLaunchSettings.bHidden ? SW_HIDE : SW_SHOWMINNOACTIVE),
 			0, NULL,
 			HANDLE(nullptr),
 			HANDLE(nullptr),
