Fix missing ZenDerivedDataBackend stats.

diff --git a/Engine/Source/Developer/DerivedDataCache/Private/ZenDerivedDataBackend.cpp b/Engine/Source/Developer/DerivedDataCache/Private/ZenDerivedDataBackend.cpp
--- a/Engine/Source/Developer/DerivedDataCache/Private/ZenDerivedDataBackend.cpp
+++ b/Engine/Source/Developer/DerivedDataCache/Private/ZenDerivedDataBackend.cpp
@@ -434,7 +434,10 @@ FZenDerivedDataBackend::GetSpeedClass() const
 
 TSharedRef<FDerivedDataCacheStatsNode> FZenDerivedDataBackend::GatherUsageStats() const
 {
-	return MakeShared<FDerivedDataCacheStatsNode>(this, "foo");
+	TSharedRef<FDerivedDataCacheStatsNode> Usage = MakeShared<FDerivedDataCacheStatsNode>(this, FString::Printf(TEXT("%s.%s"), TEXT("ZenDDC"), *GetName()));
+	Usage->Stats.Add(TEXT(""), UsageStats);
+
+	return Usage;
 }
 
 bool FZenDerivedDataBackend::TryToPrefetch(TConstArrayView<FString> CacheKeys)
