Completed FZenHttpRequest::GetResponseAsPackage() implementation, matches Zen side CbPackage serialization

#rnx

diff --git a/Engine/Source/Developer/DerivedDataCache/Private/ZenServerHttp.cpp b/Engine/Source/Developer/DerivedDataCache/Private/ZenServerHttp.cpp
--- a/Engine/Source/Developer/DerivedDataCache/Private/ZenServerHttp.cpp
+++ b/Engine/Source/Developer/DerivedDataCache/Private/ZenServerHttp.cpp
@@ -298,9 +298,9 @@ namespace UE::Zen {
 		}
 
 		TArray<CbAttachmentEntry> AttachmentEntries;
-		AttachmentEntries.SetNum(Hdr.AttachmentCount);
+		AttachmentEntries.SetNum(Hdr.AttachmentCount + 1);
 
-		Reader.Serialize(AttachmentEntries.GetData(), Hdr.AttachmentCount * sizeof(CbAttachmentEntry));
+		Reader.Serialize(AttachmentEntries.GetData(), (Hdr.AttachmentCount + 1) * sizeof(CbAttachmentEntry));
 
 		int Index = 0;
 
@@ -309,14 +309,35 @@ namespace UE::Zen {
 			FUniqueBuffer AttachmentData = FUniqueBuffer::Alloc(Entry.AttachmentSize);
 			Reader.Serialize(AttachmentData.GetData(), AttachmentData.GetSize());
 
-			if (Index == 0)
+			if (Entry.Flags & CbAttachmentEntry::IsCompressed)
 			{
+				FCompressedBuffer CompBuf(FCompressedBuffer::FromCompressed(AttachmentData.MoveToShared()));
 
+				if (Entry.Flags & CbAttachmentEntry::IsObject)
+				{
+					checkf(Index == 0, TEXT("Object attachments are not currently supported"));
+
+					Package.SetObject(FCbObject(CompBuf.Decompress()));
+				}
+				else
+				{
+					FCbAttachment Attachment(MoveTemp(CompBuf));
+					Package.AddAttachment(Attachment);
+				}
 			}
-			else
+			else /* not compressed */
 			{
-				FCbAttachment Attachment(AttachmentData.MoveToShared());
-				Package.AddAttachment(Attachment);
+				if (Entry.Flags & CbAttachmentEntry::IsObject)
+				{
+					checkf(Index == 0, TEXT("Object attachments are not currently supported"));
+
+					Package.SetObject(FCbObject(AttachmentData.MoveToShared()));
+				}
+				else
+				{
+					FCbAttachment Attachment(AttachmentData.MoveToShared());
+					Package.AddAttachment(Attachment);
+				}
 			}
 
 			++Index;
@@ -531,7 +552,7 @@ namespace UE::Zen {
 				UE_LOG(
 					LogDerivedDataCache,
 					Verbose,
-					TEXT("Finished %s zen data (response %d) from %s. %s"),
+					TEXT("Finished %s zen data (response %ld) from %s. %s"),
 					VerbStr,
 					ResponseCode,
 					Uri,
@@ -547,7 +568,7 @@ namespace UE::Zen {
 				UE_LOG(
 					LogDerivedDataCache,
 					Error,
-					TEXT("Failed %s zen data (response %d) from %s. Response: %s"),
+					TEXT("Failed %s zen data (response %ld) from %s. Response: %s"),
 					VerbStr,
 					ResponseCode,
 					Uri,
@@ -604,15 +625,15 @@ namespace UE::Zen {
 		break;
 
 		case CURLINFO_HEADER_IN:
-			UE_LOG(LogDerivedDataCache, VeryVerbose, TEXT("%p: Received header (%d bytes)"), Request, DebugInfoSize);
+			UE_LOG(LogDerivedDataCache, VeryVerbose, TEXT("%p: Received header (%zd bytes)"), Request, DebugInfoSize);
 			break;
 
 		case CURLINFO_DATA_IN:
-			UE_LOG(LogDerivedDataCache, VeryVerbose, TEXT("%p: Received data (%d bytes)"), Request, DebugInfoSize);
+			UE_LOG(LogDerivedDataCache, VeryVerbose, TEXT("%p: Received data (%zd bytes)"), Request, DebugInfoSize);
 			break;
 
 		case CURLINFO_DATA_OUT:
-			UE_LOG(LogDerivedDataCache, VeryVerbose, TEXT("%p: Sent data (%d bytes)"), Request, DebugInfoSize);
+			UE_LOG(LogDerivedDataCache, VeryVerbose, TEXT("%p: Sent data (%zd bytes)"), Request, DebugInfoSize);
 			break;
 		}
 
