DDC: Switched to FLazyEvent to reduce build execution overhead

#rnx

diff --git a/Engine/Source/Developer/DerivedDataCache/Private/DerivedDataBackendAsyncPutWrapper.cpp b/Engine/Source/Developer/DerivedDataCache/Private/DerivedDataBackendAsyncPutWrapper.cpp
--- a/Engine/Source/Developer/DerivedDataCache/Private/DerivedDataBackendAsyncPutWrapper.cpp
+++ b/Engine/Source/Developer/DerivedDataCache/Private/DerivedDataBackendAsyncPutWrapper.cpp
@@ -5,9 +5,9 @@
 #include "DerivedDataPayload.h"
 #include "DerivedDataRequest.h"
 #include "DerivedDataRequestOwner.h"
+#include "Experimental/Async/LazyEvent.h"
 #include "MemoryDerivedDataBackend.h"
 #include "Tasks/Task.h"
-#include <atomic>
 
 namespace UE::DerivedData::Backends
 {
@@ -352,22 +352,15 @@ public:
 	inline FDerivedDataAsyncWrapperRequest(IRequestOwner& InOwner, TUniqueFunction<void (bool bCancel)>&& InFunction)
 		: Owner(InOwner)
 		, Function(MoveTemp(InFunction))
-		, DoneEvent(FPlatformProcess::GetSynchEventFromPool(true))
 	{
 	}
 
-	inline ~FDerivedDataAsyncWrapperRequest() final
-	{
-		FPlatformProcess::ReturnSynchEventToPool(DoneEvent);
-	}
-
 	inline void Start(EPriority Priority)
 	{
 		FDerivedDataBackend::Get().AddToAsyncCompletionCounter(1);
 		Owner.Begin(this);
 
-		DoneEvent->Reset();
-		WorkNotFinishedCounter.fetch_add(1, std::memory_order_relaxed);
+		DoneEvent.Reset();
 		GDDCIOThreadPool->AddQueuedWork(this, GetPriority(Priority));
 	}
 
@@ -377,8 +370,7 @@ public:
 		Owner.End(this, [this, bCancel]
 		{
 			Function(bCancel);
-			WorkNotFinishedCounter.fetch_sub(1, std::memory_order_release);
-			DoneEvent->Trigger();
+			DoneEvent.Trigger();
 		});
 		// DO NOT ACCESS ANY MEMBERS PAST THIS POINT!
 		FDerivedDataBackend::Get().AddToAsyncCompletionCounter(-1);
@@ -396,7 +388,7 @@ public:
 
 	inline void Cancel() final
 	{
-		if (WorkNotFinishedCounter.load(std::memory_order_acquire) > 0)
+		if (!DoneEvent.Wait(0))
 		{
 			if (GDDCIOThreadPool->RetractQueuedWork(this))
 			{
@@ -404,14 +396,15 @@ public:
 			}
 			else
 			{
-				Wait();
+				FScopeCycleCounter Scope(GetStatId());
+				DoneEvent.Wait();
 			}
 		}
 	}
 
 	inline void Wait() final
 	{
-		if (WorkNotFinishedCounter.load(std::memory_order_acquire) > 0)
+		if (!DoneEvent.Wait(0))
 		{
 			if (GDDCIOThreadPool->RetractQueuedWork(this))
 			{
@@ -420,7 +413,7 @@ public:
 			else
 			{
 				FScopeCycleCounter Scope(GetStatId());
-				DoneEvent->Wait();
+				DoneEvent.Wait();
 			}
 		}
 	}
@@ -458,10 +451,9 @@ private:
 	}
 
 private:
-	std::atomic<uint32> WorkNotFinishedCounter{0};
 	IRequestOwner& Owner;
 	TUniqueFunction<void (bool bCancel)> Function;
-	FEvent* DoneEvent;
+	FLazyEvent DoneEvent{EEventMode::ManualReset};
 };
 
 void FDerivedDataBackendAsyncPutWrapper::Put(
diff --git a/Engine/Source/Developer/DerivedDataCache/Private/DerivedDataBuildJobContext.cpp b/Engine/Source/Developer/DerivedDataCache/Private/DerivedDataBuildJobContext.cpp
--- a/Engine/Source/Developer/DerivedDataCache/Private/DerivedDataBuildJobContext.cpp
+++ b/Engine/Source/Developer/DerivedDataCache/Private/DerivedDataBuildJobContext.cpp
@@ -130,7 +130,7 @@ void FBuildJobContext::BeginBuild(IRequestOwner& InOwner, TUniqueFunction<void (
 void FBuildJobContext::EndBuild()
 {
 	OnEndBuild();
-	BuildCompleteEvent->Trigger();
+	BuildCompleteEvent.Trigger();
 	Owner = nullptr;
 }
 
@@ -203,7 +203,7 @@ void FBuildJobContext::Cancel()
 
 void FBuildJobContext::Wait()
 {
-	BuildCompleteEvent->Wait();
+	BuildCompleteEvent.Wait();
 }
 
 } // UE::DerivedData::Private
diff --git a/Engine/Source/Developer/DerivedDataCache/Private/DerivedDataBuildJobContext.h b/Engine/Source/Developer/DerivedDataCache/Private/DerivedDataBuildJobContext.h
--- a/Engine/Source/Developer/DerivedDataCache/Private/DerivedDataBuildJobContext.h
+++ b/Engine/Source/Developer/DerivedDataCache/Private/DerivedDataBuildJobContext.h
@@ -9,7 +9,7 @@
 #include "DerivedDataBuildFunction.h"
 #include "DerivedDataCacheKey.h"
 #include "DerivedDataRequest.h"
-#include "HAL/Event.h"
+#include "Experimental/Async/LazyEvent.h"
 #include "Memory/MemoryFwd.h"
 #include "Serialization/CompactBinary.h"
 #include "Templates/Function.h"
@@ -82,7 +82,7 @@ private:
 	FBuildOutputBuilder& OutputBuilder;
 	TMap<FString, FCbObject> Constants;
 	TMap<FString, FCompressedBuffer> Inputs;
-	FEventRef BuildCompleteEvent{EEventMode::ManualReset};
+	UE::FLazyEvent BuildCompleteEvent{EEventMode::ManualReset};
 	TUniqueFunction<void ()> OnEndBuild;
 	IRequestOwner* Owner{};
 	uint64 RequiredMemory{};
diff --git a/Engine/Source/Developer/DerivedDataCache/Private/DerivedDataBuildScheduler.cpp b/Engine/Source/Developer/DerivedDataCache/Private/DerivedDataBuildScheduler.cpp
--- a/Engine/Source/Developer/DerivedDataCache/Private/DerivedDataBuildScheduler.cpp
+++ b/Engine/Source/Developer/DerivedDataCache/Private/DerivedDataBuildScheduler.cpp
@@ -7,7 +7,7 @@
 #include "DerivedDataBuildFunctionRegistry.h"
 #include "DerivedDataRequest.h"
 #include "DerivedDataRequestOwner.h"
-#include "HAL/Event.h"
+#include "Experimental/Async/LazyEvent.h"
 #include "Misc/Guid.h"
 #include "Misc/ScopeLock.h"
 #include "Tasks/Task.h"
@@ -41,73 +41,6 @@ static void ScheduleAsyncStep(IBuildJob& Job, IRequestOwner& Owner, const TCHAR*
 
 ///////////////////////////////////////////////////////////////////////////////////////////////////////////////////////
 
-static FEvent* GPretriggeredEvent;
-
-/** On-demand allocated event to reduce total event allocation */
-class FLazyManualResetEvent
-{
-public:
-	static void Init()
-	{
-		static FEventRef PretriggeredEvent(EEventMode::ManualReset);
-		PretriggeredEvent->Trigger();
-		GPretriggeredEvent = PretriggeredEvent.Get();
-	}
-
-	~FLazyManualResetEvent()
-	{
-		// Sequential consistency might not be needed here in case performance matters
-		FEvent* Event = AtomicEvent.load();
-		if ((Event != nullptr) & (Event != GPretriggeredEvent))
-		{
-			FPlatformProcess::ReturnSynchEventToPool(Event);
-		}
-	}
-
-	void Wait()
-	{
-		FEvent* Event = AtomicEvent.load(std::memory_order_relaxed);
-
-		if (Event == nullptr)
-		{
-			FEvent* NewEvent = FPlatformProcess::GetSynchEventFromPool(/* manual reset */ true);
-			if (AtomicEvent.compare_exchange_strong(/* out */ Event, NewEvent))
-			{
-				Event = NewEvent;
-			}
-			else
-			{
-				FPlatformProcess::ReturnSynchEventToPool(NewEvent);
-			}
-		}
-
-		Event->Wait();
-	}
-
-	void Trigger()
-	{
-		FEvent* Event = AtomicEvent.load(std::memory_order_relaxed);
-
-		if (Event == nullptr)
-		{
-			if (AtomicEvent.compare_exchange_strong(/* out */ Event, GPretriggeredEvent))
-			{
-				return;
-			}
-		}
-
-		if (Event != GPretriggeredEvent)
-		{
-			Event->Trigger();
-		}
-	}
-
-private:
-	std::atomic<FEvent*> AtomicEvent{nullptr};
-};
-
-///////////////////////////////////////////////////////////////////////////////////////////////////////////////////////
-
 /** Limits simultaneous build jobs to reduce peak memory usage */
 class FMemoryScheduler
 {
@@ -140,7 +73,7 @@ private:
 		IBuildJob& Job;
 		IRequestOwner& Owner;
 		const uint64 MemoryEstimate;
-		FLazyManualResetEvent Event;
+		UE::FLazyEvent Event{EEventMode::ManualReset};
 		std::atomic_flag bClaimed = ATOMIC_FLAG_INIT;
 	};
 
@@ -366,7 +299,6 @@ class FBuildScheduler final : public IBuildScheduler
 
 IBuildScheduler* CreateBuildScheduler()
 {
-	FLazyManualResetEvent::Init();
 	return new FBuildScheduler();
 }
 
diff --git a/Engine/Source/Developer/DerivedDataCache/Private/DerivedDataRequestOwner.cpp b/Engine/Source/Developer/DerivedDataCache/Private/DerivedDataRequestOwner.cpp
--- a/Engine/Source/Developer/DerivedDataCache/Private/DerivedDataRequestOwner.cpp
+++ b/Engine/Source/Developer/DerivedDataCache/Private/DerivedDataRequestOwner.cpp
@@ -5,12 +5,11 @@
 #include "Containers/Array.h"
 #include "DerivedDataRequest.h"
 #include "DerivedDataRequestTypes.h"
+#include "Experimental/Async/LazyEvent.h"
 #include "HAL/CriticalSection.h"
-#include "HAL/Event.h"
 #include "Misc/ScopeExit.h"
 #include "Misc/ScopeRWLock.h"
 #include "Templates/RefCounting.h"
-#include <atomic>
 
 namespace UE::DerivedData::Private
 {
@@ -18,7 +17,7 @@ namespace UE::DerivedData::Private
 /** A one-time-use event to work around the lack of condition variables. */
 struct FRequestBarrierEvent : public FRefCountBase
 {
-	FEventRef Event{EEventMode::ManualReset};
+	FLazyEvent Event{EEventMode::ManualReset};
 };
 
 class FRequestOwnerShared final : public IRequestOwner, public FRequestBase
@@ -136,7 +135,7 @@ void FRequestOwnerShared::EndBarrier(ERequestBarrierFlags Flags)
 	}
 	if (LocalBarrierEvent)
 	{
-		LocalBarrierEvent->Event->Trigger();
+		LocalBarrierEvent->Event.Trigger();
 	}
 }
 
@@ -195,7 +194,7 @@ void FRequestOwnerShared::Cancel()
 		}
 		else
 		{
-			LocalBarrierEvent->Event->Wait();
+			LocalBarrierEvent->Event.Wait();
 		}
 	}
 }
@@ -234,7 +233,7 @@ void FRequestOwnerShared::Wait()
 		}
 		else
 		{
-			LocalBarrierEvent->Event->Wait();
+			LocalBarrierEvent->Event.Wait();
 		}
 	}
 }
