Make content-type/accept headers mandatory

#rnx

diff --git a/Engine/Source/Developer/DerivedDataCache/Internal/ZenServerHttp.h b/Engine/Source/Developer/DerivedDataCache/Internal/ZenServerHttp.h
--- a/Engine/Source/Developer/DerivedDataCache/Internal/ZenServerHttp.h
+++ b/Engine/Source/Developer/DerivedDataCache/Internal/ZenServerHttp.h
@@ -23,11 +23,43 @@ static bool IsSuccessCode(int ResponseCode)
 
 enum class EContentType
 {
-	Binary,
-	CompactBinary,
-	CompactBinaryPackage
+	Binary				= 0,
+	Text				= 1,
+	JSON				= 2,
+	CbObject			= 3,
+	CbPackage			= 4,
+	YAML				= 5,
+	CbPackageOffer		= 6,
+	CompressedBinary	= 7,
+	UnknownContentType	= 8,
+	Count
 };
 
+inline FStringView GetMimeType(EContentType Type)
+{
+	switch (Type)
+	{
+		case EContentType::Binary:
+			return TEXT("application/octet-stream"_SV);
+		case EContentType::Text:
+			return TEXT("text/plain"_SV);
+		case EContentType::JSON:
+			return TEXT("application/json"_SV);
+		case EContentType::CbObject:
+			return TEXT("application/x-ue-cb"_SV);
+		case EContentType::CbPackage:
+			return TEXT("application/x-ue-cbpkg"_SV);
+		case EContentType::YAML:
+			return TEXT("text/yaml"_SV);
+		case EContentType::CbPackageOffer:
+			return TEXT("application/x-ue-offer"_SV);
+		case EContentType::CompressedBinary:
+			return TEXT("application/x-ue-comp"_SV);
+		default:
+			return TEXT("unknown"_SV);
+	}
+}
+
 /** Minimal HTTP request type wrapping CURL without the need for managers. This request
   * is written to allow reuse of request objects, in order to allow connections to be reused.
 	
@@ -76,6 +108,7 @@ public:
 		* Upload buffer using the request, using PUT verb
 		* @param Uri Url to use.
 		* @param Buffer Data to upload
+		* @param ContentType The content MIME type.
 		* @return Result of the request
 		*/
 	DERIVEDDATACACHE_API Result PerformBlockingPut(const TCHAR* Uri, const FCompositeBuffer& Buffer, EContentType ContentType);
@@ -85,9 +118,10 @@ public:
 	* @param Uri Url to use.
 	* @param Buffer Optional buffer where data should be downloaded to. If this is null then
 	* downloaded data will be stored in an internal buffer and accessed via GetResponseAsString
+	* @param ContentType The MIME type to accept.
 	* @return Result of the request
 	*/
-	DERIVEDDATACACHE_API Result PerformBlockingDownload(FStringView Uri, TArray64<uint8>* Buffer);
+	DERIVEDDATACACHE_API Result PerformBlockingDownload(FStringView Uri, TArray64<uint8>* Buffer, EContentType AcceptType);
 
 	/**
 		* Download an url into a buffer using the request.
@@ -100,9 +134,10 @@ public:
 	/**
 		* Query an url using the request. Queries can use either "Head" or "Delete" verbs.
 		* @param Uri Url to use.
+		* @param ContentType The MIME type to accept.
 		* @return Result of the request
 		*/
-	DERIVEDDATACACHE_API Result PerformBlockingHead(FStringView Uri);
+	DERIVEDDATACACHE_API Result PerformBlockingHead(FStringView Uri, EContentType AcceptType);
 
 	/**
 		* Query an url using the request. Queries can use either "Head" or "Delete" verbs.
@@ -156,8 +191,7 @@ private:
 	TArray<FString>			Headers;
 	FString					Domain;
 
-	void AddHeadersForContentType(EContentType ContentType);
-	void AddHeadersForAcceptType(EContentType ContentType);
+	void AddHeader(FStringView Header, FStringView Value);
 
 	/**
 	  * Supported request verb
diff --git a/Engine/Source/Developer/DerivedDataCache/Private/ZenDerivedDataBackend.cpp b/Engine/Source/Developer/DerivedDataCache/Private/ZenDerivedDataBackend.cpp
--- a/Engine/Source/Developer/DerivedDataCache/Private/ZenDerivedDataBackend.cpp
+++ b/Engine/Source/Developer/DerivedDataCache/Private/ZenDerivedDataBackend.cpp
@@ -99,7 +99,7 @@ bool FZenDerivedDataBackend::CachedDataProbablyExists(const TCHAR* CacheKey)
 	while (ResponseCode == 0 && ++Attempts < MaxAttempts)
 	{
 		Zen::FZenScopedRequestPtr Request(RequestPool.Get());
-		Zen::FZenHttpRequest::Result Result = Request->PerformBlockingHead(*Uri);
+		Zen::FZenHttpRequest::Result Result = Request->PerformBlockingHead(*Uri, Zen::EContentType::Binary);
 		ResponseCode = Request->GetResponseCode();
 
 		if (Zen::IsSuccessCode(ResponseCode) || ResponseCode == 404)
@@ -138,7 +138,7 @@ bool FZenDerivedDataBackend::GetCachedData(const TCHAR* CacheKey, TArray<uint8>&
 	double StartTime = FPlatformTime::Seconds();
 
 	TArray64<uint8> ArrayBuffer;
-	EGetResult Result = GetZenData(MakeLegacyZenKey(CacheKey), &ArrayBuffer);
+	EGetResult Result = GetZenData(MakeLegacyZenKey(CacheKey), &ArrayBuffer, Zen::EContentType::Binary);
 	check(ArrayBuffer.Num() <= UINT32_MAX);
 	OutData = MoveTemp(ArrayBuffer);
 	if (Result != EGetResult::Success)
@@ -167,7 +167,7 @@ bool FZenDerivedDataBackend::GetCachedData(const TCHAR* CacheKey, TArray<uint8>&
 }
 
 FZenDerivedDataBackend::EGetResult
-FZenDerivedDataBackend::GetZenData(FStringView Uri, TArray64<uint8>* OutData) const
+FZenDerivedDataBackend::GetZenData(FStringView Uri, TArray64<uint8>* OutData, Zen::EContentType ContentType) const
 {
 	// Retry request until we get an accepted response or exhaust allowed number of attempts.
 	EGetResult GetResult = EGetResult::NotFound;
@@ -179,11 +179,11 @@ FZenDerivedDataBackend::GetZenData(FStringView Uri, TArray64<uint8>* OutData) co
 			Zen::FZenHttpRequest::Result Result;
 			if (OutData)
 			{
-				Result = Request->PerformBlockingDownload(Uri, OutData);
+				Result = Request->PerformBlockingDownload(Uri, OutData, ContentType);
 			}
 			else
 			{
-				Result = Request->PerformBlockingHead(Uri);
+				Result = Request->PerformBlockingHead(Uri, ContentType);
 			}
 			int64 ResponseCode = Request->GetResponseCode();
 
@@ -633,7 +633,7 @@ void FZenDerivedDataBackend::GetPayload(
 			else
 			{
 				TArray64<uint8> PayloadData;
-				GetResult = GetZenData(QueryUri, &PayloadData);
+				GetResult = GetZenData(QueryUri, &PayloadData, Zen::EContentType::CompressedBinary);
 				if (GetResult == EGetResult::Success)
 				{
 					CompressedBuffer = FCompressedBuffer::FromCompressed(MakeSharedBufferFromArray(MoveTemp(PayloadData)));
@@ -719,7 +719,7 @@ bool FZenDerivedDataBackend::PutCacheRecord(const FCacheRecord& Record, FStringV
 	TStringBuilder<256> Uri;
 	AppendZenUri(Record.GetKey(), Uri);
 	AppendPolicyQueryString(Policy, Uri);
-	if (PutZenData(Uri.ToString(), Buffer, Zen::EContentType::CompactBinaryPackage)
+	if (PutZenData(Uri.ToString(), Buffer, Zen::EContentType::CbPackage)
 		!= FDerivedDataBackendInterface::EPutStatus::Cached)
 	{
 		return false;
@@ -796,7 +796,7 @@ bool FZenDerivedDataBackend::LegacyPutCacheRecord(const FCacheRecord& Record, FS
 	Writer.Save(Buffer);
 	TStringBuilder<256> Uri;
 	LegacyMakeZenKey(Key, Uri);
-	if (PutZenData(Uri.ToString(), FCompositeBuffer(Buffer.MoveToShared()), Zen::EContentType::CompactBinary)
+	if (PutZenData(Uri.ToString(), FCompositeBuffer(Buffer.MoveToShared()), Zen::EContentType::CbObject)
 		!= FDerivedDataBackendInterface::EPutStatus::Cached)
 	{
 		return false;
@@ -817,7 +817,7 @@ bool FZenDerivedDataBackend::LegacyPutCachePayload(const FCacheKey& Key, FString
 		LegacyMakePayloadKey(Key, RawHash, Uri);
 		if (Payload.HasData())
 		{
-			if (PutZenData(Uri.ToString(), CompressedBuffer, Zen::EContentType::Binary) != FDerivedDataBackendInterface::EPutStatus::Cached)
+			if (PutZenData(Uri.ToString(), CompressedBuffer, Zen::EContentType::CompressedBinary) != FDerivedDataBackendInterface::EPutStatus::Cached)
 			{
 				return false;
 			}
@@ -825,7 +825,7 @@ bool FZenDerivedDataBackend::LegacyPutCachePayload(const FCacheKey& Key, FString
 		else
 		{
 			// Put of a record with payloads-with-no-data should succeed iff the payload data is available in the CacheStore
-			if (GetZenData(Uri.ToString(), nullptr) != EGetResult::Success)
+			if (GetZenData(Uri.ToString(), nullptr, Zen::EContentType::CompressedBinary) != EGetResult::Success)
 			{
 				return false;
 			}
@@ -855,7 +855,7 @@ FOptionalCacheRecord FZenDerivedDataBackend::LegacyGetCacheRecord(const FCacheKe
 	TStringBuilder<256> Uri;
 	LegacyMakeZenKey(Key, Uri);
 	TArray64<uint8> Buffer;
-	EGetResult Result = GetZenData(Uri.ToString(), &Buffer);
+	EGetResult Result = GetZenData(Uri.ToString(), &Buffer, Zen::EContentType::CbObject);
 	if (Result != EGetResult::Success)
 	{
 		return FOptionalCacheRecord();
@@ -885,7 +885,7 @@ FPayload FZenDerivedDataBackend::LegacyGetCachePayload(const FCacheKey& Key, FSt
 	bool bSkipData = EnumHasAllFlags(Policy, ECachePolicy::SkipData);
 	EGetResult GetResult;
 	TArray64<uint8> ArrayBuffer;
-	GetResult = GetZenData(Uri.ToString(), bSkipData ? nullptr : &ArrayBuffer);
+	GetResult = GetZenData(Uri.ToString(), bSkipData ? nullptr : &ArrayBuffer, Zen::EContentType::CompressedBinary);
 	if (GetResult != EGetResult::Success)
 	{
 		switch (GetResult)
diff --git a/Engine/Source/Developer/DerivedDataCache/Private/ZenDerivedDataBackend.h b/Engine/Source/Developer/DerivedDataCache/Private/ZenDerivedDataBackend.h
--- a/Engine/Source/Developer/DerivedDataCache/Private/ZenDerivedDataBackend.h
+++ b/Engine/Source/Developer/DerivedDataCache/Private/ZenDerivedDataBackend.h
@@ -120,7 +120,7 @@ private:
 		NotFound,
 		Corrupted
 	};
-	EGetResult GetZenData(FStringView Uri, TArray64<uint8>* OutData) const;
+	EGetResult GetZenData(FStringView Uri, TArray64<uint8>* OutData, Zen::EContentType ContentType) const;
 
 	// TODO: need ability to specify content type
 	FDerivedDataBackendInterface::EPutStatus PutZenData(const TCHAR* Uri, const FCompositeBuffer& InData, Zen::EContentType ContentType);
diff --git a/Engine/Source/Developer/DerivedDataCache/Private/ZenServerHttp.cpp b/Engine/Source/Developer/DerivedDataCache/Private/ZenServerHttp.cpp
--- a/Engine/Source/Developer/DerivedDataCache/Private/ZenServerHttp.cpp
+++ b/Engine/Source/Developer/DerivedDataCache/Private/ZenServerHttp.cpp
@@ -21,6 +21,7 @@
 
 #include "Compression/CompressedBuffer.h"
 #include "Compression/OodleDataCompression.h"
+#include "Containers/StringFwd.h"
 #include "Memory/CompositeBuffer.h"
 #include "Serialization/CompactBinary.h"
 #include "Serialization/CompactBinaryPackage.h"
@@ -98,42 +99,11 @@ namespace UE::Zen {
 #endif
 	}
 
-	void FZenHttpRequest::AddHeadersForContentType(EContentType ContentType)
+	void FZenHttpRequest::AddHeader(FStringView Header, FStringView Value)
 	{
-		switch (ContentType)
-		{
-		case EContentType::Binary:
-			Headers.Add(FString(TEXT("Content-Type: application/octet-stream")));
-			break;
-		case EContentType::CompactBinary:
-			Headers.Add(FString(TEXT("Content-Type: application/x-ue-cb")));
-			break;
-		case EContentType::CompactBinaryPackage:
-			Headers.Add(FString(TEXT("Content-Type: application/x-ue-cbpkg")));
-			break;
-		default:
-			checkNoEntry();
-			break;
-		}
-	}
-
-	void FZenHttpRequest::AddHeadersForAcceptType(EContentType ContentType)
-	{
-		switch (ContentType)
-		{
-		case EContentType::Binary:
-			Headers.Add(FString(TEXT("Accept: application/octet-stream")));
-			break;
-		case EContentType::CompactBinary:
-			Headers.Add(FString(TEXT("Accept: application/x-ue-cb")));
-			break;
-		case EContentType::CompactBinaryPackage:
-			Headers.Add(FString(TEXT("Accept: application/x-ue-cbpkg")));
-			break;
-		default:
-			checkNoEntry();
-			break;
-		}
+		TStringBuilder<128> Sb;
+		Sb << Header << TEXT(": "_SV) << Value;
+		Headers.Emplace(Sb.ToString());
 	}
 
 	FZenHttpRequest::Result FZenHttpRequest::PerformBlockingPut(const TCHAR* Uri, const FCompositeBuffer& Buffer, EContentType ContentType)
@@ -146,7 +116,7 @@ namespace UE::Zen {
 		curl_easy_setopt(Curl, CURLOPT_READDATA, this);
 		curl_easy_setopt(Curl, CURLOPT_READFUNCTION, &FZenHttpRequest::FStatics::StaticReadFn);
 
-		AddHeadersForContentType(ContentType);
+		AddHeader(TEXT("Content-Type"_SV), GetMimeType(ContentType));
 
 		ReadDataView = &Buffer;
 
@@ -158,7 +128,7 @@ namespace UE::Zen {
 		FLargeMemoryWriter Out;
 		Obj.CopyTo(Out);
 
-		return PerformBlockingPost(Uri, Out.GetView(), EContentType::CompactBinary);
+		return PerformBlockingPost(Uri, Out.GetView(), EContentType::CbObject);
 	}
 
 	struct CbPackageHeader
@@ -279,7 +249,7 @@ namespace UE::Zen {
 			}
 		}
 
-		return PerformBlockingPost(Uri, Out.GetView(), EContentType::CompactBinaryPackage);
+		return PerformBlockingPost(Uri, Out.GetView(), EContentType::CbPackage);
 	}
 
 	FCbPackage FZenHttpRequest::GetResponseAsPackage() const
@@ -359,6 +329,8 @@ namespace UE::Zen {
 
 		FCompositeBuffer Buffer(FSharedBuffer::MakeView(FMemoryView{ reinterpret_cast<const uint8*>(Payload.GetData()), Payload.GetSize() }));
 		ReadDataView = &Buffer;
+		
+		AddHeader(TEXT("Content-Type"_SV), GetMimeType(EContentType::Binary));
 
 		return PerformBlocking(Uri, RequestVerb::Post, ContentLength);
 	}
@@ -372,7 +344,7 @@ namespace UE::Zen {
 		curl_easy_setopt(Curl, CURLOPT_READDATA, this);
 		curl_easy_setopt(Curl, CURLOPT_READFUNCTION, &FZenHttpRequest::FStatics::StaticReadFn);
 
-		AddHeadersForContentType(ContentType);
+		AddHeader(TEXT("Content-Type"_SV), GetMimeType(ContentType));
 
 		ContentLength = Payload.GetSize();
 
@@ -382,10 +354,12 @@ namespace UE::Zen {
 		return PerformBlocking(Uri, RequestVerb::Post, ContentLength);
 	}
 
-	FZenHttpRequest::Result FZenHttpRequest::PerformBlockingDownload(FStringView Uri, TArray64<uint8>* Buffer)
+	FZenHttpRequest::Result FZenHttpRequest::PerformBlockingDownload(FStringView Uri, TArray64<uint8>* Buffer, EContentType AcceptType)
 	{
 		curl_easy_setopt(Curl, CURLOPT_HTTPGET, 1L);
 		WriteDataBufferPtr = Buffer;
+		
+		AddHeader(TEXT("Accept"_SV), GetMimeType(AcceptType));
 
 		return PerformBlocking(Uri, RequestVerb::Get, 0u);
 	}
@@ -395,7 +369,7 @@ namespace UE::Zen {
 		curl_easy_setopt(Curl, CURLOPT_HTTPGET, 1L);
 		OutPackage.Reset();
 
-		AddHeadersForAcceptType(EContentType::CompactBinaryPackage);
+		AddHeader(TEXT("Accept"_SV), GetMimeType(EContentType::CbPackage));
 
 		// TODO: When PackageBytes can be written in segments directly, set the WritePtr to the OutPackage and use that
 		TArray64<uint8> PackageBytes;
@@ -409,9 +383,11 @@ namespace UE::Zen {
 		return LocalResult;
 	}
 
-	FZenHttpRequest::Result FZenHttpRequest::PerformBlockingHead(FStringView Uri)
+	FZenHttpRequest::Result FZenHttpRequest::PerformBlockingHead(FStringView Uri, EContentType AcceptType)
 	{
 		curl_easy_setopt(Curl, CURLOPT_NOBODY, 1L);
+		
+		AddHeader(TEXT("Accept"_SV), GetMimeType(AcceptType));
 
 		return PerformBlocking(Uri, RequestVerb::Head, 0u);
 	}
diff --git a/Engine/Source/Developer/DerivedDataCache/Private/ZenServerInterface.cpp b/Engine/Source/Developer/DerivedDataCache/Private/ZenServerInterface.cpp
--- a/Engine/Source/Developer/DerivedDataCache/Private/ZenServerInterface.cpp
+++ b/Engine/Source/Developer/DerivedDataCache/Private/ZenServerInterface.cpp
@@ -102,7 +102,7 @@ FZenServiceInstance::IsServiceReady()
 		TStringBuilder<128> ZenDomain;
 		ZenDomain << HostName << TEXT(":") << Port;
 		Zen::FZenHttpRequest Request(ZenDomain.ToString(), false);
-		Zen::FZenHttpRequest::Result Result = Request.PerformBlockingDownload(TEXT("health/ready"_SV), nullptr);
+		Zen::FZenHttpRequest::Result Result = Request.PerformBlockingDownload(TEXT("health/ready"_SV), nullptr, Zen::EContentType::Text);
 		
 		if (Result == Zen::FZenHttpRequest::Result::Success && Zen::IsSuccessCode(Request.GetResponseCode()))
 		{
diff --git a/Engine/Source/Developer/IoStoreUtilities/Private/ZenStoreHttpClient.cpp b/Engine/Source/Developer/IoStoreUtilities/Private/ZenStoreHttpClient.cpp
--- a/Engine/Source/Developer/IoStoreUtilities/Private/ZenStoreHttpClient.cpp
+++ b/Engine/Source/Developer/IoStoreUtilities/Private/ZenStoreHttpClient.cpp
@@ -61,7 +61,7 @@ FZenStoreHttpClient::Initialize(FStringView InProjectId,
 		TStringBuilder<128> ProjectUri;
 		ProjectUri << "/prj/" << InProjectId;
 		TArray64<uint8> GetBuffer;
-		UE::Zen::FZenHttpRequest::Result Res = Request->PerformBlockingDownload(ProjectUri, &GetBuffer);
+		UE::Zen::FZenHttpRequest::Result Res = Request->PerformBlockingDownload(ProjectUri, &GetBuffer, Zen::EContentType::CbObject);
 
 		// TODO: how to handle failure here? This is probably the most likely point of failure
 		// if the service is not up or not responding
@@ -127,7 +127,7 @@ void FZenStoreHttpClient::EstablishWritableOpLog(FStringView InProjectId, FStrin
 	}
 
 	TArray64<uint8> GetBuffer;
-	UE::Zen::FZenHttpRequest::Result Res = Request->PerformBlockingDownload(OplogPath, &GetBuffer);
+	UE::Zen::FZenHttpRequest::Result Res = Request->PerformBlockingDownload(OplogPath, &GetBuffer, Zen::EContentType::CbObject);
 	FCbObjectView OplogInfo;
 
 	if (Res == Zen::FZenHttpRequest::Result::Success && Request->GetResponseCode() == 200)
@@ -162,7 +162,7 @@ void FZenStoreHttpClient::EstablishWritableOpLog(FStringView InProjectId, FStrin
 
 		GetBuffer.Reset();
 		Request->Reset();
-		Res = Request->PerformBlockingDownload(OplogPath, &GetBuffer);
+		Res = Request->PerformBlockingDownload(OplogPath, &GetBuffer, Zen::EContentType::CbObject);
 
 		if (Res == Zen::FZenHttpRequest::Result::Success && Request->GetResponseCode() == 200)
 		{
@@ -185,7 +185,7 @@ void FZenStoreHttpClient::InitializeReadOnly(FStringView InProjectId, FStringVie
 		UE::Zen::FZenScopedRequestPtr Request(RequestPool.Get());
 
 		TArray64<uint8> GetBuffer;
-		UE::Zen::FZenHttpRequest::Result Res = Request->PerformBlockingDownload(WriteToString<128>("/prj/", InProjectId), &GetBuffer);
+		UE::Zen::FZenHttpRequest::Result Res = Request->PerformBlockingDownload(WriteToString<128>("/prj/", InProjectId), &GetBuffer, Zen::EContentType::CbObject);
 
 		// TODO: how to handle failure here? This is probably the most likely point of failure
 		// if the service is not up or not responding
@@ -210,7 +210,7 @@ void FZenStoreHttpClient::InitializeReadOnly(FStringView InProjectId, FStringVie
 		OplogPath = WriteToString<128>("/prj/", InProjectId, "/oplog/", InOplogId);
 
 		TArray64<uint8> GetBuffer;
-		UE::Zen::FZenHttpRequest::Result Res = Request->PerformBlockingDownload(OplogPath, &GetBuffer);
+		UE::Zen::FZenHttpRequest::Result Res = Request->PerformBlockingDownload(OplogPath, &GetBuffer, Zen::EContentType::CbObject);
 
 		if (Res != Zen::FZenHttpRequest::Result::Success || Request->GetResponseCode() != 200)
 		{
@@ -398,7 +398,7 @@ TIoStatusOr<uint64> FZenStoreHttpClient::GetChunkSize(const FIoChunkId& Id)
 	TArray64<uint8> GetBuffer;
 	TStringBuilder<128> ChunkUri;
 	ChunkUri << OplogPath << '/' << Id;
-	UE::Zen::FZenHttpRequest::Result Res = Request->PerformBlockingHead(ChunkUri);
+	UE::Zen::FZenHttpRequest::Result Res = Request->PerformBlockingHead(ChunkUri, Zen::EContentType::Binary);
 	FString ContentLengthStr;
 	if (Res == Zen::FZenHttpRequest::Result::Success && Request->GetResponseCode() == 200 && Request->GetHeader("Content-Length", ContentLengthStr))
 	{
@@ -457,7 +457,7 @@ TIoStatusOr<FIoBuffer> FZenStoreHttpClient::ReadOpLogUri(FStringBuilderBase& Chu
 	}
 
 
-	UE::Zen::FZenHttpRequest::Result Res = Request->PerformBlockingDownload(ChunkUri, &GetBuffer);
+	UE::Zen::FZenHttpRequest::Result Res = Request->PerformBlockingDownload(ChunkUri, &GetBuffer, Zen::EContentType::Binary);
 
 	if (Res == Zen::FZenHttpRequest::Result::Success && Request->GetResponseCode() == 200)
 	{
@@ -481,7 +481,7 @@ TFuture<TIoStatusOr<FCbObject>> FZenStoreHttpClient::GetOplog()
 		Uri << OplogPath << "/entries";
 
 		TArray64<uint8> GetBuffer;
-		UE::Zen::FZenHttpRequest::Result Res = Request->PerformBlockingDownload(Uri, &GetBuffer);
+		UE::Zen::FZenHttpRequest::Result Res = Request->PerformBlockingDownload(Uri, &GetBuffer, Zen::EContentType::CbObject);
 
 		if (Res == Zen::FZenHttpRequest::Result::Success && Request->GetResponseCode() == 200)
 		{
@@ -510,7 +510,7 @@ TFuture<TIoStatusOr<FCbObject>> FZenStoreHttpClient::GetFiles()
 		Uri << OplogPath << "/files";
 
 		TArray64<uint8> GetBuffer;
-		UE::Zen::FZenHttpRequest::Result Res = Request->PerformBlockingDownload(Uri, &GetBuffer);
+		UE::Zen::FZenHttpRequest::Result Res = Request->PerformBlockingDownload(Uri, &GetBuffer, Zen::EContentType::CbObject);
 
 		if (Res == Zen::FZenHttpRequest::Result::Success && Request->GetResponseCode() == 200)
 		{
diff --git a/Engine/Source/Editor/UnrealEd/Private/DerivedDataBuildZenExecutor.cpp b/Engine/Source/Editor/UnrealEd/Private/DerivedDataBuildZenExecutor.cpp
--- a/Engine/Source/Editor/UnrealEd/Private/DerivedDataBuildZenExecutor.cpp
+++ b/Engine/Source/Editor/UnrealEd/Private/DerivedDataBuildZenExecutor.cpp
@@ -173,7 +173,7 @@ public:
 
 			// Check if the server knows this worker
 
-			UE::Zen::FZenHttpRequest::Result WorkerGetResult = WorkerRequest->PerformBlockingDownload(WorkerUri, nullptr);
+			UE::Zen::FZenHttpRequest::Result WorkerGetResult = WorkerRequest->PerformBlockingDownload(WorkerUri, nullptr, Zen::EContentType::CbObject);
 
 			if (WorkerGetResult == Zen::FZenHttpRequest::Result::Failed)
 			{
@@ -593,4 +593,4 @@ void DumpDerivedDataBuildZenExecutorStats()
 {
 }
 
-#endif
\ No newline at end of file
+#endif
