DDC: Removed DDCUtils in favor of a config-based approach that can work for any project

diff --git a/Engine/Source/Developer/DerivedDataCache/DerivedDataCache.Build.cs b/Engine/Source/Developer/DerivedDataCache/DerivedDataCache.Build.cs
--- a/Engine/Source/Developer/DerivedDataCache/DerivedDataCache.Build.cs
+++ b/Engine/Source/Developer/DerivedDataCache/DerivedDataCache.Build.cs
@@ -15,13 +15,6 @@ public class DerivedDataCache : ModuleRules
 		AddEngineThirdPartyPrivateStaticDependencies(Target, "libcurl");
 		AddEngineThirdPartyPrivateStaticDependencies(Target, "OpenSSL");
 
-		// Internal (NotForLicensees) module
-		var DDCUtilsModule = Path.Combine(EngineDirectory, "Restricted/NotForLicensees/Source/Developer/DDCUtils/DDCUtils.Build.cs");
-		if (File.Exists(DDCUtilsModule))
-		{
-			DynamicallyLoadedModuleNames.Add("DDCUtils");
-		}
-
 		// Platform-specific opt-in
 		if (Target.Platform == UnrealTargetPlatform.Win64)
 		{
diff --git a/Engine/Source/Developer/DerivedDataCache/Private/DerivedDataBackends.cpp b/Engine/Source/Developer/DerivedDataCache/Private/DerivedDataBackends.cpp
--- a/Engine/Source/Developer/DerivedDataCache/Private/DerivedDataBackends.cpp
+++ b/Engine/Source/Developer/DerivedDataCache/Private/DerivedDataBackends.cpp
@@ -22,7 +22,6 @@
 #include "DerivedDataBackendCorruptionWrapper.h"
 #include "DerivedDataBackendThrottleWrapper.h"
 #include "DerivedDataBackendVerifyWrapper.h"
-#include "DerivedDataUtilsInterface.h"
 #include "Misc/EngineBuildSettings.h"
 #include "Modules/ModuleManager.h"
 #include "Misc/ConfigCacheIni.h"
@@ -578,22 +577,10 @@ public:
 			}
 		}
 
-		// Check if the Path is a real path or a special keyword
-		if (FEngineBuildSettings::IsInternalBuild())
+		// Paths starting with a '?' are looked up from config
+		if (Path.StartsWith(TEXT("?")) && !GConfig->GetString(TEXT("DerivedDataCacheSettings"), *Path + 1, Path, GEngineIni))
 		{
-			auto DDCUtils = FModuleManager::LoadModulePtr< IDDCUtilsModuleInterface >("DDCUtils");
-			if (DDCUtils)
-			{
-				FString PathFromName = DDCUtils->GetSharedCachePath(Path);
-				if (!PathFromName.IsEmpty())
-				{
-					Path = PathFromName;
-				}
-			}
-		}
-		else if (Path.StartsWith(TEXT("?")))
-		{
-			Path = TEXT("");
+			Path.Empty();
 		}
 
 		// Allow the user to override it from the editor
diff --git a/Engine/Source/Developer/DerivedDataCache/Public/DerivedDataUtilsInterface.h b/Engine/Source/Developer/DerivedDataCache/Public/DerivedDataUtilsInterface.h
deleted file mode 100644
diff --git a/Engine/Source/Programs/DerivedDataBuildWorker/DerivedDataBuildWorker.Build.cs b/Engine/Source/Programs/DerivedDataBuildWorker/DerivedDataBuildWorker.Build.cs
--- a/Engine/Source/Programs/DerivedDataBuildWorker/DerivedDataBuildWorker.Build.cs
+++ b/Engine/Source/Programs/DerivedDataBuildWorker/DerivedDataBuildWorker.Build.cs
@@ -13,10 +13,6 @@ public abstract class DerivedDataBuildWorkerTarget : TargetRules
 
 		SolutionDirectory = "Programs/BuildWorker";
 
-		// TODO: This is temporary to allow DDCUtils to not cause build errors when referenced
-		//		 by the build workers.  In the future the use of DDC should strip out backends
-		//		 and DDCUTils, so it will no longer be required to set bLegalToDistributeBinary.
-		bLegalToDistributeBinary		= true;
 		bUseXGEController				= false;
 		bCompileFreeType				= false;
 		bLoggingToMemoryEnabled			= true;
