Fixed usage of UE_LOG (arguments for "%s" formatting).

diff --git a/Engine/Source/Developer/DerivedDataCache/Private/DerivedDataBackends.cpp b/Engine/Source/Developer/DerivedDataCache/Private/DerivedDataBackends.cpp
--- a/Engine/Source/Developer/DerivedDataCache/Private/DerivedDataBackends.cpp
+++ b/Engine/Source/Developer/DerivedDataCache/Private/DerivedDataBackends.cpp
@@ -181,7 +181,7 @@ public:
 					}
 					else
 					{
-						UE_LOG( LogDerivedDataCache, Warning, TEXT("FDerivedDataBackendGraph:  Unable to create %s Boot cache because only one Boot or Sparse cache node is supported."), *NodeName );
+						UE_LOG( LogDerivedDataCache, Warning, TEXT("FDerivedDataBackendGraph:  Unable to create %s Boot cache because only one Boot or Sparse cache node is supported."), NodeName );
 					}
 				}
 				else if( NodeType == TEXT("Memory") )
@@ -201,7 +201,7 @@ public:
 					}
 					else
 					{
-						UE_LOG( LogDerivedDataCache, Warning, TEXT("FDerivedDataBackendGraph:  Unable to create %s KeyLengthWrapper because only one KeyLength node is supported."), *NodeName );
+						UE_LOG( LogDerivedDataCache, Warning, TEXT("FDerivedDataBackendGraph:  Unable to create %s KeyLengthWrapper because only one KeyLength node is supported."), NodeName );
 					}
 				}
 				else if( NodeType == TEXT("AsyncPut") )
@@ -213,7 +213,7 @@ public:
 					}
 					else
 					{
-						UE_LOG( LogDerivedDataCache, Warning, TEXT("FDerivedDataBackendGraph:  Unable to create %s AsyncPutWrapper because only one AsyncPutWrapper node is supported."), *NodeName );
+						UE_LOG( LogDerivedDataCache, Warning, TEXT("FDerivedDataBackendGraph:  Unable to create %s AsyncPutWrapper because only one AsyncPutWrapper node is supported."), NodeName );
 					}
 				}
 				else if( NodeType == TEXT("Verify") )
diff --git a/Engine/Source/Runtime/Core/Private/Serialization/Archive.cpp b/Engine/Source/Runtime/Core/Private/Serialization/Archive.cpp
--- a/Engine/Source/Runtime/Core/Private/Serialization/Archive.cpp
+++ b/Engine/Source/Runtime/Core/Private/Serialization/Archive.cpp
@@ -784,7 +784,7 @@ void FArchive::SerializeCompressedNew(void* V, int64 Length, FName CompressionFo
 				// go ahead and do it but warn :
 				CompressionFormatToDecode = FName(ZlibReplacement);
 
-				UE_LOG(LogSerialization, Warning, TEXT("Archive v2 header with ZLib not ZlibReplacement: %s"), ZlibReplacement);
+				UE_LOG(LogSerialization, Warning, TEXT("Archive v2 header with ZLib not ZlibReplacement: %s"), *CompressionFormatToDecode.ToString());
 			}
 		}
 
