DDC: Reworked ICacheStore to allow partial records, filtering of payloads, and loading parts of payloads

- ICacheStore::Put() has updated documentation to reflect the requirements of partial records.
- ICacheStore::Get() now takes a FCacheRecordPolicy, which is implicitly constructible from ECachePolicy, and allows setting the policy by payload.
- ICacheStore::GetPayload() is replaced by ICacheStore::GetChunks(), which allows loading parts of payloads.
- ICacheStore::CancelAll() is moved to ICache::CancelAll() because the cache can track requests at the top level and cancel them without exposing cancellation on individual cache stores.
- ECachePolicy::SkipLocalCopy has been removed because it is difficult to reason about.
- ECachePolicy::SkipData flags now have a documented meaning for put requests, to hint that record existence implies payload existence.
- The filesystem and memory cache stores have been updated to support partial records, filtering of payloads, and loading parts of payloads.
- Requesting part of a payload will decompress the entire payload for now, until compressed buffers expose a way to decompress only part.
- Fixed a bug in FTexturePlatformData::AreDerivedMipsAvailable() that caused it to return false for structured cache keys.

#rnx

diff --git a/Engine/Source/Developer/DerivedDataCache/Private/DerivedDataBackendAsyncPutWrapper.cpp b/Engine/Source/Developer/DerivedDataCache/Private/DerivedDataBackendAsyncPutWrapper.cpp
--- a/Engine/Source/Developer/DerivedDataCache/Private/DerivedDataBackendAsyncPutWrapper.cpp
+++ b/Engine/Source/Developer/DerivedDataCache/Private/DerivedDataBackendAsyncPutWrapper.cpp
@@ -493,7 +493,7 @@ void FDerivedDataBackendAsyncPutWrapper::Put(
 void FDerivedDataBackendAsyncPutWrapper::Get(
 	TConstArrayView<FCacheKey> Keys,
 	FStringView Context,
-	ECachePolicy Policy,
+	FCacheRecordPolicy Policy,
 	IRequestOwner& Owner,
 	FOnCacheGetComplete&& OnComplete)
 {
@@ -524,33 +524,32 @@ void FDerivedDataBackendAsyncPutWrapper::Get(
 	}
 }
 
-void FDerivedDataBackendAsyncPutWrapper::GetPayload(
-	TConstArrayView<FCachePayloadKey> Keys,
+void FDerivedDataBackendAsyncPutWrapper::GetChunks(
+	TConstArrayView<FCacheChunkRequest> Chunks,
 	FStringView Context,
-	ECachePolicy Policy,
 	IRequestOwner& Owner,
-	FOnCacheGetPayloadComplete&& OnComplete)
+	FOnCacheGetChunkComplete&& OnComplete)
 {
 	if (Owner.GetPriority() == EPriority::Blocking)
 	{
-		InnerBackend->GetPayload(Keys, Context, Policy, Owner, MoveTemp(OnComplete));
+		InnerBackend->GetChunks(Chunks, Context, Owner, MoveTemp(OnComplete));
 	}
 	else
 	{
 		FDerivedDataAsyncWrapperRequest* Request = new FDerivedDataAsyncWrapperRequest(Owner,
-			[this, Keys = TArray<FCachePayloadKey>(Keys), Context = FString(Context), Policy, OnComplete = MoveTemp(OnComplete)](bool bCancel) mutable
+			[this, Chunks = TArray<FCacheChunkRequest>(Chunks), Context = FString(Context), OnComplete = MoveTemp(OnComplete)](bool bCancel) mutable
 			{
 				if (!bCancel)
 				{
 					FRequestOwner BlockingOwner(EPriority::Blocking);
-					InnerBackend->GetPayload(Keys, Context, Policy, BlockingOwner, MoveTemp(OnComplete));
+					InnerBackend->GetChunks(Chunks, Context, BlockingOwner, MoveTemp(OnComplete));
 					BlockingOwner.Wait();
 				}
 				else if (OnComplete)
 				{
-					for (const FCachePayloadKey& Key : Keys)
+					for (const FCacheChunkRequest& Chunk : Chunks)
 					{
-						OnComplete({Key.CacheKey, FPayload(Key.Id), EStatus::Canceled});
+						OnComplete({Chunk.Key, Chunk.Id, Chunk.RawOffset, 0, {}, {}, EStatus::Canceled});
 					}
 				}
 			});
@@ -558,9 +557,4 @@ void FDerivedDataBackendAsyncPutWrapper::GetPayload(
 	}
 }
 
-void FDerivedDataBackendAsyncPutWrapper::CancelAll()
-{
-	InnerBackend->CancelAll();
-}
-
 } // UE::DerivedData::Backends
diff --git a/Engine/Source/Developer/DerivedDataCache/Private/DerivedDataBackendAsyncPutWrapper.h b/Engine/Source/Developer/DerivedDataCache/Private/DerivedDataBackendAsyncPutWrapper.h
--- a/Engine/Source/Developer/DerivedDataCache/Private/DerivedDataBackendAsyncPutWrapper.h
+++ b/Engine/Source/Developer/DerivedDataCache/Private/DerivedDataBackendAsyncPutWrapper.h
@@ -145,18 +145,15 @@ public:
 	virtual void Get(
 		TConstArrayView<FCacheKey> Keys,
 		FStringView Context,
-		ECachePolicy Policy,
+		FCacheRecordPolicy Policy,
 		IRequestOwner& Owner,
 		FOnCacheGetComplete&& OnComplete) override;
 
-	virtual void GetPayload(
-		TConstArrayView<FCachePayloadKey> Keys,
+	virtual void GetChunks(
+		TConstArrayView<FCacheChunkRequest> Chunks,
 		FStringView Context,
-		ECachePolicy Policy,
 		IRequestOwner& Owner,
-		FOnCacheGetPayloadComplete&& OnComplete) override;
-
-	virtual void CancelAll() override;
+		FOnCacheGetChunkComplete&& OnComplete) override;
 
 private:
 	FDerivedDataCacheUsageStats UsageStats;
diff --git a/Engine/Source/Developer/DerivedDataCache/Private/DerivedDataBackendCorruptionWrapper.h b/Engine/Source/Developer/DerivedDataCache/Private/DerivedDataBackendCorruptionWrapper.h
--- a/Engine/Source/Developer/DerivedDataCache/Private/DerivedDataBackendCorruptionWrapper.h
+++ b/Engine/Source/Developer/DerivedDataCache/Private/DerivedDataBackendCorruptionWrapper.h
@@ -264,26 +264,20 @@ public:
 	virtual void Get(
 		TConstArrayView<FCacheKey> Keys,
 		FStringView Context,
-		ECachePolicy Policy,
+		FCacheRecordPolicy Policy,
 		IRequestOwner& Owner,
 		FOnCacheGetComplete&& OnComplete) override
 	{
 		return InnerBackend->Get(Keys, Context, Policy, Owner, MoveTemp(OnComplete));
 	}
 
-	virtual void GetPayload(
-		TConstArrayView<FCachePayloadKey> Keys,
+	virtual void GetChunks(
+		TConstArrayView<FCacheChunkRequest> Chunks,
 		FStringView Context,
-		ECachePolicy Policy,
 		IRequestOwner& Owner,
-		FOnCacheGetPayloadComplete&& OnComplete) override
-	{
-		return InnerBackend->GetPayload(Keys, Context, Policy, Owner, MoveTemp(OnComplete));
-	}
-
-	virtual void CancelAll() override
+		FOnCacheGetChunkComplete&& OnComplete) override
 	{
-		InnerBackend->CancelAll();
+		return InnerBackend->GetChunks(Chunks, Context, Owner, MoveTemp(OnComplete));
 	}
 
 private:
diff --git a/Engine/Source/Developer/DerivedDataCache/Private/DerivedDataBackendThrottleWrapper.h b/Engine/Source/Developer/DerivedDataCache/Private/DerivedDataBackendThrottleWrapper.h
--- a/Engine/Source/Developer/DerivedDataCache/Private/DerivedDataBackendThrottleWrapper.h
+++ b/Engine/Source/Developer/DerivedDataCache/Private/DerivedDataBackendThrottleWrapper.h
@@ -176,7 +176,7 @@ public:
 	virtual void Get(
 		TConstArrayView<FCacheKey> Keys,
 		FStringView Context,
-		ECachePolicy Policy,
+		FCacheRecordPolicy Policy,
 		IRequestOwner& Owner,
 		FOnCacheGetComplete&& OnComplete) override
 	{
@@ -185,21 +185,15 @@ public:
 		InnerBackend->Get(Keys, Context, Policy, Owner, MoveTemp(OnComplete));
 	}
 
-	virtual void GetPayload(
-		TConstArrayView<FCachePayloadKey> Keys,
+	virtual void GetChunks(
+		TConstArrayView<FCacheChunkRequest> Chunks,
 		FStringView Context,
-		ECachePolicy Policy,
 		IRequestOwner& Owner,
-		FOnCacheGetPayloadComplete&& OnComplete) override
+		FOnCacheGetChunkComplete&& OnComplete) override
 	{
 		ThrottlingScope Scope(this);
 
-		InnerBackend->GetPayload(Keys, Context, Policy, Owner, MoveTemp(OnComplete));
-	}
-
-	virtual void CancelAll() override
-	{
-		InnerBackend->CancelAll();
+		InnerBackend->GetChunks(Chunks, Context, Owner, MoveTemp(OnComplete));
 	}
 
 private:
diff --git a/Engine/Source/Developer/DerivedDataCache/Private/DerivedDataBackendVerifyWrapper.h b/Engine/Source/Developer/DerivedDataCache/Private/DerivedDataBackendVerifyWrapper.h
--- a/Engine/Source/Developer/DerivedDataCache/Private/DerivedDataBackendVerifyWrapper.h
+++ b/Engine/Source/Developer/DerivedDataCache/Private/DerivedDataBackendVerifyWrapper.h
@@ -175,26 +175,20 @@ public:
 	virtual void Get(
 		TConstArrayView<FCacheKey> Keys,
 		FStringView Context,
-		ECachePolicy Policy,
+		FCacheRecordPolicy Policy,
 		IRequestOwner& Owner,
 		FOnCacheGetComplete&& OnComplete) override
 	{
 		return InnerBackend->Get(Keys, Context, Policy, Owner, MoveTemp(OnComplete));
 	}
 
-	virtual void GetPayload(
-		TConstArrayView<FCachePayloadKey> Keys,
+	virtual void GetChunks(
+		TConstArrayView<FCacheChunkRequest> Chunks,
 		FStringView Context,
-		ECachePolicy Policy,
 		IRequestOwner& Owner,
-		FOnCacheGetPayloadComplete&& OnComplete) override
-	{
-		return InnerBackend->GetPayload(Keys, Context, Policy, Owner, MoveTemp(OnComplete));
-	}
-
-	virtual void CancelAll() override
+		FOnCacheGetChunkComplete&& OnComplete) override
 	{
-		InnerBackend->CancelAll();
+		return InnerBackend->GetChunks(Chunks, Context, Owner, MoveTemp(OnComplete));
 	}
 
 private:
diff --git a/Engine/Source/Developer/DerivedDataCache/Private/DerivedDataCache.cpp b/Engine/Source/Developer/DerivedDataCache/Private/DerivedDataCache.cpp
--- a/Engine/Source/Developer/DerivedDataCache/Private/DerivedDataCache.cpp
+++ b/Engine/Source/Developer/DerivedDataCache/Private/DerivedDataCache.cpp
@@ -3,8 +3,10 @@
 #include "DerivedDataCache.h"
 #include "DerivedDataCacheInterface.h"
 
-#include "CoreMinimal.h"
+#include "Algo/Accumulate.h"
 #include "Algo/AllOf.h"
+#include "Algo/BinarySearch.h"
+#include "Algo/Sort.h"
 #include "Async/AsyncWork.h"
 #include "Async/TaskGraphInterfaces.h"
 #include "DDCCleanup.h"
@@ -154,6 +156,87 @@ void GatherDerivedDataCacheSummaryStats(FDerivedDataCacheSummaryStats& DDCSummar
 /** Whether we want to verify the DDC (pass in -VerifyDDC on the command line)*/
 bool GVerifyDDC = false;
 
+namespace UE::DerivedData::Private { class FCacheRecordPolicyShared; }
+
+namespace UE::DerivedData
+{
+
+class Private::FCacheRecordPolicyShared final : public Private::ICacheRecordPolicyShared
+{
+public:
+	inline void AddRef() const final
+	{
+		ReferenceCount.fetch_add(1, std::memory_order_relaxed);
+	}
+
+	inline void Release() const final
+	{
+		if (ReferenceCount.fetch_sub(1, std::memory_order_acq_rel) == 1)
+		{
+			delete this;
+		}
+	}
+
+	inline TConstArrayView<FCachePayloadPolicy> GetPayloadPolicies() const final
+	{
+		return Payloads;
+	}
+
+	inline void AddPayloadPolicy(const FCachePayloadPolicy& Policy) final
+	{
+		Payloads.Add(Policy);
+	}
+
+	inline void Build() final
+	{
+		Algo::SortBy(Payloads, &FCachePayloadPolicy::Id);
+	}
+
+private:
+	TArray<FCachePayloadPolicy, TInlineAllocator<14>> Payloads;
+	mutable std::atomic<uint32> ReferenceCount{0};
+};
+
+ECachePolicy FCacheRecordPolicy::GetPayloadPolicy(const FPayloadId& Id) const
+{
+	if (Shared)
+	{
+		if (TConstArrayView<FCachePayloadPolicy> Payloads = Shared->GetPayloadPolicies(); !Payloads.IsEmpty())
+		{
+			if (int32 Index = Algo::BinarySearchBy(Payloads, Id, &FCachePayloadPolicy::Id); Index != INDEX_NONE)
+			{
+				return Payloads[Index].Policy;
+			}
+		}
+	}
+	return BasePolicy;
+}
+
+void FCacheRecordPolicyBuilder::AddPayloadPolicy(const FCachePayloadPolicy& Policy)
+{
+	if (!Shared)
+	{
+		Shared = new Private::FCacheRecordPolicyShared;
+	}
+	Shared->AddPayloadPolicy(Policy);
+}
+
+FCacheRecordPolicy FCacheRecordPolicyBuilder::Build()
+{
+	FCacheRecordPolicy Policy(BasePolicy);
+	if (Shared)
+	{
+		Shared->Build();
+		const auto PolicyOr = [](ECachePolicy A, ECachePolicy B) { return A | (B & ~ECachePolicy::SkipData); };
+		const TConstArrayView<FCachePayloadPolicy> Payloads = Shared->GetPayloadPolicies();
+		Policy.RecordPolicy = Algo::TransformAccumulate(Payloads, &FCachePayloadPolicy::Policy, BasePolicy, PolicyOr);
+		Policy.Shared = MoveTemp(Shared);
+	}
+	return Policy;
+}
+
+} // UE::DerivedData
+
 namespace UE::DerivedData::Private
 {
 
@@ -694,26 +777,24 @@ public:
 	void  Get(
 		TConstArrayView<FCacheKey> Keys,
 		FStringView Context,
-		ECachePolicy Policy,
+		FCacheRecordPolicy Policy,
 		IRequestOwner& Owner,
 		FOnCacheGetComplete&& OnComplete) final
 	{
 		return FDerivedDataBackend::Get().GetRoot().Get(Keys, Context, Policy, Owner, MoveTemp(OnComplete));
 	}
 
-	void GetPayload(
-		TConstArrayView<FCachePayloadKey> Keys,
+	void GetChunks(
+		TConstArrayView<FCacheChunkRequest> Chunks,
 		FStringView Context,
-		ECachePolicy Policy,
 		IRequestOwner& Owner,
-		FOnCacheGetPayloadComplete&& OnComplete) final
+		FOnCacheGetChunkComplete&& OnComplete) final
 	{
-		return FDerivedDataBackend::Get().GetRoot().GetPayload(Keys, Context, Policy, Owner, MoveTemp(OnComplete));
+		return FDerivedDataBackend::Get().GetRoot().GetChunks(Chunks, Context, Owner, MoveTemp(OnComplete));
 	}
 
 	void CancelAll() final
 	{
-		return FDerivedDataBackend::Get().GetRoot().CancelAll();
 	}
 };
 
diff --git a/Engine/Source/Developer/DerivedDataCache/Private/DerivedDataLimitKeyLengthWrapper.h b/Engine/Source/Developer/DerivedDataCache/Private/DerivedDataLimitKeyLengthWrapper.h
--- a/Engine/Source/Developer/DerivedDataCache/Private/DerivedDataLimitKeyLengthWrapper.h
+++ b/Engine/Source/Developer/DerivedDataCache/Private/DerivedDataLimitKeyLengthWrapper.h
@@ -255,26 +255,20 @@ public:
 	virtual void Get(
 		TConstArrayView<FCacheKey> Keys,
 		FStringView Context,
-		ECachePolicy Policy,
+		FCacheRecordPolicy Policy,
 		IRequestOwner& Owner,
 		FOnCacheGetComplete&& OnComplete) override
 	{
 		return InnerBackend->Get(Keys, Context, Policy, Owner, MoveTemp(OnComplete));
 	}
 
-	virtual void GetPayload(
-		TConstArrayView<FCachePayloadKey> Keys,
+	virtual void GetChunks(
+		TConstArrayView<FCacheChunkRequest> Chunks,
 		FStringView Context,
-		ECachePolicy Policy,
 		IRequestOwner& Owner,
-		FOnCacheGetPayloadComplete&& OnComplete) override
-	{
-		return InnerBackend->GetPayload(Keys, Context, Policy, Owner, MoveTemp(OnComplete));
-	}
-
-	virtual void CancelAll() override
+		FOnCacheGetChunkComplete&& OnComplete) override
 	{
-		InnerBackend->CancelAll();
+		return InnerBackend->GetChunks(Chunks, Context, Owner, MoveTemp(OnComplete));
 	}
 
 private:
diff --git a/Engine/Source/Developer/DerivedDataCache/Private/FileSystemDerivedDataBackend.cpp b/Engine/Source/Developer/DerivedDataCache/Private/FileSystemDerivedDataBackend.cpp
--- a/Engine/Source/Developer/DerivedDataCache/Private/FileSystemDerivedDataBackend.cpp
+++ b/Engine/Source/Developer/DerivedDataCache/Private/FileSystemDerivedDataBackend.cpp
@@ -23,9 +23,10 @@
 #include "Async/TaskGraphInterfaces.h"
 #include "Async/Async.h"
 
+#include "DerivedDataBackendInterface.h"
 #include "DerivedDataCacheInterface.h"
 #include "DerivedDataCacheRecord.h"
-#include "DerivedDataBackendInterface.h"
+#include "DerivedDataChunk.h"
 #include "DerivedDataPayload.h"
 #include "DDCCleanup.h"
 
@@ -520,7 +521,7 @@ public:
 			FScopeLock Lock(&CriticalSection);
 
 			bool bIsAlreadyInSet = false;
-			PayloadKeys.FindOrAdd(FCachePayloadKey{CacheKey, Id}, &bIsAlreadyInSet);
+			PayloadKeys.FindOrAdd(MakeTuple(CacheKey, Id), &bIsAlreadyInSet);
 			if (!bIsAlreadyInSet)
 			{
 				AppendPath(Path);
@@ -543,7 +544,7 @@ public:
 		FCriticalSection CriticalSection;
 		TSet<FString> CacheKeys;
 		TSet<FCacheKey> ObjectKeys;
-		TSet<FCachePayloadKey> PayloadKeys;
+		TSet<TTuple<FCacheKey, FPayloadId>> PayloadKeys;
 	};
 	
 	
@@ -918,7 +919,7 @@ public:
 	virtual void Get(
 		TConstArrayView<FCacheKey> Keys,
 		FStringView Context,
-		ECachePolicy Policy,
+		FCacheRecordPolicy Policy,
 		IRequestOwner& Owner,
 		FOnCacheGetComplete&& OnComplete) override
 	{
@@ -949,52 +950,57 @@ public:
 		}
 	}
 
-	virtual void GetPayload(
-		TConstArrayView<FCachePayloadKey> Keys,
+	virtual void GetChunks(
+		TConstArrayView<FCacheChunkRequest> Chunks,
 		FStringView Context,
-		ECachePolicy Policy,
 		IRequestOwner& Owner,
-		FOnCacheGetPayloadComplete&& OnComplete) override
+		FOnCacheGetChunkComplete&& OnComplete) override
 	{
-		TArray<FCachePayloadKey, TInlineAllocator<16>> SortedKeys(Keys);
-		SortedKeys.StableSort();
+		TArray<FCacheChunkRequest, TInlineAllocator<16>> SortedChunks(Chunks);
+		SortedChunks.StableSort(TChunkLess());
 
 		FOptionalCacheRecord Record;
-		for (const FCachePayloadKey& Key : SortedKeys)
+		for (const FCacheChunkRequest& Chunk : SortedChunks)
 		{
 			TRACE_CPUPROFILER_EVENT_SCOPE(FileSystemDDC_Get);
 			TRACE_COUNTER_INCREMENT(FileSystemDDC_Get);
 			COOK_STAT(auto Timer = UsageStats.TimeGet());
-			if (!Record || Record.Get().GetKey() != Key.CacheKey)
+			if (!Record || Record.Get().GetKey() != Chunk.Key)
 			{
-				Record = GetCacheRecord(Key.CacheKey, Context, Policy | ECachePolicy::SkipData, /*bAlwaysLoadInlineData*/ true);
+				FCacheRecordPolicyBuilder PolicyBuilder(ECachePolicy::None);
+				PolicyBuilder.AddPayloadPolicy(Chunk.Id, Chunk.Policy | ECachePolicy::SkipData);
+				Record = GetCacheRecord(Chunk.Key, Context, PolicyBuilder.Build(), /*bAlwaysLoadInlineData*/ true);
 			}
-			if (FPayload Payload = Record ? GetCachePayload(Key.CacheKey, Context, Policy, Record.Get().GetPayload(Key.Id)) : FPayload::Null)
+			if (Record)
 			{
-				UE_LOG(LogDerivedDataCache, Verbose, TEXT("%s: Cache hit for %s from '%.*s'"),
-					*CachePath, *WriteToString<96>(Key), Context.Len(), Context.GetData());
-				TRACE_COUNTER_INCREMENT(FileSystemDDC_GetHit);
-				TRACE_COUNTER_ADD(FileSystemDDC_BytesRead, Payload.GetRawSize());
-				COOK_STAT(Timer.AddHit(Payload.GetRawSize()));
-				if (OnComplete)
+				const ECachePolicy PolicyMask = Record.Get().GetValuePayload().GetId() == Chunk.Id
+					? ECachePolicy::SkipValue : ECachePolicy::SkipAttachments;
+				const ECachePolicy Policy = Chunk.Policy | (ECachePolicy::SkipData & ~PolicyMask);
+				if (FPayload Payload = GetCachePayload(Chunk.Key, Context, Policy, Record.Get().GetPayload(Chunk.Id)))
 				{
-					OnComplete({Key.CacheKey, MoveTemp(Payload), EStatus::Ok});
+					const uint64 RawSize = FMath::Min(Payload.GetRawSize(), Chunk.RawSize);
+					UE_LOG(LogDerivedDataCache, Verbose, TEXT("%s: Cache hit for %s from '%.*s'"),
+						*CachePath, *WriteToString<96>(Chunk.Key, '/', Chunk.Id), Context.Len(), Context.GetData());
+					TRACE_COUNTER_INCREMENT(FileSystemDDC_GetHit);
+					TRACE_COUNTER_ADD(FileSystemDDC_BytesRead, RawSize);
+					COOK_STAT(Timer.AddHit(RawSize));
+					if (OnComplete)
+					{
+						FUniqueBuffer Buffer = FUniqueBuffer::Alloc(RawSize);
+						Payload.GetData().DecompressToComposite().CopyTo(Buffer, Chunk.RawOffset);
+						OnComplete({Chunk.Key, Chunk.Id, Chunk.RawOffset, RawSize, Payload.GetRawHash(), Buffer.MoveToShared(), EStatus::Ok});
+					}
+					continue;
 				}
 			}
-			else
+
+			if (OnComplete)
 			{
-				if (OnComplete)
-				{
-					OnComplete({Key.CacheKey, FPayload(Key.Id), EStatus::Error});
-				}
+				OnComplete({Chunk.Key, Chunk.Id, Chunk.RawOffset, 0, {}, {}, EStatus::Error});
 			}
 		}
 	}
 
-	virtual void CancelAll() override
-	{
-	}
-
 private:
 	uint64 MeasureCacheRecord(const FCacheRecord& Record) const
 	{
@@ -1082,7 +1088,11 @@ private:
 		return true;
 	}
 
-	FOptionalCacheRecord GetCacheRecord(const FCacheKey& Key, FStringView Context, ECachePolicy Policy, bool bAlwaysLoadInlineData = false)
+	FOptionalCacheRecord GetCacheRecord(
+		const FCacheKey& Key,
+		const FStringView Context,
+		const FCacheRecordPolicy& Policy,
+		const bool bAlwaysLoadInlineData = false)
 	{
 		if (!IsUsable())
 		{
@@ -1091,8 +1101,10 @@ private:
 			return FOptionalCacheRecord();
 		}
 
+		const ECachePolicy RecordPolicy = Policy.GetRecordPolicy();
+
 		// Skip the request if querying the cache is disabled.
-		if (!EnumHasAnyFlags(Policy, SpeedClass == ESpeedClass::Local ? ECachePolicy::QueryLocal : ECachePolicy::QueryRemote))
+		if (!EnumHasAnyFlags(RecordPolicy, SpeedClass == ESpeedClass::Local ? ECachePolicy::QueryLocal : ECachePolicy::QueryRemote))
 		{
 			UE_LOG(LogDerivedDataCache, VeryVerbose, TEXT("%s: Skipped get of %s from '%.*s' due to cache policy"),
 				*CachePath, *WriteToString<96>(Key), Context.Len(), Context.GetData());
@@ -1117,7 +1129,7 @@ private:
 			return FOptionalCacheRecord();
 		}
 
-		// Delete the record from storage if is corrupted or has missing payloads.
+		// Delete the record from storage if is corrupted.
 		bool bDeleteCacheObject = true;
 		ON_SCOPE_EXIT
 		{
@@ -1138,7 +1150,7 @@ private:
 		const FCbObject RecordObject(MoveTemp(Buffer));
 		FCacheRecordBuilder RecordBuilder(Key);
 
-		if (!EnumHasAnyFlags(Policy, ECachePolicy::SkipMeta))
+		if (!EnumHasAnyFlags(RecordPolicy, ECachePolicy::SkipMeta))
 		{
 			if (FCbFieldView MetaHash = RecordObject.FindView("MetaHash"_ASV))
 			{
@@ -1155,10 +1167,12 @@ private:
 			}
 		}
 
+		// Disable deletion now that the record is validated.
+		bDeleteCacheObject = false;
+
 		if (FCbObject ValueObject = RecordObject["Value"_ASV].AsObject())
 		{
-			const ECachePolicy ValuePolicy = Policy | (ECachePolicy::SkipData & ~ECachePolicy::SkipValue);
-			FPayload Payload = GetCachePayload(Key, Context, ValuePolicy, ValueObject, bAlwaysLoadInlineData);
+			FPayload Payload = GetCachePayload(Key, Context, Policy, ECachePolicy::SkipValue, ValueObject, bAlwaysLoadInlineData);
 			if (Payload.IsNull())
 			{
 				return FOptionalCacheRecord();
@@ -1168,8 +1182,7 @@ private:
 
 		for (FCbField AttachmentField : RecordObject["Attachments"_ASV])
 		{
-			const ECachePolicy AttachmentsPolicy = Policy | (ECachePolicy::SkipData & ~ECachePolicy::SkipAttachments);
-			FPayload Payload = GetCachePayload(Key, Context, AttachmentsPolicy, AttachmentField.AsObject(), bAlwaysLoadInlineData);
+			FPayload Payload = GetCachePayload(Key, Context, Policy, ECachePolicy::SkipAttachments, AttachmentField.AsObject(), bAlwaysLoadInlineData);
 			if (Payload.IsNull())
 			{
 				return FOptionalCacheRecord();
@@ -1191,18 +1204,13 @@ private:
 	{
 		const FIoHash& RawHash = Payload.GetRawHash();
 		const bool bHasData = Payload.HasData();
-		bStoreInline &= bHasData;
 
-		if (!bStoreInline)
+		if (bHasData && !bStoreInline)
 		{
 			TStringBuilder<256> Path;
 			BuildCachePayloadPath(Key, RawHash, Path);
 			if (!FileExists(Path))
 			{
-				if (!bHasData)
-				{
-					return false;
-				}
 				// Save the compressed buffer with its hash as a header.
 				const FCompressedBuffer& CompressedBuffer = Payload.GetData();
 				const FIoHash CompressedHash = FIoHash::HashBuffer(CompressedBuffer.GetCompressed());
@@ -1224,7 +1232,7 @@ private:
 		Writer.BeginObject();
 		Writer.AddObjectId("Id"_ASV, Payload.GetId());
 		Writer.AddInteger("RawSize"_ASV, Payload.GetRawSize());
-		if (bStoreInline)
+		if (bHasData && bStoreInline)
 		{
 			const FCompressedBuffer& CompressedBuffer = Payload.GetData();
 			Writer.AddHash("RawHash"_ASV, RawHash);
@@ -1239,12 +1247,18 @@ private:
 		return true;
 	}
 
-	FPayload GetCachePayload(const FCacheKey& Key, FStringView Context, ECachePolicy Policy, const FCbObject& Object, bool bAlwaysLoadInlineData = false) const
+	FPayload GetCachePayload(
+		const FCacheKey& Key,
+		const FStringView Context,
+		const FCacheRecordPolicy& Policy,
+		const ECachePolicy SkipFlag,
+		const FCbObject& Object,
+		const bool bAlwaysLoadInlineData) const
 	{
 		const FPayloadId Id = Object.FindView("Id"_ASV).AsObjectId();
 		const uint64 RawSize = Object.FindView("RawSize"_ASV).AsUInt64(MAX_uint64);
 		const FIoHash RawHash = Object.FindView("RawHash"_ASV).AsHash();
-		FIoHash CompressedHash = Object.FindView("CompressedHash"_ASV).AsHash();
+		const FIoHash CompressedHash = Object.FindView("CompressedHash"_ASV).AsHash();
 		FSharedBuffer CompressedData = Object["CompressedData"_ASV].AsBinary();
 
 		if (Id.IsNull() || RawSize == MAX_uint64 || RawHash.IsZero() || !(CompressedHash.IsZero() == CompressedData.IsNull()))
@@ -1254,11 +1268,12 @@ private:
 			return FPayload();
 		}
 
+		const ECachePolicy PayloadPolicy = Policy.GetPayloadPolicy(Id);
 		FPayload Payload(Id, RawHash, RawSize);
 
 		if (CompressedData)
 		{
-			if (EnumHasAllFlags(Policy, ECachePolicy::SkipData) && !bAlwaysLoadInlineData)
+			if (EnumHasAllFlags(PayloadPolicy, SkipFlag) && !bAlwaysLoadInlineData)
 			{
 				return Payload;
 			}
@@ -1268,12 +1283,12 @@ private:
 			}
 		}
 
-		return GetCachePayload(Key, Context, Policy, Payload);
+		return GetCachePayload(Key, Context, PayloadPolicy, Payload);
 	}
 
 	FPayload GetCachePayload(const FCacheKey& Key, FStringView Context, ECachePolicy Policy, const FPayload& Payload) const
 	{
-		if (Payload.HasData())
+		if (Payload.HasData() || !EnumHasAnyFlags(Policy, ECachePolicy::Query))
 		{
 			return Payload;
 		}
@@ -1316,7 +1331,7 @@ private:
 
 	FPayload ValidateCachePayload(
 		const FCacheKey& Key,
-		FStringView Context,
+		const FStringView Context,
 		const FPayload& Payload,
 		const FIoHash& CompressedHash,
 		FSharedBuffer&& CompressedData) const
diff --git a/Engine/Source/Developer/DerivedDataCache/Private/HierarchicalDerivedDataBackend.h b/Engine/Source/Developer/DerivedDataCache/Private/HierarchicalDerivedDataBackend.h
--- a/Engine/Source/Developer/DerivedDataCache/Private/HierarchicalDerivedDataBackend.h
+++ b/Engine/Source/Developer/DerivedDataCache/Private/HierarchicalDerivedDataBackend.h
@@ -3,14 +3,17 @@
 #pragma once
 
 #include "CoreTypes.h"
+#include "Algo/BinarySearch.h"
 #include "Containers/ArrayView.h"
 #include "DerivedDataBackendAsyncPutWrapper.h"
 #include "DerivedDataBackendInterface.h"
 #include "DerivedDataCacheRecord.h"
 #include "DerivedDataCacheUsageStats.h"
+#include "DerivedDataChunk.h"
 #include "DerivedDataPayload.h"
 #include "DerivedDataRequestOwner.h"
 #include "Misc/ScopeRWLock.h"
+#include "Misc/StringBuilder.h"
 #include "ProfilingDebugging/CookStats.h"
 #include "Templates/UniquePtr.h"
 
@@ -529,15 +532,15 @@ public:
 	virtual void Get(
 		TConstArrayView<FCacheKey> Keys,
 		FStringView Context,
-		ECachePolicy Policy,
+		FCacheRecordPolicy Policy,
 		IRequestOwner& Owner,
 		FOnCacheGetComplete&& OnComplete) override
 	{
-		const bool bQueryLocal = bHasLocalBackends && EnumHasAnyFlags(Policy, ECachePolicy::QueryLocal);
-		const bool bStoreLocal = bHasWritableLocalBackends && EnumHasAnyFlags(Policy, ECachePolicy::StoreLocal);
-		const bool bQueryRemote = bHasRemoteBackends && EnumHasAnyFlags(Policy, ECachePolicy::QueryRemote);
-		const bool bStoreRemote = bHasWritableRemoteBackends && EnumHasAnyFlags(Policy, ECachePolicy::StoreRemote);
-		const bool bStoreLocalCopy = bStoreLocal && bHasMultipleLocalBackends && !EnumHasAnyFlags(Policy, ECachePolicy::SkipLocalCopy);
+		const bool bQueryLocal = bHasLocalBackends && EnumHasAnyFlags(Policy.GetRecordPolicy(), ECachePolicy::QueryLocal);
+		const bool bStoreLocal = bHasWritableLocalBackends && EnumHasAnyFlags(Policy.GetRecordPolicy(), ECachePolicy::StoreLocal);
+		const bool bQueryRemote = bHasRemoteBackends && EnumHasAnyFlags(Policy.GetRecordPolicy(), ECachePolicy::QueryRemote);
+		const bool bStoreRemote = bHasWritableRemoteBackends && EnumHasAnyFlags(Policy.GetRecordPolicy(), ECachePolicy::StoreRemote);
+		const bool bStoreLocalCopy = bStoreLocal && bHasMultipleLocalBackends;
 
 		TArray<FCacheKey, TInlineAllocator<16>> RemainingKeys(Keys);
 
@@ -602,54 +605,56 @@ public:
 		}
 	}
 
-	virtual void GetPayload(
-		TConstArrayView<FCachePayloadKey> Keys,
+	virtual void GetChunks(
+		TConstArrayView<FCacheChunkRequest> Chunks,
 		FStringView Context,
-		ECachePolicy Policy,
 		IRequestOwner& Owner,
-		FOnCacheGetPayloadComplete&& OnComplete) override
+		FOnCacheGetChunkComplete&& OnComplete) override
 	{
-		TArray<FCachePayloadKey, TInlineAllocator<16>> RemainingKeys(Keys);
-
-		TSet<FCachePayloadKey> KeysOk;
+		TArray<FCacheChunkRequest, TInlineAllocator<16>> RemainingChunks(Chunks);
 
 		{
 			FReadScopeLock LockScope(Lock);
 
-			for (int32 GetCacheIndex = 0; GetCacheIndex < InnerBackends.Num() && !RemainingKeys.IsEmpty(); ++GetCacheIndex)
+			for (FDerivedDataBackendInterface* InnerBackend : InnerBackends)
 			{
+				if (RemainingChunks.IsEmpty())
+				{
+					return;
+				}
+				RemainingChunks.StableSort(TChunkLess());
+				TArray<FCacheChunkRequest, TInlineAllocator<16>> ErrorChunks;
 				FRequestOwner BlockingOwner(EPriority::Blocking);
-				InnerBackends[GetCacheIndex]->GetPayload(RemainingKeys, Context, Policy, BlockingOwner,
-					[&OnComplete, &KeysOk](FCacheGetPayloadCompleteParams&& Params)
+				InnerBackend->GetChunks(RemainingChunks, Context, BlockingOwner,
+					[&OnComplete, &RemainingChunks, &ErrorChunks](FCacheGetChunkCompleteParams&& Params)
 					{
-						if (Params.Status == EStatus::Ok)
+						if (Params.Status == EStatus::Error)
 						{
-							KeysOk.Add({Params.Key, Params.Payload.GetId()});
-							if (OnComplete)
-							{
-								OnComplete(MoveTemp(Params));
-							}
+							const int32 ChunkIndex = Algo::BinarySearch(RemainingChunks, Params, TChunkLess());
+							checkf(ChunkIndex != INDEX_NONE, TEXT("Failed to find remaining chunk %s ")
+								TEXT(" with raw offset %") UINT64_FMT TEXT("."),
+								*WriteToString<96>(Params.Key, '/', Params.Id), Params.RawOffset);
+							ErrorChunks.Add(RemainingChunks[ChunkIndex]);
+						}
+						else if (OnComplete)
+						{
+							OnComplete(MoveTemp(Params));
 						}
 					});
 				BlockingOwner.Wait();
-
-				RemainingKeys.RemoveAll([&KeysOk](const FCachePayloadKey& Key) { return KeysOk.Contains(Key); });
+				RemainingChunks = MoveTemp(ErrorChunks);
 			}
 		}
 
 		if (OnComplete)
 		{
-			for (const FCachePayloadKey& Key : RemainingKeys)
+			for (const FCacheChunkRequest& Chunk : RemainingChunks)
 			{
-				OnComplete({Key.CacheKey, FPayload(Key.Id), EStatus::Error});
+				OnComplete({Chunk.Key, Chunk.Id, Chunk.RawOffset, 0, {}, {}, EStatus::Error});
 			}
 		}
 	}
 
-	virtual void CancelAll() override
-	{
-	}
-
 private:
 	FDerivedDataCacheUsageStats UsageStats;
 
diff --git a/Engine/Source/Developer/DerivedDataCache/Private/HttpDerivedDataBackend.cpp b/Engine/Source/Developer/DerivedDataCache/Private/HttpDerivedDataBackend.cpp
--- a/Engine/Source/Developer/DerivedDataCache/Private/HttpDerivedDataBackend.cpp
+++ b/Engine/Source/Developer/DerivedDataCache/Private/HttpDerivedDataBackend.cpp
@@ -2449,7 +2449,7 @@ void FHttpDerivedDataBackend::Put(
 void FHttpDerivedDataBackend::Get(
 	TConstArrayView<FCacheKey> Keys,
 	FStringView Context,
-	ECachePolicy Policy,
+	FCacheRecordPolicy Policy,
 	IRequestOwner& Owner,
 	FOnCacheGetComplete&& OnComplete)
 {
@@ -2462,18 +2462,17 @@ void FHttpDerivedDataBackend::Get(
 	}
 }
 
-void FHttpDerivedDataBackend::GetPayload(
-	TConstArrayView<FCachePayloadKey> Keys,
+void FHttpDerivedDataBackend::GetChunks(
+	TConstArrayView<FCacheChunkRequest> Chunks,
 	FStringView Context,
-	ECachePolicy Policy,
 	IRequestOwner& Owner,
-	FOnCacheGetPayloadComplete&& OnComplete)
+	FOnCacheGetChunkComplete&& OnComplete)
 {
 	if (OnComplete)
 	{
-		for (const FCachePayloadKey& Key : Keys)
+		for (const FCacheChunkRequest& Chunk : Chunks)
 		{
-			OnComplete({Key.CacheKey, FPayload(Key.Id), EStatus::Error});
+			OnComplete({Chunk.Key, Chunk.Id, Chunk.RawOffset, 0, {}, {}, EStatus::Error});
 		}
 	}
 }
diff --git a/Engine/Source/Developer/DerivedDataCache/Private/HttpDerivedDataBackend.h b/Engine/Source/Developer/DerivedDataCache/Private/HttpDerivedDataBackend.h
--- a/Engine/Source/Developer/DerivedDataCache/Private/HttpDerivedDataBackend.h
+++ b/Engine/Source/Developer/DerivedDataCache/Private/HttpDerivedDataBackend.h
@@ -83,20 +83,15 @@ public:
 	virtual void Get(
 		TConstArrayView<FCacheKey> Keys,
 		FStringView Context,
-		ECachePolicy Policy,
+		FCacheRecordPolicy Policy,
 		IRequestOwner& Owner,
 		FOnCacheGetComplete&& OnComplete) override;
 
-	virtual void GetPayload(
-		TConstArrayView<FCachePayloadKey> Keys,
+	virtual void GetChunks(
+		TConstArrayView<FCacheChunkRequest> Chunks,
 		FStringView Context,
-		ECachePolicy Policy,
 		IRequestOwner& Owner,
-		FOnCacheGetPayloadComplete&& OnComplete) override;
-
-	virtual void CancelAll() override
-	{
-	}
+		FOnCacheGetChunkComplete&& OnComplete) override;
 
 	static FHttpDerivedDataBackend* GetAny()
 	{
diff --git a/Engine/Source/Developer/DerivedDataCache/Private/MemoryDerivedDataBackend.cpp b/Engine/Source/Developer/DerivedDataCache/Private/MemoryDerivedDataBackend.cpp
--- a/Engine/Source/Developer/DerivedDataCache/Private/MemoryDerivedDataBackend.cpp
+++ b/Engine/Source/Developer/DerivedDataCache/Private/MemoryDerivedDataBackend.cpp
@@ -503,7 +503,7 @@ void FMemoryDerivedDataBackend::Put(
 void FMemoryDerivedDataBackend::Get(
 	TConstArrayView<FCacheKey> Keys,
 	FStringView Context,
-	ECachePolicy Policy,
+	FCacheRecordPolicy Policy,
 	IRequestOwner& Owner,
 	FOnCacheGetComplete&& OnComplete)
 {
@@ -538,39 +538,41 @@ void FMemoryDerivedDataBackend::Get(
 	}
 }
 
-void FMemoryDerivedDataBackend::GetPayload(
-	TConstArrayView<FCachePayloadKey> Keys,
+void FMemoryDerivedDataBackend::GetChunks(
+	TConstArrayView<FCacheChunkRequest> Chunks,
 	FStringView Context,
-	ECachePolicy Policy,
 	IRequestOwner& Owner,
-	FOnCacheGetPayloadComplete&& OnComplete)
+	FOnCacheGetChunkComplete&& OnComplete)
 {
-	for (const FCachePayloadKey& PayloadKey : Keys)
+	for (const FCacheChunkRequest& Chunk : Chunks)
 	{
 		COOK_STAT(auto Timer = UsageStats.TimeGet());
 		FPayload Payload;
-		if (ShouldSimulateMiss(PayloadKey.CacheKey))
+		if (ShouldSimulateMiss(Chunk.Key))
 		{
 			UE_LOG(LogDerivedDataCache, Verbose, TEXT("%s: Simulated miss for get of %s from '%.*s'"),
-				*GetName(), *WriteToString<96>(PayloadKey), Context.Len(), Context.GetData());
+				*GetName(), *WriteToString<96>(Chunk.Key, '/', Chunk.Id), Context.Len(), Context.GetData());
 		}
-		else if (FScopeLock ScopeLock(&SynchronizationObject); const FCacheRecord* Record = CacheRecords.Find(PayloadKey.CacheKey))
+		else if (FScopeLock ScopeLock(&SynchronizationObject); const FCacheRecord* Record = CacheRecords.Find(Chunk.Key))
 		{
-			Payload = Record->GetAttachmentPayload(PayloadKey.Id);
+			Payload = Record->GetAttachmentPayload(Chunk.Id);
 		}
-		if (Payload)
+		if (Payload && Chunk.RawOffset <= Payload.GetRawSize())
 		{
-			COOK_STAT(Timer.AddHit(Payload.GetData().GetRawSize()));
+			const uint64 RawSize = FMath::Min(Payload.GetRawSize() - Chunk.RawOffset, Chunk.RawSize);
+			COOK_STAT(Timer.AddHit(RawSize));
 			if (OnComplete)
 			{
-				OnComplete({PayloadKey.CacheKey, MoveTemp(Payload), EStatus::Ok});
+				FUniqueBuffer Buffer = FUniqueBuffer::Alloc(RawSize);
+				Payload.GetData().DecompressToComposite().CopyTo(Buffer, Chunk.RawOffset);
+				OnComplete({Chunk.Key, Chunk.Id, Chunk.RawOffset, RawSize, Payload.GetRawHash(), Buffer.MoveToShared(), EStatus::Ok});
 			}
 		}
 		else
 		{
 			if (OnComplete)
 			{
-				OnComplete({PayloadKey.CacheKey, FPayload(PayloadKey.Id), EStatus::Error});
+				OnComplete({Chunk.Key, Chunk.Id, Chunk.RawOffset, 0, {}, {}, EStatus::Error});
 			}
 		}
 	}
diff --git a/Engine/Source/Developer/DerivedDataCache/Private/MemoryDerivedDataBackend.h b/Engine/Source/Developer/DerivedDataCache/Private/MemoryDerivedDataBackend.h
--- a/Engine/Source/Developer/DerivedDataCache/Private/MemoryDerivedDataBackend.h
+++ b/Engine/Source/Developer/DerivedDataCache/Private/MemoryDerivedDataBackend.h
@@ -104,20 +104,15 @@ public:
 	virtual void Get(
 		TConstArrayView<FCacheKey> Keys,
 		FStringView Context,
-		ECachePolicy Policy,
+		FCacheRecordPolicy Policy,
 		IRequestOwner& Owner,
 		FOnCacheGetComplete&& OnComplete) override;
 
-	virtual void GetPayload(
-		TConstArrayView<FCachePayloadKey> Keys,
+	virtual void GetChunks(
+		TConstArrayView<FCacheChunkRequest> Chunks,
 		FStringView Context,
-		ECachePolicy Policy,
 		IRequestOwner& Owner,
-		FOnCacheGetPayloadComplete&& OnComplete) override;
-
-	virtual void CancelAll() override
-	{
-	}
+		FOnCacheGetChunkComplete&& OnComplete) override;
 
 private:
 	/** Name of the cache file loaded (if any). */
diff --git a/Engine/Source/Developer/DerivedDataCache/Private/PakFileDerivedDataBackend.cpp b/Engine/Source/Developer/DerivedDataCache/Private/PakFileDerivedDataBackend.cpp
--- a/Engine/Source/Developer/DerivedDataCache/Private/PakFileDerivedDataBackend.cpp
+++ b/Engine/Source/Developer/DerivedDataCache/Private/PakFileDerivedDataBackend.cpp
@@ -460,7 +460,7 @@ void FPakFileDerivedDataBackend::Put(
 void FPakFileDerivedDataBackend::Get(
 	TConstArrayView<FCacheKey> Keys,
 	FStringView Context,
-	ECachePolicy Policy,
+	FCacheRecordPolicy Policy,
 	IRequestOwner& Owner,
 	FOnCacheGetComplete&& OnComplete)
 {
@@ -473,18 +473,17 @@ void FPakFileDerivedDataBackend::Get(
 	}
 }
 
-void FPakFileDerivedDataBackend::GetPayload(
-	TConstArrayView<FCachePayloadKey> Keys,
+void FPakFileDerivedDataBackend::GetChunks(
+	TConstArrayView<FCacheChunkRequest> Chunks,
 	FStringView Context,
-	ECachePolicy Policy,
 	IRequestOwner& Owner,
-	FOnCacheGetPayloadComplete&& OnComplete)
+	FOnCacheGetChunkComplete&& OnComplete)
 {
 	if (OnComplete)
 	{
-		for (const FCachePayloadKey& Key : Keys)
+		for (const FCacheChunkRequest& Chunk : Chunks)
 		{
-			OnComplete({Key.CacheKey, FPayload(Key.Id), EStatus::Error});
+			OnComplete({Chunk.Key, Chunk.Id, Chunk.RawOffset, 0, {}, {}, EStatus::Error});
 		}
 	}
 }
diff --git a/Engine/Source/Developer/DerivedDataCache/Private/PakFileDerivedDataBackend.h b/Engine/Source/Developer/DerivedDataCache/Private/PakFileDerivedDataBackend.h
--- a/Engine/Source/Developer/DerivedDataCache/Private/PakFileDerivedDataBackend.h
+++ b/Engine/Source/Developer/DerivedDataCache/Private/PakFileDerivedDataBackend.h
@@ -114,20 +114,15 @@ public:
 	virtual void Get(
 		TConstArrayView<FCacheKey> Keys,
 		FStringView Context,
-		ECachePolicy Policy,
+		FCacheRecordPolicy Policy,
 		IRequestOwner& Owner,
 		FOnCacheGetComplete&& OnComplete) override;
 
-	virtual void GetPayload(
-		TConstArrayView<FCachePayloadKey> Keys,
+	virtual void GetChunks(
+		TConstArrayView<FCacheChunkRequest> Chunks,
 		FStringView Context,
-		ECachePolicy Policy,
 		IRequestOwner& Owner,
-		FOnCacheGetPayloadComplete&& OnComplete) override;
-
-	virtual void CancelAll() override
-	{
-	}
+		FOnCacheGetChunkComplete&& OnComplete) override;
 
 private:
 	FDerivedDataCacheUsageStats UsageStats;
diff --git a/Engine/Source/Developer/DerivedDataCache/Private/S3DerivedDataBackend.cpp b/Engine/Source/Developer/DerivedDataCache/Private/S3DerivedDataBackend.cpp
--- a/Engine/Source/Developer/DerivedDataCache/Private/S3DerivedDataBackend.cpp
+++ b/Engine/Source/Developer/DerivedDataCache/Private/S3DerivedDataBackend.cpp
@@ -1196,7 +1196,7 @@ void FS3DerivedDataBackend::Put(
 void FS3DerivedDataBackend::Get(
 	TConstArrayView<FCacheKey> Keys,
 	FStringView Context,
-	ECachePolicy Policy,
+	FCacheRecordPolicy Policy,
 	IRequestOwner& Owner,
 	FOnCacheGetComplete&& OnComplete)
 {
@@ -1209,18 +1209,17 @@ void FS3DerivedDataBackend::Get(
 	}
 }
 
-void FS3DerivedDataBackend::GetPayload(
-	TConstArrayView<FCachePayloadKey> Keys,
+void FS3DerivedDataBackend::GetChunks(
+	TConstArrayView<FCacheChunkRequest> Chunks,
 	FStringView Context,
-	ECachePolicy Policy,
 	IRequestOwner& Owner,
-	FOnCacheGetPayloadComplete&& OnComplete)
+	FOnCacheGetChunkComplete&& OnComplete)
 {
 	if (OnComplete)
 	{
-		for (const FCachePayloadKey& Key : Keys)
+		for (const FCacheChunkRequest& Chunk : Chunks)
 		{
-			OnComplete({Key.CacheKey, FPayload(Key.Id), EStatus::Error});
+			OnComplete({Chunk.Key, Chunk.Id, Chunk.RawOffset, 0, {}, {}, EStatus::Error});
 		}
 	}
 }
diff --git a/Engine/Source/Developer/DerivedDataCache/Private/S3DerivedDataBackend.h b/Engine/Source/Developer/DerivedDataCache/Private/S3DerivedDataBackend.h
--- a/Engine/Source/Developer/DerivedDataCache/Private/S3DerivedDataBackend.h
+++ b/Engine/Source/Developer/DerivedDataCache/Private/S3DerivedDataBackend.h
@@ -69,20 +69,15 @@ public:
 	virtual void Get(
 		TConstArrayView<FCacheKey> Keys,
 		FStringView Context,
-		ECachePolicy Policy,
+		FCacheRecordPolicy Policy,
 		IRequestOwner& Owner,
 		FOnCacheGetComplete&& OnComplete) override;
 
-	virtual void GetPayload(
-		TConstArrayView<FCachePayloadKey> Keys,
+	virtual void GetChunks(
+		TConstArrayView<FCacheChunkRequest> Chunks,
 		FStringView Context,
-		ECachePolicy Policy,
 		IRequestOwner& Owner,
-		FOnCacheGetPayloadComplete&& OnComplete) override;
-
-	virtual void CancelAll() override
-	{
-	}
+		FOnCacheGetChunkComplete&& OnComplete) override;
 
 private:
 	struct FBundle;
diff --git a/Engine/Source/Developer/DerivedDataCache/Private/ZenDerivedDataBackend.cpp b/Engine/Source/Developer/DerivedDataCache/Private/ZenDerivedDataBackend.cpp
--- a/Engine/Source/Developer/DerivedDataCache/Private/ZenDerivedDataBackend.cpp
+++ b/Engine/Source/Developer/DerivedDataCache/Private/ZenDerivedDataBackend.cpp
@@ -5,6 +5,7 @@
 
 #include "Algo/Accumulate.h"
 #include "DerivedDataCacheRecord.h"
+#include "DerivedDataChunk.h"
 #include "DerivedDataPayload.h"
 #include "Misc/ConfigCacheIni.h"
 #include "Misc/ScopeLock.h"
@@ -555,7 +556,7 @@ void FZenDerivedDataBackend::Put(
 void FZenDerivedDataBackend::Get(
 	TConstArrayView<FCacheKey> Keys,
 	FStringView Context,
-	ECachePolicy Policy,
+	FCacheRecordPolicy Policy,
 	IRequestOwner& Owner,
 	FOnCacheGetComplete&& OnComplete)
 {
@@ -603,31 +604,30 @@ void FZenDerivedDataBackend::Get(
 	}
 }
 
-void FZenDerivedDataBackend::GetPayload(
-	TConstArrayView<FCachePayloadKey> Keys,
+void FZenDerivedDataBackend::GetChunks(
+	TConstArrayView<FCacheChunkRequest> Chunks,
 	FStringView Context,
-	ECachePolicy Policy,
 	IRequestOwner& Owner,
-	FOnCacheGetPayloadComplete&& OnComplete)
+	FOnCacheGetChunkComplete&& OnComplete)
 {
 	if (bCacheRecordEndpointEnabled)
 	{
-		for (const FCachePayloadKey& Key : Keys)
+		for (const FCacheChunkRequest& Chunk : Chunks)
 		{
 			TRACE_CPUPROFILER_EVENT_SCOPE(ZenDDC_Get);
 			TRACE_COUNTER_ADD(ZenDDC_Get, int64(1));
 			COOK_STAT(auto Timer = UsageStats.TimeGet());
 
 			TStringBuilder<256> QueryUri;
-			AppendZenUri(Key.CacheKey, Key.Id, QueryUri);
-			AppendPolicyQueryString(Policy, QueryUri);
+			AppendZenUri(Chunk.Key, Chunk.Id, QueryUri);
+			AppendPolicyQueryString(Chunk.Policy, QueryUri);
 
 			EGetResult GetResult;
 			FCompressedBuffer CompressedBuffer;
-			if (ShouldSimulateMiss(Key.CacheKey))
+			if (ShouldSimulateMiss(Chunk.Key))
 			{
 				UE_LOG(LogDerivedDataCache, Verbose, TEXT("%s: Simulated miss for get of %s from '%.*s'"),
-					*GetName(), *WriteToString<96>(Key), Context.Len(), Context.GetData());
+					*GetName(), *WriteToString<96>(Chunk.Key, '/', Chunk.Id), Context.Len(), Context.GetData());
 				GetResult = EGetResult::NotFound;
 			}
 			else
@@ -641,74 +641,80 @@ void FZenDerivedDataBackend::GetPayload(
 			}
 			if (CompressedBuffer)
 			{
-				FPayload Payload = FPayload(Key.Id, MoveTemp(CompressedBuffer));
+				const uint64 RawSize = FMath::Min(CompressedBuffer.GetRawSize(), Chunk.RawSize);
 				UE_LOG(LogDerivedDataCache, Verbose, TEXT("%s: Cache hit for %s from '%.*s'"),
-					*GetName(), *WriteToString<96>(Key), Context.Len(), Context.GetData());
+					*GetName(), *WriteToString<96>(Chunk.Key, '/', Chunk.Id), Context.Len(), Context.GetData());
 				TRACE_COUNTER_ADD(ZenDDC_GetHit, int64(1));
-				TRACE_COUNTER_ADD(ZenDDC_BytesReceived, Payload.GetRawSize());
-				COOK_STAT(Timer.AddHit(Payload.GetRawSize()));
+				TRACE_COUNTER_ADD(ZenDDC_BytesReceived, RawSize);
+				COOK_STAT(Timer.AddHit(RawSize));
 				if (OnComplete)
 				{
-					OnComplete({ Key.CacheKey, MoveTemp(Payload), EStatus::Ok });
+					FUniqueBuffer Buffer = FUniqueBuffer::Alloc(RawSize);
+					CompressedBuffer.DecompressToComposite().CopyTo(Buffer, Chunk.RawOffset);
+					OnComplete({ Chunk.Key, Chunk.Id, Chunk.RawOffset, RawSize, CompressedBuffer.GetRawHash(), Buffer.MoveToShared(), EStatus::Ok });
 				}
 			}
 			else
 			{
 				UE_LOG(LogDerivedDataCache, Display, TEXT("%s: Cache miss with %s payload %s for %s from '%.*s'"),
 					*GetName(), (GetResult == EGetResult::NotFound ? TEXT("missing") : TEXT("corrupted")),
-					*WriteToString<16>(Key.Id), *WriteToString<96>(Key), Context.Len(), Context.GetData());
+					*WriteToString<16>(Chunk.Id), *WriteToString<96>(Chunk.Key), Context.Len(), Context.GetData());
 				if (OnComplete)
 				{
-					OnComplete({ Key.CacheKey, FPayload(Key.Id), EStatus::Error });
+					OnComplete({ Chunk.Key, Chunk.Id, Chunk.RawOffset, 0, {}, {}, EStatus::Error });
 				}
 			}
 		}
 	}
 	else
 	{
-		// We have to load the CacheRecord for each Payload, so sort all the Payloads sharing
-		// the same key to occur together and we can load the CacheRecord at the start of each run of Payloads
-		TArray<FCachePayloadKey, TInlineAllocator<16>> SortedKeys(Keys);
-		SortedKeys.StableSort();
+		// We have to load the CacheRecord for each chunks, so sort all the chunks sharing
+		// the same key to occur together and we can load the CacheRecord at the start of each run.
+		TArray<FCacheChunkRequest, TInlineAllocator<16>> SortedChunks(Chunks);
+		SortedChunks.StableSort(TChunkLess());
 
 		FOptionalCacheRecord Record;
-		for (const FCachePayloadKey& Key : SortedKeys)
+		for (const FCacheChunkRequest& Chunk : SortedChunks)
 		{
 			TRACE_CPUPROFILER_EVENT_SCOPE(ZenDDC_Get);
 			TRACE_COUNTER_ADD(ZenDDC_Get, int64(1));
 			COOK_STAT(auto Timer = UsageStats.TimeGet());
 
-			if (!Record || Record.Get().GetKey() != Key.CacheKey)
+			if (!Record || Record.Get().GetKey() != Chunk.Key)
 			{
-				Record = LegacyGetCacheRecord(Key.CacheKey, Context, Policy | ECachePolicy::SkipData, /*bAlwaysLoadInlineData*/ true);
+				Record = LegacyGetCacheRecord(Chunk.Key, Context, Chunk.Policy | ECachePolicy::SkipData, /*bAlwaysLoadInlineData*/ true);
 			}
-			if (FPayload Payload = Record ? LegacyGetCachePayload(Key.CacheKey, Context, Policy, Record.Get().GetPayload(Key.Id)) : FPayload::Null)
+			if (Record)
 			{
-				UE_LOG(LogDerivedDataCache, Verbose, TEXT("%s: Cache hit for %s from '%.*s'"),
-					*GetName(), *WriteToString<96>(Key), Context.Len(), Context.GetData());
-				TRACE_COUNTER_ADD(ZenDDC_GetHit, int64(1));
-				TRACE_COUNTER_ADD(ZenDDC_BytesReceived, Payload.GetRawSize());
-				COOK_STAT(Timer.AddHit(Payload.GetRawSize()));
-				if (OnComplete)
+				const ECachePolicy PolicyMask = Record.Get().GetValuePayload().GetId() == Chunk.Id
+					? ECachePolicy::SkipValue : ECachePolicy::SkipAttachments;
+				const ECachePolicy Policy = Chunk.Policy | (ECachePolicy::SkipData & ~PolicyMask);
+				if (FPayload Payload = LegacyGetCachePayload(Chunk.Key, Context, Policy, Record.Get().GetPayload(Chunk.Id)))
 				{
-					OnComplete({ Key.CacheKey, MoveTemp(Payload), EStatus::Ok });
+					const uint64 RawSize = FMath::Min(Payload.GetRawSize(), Chunk.RawSize);
+					UE_LOG(LogDerivedDataCache, Verbose, TEXT("%s: Cache hit for %s from '%.*s'"),
+						*GetName(), *WriteToString<96>(Chunk.Key, '/', Chunk.Id), Context.Len(), Context.GetData());
+					TRACE_COUNTER_ADD(ZenDDC_GetHit, int64(1));
+					TRACE_COUNTER_ADD(ZenDDC_BytesReceived, RawSize);
+					COOK_STAT(Timer.AddHit(RawSize));
+					if (OnComplete)
+					{
+						FUniqueBuffer Buffer = FUniqueBuffer::Alloc(RawSize);
+						Payload.GetData().DecompressToComposite().CopyTo(Buffer, Chunk.RawOffset);
+						OnComplete({ Chunk.Key, Chunk.Id, Chunk.RawOffset, RawSize, Payload.GetRawHash(), Buffer.MoveToShared(), EStatus::Ok });
+					}
+					continue;
 				}
 			}
-			else
+
+			if (OnComplete)
 			{
-				if (OnComplete)
-				{
-					OnComplete({ Key.CacheKey, FPayload(Key.Id), EStatus::Error });
-				}
+				OnComplete({ Chunk.Key, Chunk.Id, Chunk.RawOffset, 0, {}, {}, EStatus::Error });
 			}
 		}
 	}
 }
 
-void FZenDerivedDataBackend::CancelAll()
-{
-}
-
 bool FZenDerivedDataBackend::PutCacheRecord(const FCacheRecord& Record, FStringView Context, ECachePolicy Policy)
 {
 	const FCacheKey& Key = Record.GetKey();
@@ -728,11 +734,13 @@ bool FZenDerivedDataBackend::PutCacheRecord(const FCacheRecord& Record, FStringV
 	return true;
 }
 
-FOptionalCacheRecord FZenDerivedDataBackend::GetCacheRecord(const FCacheKey& Key, FStringView Context,
-	ECachePolicy Policy) const
+FOptionalCacheRecord FZenDerivedDataBackend::GetCacheRecord(
+	const FCacheKey& Key,
+	const FStringView Context,
+	const FCacheRecordPolicy& Policy) const
 {
 	FCbPackage Package;
-	EGetResult GetResult = GetZenData(Key, Policy, Package);
+	EGetResult GetResult = GetZenData(Key, Policy.GetRecordPolicy(), Package);
 	if (GetResult != EGetResult::Success)
 	{
 		switch (GetResult)
@@ -849,8 +857,11 @@ bool FZenDerivedDataBackend::LegacyPutCachePayload(const FCacheKey& Key, FString
 	return true;
 }
 
-FOptionalCacheRecord FZenDerivedDataBackend::LegacyGetCacheRecord(const FCacheKey& Key,
-	FStringView Context, ECachePolicy Policy, bool bAlwaysLoadInlineData) const
+FOptionalCacheRecord FZenDerivedDataBackend::LegacyGetCacheRecord(
+	const FCacheKey& Key,
+	const FStringView Context,
+	const FCacheRecordPolicy& Policy,
+	const bool bAlwaysLoadInlineData) const
 {
 	TStringBuilder<256> Uri;
 	LegacyMakeZenKey(Key, Uri);
@@ -917,8 +928,12 @@ FPayload FZenDerivedDataBackend::LegacyGetCachePayload(const FCacheKey& Key, FSt
 	return FPayload(PayloadIdOnly.GetId(), MoveTemp(CompressedBuffer));
 }
 
-FOptionalCacheRecord FZenDerivedDataBackend::LegacyCreateRecord(FSharedBuffer&& RecordBytes, const FCacheKey& Key,
-	FStringView Context, ECachePolicy Policy, bool bAlwaysLoadInlineData) const
+FOptionalCacheRecord FZenDerivedDataBackend::LegacyCreateRecord(
+	FSharedBuffer&& RecordBytes,
+	const FCacheKey& Key,
+	const FStringView Context,
+	const FCacheRecordPolicy& Policy,
+	const bool bAlwaysLoadInlineData) const
 {
 	// TODO: Temporary function copied from FFileSystemDerivedDataBackend.
 	// Eventually We will instead send the Policy to zen sever along with the request for the FCacheKey,
@@ -926,7 +941,7 @@ FOptionalCacheRecord FZenDerivedDataBackend::LegacyCreateRecord(FSharedBuffer&&
 	const FCbObject RecordObject(MoveTemp(RecordBytes));
 	FCacheRecordBuilder RecordBuilder(Key);
 
-	if (!EnumHasAnyFlags(Policy, ECachePolicy::SkipMeta))
+	if (!EnumHasAnyFlags(Policy.GetRecordPolicy(), ECachePolicy::SkipMeta))
 	{
 		if (FCbFieldView MetaHash = RecordObject.FindView("MetaHash"_ASV))
 		{
@@ -945,8 +960,8 @@ FOptionalCacheRecord FZenDerivedDataBackend::LegacyCreateRecord(FSharedBuffer&&
 
 	if (FCbObject ValueObject = RecordObject["Value"_ASV].AsObject())
 	{
-		const ECachePolicy ValuePolicy = Policy | (ECachePolicy::SkipData & ~ECachePolicy::SkipValue);
-		FPayload Payload = LegacyGetCachePayload(Key, Context, ValuePolicy, ValueObject, bAlwaysLoadInlineData);
+		FPayload Payload = LegacyGetCachePayload(Key, Context, Policy,
+			ECachePolicy::SkipValue, ValueObject, bAlwaysLoadInlineData);
 		if (Payload.IsNull())
 		{
 			return FOptionalCacheRecord();
@@ -956,9 +971,8 @@ FOptionalCacheRecord FZenDerivedDataBackend::LegacyCreateRecord(FSharedBuffer&&
 
 	for (FCbField AttachmentField : RecordObject["Attachments"_ASV])
 	{
-		const ECachePolicy AttachmentsPolicy = Policy | (ECachePolicy::SkipData & ~ECachePolicy::SkipAttachments);
-		FPayload Payload = LegacyGetCachePayload(Key, Context, AttachmentsPolicy, AttachmentField.AsObject(),
-			bAlwaysLoadInlineData);
+		FPayload Payload = LegacyGetCachePayload(Key, Context, Policy,
+			ECachePolicy::SkipAttachments, AttachmentField.AsObject(), bAlwaysLoadInlineData);
 		if (Payload.IsNull())
 		{
 			return FOptionalCacheRecord();
@@ -969,8 +983,13 @@ FOptionalCacheRecord FZenDerivedDataBackend::LegacyCreateRecord(FSharedBuffer&&
 	return RecordBuilder.Build();
 }
 
-FPayload FZenDerivedDataBackend::LegacyGetCachePayload(const FCacheKey& Key, FStringView Context, ECachePolicy Policy,
-	const FCbObject& Object, bool bAlwaysLoadInlineData) const
+FPayload FZenDerivedDataBackend::LegacyGetCachePayload(
+	const FCacheKey& Key,
+	const FStringView Context,
+	const FCacheRecordPolicy& Policy,
+	const ECachePolicy SkipFlag,
+	const FCbObject& Object,
+	const bool bAlwaysLoadInlineData) const
 {
 	// TODO: Temporary function copied from FFileSystemDerivedDataBackend, helper for CreateRecord
 	const FPayloadId Id = Object.FindView("Id"_ASV).AsObjectId();
@@ -986,11 +1005,12 @@ FPayload FZenDerivedDataBackend::LegacyGetCachePayload(const FCacheKey& Key, FSt
 		return FPayload();
 	}
 
+	const ECachePolicy PayloadPolicy = Policy.GetPayloadPolicy(Id);
 	FPayload Payload(Id, RawHash, RawSize);
 
 	if (CompressedData)
 	{
-		if (EnumHasAllFlags(Policy, ECachePolicy::SkipData) && !bAlwaysLoadInlineData)
+		if (EnumHasAllFlags(PayloadPolicy, SkipFlag) && !bAlwaysLoadInlineData)
 		{
 			return Payload;
 		}
@@ -1000,7 +1020,7 @@ FPayload FZenDerivedDataBackend::LegacyGetCachePayload(const FCacheKey& Key, FSt
 		}
 	}
 
-	return LegacyGetCachePayload(Key, Context, Policy, Payload);
+	return LegacyGetCachePayload(Key, Context, PayloadPolicy, Payload);
 }
 
 FPayload FZenDerivedDataBackend::LegacyValidateCachePayload(const FCacheKey& Key, FStringView Context,
diff --git a/Engine/Source/Developer/DerivedDataCache/Private/ZenDerivedDataBackend.h b/Engine/Source/Developer/DerivedDataCache/Private/ZenDerivedDataBackend.h
--- a/Engine/Source/Developer/DerivedDataCache/Private/ZenDerivedDataBackend.h
+++ b/Engine/Source/Developer/DerivedDataCache/Private/ZenDerivedDataBackend.h
@@ -100,18 +100,15 @@ public:
 	virtual void Get(
 		TConstArrayView<FCacheKey> Keys,
 		FStringView Context,
-		ECachePolicy Policy,
+		FCacheRecordPolicy Policy,
 		IRequestOwner& Owner,
 		FOnCacheGetComplete&& OnComplete) override;
 
-	virtual void GetPayload(
-		TConstArrayView<FCachePayloadKey> Keys,
+	virtual void GetChunks(
+		TConstArrayView<FCacheChunkRequest> Chunks,
 		FStringView Context,
-		ECachePolicy Policy,
 		IRequestOwner& Owner,
-		FOnCacheGetPayloadComplete&& OnComplete) override;
-
-	virtual void CancelAll() override;
+		FOnCacheGetChunkComplete&& OnComplete) override;
 
 private:
 	enum class EGetResult
@@ -127,8 +124,7 @@ private:
 	EGetResult GetZenData(const FCacheKey& Key, ECachePolicy CachePolicy, FCbPackage& OutPackage) const;
 
 	bool PutCacheRecord(const FCacheRecord& Record, FStringView Context, ECachePolicy Policy);
-	FOptionalCacheRecord GetCacheRecord(const FCacheKey& Key, FStringView Context,
-		ECachePolicy Policy) const;
+	FOptionalCacheRecord GetCacheRecord(const FCacheKey& Key, FStringView Context, const FCacheRecordPolicy& Policy) const;
 
 	bool IsServiceReady();
 	static FString MakeLegacyZenKey(const TCHAR* CacheKey);
@@ -143,14 +139,14 @@ private:
 	bool LegacyPutCacheRecord(const FCacheRecord& Record, FStringView Context, ECachePolicy Policy);
 	bool LegacyPutCachePayload(const FCacheKey& Key, FStringView Context, const FPayload& Payload, FCbWriter& Writer);
 	FOptionalCacheRecord LegacyGetCacheRecord(const FCacheKey& Key, FStringView Context,
-		ECachePolicy Policy, bool bAlwaysLoadInlineData = false) const;
+		const FCacheRecordPolicy& Policy, bool bAlwaysLoadInlineData = false) const;
 	void LegacyMakeZenKey(const FCacheKey& CacheKey, FStringBuilderBase& Out) const;
 	void LegacyMakePayloadKey(const FCacheKey& CacheKey, const FIoHash& RawHash, FStringBuilderBase& Out) const;
 	FPayload LegacyGetCachePayload(const FCacheKey& Key, FStringView Context, ECachePolicy Policy, const FPayload& Payload) const;
 	FOptionalCacheRecord LegacyCreateRecord(FSharedBuffer&& RecordBytes, const FCacheKey& Key, FStringView Context,
-		ECachePolicy Policy, bool bAlwaysLoadInlineData) const;
-	FPayload LegacyGetCachePayload(const FCacheKey& Key, FStringView Context, ECachePolicy Policy, const FCbObject& Object,
-		bool bAlwaysLoadInlineData = false) const;
+		const FCacheRecordPolicy& Policy, bool bAlwaysLoadInlineData) const;
+	FPayload LegacyGetCachePayload(const FCacheKey& Key, FStringView Context, const FCacheRecordPolicy& Policy,
+		ECachePolicy PolicyMask, const FCbObject& Object, bool bAlwaysLoadInlineData = false) const;
 	FPayload LegacyValidateCachePayload(const FCacheKey& Key, FStringView Context, const FPayload& Payload,
 		const FIoHash& CompressedHash, FSharedBuffer&& CompressedData) const;
 
diff --git a/Engine/Source/Developer/DerivedDataCache/Public/DerivedDataCache.h b/Engine/Source/Developer/DerivedDataCache/Public/DerivedDataCache.h
--- a/Engine/Source/Developer/DerivedDataCache/Public/DerivedDataCache.h
+++ b/Engine/Source/Developer/DerivedDataCache/Public/DerivedDataCache.h
@@ -3,26 +3,32 @@
 #pragma once
 
 #include "CoreTypes.h"
+#include "Containers/StringFwd.h"
 #include "DerivedDataCacheKey.h"
+#include "DerivedDataPayloadId.h"
 #include "DerivedDataRequestTypes.h"
+#include "Math/NumericLimits.h"
 #include "Misc/EnumClassFlags.h"
 #include "Templates/Function.h"
+#include "Templates/RefCounting.h"
 
 #define UE_API DERIVEDDATACACHE_API
 
 namespace UE::DerivedData { class FCacheRecord; }
-namespace UE::DerivedData { class FPayload; }
+namespace UE::DerivedData { class FCacheRecordPolicy; }
 namespace UE::DerivedData { class IRequestOwner; }
+namespace UE::DerivedData { struct FCacheChunkRequest; }
+namespace UE::DerivedData { struct FCacheGetChunkCompleteParams; }
 namespace UE::DerivedData { struct FCacheGetCompleteParams; }
-namespace UE::DerivedData { struct FCacheGetPayloadCompleteParams; }
 namespace UE::DerivedData { struct FCachePutCompleteParams; }
+namespace UE::DerivedData::Private { class ICacheRecordPolicyShared; }
 
 namespace UE::DerivedData
 {
 
 using FOnCachePutComplete = TUniqueFunction<void (FCachePutCompleteParams&& Params)>;
 using FOnCacheGetComplete = TUniqueFunction<void (FCacheGetCompleteParams&& Params)>;
-using FOnCacheGetPayloadComplete = TUniqueFunction<void (FCacheGetPayloadCompleteParams&& Params)>;
+using FOnCacheGetChunkComplete = TUniqueFunction<void (FCacheGetChunkCompleteParams&& Params)>;
 
 ///////////////////////////////////////////////////////////////////////////////////////////////////////////////////////
 
@@ -50,34 +56,40 @@ enum class ECachePolicy : uint32
 	/** A value without any flags set. */
 	None            = 0,
 
-	/** Allow a cache get request or payload request to query local caches. */
+	/** Allow a cache request to query local caches. */
 	QueryLocal      = 1 << 0,
-	/** Allow a cache get request or payload request to query remote caches. */
+	/** Allow a cache request to query remote caches. */
 	QueryRemote     = 1 << 1,
-	/** Allow a cache get request or payload request to query local and remote caches. */
+	/** Allow a cache request to query local and remote caches. */
 	Query           = QueryLocal | QueryRemote,
 
-	/** Allow cache records to be stored in local caches by any cache request. */
+	/** Allow cache records and values to be stored in local caches. */
 	StoreLocal      = 1 << 2,
-	/** Allow cache records to be stored in remote caches by any cache request. */
+	/** Allow cache records and values to be stored in remote caches. */
 	StoreRemote     = 1 << 3,
-	/** Allow cache records to be stored in local and remote caches by any cache request. */
+	/** Allow cache records and values to be stored in local and remote caches. */
 	Store           = StoreLocal | StoreRemote,
 
-	/** Skip fetching the metadata for get requests. */
+	/** Skip fetching the metadata for record requests. */
 	SkipMeta        = 1 << 4,
-	/** Skip fetching the value for get requests and the payload for payload requests. */
+	/** Skip fetching the value for record, chunk, or value requests. */
 	SkipValue       = 1 << 5,
-	/** Skip fetching the attachments for get requests. */
+	/** Skip fetching the attachments for record requests. */
 	SkipAttachments = 1 << 6,
-	/** Skip fetching the metadata or payloads for get requests or payload requests. */
+	/**
+	 * Skip fetching the data for any requests.
+	 *
+	 * Put requests with skip flags may assume that record existence implies payload existence.
+	 */
 	SkipData        = SkipMeta | SkipValue | SkipAttachments,
 
-	/** Skip copying records to local caches for get requests if they already exist locally. */
-	SkipLocalCopy   = 1 << 7,
-
-	/** Keep records in the cache for at least the duration of the session. */
-	KeepAlive       = 1 << 8,
+	/**
+	 * Keep records in the cache for at least the duration of the session.
+	 *
+	 * This is a hint that the record may be accessed again in this session. This is mainly meant
+	 * to be used when subsequent accesses will not tolerate a cache miss.
+	 */
+	KeepAlive       = 1 << 7,
 
 	/** Allow cache requests to query and store records in local caches. */
 	Local           = QueryLocal | StoreLocal,
@@ -95,6 +107,92 @@ ENUM_CLASS_FLAGS(ECachePolicy);
 
 ///////////////////////////////////////////////////////////////////////////////////////////////////////////////////////
 
+/** A payload ID and the policy to use for that payload. */
+struct FCachePayloadPolicy
+{
+	FPayloadId Id;
+	ECachePolicy Policy = ECachePolicy::Default;
+};
+
+/** Interface for the private implementation of the cache record policy. */
+class Private::ICacheRecordPolicyShared
+{
+public:
+	virtual ~ICacheRecordPolicyShared() = default;
+	virtual void AddRef() const = 0;
+	virtual void Release() const = 0;
+	virtual TConstArrayView<FCachePayloadPolicy> GetPayloadPolicies() const = 0;
+	virtual void AddPayloadPolicy(const FCachePayloadPolicy& Policy) = 0;
+	virtual void Build() = 0;
+};
+
+/** A policy for cache record requests, with optional overrides by payload. */
+class FCacheRecordPolicy
+{
+public:
+	/** Construct a cache record policy that uses the default policy for records and payloads. */
+	FCacheRecordPolicy() = default;
+
+	/** Construct a cache record policy from a single policy for both records and payloads. */
+	inline FCacheRecordPolicy(ECachePolicy Policy)
+		: BasePolicy(Policy)
+		, RecordPolicy(Policy)
+	{
+	}
+
+	/** Returns true if the record and every payload use the same policy. */
+	inline bool IsUniform() const { return !Shared && BasePolicy == RecordPolicy; }
+
+	/** Returns the cache policy to use for the record. */
+	inline ECachePolicy GetRecordPolicy() const { return RecordPolicy; }
+
+	/** Returns the cache policy to use for the payload. */
+	UE_API ECachePolicy GetPayloadPolicy(const FPayloadId& Id) const;
+
+	/** Returns the cache policy to use for payloads with no override. */
+	inline ECachePolicy GetDefaultPayloadPolicy() const { return BasePolicy; }
+
+	/** Returns the array of cache policy overrides for payloads, sorted by ID. */
+	inline TConstArrayView<FCachePayloadPolicy> GetPayloadPolicies() const
+	{
+		return Shared ? Shared->GetPayloadPolicies() : TConstArrayView<FCachePayloadPolicy>();
+	}
+
+private:
+	friend class FCacheRecordPolicyBuilder;
+
+	ECachePolicy BasePolicy = ECachePolicy::Default;
+	ECachePolicy RecordPolicy = ECachePolicy::Default;
+	TRefCountPtr<const Private::ICacheRecordPolicyShared> Shared;
+};
+
+/** A cache record policy builder is used to construct a cache record policy. */
+class FCacheRecordPolicyBuilder
+{
+public:
+	/** Construct a cache record policy builder that uses the default policy as its base policy. */
+	FCacheRecordPolicyBuilder() = default;
+
+	/** Construct a cache record policy builder from the base policy for the record and its payloads. */
+	inline explicit FCacheRecordPolicyBuilder(ECachePolicy Policy)
+		: BasePolicy(Policy)
+	{
+	}
+
+	/** Adds a cache policy override for a payload. */
+	UE_API void AddPayloadPolicy(const FCachePayloadPolicy& Policy);
+	inline void AddPayloadPolicy(const FPayloadId& Id, ECachePolicy Policy) { AddPayloadPolicy({Id, Policy}); }
+
+	/** Build a cache record policy, which makes this builder subsequently unusable. */
+	UE_API FCacheRecordPolicy Build();
+
+private:
+	ECachePolicy BasePolicy = ECachePolicy::Default;
+	TRefCountPtr<Private::ICacheRecordPolicyShared> Shared;
+};
+
+///////////////////////////////////////////////////////////////////////////////////////////////////////////////////////
+
 /**
  * Interface to a store of cache records.
  *
@@ -114,13 +212,14 @@ public:
 	 * Records may finish storing in any order, and from multiple threads concurrently.
 	 *
 	 * A cache store is free to interpret a record containing only a key as a request to delete that
-	 * record from the store. Records containing payloads without data must be supported, and should
-	 * reference matching payloads that it already has stored, or otherwise return an error.
+	 * record from the store. Records may contain payloads that do not have data, and these payloads
+	 * must reference an existing copy of the payload in the store, if available, and must otherwise
+	 * be stored as a partial record that can attempt recovery of missing payloads when fetched.
 	 *
 	 * @param Records      The cache records to store. Must have a key.
 	 * @param Context      A description of the request. An object path is typically sufficient.
 	 * @param Policy       Flags to control the behavior of the request. See ECachePolicy.
-	 * @param Owner        The owner to execute the build within.
+	 * @param Owner        The owner to execute the request within.
 	 * @param OnComplete   A callback invoked for every record as it completes or is canceled.
 	 */
 	virtual void Put(
@@ -136,47 +235,46 @@ public:
 	 * The callback will always be called for every key, and may be called from an arbitrary thread.
 	 * Records may become available in any order, and from multiple threads concurrently.
 	 *
+	 * Records may propagate into other cache stores, in accordance with the policy. A propagated
+	 * record may be a partial record, with some payloads missing data, depending on the policy.
+	 *
+	 * When payloads are required by the policy, but not available, the status must be Error, but
+	 * the cache store may provide a partial record with the available payloads, to be completed,
+	 * when possible, by another cache store.
+	 *
 	 * @param Keys         The keys identifying the cache records to query.
 	 * @param Context      A description of the request. An object path is typically sufficient.
-	 * @param Policy       Flags to control the behavior of the request. See ECachePolicy.
-	 * @param Owner        The owner to execute the build within.
+	 * @param Policy       Flags to control the behavior of the request. See FCacheRecordPolicy.
+	 * @param Owner        The owner to execute the request within.
 	 * @param OnComplete   A callback invoked for every key as it completes or is canceled.
 	 */
 	virtual void Get(
 		TConstArrayView<FCacheKey> Keys,
 		FStringView Context,
-		ECachePolicy Policy,
+		FCacheRecordPolicy Policy,
 		IRequestOwner& Owner,
 		FOnCacheGetComplete&& OnComplete) = 0;
 
 	/**
-	 * Asynchronous request to get cache record payloads according to the policy.
+	 * Asynchronous request to get chunks, which are subsets of payloads, from cache records.
 	 *
-	 * The callback will always be called for every key, and may be called from an arbitrary thread.
-	 * Payloads may become available in any order, and from multiple threads concurrently.
+	 * The callback will always be called for every chunk, and may be called from an arbitrary thread.
+	 * Chunks may become available in any order, and from multiple threads concurrently.
+	 *
+	 * There is no requirement that the cache store validate the existence of other payloads from
+	 * the requested records. Requests for whole payloads may propagate those payloads into other
+	 * cache stores, in accordance with the policy, through a put of a partial record.
 	 *
-	 * @param Keys         The keys identifying the cache record payloads to query.
+	 * @param Chunks       The keys, IDs, offsets, and sizes of the chunks to query.
 	 * @param Context      A description of the request. An object path is typically sufficient.
-	 * @param Policy       Flags to control the behavior of the request. See ECachePolicy.
-	 * @param Owner        The owner to execute the build within.
-	 * @param OnComplete   A callback invoked for every key as it completes or is canceled.
+	 * @param Owner        The owner to execute the request within.
+	 * @param OnComplete   A callback invoked for every chunk as it completes or is canceled.
 	 */
-	virtual void GetPayload(
-		TConstArrayView<FCachePayloadKey> Keys,
+	virtual void GetChunks(
+		TConstArrayView<FCacheChunkRequest> Chunks,
 		FStringView Context,
-		ECachePolicy Policy,
 		IRequestOwner& Owner,
-		FOnCacheGetPayloadComplete&& OnComplete) = 0;
-
-	/**
-	 * Cancel all queued and active cache requests and invoke their callbacks.
-	 *
-	 * This is meant to be called before destroying the cache. Any new requests that are made during
-	 * execution of this function will also be canceled, because callbacks might queue requests when
-	 * they are invoked. This does not block any future cache requests, which must be handled by the
-	 * owner of the cache if future requests are meant to be avoided.
-	 */
-	virtual void CancelAll() = 0;
+		FOnCacheGetChunkComplete&& OnComplete) = 0;
 };
 
 ///////////////////////////////////////////////////////////////////////////////////////////////////////////////////////
@@ -190,6 +288,16 @@ class ICache : public ICacheStore
 {
 public:
 	virtual ~ICache() = default;
+
+	/**
+	 * Cancel all queued and active cache requests and invoke their callbacks.
+	 *
+	 * This is meant to be called before destroying the cache. Any new requests that are made during
+	 * execution of this function will also be canceled, because callbacks might queue requests when
+	 * they are invoked. This does not block any future cache requests, which must be handled by the
+	 * owner of the cache if future requests are meant to be avoided.
+	 */
+	virtual void CancelAll() = 0;
 };
 
 ///////////////////////////////////////////////////////////////////////////////////////////////////////////////////////
@@ -221,20 +329,45 @@ struct FCacheGetCompleteParams
 	EStatus Status = EStatus::Error;
 };
 
-/** Parameters for the completion callback for cache payload requests. */
-struct FCacheGetPayloadCompleteParams
+/** Parameters to request a chunk, which is a subset of a payload, from a cache record. */
+struct FCacheChunkRequest
 {
-	/** Key for the part of the payload request that completed or was canceled. See Payload for ID. */
+	/** The key identifying the cache record to fetch the payload from. */
 	FCacheKey Key;
 
-	/**
-	 * Payload for the part of the payload request that completed or was canceled.
-	 *
-	 * The ID is always populated.
-	 * The hash and size are populated when Status is Ok.
-	 * The data is populated when Status is Ok and the data was not skipped by the cache policy.
-	 */
-	FPayload&& Payload;
+	/** The ID of the payload to fetch from the cache record. */
+	FPayloadId Id;
+
+	/** The offset into the raw bytes of the payload at which to start fetching. */
+	uint64 RawOffset = 0;
+
+	/** The maximum number of raw bytes of the payload to fetch, starting from the offset. */
+	uint64 RawSize = MAX_uint64;
+
+	/** Flags to control the behavior of the request. See ECachePolicy. */
+	ECachePolicy Policy = ECachePolicy::Default;
+};
+
+/** Parameters for the completion callback for cache chunk requests. */
+struct FCacheGetChunkCompleteParams
+{
+	/** Key from the chunk request that completed or was canceled. */
+	FCacheKey Key;
+
+	/** ID from the chunk request that completed or was canceled. */
+	FPayloadId Id;
+
+	/** Offset from the chunk request that completed or was canceled. */
+	uint64 RawOffset = 0;
+
+	/** Size, in bytes, of the subset of the payload that was fetched, if any. */
+	uint64 RawSize = 0;
+
+	/** Hash of the entire payload, even if only a subset was fetched. */
+	const FIoHash& RawHash;
+
+	/** Data for the subset of the payload that was fetched when Status is Ok, otherwise null. */
+	FSharedBuffer&& RawData;
 
 	/** Status of the cache request. */
 	EStatus Status = EStatus::Error;
diff --git a/Engine/Source/Developer/DerivedDataCache/Public/DerivedDataCacheKey.h b/Engine/Source/Developer/DerivedDataCache/Public/DerivedDataCacheKey.h
--- a/Engine/Source/Developer/DerivedDataCache/Public/DerivedDataCacheKey.h
+++ b/Engine/Source/Developer/DerivedDataCache/Public/DerivedDataCacheKey.h
@@ -5,7 +5,6 @@
 #include "CoreTypes.h"
 #include "Containers/StringFwd.h"
 #include "Containers/StringView.h"
-#include "DerivedDataPayloadId.h"
 #include "IO/IoHash.h"
 #include "Templates/TypeHash.h"
 
@@ -76,18 +75,6 @@ struct FCacheKey
 
 inline const FCacheKey FCacheKey::Empty;
 
-/** A key that uniquely identifies a payload within a cache record. */
-struct FCachePayloadKey
-{
-	FCacheKey CacheKey;
-	FPayloadId Id;
-
-	/** A payload key with an empty cache key and a null payload identifier. */
-	static const FCachePayloadKey Empty;
-};
-
-inline const FCachePayloadKey FCachePayloadKey::Empty;
-
 ///////////////////////////////////////////////////////////////////////////////////////////////////////////////////////
 
 inline FAnsiStringView FCacheBucket::ToString() const
@@ -133,36 +120,6 @@ inline TStringBuilderBase<CharType>& operator<<(TStringBuilderBase<CharType>& Bu
 
 ///////////////////////////////////////////////////////////////////////////////////////////////////////////////////////
 
-inline bool operator==(const FCachePayloadKey& A, const FCachePayloadKey& B)
-{
-	return A.CacheKey == B.CacheKey && A.Id == B.Id;
-}
-
-inline bool operator!=(const FCachePayloadKey& A, const FCachePayloadKey& B)
-{
-	return A.CacheKey != B.CacheKey || A.Id != B.Id;
-}
-
-inline bool operator<(const FCachePayloadKey& A, const FCachePayloadKey& B)
-{
-	const FCacheKey& KeyA = A.CacheKey;
-	const FCacheKey& KeyB = B.CacheKey;
-	return KeyA == KeyB ? A.Id < B.Id : KeyA < KeyB;
-}
-
-inline uint32 GetTypeHash(const FCachePayloadKey& Key)
-{
-	return HashCombine(GetTypeHash(Key.CacheKey), GetTypeHash(Key.Id));
-}
-
-template <typename CharType>
-inline TStringBuilderBase<CharType>& operator<<(TStringBuilderBase<CharType>& Builder, const FCachePayloadKey& Key)
-{
-	return Builder << Key.CacheKey << CharType('/') << Key.Id;
-}
-
-///////////////////////////////////////////////////////////////////////////////////////////////////////////////////////
-
 } // UE::DerivedData
 
 #undef UE_API
diff --git a/Engine/Source/Developer/DerivedDataCache/Public/DerivedDataChunk.h b/Engine/Source/Developer/DerivedDataCache/Public/DerivedDataChunk.h
new file mode 100644
--- /dev/null
+++ b/Engine/Source/Developer/DerivedDataCache/Public/DerivedDataChunk.h
@@ -0,0 +1,45 @@
+// Copyright Epic Games, Inc. All Rights Reserved.
+
+#pragma once
+
+#include "CoreTypes.h"
+
+namespace UE::DerivedData
+{
+
+/** Binary predicate that compares chunks by key, then payload ID, then raw offset. */
+struct TChunkEqual
+{
+	template <typename ChunkTypeA, typename ChunkTypeB>
+	inline bool operator()(ChunkTypeA&& A, ChunkTypeB&& B) const
+	{
+		return A.Key == B.Key && A.Id == B.Id && A.RawOffset == B.RawOffset;
+	}
+};
+/** Binary predicate that compares chunks by key, then payload ID, then raw offset. */
+struct TChunkLess
+{
+	template <typename ChunkTypeA, typename ChunkTypeB>
+	inline bool operator()(ChunkTypeA&& A, ChunkTypeB&& B) const
+	{
+		if (A.Key < B.Key)
+		{
+			return true;
+		}
+		if (B.Key < A.Key)
+		{
+			return false;
+		}
+		if (A.Id < B.Id)
+		{
+			return true;
+		}
+		if (B.Id < A.Id)
+		{
+			return false;
+		}
+		return A.RawOffset < B.RawOffset;
+	}
+};
+
+} // UE::DerivedData
diff --git a/Engine/Source/Editor/UnrealEd/Private/EditorDomain/EditorDomainArchive.cpp b/Engine/Source/Editor/UnrealEd/Private/EditorDomain/EditorDomainArchive.cpp
--- a/Engine/Source/Editor/UnrealEd/Private/EditorDomain/EditorDomainArchive.cpp
+++ b/Engine/Source/Editor/UnrealEd/Private/EditorDomain/EditorDomainArchive.cpp
@@ -400,10 +400,11 @@ void FEditorDomainPackageSegments::SendSegmentRequest(FSegment& Segment)
 	ICache& Cache = GetCache();
 
 	// Note that Segment.RequestOwner is Interface-only and so we can write it outside the lock
-	Segment.RequestOwner.Emplace(UE::DerivedData::EPriority::Normal);
-	Cache.GetPayload({ FCachePayloadKey{ UE::EditorDomain::GetEditorDomainPackageKey(PackageDigest), Segment.PayloadId} },
-		PackagePath.GetDebugName(), ECachePolicy::Local, *Segment.RequestOwner,
-		[this, &Segment](FCacheGetPayloadCompleteParams&& Params)
+	Segment.RequestOwner.Emplace(EPriority::Normal);
+	FCacheChunkRequest SegmentChunk{UE::EditorDomain::GetEditorDomainPackageKey(PackageDigest), Segment.PayloadId};
+	SegmentChunk.Policy = ECachePolicy::Local;
+	Cache.GetChunks({SegmentChunk}, PackagePath.GetDebugName(), *Segment.RequestOwner,
+		[this, &Segment](FCacheGetChunkCompleteParams&& Params)
 		{
 			if (AsyncSource == ESource::Closed)
 			{
@@ -421,7 +422,7 @@ void FEditorDomainPackageSegments::SendSegmentRequest(FSegment& Segment)
 			}
 			else
 			{
-				Segment.Data = Params.Payload.GetData().Decompress();
+				Segment.Data = MoveTemp(Params.RawData);
 				uint64 SegmentEnd = &Segment < &Segments.Last() ? (&Segment + 1)->Start : Size;
 				uint64 SegmentSize = SegmentEnd - Segment.Start;
 				if (Segment.Data.GetSize() != SegmentSize)
diff --git a/Engine/Source/Runtime/Engine/Classes/Engine/Texture.h b/Engine/Source/Runtime/Engine/Classes/Engine/Texture.h
--- a/Engine/Source/Runtime/Engine/Classes/Engine/Texture.h
+++ b/Engine/Source/Runtime/Engine/Classes/Engine/Texture.h
@@ -35,6 +35,10 @@ class ITargetPlatform;
 class UAssetUserData;
 struct FPropertyChangedEvent;
 
+#if WITH_EDITORONLY_DATA
+namespace UE::DerivedData { struct FPayloadId; }
+#endif
+
 UENUM()
 enum TextureFilter
 {
@@ -691,7 +695,7 @@ public:
 
 #if WITH_EDITORONLY_DATA
 	FString GetDerivedDataMipKeyString(int32 MipIndex, const FTexture2DMipMap& Mip) const;
-	UE::DerivedData::FCachePayloadKeyProxy GetDerivedDataMipKeyProxy(int32 MipIndex, const FTexture2DMipMap& Mip) const;
+	static UE::DerivedData::FPayloadId MakeMipId(int32 MipIndex);
 #endif // WITH_EDITORONLY_DATA
 
 	/** 
diff --git a/Engine/Source/Runtime/Engine/Private/DerivedDataCacheKeyProxy.cpp b/Engine/Source/Runtime/Engine/Private/DerivedDataCacheKeyProxy.cpp
--- a/Engine/Source/Runtime/Engine/Private/DerivedDataCacheKeyProxy.cpp
+++ b/Engine/Source/Runtime/Engine/Private/DerivedDataCacheKeyProxy.cpp
@@ -3,6 +3,7 @@
 #include "DerivedDataCacheKeyProxy.h"
 
 #if WITH_EDITORONLY_DATA
+
 #include "DerivedDataCacheKey.h"
 
 namespace UE::DerivedData
@@ -11,9 +12,6 @@ namespace UE::DerivedData
 static_assert(sizeof(FCacheKeyProxy) == sizeof(FCacheKeyProxy));
 static_assert(alignof(FCacheKeyProxy) == alignof(FCacheKeyProxy));
 
-static_assert(sizeof(FCachePayloadKeyProxy) == sizeof(FCachePayloadKey));
-static_assert(alignof(FCachePayloadKeyProxy) == alignof(FCachePayloadKey));
-
 FCacheKeyProxy::FCacheKeyProxy(const FCacheKey& InKey)
 {
 	new(AsCacheKey()) FCacheKey(InKey);
@@ -24,16 +22,6 @@ FCacheKeyProxy::~FCacheKeyProxy()
 	AsCacheKey()->~FCacheKey();
 }
 
-FCachePayloadKeyProxy::FCachePayloadKeyProxy(const FCachePayloadKey& InKey)
-{
-	new(AsCachePayloadKey()) FCachePayloadKey(InKey);
-}
-
-FCachePayloadKeyProxy::~FCachePayloadKeyProxy()
-{
-	AsCachePayloadKey()->~FCachePayloadKey();
-}
-
-}
+} // UE::DerivedData
 
 #endif // WITH_EDITORONLY_DATA
\ No newline at end of file
diff --git a/Engine/Source/Runtime/Engine/Private/Streaming/Texture2DMipDataProvider_DDC.cpp b/Engine/Source/Runtime/Engine/Private/Streaming/Texture2DMipDataProvider_DDC.cpp
--- a/Engine/Source/Runtime/Engine/Private/Streaming/Texture2DMipDataProvider_DDC.cpp
+++ b/Engine/Source/Runtime/Engine/Private/Streaming/Texture2DMipDataProvider_DDC.cpp
@@ -56,32 +56,31 @@ void FTexture2DMipDataProvider_DDC::Init(const FTextureUpdateContext& Context, c
 			else if (PlatformData->DerivedDataKey.IsType<UE::DerivedData::FCacheKeyProxy>())
 			{
 				using namespace UE::DerivedData;
-				TArray<FCachePayloadKey> MipKeys;
+				TArray<FCacheChunkRequest> MipKeys;
 
+				const FCacheKey& Key = *PlatformData->DerivedDataKey.Get<UE::DerivedData::FCacheKeyProxy>().AsCacheKey();
 				for (int32 MipIndex = PendingFirstLODIdx; MipIndex < CurrentFirstLODIdx; ++MipIndex)
 				{
 					const FTexture2DMipMap& MipMap = *Context.MipsView[MipIndex];
 					if (MipMap.IsPagedToDerivedData())
 					{
-						MipKeys.Emplace(*PlatformData->GetDerivedDataMipKeyProxy(MipIndex + LODBias, MipMap).AsCachePayloadKey());
+						MipKeys.Add({Key, FTexturePlatformData::MakeMipId(MipIndex + LODBias)});
 					}
 				}
 
 				if (MipKeys.Num())
 				{
-					GetCache().GetPayload(MipKeys, Context.Texture->GetPathName(), ECachePolicy::Default, DDCRequestOwner,
-						[this, &Context, LODBias](FCacheGetPayloadCompleteParams&& Params)
+					GetCache().GetChunks(MipKeys, Context.Texture->GetPathName(), DDCRequestOwner,
+						[this, &Context, LODBias](FCacheGetChunkCompleteParams&& Params)
 						{
 							if (Params.Status == EStatus::Ok)
 							{
 								for (int32 MipIndex = PendingFirstLODIdx; MipIndex < CurrentFirstLODIdx; ++MipIndex)
 								{
-									TStringBuilder<16> MipName;
-									MipName << TEXT("Mip") << (MipIndex + LODBias);
-									if (Params.Payload.GetId() == FPayloadId::FromName(MipName))
+									if (Params.Id == FTexturePlatformData::MakeMipId(MipIndex + LODBias))
 									{
 										check(!DDCBuffers[MipIndex]);
-										DDCBuffers[MipIndex] = Params.Payload.GetData().Decompress();
+										DDCBuffers[MipIndex] = MoveTemp(Params.RawData);
 										break;
 									}
 								}
diff --git a/Engine/Source/Runtime/Engine/Private/Streaming/Texture2DStreamIn_DDC.cpp b/Engine/Source/Runtime/Engine/Private/Streaming/Texture2DStreamIn_DDC.cpp
--- a/Engine/Source/Runtime/Engine/Private/Streaming/Texture2DStreamIn_DDC.cpp
+++ b/Engine/Source/Runtime/Engine/Private/Streaming/Texture2DStreamIn_DDC.cpp
@@ -159,32 +159,31 @@ void FTexture2DStreamIn_DDC::DoCreateAsyncDDCRequests(const FContext& Context)
 		else if (PlatformData->DerivedDataKey.IsType<UE::DerivedData::FCacheKeyProxy>())
 		{
 			using namespace UE::DerivedData;
-			TArray<FCachePayloadKey> MipKeys;
+			TArray<FCacheChunkRequest> MipKeys;
 
+			const FCacheKey& Key = *PlatformData->DerivedDataKey.Get<UE::DerivedData::FCacheKeyProxy>().AsCacheKey();
 			for (int32 MipIndex = PendingFirstLODIdx; MipIndex < CurrentFirstLODIdx && !IsCancelled(); ++MipIndex)
 			{
 				const FTexture2DMipMap& MipMap = *Context.MipsView[MipIndex];
 				if (MipMap.IsPagedToDerivedData())
 				{
-					MipKeys.Emplace(*PlatformData->GetDerivedDataMipKeyProxy(MipIndex + LODBias, MipMap).AsCachePayloadKey());
+					MipKeys.Add({Key, FTexturePlatformData::MakeMipId(MipIndex + LODBias)});
 				}
 			}
 
 			if (MipKeys.Num())
 			{
-				GetCache().GetPayload(MipKeys, Context.Texture->GetPathName(), ECachePolicy::Default, DDCRequestOwner,
-					[this, LODBias](FCacheGetPayloadCompleteParams&& Params)
+				GetCache().GetChunks(MipKeys, Context.Texture->GetPathName(), DDCRequestOwner,
+					[this, LODBias](FCacheGetChunkCompleteParams&& Params)
 					{
 						if (Params.Status == EStatus::Ok)
 						{
 							for (int32 MipIndex = PendingFirstLODIdx; MipIndex < CurrentFirstLODIdx; ++MipIndex)
 							{
-								TStringBuilder<16> MipName;
-								MipName << TEXT("Mip") << (MipIndex + LODBias);
-								if (Params.Payload.GetId() == FPayloadId::FromName(MipName))
+								if (Params.Id == FTexturePlatformData::MakeMipId(MipIndex + LODBias))
 								{
 									check(!DDCBuffers[MipIndex]);
-									DDCBuffers[MipIndex] = Params.Payload.GetData().Decompress();
+									DDCBuffers[MipIndex] = MoveTemp(Params.RawData);
 									break;
 								}
 							}
diff --git a/Engine/Source/Runtime/Engine/Private/TextureDerivedData.cpp b/Engine/Source/Runtime/Engine/Private/TextureDerivedData.cpp
--- a/Engine/Source/Runtime/Engine/Private/TextureDerivedData.cpp
+++ b/Engine/Source/Runtime/Engine/Private/TextureDerivedData.cpp
@@ -42,6 +42,7 @@
 #include "Interfaces/ITargetPlatform.h"
 #include "Interfaces/ITargetPlatformManagerModule.h"
 #include "Interfaces/ITextureFormat.h"
+#include "Misc/StringBuilder.h"
 #include "ProfilingDebugging/CookStats.h"
 #include "VT/VirtualTextureDataBuilder.h"
 #include "VT/LightmapVirtualTexture.h"
@@ -963,13 +964,15 @@ static bool BeginLoadDerivedMips(FTexturePlatformData& PlatformData, int32 First
 	}
 	else if (PlatformData.DerivedDataKey.IsType<FCacheKeyProxy>())
 	{
-		TArray<FCachePayloadKey> MipKeys;
+		TArray<FCacheChunkRequest> MipKeys;
+
+		const FCacheKey& Key = *PlatformData.DerivedDataKey.Get<UE::DerivedData::FCacheKeyProxy>().AsCacheKey();
 		for (int32 MipIndex = FirstMipToLoad; MipIndex < Mips.Num(); ++MipIndex)
 		{
 			const FTexture2DMipMap& Mip = Mips[MipIndex];
 			if (Mip.IsPagedToDerivedData())
 			{
-				MipKeys.Emplace(*PlatformData.GetDerivedDataMipKeyProxy(MipIndex, Mip).AsCachePayloadKey());
+				MipKeys.Add({Key, FTexturePlatformData::MakeMipId(MipIndex)});
 			}
 		}
 
@@ -978,18 +981,16 @@ static bool BeginLoadDerivedMips(FTexturePlatformData& PlatformData, int32 First
 			check(Callback);
 			bool bMiss = false;
 			FRequestOwner RequestOwner(EPriority::Blocking);
-			GetCache().GetPayload(MipKeys, Texture ? Texture->GetPathName() : TEXT("Unknown Texture"_SV), ECachePolicy::Default, RequestOwner,
-				[&Mips, Callback = MoveTemp(Callback), FirstMipToLoad, &bMiss](FCacheGetPayloadCompleteParams&& Params)
+			GetCache().GetChunks(MipKeys, Texture ? Texture->GetPathName() : TEXT("Unknown Texture"_SV), RequestOwner,
+				[&Mips, Callback = MoveTemp(Callback), FirstMipToLoad, &bMiss](FCacheGetChunkCompleteParams&& Params)
 			{
 				if (Params.Status == EStatus::Ok)
 				{
 					for (int32 MipIndex = FirstMipToLoad; MipIndex < Mips.Num(); ++MipIndex)
 					{
-						TStringBuilder<16> MipName;
-						MipName << TEXT("Mip") << MipIndex;
-						if (Params.Payload.GetId() == FPayloadId::FromName(MipName))
+						if (Params.Id == FTexturePlatformData::MakeMipId(MipIndex))
 						{
-							Callback(MipIndex, Params.Payload.GetData().Decompress());
+							Callback(MipIndex, MoveTemp(Params.RawData));
 							break;
 						}
 					}
@@ -1492,7 +1493,7 @@ bool FTexturePlatformData::AreDerivedMipsAvailable(FStringView Context) const
 			GetDerivedDataCacheRef().TryToPrefetch({MipKeys.Last()}, Context);
 		}
 	}
-	else if (DerivedDataKey.IsType<UE::DerivedData::FCachePayloadKeyProxy>())
+	else if (DerivedDataKey.IsType<UE::DerivedData::FCacheKeyProxy>())
 	{
 		bAreDerivedMipsAvailable = true;
 	}
@@ -1767,21 +1768,16 @@ void FTexturePlatformData::Serialize(FArchive& Ar, UTexture* Owner)
 }
 
 #if WITH_EDITORONLY_DATA
+
 FString FTexturePlatformData::GetDerivedDataMipKeyString(int32 MipIndex, const FTexture2DMipMap& Mip) const
 {
 	const FString& KeyString = DerivedDataKey.Get<FString>();
 	return FDerivedDataCacheInterface::SanitizeCacheKey(*FString::Printf(TEXT("%s_MIP%u_%dx%d"), *KeyString, MipIndex, Mip.SizeX, Mip.SizeY));
 }
 
-UE::DerivedData::FCachePayloadKeyProxy FTexturePlatformData::GetDerivedDataMipKeyProxy(int32 MipIndex, const FTexture2DMipMap& Mip) const
+UE::DerivedData::FPayloadId FTexturePlatformData::MakeMipId(int32 MipIndex)
 {
-	using namespace UE::DerivedData;
-	const FCacheKeyProxy& KeyProxy = DerivedDataKey.Get<FCacheKeyProxy>();
-
-	TStringBuilder<16> MipName;
-	MipName << TEXT("Mip") << MipIndex;
-	FCachePayloadKey PayloadKey {*KeyProxy.AsCacheKey(), FPayloadId::FromName(MipName)};
-	return FCachePayloadKeyProxy(PayloadKey);
+	return UE::DerivedData::FPayloadId::FromName(WriteToString<16>(TEXT("Mip"_SV), MipIndex));
 }
 
 #endif // WITH_EDITORONLY_DATA
diff --git a/Engine/Source/Runtime/Engine/Private/TextureDerivedDataTask.cpp b/Engine/Source/Runtime/Engine/Private/TextureDerivedDataTask.cpp
--- a/Engine/Source/Runtime/Engine/Private/TextureDerivedDataTask.cpp
+++ b/Engine/Source/Runtime/Engine/Private/TextureDerivedDataTask.cpp
@@ -1293,11 +1293,11 @@ private:
 			}
 			else if (bInlineMips && (MipIndex >= FirstMipToLoad))
 			{
-				const FPayload& StreamingMipPayload = Output.GetPayload(FPayloadId::FromName(WriteToString<8>(TEXT("Mip"), MipIndex)));
+				const FPayload& StreamingMipPayload = Output.GetPayload(FTexturePlatformData::MakeMipId(MipIndex));
 				if (!StreamingMipPayload)
 				{
-					UE_LOG(LogTexture, Error, TEXT("Missing texture streaming mip '%s' for build of '%s' by %s."),
-						*WriteToString<8>(TEXT("Mip"), MipIndex), *WriteToString<128>(Output.GetName()), *WriteToString<32>(Output.GetFunction()));
+					UE_LOG(LogTexture, Error, TEXT("Missing texture streaming mip '%d' for build of '%s' by %s."),
+						MipIndex, *WriteToString<128>(Output.GetName()), *WriteToString<32>(Output.GetFunction()));
 					return false;
 				}
 				FSharedBuffer StreamingMipData = StreamingMipPayload.GetData().Decompress();
diff --git a/Engine/Source/Runtime/Engine/Public/DerivedDataCacheKeyProxy.h b/Engine/Source/Runtime/Engine/Public/DerivedDataCacheKeyProxy.h
--- a/Engine/Source/Runtime/Engine/Public/DerivedDataCacheKeyProxy.h
+++ b/Engine/Source/Runtime/Engine/Public/DerivedDataCacheKeyProxy.h
@@ -10,7 +10,6 @@ namespace UE::DerivedData
 {
 
 struct FCacheKey;
-struct FCachePayloadKey;
 
 namespace Private
 {
@@ -19,12 +18,6 @@ namespace Private
 		const ANSICHAR* BucketNamePtrDummy = nullptr;
 		FIoHash HashDummy;
 	};
-
-	struct FCachePayloadKeyDummy
-	{
-		FCacheKeyDummy CacheKeyDummy;
-		alignas(uint32) uint8 PayloadBytesDummy[12];
-	};
 } // Private
 
 struct FCacheKeyProxy : private TAlignedBytes<sizeof(UE::DerivedData::Private::FCacheKeyDummy), alignof(UE::DerivedData::Private::FCacheKeyDummy)>
@@ -35,12 +28,4 @@ struct FCacheKeyProxy : private TAlignedBytes<sizeof(UE::DerivedData::Private::F
 	const FCacheKey* AsCacheKey() const { return (const FCacheKey*)this; }
 };
 
-struct FCachePayloadKeyProxy : private TAlignedBytes<sizeof(UE::DerivedData::Private::FCachePayloadKeyDummy), alignof(UE::DerivedData::Private::FCachePayloadKeyDummy)>
-{
-	FCachePayloadKeyProxy(const FCachePayloadKey& InKey);
-	~FCachePayloadKeyProxy();
-	FCachePayloadKey* AsCachePayloadKey() { return (FCachePayloadKey*)this;  }
-	const FCachePayloadKey* AsCachePayloadKey() const { return (const FCachePayloadKey*)this; }
-};
-
 } // UE::DerivedData
