DDC: Added ECachePolicy::PartialOnError and implemented partial cache records in the filesystem and memory backends

#rnx

diff --git a/Engine/Source/Developer/DerivedDataCache/Private/DerivedDataCache.cpp b/Engine/Source/Developer/DerivedDataCache/Private/DerivedDataCache.cpp
--- a/Engine/Source/Developer/DerivedDataCache/Private/DerivedDataCache.cpp
+++ b/Engine/Source/Developer/DerivedDataCache/Private/DerivedDataCache.cpp
@@ -209,7 +209,7 @@ ECachePolicy FCacheRecordPolicy::GetPayloadPolicy(const FPayloadId& Id) const
 			}
 		}
 	}
-	return BasePolicy;
+	return DefaultPayloadPolicy;
 }
 
 void FCacheRecordPolicyBuilder::AddPayloadPolicy(const FCachePayloadPolicy& Policy)
diff --git a/Engine/Source/Developer/DerivedDataCache/Private/FileSystemDerivedDataBackend.cpp b/Engine/Source/Developer/DerivedDataCache/Private/FileSystemDerivedDataBackend.cpp
--- a/Engine/Source/Developer/DerivedDataCache/Private/FileSystemDerivedDataBackend.cpp
+++ b/Engine/Source/Developer/DerivedDataCache/Private/FileSystemDerivedDataBackend.cpp
@@ -898,8 +898,8 @@ public:
 				UE_LOG(LogDerivedDataCache, Verbose, TEXT("%s: Cache put complete for %s from '%.*s'"),
 					*CachePath, *WriteToString<96>(Record.GetKey()), Context.Len(), Context.GetData());
 				TRACE_COUNTER_INCREMENT(FileSystemDDC_PutHit);
-				TRACE_COUNTER_ADD(FileSystemDDC_BytesWritten, MeasureCacheRecord(Record));
-				COOK_STAT(Timer.AddHit(MeasureCacheRecord(Record)));
+				TRACE_COUNTER_ADD(FileSystemDDC_BytesWritten, MeasureCompressedCacheRecord(Record));
+				COOK_STAT(Timer.AddHit(MeasureRawCacheRecord(Record)));
 				if (OnComplete)
 				{
 					OnComplete({Record.GetKey(), EStatus::Ok});
@@ -907,7 +907,7 @@ public:
 			}
 			else
 			{
-				COOK_STAT(Timer.AddMiss(MeasureCacheRecord(Record)));
+				COOK_STAT(Timer.AddMiss(MeasureRawCacheRecord(Record)));
 				if (OnComplete)
 				{
 					OnComplete({Record.GetKey(), EStatus::Error});
@@ -928,23 +928,25 @@ public:
 			TRACE_CPUPROFILER_EVENT_SCOPE(FileSystemDDC_Get);
 			TRACE_COUNTER_INCREMENT(FileSystemDDC_Get);
 			COOK_STAT(auto Timer = UsageStats.TimeGet());
-			if (FOptionalCacheRecord Record = GetCacheRecord(Key, Context, Policy))
+			const bool bAlwaysLoadInlineData = false;
+			EStatus Status = EStatus::Error;
+			if (FOptionalCacheRecord Record = GetCacheRecord(Key, Context, Policy, bAlwaysLoadInlineData, Status))
 			{
 				UE_LOG(LogDerivedDataCache, Verbose, TEXT("%s: Cache hit for %s from '%.*s'"),
 					*CachePath, *WriteToString<96>(Key), Context.Len(), Context.GetData());
 				TRACE_COUNTER_INCREMENT(FileSystemDDC_GetHit);
-				TRACE_COUNTER_ADD(FileSystemDDC_BytesRead, MeasureCacheRecord(Record.Get()));
-				COOK_STAT(Timer.AddHit(MeasureCacheRecord(Record.Get())));
+				TRACE_COUNTER_ADD(FileSystemDDC_BytesRead, MeasureCompressedCacheRecord(Record.Get()));
+				COOK_STAT(Timer.AddHit(MeasureRawCacheRecord(Record.Get())));
 				if (OnComplete)
 				{
-					OnComplete({MoveTemp(Record).Get(), EStatus::Ok});
+					OnComplete({MoveTemp(Record).Get(), Status});
 				}
 			}
 			else
 			{
 				if (OnComplete)
 				{
-					OnComplete({FCacheRecordBuilder(Key).Build(), EStatus::Error});
+					OnComplete({FCacheRecordBuilder(Key).Build(), Status});
 				}
 			}
 		}
@@ -960,35 +962,43 @@ public:
 		SortedChunks.StableSort(TChunkLess());
 
 		FOptionalCacheRecord Record;
+		EStatus RecordStatus = EStatus::Error;
 		for (const FCacheChunkRequest& Chunk : SortedChunks)
 		{
+			const bool bExistsOnly = EnumHasAnyFlags(Chunk.Policy, ECachePolicy::SkipValue);
 			TRACE_CPUPROFILER_EVENT_SCOPE(FileSystemDDC_Get);
 			TRACE_COUNTER_INCREMENT(FileSystemDDC_Get);
-			COOK_STAT(auto Timer = UsageStats.TimeGet());
+			COOK_STAT(auto Timer = bExistsOnly ? UsageStats.TimeProbablyExists() : UsageStats.TimeGet());
 			if (!Record || Record.Get().GetKey() != Chunk.Key)
 			{
+				const bool bAlwaysLoadInlineData = true;
 				FCacheRecordPolicyBuilder PolicyBuilder(ECachePolicy::None);
 				PolicyBuilder.AddPayloadPolicy(Chunk.Id, Chunk.Policy | ECachePolicy::SkipData);
-				Record = GetCacheRecord(Chunk.Key, Context, PolicyBuilder.Build(), /*bAlwaysLoadInlineData*/ true);
+				Record = GetCacheRecord(Chunk.Key, Context, PolicyBuilder.Build(), bAlwaysLoadInlineData, RecordStatus);
 			}
-			if (Record)
+			if (Record && RecordStatus == EStatus::Ok)
 			{
-				const ECachePolicy PolicyMask = Record.Get().GetValuePayload().GetId() == Chunk.Id
-					? ECachePolicy::SkipValue : ECachePolicy::SkipAttachments;
-				const ECachePolicy Policy = Chunk.Policy | (ECachePolicy::SkipData & ~PolicyMask);
-				if (FPayload Payload = GetCachePayload(Chunk.Key, Context, Policy, Record.Get().GetPayload(Chunk.Id)))
+				EStatus PayloadStatus = EStatus::Error;
+				const ECachePolicy Policy = Chunk.Policy | (ECachePolicy::SkipData & ~ECachePolicy::SkipValue);
+				const FPayload& RecordPayload = Record.Get().GetPayload(Chunk.Id);
+				if (FPayload Payload = GetCachePayload(Chunk.Key, Context, Policy, RecordPayload, PayloadStatus))
 				{
 					const uint64 RawSize = FMath::Min(Payload.GetRawSize(), Chunk.RawSize);
 					UE_LOG(LogDerivedDataCache, Verbose, TEXT("%s: Cache hit for %s from '%.*s'"),
 						*CachePath, *WriteToString<96>(Chunk.Key, '/', Chunk.Id), Context.Len(), Context.GetData());
 					TRACE_COUNTER_INCREMENT(FileSystemDDC_GetHit);
-					TRACE_COUNTER_ADD(FileSystemDDC_BytesRead, RawSize);
-					COOK_STAT(Timer.AddHit(RawSize));
+					TRACE_COUNTER_ADD(FileSystemDDC_BytesRead, Payload.GetData().GetCompressedSize());
+					COOK_STAT(Timer.AddHit(Payload.HasData() ? RawSize : 0));
 					if (OnComplete)
 					{
-						FUniqueBuffer Buffer = FUniqueBuffer::Alloc(RawSize);
-						Payload.GetData().DecompressToComposite().CopyTo(Buffer, Chunk.RawOffset);
-						OnComplete({Chunk.Key, Chunk.Id, Chunk.RawOffset, RawSize, Payload.GetRawHash(), Buffer.MoveToShared(), EStatus::Ok});
+						FUniqueBuffer Buffer;
+						if (Payload.HasData() && !bExistsOnly)
+						{
+							Buffer = FUniqueBuffer::Alloc(RawSize);
+							Payload.GetData().DecompressToComposite().CopyTo(Buffer, Chunk.RawOffset);
+						}
+						OnComplete({Chunk.Key, Chunk.Id, Chunk.RawOffset,
+							RawSize, Payload.GetRawHash(), Buffer.MoveToShared(), PayloadStatus});
 					}
 					continue;
 				}
@@ -1002,7 +1012,15 @@ public:
 	}
 
 private:
-	uint64 MeasureCacheRecord(const FCacheRecord& Record) const
+	uint64 MeasureCompressedCacheRecord(const FCacheRecord& Record) const
+	{
+		return Record.GetMeta().GetSize() +
+			Record.GetValuePayload().GetData().GetCompressedSize() +
+			Algo::TransformAccumulate(Record.GetAttachmentPayloads(),
+				[](const FPayload& Payload) { return Payload.GetData().GetCompressedSize(); }, uint64(0));
+	}
+
+	uint64 MeasureRawCacheRecord(const FCacheRecord& Record) const
 	{
 		return Record.GetMeta().GetSize() +
 			Record.GetValuePayload().GetData().GetRawSize() +
@@ -1014,7 +1032,8 @@ private:
 	{
 		if (!IsWritable())
 		{
-			UE_LOG(LogDerivedDataCache, VeryVerbose, TEXT("%s: Skipped put of %s from '%.*s' because this is not writable"),
+			UE_LOG(LogDerivedDataCache, VeryVerbose,
+				TEXT("%s: Skipped put of %s from '%.*s' because this cache store is not writable"),
 				*CachePath, *WriteToString<96>(Record.GetKey()), Context.Len(), Context.GetData());
 			return false;
 		}
@@ -1092,11 +1111,15 @@ private:
 		const FCacheKey& Key,
 		const FStringView Context,
 		const FCacheRecordPolicy& Policy,
-		const bool bAlwaysLoadInlineData = false)
+		const bool bAlwaysLoadInlineData,
+		EStatus& OutStatus)
 	{
+		OutStatus = EStatus::Error;
+
 		if (!IsUsable())
 		{
-			UE_LOG(LogDerivedDataCache, VeryVerbose, TEXT("%s: Skipped get of %s from '%.*s' because this cache store is not available"),
+			UE_LOG(LogDerivedDataCache, VeryVerbose,
+				TEXT("%s: Skipped get of %s from '%.*s' because this cache store is not available"),
 				*CachePath, *WriteToString<96>(Key), Context.Len(), Context.GetData());
 			return FOptionalCacheRecord();
 		}
@@ -1154,13 +1177,14 @@ private:
 		{
 			if (FCbFieldView MetaHash = RecordObject.FindView("MetaHash"_ASV))
 			{
-				if (FCbObject MetaObject = RecordObject["Meta"_ASV].AsObject(); MetaObject.GetHash() == MetaHash.AsHash())
+				if (FCbObject Meta = RecordObject["Meta"_ASV].AsObject(); Meta.GetHash() == MetaHash.AsHash())
 				{
-					RecordBuilder.SetMeta(MoveTemp(MetaObject));
+					RecordBuilder.SetMeta(MoveTemp(Meta));
 				}
 				else
 				{
-					UE_LOG(LogDerivedDataCache, Display, TEXT("%s: Cache miss with corrupted metadata for %s from '%.*s'"),
+					UE_LOG(LogDerivedDataCache, Display,
+						TEXT("%s: Cache miss with corrupted metadata for %s from '%.*s'"),
 						*CachePath, *WriteToString<96>(Key), Context.Len(), Context.GetData());
 					return FOptionalCacheRecord();
 				}
@@ -1169,10 +1193,17 @@ private:
 
 		// Disable deletion now that the record is validated.
 		bDeleteCacheObject = false;
+		OutStatus = EStatus::Ok;
 
 		if (FCbObject ValueObject = RecordObject["Value"_ASV].AsObject())
 		{
-			FPayload Payload = GetCachePayload(Key, Context, Policy, ECachePolicy::SkipValue, ValueObject, bAlwaysLoadInlineData);
+			EStatus PayloadStatus = EStatus::Error;
+			FPayload Payload = GetCachePayload(Key, Context, Policy,
+				ECachePolicy::SkipValue, ValueObject, bAlwaysLoadInlineData, PayloadStatus);
+			if (PayloadStatus == EStatus::Error)
+			{
+				OutStatus = EStatus::Error;
+			}
 			if (Payload.IsNull())
 			{
 				return FOptionalCacheRecord();
@@ -1182,7 +1213,13 @@ private:
 
 		for (FCbField AttachmentField : RecordObject["Attachments"_ASV])
 		{
-			FPayload Payload = GetCachePayload(Key, Context, Policy, ECachePolicy::SkipAttachments, AttachmentField.AsObject(), bAlwaysLoadInlineData);
+			EStatus PayloadStatus = EStatus::Error;
+			FPayload Payload = GetCachePayload(Key, Context, Policy,
+				ECachePolicy::SkipAttachments, AttachmentField.AsObject(), bAlwaysLoadInlineData, PayloadStatus);
+			if (PayloadStatus == EStatus::Error)
+			{
+				OutStatus = EStatus::Error;
+			}
 			if (Payload.IsNull())
 			{
 				return FOptionalCacheRecord();
@@ -1200,7 +1237,12 @@ private:
 		return RecordBuilder.Build();
 	}
 
-	bool PutCachePayload(const FCacheKey& Key, FStringView Context, const FPayload& Payload, FCbWriter& Writer, bool bStoreInline) const
+	bool PutCachePayload(
+		const FCacheKey& Key,
+		const FStringView Context,
+		const FPayload& Payload,
+		FCbWriter& Writer,
+		const bool bStoreInline) const
 	{
 		const FIoHash& RawHash = Payload.GetRawHash();
 		const bool bHasData = Payload.HasData();
@@ -1253,7 +1295,8 @@ private:
 		const FCacheRecordPolicy& Policy,
 		const ECachePolicy SkipFlag,
 		const FCbObject& Object,
-		const bool bAlwaysLoadInlineData) const
+		const bool bAlwaysLoadInlineData,
+		EStatus& OutStatus) const
 	{
 		const FPayloadId Id = Object.FindView("Id"_ASV).AsObjectId();
 		const uint64 RawSize = Object.FindView("RawSize"_ASV).AsUInt64(MAX_uint64);
@@ -1265,6 +1308,7 @@ private:
 		{
 			UE_LOG(LogDerivedDataCache, Display, TEXT("%s: Cache miss with invalid record format for %s from '%.*s'"),
 				*CachePath, *WriteToString<96>(Key), Context.Len(), Context.GetData());
+			OutStatus = EStatus::Error;
 			return FPayload();
 		}
 
@@ -1275,21 +1319,26 @@ private:
 		{
 			if (EnumHasAllFlags(PayloadPolicy, SkipFlag) && !bAlwaysLoadInlineData)
 			{
+				OutStatus = EStatus::Ok;
 				return Payload;
 			}
-			else
-			{
-				return ValidateCachePayload(Key, Context, Payload, CompressedHash, MoveTemp(CompressedData));
-			}
+			return ValidateCachePayload(Key, Context, PayloadPolicy,
+				Payload, CompressedHash, MoveTemp(CompressedData), OutStatus);
 		}
 
-		return GetCachePayload(Key, Context, PayloadPolicy, Payload);
+		return GetCachePayload(Key, Context, PayloadPolicy, Payload, OutStatus);
 	}
 
-	FPayload GetCachePayload(const FCacheKey& Key, FStringView Context, ECachePolicy Policy, const FPayload& Payload) const
+	FPayload GetCachePayload(
+		const FCacheKey& Key,
+		const FStringView Context,
+		const ECachePolicy Policy,
+		const FPayload& Payload,
+		EStatus& OutStatus) const
 	{
 		if (Payload.HasData() || !EnumHasAnyFlags(Policy, ECachePolicy::Query))
 		{
+			OutStatus = EStatus::Ok;
 			return Payload;
 		}
 
@@ -1303,6 +1352,7 @@ private:
 				{
 					AccessLogWriter->Append(Key, Payload.GetId(), Path);
 				}
+				OutStatus = EStatus::Ok;
 				return Payload;
 			}
 		}
@@ -1313,7 +1363,8 @@ private:
 				const FMemoryView CompressedDataView = CompressedData;
 				const FIoHash CompressedHash(CompressedDataView.Left(sizeof(FIoHash)));
 				CompressedData = FSharedBuffer::MakeView(CompressedDataView + sizeof(FIoHash), MoveTemp(CompressedData));
-				const FPayload FullPayload = ValidateCachePayload(Key, Context, Payload, CompressedHash, MoveTemp(CompressedData));
+				const FPayload FullPayload = ValidateCachePayload(Key, Context, Policy,
+					Payload, CompressedHash, MoveTemp(CompressedData), OutStatus);
 				if (FullPayload && AccessLogWriter)
 				{
 					AccessLogWriter->Append(Key, FullPayload.GetId(), Path);
@@ -1326,28 +1377,33 @@ private:
 			TEXT("%s: Cache miss with missing payload %s with hash %s for %s from '%.*s'"),
 			*CachePath, *WriteToString<16>(Payload.GetId()), *WriteToString<48>(Payload.GetRawHash()), *WriteToString<96>(Key),
 			Context.Len(), Context.GetData());
-		return FPayload();
+		OutStatus = EStatus::Error;
+		return EnumHasAllFlags(Policy, ECachePolicy::PartialOnError) ? Payload : FPayload::Null;
 	}
 
 	FPayload ValidateCachePayload(
 		const FCacheKey& Key,
 		const FStringView Context,
+		const ECachePolicy Policy,
 		const FPayload& Payload,
 		const FIoHash& CompressedHash,
-		FSharedBuffer&& CompressedData) const
+		FSharedBuffer&& CompressedData,
+		EStatus& OutStatus) const
 	{
 		if (FCompressedBuffer CompressedBuffer = FCompressedBuffer::FromCompressed(MoveTemp(CompressedData));
 			CompressedBuffer &&
 			CompressedBuffer.GetRawHash() == Payload.GetRawHash() &&
 			FIoHash::HashBuffer(CompressedBuffer.GetCompressed()) == CompressedHash)
 		{
+			OutStatus = EStatus::Ok;
 			return FPayload(Payload.GetId(), MoveTemp(CompressedBuffer));
 		}
 		UE_LOG(LogDerivedDataCache, Display,
 			TEXT("%s: Cache miss with corrupted payload %s with hash %s for %s from '%.*s'"),
-			*CachePath, *WriteToString<16>(Payload.GetId()), *WriteToString<48>(Payload.GetRawHash()), *WriteToString<96>(Key),
-			Context.Len(), Context.GetData());
-		return FPayload();
+			*CachePath, *WriteToString<16>(Payload.GetId()), *WriteToString<48>(Payload.GetRawHash()),
+			*WriteToString<96>(Key), Context.Len(), Context.GetData());
+		OutStatus = EStatus::Error;
+		return EnumHasAllFlags(Policy, ECachePolicy::PartialOnError) ? Payload : FPayload::Null;
 	}
 
 	void BuildCacheObjectPath(const FCacheKey& CacheKey, FStringBuilderBase& Path) const
diff --git a/Engine/Source/Developer/DerivedDataCache/Private/MemoryDerivedDataBackend.cpp b/Engine/Source/Developer/DerivedDataCache/Private/MemoryDerivedDataBackend.cpp
--- a/Engine/Source/Developer/DerivedDataCache/Private/MemoryDerivedDataBackend.cpp
+++ b/Engine/Source/Developer/DerivedDataCache/Private/MemoryDerivedDataBackend.cpp
@@ -509,8 +509,10 @@ void FMemoryDerivedDataBackend::Get(
 {
 	for (const FCacheKey& Key : Keys)
 	{
-		COOK_STAT(auto Timer = UsageStats.TimeGet());
+		const bool bExistsOnly = EnumHasAllFlags(Policy.GetRecordPolicy(), ECachePolicy::SkipData);
+		COOK_STAT(auto Timer = bExistsOnly ? UsageStats.TimeProbablyExists() : UsageStats.TimeGet());
 		FOptionalCacheRecord Record;
+		EStatus Status = EStatus::Error;
 		if (ShouldSimulateMiss(Key))
 		{
 			UE_LOG(LogDerivedDataCache, Verbose, TEXT("%s: Simulated miss for get of %s from '%.*s'"),
@@ -518,14 +520,42 @@ void FMemoryDerivedDataBackend::Get(
 		}
 		else if (FScopeLock ScopeLock(&SynchronizationObject); const FCacheRecord* CacheRecord = CacheRecords.Find(Key))
 		{
+			Status = EStatus::Ok;
 			Record = *CacheRecord;
 		}
 		if (Record)
+		{
+			if (const FPayload& Payload = Record.Get().GetValuePayload();
+				!Payload.HasData() && !EnumHasAnyFlags(Policy.GetPayloadPolicy(Payload.GetId()), ECachePolicy::SkipValue))
+			{
+				Status = EStatus::Error;
+				if (!EnumHasAllFlags(Policy.GetPayloadPolicy(Payload.GetId()), ECachePolicy::PartialOnError))
+				{
+					Record.Reset();
+				}
+			}
+			if (Record)
+			{
+				for (const FPayload& Payload : Record.Get().GetAttachmentPayloads())
+				{
+					if (!Payload.HasData() && !EnumHasAnyFlags(Policy.GetPayloadPolicy(Payload.GetId()), ECachePolicy::SkipAttachments))
+					{
+						Status = EStatus::Error;
+						if (!EnumHasAllFlags(Policy.GetPayloadPolicy(Payload.GetId()), ECachePolicy::PartialOnError))
+						{
+							Record.Reset();
+							break;
+						}
+					}
+				}
+			}
+		}
+		if (Record)
 		{
 			COOK_STAT(Timer.AddHit(CalcRawCacheRecordSize(Record.Get())));
 			if (OnComplete)
 			{
-				OnComplete({MoveTemp(Record).Get(), EStatus::Ok});
+				OnComplete({MoveTemp(Record).Get(), Status});
 			}
 		}
 		else
@@ -546,7 +576,8 @@ void FMemoryDerivedDataBackend::GetChunks(
 {
 	for (const FCacheChunkRequest& Chunk : Chunks)
 	{
-		COOK_STAT(auto Timer = UsageStats.TimeGet());
+		const bool bExistsOnly = EnumHasAnyFlags(Chunk.Policy, ECachePolicy::SkipValue);
+		COOK_STAT(auto Timer = bExistsOnly ? UsageStats.TimeProbablyExists() : UsageStats.TimeGet());
 		FPayload Payload;
 		if (ShouldSimulateMiss(Chunk.Key))
 		{
@@ -563,9 +594,15 @@ void FMemoryDerivedDataBackend::GetChunks(
 			COOK_STAT(Timer.AddHit(RawSize));
 			if (OnComplete)
 			{
-				FUniqueBuffer Buffer = FUniqueBuffer::Alloc(RawSize);
-				Payload.GetData().DecompressToComposite().CopyTo(Buffer, Chunk.RawOffset);
-				OnComplete({Chunk.Key, Chunk.Id, Chunk.RawOffset, RawSize, Payload.GetRawHash(), Buffer.MoveToShared(), EStatus::Ok});
+				FUniqueBuffer Buffer;
+				if (Payload.HasData() && !bExistsOnly)
+				{
+					Buffer = FUniqueBuffer::Alloc(RawSize);
+					Payload.GetData().DecompressToComposite().CopyTo(Buffer, Chunk.RawOffset);
+				}
+				const EStatus Status = bExistsOnly || Buffer ? EStatus::Ok : EStatus::Error;
+				OnComplete({Chunk.Key, Chunk.Id, Chunk.RawOffset,
+					RawSize, Payload.GetRawHash(), Buffer.MoveToShared(), Status});
 			}
 		}
 		else
diff --git a/Engine/Source/Developer/DerivedDataCache/Public/DerivedDataCache.h b/Engine/Source/Developer/DerivedDataCache/Public/DerivedDataCache.h
--- a/Engine/Source/Developer/DerivedDataCache/Public/DerivedDataCache.h
+++ b/Engine/Source/Developer/DerivedDataCache/Public/DerivedDataCache.h
@@ -44,8 +44,6 @@ using FOnCacheGetChunkComplete = TUniqueFunction<void (FCacheGetChunkCompletePar
  * Get(Query | StoreLocal): read from any cache; put to writable local caches if missing.
  * Get(Query | SkipData): check for existence in any cache; do not modify any cache.
  * Get(Default | SkipData): check for existence in any cache; put to writable caches if missing.
- * Get(Default | SkipLocalCopy): read from any cache; put to writable remote caches if missing,
- *                               and to writable local caches if not present in any local cache.
  *
  * Put(Default): write to every writable cache.
  * Put(Remote): write to every writable remote cache, skipping local caches.
@@ -60,14 +58,14 @@ enum class ECachePolicy : uint32
 	QueryLocal      = 1 << 0,
 	/** Allow a cache request to query remote caches. */
 	QueryRemote     = 1 << 1,
-	/** Allow a cache request to query local and remote caches. */
+	/** Allow a cache request to query any caches. */
 	Query           = QueryLocal | QueryRemote,
 
 	/** Allow cache records and values to be stored in local caches. */
 	StoreLocal      = 1 << 2,
 	/** Allow cache records and values to be stored in remote caches. */
 	StoreRemote     = 1 << 3,
-	/** Allow cache records and values to be stored in local and remote caches. */
+	/** Allow cache records and values to be stored in any caches. */
 	Store           = StoreLocal | StoreRemote,
 
 	/** Skip fetching the metadata for record requests. */
@@ -91,15 +89,27 @@ enum class ECachePolicy : uint32
 	 */
 	KeepAlive       = 1 << 7,
 
-	/** Allow cache requests to query and store records in local caches. */
+	/**
+	 * Partial output will be provided with the error status when a required payload is missing.
+	 *
+	 * This is meant for cases when the missing payloads can be individually recovered or rebuilt
+	 * without rebuilding the whole record. The cache automatically adds this flag when there are
+	 * other cache stores that it may be able to recover missing payloads from.
+	 *
+	 * Requests for records would return records where the missing payloads have a hash and size,
+	 * but no data. Requests for chunks or values would return the hash and size, but no data.
+	 */
+	PartialOnError  = 1 << 8,
+
+	/** Allow cache requests to query and store records and values in local caches. */
 	Local           = QueryLocal | StoreLocal,
-	/** Allow cache requests to query and store records in remote caches. */
+	/** Allow cache requests to query and store records and values in remote caches. */
 	Remote          = QueryRemote | StoreRemote,
 
-	/** Allow cache requests to query and store records in local and remote caches. */
+	/** Allow cache requests to query and store records and values in any caches. */
 	Default         = Query | Store,
 
-	/** Do not allow cache requests to query or store records in local or remote caches. */
+	/** Do not allow cache requests to query or store records and values in any caches. */
 	Disable         = None,
 };
 
@@ -107,7 +117,7 @@ ENUM_CLASS_FLAGS(ECachePolicy);
 
 ///////////////////////////////////////////////////////////////////////////////////////////////////////////////////////
 
-/** A payload ID and the policy to use for that payload. */
+/** A payload ID and the cache policy to use for that payload. */
 struct FCachePayloadPolicy
 {
 	FPayloadId Id;
@@ -126,22 +136,30 @@ public:
 	virtual void Build() = 0;
 };
 
-/** A policy for cache record requests, with optional overrides by payload. */
+/**
+ * Flags to control the behavior of cache requests, with optional overrides by payload.
+ *
+ * Examples:
+ * - A base policy of Disable with payload policy overrides of Default will fetch those payloads,
+ *   if they exist in the record, and skip data for any other payloads.
+ * - A base policy of Default with payload policy overrides of (Query | SkipData) will skip those
+ *   payloads, but still check if they exist, and will load any other payloads.
+ */
 class FCacheRecordPolicy
 {
 public:
-	/** Construct a cache record policy that uses the default policy for records and payloads. */
+	/** Construct a cache record policy that uses the default policy. */
 	FCacheRecordPolicy() = default;
 
-	/** Construct a cache record policy from a single policy for both records and payloads. */
+	/** Construct a cache record policy with a uniform policy for the record and every payload. */
 	inline FCacheRecordPolicy(ECachePolicy Policy)
-		: BasePolicy(Policy)
-		, RecordPolicy(Policy)
+		: RecordPolicy(Policy)
+		, DefaultPayloadPolicy(Policy)
 	{
 	}
 
-	/** Returns true if the record and every payload use the same policy. */
-	inline bool IsUniform() const { return !Shared && BasePolicy == RecordPolicy; }
+	/** Returns true if the record and every payload use the same cache policy. */
+	inline bool IsUniform() const { return !Shared && RecordPolicy == DefaultPayloadPolicy; }
 
 	/** Returns the cache policy to use for the record. */
 	inline ECachePolicy GetRecordPolicy() const { return RecordPolicy; }
@@ -150,7 +168,7 @@ public:
 	UE_API ECachePolicy GetPayloadPolicy(const FPayloadId& Id) const;
 
 	/** Returns the cache policy to use for payloads with no override. */
-	inline ECachePolicy GetDefaultPayloadPolicy() const { return BasePolicy; }
+	inline ECachePolicy GetDefaultPayloadPolicy() const { return DefaultPayloadPolicy; }
 
 	/** Returns the array of cache policy overrides for payloads, sorted by ID. */
 	inline TConstArrayView<FCachePayloadPolicy> GetPayloadPolicies() const
@@ -161,8 +179,8 @@ public:
 private:
 	friend class FCacheRecordPolicyBuilder;
 
-	ECachePolicy BasePolicy = ECachePolicy::Default;
 	ECachePolicy RecordPolicy = ECachePolicy::Default;
+	ECachePolicy DefaultPayloadPolicy = ECachePolicy::Default;
 	TRefCountPtr<const Private::ICacheRecordPolicyShared> Shared;
 };
 
@@ -170,10 +188,10 @@ private:
 class FCacheRecordPolicyBuilder
 {
 public:
-	/** Construct a cache record policy builder that uses the default policy as its base policy. */
+	/** Construct a policy builder that uses the default policy as its base policy. */
 	FCacheRecordPolicyBuilder() = default;
 
-	/** Construct a cache record policy builder from the base policy for the record and its payloads. */
+	/** Construct a policy builder that uses the provided policy for the record and payloads with no override. */
 	inline explicit FCacheRecordPolicyBuilder(ECachePolicy Policy)
 		: BasePolicy(Policy)
 	{
@@ -238,9 +256,9 @@ public:
 	 * Records may propagate into other cache stores, in accordance with the policy. A propagated
 	 * record may be a partial record, with some payloads missing data, depending on the policy.
 	 *
-	 * When payloads are required by the policy, but not available, the status must be Error, but
-	 * the cache store may provide a partial record with the available payloads, to be completed,
-	 * when possible, by another cache store.
+	 * When payloads are required by the policy, but not available, the status must be Error. The
+	 * cache store must produce a partial record with the available payloads when the policy flag
+	 * PartialOnError is set for the missing payloads.
 	 *
 	 * @param Keys         The keys identifying the cache records to query.
 	 * @param Context      A description of the request. An object path is typically sufficient.
