DDC: Reset the cache record builder in Build()

#rnx

diff --git a/Engine/Source/Developer/DerivedDataCache/Public/DerivedDataCacheRecord.h b/Engine/Source/Developer/DerivedDataCache/Public/DerivedDataCacheRecord.h
--- a/Engine/Source/Developer/DerivedDataCache/Public/DerivedDataCacheRecord.h
+++ b/Engine/Source/Developer/DerivedDataCache/Public/DerivedDataCacheRecord.h
@@ -5,6 +5,7 @@
 #include "DerivedDataCacheKey.h"
 #include "DerivedDataPayloadId.h"
 #include "Memory/SharedBuffer.h"
+#include "Misc/ScopeExit.h"
 #include "Templates/Function.h"
 #include "Templates/RefCounting.h"
 #include "Templates/UniquePtr.h"
@@ -219,6 +220,7 @@ public:
 	 */
 	inline FCacheRecord Build()
 	{
+		ON_SCOPE_EXIT { RecordBuilder = nullptr; };
 		return RecordBuilder->Build();
 	}
 
