DDC: Fixed fetching of partial records in the filesystem backend

#rnx

diff --git a/Engine/Source/Developer/DerivedDataCache/Private/FileSystemDerivedDataBackend.cpp b/Engine/Source/Developer/DerivedDataCache/Private/FileSystemDerivedDataBackend.cpp
--- a/Engine/Source/Developer/DerivedDataCache/Private/FileSystemDerivedDataBackend.cpp
+++ b/Engine/Source/Developer/DerivedDataCache/Private/FileSystemDerivedDataBackend.cpp
@@ -972,7 +972,8 @@ public:
 		EStatus RecordStatus = EStatus::Error;
 		for (const FCacheChunkRequest& Chunk : SortedChunks)
 		{
-			const bool bExistsOnly = EnumHasAnyFlags(Chunk.Policy, ECachePolicy::SkipValue);
+			constexpr ECachePolicy SkipFlag = ECachePolicy::SkipValue;
+			const bool bExistsOnly = EnumHasAnyFlags(Chunk.Policy, SkipFlag);
 			TRACE_CPUPROFILER_EVENT_SCOPE(FileSystemDDC_Get);
 			TRACE_COUNTER_INCREMENT(FileSystemDDC_Get);
 			COOK_STAT(auto Timer = bExistsOnly ? UsageStats.TimeProbablyExists() : UsageStats.TimeGet());
@@ -980,15 +981,15 @@ public:
 			{
 				const bool bAlwaysLoadInlineData = true;
 				FCacheRecordPolicyBuilder PolicyBuilder(ECachePolicy::None);
-				PolicyBuilder.AddPayloadPolicy(Chunk.Id, Chunk.Policy | ECachePolicy::SkipData);
+				const ECachePolicy RecordSkipFlags = bExistsOnly ? ECachePolicy::SkipData : ECachePolicy::None;
+				PolicyBuilder.AddPayloadPolicy(Chunk.Id, Chunk.Policy | RecordSkipFlags);
 				Record = GetCacheRecord(Chunk.Key, Context, PolicyBuilder.Build(), bAlwaysLoadInlineData, RecordStatus);
 			}
 			if (Record && RecordStatus == EStatus::Ok)
 			{
 				EStatus PayloadStatus = EStatus::Error;
-				const ECachePolicy Policy = Chunk.Policy | (ECachePolicy::SkipData & ~ECachePolicy::SkipValue);
 				const FPayload& RecordPayload = Record.Get().GetPayload(Chunk.Id);
-				if (FPayload Payload = GetCachePayload(Chunk.Key, Context, Policy, RecordPayload, PayloadStatus))
+				if (FPayload Payload = GetCachePayload(Chunk.Key, Context, Chunk.Policy, SkipFlag, RecordPayload, PayloadStatus))
 				{
 					const uint64 RawSize = FMath::Min(Payload.GetRawSize(), Chunk.RawSize);
 					UE_LOG(LogDerivedDataCache, Verbose, TEXT("%s: Cache hit for %s from '%.*s'"),
@@ -1333,13 +1334,14 @@ private:
 				Payload, CompressedHash, MoveTemp(CompressedData), OutStatus);
 		}
 
-		return GetCachePayload(Key, Context, PayloadPolicy, Payload, OutStatus);
+		return GetCachePayload(Key, Context, PayloadPolicy, SkipFlag, Payload, OutStatus);
 	}
 
 	FPayload GetCachePayload(
 		const FCacheKey& Key,
 		const FStringView Context,
 		const ECachePolicy Policy,
+		const ECachePolicy SkipFlag,
 		const FPayload& Payload,
 		EStatus& OutStatus) const
 	{
@@ -1351,7 +1353,7 @@ private:
 
 		TStringBuilder<256> Path;
 		BuildCachePayloadPath(Key, Payload.GetRawHash(), Path);
-		if (EnumHasAllFlags(Policy, ECachePolicy::SkipData))
+		if (EnumHasAllFlags(Policy, SkipFlag))
 		{
 			if (FileExists(Path))
 			{
