DDC: Limit async put priority to Highest to avoid blocking

#rnx

diff --git a/Engine/Source/Developer/DerivedDataCache/Private/HierarchicalDerivedDataBackend.h b/Engine/Source/Developer/DerivedDataCache/Private/HierarchicalDerivedDataBackend.h
--- a/Engine/Source/Developer/DerivedDataCache/Private/HierarchicalDerivedDataBackend.h
+++ b/Engine/Source/Developer/DerivedDataCache/Private/HierarchicalDerivedDataBackend.h
@@ -12,6 +12,7 @@
 #include "DerivedDataChunk.h"
 #include "DerivedDataPayload.h"
 #include "DerivedDataRequestOwner.h"
+#include "HAL/PlatformMath.h"
 #include "Misc/ScopeRWLock.h"
 #include "Misc/StringBuilder.h"
 #include "ProfilingDebugging/CookStats.h"
@@ -483,7 +484,7 @@ public:
 		IRequestOwner& Owner,
 		FOnCachePutComplete&& OnComplete) override
 	{
-		FRequestOwner AsyncOwner(Owner.GetPriority());
+		FRequestOwner AsyncOwner(FPlatformMath::Min(Owner.GetPriority(), EPriority::Highest));
 		FRequestBarrier AsyncBarrier(AsyncOwner);
 		AsyncOwner.KeepAlive();
 
@@ -574,7 +575,7 @@ public:
 						{
 							if (bFill)
 							{
-								FRequestOwner AsyncOwner(Owner.GetPriority());
+								FRequestOwner AsyncOwner(FPlatformMath::Min(Owner.GetPriority(), EPriority::Highest));
 								AsyncOwner.KeepAlive();
 								for (int32 FillCacheIndex = 0; FillCacheIndex < InnerBackends.Num(); ++FillCacheIndex)
 								{
