DDC: Added structured cache support to the throttle wrapper

#rnx

diff --git a/Engine/Source/Developer/DerivedDataCache/Private/DerivedDataBackendThrottleWrapper.h b/Engine/Source/Developer/DerivedDataCache/Private/DerivedDataBackendThrottleWrapper.h
--- a/Engine/Source/Developer/DerivedDataCache/Private/DerivedDataBackendThrottleWrapper.h
+++ b/Engine/Source/Developer/DerivedDataCache/Private/DerivedDataBackendThrottleWrapper.h
@@ -2,8 +2,12 @@
 
 #pragma once
 
-#include "CoreMinimal.h"
+#include "CoreTypes.h"
+#include "Algo/Find.h"
+#include "Algo/Transform.h"
 #include "DerivedDataBackendInterface.h"
+#include "DerivedDataCachePrivate.h"
+#include "DerivedDataCacheRecord.h"
 #include "DerivedDataCacheUsageStats.h"
 #include <atomic>
 
@@ -94,7 +98,7 @@ public:
 	 */
 	virtual bool CachedDataProbablyExists(const TCHAR* CacheKey) override
 	{
-		ThrottlingScope Scope(this);
+		FThrottlingScope Scope(this);
 
 		return InnerBackend->CachedDataProbablyExists(CacheKey);
 	}
@@ -107,7 +111,7 @@ public:
 	 */
 	virtual TBitArray<> CachedDataProbablyExistsBatch(TConstArrayView<FString> CacheKeys) override
 	{
-		ThrottlingScope Scope(this);
+		FThrottlingScope Scope(this);
 
 		return InnerBackend->CachedDataProbablyExistsBatch(CacheKeys);
 	}
@@ -121,7 +125,7 @@ public:
 	 */
 	virtual bool GetCachedData(const TCHAR* CacheKey, TArray<uint8>& OutData) override
 	{
-		ThrottlingScope Scope(this, [&OutData]() { return OutData.Num(); });
+		FThrottlingScope Scope(this, [&OutData]() { return OutData.Num(); });
 
 		return InnerBackend->GetCachedData(CacheKey, OutData);
 	}
@@ -135,14 +139,14 @@ public:
 	 */
 	virtual EPutStatus PutCachedData(const TCHAR* CacheKey, TArrayView<const uint8> InData, bool bPutEvenIfExists) override
 	{
-		ThrottlingScope Scope(this, [&InData]() { return InData.Num(); });
+		FThrottlingScope Scope(this, [&InData]() { return InData.Num(); });
 
 		return InnerBackend->PutCachedData(CacheKey, InData, bPutEvenIfExists);
 	}
 	
 	virtual void RemoveCachedData(const TCHAR* CacheKey, bool bTransient) override
 	{
-		ThrottlingScope Scope(this);
+		FThrottlingScope Scope(this);
 
 		return InnerBackend->RemoveCachedData(CacheKey, bTransient);
 	}
@@ -158,7 +162,7 @@ public:
 
 	virtual bool TryToPrefetch(TConstArrayView<FString> CacheKeys) override
 	{
-		ThrottlingScope Scope(this);
+		FThrottlingScope Scope(this);
 
 		return InnerBackend->TryToPrefetch(CacheKeys);
 	}
@@ -180,9 +184,28 @@ public:
 		IRequestOwner& Owner,
 		FOnCachePutComplete&& OnComplete) override
 	{
-		ThrottlingScope Scope(this);
+		struct FRecordSize
+		{
+			FCacheKey Key;
+			uint64 Size;
+		};
+		TArray<FRecordSize, TInlineAllocator<1>> RecordSizes;
+		RecordSizes.Reserve(Records.Num());
+		Algo::Transform(Records, RecordSizes, [](const FCacheRecord& Record) -> FRecordSize
+		{
+			return {Record.GetKey(), Private::GetCacheRecordCompressedSize(Record)};
+		});
 
-		InnerBackend->Put(Records, Context, Policy, Owner, MoveTemp(OnComplete));
+		InnerBackend->Put(Records, Context, Policy, Owner,
+			[this, RecordSizes = MoveTemp(RecordSizes), State = EnterThrottlingScope(), OnComplete = MoveTemp(OnComplete)](FCachePutCompleteParams&& Params)
+			{
+				const FRecordSize* Size = Algo::FindBy(RecordSizes, Params.Key, &FRecordSize::Key);
+				CloseThrottlingScope(State, FThrottlingState(this, Size ? Size->Size : 0));
+				if (OnComplete)
+				{
+					OnComplete(MoveTemp(Params));
+				}
+			});
 	}
 
 	virtual void Get(
@@ -192,9 +215,15 @@ public:
 		IRequestOwner& Owner,
 		FOnCacheGetComplete&& OnComplete) override
 	{
-		ThrottlingScope Scope(this);
-
-		InnerBackend->Get(Keys, Context, Policy, Owner, MoveTemp(OnComplete));
+		InnerBackend->Get(Keys, Context, Policy, Owner,
+			[this, State = EnterThrottlingScope(), OnComplete = MoveTemp(OnComplete)](FCacheGetCompleteParams&& Params)
+			{
+				CloseThrottlingScope(State, FThrottlingState(this, Private::GetCacheRecordCompressedSize(Params.Record)));
+				if (OnComplete)
+				{
+					OnComplete(MoveTemp(Params));
+				}
+			});
 	}
 
 	virtual void GetChunks(
@@ -203,32 +232,54 @@ public:
 		IRequestOwner& Owner,
 		FOnCacheGetChunkComplete&& OnComplete) override
 	{
-		ThrottlingScope Scope(this);
-
-		InnerBackend->GetChunks(Chunks, Context, Owner, MoveTemp(OnComplete));
+		InnerBackend->GetChunks(Chunks, Context, Owner,
+			[this, State = EnterThrottlingScope(), OnComplete = MoveTemp(OnComplete)](FCacheGetChunkCompleteParams&& Params)
+			{
+				CloseThrottlingScope(State, FThrottlingState(this, Params.RawData.GetSize()));
+				if (OnComplete)
+				{
+					OnComplete(MoveTemp(Params));
+				}
+			});
 	}
 
 private:
 
-	void EnterThrottlingScope(uint64& PreviousBytesTransferred)
+	struct FThrottlingState
+	{
+		double Time;
+		uint64 TotalBytesTransferred;
+
+		explicit FThrottlingState(FDerivedDataBackendThrottleWrapper* ThrottleWrapper)
+		: Time(FPlatformTime::Seconds())
+		, TotalBytesTransferred(ThrottleWrapper->TotalBytesTransferred.load(std::memory_order_relaxed))
+		{
+		}
+
+		explicit FThrottlingState(FDerivedDataBackendThrottleWrapper* ThrottleWrapper, uint64 BytesTransferred)
+		: Time(FPlatformTime::Seconds())
+		, TotalBytesTransferred(ThrottleWrapper->TotalBytesTransferred.fetch_add(BytesTransferred, std::memory_order_relaxed) + BytesTransferred)
+		{
+		}
+	};
+
+	FThrottlingState EnterThrottlingScope()
 	{
 		if (Latency > 0)
 		{
 			TRACE_CPUPROFILER_EVENT_SCOPE(ThrottlingLatency);
 			FPlatformProcess::Sleep(Latency);
 		}
-		PreviousBytesTransferred = TotalBytesTransferred.load();
+		return FThrottlingState(this);
 	}
 
-	void CloseThrottlingScope(double PreviousTime, uint64 PreviousBytesTransferred, TFunction<int64()> GetTransferredBytes = nullptr)
+	void CloseThrottlingScope(FThrottlingState PreviousState, FThrottlingState CurrentState)
 	{
-		if (MaxBytesPerSecond && GetTransferredBytes)
+		if (MaxBytesPerSecond)
 		{
-			uint64 BytesTransferred = GetTransferredBytes();
 			// Take into account any other transfer that might have happened during that time from any other thread so we have a global limit
-			uint64 NewBytesTransferred = TotalBytesTransferred += BytesTransferred;
-			double ActualTime = FPlatformTime::Seconds() - PreviousTime;
-			double ExpectedTime = double(NewBytesTransferred - PreviousBytesTransferred) / MaxBytesPerSecond;
+			const double ExpectedTime = double(CurrentState.TotalBytesTransferred - PreviousState.TotalBytesTransferred) / MaxBytesPerSecond;
+			const double ActualTime = CurrentState.Time - PreviousState.Time;
 			if (ExpectedTime > ActualTime)
 			{
 				TRACE_CPUPROFILER_EVENT_SCOPE(ThrottlingBandwidth);
@@ -237,24 +288,26 @@ private:
 		}
 	}
 
-	struct ThrottlingScope
+	struct FThrottlingScope
 	{
 		FDerivedDataBackendThrottleWrapper* ThrottleWrapper;
 		TFunction<int64()> GetTransferredBytes;
-		double PreviousTime;
-		uint64 PreviousBytesTransferred;
+		FThrottlingState State;
+
+		FThrottlingScope(const FThrottlingScope&) = delete;
+		FThrottlingScope& operator=(const FThrottlingScope&) = delete;
 
-		ThrottlingScope(FDerivedDataBackendThrottleWrapper* InThrottleWrapper, TFunction<int64()> InGetTransferredBytes = nullptr)
+		explicit FThrottlingScope(FDerivedDataBackendThrottleWrapper* InThrottleWrapper, TFunction<int64()> InGetTransferredBytes = nullptr)
 			: ThrottleWrapper(InThrottleWrapper)
 			, GetTransferredBytes(InGetTransferredBytes)
-			, PreviousTime(FPlatformTime::Seconds())
+			, State(ThrottleWrapper->EnterThrottlingScope())
 		{
-			ThrottleWrapper->EnterThrottlingScope(PreviousBytesTransferred);
 		}
 
-		~ThrottlingScope()
+		~FThrottlingScope()
 		{
-			ThrottleWrapper->CloseThrottlingScope(PreviousTime, PreviousBytesTransferred, GetTransferredBytes);
+			uint64 BytesTransferred = GetTransferredBytes ? GetTransferredBytes() : 0;
+			ThrottleWrapper->CloseThrottlingScope(State, FThrottlingState(ThrottleWrapper, BytesTransferred));
 		}
 	};
 
diff --git a/Engine/Source/Developer/DerivedDataCache/Private/DerivedDataCachePrivate.h b/Engine/Source/Developer/DerivedDataCache/Private/DerivedDataCachePrivate.h
--- a/Engine/Source/Developer/DerivedDataCache/Private/DerivedDataCachePrivate.h
+++ b/Engine/Source/Developer/DerivedDataCache/Private/DerivedDataCachePrivate.h
@@ -2,11 +2,13 @@
 
 #pragma once
 
+#include "CoreTypes.h"
 #include "Containers/StringFwd.h"
 #include "Misc/AsciiSet.h"
 
 class FDerivedDataCacheInterface;
 
+namespace UE::DerivedData { class FCacheRecord; }
 namespace UE::DerivedData { class ICache; }
 namespace UE::DerivedData { struct FCacheKey; }
 
@@ -16,6 +18,11 @@ namespace UE::DerivedData::Private
 // Implemented in DerivedDataCache.cpp
 ICache* CreateCache(FDerivedDataCacheInterface** OutLegacyCache);
 
+// Implemented in DerivedDataCacheRecord.cpp
+uint64 GetCacheRecordCompressedSize(const FCacheRecord& Record);
+uint64 GetCacheRecordTotalRawSize(const FCacheRecord& Record);
+uint64 GetCacheRecordRawSize(const FCacheRecord& Record);
+
 ///////////////////////////////////////////////////////////////////////////////////////////////////////////////////////
 
 template <typename CharType>
diff --git a/Engine/Source/Developer/DerivedDataCache/Private/DerivedDataCacheRecord.cpp b/Engine/Source/Developer/DerivedDataCache/Private/DerivedDataCacheRecord.cpp
--- a/Engine/Source/Developer/DerivedDataCache/Private/DerivedDataCacheRecord.cpp
+++ b/Engine/Source/Developer/DerivedDataCache/Private/DerivedDataCacheRecord.cpp
@@ -2,6 +2,7 @@
 
 #include "DerivedDataCacheRecord.h"
 
+#include "Algo/Accumulate.h"
 #include "DerivedDataCacheKey.h"
 #include "DerivedDataCachePrivate.h"
 #include "DerivedDataPayload.h"
@@ -404,3 +405,30 @@ FCacheRecordBuilder::FCacheRecordBuilder(const FCacheKey& Key)
 }
 
 } // UE::DerivedData
+
+///////////////////////////////////////////////////////////////////////////////////////////////////////////////////////
+
+namespace UE::DerivedData::Private
+{
+
+uint64 GetCacheRecordCompressedSize(const FCacheRecord& Record)
+{
+	const uint64 ValueSize = Record.GetValuePayload().GetData().GetCompressedSize();
+	return int64(Algo::TransformAccumulate(Record.GetAttachmentPayloads(),
+		[](const FPayload& Payload) { return Payload.GetData().GetCompressedSize(); }, ValueSize));
+}
+
+uint64 GetCacheRecordTotalRawSize(const FCacheRecord& Record)
+{
+	const uint64 ValueSize = Record.GetValuePayload().GetRawSize();
+	return int64(Algo::TransformAccumulate(Record.GetAttachmentPayloads(), &FPayload::GetRawSize, ValueSize));
+}
+
+uint64 GetCacheRecordRawSize(const FCacheRecord& Record)
+{
+	const uint64 ValueSize = Record.GetValuePayload().GetData().GetRawSize();
+	return int64(Algo::TransformAccumulate(Record.GetAttachmentPayloads(),
+		[](const FPayload& Payload) { return Payload.GetData().GetRawSize(); }, ValueSize));
+}
+
+} // UE::DerivedData::Private
