DDC: Reworked FBuildSession to allow filtering of payloads

- Build now takes optional inputs for consistency with BuildAction.
- BuildAction has been renamed to Build and works the same as Build except it starts from an action.
- Build now takes a FBuildPolicy, which is implicitly constructible from EBuildPolicy, and allows setting the policy by payload.
- FBuildContext and FBuildConfigContext now have policy masks that allow clearing flags across every payload in the build.
- EBuildPolicy has CacheStore split into CacheStoreOnQuery and CacheStoreOnBuild for more control of cache behavior.
- Builds now use consistent completion callback parameters, FBuildCompleteParams.
- Build types now have debugger visualization.

#rnx

diff --git a/Engine/Source/Developer/DerivedDataCache/Private/DerivedDataBuild.cpp b/Engine/Source/Developer/DerivedDataCache/Private/DerivedDataBuild.cpp
--- a/Engine/Source/Developer/DerivedDataCache/Private/DerivedDataBuild.cpp
+++ b/Engine/Source/Developer/DerivedDataCache/Private/DerivedDataBuild.cpp
@@ -2,6 +2,7 @@
 
 #include "DerivedDataBuild.h"
 
+#include "Algo/Accumulate.h"
 #include "DerivedDataBuildAction.h"
 #include "DerivedDataBuildDefinition.h"
 #include "DerivedDataBuildFunctionRegistry.h"
@@ -10,6 +11,7 @@
 #include "DerivedDataBuildPrivate.h"
 #include "DerivedDataBuildScheduler.h"
 #include "DerivedDataBuildSession.h"
+#include "DerivedDataBuildTypes.h"
 #include "DerivedDataBuildWorkerRegistry.h"
 #include "DerivedDataCache.h"
 #include "Misc/Guid.h"
@@ -128,3 +130,84 @@ IBuild* CreateBuild(ICache& Cache)
 }
 
 } // UE::DerivedData::Private
+
+namespace UE::DerivedData::Private { class FBuildPolicyShared; }
+
+namespace UE::DerivedData
+{
+
+class Private::FBuildPolicyShared final : public Private::IBuildPolicyShared
+{
+public:
+	inline void AddRef() const final
+	{
+		ReferenceCount.fetch_add(1, std::memory_order_relaxed);
+	}
+
+	inline void Release() const final
+	{
+		if (ReferenceCount.fetch_sub(1, std::memory_order_acq_rel) == 1)
+		{
+			delete this;
+		}
+	}
+
+	inline TConstArrayView<FBuildPayloadPolicy> GetPayloadPolicies() const final
+	{
+		return Payloads;
+	}
+
+	inline void AddPayloadPolicy(const FBuildPayloadPolicy& Policy) final
+	{
+		Payloads.Add(Policy);
+	}
+
+	inline void Build() final
+	{
+		Algo::SortBy(Payloads, &FBuildPayloadPolicy::Id);
+	}
+
+private:
+	TArray<FBuildPayloadPolicy, TInlineAllocator<14>> Payloads;
+	mutable std::atomic<uint32> ReferenceCount{0};
+};
+
+EBuildPolicy FBuildPolicy::GetPayloadPolicy(const FPayloadId& Id) const
+{
+	if (Shared)
+	{
+		if (TConstArrayView<FBuildPayloadPolicy> Payloads = Shared->GetPayloadPolicies(); !Payloads.IsEmpty())
+		{
+			if (int32 Index = Algo::BinarySearchBy(Payloads, Id, &FBuildPayloadPolicy::Id); Index != INDEX_NONE)
+			{
+				return Payloads[Index].Policy;
+			}
+		}
+	}
+	return DefaultPayloadPolicy;
+}
+
+void FBuildPolicyBuilder::AddPayloadPolicy(const FBuildPayloadPolicy& Policy)
+{
+	if (!Shared)
+	{
+		Shared = new Private::FBuildPolicyShared;
+	}
+	Shared->AddPayloadPolicy(Policy);
+}
+
+FBuildPolicy FBuildPolicyBuilder::Build()
+{
+	FBuildPolicy Policy(BasePolicy);
+	if (Shared)
+	{
+		Shared->Build();
+		const auto PolicyOr = [](EBuildPolicy A, EBuildPolicy B) { return A | (B & EBuildPolicy::Default); };
+		const TConstArrayView<FBuildPayloadPolicy> Payloads = Shared->GetPayloadPolicies();
+		Policy.CombinedPolicy = Algo::TransformAccumulate(Payloads, &FBuildPayloadPolicy::Policy, BasePolicy, PolicyOr);
+		Policy.Shared = MoveTemp(Shared);
+	}
+	return Policy;
+}
+
+} // UE::DerivedData
diff --git a/Engine/Source/Developer/DerivedDataCache/Private/DerivedDataBuildJob.cpp b/Engine/Source/Developer/DerivedDataCache/Private/DerivedDataBuildJob.cpp
--- a/Engine/Source/Developer/DerivedDataCache/Private/DerivedDataBuildJob.cpp
+++ b/Engine/Source/Developer/DerivedDataCache/Private/DerivedDataBuildJob.cpp
@@ -94,11 +94,27 @@ class FBuildJob final : public IBuildJob
 {
 public:
 	/** Resolve the key to a definition, then build like the definition constructor. */
-	FBuildJob(const FBuildJobCreateParams& Params, const FBuildKey& Key, FOnBuildJobComplete&& OnComplete);
+	FBuildJob(
+		const FBuildJobCreateParams& Params,
+		const FBuildKey& Key,
+		const FBuildPolicy& Policy,
+		FOnBuildComplete&& OnComplete);
+
 	/** Resolve the definition to an action, then build like the action constructor. */
-	FBuildJob(const FBuildJobCreateParams& Params, const FBuildDefinition& Definition, FOnBuildJobComplete&& OnComplete);
+	FBuildJob(
+		const FBuildJobCreateParams& Params,
+		const FBuildDefinition& Definition,
+		const FOptionalBuildInputs& Inputs,
+		const FBuildPolicy& Policy,
+		FOnBuildComplete&& OnComplete);
+
 	/** Query the cache, attempt remote execution, resolve inputs, fall back to local execution, store to the cache. */
-	FBuildJob(const FBuildJobCreateParams& Params, const FBuildAction& Action, const FOptionalBuildInputs& Inputs, FOnBuildJobComplete&& OnComplete);
+	FBuildJob(
+		const FBuildJobCreateParams& Params,
+		const FBuildAction& Action,
+		const FOptionalBuildInputs& Inputs,
+		const FBuildPolicy& Policy,
+		FOnBuildComplete&& OnComplete);
 
 	/** Destroy the job, which must be complete or not started. */
 	~FBuildJob();
@@ -114,7 +130,7 @@ public:
 	void StepExecution() final;
 
 private:
-	FBuildJob(const FBuildJobCreateParams& Params, FStringView Name, FStringView FunctionName, FOnBuildJobComplete&& OnComplete);
+	FBuildJob(const FBuildJobCreateParams& Params, FStringView Name, FStringView FunctionName, FOnBuildComplete&& OnComplete);
 
 	void BeginJob();
 	void EndJob();
@@ -185,10 +201,10 @@ private:
 	/** Next state for the job. Used to handle re-entrant calls to AdvanceToState. */
 	EBuildJobState NextState{EBuildJobState::NotStarted};
 
+	/** Policy requested and/or configured for the job. */
+	FBuildPolicy BuildPolicy{EBuildPolicy::None};
 	/** Status flags that are added to as the job moves through its states. */
 	EBuildStatus BuildStatus{EBuildStatus::None};
-	/** Policy requested and/or configured for the job. */
-	EBuildPolicy BuildPolicy{EBuildPolicy::None};
 
 	/** Available in [ResolveKey, Complete] for jobs created from a key or definition. */
 	FBuildKey DefinitionKey;
@@ -223,7 +239,7 @@ private:
 	/** Worker executor to use for remote execution. Available in [ExecuteRemote, ExecuteRemoteRetryWait]. */
 	IBuildWorkerExecutor* WorkerExecutor{};
 	/** Invoked exactly once when the output is complete or when the job fails. */
-	FOnBuildJobComplete OnComplete;
+	FOnBuildComplete OnComplete;
 
 	/** Keys for missing inputs. */
 	TArray<FString> MissingInputs;
@@ -271,10 +287,9 @@ FBuildJob::FBuildJob(
 	const FBuildJobCreateParams& Params,
 	FStringView InName,
 	FStringView InFunctionName,
-	FOnBuildJobComplete&& InOnComplete)
+	FOnBuildComplete&& InOnComplete)
 	: Name(InName)
 	, FunctionName(InFunctionName)
-	, BuildPolicy(Params.Policy)
 	, OutputBuilder(Params.BuildSystem.CreateOutput(Name, FunctionName))
 	, Owner(Params.Owner)
 	, Scheduler(Params.Scheduler)
@@ -288,9 +303,11 @@ FBuildJob::FBuildJob(
 FBuildJob::FBuildJob(
 	const FBuildJobCreateParams& Params,
 	const FBuildKey& InKey,
-	FOnBuildJobComplete&& InOnComplete)
+	const FBuildPolicy& InPolicy,
+	FOnBuildComplete&& InOnComplete)
 	: FBuildJob(Params, WriteToString<64>(TEXT("Resolve: "_SV), InKey), TEXT("Unknown"_SV), MoveTemp(InOnComplete))
 {
+	BuildPolicy = InPolicy;
 	DefinitionKey = InKey;
 	FRequestBarrier Barrier(Owner, ERequestBarrierFlags::Priority);
 	AdvanceToState(EBuildJobState::ResolveKey);
@@ -299,11 +316,15 @@ FBuildJob::FBuildJob(
 FBuildJob::FBuildJob(
 	const FBuildJobCreateParams& Params,
 	const FBuildDefinition& InDefinition,
-	FOnBuildJobComplete&& InOnComplete)
+	const FOptionalBuildInputs& InInputs,
+	const FBuildPolicy& InPolicy,
+	FOnBuildComplete&& InOnComplete)
 	: FBuildJob(Params, InDefinition.GetName(), InDefinition.GetFunction(), MoveTemp(InOnComplete))
 {
+	BuildPolicy = InPolicy;
 	DefinitionKey = InDefinition.GetKey();
 	Definition = InDefinition;
+	Inputs = InInputs;
 	FRequestBarrier Barrier(Owner, ERequestBarrierFlags::Priority);
 	AdvanceToState(EBuildJobState::ResolveKey);
 }
@@ -312,9 +333,11 @@ FBuildJob::FBuildJob(
 	const FBuildJobCreateParams& Params,
 	const FBuildAction& InAction,
 	const FOptionalBuildInputs& InInputs,
-	FOnBuildJobComplete&& InOnComplete)
+	const FBuildPolicy& InPolicy,
+	FOnBuildComplete&& InOnComplete)
 	: FBuildJob(Params, InAction.GetName(), InAction.GetFunction(), MoveTemp(InOnComplete))
 {
+	BuildPolicy = InPolicy;
 	Action = InAction;
 	Inputs = InInputs;
 	FRequestBarrier Barrier(Owner, ERequestBarrierFlags::Priority);
@@ -397,13 +420,12 @@ void FBuildJob::CreateContext()
 	else
 	{
 		const FCacheKey CacheKey{FCacheBucket(FunctionName), Action.Get().GetKey().Hash};
-		Context = new FBuildJobContext(*this, CacheKey, *Function, OutputBuilder, BuildPolicy);
+		Context = new FBuildJobContext(*this, CacheKey, *Function, OutputBuilder);
 		Action.Get().IterateConstants([this](FStringView Key, FCbObject&& Value)
 		{
 			Context->AddConstant(Key, MoveTemp(Value));
 		});
 		Function->Configure(*Context);
-		BuildPolicy = Context->GetBuildPolicy();
 		EnumAddFlags(BuildStatus, ShouldExportBuild() ? EBuildStatus::BuildTryExport : EBuildStatus::None);
 		EnumAddFlags(BuildStatus, EBuildStatus::CacheKey);
 	}
@@ -433,11 +455,45 @@ void FBuildJob::CreateContext()
 
 ///////////////////////////////////////////////////////////////////////////////////////////////////////////////////////
 
+static ECachePolicy MakeCacheQueryPolicy(EBuildPolicy BuildPolicy, const FBuildJobContext& Context)
+{
+	BuildPolicy &= Context.GetBuildPolicyMask();
+	ECachePolicy CachePolicy = ECachePolicy::None;
+	if (EnumHasAnyFlags(BuildPolicy, EBuildPolicy::CacheQuery))
+	{
+		EnumAddFlags(CachePolicy, ECachePolicy::Query);
+	}
+	if (EnumHasAnyFlags(BuildPolicy, EBuildPolicy::CacheStoreOnQuery))
+	{
+		EnumAddFlags(CachePolicy, ECachePolicy::Store);
+	}
+	if (EnumHasAnyFlags(BuildPolicy, EBuildPolicy::CacheKeepAlive))
+	{
+		EnumAddFlags(CachePolicy, ECachePolicy::KeepAlive);
+	}
+	if (EnumHasAnyFlags(BuildPolicy, EBuildPolicy::SkipData))
+	{
+		EnumAddFlags(CachePolicy, ECachePolicy::SkipAttachments);
+	}
+	CachePolicy &= Context.GetCachePolicyMask();
+	return CachePolicy;
+}
+
+static FCacheRecordPolicy MakeCacheRecordQueryPolicy(const FBuildPolicy& BuildPolicy, const FBuildJobContext& Context)
+{
+	FCacheRecordPolicyBuilder Builder(MakeCacheQueryPolicy(BuildPolicy.GetDefaultPayloadPolicy(), Context));
+	for (const FBuildPayloadPolicy& PayloadPolicy : BuildPolicy.GetPayloadPolicies())
+	{
+		Builder.AddPayloadPolicy(PayloadPolicy.Id, MakeCacheQueryPolicy(PayloadPolicy.Policy, Context));
+	}
+	return Builder.Build();
+}
+
 void FBuildJob::EnterCacheQuery()
 {
-	ECachePolicy CachePolicy = Context ? Context->GetCachePolicy() : ECachePolicy::None;
-	if (!EnumHasAnyFlags(CachePolicy, ECachePolicy::Query) ||
-		!EnumHasAnyFlags(BuildPolicy, EBuildPolicy::CacheQuery))
+	if (!Context ||
+		!EnumHasAnyFlags(Context->GetCachePolicyMask(), ECachePolicy::Query) ||
+		!EnumHasAnyFlags(Context->GetBuildPolicyMask() & BuildPolicy.GetCombinedPolicy(), EBuildPolicy::CacheQuery))
 	{
 		return AdvanceToState(EBuildJobState::ExecuteRemote);
 	}
@@ -446,13 +502,8 @@ void FBuildJob::EnterCacheQuery()
 void FBuildJob::BeginCacheQuery()
 {
 	TRACE_CPUPROFILER_EVENT_SCOPE(FBuildJob::CacheQuery);
-	ECachePolicy CachePolicy = Context->GetCachePolicy();
-	if (EnumHasAnyFlags(BuildPolicy, EBuildPolicy::SkipData))
-	{
-		EnumAddFlags(CachePolicy, ECachePolicy::SkipAttachments);
-	}
 	EnumAddFlags(BuildStatus, EBuildStatus::CacheQuery);
-	Cache.Get({Context->GetCacheKey()}, Name, CachePolicy, Owner,
+	Cache.Get({Context->GetCacheKey()}, Name, MakeCacheRecordQueryPolicy(BuildPolicy, *Context), Owner,
 		[this](FCacheGetCompleteParams&& Params) { EndCacheQuery(MoveTemp(Params)); });
 }
 
@@ -470,7 +521,7 @@ void FBuildJob::EndCacheQuery(FCacheGetCompleteParams&& Params)
 			return SetOutputNoCheck(MoveTemp(CacheOutput).Get());
 		}
 	}
-	if (!EnumHasAnyFlags(BuildPolicy, EBuildPolicy::Build))
+	if (!EnumHasAnyFlags(BuildPolicy.GetCombinedPolicy() & Context->GetBuildPolicyMask(), EBuildPolicy::Build))
 	{
 		return CompleteWithError(TEXT("Failed to fetch from the cache and build policy does not allow execution."));
 	}
@@ -479,11 +530,31 @@ void FBuildJob::EndCacheQuery(FCacheGetCompleteParams&& Params)
 
 ///////////////////////////////////////////////////////////////////////////////////////////////////////////////////////
 
+static ECachePolicy MakeCacheStorePolicy(EBuildPolicy BuildPolicy, const FBuildJobContext& Context)
+{
+	BuildPolicy &= Context.GetBuildPolicyMask();
+	ECachePolicy CachePolicy = ECachePolicy::None;
+	if (EnumHasAnyFlags(BuildPolicy, EBuildPolicy::CacheQuery))
+	{
+		EnumAddFlags(CachePolicy, ECachePolicy::Query);
+	}
+	if (EnumHasAnyFlags(BuildPolicy, EBuildPolicy::CacheStoreOnBuild))
+	{
+		EnumAddFlags(CachePolicy, ECachePolicy::Store);
+	}
+	if (EnumHasAnyFlags(BuildPolicy, EBuildPolicy::CacheKeepAlive))
+	{
+		EnumAddFlags(CachePolicy, ECachePolicy::KeepAlive);
+	}
+	CachePolicy &= Context.GetCachePolicyMask();
+	return CachePolicy;
+}
+
 void FBuildJob::EnterCacheStore()
 {
-	ECachePolicy CachePolicy = Context ? Context->GetCachePolicy() : ECachePolicy::None;
-	if (!EnumHasAnyFlags(CachePolicy, ECachePolicy::Store) ||
-		!EnumHasAnyFlags(BuildPolicy, EBuildPolicy::CacheStore) ||
+	if (!Context ||
+		!EnumHasAnyFlags(Context->GetCachePolicyMask(), ECachePolicy::Store) ||
+		!EnumHasAnyFlags(Context->GetBuildPolicyMask() & BuildPolicy.GetCombinedPolicy(), EBuildPolicy::CacheStoreOnBuild) ||
 		EnumHasAnyFlags(BuildStatus, EBuildStatus::CacheQueryHit) ||
 		Output.Get().HasError())
 	{
@@ -494,10 +565,10 @@ void FBuildJob::EnterCacheStore()
 void FBuildJob::BeginCacheStore()
 {
 	TRACE_CPUPROFILER_EVENT_SCOPE(FBuildJob::CacheStore);
+	EnumAddFlags(BuildStatus, EBuildStatus::CacheStore);
 	FCacheRecordBuilder RecordBuilder(Context->GetCacheKey());
 	Output.Get().Save(RecordBuilder);
-	EnumAddFlags(BuildStatus, EBuildStatus::CacheStore);
-	Cache.Put({RecordBuilder.Build()}, Name, Context->GetCachePolicy(), Owner,
+	Cache.Put({RecordBuilder.Build()}, Name, MakeCacheStorePolicy(BuildPolicy.GetCombinedPolicy(), *Context), Owner,
 		[this](FCachePutCompleteParams&& Params) { EndCacheStore(MoveTemp(Params)); });
 }
 
@@ -615,7 +686,7 @@ void FBuildJob::EnterResolveInputData()
 			return AdvanceToState(EBuildJobState::ExecuteLocal);
 		}
 		// Skip resolving input data if it will not be used anyway.
-		if (!EnumHasAnyFlags(BuildPolicy, EBuildPolicy::BuildLocal) &&
+		if (!EnumHasAnyFlags(BuildPolicy.GetCombinedPolicy() & Context->GetBuildPolicyMask(), EBuildPolicy::BuildLocal) &&
 			!EnumHasAnyFlags(BuildStatus, EBuildStatus::BuildTryExport))
 		{
 			return AdvanceToState(EBuildJobState::ExecuteLocal);
@@ -693,7 +764,7 @@ void FBuildJob::EndResolveInputData(FBuildInputDataResolvedParams&& Params)
 
 void FBuildJob::EnterExecuteRemote()
 {
-	if (!EnumHasAnyFlags(BuildPolicy, EBuildPolicy::BuildRemote) ||
+	if (!EnumHasAnyFlags(BuildPolicy.GetCombinedPolicy() & Context->GetBuildPolicyMask(), EBuildPolicy::BuildRemote) ||
 		EnumHasAnyFlags(BuildStatus, EBuildStatus::BuildTryExport))
 	{
 		return AdvanceToState(EBuildJobState::ResolveInputData);
@@ -720,7 +791,7 @@ void FBuildJob::BeginExecuteRemote()
 	checkf(Worker && WorkerExecutor, TEXT("Job requires a worker in state %s for build of '%s' by %s."),
 		LexToString(State), *Name, *FunctionName);
 	EnumAddFlags(BuildStatus, EBuildStatus::BuildTryRemote);
-	WorkerExecutor->BuildAction(Action.Get(), Inputs, *Worker, BuildSystem, BuildPolicy, Owner,
+	WorkerExecutor->BuildAction(Action.Get(), Inputs, *Worker, BuildSystem, BuildPolicy.GetCombinedPolicy(), Owner,
 		[this](FBuildWorkerActionCompleteParams&& Params) { EndExecuteRemote(MoveTemp(Params)); });
 }
 
@@ -778,9 +849,10 @@ void FBuildJob::EnterExecuteLocal()
 		ExportBuild();
 		EnumAddFlags(BuildStatus, EBuildStatus::BuildExport);
 	}
-	if (!EnumHasAnyFlags(BuildPolicy, EBuildPolicy::BuildLocal))
+	const EBuildPolicy LocalBuildPolicy = BuildPolicy.GetCombinedPolicy() & Context->GetBuildPolicyMask();
+	if (!EnumHasAnyFlags(LocalBuildPolicy, EBuildPolicy::BuildLocal))
 	{
-		if (!EnumHasAnyFlags(BuildPolicy, EBuildPolicy::BuildRemote))
+		if (!EnumHasAnyFlags(LocalBuildPolicy, EBuildPolicy::BuildRemote))
 		{
 			return AdvanceToState(EBuildJobState::Complete);
 		}
@@ -915,7 +987,7 @@ void FBuildJob::SetOutputNoCheck(FBuildOutput&& InOutput, EBuildJobState NewStat
 	{
 		const FCacheKey& CacheKey = Context ? Context->GetCacheKey() : FCacheKey::Empty;
 		const EStatus Status = Owner.IsCanceled() ? EStatus::Canceled : Output.Get().HasError() ? EStatus::Error : EStatus::Ok;
-		OnComplete({*this, CacheKey, FBuildOutput(Output.Get()), BuildStatus, Status});
+		OnComplete({CacheKey, FBuildOutput(Output.Get()), BuildStatus, Status});
 		OnComplete = nullptr;
 	}
 
@@ -1176,19 +1248,33 @@ TArray<FName> FBuildJob::ParseExportBuildTypes(bool& bOutExportAll)
 
 ///////////////////////////////////////////////////////////////////////////////////////////////////////////////////////
 
-void CreateBuildJob(const FBuildJobCreateParams& Params, const FBuildKey& Key, FOnBuildJobComplete&& OnComplete)
+void CreateBuildJob(
+	const FBuildJobCreateParams& Params,
+	const FBuildKey& Key,
+	const FBuildPolicy& Policy,
+	FOnBuildComplete&& OnComplete)
 {
-	new FBuildJob(Params, Key, MoveTemp(OnComplete));
+	new FBuildJob(Params, Key, Policy, MoveTemp(OnComplete));
 }
 
-void CreateBuildJob(const FBuildJobCreateParams& Params, const FBuildDefinition& Definition, FOnBuildJobComplete&& OnComplete)
+void CreateBuildJob(
+	const FBuildJobCreateParams& Params,
+	const FBuildDefinition& Definition,
+	const FOptionalBuildInputs& Inputs,
+	const FBuildPolicy& Policy,
+	FOnBuildComplete&& OnComplete)
 {
-	new FBuildJob(Params, Definition, MoveTemp(OnComplete));
+	new FBuildJob(Params, Definition, Inputs, Policy, MoveTemp(OnComplete));
 }
 
-void CreateBuildJob(const FBuildJobCreateParams& Params, const FBuildAction& Action, const FOptionalBuildInputs& Inputs, FOnBuildJobComplete&& OnComplete)
+void CreateBuildJob(
+	const FBuildJobCreateParams& Params,
+	const FBuildAction& Action,
+	const FOptionalBuildInputs& Inputs,
+	const FBuildPolicy& Policy,
+	FOnBuildComplete&& OnComplete)
 {
-	new FBuildJob(Params, Action, Inputs, MoveTemp(OnComplete));
+	new FBuildJob(Params, Action, Inputs, Policy, MoveTemp(OnComplete));
 }
 
 } // UE::DerivedData::Private
diff --git a/Engine/Source/Developer/DerivedDataCache/Private/DerivedDataBuildJobContext.cpp b/Engine/Source/Developer/DerivedDataCache/Private/DerivedDataBuildJobContext.cpp
--- a/Engine/Source/Developer/DerivedDataCache/Private/DerivedDataBuildJobContext.cpp
+++ b/Engine/Source/Developer/DerivedDataCache/Private/DerivedDataBuildJobContext.cpp
@@ -21,19 +21,14 @@ FBuildJobContext::FBuildJobContext(
 	IBuildJob& InJob,
 	const FCacheKey& InCacheKey,
 	const IBuildFunction& InFunction,
-	FBuildOutputBuilder& InOutputBuilder,
-	EBuildPolicy InBuildPolicy)
+	FBuildOutputBuilder& InOutputBuilder)
 	: Job(InJob)
 	, CacheKey(InCacheKey)
 	, Function(InFunction)
 	, OutputBuilder(InOutputBuilder)
-	, CachePolicy(ECachePolicy::Default)
-	, BuildPolicy(InBuildPolicy)
+	, CachePolicyMask(~ECachePolicy::None)
+	, BuildPolicyMask(~EBuildPolicy::None)
 {
-	if (EnumHasAnyFlags(BuildPolicy, EBuildPolicy::CacheKeepAlive))
-	{
-		EnumAddFlags(CachePolicy, ECachePolicy::KeepAlive);
-	}
 }
 
 FStringView FBuildJobContext::GetName() const
@@ -160,34 +155,37 @@ void FBuildJobContext::SetCacheBucket(FCacheBucket Bucket)
 	CacheKey.Bucket = Bucket;
 }
 
-void FBuildJobContext::SetCachePolicy(ECachePolicy Policy)
+void FBuildJobContext::SetCachePolicyMask(ECachePolicy Policy)
 {
-	checkf(!EnumHasAnyFlags(CachePolicy ^ Policy, ECachePolicy::SkipData),
-		TEXT("SkipData flags may not be modified on the cache policy for build of '%s' by %s. ")
+	checkf(EnumHasAllFlags(Policy, ECachePolicy::SkipData),
+		TEXT("SkipData flags may not be masked out on the cache policy for build of '%s' by %s. ")
 		TEXT("Flags for skipping data may be set indirectly through EBuildPolicy."),
 		*WriteToString<128>(Job.GetName()), *WriteToString<32>(Job.GetFunction()));
-	checkf(!EnumHasAnyFlags(CachePolicy ^ Policy, ECachePolicy::KeepAlive),
-		TEXT("KeepAlive flag may not be modified on the cache policy for build of '%s' by %s. ")
+	checkf(EnumHasAllFlags(Policy, ECachePolicy::KeepAlive),
+		TEXT("KeepAlive flag may not be masked out on the cache policy for build of '%s' by %s. ")
 		TEXT("Flags for cache record lifetime may be set indirectly through EBuildPolicy."),
 		*WriteToString<128>(Job.GetName()), *WriteToString<32>(Job.GetFunction()));
-	CachePolicy = Policy;
+	checkf(EnumHasAllFlags(Policy, ECachePolicy::PartialOnError),
+		TEXT("PartialOnError flag may not be masked out on the cache policy for build of '%s' by %s."),
+		*WriteToString<128>(Job.GetName()), *WriteToString<32>(Job.GetFunction()));
+	CachePolicyMask = Policy;
 }
 
-void FBuildJobContext::SetBuildPolicy(EBuildPolicy Policy)
+void FBuildJobContext::SetBuildPolicyMask(EBuildPolicy Policy)
 {
-	checkf(!EnumHasAnyFlags(BuildPolicy ^ Policy, EBuildPolicy::Cache),
-		TEXT("Cache flags may not be modified on the build policy for build of '%s' by %s. ")
+	checkf(EnumHasAllFlags(Policy, EBuildPolicy::Cache),
+		TEXT("Cache flags may not be masked out on the build policy for build of '%s' by %s. ")
 		TEXT("Flags for modifying cache operations may be set through ECachePolicy."),
 		*WriteToString<128>(Job.GetName()), *WriteToString<32>(Job.GetFunction()));
-	checkf(!EnumHasAnyFlags(BuildPolicy ^ Policy, EBuildPolicy::CacheKeepAlive),
-		TEXT("CacheKeepAlive flag may not be modified on the build policy for build of '%s' by %s. ")
+	checkf(EnumHasAllFlags(Policy, EBuildPolicy::CacheKeepAlive),
+		TEXT("CacheKeepAlive flag may not be masked out on the build policy for build of '%s' by %s. ")
 		TEXT("Flags for cache record lifetime may only be set through the build session."),
 		*WriteToString<128>(Job.GetName()), *WriteToString<32>(Job.GetFunction()));
-	checkf(!EnumHasAnyFlags(BuildPolicy ^ Policy, EBuildPolicy::SkipData),
-		TEXT("SkipData flags may not be modified on the build policy for build of '%s' by %s. ")
+	checkf(EnumHasAllFlags(Policy, EBuildPolicy::SkipData),
+		TEXT("SkipData flags may not be masked out on the build policy for build of '%s' by %s. ")
 		TEXT("Flags for skipping the data may only be set through the build session."),
 		*WriteToString<128>(Job.GetName()), *WriteToString<32>(Job.GetFunction()));
-	BuildPolicy = Policy;
+	BuildPolicyMask = Policy;
 }
 
 void FBuildJobContext::SetPriority(EPriority Priority)
diff --git a/Engine/Source/Developer/DerivedDataCache/Private/DerivedDataBuildJobContext.h b/Engine/Source/Developer/DerivedDataCache/Private/DerivedDataBuildJobContext.h
--- a/Engine/Source/Developer/DerivedDataCache/Private/DerivedDataBuildJobContext.h
+++ b/Engine/Source/Developer/DerivedDataCache/Private/DerivedDataBuildJobContext.h
@@ -32,8 +32,7 @@ public:
 		IBuildJob& Job,
 		const FCacheKey& CacheKey,
 		const IBuildFunction& Function,
-		FBuildOutputBuilder& OutputBuilder,
-		EBuildPolicy BuildPolicy);
+		FBuildOutputBuilder& OutputBuilder);
 
 	void BeginBuild(IRequestOwner& Owner, TUniqueFunction<void ()>&& OnEndBuild);
 
@@ -41,8 +40,8 @@ public:
 
 	inline const FCacheKey& GetCacheKey() const { return CacheKey; }
 
-	inline ECachePolicy GetCachePolicy() const final { return CachePolicy; }
-	inline EBuildPolicy GetBuildPolicy() const final { return BuildPolicy; }
+	inline ECachePolicy GetCachePolicyMask() const final { return CachePolicyMask; }
+	inline EBuildPolicy GetBuildPolicyMask() const final { return BuildPolicyMask; }
 	inline uint64 GetRequiredMemory() const { return RequiredMemory; }
 	inline bool ShouldCheckDeterministicOutput() const { return bDeterministicOutputCheck; }
 
@@ -66,8 +65,8 @@ private:
 	void EndAsyncBuild() final;
 
 	void SetCacheBucket(FCacheBucket Bucket) final;
-	void SetCachePolicy(ECachePolicy Policy) final;
-	void SetBuildPolicy(EBuildPolicy Policy) final;
+	void SetCachePolicyMask(ECachePolicy Policy) final;
+	void SetBuildPolicyMask(EBuildPolicy Policy) final;
 	void SetRequiredMemory(uint64 InRequiredMemory) final { RequiredMemory = InRequiredMemory; }
 	void SkipDeterministicOutputCheck() final { bDeterministicOutputCheck = false; }
 
@@ -86,8 +85,8 @@ private:
 	TUniqueFunction<void ()> OnEndBuild;
 	IRequestOwner* Owner{};
 	uint64 RequiredMemory{};
-	ECachePolicy CachePolicy;
-	EBuildPolicy BuildPolicy;
+	ECachePolicy CachePolicyMask;
+	EBuildPolicy BuildPolicyMask;
 	bool bIsAsyncBuild{false};
 	bool bIsAsyncBuildComplete{false};
 	bool bDeterministicOutputCheck{true};
diff --git a/Engine/Source/Developer/DerivedDataCache/Private/DerivedDataBuildPrivate.h b/Engine/Source/Developer/DerivedDataCache/Private/DerivedDataBuildPrivate.h
--- a/Engine/Source/Developer/DerivedDataCache/Private/DerivedDataBuildPrivate.h
+++ b/Engine/Source/Developer/DerivedDataCache/Private/DerivedDataBuildPrivate.h
@@ -3,6 +3,7 @@
 #pragma once
 
 #include "Containers/StringFwd.h"
+#include "DerivedDataBuildTypes.h"
 #include "Logging/LogMacros.h"
 #include "Misc/AsciiSet.h"
 
@@ -16,6 +17,7 @@ namespace UE::DerivedData { class FBuildDefinition; }
 namespace UE::DerivedData { class FBuildDefinitionBuilder; }
 namespace UE::DerivedData { class FBuildInputsBuilder; }
 namespace UE::DerivedData { class FBuildOutputBuilder; }
+namespace UE::DerivedData { class FBuildPolicy; }
 namespace UE::DerivedData { class FBuildSession; }
 namespace UE::DerivedData { class FOptionalBuildInputs; }
 namespace UE::DerivedData { class IBuild; }
@@ -25,10 +27,7 @@ namespace UE::DerivedData { class IBuildScheduler; }
 namespace UE::DerivedData { class IBuildWorkerRegistry; }
 namespace UE::DerivedData { class ICache; }
 namespace UE::DerivedData { class IRequestOwner; }
-namespace UE::DerivedData { struct FBuildJobCompleteParams; }
 namespace UE::DerivedData { struct FBuildKey; }
-namespace UE::DerivedData { enum class EBuildPolicy : uint32; }
-namespace UE::DerivedData { using FOnBuildJobComplete = TUniqueFunction<void (FBuildJobCompleteParams&& Params)>; }
 
 namespace UE::DerivedData::Private
 {
@@ -75,11 +74,24 @@ struct FBuildJobCreateParams
 	IBuildScheduler& Scheduler;
 	IBuildInputResolver* InputResolver{};
 	IRequestOwner& Owner;
-	EBuildPolicy Policy{};
 };
-void CreateBuildJob(const FBuildJobCreateParams& Params, const FBuildKey& Key, FOnBuildJobComplete&& OnComplete);
-void CreateBuildJob(const FBuildJobCreateParams& Params, const FBuildDefinition& Definition, FOnBuildJobComplete&& OnComplete);
-void CreateBuildJob(const FBuildJobCreateParams& Params, const FBuildAction& Action, const FOptionalBuildInputs& Inputs, FOnBuildJobComplete&& OnComplete);
+void CreateBuildJob(
+	const FBuildJobCreateParams& Params,
+	const FBuildKey& Key,
+	const FBuildPolicy& Policy,
+	FOnBuildComplete&& OnComplete);
+void CreateBuildJob(
+	const FBuildJobCreateParams& Params,
+	const FBuildDefinition& Definition,
+	const FOptionalBuildInputs& Inputs,
+	const FBuildPolicy& Policy,
+	FOnBuildComplete&& OnComplete);
+void CreateBuildJob(
+	const FBuildJobCreateParams& Params,
+	const FBuildAction& Action,
+	const FOptionalBuildInputs& Inputs,
+	const FBuildPolicy& Policy,
+	FOnBuildComplete&& OnComplete);
 
 ///////////////////////////////////////////////////////////////////////////////////////////////////////////////////////
 
diff --git a/Engine/Source/Developer/DerivedDataCache/Private/DerivedDataBuildSession.cpp b/Engine/Source/Developer/DerivedDataCache/Private/DerivedDataBuildSession.cpp
--- a/Engine/Source/Developer/DerivedDataCache/Private/DerivedDataBuildSession.cpp
+++ b/Engine/Source/Developer/DerivedDataCache/Private/DerivedDataBuildSession.cpp
@@ -34,16 +34,17 @@ public:
 
 	void Build(
 		const FBuildDefinition& Definition,
-		EBuildPolicy Policy,
+		const FOptionalBuildInputs& Inputs,
+		const FBuildPolicy& Policy,
 		IRequestOwner& Owner,
 		FOnBuildComplete&& OnComplete) final;
 
-	void BuildAction(
+	void Build(
 		const FBuildAction& Action,
 		const FOptionalBuildInputs& Inputs,
-		EBuildPolicy Policy,
+		const FBuildPolicy& Policy,
 		IRequestOwner& Owner,
-		FOnBuildActionComplete&& OnComplete) final;
+		FOnBuildComplete&& OnComplete) final;
 
 	void BuildPayload(
 		const FBuildPayloadKey& Payload,
@@ -62,37 +63,24 @@ public:
 
 void FBuildSessionInternal::Build(
 	const FBuildDefinition& Definition,
-	EBuildPolicy Policy,
+	const FOptionalBuildInputs& Inputs,
+	const FBuildPolicy& Policy,
 	IRequestOwner& Owner,
 	FOnBuildComplete&& OnComplete)
 {
-	FOnBuildJobComplete OnJobComplete;
-	if (OnComplete)
-	{
-		OnJobComplete = [Definition, OnComplete = MoveTemp(OnComplete)](FBuildJobCompleteParams&& Params)
-		{
-			OnComplete({Definition.GetKey(), Params.CacheKey, MoveTemp(Params.Output), Params.BuildStatus, Params.Status});
-		};
-	}
-	CreateBuildJob({Cache, BuildSystem, Scheduler, InputResolver, Owner, Policy}, Definition, MoveTemp(OnJobComplete));
+	CreateBuildJob({Cache, BuildSystem, Scheduler, InputResolver, Owner}, Definition, Inputs, Policy,
+		OnComplete ? MoveTemp(OnComplete) : [](FBuildCompleteParams&&){});
 }
 
-void FBuildSessionInternal::BuildAction(
+void FBuildSessionInternal::Build(
 	const FBuildAction& Action,
 	const FOptionalBuildInputs& Inputs,
-	EBuildPolicy Policy,
+	const FBuildPolicy& Policy,
 	IRequestOwner& Owner,
-	FOnBuildActionComplete&& OnComplete)
+	FOnBuildComplete&& OnComplete)
 {
-	FOnBuildJobComplete OnJobComplete;
-	if (OnComplete)
-	{
-		OnJobComplete = [Action, OnComplete = MoveTemp(OnComplete)](FBuildJobCompleteParams&& Params)
-		{
-			OnComplete({Action.GetKey(), Params.CacheKey, MoveTemp(Params.Output), Params.BuildStatus, Params.Status});
-		};
-	}
-	CreateBuildJob({Cache, BuildSystem, Scheduler, InputResolver, Owner, Policy}, Action, Inputs, MoveTemp(OnJobComplete));
+	CreateBuildJob({Cache, BuildSystem, Scheduler, InputResolver, Owner}, Action, Inputs, Policy,
+		OnComplete ? MoveTemp(OnComplete) : [](FBuildCompleteParams&&){});
 }
 
 void FBuildSessionInternal::BuildPayload(
@@ -102,10 +90,10 @@ void FBuildSessionInternal::BuildPayload(
 	FOnBuildPayloadComplete&& OnComplete)
 {
 	// This requests the entire output to get one payload. It will be optimized later to request only one payload.
-	FOnBuildJobComplete OnJobComplete;
+	FOnBuildComplete OnJobComplete;
 	if (OnComplete)
 	{
-		OnJobComplete = [PayloadKey, OnComplete = MoveTemp(OnComplete)](FBuildJobCompleteParams&& Params)
+		OnJobComplete = [PayloadKey, OnComplete = MoveTemp(OnComplete)](FBuildCompleteParams&& Params)
 		{
 			FPayload Payload;
 			EStatus Status = Params.Status;
@@ -121,7 +109,7 @@ void FBuildSessionInternal::BuildPayload(
 			OnComplete({PayloadKey.BuildKey, MoveTemp(Payload), Status});
 		};
 	}
-	CreateBuildJob({Cache, BuildSystem, Scheduler, InputResolver, Owner, Policy}, PayloadKey.BuildKey, MoveTemp(OnJobComplete));
+	CreateBuildJob({Cache, BuildSystem, Scheduler, InputResolver, Owner}, PayloadKey.BuildKey, Policy, MoveTemp(OnJobComplete));
 }
 
 ///////////////////////////////////////////////////////////////////////////////////////////////////////////////////////
diff --git a/Engine/Source/Developer/DerivedDataCache/Public/DerivedDataBuildFunction.h b/Engine/Source/Developer/DerivedDataCache/Public/DerivedDataBuildFunction.h
--- a/Engine/Source/Developer/DerivedDataCache/Public/DerivedDataBuildFunction.h
+++ b/Engine/Source/Developer/DerivedDataCache/Public/DerivedDataBuildFunction.h
@@ -97,8 +97,8 @@ public:
 	virtual void AddPayload(const FPayloadId& Id, const FSharedBuffer& Buffer) = 0;
 	virtual void AddPayload(const FPayloadId& Id, const FCbObject& Object) = 0;
 
-	/** Overrides the default cache policy used when writing this build in the cache. */
-	virtual void SetCachePolicy(ECachePolicy Policy) = 0;
+	/** Overrides the cache policy mask used when writing this build in the cache. */
+	virtual void SetCachePolicyMask(ECachePolicy Policy) = 0;
 
 	/**
 	 * Make this an asynchronous build by making the caller responsible for completing the build.
@@ -137,27 +137,29 @@ public:
 	/** Overrides the cache bucket used when reading or writing this build in the cache. */
 	virtual void SetCacheBucket(FCacheBucket Bucket) = 0;
 
-	/** Returns the cache policy used when reading or writing this build in the cache. */
-	virtual ECachePolicy GetCachePolicy() const = 0;
+	/** Returns the cache policy mask used when reading or writing this build in the cache. */
+	virtual ECachePolicy GetCachePolicyMask() const = 0;
 
 	/**
-	 * Overrides the cache policy used when reading or writing this build in the cache.
+	 * Overrides the cache policy mask used when reading or writing this build in the cache.
 	 *
-	 * Always set the cache policy by modifying the policy returned by GetCachePolicy(), to maintain
-	 * the state of other flags. Changes to Skip* and KeepAlive flags are not allowed.
+	 * Allows Query and Store flags to be cleared for this build by removing them from the mask.
+	 *
+	 * Set the mask to the bitwise complement of the flags to be cleared: ~ECachePolicy::Remote.
 	 */
-	virtual void SetCachePolicy(ECachePolicy Policy) = 0;
+	virtual void SetCachePolicyMask(ECachePolicy Policy) = 0;
 
-	/** Returns the build policy used when executing this build. */
-	virtual EBuildPolicy GetBuildPolicy() const = 0;
+	/** Returns the build policy mask used when executing this build. */
+	virtual EBuildPolicy GetBuildPolicyMask() const = 0;
 
 	/**
-	 * Overrides the build policy used when executing this build.
+	 * Overrides the build policy mask used when executing this build.
+	 *
+	 * Allows Build flags to be cleared for this build by removing them from the mask.
 	 *
-	 * Always set the build policy by modifying the policy returned by GetBuildPolicy(), to maintain
-	 * the state of other flags. Changes to Cache* and Skip* flags are not allowed.
+	 * Set the mask to the bitwise complement of the flags to be cleared: ~EBuildPolicy::BuildRemote.
 	 */
-	virtual void SetBuildPolicy(EBuildPolicy Policy) = 0;
+	virtual void SetBuildPolicyMask(EBuildPolicy Policy) = 0;
 
 	/** Sets the estimated peak memory required to execute the build, excluding input size. */
 	virtual void SetRequiredMemory(uint64 RequiredMemory) = 0;
diff --git a/Engine/Source/Developer/DerivedDataCache/Public/DerivedDataBuildSession.h b/Engine/Source/Developer/DerivedDataCache/Public/DerivedDataBuildSession.h
--- a/Engine/Source/Developer/DerivedDataCache/Public/DerivedDataBuildSession.h
+++ b/Engine/Source/Developer/DerivedDataCache/Public/DerivedDataBuildSession.h
@@ -6,27 +6,20 @@
 #include "Containers/StringView.h"
 #include "DerivedDataBuildKey.h"
 #include "DerivedDataBuildTypes.h"
-#include "DerivedDataRequestTypes.h"
 #include "Templates/Function.h"
 #include "Templates/UniquePtr.h"
 
 namespace UE::DerivedData { class FBuildAction; }
 namespace UE::DerivedData { class FBuildDefinition; }
-namespace UE::DerivedData { class FBuildOutput; }
 namespace UE::DerivedData { class FBuildSession; }
 namespace UE::DerivedData { class FOptionalBuildInputs; }
 namespace UE::DerivedData { class FPayload; }
 namespace UE::DerivedData { class IRequestOwner; }
-namespace UE::DerivedData { struct FBuildActionCompleteParams; }
-namespace UE::DerivedData { struct FBuildCompleteParams; }
 namespace UE::DerivedData { struct FBuildPayloadCompleteParams; }
-namespace UE::DerivedData { struct FCacheKey; }
 
 namespace UE::DerivedData
 {
 
-using FOnBuildComplete = TUniqueFunction<void (FBuildCompleteParams&& Params)>;
-using FOnBuildActionComplete = TUniqueFunction<void (FBuildActionCompleteParams&& Params)>;
 using FOnBuildPayloadComplete = TUniqueFunction<void (FBuildPayloadCompleteParams&& Params)>;
 
 } // UE::DerivedData
@@ -41,15 +34,16 @@ public:
 	virtual FStringView GetName() const = 0;
 	virtual void Build(
 		const FBuildDefinition& Definition,
-		EBuildPolicy Policy,
+		const FOptionalBuildInputs& Inputs,
+		const FBuildPolicy& Policy,
 		IRequestOwner& Owner,
 		FOnBuildComplete&& OnComplete) = 0;
-	virtual void BuildAction(
+	virtual void Build(
 		const FBuildAction& Action,
 		const FOptionalBuildInputs& Inputs,
-		EBuildPolicy Policy,
+		const FBuildPolicy& Policy,
 		IRequestOwner& Owner,
-		FOnBuildActionComplete&& OnComplete) = 0;
+		FOnBuildComplete&& OnComplete) = 0;
 	virtual void BuildPayload(
 		const FBuildPayloadKey& Payload,
 		EBuildPolicy Policy,
@@ -86,17 +80,19 @@ public:
 	 * The callback will always be called, and may be called from an arbitrary thread.
 	 *
 	 * @param Definition   The build function to execute and references to its inputs.
-	 * @param Policy       Flags to control the behavior of the request. See EBuildPolicy.
+	 * @param Inputs       The build inputs referenced by the definition, if it has any. Optional.
+	 * @param Policy       Flags to control the behavior of the request. See FBuildPolicy.
 	 * @param Owner        The owner to execute the build within.
 	 * @param OnComplete   A callback invoked when the build completes or is canceled.
 	 */
 	inline void Build(
 		const FBuildDefinition& Definition,
-		EBuildPolicy Policy,
+		const FOptionalBuildInputs& Inputs,
+		const FBuildPolicy& Policy,
 		IRequestOwner& Owner,
 		FOnBuildComplete&& OnComplete)
 	{
-		Session->Build(Definition, Policy, Owner, MoveTemp(OnComplete));
+		Session->Build(Definition, Inputs, Policy, Owner, MoveTemp(OnComplete));
 	}
 
 	/**
@@ -105,19 +101,19 @@ public:
 	 * The callback will always be called, and may be called from an arbitrary thread.
 	 *
 	 * @param Action       The build function to execute and references to its inputs.
-	 * @param Inputs       The build inputs referenced by the action, if it has any.
-	 * @param Policy       Flags to control the behavior of the request. See EBuildPolicy.
+	 * @param Inputs       The build inputs referenced by the action, if it has any. Optional.
+	 * @param Policy       Flags to control the behavior of the request. See FBuildPolicy.
 	 * @param Owner        The owner to execute the build within.
 	 * @param OnComplete   A callback invoked when the build completes or is canceled.
 	 */
-	inline void BuildAction(
+	inline void Build(
 		const FBuildAction& Action,
 		const FOptionalBuildInputs& Inputs,
-		EBuildPolicy Policy,
+		const FBuildPolicy& Policy,
 		IRequestOwner& Owner,
-		FOnBuildActionComplete&& OnComplete)
+		FOnBuildComplete&& OnComplete)
 	{
-		Session->BuildAction(Action, Inputs, Policy, Owner, MoveTemp(OnComplete));
+		Session->Build(Action, Inputs, Policy, Owner, MoveTemp(OnComplete));
 	}
 
 	/**
@@ -179,56 +175,6 @@ public:
 	inline void Reset() { *this = FOptionalBuildSession(); }
 };
 
-/** Parameters for the completion callback for build requests. */
-struct FBuildCompleteParams
-{
-	/** Key for the build request that completed or was canceled. */
-	const FBuildKey& BuildKey;
-
-	/** Key for the build in the cache. Empty if the build completes before the key is assigned. */
-	const FCacheKey& CacheKey;
-
-	/**
-	 * Output for the build request that completed or was canceled.
-	 *
-	 * The name, function, and diagnostics are always populated.
-	 *
-	 * The payloads are populated when Status is Ok, but with null data if skipped by the policy.
-	 */
-	FBuildOutput&& Output;
-
-	/** Detailed status of the build request. */
-	EBuildStatus BuildStatus = EBuildStatus::None;
-
-	/** Basic status of the build request. */
-	EStatus Status = EStatus::Error;
-};
-
-/** Parameters for the completion callback for build action requests. */
-struct FBuildActionCompleteParams
-{
-	/** Key for the build action request that completed or was canceled. */
-	FBuildActionKey BuildKey;
-
-	/** Key for the build in the cache. Empty if the build completes before the key is assigned. */
-	const FCacheKey& CacheKey;
-
-	/**
-	 * Output for the build action request that completed or was canceled.
-	 *
-	 * The name, function, and diagnostics are always populated.
-	 *
-	 * The payloads are populated when Status is Ok, but with null data if skipped by the policy.
-	 */
-	FBuildOutput&& Output;
-
-	/** Detailed status of the build request. */
-	EBuildStatus BuildStatus = EBuildStatus::None;
-
-	/** Basic status of the build request. */
-	EStatus Status = EStatus::Error;
-};
-
 /** Parameters for the completion callback for build payload requests. */
 struct FBuildPayloadCompleteParams
 {
diff --git a/Engine/Source/Developer/DerivedDataCache/Public/DerivedDataBuildTypes.h b/Engine/Source/Developer/DerivedDataCache/Public/DerivedDataBuildTypes.h
--- a/Engine/Source/Developer/DerivedDataCache/Public/DerivedDataBuildTypes.h
+++ b/Engine/Source/Developer/DerivedDataCache/Public/DerivedDataBuildTypes.h
@@ -3,11 +3,23 @@
 #pragma once
 
 #include "CoreTypes.h"
+#include "DerivedDataPayloadId.h"
+#include "DerivedDataRequestTypes.h"
 #include "Misc/EnumClassFlags.h"
+#include "Templates/RefCounting.h"
+
+#define UE_API DERIVEDDATACACHE_API
+
+namespace UE::DerivedData { class FBuildOutput; }
+namespace UE::DerivedData { struct FBuildCompleteParams; }
+namespace UE::DerivedData { struct FCacheKey; }
+namespace UE::DerivedData::Private { class IBuildPolicyShared; }
 
 namespace UE::DerivedData
 {
 
+using FOnBuildComplete = TUniqueFunction<void (FBuildCompleteParams&& Params)>;
+
 /**
  * Flags to control the behavior of build requests.
  *
@@ -24,34 +36,122 @@ namespace UE::DerivedData
 enum class EBuildPolicy : uint32
 {
 	/** A value without any flags set. */
-	None            = 0,
+	None                = 0,
 
 	/** Allow local execution of the build function. */
-	BuildLocal      = 1 << 0,
+	BuildLocal          = 1 << 0,
 	/** Allow remote execution of the build function if it has a registered build worker. */
-	BuildRemote     = 1 << 1,
+	BuildRemote         = 1 << 1,
 	/** Allow local and remote execution of the build function. */
-	Build           = BuildLocal | BuildRemote,
+	Build               = BuildLocal | BuildRemote,
 
 	/** Allow a cache query to avoid having to build. */
-	CacheQuery      = 1 << 2,
+	CacheQuery          = 1 << 2,
+	/** Allow a cache store to persist the build output when another cache store contains it. */
+	CacheStoreOnQuery   = 1 << 3,
+	/** Allow a cache store to persist the build output when the build function executes. */
+	CacheStoreOnBuild   = 1 << 4,
 	/** Allow a cache store to persist the build output. */
-	CacheStore      = 1 << 3,
+	CacheStore          = CacheStoreOnQuery | CacheStoreOnBuild,
 	/** Allow a cache query and a cache store for the build. */
-	Cache           = CacheQuery | CacheStore,
+	Cache               = CacheQuery | CacheStore,
 
 	/** Keep records in the cache for at least the duration of the session. */
-	CacheKeepAlive  = 1 << 4,
+	CacheKeepAlive      = 1 << 5,
 
-	/** Skip fetching the payload data from the cache. */
-	SkipData        = 1 << 5,
+	/** Skip fetching or returning data for the payloads. */
+	SkipData            = 1 << 6,
 
 	/** Allow cache query+store, allow local+remote build when missed or skipped, and fetch the payload(s). */
-	Default         = Build | Cache,
+	Default             = Build | Cache,
 };
 
 ENUM_CLASS_FLAGS(EBuildPolicy);
 
+/** A payload ID and the build policy to use for that payload. */
+struct FBuildPayloadPolicy
+{
+	FPayloadId Id;
+	EBuildPolicy Policy = EBuildPolicy::Default;
+};
+
+/** Interface for the private implementation of the build policy. */
+class Private::IBuildPolicyShared
+{
+public:
+	virtual ~IBuildPolicyShared() = default;
+	virtual void AddRef() const = 0;
+	virtual void Release() const = 0;
+	virtual TConstArrayView<FBuildPayloadPolicy> GetPayloadPolicies() const = 0;
+	virtual void AddPayloadPolicy(const FBuildPayloadPolicy& Policy) = 0;
+	virtual void Build() = 0;
+};
+
+/** Flags to control the behavior of build requests, with optional overrides by payload. */
+class FBuildPolicy
+{
+public:
+	/** Construct a build policy that uses the default policy. */
+	FBuildPolicy() = default;
+
+	/** Construct a build policy with a single policy for every payload. */
+	inline FBuildPolicy(EBuildPolicy Policy)
+		: CombinedPolicy(Policy)
+		, DefaultPayloadPolicy(Policy)
+	{
+	}
+
+	/** Returns true if every payload uses the same build policy. */
+	inline bool IsUniform() const { return !Shared; }
+
+	/** Returns the build policy combined from the payload policies. */
+	inline EBuildPolicy GetCombinedPolicy() const { return CombinedPolicy; }
+
+	/** Returns the build policy to use for the payload. */
+	UE_API EBuildPolicy GetPayloadPolicy(const FPayloadId& Id) const;
+
+	/** Returns the build policy to use for payloads with no override. */
+	inline EBuildPolicy GetDefaultPayloadPolicy() const { return DefaultPayloadPolicy; }
+
+	/** Returns the array of build policy overrides for payloads, sorted by ID. */
+	inline TConstArrayView<FBuildPayloadPolicy> GetPayloadPolicies() const
+	{
+		return Shared ? Shared->GetPayloadPolicies() : TConstArrayView<FBuildPayloadPolicy>();
+	}
+
+private:
+	friend class FBuildPolicyBuilder;
+
+	EBuildPolicy CombinedPolicy = EBuildPolicy::Default;
+	EBuildPolicy DefaultPayloadPolicy = EBuildPolicy::Default;
+	TRefCountPtr<const Private::IBuildPolicyShared> Shared;
+};
+
+/** A build policy builder is used to construct a build policy. */
+class FBuildPolicyBuilder
+{
+public:
+	/** Construct a policy builder that uses the default policy. */
+	FBuildPolicyBuilder() = default;
+
+	/** Construct a policy builder that uses the provided policy for payloads with no override. */
+	inline explicit FBuildPolicyBuilder(EBuildPolicy Policy)
+		: BasePolicy(Policy)
+	{
+	}
+
+	/** Adds a build policy override for a payload. */
+	UE_API void AddPayloadPolicy(const FBuildPayloadPolicy& Policy);
+	inline void AddPayloadPolicy(const FPayloadId& Id, EBuildPolicy Policy) { AddPayloadPolicy({Id, Policy}); }
+
+	/** Build a build policy, which makes this builder subsequently unusable. */
+	UE_API FBuildPolicy Build();
+
+private:
+	EBuildPolicy BasePolicy = EBuildPolicy::Default;
+	TRefCountPtr<Private::IBuildPolicyShared> Shared;
+};
+
 /** Flags for build request completion callbacks. */
 enum class EBuildStatus : uint32
 {
@@ -86,4 +186,28 @@ enum class EBuildStatus : uint32
 
 ENUM_CLASS_FLAGS(EBuildStatus);
 
+/** Parameters for the completion callback for build requests. */
+struct FBuildCompleteParams
+{
+	/** Key for the build in the cache. Empty if the build completes before the key is assigned. */
+	const FCacheKey& CacheKey;
+
+	/**
+	 * Output for the build request that completed or was canceled.
+	 *
+	 * The name, function, and diagnostics are always populated.
+	 *
+	 * The payloads are populated when Status is Ok, but with null data if skipped by the policy.
+	 */
+	FBuildOutput&& Output;
+
+	/** Detailed status of the build request. */
+	EBuildStatus BuildStatus = EBuildStatus::None;
+
+	/** Basic status of the build request. */
+	EStatus Status = EStatus::Error;
+};
+
 } // UE::DerivedData
+
+#undef UE_API
diff --git a/Engine/Source/Runtime/Engine/Private/TextureDerivedDataTask.cpp b/Engine/Source/Runtime/Engine/Private/TextureDerivedDataTask.cpp
--- a/Engine/Source/Runtime/Engine/Private/TextureDerivedDataTask.cpp
+++ b/Engine/Source/Runtime/Engine/Private/TextureDerivedDataTask.cpp
@@ -1035,7 +1035,7 @@ public:
 				bBuildKicked = true;
 				this->UsedEncodeSpeed = (ETextureEncodeSpeed)InSettingsFetchFirst->RepresentsEncodeSpeedNoSend;
 
-				BuildSession.Get().Build(FetchDefinition, EBuildPolicy::Cache, *Owner,
+				BuildSession.Get().Build(FetchDefinition, {}, EBuildPolicy::Cache, *Owner,
 					[this, FetchOrBuildDefinition = MoveTemp(FetchOrBuildDefinition), Flags, FetchOrBuildEncodeSpeed = (ETextureEncodeSpeed)InSettingsFetchOrBuild.RepresentsEncodeSpeedNoSend](FBuildCompleteParams&& Params)
 					{
 						switch (Params.Status)
@@ -1084,9 +1084,9 @@ private:
 		EBuildPolicy BuildPolicy = EBuildPolicy::Default;
 		if (EnumHasAnyFlags(Flags, ETextureCacheFlags::ForceRebuild))
 		{
-			BuildPolicy &= ~EBuildPolicy::CacheQuery;
+			EnumRemoveFlags(BuildPolicy, EBuildPolicy::CacheQuery);
 		}
-		BuildSession.Get().Build(Definition, BuildPolicy, *Owner,
+		BuildSession.Get().Build(Definition, {}, BuildPolicy, *Owner,
 			[this](FBuildCompleteParams&& Params)
 			{
 				EndBuild(Params.CacheKey, MoveTemp(Params.Output), Params.BuildStatus);
