Support reading HttpBackend secret from a file

diff --git a/Engine/Source/Developer/DerivedDataCache/Private/HttpDerivedDataBackend.cpp b/Engine/Source/Developer/DerivedDataCache/Private/HttpDerivedDataBackend.cpp
--- a/Engine/Source/Developer/DerivedDataCache/Private/HttpDerivedDataBackend.cpp
+++ b/Engine/Source/Developer/DerivedDataCache/Private/HttpDerivedDataBackend.cpp
@@ -2036,10 +2036,11 @@ bool FHttpDerivedDataBackend::AcquireAccessToken()
 		FHttpRequest Request(*AuthDomain, nullptr, false);
 
 		// If contents of the secret string is a file path, resolve and read form data.
-		if (OAuthSecret.StartsWith(TEXT("\\\\")))
+		if (OAuthSecret.StartsWith(TEXT("file://")))
 		{
+			FString FilePath = OAuthSecret.Mid(7, OAuthSecret.Len() - 7);
 			FString SecretFileContents;
-			if (FFileHelper::LoadFileToString(SecretFileContents, *OAuthSecret))
+			if (FFileHelper::LoadFileToString(SecretFileContents, *FilePath))
 			{
 				// Overwrite the filepath with the actual content.
 				OAuthSecret = SecretFileContents;
