DDC: Fixed a static analysis warning in the memory backend

#rnx

diff --git a/Engine/Source/Developer/DerivedDataCache/Private/MemoryDerivedDataBackend.cpp b/Engine/Source/Developer/DerivedDataCache/Private/MemoryDerivedDataBackend.cpp
--- a/Engine/Source/Developer/DerivedDataCache/Private/MemoryDerivedDataBackend.cpp
+++ b/Engine/Source/Developer/DerivedDataCache/Private/MemoryDerivedDataBackend.cpp
@@ -197,7 +197,7 @@ void FMemoryDerivedDataBackend::RemoveCachedData(const TCHAR* CacheKey, bool bTr
 	{
 		FWriteScopeLock WriteScopeLock(SynchronizationObject);
 		
-		if (CacheItems.RemoveAndCopyValue(Key, Item))
+		if (CacheItems.RemoveAndCopyValue(Key, Item) && Item)
 		{
 			CurrentCacheSize -= CalcSerializedCacheValueSize(Key, *Item);
 			bMaxSizeExceeded = false;
