DDC: Added a ServerID field to the Jupiter configuration to allow reading settings from another section

When ServerID is used, configuration is first loaded from a key matching the ServerID field, in the HordeStorageServers section of the same INI. This makes it easier to put Jupiter configuration in a consistent place and have it stripped by the cooker to avoid leaking any of the configuration.

#rnx

diff --git a/Engine/Config/BaseGame.ini b/Engine/Config/BaseGame.ini
--- a/Engine/Config/BaseGame.ini
+++ b/Engine/Config/BaseGame.ini
@@ -149,7 +149,8 @@ bPakUsesSecondaryOrder=True
 +IniKeyBlacklist=IniKeyBlacklist
 +IniKeyBlacklist=IniSectionBlacklist
 
- 
++IniSectionBlacklist=HordeStorageServers
+
 [/Script/Engine.HUD]
 DebugDisplay=AI
 
diff --git a/Engine/Source/Developer/DerivedDataCache/Private/DerivedDataBackends.cpp b/Engine/Source/Developer/DerivedDataCache/Private/DerivedDataBackends.cpp
--- a/Engine/Source/Developer/DerivedDataCache/Private/DerivedDataBackends.cpp
+++ b/Engine/Source/Developer/DerivedDataCache/Private/DerivedDataBackends.cpp
@@ -234,7 +234,7 @@ public:
 				}
 				else if (NodeType == TEXT("Http"))
 				{
-					ParsedNode = ParseHttpCache(NodeName, *Entry);
+					ParsedNode = ParseHttpCache(NodeName, *Entry, IniFilename, IniSection);
 				}
 				else if (NodeType == TEXT("Zen"))
 				{
@@ -738,112 +738,158 @@ public:
 #endif
 	}
 
-	/**
-	 * Creates a HTTP data cache interface.
-	 */
-	FDerivedDataBackendInterface* ParseHttpCache(const TCHAR* NodeName, const TCHAR* Entry)
+	void ParseHttpCacheParams(
+		const TCHAR* NodeName,
+		const TCHAR* Entry,
+		const FString& IniFilename,
+		const TCHAR* IniSection,
+		FString& Host,
+		FString& Namespace,
+		FString& OAuthProvider,
+		FString& OAuthClientId,
+		FString& OAuthSecret,
+		bool& bReadOnly)
 	{
-#if WITH_HTTP_DDC_BACKEND
-		FString ServiceUrl;
-		if (!FParse::Value(Entry, TEXT("Host="), ServiceUrl))
+		FString ServerId;
+		if (FParse::Value(Entry, TEXT("ServerID="), ServerId))
 		{
-			UE_LOG(LogDerivedDataCache, Error, TEXT("Node %s does not specify 'Host'."), NodeName);
-			return nullptr;
+			FString ServerEntry;
+			const TCHAR* ServerSection = TEXT("HordeStorageServers");
+			if (GConfig->GetString(ServerSection, *ServerId, ServerEntry, IniFilename))
+			{
+				ParseHttpCacheParams(NodeName, *ServerEntry, IniFilename, IniSection, Host, Namespace, OAuthProvider, OAuthClientId, OAuthSecret, bReadOnly);
+			}
+			else
+			{
+				UE_LOG(LogDerivedDataCache, Warning, TEXT("Node %s is using ServerID=%s which was not found in [%s]"), NodeName, *ServerId, ServerSection);
+			}
 		}
 
-		// Check the EnvHostOverride environment variable to allow persistent overriding of data cache path, eg for offsite workers.
+		FParse::Value(Entry, TEXT("Host="), Host);
+
 		FString EnvHostOverride;
-		if( FParse::Value( Entry, TEXT("EnvHostOverride="), EnvHostOverride ) )
+		if (FParse::Value(Entry, TEXT("EnvHostOverride="), EnvHostOverride))
 		{
-			FString ServiceUrlEnv = FPlatformMisc::GetEnvironmentVariable( *EnvHostOverride );
-			if( ServiceUrlEnv.Len() > 0 )
+			FString HostEnv = FPlatformMisc::GetEnvironmentVariable(*EnvHostOverride);
+			if (!HostEnv.IsEmpty())
 			{
-				ServiceUrl = ServiceUrlEnv;
-				UE_LOG( LogDerivedDataCache, Log, TEXT("Found environment variable for Host %s=%s"), *EnvHostOverride, *ServiceUrl );
+				Host = HostEnv;
+				UE_LOG(LogDerivedDataCache, Log, TEXT("Node %s found environment variable for Host %s=%s"), NodeName, *EnvHostOverride, *Host);
 			}
 		}
 
-		// Allow Host to be overriden via CommandLineHostOverride
 		FString CommandLineOverride;
-		if( FParse::Value( Entry, TEXT("CommandLineHostOverride="), CommandLineOverride ) )
+		if (FParse::Value(Entry, TEXT("CommandLineHostOverride="), CommandLineOverride))
 		{
-			FString Value;
-			if (FParse::Value(FCommandLine::Get(), *(CommandLineOverride + TEXT("=")), Value))
+			if (FParse::Value(FCommandLine::Get(), *(CommandLineOverride + TEXT("=")), Host))
 			{
-				ServiceUrl = Value;
-				UE_LOG(LogDerivedDataCache, Warning, TEXT("Found command line override for Host %s=%s"), *CommandLineOverride, *ServiceUrl);
+				UE_LOG(LogDerivedDataCache, Log, TEXT("Node %s found command line override for Host %s=%s"), NodeName, *CommandLineOverride, *Host);
 			}
 		}
 
-		if (ServiceUrl == TEXT("None"))
+		FParse::Value(Entry, TEXT("Namespace="), Namespace);
+		FParse::Value(Entry, TEXT("OAuthProvider="), OAuthProvider);
+		FParse::Value(Entry, TEXT("OAuthClientId="), OAuthClientId);
+		FParse::Value(Entry, TEXT("OAuthSecret="), OAuthSecret);
+		FParse::Bool(Entry, TEXT("ReadOnly="), bReadOnly);
+	}
+
+	/**
+	 * Creates a HTTP data cache interface.
+	 */
+	FDerivedDataBackendInterface* ParseHttpCache(
+		const TCHAR* NodeName,
+		const TCHAR* Entry,
+		const FString& IniFilename,
+		const TCHAR* IniSection)
+	{
+#if WITH_HTTP_DDC_BACKEND
+		FString Host;
+		FString Namespace;
+		FString OAuthProvider;
+		FString OAuthClientId;
+		FString OAuthSecret;
+		bool bReadOnly = false;
+
+		ParseHttpCacheParams(NodeName, Entry, IniFilename, IniSection, Host, Namespace, OAuthProvider, OAuthClientId, OAuthSecret, bReadOnly);
+
+		if (Host.IsEmpty())
 		{
-			UE_LOG( LogDerivedDataCache, Log, TEXT("Disabling %s data cache - host set to 'None'."), NodeName );
+			UE_LOG(LogDerivedDataCache, Error, TEXT("Node %s does not specify 'Host'"), NodeName);
 			return nullptr;
 		}
 
-		FDerivedDataBackendInterface::ESpeedClass ForceSpeedClass = FDerivedDataBackendInterface::ESpeedClass::Unknown;
-		FString ForceSpeedClassValue;
-		if (FParse::Value(FCommandLine::Get(), TEXT("HttpForceSpeedClass="), ForceSpeedClassValue))
+		if (Host == TEXT("None"))
 		{
-			UE_LOG(LogDerivedDataCache, Warning, TEXT("Jupiter speed class overridden due to HttpForceSpeedClass=%s"), *ForceSpeedClassValue);
-			if (ForceSpeedClassValue == TEXT("Slow"))
-			{
-				ForceSpeedClass = FDerivedDataBackendInterface::ESpeedClass::Slow;
-			}
-			else if (ForceSpeedClassValue == TEXT("Fast"))
-			{
-				ForceSpeedClass = FDerivedDataBackendInterface::ESpeedClass::Fast;
-			}
-			else if (ForceSpeedClassValue == TEXT("Local"))
-			{
-				ForceSpeedClass = FDerivedDataBackendInterface::ESpeedClass::Local;
-			}
+			UE_LOG(LogDerivedDataCache, Log, TEXT("Node %s is disabled because Host is set to 'None'"), NodeName);
+			return nullptr;
 		}
-		
-		FString Namespace;
-		if (!FParse::Value(Entry, TEXT("Namespace="), Namespace))
+
+		if (Namespace.IsEmpty())
 		{
 			Namespace = FApp::GetProjectName();
 			UE_LOG(LogDerivedDataCache, Warning, TEXT("Node %s does not specify 'Namespace', falling back to '%s'"), NodeName, *Namespace);
 		}
 
-		FString OAuthProvider;
-		if (!FParse::Value(Entry, TEXT("OAuthProvider="), OAuthProvider))
+		if (OAuthProvider.IsEmpty())
 		{
-			UE_LOG(LogDerivedDataCache, Error, TEXT("Node %s does not specify 'OAuthProvider'."), NodeName);
+			UE_LOG(LogDerivedDataCache, Error, TEXT("Node %s does not specify 'OAuthProvider'"), NodeName);
 			return nullptr;
 		}
 
-		FString OAuthSecret;
-		if (!FParse::Value(Entry, TEXT("OAuthSecret="), OAuthSecret))
+		if (OAuthClientId.IsEmpty())
 		{
-			UE_LOG(LogDerivedDataCache, Error, TEXT("Node %s does not specify 'OAuthSecret'."), NodeName);
+			UE_LOG(LogDerivedDataCache, Error, TEXT("Node %s does not specify 'OAuthClientId'"), NodeName);
 			return nullptr;
 		}
 
-		FString OAuthClientId;
-		if (!FParse::Value(Entry, TEXT("OAuthClientId="), OAuthClientId))
+		if (OAuthSecret.IsEmpty())
 		{
-			UE_LOG(LogDerivedDataCache, Error, TEXT("Node %s does not specify 'OAuthClientId'."), NodeName);
+			UE_LOG(LogDerivedDataCache, Error, TEXT("Node %s does not specify 'OAuthSecret'"), NodeName);
 			return nullptr;
 		}
 
-		const bool bReadOnly = GetParsedBool(Entry, TEXT("ReadOnly="));
+		FDerivedDataBackendInterface::ESpeedClass ForceSpeedClass = FDerivedDataBackendInterface::ESpeedClass::Unknown;
+		FString ForceSpeedClassValue;
+		if (FParse::Value(FCommandLine::Get(), TEXT("HttpForceSpeedClass="), ForceSpeedClassValue))
+		{
+			if (ForceSpeedClassValue == TEXT("Slow"))
+			{
+				ForceSpeedClass = FDerivedDataBackendInterface::ESpeedClass::Slow;
+			}
+			else if (ForceSpeedClassValue == TEXT("Ok"))
+			{
+				ForceSpeedClass = FDerivedDataBackendInterface::ESpeedClass::Ok;
+			}
+			else if (ForceSpeedClassValue == TEXT("Fast"))
+			{
+				ForceSpeedClass = FDerivedDataBackendInterface::ESpeedClass::Fast;
+			}
+			else if (ForceSpeedClassValue == TEXT("Local"))
+			{
+				ForceSpeedClass = FDerivedDataBackendInterface::ESpeedClass::Local;
+			}
+			else
+			{
+				UE_LOG(LogDerivedDataCache, Warning, TEXT("Node %s found unknown speed class override HttpForceSpeedClass=%s"), NodeName, *ForceSpeedClassValue);
+			}
+		}
 
-		FHttpDerivedDataBackend* backend = new FHttpDerivedDataBackend(*ServiceUrl, *Namespace, *OAuthProvider, *OAuthClientId, *OAuthSecret, bReadOnly);
+		FHttpDerivedDataBackend* Backend = new FHttpDerivedDataBackend(*Host, *Namespace, *OAuthProvider, *OAuthClientId, *OAuthSecret, bReadOnly);
 
 		if (ForceSpeedClass != FDerivedDataBackendInterface::ESpeedClass::Unknown)
 		{
-			backend->SetSpeedClass(ForceSpeedClass);
+			UE_LOG(LogDerivedDataCache, Log, TEXT("Node %s found speed class override ForceSpeedClass=%s"), NodeName, *ForceSpeedClassValue);
+			Backend->SetSpeedClass(ForceSpeedClass);
 		}
 
-		if (!backend->IsUsable())
+		if (!Backend->IsUsable())
 		{
-			UE_LOG(LogDerivedDataCache, Warning, TEXT("%s could not contact the service (%s), will not use it."), NodeName, *ServiceUrl);
-			delete backend;
+			UE_LOG(LogDerivedDataCache, Warning, TEXT("Node %s could not contact the service (%s), will not use it"), NodeName, *Host);
+			delete Backend;
 			return nullptr;
 		}
-		return backend;
+		return Backend;
 #else
 		UE_LOG(LogDerivedDataCache, Warning, TEXT("HTTP backend is not yet supported in the current build configuration."));
 		return nullptr;
