DDC: Added a missing request barrier in the hierarchical backend

#rnx

diff --git a/Engine/Source/Developer/DerivedDataCache/Private/HierarchicalDerivedDataBackend.h b/Engine/Source/Developer/DerivedDataCache/Private/HierarchicalDerivedDataBackend.h
--- a/Engine/Source/Developer/DerivedDataCache/Private/HierarchicalDerivedDataBackend.h
+++ b/Engine/Source/Developer/DerivedDataCache/Private/HierarchicalDerivedDataBackend.h
@@ -576,6 +576,7 @@ public:
 							if (bFill)
 							{
 								FRequestOwner AsyncOwner(FPlatformMath::Min(Owner.GetPriority(), EPriority::Highest));
+								FRequestBarrier AsyncBarrier(AsyncOwner);
 								AsyncOwner.KeepAlive();
 								for (int32 FillCacheIndex = 0; FillCacheIndex < InnerBackends.Num(); ++FillCacheIndex)
 								{
