DDC: New filesystem cache store directory layout and file format

- Implemented only for the structured cache initially.
- Records are stored as compact binary packages with payloads stored within the record if small enough.
- Records are stored in Records/{Bucket}/{01}/{23}/{456789abcdef0123456789abcdef01234567} based on the hash in the cache key.
- Content that exceeds the size allowed for storage within the record is stored in Content/{01}/{23}/{456789abcdef0123456789abcdef01234567} based on the hash of the content.
- Both records and content are written to disk with their hash as a suffix, and are always checked on load.
- Maximum size of compressed data stored in the package can be configured by MaxRecordSizeKB and MaxValueSizeKB.
- Maintenance of the cache, to delete unused files, has been rewritten and can now scan ~175k files/second on modern hardware, though is limited to 10k files/second in the config to avoid adding too much overhead.

diff --git a/Engine/Config/BaseEngine.ini b/Engine/Config/BaseEngine.ini
--- a/Engine/Config/BaseEngine.ini
+++ b/Engine/Config/BaseEngine.ini
@@ -2046,162 +2046,143 @@ HostName=localhost
 Port=1337
 
 ; Derived Data backend graphs
-; Each of the below backend graph sections contains a set of nodes used to create derived data backed graph.
+; Each of the below backend graph sections contains a set of nodes used to create derived data backend graph.
 ; DerivedDataBackendGraph is the default graph for source builds, and InstalledDerivedDataBackendGraph is the default for installed builds.
-; Others can be specified othe command line using:
-;     -DDC=GraphSectionName     (for example: -DDC=VerifyDerivedDataBackendGraph)
-; Each graph should start with 'Root' node. Names of all the other nodes are not predefined.
-; Supported node types are: KeyLength, AsyncPut, Hierarchical, Boot, Filesystem, ReadPak, WritePak, Verify
-; The order nodes are define in is not relevant
-; Filesystem backends can be disabled by setting the corresponding environment variable to 'None'. E.g UE-SharedDataCachePath=None
+; Others can be specified on the command line using:
+;     -DDC=GraphSectionName     (for example: -DDC=NoShared)
+; Each graph should start with a 'Root' node. The names of the other nodes are not predefined.
+; Supported node types are: KeyLength, AsyncPut, Hierarchical, FileSystem, Zen, Http, S3, ReadPak, WritePak, Verify
+; The order that nodes are defined in is not relevant.
+; FileSystem backends can be disabled by setting the corresponding environment variable to 'None'. E.g., UE-SharedDataCachePath=None
+; Http backends can be disabled by setting the corresponding environment variable to 'None'. E.g., UE-CloudDataCachePath=None
 
 [DerivedDataBackendGraph]
-; Keep files for at least 7 days
-MinimumDaysToKeepFile=7
 ; Configure the root entry. It uses a KeyLength node to hash long strings then passes requests to AsyncPut
 Root=(Type=KeyLength, Length=120, Inner=AsyncPut)
 ; Configure the AsyncPut entry. It uses an AsyncPut node that then passes requests to Hierarchy
 AsyncPut=(Type=AsyncPut, Inner=Hierarchy)
-; Configure the hierarchy entry. This uses multiple nodes that are used in order until a read is found (writes go to all writable entries)
+; Configure the hierarchy entry. This contains multiple nodes that are used in order until a read is found. Writes go to all writable backends.
 Hierarchy=(Type=Hierarchical, Inner=Pak, Inner=EnginePak, Inner=Local, Inner=Shared)
-; Configure the local node. This is a filesystem node with the following set -
-;	Readonly: Can data be written to this layer
-;	Clean: Perform a cleanup of old files on launch
-;	Flush: Clear the DDC and start over
-;	PurgeTransient: Don't keep transient data in the DDC
-;	DeleteUnused: Clean up old files (happens on a background thread)
-;	UnusedFileAge: Age after which files are removed
-;	FoldersToClean: Max number of folders to clean in a session. -1 = Unlimited
-;	PromptIfMissing: If the path does not exist a warning prompt will be shown. This is useful for shared DDCs where the mount may have dropped
-;	ConsiderSlowAt: If access times are < this value in ms then some operations will be disabled to improve performance.
-;	MaxFileChecksPerSec: How many files to check a second.
-;	Path: path to use for the filesystem DDC
-;	EnvPathOverride: An environment variable that if set will be used instead of path. E.g. UE-LocalDataCachePath=d:\DDC. ('None' disables the DDC)
-;	CommandLineOverride: A command line argument used in preference to the default / envvar setting. E.g. -SharedDataCachePath=\\someshare\folder
-;	EditorOverrideSetting: Editor user setting that overrides the default/envvarcommand line values
-Local=(Type=FileSystem, ReadOnly=false, Clean=false, Flush=false, PurgeTransient=true, DeleteUnused=true, UnusedFileAge=34, FoldersToClean=-1, PromptIfMissing=true, Path=%ENGINEDIR%DerivedDataCache, EnvPathOverride=UE-LocalDataCachePath, EditorOverrideSetting=LocalDerivedDataCache)
+; Configure the local cache. This is a FileSystem cache with the following parameters:
+;	ReadOnly: Can data be written to this layer. Default is false.
+;	Touch: Update timestamps on every access instead of only after 1/4 of UnusedFileAge. Default is false.
+;	PurgeTransient: Delete transient data from the cache. Default is false.
+;	Clean: Perform a cleanup of old files on launch. Default is false.
+;	Flush: Clear the DDC and start over on launch. Default is false.
+;	DeleteUnused: Whether to clean up old files on a background thread. Default is true.
+;	UnusedFileAge: Age since last access after which files are removed. Default is 15 days.
+;	FoldersToClean: Max number of folders to clean in a session. -1 = Unlimited. Default is -1.
+;	PromptIfMissing: If the path does not exist a warning prompt will be shown. This is useful for shared DDCs where the mount may have dropped. Default is false.
+;	ConsiderSlowAt: If access times are less than this value in ms then some operations will be disabled to improve performance. Default is 50ms.
+;	MaxFileChecksPerSec: Maximum file cleanup checks to do in one second. Default is configured in Engine:[DDCCleanup]:MaxFileChecksPerSec.
+;	Path: Path to use for the filesystem DDC.
+;	EnvPathOverride: An environment variable that if set will be used instead of Path. E.g., UE-LocalDataCachePath=D:\DDC. ('None' disables the DDC)
+;	CommandLineOverride: A command line argument used in preference to the Path / EnvPathOverride settings. E.g., -SharedDataCachePath=\\someshare\folder
+;	EditorOverrideSetting: Editor user setting that overrides the Path / EnvPathOverride line values.
+Local=(Type=FileSystem, PurgeTransient=true, UnusedFileAge=34, PromptIfMissing=true, Path=%ENGINEDIR%DerivedDataCache, EnvPathOverride=UE-LocalDataCachePath, EditorOverrideSetting=LocalDerivedDataCache)
 ; Configure the shared DDC that is accessed after local. It's a filesystem DDC and the parameters are explained above
-Shared=(Type=FileSystem, ReadOnly=false, Clean=false, Flush=false, DeleteUnused=true, UnusedFileAge=10, FoldersToClean=10, MaxFileChecksPerSec=1, ConsiderSlowAt=70, PromptIfMissing=false, Path=?EpicDDC, EnvPathOverride=UE-SharedDataCachePath, EditorOverrideSetting=SharedDerivedDataCache, CommandLineOverride=SharedDataCachePath)
-; Configure an alternate shared DDC that is accessed after local. It's a filesystem DDC and the parameters are explained above
-AltShared=(Type=FileSystem, ReadOnly=true, Clean=false, Flush=false, DeleteUnused=true, UnusedFileAge=23, FoldersToClean=10, MaxFileChecksPerSec=1, ConsiderSlowAt=70, PromptIfMissing=false, Path=?EpicDDC2, EnvPathOverride=UE-SharedDataCachePath2)
+Shared=(Type=FileSystem, UnusedFileAge=10, FoldersToClean=10, ConsiderSlowAt=70, Path=?EpicDDC, EnvPathOverride=UE-SharedDataCachePath, EditorOverrideSetting=SharedDerivedDataCache, CommandLineOverride=SharedDataCachePath)
 ; Configure a Project Pak node. This is a pre-generated DDC data file for the project that can be distributed to reduce runtime fetches/generation
 ; See documentation for how to create a DDP via the DerivedDataCache commandlet
 Pak=(Type=ReadPak, Filename="%GAMEDIR%DerivedDataCache/DDC.ddp")
 ; Configure a Project Pak node. This is a pre-generated DDC data file for the engine that can be distributed to reduce runtime fetches/generation
-EnginePak=(Type=ReadPak, Filename=%ENGINEDIR%DerivedDataCache/DDC.ddp)
+EnginePak=(Type=ReadPak, Filename="%ENGINEDIR%DerivedDataCache/DDC.ddp")
+
+[InstalledDerivedDataBackendGraph]
+Root=(Type=KeyLength, Length=120, Inner=AsyncPut)
+AsyncPut=(Type=AsyncPut, Inner=Hierarchy)
+Hierarchy=(Type=Hierarchical, Inner=Pak, Inner=CompressedPak, Inner=EnginePak, Inner=EnterprisePak, Inner=Local, Inner=Shared)
+Local=(Type=FileSystem, PurgeTransient=true, UnusedFileAge=34, Path="%ENGINEVERSIONAGNOSTICUSERDIR%DerivedDataCache", EnvPathOverride=UE-LocalDataCachePath, EditorOverrideSetting=LocalDerivedDataCache)
+Shared=(Type=FileSystem, UnusedFileAge=10, FoldersToClean=10, ConsiderSlowAt=70, Path=?EpicDDC, EnvPathOverride=UE-SharedDataCachePath, EditorOverrideSetting=SharedDerivedDataCache, CommandLineOverride=SharedDataCachePath)
+Pak=(Type=ReadPak, Filename="%GAMEDIR%DerivedDataCache/DDC.ddp")
+CompressedPak=(Type=ReadPak, Filename="%GAMEDIR%DerivedDataCache/Compressed.ddp", Compressed=true)
+EnginePak=(Type=ReadPak, Filename=../../../Engine/DerivedDataCache/Compressed.ddp, Compressed=true)
+EnterprisePak=(Type=ReadPak, Filename=../../../Enterprise/DerivedDataCache/Compressed.ddp, Compressed=true)
 
 [ZenOnlyDDC]
-MinimumDaysToKeepFile=7
 Root=(Type=KeyLength, Length=120, Inner=AsyncPut)
 AsyncPut=(Type=AsyncPut, Inner=Local)
-Local=(Type=Zen, ReadOnly=false, Namespace="ue4.ddc")
+Local=(Type=Zen, Namespace="ue4.ddc")
 
 [ZenDDC]
-MinimumDaysToKeepFile=7
 Root=(Type=KeyLength, Length=120, Inner=AsyncPut)
 AsyncPut=(Type=AsyncPut, Inner=Hierarchy)
 Hierarchy=(Type=Hierarchical, Inner=Local, Inner=Shared)
-Local=(Type=Zen, ReadOnly=false, Namespace="ue4.ddc")
-Shared=(Type=FileSystem, ReadOnly=false, Clean=false, Flush=false, DeleteUnused=true, UnusedFileAge=10, FoldersToClean=10, MaxFileChecksPerSec=1, ConsiderSlowAt=70, PromptIfMissing=false, Path=?EpicDDC, EnvPathOverride=UE-SharedDataCachePath, EditorOverrideSetting=SharedDerivedDataCache, CommandLineOverride=SharedDataCachePath)
+Local=(Type=Zen, Namespace="ue4.ddc")
+Shared=(Type=FileSystem, UnusedFileAge=10, FoldersToClean=10, ConsiderSlowAt=70, Path=?EpicDDC, EnvPathOverride=UE-SharedDataCachePath, EditorOverrideSetting=SharedDerivedDataCache, CommandLineOverride=SharedDataCachePath)
 
 [DerivedDataBackendGraph_Fill_Seattle]
-MinimumDaysToKeepFile=7
 Root=(Type=KeyLength, Length=120, Inner=AsyncPut)
 AsyncPut=(Type=AsyncPut, Inner=Hierarchy)
 Hierarchy=(Type=Hierarchical, Inner=Pak, Inner=EnginePak, Inner=Local, Inner=Seattle)
-Local=(Type=FileSystem, ReadOnly=false, Clean=false, Flush=false, PurgeTransient=true, DeleteUnused=true, UnusedFileAge=34, FoldersToClean=-1, Path=%ENGINEDIR%DerivedDataCache, EnvPathOverride=UE-LocalDataCachePath, EditorOverrideSetting=LocalDerivedDataCache)
-Seattle=(Type=FileSystem, ReadOnly=false, Clean=false, Flush=false, DeleteUnused=true, UnusedFileAge=23, FoldersToClean=10, MaxFileChecksPerSec=1, Path=?EpicSeaDDC, EnvPathOverride=UE-SharedDataCachePath_Seattle)
+Local=(Type=FileSystem, PurgeTransient=true, UnusedFileAge=34, Path=%ENGINEDIR%DerivedDataCache, EnvPathOverride=UE-LocalDataCachePath, EditorOverrideSetting=LocalDerivedDataCache)
+Seattle=(Type=FileSystem, UnusedFileAge=23, FoldersToClean=10, Path=?EpicSeaDDC, EnvPathOverride=UE-SharedDataCachePath_Seattle)
 Pak=(Type=ReadPak, Filename="%GAMEDIR%DerivedDataCache/DDC.ddp")
-EnginePak=(Type=ReadPak, Filename=%ENGINEDIR%DerivedDataCache/DDC.ddp)
-
-[InstalledDerivedDataBackendGraph]
-MinimumDaysToKeepFile=7
-Root=(Type=KeyLength, Length=120, Inner=AsyncPut)
-AsyncPut=(Type=AsyncPut, Inner=Hierarchy)
-Hierarchy=(Type=Hierarchical, Inner=Pak, Inner=CompressedPak, Inner=EnginePak, Inner=EnterprisePak, Inner=Local, Inner=Shared)
-Local=(Type=FileSystem, ReadOnly=false, Clean=false, Flush=false, PurgeTransient=true, DeleteUnused=true, UnusedFileAge=34, FoldersToClean=-1, Path="%ENGINEVERSIONAGNOSTICUSERDIR%DerivedDataCache", EnvPathOverride=UE-LocalDataCachePath, EditorOverrideSetting=LocalDerivedDataCache)
-Shared=(Type=FileSystem, ReadOnly=false, Clean=false, Flush=false, DeleteUnused=true, UnusedFileAge=10, FoldersToClean=10, MaxFileChecksPerSec=1, ConsiderSlowAt=70, PromptIfMissing=false, Path=?EpicDDC, EnvPathOverride=UE-SharedDataCachePath, EditorOverrideSetting=SharedDerivedDataCache, CommandLineOverride=SharedDataCachePath)
-Pak=(Type=ReadPak, Filename="%GAMEDIR%DerivedDataCache/DDC.ddp")
-CompressedPak=(Type=ReadPak, Filename="%GAMEDIR%DerivedDataCache/Compressed.ddp", Compressed=true)
-EnginePak=(Type=ReadPak, Filename=../../../Engine/DerivedDataCache/Compressed.ddp, Compressed=true)
-EnterprisePak=(Type=ReadPak, Filename=../../../Enterprise/DerivedDataCache/Compressed.ddp, Compressed=true)
+EnginePak=(Type=ReadPak, Filename="%ENGINEDIR%DerivedDataCache/DDC.ddp")
 
 [NoShared]
-MinimumDaysToKeepFile=7
 Root=(Type=KeyLength, Length=120, Inner=AsyncPut)
 AsyncPut=(Type=AsyncPut, Inner=Hierarchy)
 Hierarchy=(Type=Hierarchical, Inner=Pak, Inner=Local)
-Local=(Type=FileSystem, ReadOnly=false, Clean=false, Flush=false, PurgeTransient=true, DeleteUnused=true, UnusedFileAge=34, FoldersToClean=-1, Path=%ENGINEDIR%DerivedDataCache, EnvPathOverride=UE-LocalDataCachePath, EditorOverrideSetting=LocalDerivedDataCache)
+Local=(Type=FileSystem, PurgeTransient=true, UnusedFileAge=34, Path=%ENGINEDIR%DerivedDataCache, EnvPathOverride=UE-LocalDataCachePath, EditorOverrideSetting=LocalDerivedDataCache)
 Pak=(Type=ReadPak, Filename="%GAMEDIR%DerivedDataCache/DDC.ddp")
 
 ; Creates a completely new standalone cache so as to not destroy a valuable populated local cache
 [Cold]
-MinimumDaysToKeepFile=0
 Root=(Type=KeyLength, Length=120, Inner=AsyncPut)
 AsyncPut=(Type=AsyncPut, Inner=Hierarchy)
 Hierarchy=(Type=Hierarchical, Inner=Local)
-Local=(Type=FileSystem, ReadOnly=false, Clean=false, Flush=true, PurgeTransient=true, DeleteUnused=true, UnusedFileAge=0, FoldersToClean=-1, Path="%GAMEDIR%ColdDerivedDataCache\")
+Local=(Type=FileSystem, Flush=true, PurgeTransient=true, UnusedFileAge=1, Path="%GAMEDIR%ColdDerivedDataCache")
 
 ; Inherits the same standalone cache from [Cold], for incremental performance measurement purposes
 [Warm]
-MinimumDaysToKeepFile=0
 Root=(Type=KeyLength, Length=120, Inner=AsyncPut)
 AsyncPut=(Type=AsyncPut, Inner=Hierarchy)
 Hierarchy=(Type=Hierarchical, Inner=Local)
-Local=(Type=FileSystem, ReadOnly=false, Clean=false, Flush=false, PurgeTransient=true, DeleteUnused=true, UnusedFileAge=1, FoldersToClean=-1, Path="%GAMEDIR%ColdDerivedDataCache\")
+Local=(Type=FileSystem, Clean=true, PurgeTransient=true, UnusedFileAge=1, Path="%GAMEDIR%ColdDerivedDataCache")
 
 [CreatePak]
-MinimumDaysToKeepFile=7
 Root=(Type=KeyLength, Length=120, Inner=AsyncPut)
 AsyncPut=(Type=AsyncPut, Inner=Hierarchy)
 Hierarchy=(Type=Hierarchical, Inner=PakWrite, Inner=PakRead, Inner=Local, Inner=Shared)
-Local=(Type=FileSystem, ReadOnly=false, Clean=false, Flush=false, PurgeTransient=true, DeleteUnused=true, UnusedFileAge=34, FoldersToClean=-1, Path=%ENGINEDIR%DerivedDataCache, EnvPathOverride=UE-LocalDataCachePath, EditorOverrideSetting=LocalDerivedDataCache)
-Shared=(Type=FileSystem, ReadOnly=false, Clean=false, Flush=false, DeleteUnused=false, UnusedFileAge=23, FoldersToClean=-1, Path=?EpicDDC, EnvPathOverride=UE-SharedDataCachePath)
-AltShared=(Type=FileSystem, ReadOnly=true, Clean=false, Flush=false, DeleteUnused=false, UnusedFileAge=23, FoldersToClean=-1, Path=?EpicDDC2, EnvPathOverride=UE-SharedDataCachePath2)
+Local=(Type=FileSystem, PurgeTransient=true, UnusedFileAge=34, Path=%ENGINEDIR%DerivedDataCache, EnvPathOverride=UE-LocalDataCachePath, EditorOverrideSetting=LocalDerivedDataCache)
+Shared=(Type=FileSystem, DeleteUnused=false, ConsiderSlowAt=70, Path=?EpicDDC, EnvPathOverride=UE-SharedDataCachePath, EditorOverrideSetting=SharedDerivedDataCache, CommandLineOverride=SharedDataCachePath)
 PakRead=(Type=ReadPak, Filename="%GAMEDIR%DerivedDataCache/DDC.ddp")
 PakWrite=(Type=WritePak, Filename="%GAMEDIR%DerivedDataCache/DDC.ddp")
 
 [CreateInstalledProjectPak]
-MinimumDaysToKeepFile=7
 Root=(Type=KeyLength, Length=120, Inner=AsyncPut)
 AsyncPut=(Type=AsyncPut, Inner=Hierarchy)
 Hierarchy=(Type=Hierarchical, Inner=EnginePak, Inner=PakWrite, Inner=PakRead, Inner=Local, Inner=Shared)
-Local=(Type=FileSystem, ReadOnly=false, Clean=false, Flush=false, PurgeTransient=true, DeleteUnused=true, UnusedFileAge=34, FoldersToClean=-1, Path=%ENGINEDIR%DerivedDataCache, EnvPathOverride=UE-LocalDataCachePath, EditorOverrideSetting=LocalDerivedDataCache)
-Shared=(Type=FileSystem, ReadOnly=false, Clean=false, Flush=false, DeleteUnused=false, UnusedFileAge=23, FoldersToClean=-1, Path=?EpicDDC, EnvPathOverride=UE-SharedDataCachePath, CommandLineOverride=SharedDataCachePath)
-AltShared=(Type=FileSystem, ReadOnly=true, Clean=false, Flush=false, DeleteUnused=false, UnusedFileAge=23, FoldersToClean=-1, Path=?EpicDDC2, EnvPathOverride=UE-SharedDataCachePath2)
+Local=(Type=FileSystem, PurgeTransient=true, UnusedFileAge=34, Path=%ENGINEDIR%DerivedDataCache, EnvPathOverride=UE-LocalDataCachePath, EditorOverrideSetting=LocalDerivedDataCache)
+Shared=(Type=FileSystem, DeleteUnused=false, ConsiderSlowAt=70, Path=?EpicDDC, EnvPathOverride=UE-SharedDataCachePath, EditorOverrideSetting=SharedDerivedDataCache, CommandLineOverride=SharedDataCachePath)
 EnginePak=(Type=ReadPak, Filename=../../../Engine/DerivedDataCache/Compressed.ddp, Compressed=true)
 PakRead=(Type=ReadPak, Filename="%GAMEDIR%DerivedDataCache/Compressed.ddp", Compressed=true)
 PakWrite=(Type=WritePak, Filename="%GAMEDIR%DerivedDataCache/Compressed.ddp", Compressed=true)
 
 [CreateInstalledEnginePak]
-MinimumDaysToKeepFile=7
 Root=(Type=KeyLength, Length=120, Inner=AsyncPut)
 AsyncPut=(Type=AsyncPut, Inner=Hierarchy)
 Hierarchy=(Type=Hierarchical, Inner=PakWrite, Inner=Local, Inner=Shared)
-Local=(Type=FileSystem, ReadOnly=false, Clean=false, Flush=false, PurgeTransient=true, DeleteUnused=true, UnusedFileAge=34, FoldersToClean=-1, Path=%ENGINEDIR%DerivedDataCache, EnvPathOverride=UE-LocalDataCachePath, EditorOverrideSetting=LocalDerivedDataCache)
-Shared=(Type=FileSystem, ReadOnly=false, Clean=false, Flush=false, DeleteUnused=false, UnusedFileAge=23, FoldersToClean=-1, Path=?EpicDDC, EnvPathOverride=UE-SharedDataCachePath)
-AltShared=(Type=FileSystem, ReadOnly=true, Clean=false, Flush=false, DeleteUnused=false, UnusedFileAge=23, FoldersToClean=-1, Path=?EpicDDC2, EnvPathOverride=UE-SharedDataCachePath2)
-PakWrite=(Type=WritePak, Filename=%ENGINEDIR%DerivedDataCache/Compressed.ddp, Compressed=true)
+Local=(Type=FileSystem, PurgeTransient=true, UnusedFileAge=34, Path=%ENGINEDIR%DerivedDataCache, EnvPathOverride=UE-LocalDataCachePath, EditorOverrideSetting=LocalDerivedDataCache)
+Shared=(Type=FileSystem, DeleteUnused=false, ConsiderSlowAt=70, Path=?EpicDDC, EnvPathOverride=UE-SharedDataCachePath, EditorOverrideSetting=SharedDerivedDataCache, CommandLineOverride=SharedDataCachePath)
+PakWrite=(Type=WritePak, Filename="%ENGINEDIR%DerivedDataCache/Compressed.ddp", Compressed=true)
 
 [CreateInstalledEnterprisePak]
-MinimumDaysToKeepFile=7
 Root=(Type=KeyLength, Length=120, Inner=AsyncPut)
 AsyncPut=(Type=AsyncPut, Inner=Hierarchy)
 Hierarchy=(Type=Hierarchical, Inner=PakWrite, Inner=Local, Inner=Shared)
-Local=(Type=FileSystem, ReadOnly=false, Clean=false, Flush=false, PurgeTransient=true, DeleteUnused=true, UnusedFileAge=34, FoldersToClean=-1, Path=%ENGINEDIR%DerivedDataCache, EnvPathOverride=UE-LocalDataCachePath, EditorOverrideSetting=LocalDerivedDataCache)
-Shared=(Type=FileSystem, ReadOnly=false, Clean=false, Flush=false, DeleteUnused=false, UnusedFileAge=23, FoldersToClean=-1, Path=?EpicDDC, EnvPathOverride=UE-SharedDataCachePath)
-AltShared=(Type=FileSystem, ReadOnly=true, Clean=false, Flush=false, DeleteUnused=false, UnusedFileAge=23, FoldersToClean=-1, Path=?EpicDDC2, EnvPathOverride=UE-SharedDataCachePath2)
+Local=(Type=FileSystem, PurgeTransient=true, UnusedFileAge=34, Path=%ENGINEDIR%DerivedDataCache, EnvPathOverride=UE-LocalDataCachePath, EditorOverrideSetting=LocalDerivedDataCache)
+Shared=(Type=FileSystem, DeleteUnused=false, ConsiderSlowAt=70, Path=?EpicDDC, EnvPathOverride=UE-SharedDataCachePath, EditorOverrideSetting=SharedDerivedDataCache, CommandLineOverride=SharedDataCachePath)
 PakWrite=(Type=WritePak, Filename=../../../Enterprise/DerivedDataCache/Compressed.ddp, Compressed=true)
 
 [CreateProjectCache]
-MinimumDaysToKeepFile=7
 Root=(Type=KeyLength, Length=120, Inner=AsyncPut)
 AsyncPut=(Type=AsyncPut, Inner=Hierarchy)
 Hierarchy=(Type=Hierarchical, Inner=Local, Inner=Project, Inner=Shared)
-Local=(Type=FileSystem, ReadOnly=false, Clean=false, Flush=false, PurgeTransient=true, DeleteUnused=true, UnusedFileAge=34, FoldersToClean=-1, Path=%ENGINEDIR%DerivedDataCache, EnvPathOverride=UE-LocalDataCachePath, EditorOverrideSetting=LocalDerivedDataCache)
-Project=(Type=FileSystem, ReadOnly=false, Clean=false, Flush=true, PurgeTransient=true, DeleteUnused=true, UnusedFileAge=34, FoldersToClean=-1, Path=%GAMEDIR%ProjectDerivedData)
-Shared=(Type=FileSystem, ReadOnly=false, Clean=false, Flush=false, DeleteUnused=false, UnusedFileAge=23, FoldersToClean=-1, Path=?EpicDDC, EnvPathOverride=UE-SharedDataCachePath)
-AltShared=(Type=FileSystem, ReadOnly=true, Clean=false, Flush=false, DeleteUnused=false, UnusedFileAge=23, FoldersToClean=-1, Path=?EpicDDC2, EnvPathOverride=UE-SharedDataCachePath2)
+Local=(Type=FileSystem, PurgeTransient=true, UnusedFileAge=34, Path=%ENGINEDIR%DerivedDataCache, EnvPathOverride=UE-LocalDataCachePath, EditorOverrideSetting=LocalDerivedDataCache)
+Project=(Type=FileSystem, Flush=true, PurgeTransient=true, UnusedFileAge=34, Path=%GAMEDIR%ProjectDerivedData)
+Shared=(Type=FileSystem, DeleteUnused=false, ConsiderSlowAt=70, Path=?EpicDDC, EnvPathOverride=UE-SharedDataCachePath, EditorOverrideSetting=SharedDerivedDataCache, CommandLineOverride=SharedDataCachePath)
 
 [VirtualTextureChunkDDCCache]
 UnusedFileAge=34
@@ -2210,8 +2191,7 @@ Path=%GAMEDIR%DerivedDataCache/VT
 
 [DDCCleanup]
 TimeToWaitAfterInit=120
-TimeBetweenDeleteingDirectories=5
-TimeBetweenDeletingFiles=2
+MaxFileChecksPerSec=10000
 
 [/Script/Engine.LocalPlayer]
 AspectRatioAxisConstraint=AspectRatio_MaintainXFOV
diff --git a/Engine/Source/Developer/DerivedDataCache/Private/DDCCleanup.cpp b/Engine/Source/Developer/DerivedDataCache/Private/DDCCleanup.cpp
deleted file mode 100644
diff --git a/Engine/Source/Developer/DerivedDataCache/Private/DerivedDataBackends.cpp b/Engine/Source/Developer/DerivedDataCache/Private/DerivedDataBackends.cpp
--- a/Engine/Source/Developer/DerivedDataCache/Private/DerivedDataBackends.cpp
+++ b/Engine/Source/Developer/DerivedDataCache/Private/DerivedDataBackends.cpp
@@ -39,11 +39,16 @@ DEFINE_LOG_CATEGORY(LogDerivedDataCache);
 #define MAX_BACKEND_KEY_LENGTH (120)
 #define LOCTEXT_NAMESPACE "DerivedDataBackendGraph"
 
-namespace UE::DerivedData::Backends
+namespace UE::DerivedData::CacheStore::FileSystem
 {
 
 FDerivedDataBackendInterface* CreateFileSystemDerivedDataBackend(const TCHAR* CacheDirectory, const TCHAR* InParams, const TCHAR* InAccessLogFileName = nullptr);
 
+} // UE::DerivedData::CacheStore::FileSystem
+
+namespace UE::DerivedData::Backends
+{
+
 /**
   * This class is used to create a singleton that represents the derived data cache hierarchy and all of the wrappers necessary
   * ideally this would be data driven and the backends would be plugins...
@@ -632,7 +637,7 @@ public:
 
 				if (!bShared || IFileManager::Get().DirectoryExists(*Path))
 				{
-					InnerFileSystem = CreateFileSystemDerivedDataBackend(*Path, Entry, *WriteAccessLog);
+					InnerFileSystem = CacheStore::FileSystem::CreateFileSystemDerivedDataBackend(*Path, Entry, *WriteAccessLog);
 				}
 
 				if (InnerFileSystem)
diff --git a/Engine/Source/Developer/DerivedDataCache/Private/DerivedDataCache.cpp b/Engine/Source/Developer/DerivedDataCache/Private/DerivedDataCache.cpp
--- a/Engine/Source/Developer/DerivedDataCache/Private/DerivedDataCache.cpp
+++ b/Engine/Source/Developer/DerivedDataCache/Private/DerivedDataCache.cpp
@@ -12,17 +12,19 @@
 #include "DDCCleanup.h"
 #include "DerivedDataBackendInterface.h"
 #include "DerivedDataCache.h"
+#include "DerivedDataCacheMaintainer.h"
 #include "DerivedDataCachePrivate.h"
 #include "DerivedDataCacheUsageStats.h"
 #include "DerivedDataPluginInterface.h"
-#include "ZenServerInterface.h"
+#include "Features/IModularFeatures.h"
 #include "HAL/ThreadSafeCounter.h"
-#include "Misc/CoreMisc.h"
 #include "Misc/CommandLine.h"
+#include "Misc/CoreMisc.h"
 #include "Misc/ScopeLock.h"
 #include "ProfilingDebugging/CookStats.h"
 #include "Stats/Stats.h"
 #include "Stats/StatsMisc.h"
+#include "ZenServerInterface.h"
 #include <atomic>
 
 DEFINE_STAT(STAT_DDC_NumGets);
@@ -268,7 +270,11 @@ namespace UE::DerivedData::Private
  * Implementation of the derived data cache
  * This API is fully threadsafe
 **/
-class FDerivedDataCache final : public FDerivedDataCacheInterface, public ICache
+class FDerivedDataCache final
+	: public FDerivedDataCacheInterface
+	, public ICache
+	, public ICacheStoreMaintainer
+	, public IDDCCleanup
 {
 
 	/** 
@@ -449,7 +455,9 @@ public:
 	FDerivedDataCache()
 		: CurrentHandle(19248) // we will skip some potential handles to catch errors
 	{
-		FDerivedDataBackend::Get(); // we need to make sure this starts before we all us to start
+		FDerivedDataBackend::Get(); // we need to make sure this starts before we allow us to start
+
+		CacheStoreMaintainers = IModularFeatures::Get().GetModularFeatureImplementations<ICacheStoreMaintainer>(FeatureName);
 
 		GVerifyDDC = FParse::Param(FCommandLine::Get(), TEXT("VerifyDDC"));
 
@@ -459,8 +467,6 @@ public:
 	/** Destructor, flushes all sync tasks **/
 	~FDerivedDataCache()
 	{
-		FDDCCleanup::Shutdown();
-
 		WaitForQuiescence(true);
 		FScopeLock ScopeLock(&SynchronizationObject);
 		for (TMap<uint32,FAsyncTask<FBuildAsyncWorker>*>::TIterator It(PendingTasks); It; ++It)
@@ -716,9 +722,24 @@ public:
 		FDerivedDataBackend::Get().GetDirectories(OutResults);
 	}
 
+	PRAGMA_DISABLE_DEPRECATION_WARNINGS
 	virtual IDDCCleanup* GetCleanup() const override
 	{
-		return FDDCCleanup::Get();
+		return const_cast<IDDCCleanup*>(static_cast<const IDDCCleanup*>(this));
+	}
+	PRAGMA_ENABLE_DEPRECATION_WARNINGS
+
+	virtual bool IsFinished() const override
+	{
+		return IsIdle();
+	}
+
+	virtual void WaitBetweenDeletes(bool bWait) override
+	{
+		if (!bWait)
+		{
+			BoostPriority();
+		}
 	}
 
 	virtual void GatherUsageStats(TMap<FString, FDerivedDataCacheUsageStats>& UsageStats) override
@@ -798,7 +819,7 @@ public:
 		return FDerivedDataBackend::Get().GetRoot().Put(Records, Context, Policy, Owner, MoveTemp(OnComplete));
 	}
 
-	void  Get(
+	void Get(
 		TConstArrayView<FCacheKey> Keys,
 		FStringView Context,
 		FCacheRecordPolicy Policy,
@@ -820,6 +841,29 @@ public:
 	void CancelAll() final
 	{
 	}
+
+	ICacheStoreMaintainer& GetMaintainer() final
+	{
+		return *this;
+	}
+
+	// ICacheStoreMaintainer Interface
+
+	bool IsIdle() const final
+	{
+		return Algo::AllOf(CacheStoreMaintainers, &ICacheStoreMaintainer::IsIdle);
+	}
+
+	void BoostPriority() final
+	{
+		for (ICacheStoreMaintainer* Maintainer : CacheStoreMaintainers)
+		{
+			Maintainer->BoostPriority();
+		}
+	}
+
+private:
+	TArray<ICacheStoreMaintainer*> CacheStoreMaintainers;
 };
 
 ICache* CreateCache(FDerivedDataCacheInterface** OutLegacyCache)
diff --git a/Engine/Source/Developer/DerivedDataCache/Private/DerivedDataCacheRecord.cpp b/Engine/Source/Developer/DerivedDataCache/Private/DerivedDataCacheRecord.cpp
--- a/Engine/Source/Developer/DerivedDataCache/Private/DerivedDataCacheRecord.cpp
+++ b/Engine/Source/Developer/DerivedDataCache/Private/DerivedDataCacheRecord.cpp
@@ -296,11 +296,13 @@ FCbPackage FCacheRecord::Save() const
 	}
 	auto SavePayload = [&Package, &Writer](const FPayload& Payload)
 	{
-		FCbAttachment Attachment(Payload.GetData());
-		Package.AddAttachment(Attachment);
+		if (Payload.HasData())
+		{
+			Package.AddAttachment(FCbAttachment(Payload.GetData()));
+		}
 		Writer.BeginObject();
 		Writer.AddObjectId("Id"_ASV, Payload.GetId());
-		Writer.AddAttachment("RawHash"_ASV, Attachment);
+		Writer.AddBinaryAttachment("RawHash"_ASV, Payload.GetRawHash());
 		Writer.AddInteger("RawSize"_ASV, Payload.GetRawSize());
 		Writer.EndObject();
 	};
diff --git a/Engine/Source/Developer/DerivedDataCache/Private/FileSystemDerivedDataBackend.cpp b/Engine/Source/Developer/DerivedDataCache/Private/FileSystemDerivedDataBackend.cpp
--- a/Engine/Source/Developer/DerivedDataCache/Private/FileSystemDerivedDataBackend.cpp
+++ b/Engine/Source/Developer/DerivedDataCache/Private/FileSystemDerivedDataBackend.cpp
@@ -1,37 +1,44 @@
 // Copyright Epic Games, Inc. All Rights Reserved.
 
-
-#include "CoreMinimal.h"
-#include "Algo/AllOf.h"
 #include "Algo/Accumulate.h"
-#include "Misc/MessageDialog.h"
+#include "Algo/AllOf.h"
+#include "Algo/StableSort.h"
+#include "Algo/Transform.h"
+#include "Async/Async.h"
+#include "Async/TaskGraphInterfaces.h"
+#include "Containers/StaticBitArray.h"
+#include "DerivedDataBackendInterface.h"
+#include "DerivedDataCacheInterface.h"
+#include "DerivedDataCacheMaintainer.h"
+#include "DerivedDataCacheRecord.h"
+#include "DerivedDataCacheUsageStats.h"
+#include "DerivedDataChunk.h"
+#include "DerivedDataPayload.h"
+#include "Experimental/Async/LazyEvent.h"
+#include "Features/IModularFeatures.h"
+#include "HAL/Event.h"
 #include "HAL/FileManager.h"
-#include "Misc/CoreMisc.h"
+#include "HAL/Thread.h"
+#include "Hash/xxhash.h"
+#include "HashingArchiveProxy.h"
 #include "Misc/CommandLine.h"
+#include "Misc/CoreMisc.h"
 #include "Misc/FileHelper.h"
+#include "Misc/Guid.h"
+#include "Misc/MessageDialog.h"
 #include "Misc/Paths.h"
 #include "Misc/PathViews.h"
-#include "Misc/Guid.h"
 #include "Misc/ScopeExit.h"
 #include "Misc/ScopeLock.h"
 #include "Misc/StringBuilder.h"
-#include "ProfilingDebugging/CpuProfilerTrace.h"
+#include "ProfilingDebugging/CookStats.h"
 #include "ProfilingDebugging/CountersTrace.h"
+#include "ProfilingDebugging/CpuProfilerTrace.h"
 #include "Serialization/CompactBinary.h"
+#include "Serialization/CompactBinaryPackage.h"
 #include "Serialization/CompactBinaryValidation.h"
 #include "Serialization/CompactBinaryWriter.h"
-#include "Async/TaskGraphInterfaces.h"
-#include "Async/Async.h"
-
-#include "DerivedDataBackendInterface.h"
-#include "DerivedDataCacheInterface.h"
-#include "DerivedDataCacheRecord.h"
-#include "DerivedDataChunk.h"
-#include "DerivedDataPayload.h"
-#include "DDCCleanup.h"
-
-#include "ProfilingDebugging/CookStats.h"
-#include "DerivedDataCacheUsageStats.h"
+#include "Templates/Greater.h"
 
 #define MAX_BACKEND_KEY_LENGTH (120)
 #define MAX_BACKEND_NUMBERED_SUBFOLDER_LENGTH (9)
@@ -43,7 +50,7 @@
 #endif // PLATFORM_LINUX
 #define MAX_CACHE_EXTENTION_LEN (4)
 
-namespace UE::DerivedData::Backends
+namespace UE::DerivedData::CacheStore::FileSystem
 {
 
 TRACE_DECLARE_INT_COUNTER(FileSystemDDC_Exist, TEXT("FileSystemDDC Exist"));
@@ -55,6 +62,13 @@ TRACE_DECLARE_INT_COUNTER(FileSystemDDC_PutHit, TEXT("FileSystemDDC Put Hit"));
 TRACE_DECLARE_INT_COUNTER(FileSystemDDC_BytesRead, TEXT("FileSystemDDC Bytes Read"));
 TRACE_DECLARE_INT_COUNTER(FileSystemDDC_BytesWritten, TEXT("FileSystemDDC Bytes Written"));
 
+///////////////////////////////////////////////////////////////////////////////////////////////////////////////////////
+
+static const TCHAR GBucketsDirectoryName[] = TEXT("Buckets");
+static const TCHAR GContentDirectoryName[] = TEXT("Content");
+
+///////////////////////////////////////////////////////////////////////////////////////////////////////////////////////
+
 FString BuildPathForCacheKey(const TCHAR* CacheKey)
 {
 	FString Key = FString(CacheKey).ToUpper();
@@ -66,38 +80,544 @@ FString BuildPathForCacheKey(const TCHAR* CacheKey)
 	return HashPath / Key + TEXT(".udd");
 }
 
-void BuildPathForCacheObject(const FCacheKey& CacheKey, FStringBuilderBase& Path)
+void BuildPathForCacheRecord(const FCacheKey& CacheKey, FStringBuilderBase& Path)
+{
+	const FIoHash::ByteArray& Bytes = CacheKey.Hash.GetBytes();
+	Path.Appendf(TEXT("%s/%hs/%02x/%02x/"), GBucketsDirectoryName, CacheKey.Bucket.ToCString(), Bytes[0], Bytes[1]);
+	UE::String::BytesToHexLower(MakeArrayView(Bytes).RightChop(2), Path);
+	Path << TEXT(".udd"_SV);
+}
+
+void BuildPathForCacheContent(const FIoHash& RawHash, FStringBuilderBase& Path)
+{
+	const FIoHash::ByteArray& Bytes = RawHash.GetBytes();
+	Path.Appendf(TEXT("%s/%02x/%02x/"), GContentDirectoryName, Bytes[0], Bytes[1]);
+	UE::String::BytesToHexLower(MakeArrayView(Bytes).RightChop(2), Path);
+	Path << TEXT(".udd"_SV);
+}
+
+///////////////////////////////////////////////////////////////////////////////////////////////////////////////////////
+
+static uint64 RandFromGuid()
+{
+	const FGuid Guid = FGuid::NewGuid();
+	return FXxHash64::HashBuffer(&Guid, sizeof(FGuid)).Hash;
+}
+
+/** A LCG in which the modulus is a power of two where the exponent is the bit width of T. */
+template <typename T, T Modulus = 0>
+class TLinearCongruentialGenerator
+{
+	static_assert(!TIsSigned<T>::Value);
+	static_assert((Modulus & (Modulus - 1)) == 0, "Modulus must be a power of two.");
+
+public:
+	constexpr inline TLinearCongruentialGenerator(T InMultiplier, T InIncrement)
+		: Multiplier(InMultiplier)
+		, Increment(InIncrement)
+	{
+	}
+
+	constexpr inline T GetNext(T& Value)
+	{
+		Value = (Value * Multiplier + Increment) & (Modulus - 1);
+		return Value;
+	}
+
+private:
+	const T Multiplier;
+	const T Increment;
+};
+
+class FRandomStream
+{
+public:
+	inline explicit FRandomStream(uint32 Seed)
+		: Random(1103515245, 12345) // From ANSI C
+		, Value(Seed)
+	{
+	}
+
+	/** Returns a random value in [Min, Max). */
+	inline uint32 GetRandRange(uint32 Min, uint32 Max)
+	{
+		return Min + uint32((uint64(Max - Min) * Random.GetNext(Value)) >> 32);
+	}
+
+private:
+	TLinearCongruentialGenerator<uint32> Random;
+	uint32 Value;
+};
+
+template <uint32 Modulus, uint32 Count = Modulus>
+class TRandomOrder
+{
+	static_assert((Modulus & (Modulus - 1)) == 0 && Modulus > 16, "Modulus must be a power of two greater than 16.");
+	static_assert(Count > 0 && Count <= Modulus, "Count must be in the range (0, Modulus].");
+
+public:
+	inline explicit TRandomOrder(FRandomStream& Stream)
+		: Random(Stream.GetRandRange(0, Modulus / 16) * 8 + 5, 12345)
+		, First(Stream.GetRandRange(0, Count))
+		, Value(First)
+	{
+	}
+
+	inline uint32 GetFirst() const
+	{
+		return First;
+	}
+
+	inline uint32 GetNext()
+	{
+		if constexpr (Count < Modulus)
+		{
+			for (;;)
+			{
+				if (const uint32 Next = Random.GetNext(Value); Next < Count)
+				{
+					return Next;
+				}
+			}
+		}
+		return Random.GetNext(Value);
+	}
+
+private:
+	TLinearCongruentialGenerator<uint32, Modulus> Random;
+	uint32 First;
+	uint32 Value;
+};
+
+///////////////////////////////////////////////////////////////////////////////////////////////////////////////////////
+
+struct FFileSystemCacheStoreMaintainerParams
+{
+	/** Files older than this will be deleted. */
+	FTimespan MaxFileAge = FTimespan::FromDays(15.0);
+	/** Limits the number of files scanned in one second. */
+	uint32 MaxFileScanRate = MAX_uint32;
+	/** Limits the number of directories scanned in each cache bucket or content root. */
+	uint32 MaxDirectoryScanCount = MAX_uint32;
+	/** Minimum duration between the start of consecutive scans. */
+	FTimespan ScanFrequency = FTimespan::FromHours(1.0);
+	/** Time to wait after initialization before maintenance begins. */
+	FTimespan TimeToWaitAfterInit = FTimespan::FromMinutes(1.0);
+};
+
+class FFileSystemCacheStoreMaintainer final : public ICacheStoreMaintainer
+{
+public:
+	FFileSystemCacheStoreMaintainer(const FFileSystemCacheStoreMaintainerParams& Params, FStringView CachePath);
+	~FFileSystemCacheStoreMaintainer();
+
+	bool IsIdle() const final { return bIdle; }
+	void WaitForIdle() const { IdleEvent.Wait(); }
+	void BoostPriority() final;
+
+private:
+	void Loop();
+	void Scan();
+
+	void CreateContentRoot();
+	void CreateBucketRoots();
+	void ScanHashRoot(uint32 RootIndex);
+	TStaticBitArray<256> ScanHashDirectory(FStringBuilderBase& Path);
+
+	TStaticBitArray<10> ScanLegacyDirectory(FStringBuilderBase& Path);
+	void CreateLegacyRoot();
+	void ScanLegacyRoot();
+
+	void ResetRoots();
+
+	void ProcessFile(const TCHAR* Path, const FFileStatData& StatData);
+
+private:
+	struct FRoot;
+	struct FLegacyRoot;
+
+	FFileSystemCacheStoreMaintainerParams Params;
+	/** Path to the root of the cache store. */
+	FString CachePath;
+	/** True when maintenance is expected to exit. */
+	bool bExit = false;
+	/** True when there is no active maintenance scan. */
+	bool bIdle = false;
+	/** Ignore the file scan rate for one maintenance scan. */
+	bool bIgnoreFileScanRate = false;
+
+	uint32 ProcessCount = 0;
+	uint32 DeleteCount = 0;
+	uint64 DeleteSize = 0;
+
+	double BatchStartTime = 0.0;
+
+	IFileManager& FileManager = IFileManager::Get();
+
+	mutable FLazyEvent IdleEvent;
+	FEventRef WaitEvent;
+	FThread Thread;
+
+	TArray<TUniquePtr<FRoot>> Roots;
+	TUniquePtr<FLegacyRoot> LegacyRoot;
+	FRandomStream Random{uint32(RandFromGuid())};
+};
+
+struct FFileSystemCacheStoreMaintainer::FRoot
+{
+	inline FRoot(FStringView RootPath, FRandomStream& Stream)
+		: Order(Stream)
+	{
+		Path.Append(RootPath);
+	}
+
+	TStringBuilder<256> Path;
+	TRandomOrder<256 * 256> Order;
+	TStaticBitArray<256> ScannedLevel0;
+	TStaticBitArray<256> ExistsLevel0;
+	TStaticBitArray<256> ExistsLevel1[256];
+	uint32 DirectoryScanCount = 0;
+	bool bScannedRoot = false;
+};
+
+struct FFileSystemCacheStoreMaintainer::FLegacyRoot
+{
+	inline explicit FLegacyRoot(FRandomStream& Stream)
+		: Order(Stream)
+	{
+	}
+
+	TRandomOrder<1024, 1000> Order;
+	TStaticBitArray<10> ScannedLevel0;
+	TStaticBitArray<10> ScannedLevel1[10];
+	TStaticBitArray<10> ExistsLevel0;
+	TStaticBitArray<10> ExistsLevel1[10];
+	TStaticBitArray<10> ExistsLevel2[10][10];
+	uint32 DirectoryScanCount = 0;
+};
+
+FFileSystemCacheStoreMaintainer::FFileSystemCacheStoreMaintainer(
+	const FFileSystemCacheStoreMaintainerParams& InParams,
+	FStringView InCachePath)
+	: Params(InParams)
+	, CachePath(InCachePath)
+	, IdleEvent(EEventMode::ManualReset)
+	, WaitEvent(EEventMode::AutoReset)
+	, Thread(
+		TEXT("FileSystemCacheStoreMaintainer"),
+		[this] { Loop(); },
+		/*StackSize*/ 32 * 1024,
+		TPri_BelowNormal)
+{
+	IModularFeatures::Get().RegisterModularFeature(FeatureName, this);
+}
+
+FFileSystemCacheStoreMaintainer::~FFileSystemCacheStoreMaintainer()
+{
+	bExit = true;
+	IModularFeatures::Get().UnregisterModularFeature(FeatureName, this);
+	WaitEvent->Trigger();
+	Thread.Join();
+}
+
+void FFileSystemCacheStoreMaintainer::BoostPriority()
+{
+	bIgnoreFileScanRate = true;
+	WaitEvent->Trigger();
+}
+
+void FFileSystemCacheStoreMaintainer::Loop()
+{
+	WaitEvent->Wait(Params.TimeToWaitAfterInit, /*bIgnoreThreadIdleStats*/ true);
+
+	while (!bExit)
+	{
+		const FDateTime ScanStart = FDateTime::Now();
+		DeleteCount = 0;
+		DeleteSize = 0;
+		IdleEvent.Reset();
+		bIdle = false;
+		Scan();
+		bIdle = true;
+		IdleEvent.Trigger();
+		bIgnoreFileScanRate = false;
+		const FDateTime ScanEnd = FDateTime::Now();
+
+		UE_LOG(LogDerivedDataCache, Log,
+			TEXT("%s: Maintenance finished in %s and deleted %u file(s) with total size %" UINT64_FMT " MiB."),
+			*CachePath, *(ScanEnd - ScanStart).ToString(), DeleteCount, DeleteSize / 1024 / 1024);
+
+		if (bExit || Params.ScanFrequency.GetTotalDays() > 365.0)
+		{
+			break;
+		}
+
+		const FDateTime ScanTime = ScanStart + Params.ScanFrequency;
+		UE_LOG(LogDerivedDataCache, Verbose, TEXT("%s: Maintenance is paused until the next scan at %s."),
+			*CachePath, *ScanTime.ToString());
+		for (FDateTime Now = ScanEnd; !bExit && Now < ScanTime; Now = FDateTime::Now())
+		{
+			WaitEvent->Wait(ScanTime - Now, /*bIgnoreThreadIdleStats*/ true);
+		}
+	}
+}
+
+void FFileSystemCacheStoreMaintainer::Scan()
+{
+	CreateContentRoot();
+	CreateBucketRoots();
+	CreateLegacyRoot();
+
+	while (!bExit)
+	{
+		const uint32 RootCount = uint32(Roots.Num());
+		const uint32 TotalRootCount = uint32(RootCount + LegacyRoot.IsValid());
+		if (TotalRootCount == 0)
+		{
+			break;
+		}
+		if (const uint32 RootIndex = Random.GetRandRange(0, TotalRootCount); RootIndex < RootCount)
+		{
+			ScanHashRoot(RootIndex);
+		}
+		else
+		{
+			ScanLegacyRoot();
+		}
+	}
+
+	ResetRoots();
+}
+
+void FFileSystemCacheStoreMaintainer::CreateContentRoot()
+{
+	TStringBuilder<256> ContentPath;
+	FPathViews::Append(ContentPath, CachePath, GContentDirectoryName);
+	if (FileManager.DirectoryExists(*ContentPath))
+	{
+		Roots.Add(MakeUnique<FRoot>(ContentPath, Random));
+	}
+}
+
+void FFileSystemCacheStoreMaintainer::CreateBucketRoots()
+{
+	TStringBuilder<256> BucketsPath;
+	FPathViews::Append(BucketsPath, CachePath, GBucketsDirectoryName);
+	FileManager.IterateDirectoryStat(*BucketsPath, [this](const TCHAR* Path, const FFileStatData& StatData) -> bool
+	{
+		if (StatData.bIsDirectory)
+		{
+			Roots.Add(MakeUnique<FRoot>(Path, Random));
+		}
+		return !bExit;
+	});
+}
+
+void FFileSystemCacheStoreMaintainer::ScanHashRoot(uint32 RootIndex)
+{
+	FRoot& Root = *Roots[int32(RootIndex)];
+	const uint32 DirectoryIndex = Root.Order.GetNext();
+	const uint32 IndexLevel0 = DirectoryIndex / 256;
+	const uint32 IndexLevel1 = DirectoryIndex % 256;
+
+	bool bScanned = false;
+	ON_SCOPE_EXIT
+	{
+		if ((DirectoryIndex == Root.Order.GetFirst()) ||
+			(bScanned && ++Root.DirectoryScanCount >= Params.MaxDirectoryScanCount))
+		{
+			UE_LOG(LogDerivedDataCache, VeryVerbose, TEXT("%s: Maintenance finished scanning %s."), *CachePath, *Root.Path);
+			Roots.RemoveAt(int32(RootIndex));
+		}
+	};
+
+	if (!Root.bScannedRoot)
+	{
+		Root.ExistsLevel0 = ScanHashDirectory(Root.Path);
+		Root.bScannedRoot = true;
+	}
+
+	if (!Root.ExistsLevel0[IndexLevel0])
+	{
+		return;
+	}
+	if (!Root.ScannedLevel0[IndexLevel0])
+	{
+		TStringBuilder<256> Path;
+		Path.Appendf(TEXT("%s/%02x"), *Root.Path, IndexLevel0);
+		Root.ExistsLevel1[IndexLevel0] = ScanHashDirectory(Path);
+		Root.ScannedLevel0[IndexLevel0] = true;
+	}
+
+	if (!Root.ExistsLevel1[IndexLevel0][IndexLevel1])
+	{
+		return;
+	}
+
+	TStringBuilder<256> Path;
+	Path.Appendf(TEXT("%s/%02x/%02x"), *Root.Path, IndexLevel0, IndexLevel1);
+	FileManager.IterateDirectoryStat(*Path, [this](const TCHAR* const Path, const FFileStatData& StatData) -> bool
+	{
+		ProcessFile(Path, StatData);
+		return !bExit;
+	});
+
+	bScanned = true;
+}
+
+TStaticBitArray<256> FFileSystemCacheStoreMaintainer::ScanHashDirectory(FStringBuilderBase& Path)
+{
+	TStaticBitArray<256> Exists;
+	FileManager.IterateDirectoryStat(*Path, [this, &Exists](const TCHAR* ChildPath, const FFileStatData& StatData) -> bool
+	{
+		FStringView ChildView = FPathViews::GetCleanFilename(ChildPath);
+		if (StatData.bIsDirectory && ChildView.Len() == 2 && Algo::AllOf(ChildView, FChar::IsHexDigit))
+		{
+			uint8 Byte;
+			if (String::HexToBytes(ChildView, &Byte) == 1)
+			{
+				Exists[Byte] = true;
+			}
+		}
+		return !bExit;
+	});
+	return Exists;
+}
+
+TStaticBitArray<10> FFileSystemCacheStoreMaintainer::ScanLegacyDirectory(FStringBuilderBase& Path)
 {
-	uint32 Hash;
-	FMemory::Memcpy(&Hash, &CacheKey.Hash, sizeof(uint32));
-	Path.Appendf(TEXT("%1u/%1u/%1u/"), (Hash / 100) % 10, (Hash / 10) % 10, Hash % 10);
-	Path << CacheKey.Bucket << TEXT('_') << CacheKey.Hash << TEXT(".udd"_SV);
+	TStaticBitArray<10> Exists;
+	const bool bDeleteFiles = CachePath.Len() < Path.Len();
+	FileManager.IterateDirectoryStat(*Path, [this, &Exists, bDeleteFiles](const TCHAR* ChildPath, const FFileStatData& StatData) -> bool
+	{
+		FStringView ChildView = FPathViews::GetCleanFilename(ChildPath);
+		if (StatData.bIsDirectory && ChildView.Len() == 1 && Algo::AllOf(ChildView, FChar::IsDigit))
+		{
+			Exists[FChar::ConvertCharDigitToInt(ChildView[0])] = true;
+		}
+		return !bExit;
+	});
+	return Exists;
 }
 
-void BuildPathForCachePayload(const FCacheKey& CacheKey, const FIoHash& RawHash, FStringBuilderBase& Path)
+void FFileSystemCacheStoreMaintainer::CreateLegacyRoot()
 {
-	uint32 Hash;
-	FMemory::Memcpy(&Hash, &RawHash, sizeof(uint32));
-	Path.Appendf(TEXT("%1u/%1u/%1u/"), (Hash / 100) % 10, (Hash / 10) % 10, Hash % 10);
-	Path << CacheKey.Bucket << TEXT("_CAS_"_SV) << RawHash << TEXT(".udd"_SV);
+	TStringBuilder<256> Path;
+	FPathViews::Append(Path, CachePath);
+	TStaticBitArray<10> Exists = ScanLegacyDirectory(Path);
+	if (Exists.FindFirstSetBit() != INDEX_NONE)
+	{
+		LegacyRoot = MakeUnique<FLegacyRoot>(Random);
+		LegacyRoot->ExistsLevel0 = Exists;
+	}
+}
+
+void FFileSystemCacheStoreMaintainer::ScanLegacyRoot()
+{
+	FLegacyRoot& Root = *LegacyRoot;
+	const uint32 DirectoryIndex = Root.Order.GetNext();
+	const int32 IndexLevel0 = int32(DirectoryIndex / 100) % 10;
+	const int32 IndexLevel1 = int32(DirectoryIndex / 10) % 10;
+	const int32 IndexLevel2 = int32(DirectoryIndex / 1) % 10;
+
+	bool bScanned = false;
+	ON_SCOPE_EXIT
+	{
+		if ((DirectoryIndex == Root.Order.GetFirst()) ||
+			(bScanned && ++Root.DirectoryScanCount >= Params.MaxDirectoryScanCount))
+		{
+			LegacyRoot.Reset();
+		}
+	};
+
+	if (!Root.ExistsLevel0[IndexLevel0])
+	{
+		return;
+	}
+	if (!Root.ScannedLevel0[IndexLevel0])
+	{
+		TStringBuilder<256> Path;
+		FPathViews::Append(Path, CachePath, IndexLevel0);
+		Root.ExistsLevel1[IndexLevel0] = ScanLegacyDirectory(Path);
+		Root.ScannedLevel0[IndexLevel0] = true;
+	}
+
+	if (!Root.ExistsLevel1[IndexLevel0][IndexLevel1])
+	{
+		return;
+	}
+	if (!Root.ScannedLevel1[IndexLevel0][IndexLevel1])
+	{
+		TStringBuilder<256> Path;
+		FPathViews::Append(Path, CachePath, IndexLevel0, IndexLevel1);
+		Root.ExistsLevel2[IndexLevel0][IndexLevel1] = ScanLegacyDirectory(Path);
+		Root.ScannedLevel1[IndexLevel0][IndexLevel1] = true;
+	}
+
+	if (!Root.ExistsLevel2[IndexLevel0][IndexLevel1][IndexLevel2])
+	{
+		return;
+	}
+
+	TStringBuilder<256> Path;
+	FPathViews::Append(Path, CachePath, IndexLevel0, IndexLevel1, IndexLevel2);
+	FileManager.IterateDirectoryStat(*Path, [this](const TCHAR* const Path, const FFileStatData& StatData) -> bool
+	{
+		ProcessFile(Path, StatData);
+		return !bExit;
+	});
+
+	bScanned = true;
 }
 
-/**
- * Helper function to get the value of parsed bool as the return value
- **/
-bool GetParsedBool(const TCHAR* Stream, const TCHAR* Match) 
+void FFileSystemCacheStoreMaintainer::ResetRoots()
 {
-	bool bValue = 0;
-	FParse::Bool(Stream, Match, bValue);
-	return bValue;
+	Roots.Empty();
+	LegacyRoot.Reset();
 }
 
-/** Delete the old files in a directory **/
-void DeleteOldFiles(const TCHAR* Directory, int Age)
+void FFileSystemCacheStoreMaintainer::ProcessFile(const TCHAR* const Path, const FFileStatData& StatData)
 {
-	// @todo(agrant) the original implementation of this did nothing. Do we need this?
+	if (StatData.bIsDirectory)
+	{
+		return;
+	}
+
+	if (StatData.ModificationTime + Params.MaxFileAge < FDateTime::UtcNow())
+	{
+		++DeleteCount;
+		DeleteSize += StatData.FileSize > 0 ? uint64(StatData.FileSize) : 0;
+		if (FileManager.Delete(Path, /*bRequireExists*/ false, /*bEvenReadOnly*/ false, /*bQuiet*/ true))
+		{
+			UE_LOG(LogDerivedDataCache, VeryVerbose,
+				TEXT("%s: Maintenance deleted file %s that was last modified at %s."),
+				*CachePath, Path, *StatData.ModificationTime.ToIso8601());
+		}
+		else
+		{
+			UE_LOG(LogDerivedDataCache, Verbose,
+				TEXT("%s: Maintenance failed to delete file %s that was last modified at %s."),
+				*CachePath, Path, *StatData.ModificationTime.ToIso8601());
+		}
+	}
+
+	if (!bExit && !bIgnoreFileScanRate && Params.MaxFileScanRate && ++ProcessCount % Params.MaxFileScanRate == 0)
+	{
+		const double BatchEndTime = FPlatformTime::Seconds();
+		if (const double BatchWaitTime = 1.0 - (BatchEndTime - BatchStartTime); BatchWaitTime > 0.0)
+		{
+			WaitEvent->Wait(FTimespan::FromSeconds(BatchWaitTime), /*bIgnoreThreadIdleStats*/ true);
+			BatchStartTime = FPlatformTime::Seconds();
+		}
+		else
+		{
+			BatchStartTime = BatchEndTime;
+		}
+	}
 }
 
+///////////////////////////////////////////////////////////////////////////////////////////////////////////////////////
+
 /** 
  * Cache server that uses the OS filesystem
  * The entire API should be callable from any thread (except the singleton can be assumed to be called at least once before concurrent access).
@@ -105,22 +625,17 @@ void DeleteOldFiles(const TCHAR* Directory, int Age)
 class FFileSystemDerivedDataBackend : public FDerivedDataBackendInterface
 {
 public:
-	/**
-	 * Constructor that should only be called once by the singleton, grabs the cache path from the ini
-	 * @param InCacheDirectory	directory to store the cache in
-	 * @param bForceReadOnly	if true, do not attempt to write to this cache
-	*/
-	FFileSystemDerivedDataBackend(const TCHAR* InCacheDirectory, const TCHAR* InParams, const TCHAR* InAccessLogFileName)
-		: CachePath(InCacheDirectory)
+	FFileSystemDerivedDataBackend(const TCHAR* InCachePath, const TCHAR* InParams, const TCHAR* InAccessLogFileName)
+		: CachePath(InCachePath)
 		, SpeedClass(ESpeedClass::Unknown)
+		, bDisabled(false)
 		, bReadOnly(false)
 		, bTouch(false)
 		, bPurgeTransient(false)
-		, DaysToDeleteUnusedFiles(15)
-		, bDisabled(false)
+		, DaysToDeleteUnusedFiles(15.0)
 		, TotalEstimatedBuildTime(0)
 	{
-		// If we find a platform that has more stingent limits, this needs to be rethought.
+		// If we find a platform that has more stringent limits, this needs to be rethought.
 		checkf(MAX_BACKEND_KEY_LENGTH + MAX_CACHE_DIR_LEN + MAX_BACKEND_NUMBERED_SUBFOLDER_LENGTH + MAX_CACHE_EXTENTION_LEN < FPlatformMisc::GetMaxPathLength(),
 			TEXT("Not enough room left for cache keys in max path."));
 
@@ -128,36 +643,18 @@ public:
 		FPaths::NormalizeFilename(CachePath);
 
 		// Params that override our instance defaults
-		bReadOnly = GetParsedBool(InParams, TEXT("ReadOnly="));
-		bTouch = GetParsedBool(InParams, TEXT("Touch="));
-		bPurgeTransient = GetParsedBool(InParams, TEXT("PurgeTransient="));
+		FParse::Bool(InParams, TEXT("ReadOnly="), bReadOnly);
+		FParse::Bool(InParams, TEXT("Touch="), bTouch);
+		FParse::Bool(InParams, TEXT("PurgeTransient="), bPurgeTransient);
 		FParse::Value(InParams, TEXT("UnusedFileAge="), DaysToDeleteUnusedFiles);
+		FParse::Value(InParams, TEXT("MaxRecordSizeKB="), MaxRecordSizeKB);
+		FParse::Value(InParams, TEXT("MaxValueSizeKB="), MaxValueSizeKB);
 
-		// Params that are used when setting up our path
-		const bool bClean = GetParsedBool(InParams, TEXT("Clean="));
-		const bool bFlush = GetParsedBool(InParams, TEXT("Flush="));
-	
-
-		// These are used to determine if we kick off a worker to cleanup our cache
-		bool bDeleteUnused = true; // On by default
-		int32 MaxFoldersToClean = -1;
-		int32 MaxFileChecksPerSec = -1;
-
-		FParse::Bool(InParams, TEXT("DeleteUnused="), bDeleteUnused);
-
-		if (bDeleteUnused)
+		// Flush the cache if requested.
+		bool bFlush = false;
+		if (!bReadOnly && FParse::Bool(InParams, TEXT("Flush="), bFlush) && bFlush)
 		{
-			FParse::Value(InParams, TEXT("FoldersToClean="), MaxFoldersToClean);
-			FParse::Value(InParams, TEXT("MaxFileChecksPerSec="), MaxFileChecksPerSec);
-		}
-
-		if (bFlush)
-		{
-			IFileManager::Get().DeleteDirectory(*(CachePath / TEXT("")), false, true);
-		}
-		else if (bClean)
-		{
-			DeleteOldFiles(InCacheDirectory, DaysToDeleteUnusedFiles);
+			IFileManager::Get().DeleteDirectory(*(CachePath / TEXT("")), /*bRequireExists*/ false, /*bTree*/ true);
 		}
 
 		// check latency and speed. Read values should always be valid
@@ -244,9 +741,49 @@ public:
 					UE_LOG(LogDerivedDataCache, Display, TEXT("Files in %s will be touched."), *CachePath);
 				}
 
-				if (bDeleteUnused && !FParse::Param(FCommandLine::Get(), TEXT("NODDCCLEANUP")) && FDDCCleanup::Get())
+				bool bClean = false;
+				bool bDeleteUnused = true;
+				FParse::Bool(InParams, TEXT("Clean="), bClean);
+				FParse::Bool(InParams, TEXT("DeleteUnused="), bDeleteUnused);
+				bDeleteUnused = bDeleteUnused && !FParse::Param(FCommandLine::Get(), TEXT("NODDCCLEANUP"));
+
+				if (bClean || bDeleteUnused)
 				{
-					FDDCCleanup::Get()->AddFilesystem(CachePath, DaysToDeleteUnusedFiles, MaxFoldersToClean, MaxFileChecksPerSec);
+					FFileSystemCacheStoreMaintainerParams MaintainerParams;
+					MaintainerParams.MaxFileAge = FTimespan::FromDays(DaysToDeleteUnusedFiles);
+					if (bDeleteUnused)
+					{
+						if (!FParse::Value(InParams, TEXT("MaxFileChecksPerSec="), MaintainerParams.MaxFileScanRate))
+						{
+							int32 MaxFileScanRate;
+							if (GConfig->GetInt(TEXT("DDCCleanup"), TEXT("MaxFileChecksPerSec"), MaxFileScanRate, GEngineIni))
+							{
+								MaintainerParams.MaxFileScanRate = uint32(MaxFileScanRate);
+							}
+						}
+						FParse::Value(InParams, TEXT("FoldersToClean="), MaintainerParams.MaxDirectoryScanCount);
+					}
+					else
+					{
+						MaintainerParams.ScanFrequency = FTimespan::FromDays(500.0);
+					}
+					double TimeToWaitAfterInit;
+					if (bClean)
+					{
+						MaintainerParams.TimeToWaitAfterInit = FTimespan::Zero();
+					}
+					else if (GConfig->GetDouble(TEXT("DDCCleanup"), TEXT("TimeToWaitAfterInit"), TimeToWaitAfterInit, GEngineIni))
+					{
+						MaintainerParams.TimeToWaitAfterInit = FTimespan::FromSeconds(TimeToWaitAfterInit);
+					}
+
+					Maintainer = MakeUnique<FFileSystemCacheStoreMaintainer>(MaintainerParams, CachePath);
+
+					if (bClean)
+					{
+						Maintainer->BoostPriority();
+						Maintainer->WaitForIdle();
+					}
 				}
 			}
 			
@@ -309,8 +846,6 @@ public:
 
 		// measure latency by checking for the presence of all these files. We'll also track which don't exist..
 		const double StatStartTime = FPlatformTime::Seconds();
-		int TotalStatChecks = 0;
-		FDateTime CurrentTime = FDateTime::Now();
  		for (auto& KV : TestFileEntries)
 		{
 			FFileStatData StatData = IFileManager::Get().GetStatData(*KV.Key);
@@ -511,24 +1046,24 @@ public:
 			}
 		}
 
-		void Append(const FCacheKey& CacheKey, FStringView Path)
+		void Append(const FIoHash& RawHash, FStringView Path)
 		{
 			FScopeLock Lock(&CriticalSection);
 
 			bool bIsAlreadyInSet = false;
-			ObjectKeys.FindOrAdd(CacheKey, &bIsAlreadyInSet);
+			ContentKeys.FindOrAdd(RawHash, &bIsAlreadyInSet);
 			if (!bIsAlreadyInSet)
 			{
 				AppendPath(Path);
 			}
 		}
 
-		void Append(const FCacheKey& CacheKey, const FPayloadId& Id, FStringView Path)
+		void Append(const FCacheKey& CacheKey, FStringView Path)
 		{
 			FScopeLock Lock(&CriticalSection);
 
 			bool bIsAlreadyInSet = false;
-			PayloadKeys.FindOrAdd(MakeTuple(CacheKey, Id), &bIsAlreadyInSet);
+			RecordKeys.FindOrAdd(CacheKey, &bIsAlreadyInSet);
 			if (!bIsAlreadyInSet)
 			{
 				AppendPath(Path);
@@ -550,8 +1085,8 @@ public:
 		FString BasePath;
 		FCriticalSection CriticalSection;
 		TSet<FString> CacheKeys;
-		TSet<FCacheKey> ObjectKeys;
-		TSet<TTuple<FCacheKey, FPayloadId>> PayloadKeys;
+		TSet<FIoHash> ContentKeys;
+		TSet<FCacheKey> RecordKeys;
 	};
 	
 	
@@ -595,7 +1130,7 @@ public:
 
 			// Update file timestamp to prevent it from being deleted by DDC Cleanup.
 			if (bTouch || 
-				 (!bReadOnly && (FDateTime::UtcNow() - TimeStamp).GetDays() > (DaysToDeleteUnusedFiles / 4)))
+				 (!bReadOnly && (FDateTime::UtcNow() - TimeStamp).GetTotalDays() > (DaysToDeleteUnusedFiles / 4)))
 			{
 				IFileManager::Get().SetTimeStamp(*Filename, FDateTime::UtcNow());
 			}
@@ -656,7 +1191,7 @@ public:
 					FDateTime TimeStamp = FileStat.ModificationTime;
 					// Update file timestamp to prevent it from being deleted by DDC Cleanup.
 					if (bTouch ||
-						(!bReadOnly && (FDateTime::UtcNow() - TimeStamp).GetDays() > (DaysToDeleteUnusedFiles / 4)))
+						(!bReadOnly && (FDateTime::UtcNow() - TimeStamp).GetTotalDays() > (DaysToDeleteUnusedFiles / 4)))
 					{
 						IFileManager::Get().SetTimeStamp(*Filename, FDateTime::UtcNow());
 					}
@@ -932,9 +1467,8 @@ public:
 			TRACE_CPUPROFILER_EVENT_SCOPE(FileSystemDDC_Get);
 			TRACE_COUNTER_INCREMENT(FileSystemDDC_Get);
 			COOK_STAT(auto Timer = UsageStats.TimeGet());
-			const bool bAlwaysLoadInlineData = false;
-			EStatus Status = EStatus::Error;
-			if (FOptionalCacheRecord Record = GetCacheRecord(Key, Context, Policy, bAlwaysLoadInlineData, Status))
+			EStatus Status = EStatus::Ok;
+			if (FOptionalCacheRecord Record = GetCacheRecord(Key, Context, Policy, Status))
 			{
 				UE_LOG(LogDerivedDataCache, Verbose, TEXT("%s: Cache hit for %s from '%.*s'"),
 					*CachePath, *WriteToString<96>(Key), Context.Len(), Context.GetData());
@@ -966,7 +1500,6 @@ public:
 		SortedChunks.StableSort(TChunkLess());
 
 		FOptionalCacheRecord Record;
-		EStatus RecordStatus = EStatus::Error;
 		for (const FCacheChunkRequest& Chunk : SortedChunks)
 		{
 			constexpr ECachePolicy SkipFlag = ECachePolicy::SkipValue;
@@ -976,17 +1509,17 @@ public:
 			COOK_STAT(auto Timer = bExistsOnly ? UsageStats.TimeProbablyExists() : UsageStats.TimeGet());
 			if (!Record || Record.Get().GetKey() != Chunk.Key)
 			{
-				const bool bAlwaysLoadInlineData = true;
 				FCacheRecordPolicyBuilder PolicyBuilder(ECachePolicy::None);
 				const ECachePolicy RecordSkipFlags = bExistsOnly ? ECachePolicy::SkipData : ECachePolicy::None;
 				PolicyBuilder.AddPayloadPolicy(Chunk.Id, Chunk.Policy | RecordSkipFlags);
-				Record = GetCacheRecord(Chunk.Key, Context, PolicyBuilder.Build(), bAlwaysLoadInlineData, RecordStatus);
+				Record = GetCacheRecordOnly(Chunk.Key, Context, PolicyBuilder.Build());
 			}
-			if (Record && RecordStatus == EStatus::Ok)
+			if (Record)
 			{
-				EStatus PayloadStatus = EStatus::Error;
-				const FPayload& RecordPayload = Record.Get().GetPayload(Chunk.Id);
-				if (FPayload Payload = GetCachePayload(Chunk.Key, Context, Chunk.Policy, SkipFlag, RecordPayload, PayloadStatus))
+				EStatus PayloadStatus = EStatus::Ok;
+				FPayload Payload = Record.Get().GetPayload(Chunk.Id);
+				GetCachePayload(Chunk.Key, Context, Chunk.Policy, SkipFlag, Payload, PayloadStatus);
+				if (Payload)
 				{
 					const uint64 RawSize = FMath::Min(Payload.GetRawSize(), Chunk.RawSize);
 					UE_LOG(LogDerivedDataCache, Verbose, TEXT("%s: Cache hit for %s from '%.*s'"),
@@ -1045,7 +1578,7 @@ private:
 
 		const FCacheKey& Key = Record.GetKey();
 
-		// Skip the request if querying the cache is disabled.
+		// Skip the request if storing to the cache is disabled.
 		if (!EnumHasAnyFlags(Policy, SpeedClass == ESpeedClass::Local ? ECachePolicy::StoreLocal : ECachePolicy::StoreRemote))
 		{
 			UE_LOG(LogDerivedDataCache, VeryVerbose, TEXT("%s: Skipped put of %s from '%.*s' due to cache policy"),
@@ -1060,48 +1593,74 @@ private:
 			return false;
 		}
 
-		// Save the payloads and build the cache object.
-		const TConstArrayView<FPayload> Attachments = Record.GetAttachmentPayloads();
-		TCbWriter<512> Writer;
-		Writer.BeginObject();
-		if (const FCbObject& Meta = Record.GetMeta())
+		// Check if there is an existing record package.
+		bool bRecordExists;
+		FCbPackage ExistingPackage;
+		TStringBuilder<256> Path;
+		BuildCacheRecordPath(Key, Path);
+		if (EnumHasAllFlags(Policy, ECachePolicy::SkipValue | ECachePolicy::SkipAttachments))
 		{
-			Writer.AddObject("Meta"_ASV, Meta);
-			Writer.AddHash("MetaHash"_ASV, Meta.GetHash());
+			bRecordExists = FileExists(Path);
 		}
-		if (const FPayload& Value = Record.GetValuePayload())
+		else
 		{
-			Writer.SetName("Value"_ASV);
-			if (!PutCachePayload(Key, Context, Value, Writer, /*bStoreInline*/ Attachments.IsEmpty()))
+			FSharedBuffer Buffer = LoadFile(Path, Context);
+			FCbFieldIterator It = FCbFieldIterator::MakeRange(Buffer);
+			bRecordExists = ExistingPackage.TryLoad(It);
+		}
+
+		// Save the record to a package and remove attachments that will be stored externally.
+		FCbPackage Package = Record.Save();
+		TArray<FCompressedBuffer, TInlineAllocator<8>> ExternalContent;
+		if (ExistingPackage)
+		{
+			// Mirror the existing internal/external attachment storage.
+			TArray<FCompressedBuffer, TInlineAllocator<8>> AllContent;
+			Algo::Transform(Package.GetAttachments(), AllContent, &FCbAttachment::AsCompressedBinary);
+			for (FCompressedBuffer& Content : AllContent)
 			{
-				return false;
+				const FIoHash RawHash = Content.GetRawHash();
+				if (!ExistingPackage.FindAttachment(RawHash))
+				{
+					Package.RemoveAttachment(RawHash);
+					ExternalContent.Add(MoveTemp(Content));
+				}
 			}
 		}
-		if (!Attachments.IsEmpty())
+		else
 		{
-			Writer.BeginArray("Attachments"_ASV);
-			for (const FPayload& Attachment : Attachments)
+			// Remove the largest attachments from the package until it fits within the size limits.
+			TArray<FCompressedBuffer, TInlineAllocator<8>> AllContent;
+			Algo::Transform(Package.GetAttachments(), AllContent, &FCbAttachment::AsCompressedBinary);
+			uint64 TotalSize = Algo::TransformAccumulate(AllContent, &FCompressedBuffer::GetCompressedSize, uint64(0));
+			const uint64 MaxSize = (AllContent.Num() == 1 ? MaxValueSizeKB : MaxRecordSizeKB) * 1024;
+			if (TotalSize > MaxSize)
 			{
-				if (!PutCachePayload(Key, Context, Attachment, Writer, /*bStoreInline*/ false))
+				Algo::StableSortBy(AllContent, &FCompressedBuffer::GetCompressedSize, TGreater<>());
+				for (FCompressedBuffer& Content : AllContent)
 				{
-					return false;
+					const uint64 CompressedSize = Content.GetCompressedSize();
+					Package.RemoveAttachment(Content.GetRawHash());
+					ExternalContent.Add(MoveTemp(Content));
+					TotalSize -= CompressedSize;
+					if (TotalSize <= MaxSize)
+					{
+						break;
+					}
 				}
 			}
-			Writer.EndArray();
 		}
-		Writer.EndObject();
 
-		// Save the record to storage.
-		TStringBuilder<256> Path;
-		BuildCacheObjectPath(Key, Path);
-		if (!FileExists(Path))
+		// Save the external content to storage.
+		for (FCompressedBuffer& Content : ExternalContent)
 		{
-			FUniqueBuffer Buffer = FUniqueBuffer::Alloc(Writer.GetSaveSize());
-			Writer.Save(Buffer);
-			if (!SaveFile(Path, FCompositeBuffer(Buffer.MoveToShared()), Context))
-			{
-				return false;
-			}
+			PutCacheContent(Content, Context);
+		}
+
+		// Save the record package to storage.
+		if (!bRecordExists && !SaveFile(Path, Context, [&Package](FArchive& Ar) { Package.Save(Ar); }))
+		{
+			return false;
 		}
 
 		if (AccessLogWriter)
@@ -1112,15 +1671,11 @@ private:
 		return true;
 	}
 
-	FOptionalCacheRecord GetCacheRecord(
+	FOptionalCacheRecord GetCacheRecordOnly(
 		const FCacheKey& Key,
 		const FStringView Context,
-		const FCacheRecordPolicy& Policy,
-		const bool bAlwaysLoadInlineData,
-		EStatus& OutStatus)
+		const FCacheRecordPolicy& Policy)
 	{
-		OutStatus = EStatus::Error;
-
 		if (!IsUsable())
 		{
 			UE_LOG(LogDerivedDataCache, VeryVerbose,
@@ -1129,10 +1684,8 @@ private:
 			return FOptionalCacheRecord();
 		}
 
-		const ECachePolicy RecordPolicy = Policy.GetRecordPolicy();
-
 		// Skip the request if querying the cache is disabled.
-		if (!EnumHasAnyFlags(RecordPolicy, SpeedClass == ESpeedClass::Local ? ECachePolicy::QueryLocal : ECachePolicy::QueryRemote))
+		if (!EnumHasAnyFlags(Policy.GetRecordPolicy(), SpeedClass == ESpeedClass::Local ? ECachePolicy::QueryLocal : ECachePolicy::QueryRemote))
 		{
 			UE_LOG(LogDerivedDataCache, VeryVerbose, TEXT("%s: Skipped get of %s from '%.*s' due to cache policy"),
 				*CachePath, *WriteToString<96>(Key), Context.Len(), Context.GetData());
@@ -1148,7 +1701,7 @@ private:
 
 		// Request the record from storage.
 		TStringBuilder<256> Path;
-		BuildCacheObjectPath(Key, Path);
+		BuildCacheRecordPath(Key, Path);
 		FSharedBuffer Buffer = LoadFile(Path, Context);
 		if (Buffer.IsNull())
 		{
@@ -1157,7 +1710,7 @@ private:
 			return FOptionalCacheRecord();
 		}
 
-		// Delete the record from storage if is corrupted.
+		// Delete the record from storage if it is invalid.
 		bool bDeleteCacheObject = true;
 		ON_SCOPE_EXIT
 		{
@@ -1167,71 +1720,34 @@ private:
 			}
 		};
 
-		// Validate that the record can be read as compact binary without crashing.
-		if (ValidateCompactBinaryRange(Buffer, ECbValidateMode::Default) != ECbValidateError::None)
+		// Validate that the record can be read as a compact binary package without crashing.
+		if (ValidateCompactBinaryPackage(Buffer, ECbValidateMode::Default | ECbValidateMode::Package) != ECbValidateError::None)
 		{
-			UE_LOG(LogDerivedDataCache, Display, TEXT("%s: Cache miss with corrupted record for %s from '%.*s'"),
+			UE_LOG(LogDerivedDataCache, Display, TEXT("%s: Cache miss with invalid package for %s from '%.*s'"),
 				*CachePath, *WriteToString<96>(Key), Context.Len(), Context.GetData());
 			return FOptionalCacheRecord();
 		}
 
-		const FCbObject RecordObject(MoveTemp(Buffer));
-		FCacheRecordBuilder RecordBuilder(Key);
-
-		if (!EnumHasAnyFlags(RecordPolicy, ECachePolicy::SkipMeta))
-		{
-			if (FCbFieldView MetaHash = RecordObject.FindView("MetaHash"_ASV))
-			{
-				if (FCbObject Meta = RecordObject["Meta"_ASV].AsObject(); Meta.GetHash() == MetaHash.AsHash())
-				{
-					RecordBuilder.SetMeta(MoveTemp(Meta));
-				}
-				else
-				{
-					UE_LOG(LogDerivedDataCache, Display,
-						TEXT("%s: Cache miss with corrupted metadata for %s from '%.*s'"),
-						*CachePath, *WriteToString<96>(Key), Context.Len(), Context.GetData());
-					return FOptionalCacheRecord();
-				}
-			}
-		}
-
-		// Disable deletion now that the record is validated.
-		bDeleteCacheObject = false;
-		OutStatus = EStatus::Ok;
-
-		if (FCbObject ValueObject = RecordObject["Value"_ASV].AsObject())
+		// Load the record from the package.
+		FOptionalCacheRecord Record;
 		{
-			EStatus PayloadStatus = EStatus::Error;
-			FPayload Payload = GetCachePayload(Key, Context, Policy,
-				ECachePolicy::SkipValue, ValueObject, bAlwaysLoadInlineData, PayloadStatus);
-			if (PayloadStatus == EStatus::Error)
-			{
-				OutStatus = EStatus::Error;
-			}
-			if (Payload.IsNull())
+			FCbPackage Package;
+			if (FCbFieldIterator It = FCbFieldIterator::MakeRange(Buffer); !Package.TryLoad(It))
 			{
+				UE_LOG(LogDerivedDataCache, Display, TEXT("%s: Cache miss with package load failure for %s from '%.*s'"),
+					*CachePath, *WriteToString<96>(Key), Context.Len(), Context.GetData());
 				return FOptionalCacheRecord();
 			}
-			RecordBuilder.SetValue(MoveTemp(Payload));
-		}
-
-		for (FCbField AttachmentField : RecordObject["Attachments"_ASV])
-		{
-			EStatus PayloadStatus = EStatus::Error;
-			FPayload Payload = GetCachePayload(Key, Context, Policy,
-				ECachePolicy::SkipAttachments, AttachmentField.AsObject(), bAlwaysLoadInlineData, PayloadStatus);
-			if (PayloadStatus == EStatus::Error)
-			{
-				OutStatus = EStatus::Error;
-			}
-			if (Payload.IsNull())
+			Record = FCacheRecord::Load(Package);
+			if (Record.IsNull())
 			{
+				UE_LOG(LogDerivedDataCache, Display, TEXT("%s: Cache miss with record load failure for %s from '%.*s'"),
+					*CachePath, *WriteToString<96>(Key), Context.Len(), Context.GetData());
 				return FOptionalCacheRecord();
 			}
-			RecordBuilder.AddAttachment(MoveTemp(Payload));
 		}
 
+		// Disable deletion now that the record is loaded and validated.
 		bDeleteCacheObject = false;
 
 		if (AccessLogWriter)
@@ -1239,204 +1755,173 @@ private:
 			AccessLogWriter->Append(Key, Path);
 		}
 
-		return RecordBuilder.Build();
+		return Record.Get();
 	}
 
-	bool PutCachePayload(
+	FOptionalCacheRecord GetCacheRecord(
 		const FCacheKey& Key,
 		const FStringView Context,
-		const FPayload& Payload,
-		FCbWriter& Writer,
-		const bool bStoreInline) const
+		const FCacheRecordPolicy& Policy,
+		EStatus& OutStatus)
 	{
-		const FIoHash& RawHash = Payload.GetRawHash();
-		const bool bHasData = Payload.HasData();
+		FOptionalCacheRecord Record = GetCacheRecordOnly(Key, Context, Policy);
+		if (Record.IsNull())
+		{
+			OutStatus = EStatus::Error;
+			return Record;
+		}
+
+		OutStatus = EStatus::Ok;
 
-		if (bHasData && !bStoreInline)
+		FCacheRecordBuilder RecordBuilder(Key);
+
+		if (!EnumHasAnyFlags(Policy.GetRecordPolicy(), ECachePolicy::SkipMeta))
 		{
-			TStringBuilder<256> Path;
-			BuildCachePayloadPath(Key, RawHash, Path);
-			if (!FileExists(Path))
+			RecordBuilder.SetMeta(FCbObject(Record.Get().GetMeta()));
+		}
+
+		if (FPayload Payload = Record.Get().GetValuePayload())
+		{
+			const ECachePolicy PayloadPolicy = Policy.GetPayloadPolicy(Payload.GetId());
+			GetCachePayload(Key, Context, PayloadPolicy, ECachePolicy::SkipValue, Payload, OutStatus);
+			if (Payload.IsNull())
 			{
-				// Save the compressed buffer with its hash as a header.
-				const FCompressedBuffer& CompressedBuffer = Payload.GetData();
-				const FIoHash CompressedHash = FIoHash::HashBuffer(CompressedBuffer.GetCompressed());
-				const FCompositeBuffer BufferWithHash(
-					FSharedBuffer::Clone(MakeMemoryView(CompressedHash.GetBytes())),
-					CompressedBuffer.GetCompressed());
-				if (!SaveFile(Path, BufferWithHash, Context))
-				{
-					return false;
-				}
+				return FOptionalCacheRecord();
 			}
+			RecordBuilder.SetValue(MoveTemp(Payload));
+		}
 
-			if (AccessLogWriter)
+		for (FPayload Payload : Record.Get().GetAttachmentPayloads())
+		{
+			const ECachePolicy PayloadPolicy = Policy.GetPayloadPolicy(Payload.GetId());
+			GetCachePayload(Key, Context, PayloadPolicy, ECachePolicy::SkipAttachments, Payload, OutStatus);
+			if (Payload.IsNull())
 			{
-				AccessLogWriter->Append(Key, Payload.GetId(), Path);
+				return FOptionalCacheRecord();
 			}
+			RecordBuilder.AddAttachment(MoveTemp(Payload));
 		}
 
-		Writer.BeginObject();
-		Writer.AddObjectId("Id"_ASV, Payload.GetId());
-		Writer.AddInteger("RawSize"_ASV, Payload.GetRawSize());
-		if (bHasData && bStoreInline)
+		return RecordBuilder.Build();
+	}
+
+	bool PutCacheContent(const FCompressedBuffer& Content, const FStringView Context) const
+	{
+		const FIoHash& RawHash = Content.GetRawHash();
+		TStringBuilder<256> Path;
+		BuildCacheContentPath(RawHash, Path);
+		if (!FileExists(Path))
 		{
-			const FCompressedBuffer& CompressedBuffer = Payload.GetData();
-			Writer.AddHash("RawHash"_ASV, RawHash);
-			Writer.AddHash("CompressedHash"_ASV, FIoHash::HashBuffer(CompressedBuffer.GetCompressed()));
-			Writer.AddBinary("CompressedData"_ASV, CompressedBuffer.GetCompressed());
+			if (!SaveFile(Path, Context, [&Content](FArchive& Ar) { Ar << const_cast<FCompressedBuffer&>(Content); }))
+			{
+				return false;
+			}
 		}
-		else
+		if (AccessLogWriter)
 		{
-			Writer.AddBinaryAttachment("RawHash"_ASV, RawHash);
+			AccessLogWriter->Append(RawHash, Path);
 		}
-		Writer.EndObject();
 		return true;
 	}
 
-	FPayload GetCachePayload(
+	void GetCachePayload(
 		const FCacheKey& Key,
 		const FStringView Context,
-		const FCacheRecordPolicy& Policy,
+		const ECachePolicy Policy,
 		const ECachePolicy SkipFlag,
-		const FCbObject& Object,
-		const bool bAlwaysLoadInlineData,
-		EStatus& OutStatus) const
+		FPayload& InOutPayload,
+		EStatus& InOutStatus) const
 	{
-		const FPayloadId Id = Object.FindView("Id"_ASV).AsObjectId();
-		const uint64 RawSize = Object.FindView("RawSize"_ASV).AsUInt64(MAX_uint64);
-		const FIoHash RawHash = Object.FindView("RawHash"_ASV).AsHash();
-		const FIoHash CompressedHash = Object.FindView("CompressedHash"_ASV).AsHash();
-		FSharedBuffer CompressedData = Object["CompressedData"_ASV].AsBinary();
+		const FIoHash& RawHash = InOutPayload.GetRawHash();
 
-		if (Id.IsNull() || RawSize == MAX_uint64 || RawHash.IsZero() || !(CompressedHash.IsZero() == CompressedData.IsNull()))
+		if (!EnumHasAnyFlags(Policy, ECachePolicy::Query) || (EnumHasAnyFlags(Policy, SkipFlag) && InOutPayload.HasData()))
 		{
-			UE_LOG(LogDerivedDataCache, Display, TEXT("%s: Cache miss with invalid record format for %s from '%.*s'"),
-				*CachePath, *WriteToString<96>(Key), Context.Len(), Context.GetData());
-			OutStatus = EStatus::Error;
-			return FPayload();
-		}
-
-		const ECachePolicy PayloadPolicy = Policy.GetPayloadPolicy(Id);
-		FPayload Payload(Id, RawHash, RawSize);
-
-		if (CompressedData)
-		{
-			if (EnumHasAllFlags(PayloadPolicy, SkipFlag) && !bAlwaysLoadInlineData)
-			{
-				OutStatus = EStatus::Ok;
-				return Payload;
-			}
-			return ValidateCachePayload(Key, Context, PayloadPolicy,
-				Payload, CompressedHash, MoveTemp(CompressedData), OutStatus);
+			InOutPayload = FPayload(InOutPayload.GetId(), RawHash, InOutPayload.GetRawSize());
+			return;
 		}
 
-		return GetCachePayload(Key, Context, PayloadPolicy, SkipFlag, Payload, OutStatus);
-	}
-
-	FPayload GetCachePayload(
-		const FCacheKey& Key,
-		const FStringView Context,
-		const ECachePolicy Policy,
-		const ECachePolicy SkipFlag,
-		const FPayload& Payload,
-		EStatus& OutStatus) const
-	{
-		if (Payload.HasData() || !EnumHasAnyFlags(Policy, ECachePolicy::Query))
+		if (InOutPayload.HasData())
 		{
-			OutStatus = EStatus::Ok;
-			return Payload;
+			return;
 		}
 
 		TStringBuilder<256> Path;
-		BuildCachePayloadPath(Key, Payload.GetRawHash(), Path);
+		BuildCacheContentPath(RawHash, Path);
 		if (EnumHasAllFlags(Policy, SkipFlag))
 		{
 			if (FileExists(Path))
 			{
 				if (AccessLogWriter)
 				{
-					AccessLogWriter->Append(Key, Payload.GetId(), Path);
+					AccessLogWriter->Append(RawHash, Path);
 				}
-				OutStatus = EStatus::Ok;
-				return Payload;
+				return;
 			}
 		}
 		else
 		{
-			if (FSharedBuffer CompressedData = LoadFile(Path, Context); CompressedData.GetSize() >= sizeof(FIoHash))
+			if (FSharedBuffer CompressedData = LoadFile(Path, Context))
 			{
-				const FMemoryView CompressedDataView = CompressedData;
-				const FIoHash CompressedHash(CompressedDataView.Left(sizeof(FIoHash)));
-				CompressedData = FSharedBuffer::MakeView(CompressedDataView + sizeof(FIoHash), MoveTemp(CompressedData));
-				const FPayload FullPayload = ValidateCachePayload(Key, Context, Policy,
-					Payload, CompressedHash, MoveTemp(CompressedData), OutStatus);
-				if (FullPayload && AccessLogWriter)
+				if (FCompressedBuffer CompressedBuffer = FCompressedBuffer::FromCompressed(MoveTemp(CompressedData));
+					CompressedBuffer && CompressedBuffer.GetRawHash() == RawHash)
+				{
+					if (AccessLogWriter)
+					{
+						AccessLogWriter->Append(RawHash, Path);
+					}
+					InOutPayload = FPayload(InOutPayload.GetId(), MoveTemp(CompressedBuffer));
+					return;
+				}
+				UE_LOG(LogDerivedDataCache, Display,
+					TEXT("%s: Cache miss with corrupted payload %s with hash %s for %s from '%.*s'"),
+					*CachePath, *WriteToString<16>(InOutPayload.GetId()), *WriteToString<48>(RawHash),
+					*WriteToString<96>(Key), Context.Len(), Context.GetData());
+				InOutStatus = EStatus::Error;
+				if (!EnumHasAnyFlags(Policy, ECachePolicy::PartialOnError))
 				{
-					AccessLogWriter->Append(Key, FullPayload.GetId(), Path);
+					InOutPayload = FPayload::Null;
 				}
-				return FullPayload;
+				return;
 			}
 		}
 
 		UE_LOG(LogDerivedDataCache, Verbose,
 			TEXT("%s: Cache miss with missing payload %s with hash %s for %s from '%.*s'"),
-			*CachePath, *WriteToString<16>(Payload.GetId()), *WriteToString<48>(Payload.GetRawHash()), *WriteToString<96>(Key),
+			*CachePath, *WriteToString<16>(InOutPayload.GetId()), *WriteToString<48>(RawHash), *WriteToString<96>(Key),
 			Context.Len(), Context.GetData());
-		OutStatus = EStatus::Error;
-		return EnumHasAllFlags(Policy, ECachePolicy::PartialOnError) ? Payload : FPayload::Null;
-	}
-
-	FPayload ValidateCachePayload(
-		const FCacheKey& Key,
-		const FStringView Context,
-		const ECachePolicy Policy,
-		const FPayload& Payload,
-		const FIoHash& CompressedHash,
-		FSharedBuffer&& CompressedData,
-		EStatus& OutStatus) const
-	{
-		if (FCompressedBuffer CompressedBuffer = FCompressedBuffer::FromCompressed(MoveTemp(CompressedData));
-			CompressedBuffer &&
-			CompressedBuffer.GetRawHash() == Payload.GetRawHash() &&
-			FIoHash::HashBuffer(CompressedBuffer.GetCompressed()) == CompressedHash)
-		{
-			OutStatus = EStatus::Ok;
-			return FPayload(Payload.GetId(), MoveTemp(CompressedBuffer));
-		}
-		UE_LOG(LogDerivedDataCache, Display,
-			TEXT("%s: Cache miss with corrupted payload %s with hash %s for %s from '%.*s'"),
-			*CachePath, *WriteToString<16>(Payload.GetId()), *WriteToString<48>(Payload.GetRawHash()),
-			*WriteToString<96>(Key), Context.Len(), Context.GetData());
-		OutStatus = EStatus::Error;
-		return EnumHasAllFlags(Policy, ECachePolicy::PartialOnError) ? Payload : FPayload::Null;
+		InOutStatus = EStatus::Error;
+		if (!EnumHasAnyFlags(Policy, ECachePolicy::PartialOnError))
+		{
+			InOutPayload = FPayload::Null;
+		}
 	}
 
-	void BuildCacheObjectPath(const FCacheKey& CacheKey, FStringBuilderBase& Path) const
+	void BuildCacheRecordPath(const FCacheKey& CacheKey, FStringBuilderBase& Path) const
 	{
 		Path << CachePath << TEXT('/');
-		BuildPathForCacheObject(CacheKey, Path);
+		BuildPathForCacheRecord(CacheKey, Path);
 	}
 
-	void BuildCachePayloadPath(const FCacheKey& CacheKey, const FIoHash& RawHash, FStringBuilderBase& Path) const
+	void BuildCacheContentPath(const FIoHash& RawHash, FStringBuilderBase& Path) const
 	{
 		Path << CachePath << TEXT('/');
-		BuildPathForCachePayload(CacheKey, RawHash, Path);
+		BuildPathForCacheContent(RawHash, Path);
 	}
 
-	bool SaveFile(FStringBuilderBase& Path, const FCompositeBuffer& Buffer, FStringView Context) const
+	bool SaveFile(FStringBuilderBase& Path, FStringView Context, TFunctionRef<void (FArchive&)> WriteFunction) const
 	{
 		TStringBuilder<256> TempPath;
-		TempPath << FPathViews::GetPath(Path) << TEXT("/Temp.") << FGuid::NewGuid().ToString();
+		TempPath << FPathViews::GetPath(Path) << TEXT("/Temp.") << FGuid::NewGuid();
 
 		ON_SCOPE_EXIT
 		{
 			IFileManager::Get().Delete(*TempPath, /*bRequireExists*/ false, /*bEvenReadOnly*/ false, /*bQuiet*/ true);
 		};
 
-		if (SaveBufferToFile(TempPath, Buffer))
+		int64 ExpectedSize = 0;
+		if (SaveDataToFile(TempPath, MoveTemp(WriteFunction), ExpectedSize))
 		{
-			if (IFileManager::Get().FileSize(*TempPath) == int64(Buffer.GetSize()))
+			if (IFileManager::Get().FileSize(*TempPath) == ExpectedSize)
 			{
 				if (IFileManager::Get().Move(*Path, *TempPath, /*bReplace*/ false, /*bEvenIfReadOnly*/ false, /*bAttributes*/ false, /*bDoNotRetryOrError*/ true))
 				{
@@ -1454,31 +1939,36 @@ private:
 			{
 				UE_LOG(LogDerivedDataCache, Warning,
 					TEXT("%s: Failed to write to temp file %s when saving %s from '%.*s'. ")
-					TEXT("File is %" INT64_FMT " bytes when %" UINT64_FMT " bytes are expected."),
+					TEXT("File is %" INT64_FMT " bytes when %" INT64_FMT " bytes are expected."),
 					*CachePath, *TempPath, *Path, Context.Len(), Context.GetData(),
-					IFileManager::Get().FileSize(*TempPath), Buffer.GetSize());
+					IFileManager::Get().FileSize(*TempPath), ExpectedSize);
 				return false;
 			}
 		}
 		else
 		{
 			UE_LOG(LogDerivedDataCache, Warning,
-				TEXT("%s: Failed to write to temp file %s when saving %s from '%.*s' (error 0x%08x)"),
+				TEXT("%s: Failed to write to temp file %s when saving %s from '%.*s'. Error 0x%08x"),
 				*CachePath, *TempPath, *Path, Context.Len(), Context.GetData(), FPlatformMisc::GetLastError());
 			return false;
 		}
 	}
 
-	bool SaveBufferToFile(FStringBuilderBase& Path, const FCompositeBuffer& Buffer) const
+	bool SaveDataToFile(
+		FStringBuilderBase& Path,
+		TFunctionRef<void (FArchive&)> WriteFunction,
+		int64& OutWriteSize) const
 	{
 		if (TUniquePtr<FArchive> Ar{IFileManager::Get().CreateFileWriter(*Path, FILEWRITE_Silent)})
 		{
-			for (const FSharedBuffer& Segment : Buffer.GetSegments())
-			{
-				Ar->Serialize(const_cast<void*>(Segment.GetData()), Segment.GetSize());
-			}
+			THashingArchiveProxy<FBlake3> HashAr(*Ar);
+			WriteFunction(HashAr);
+			FBlake3Hash Hash = HashAr.GetHash();
+			*Ar << Hash;
+			OutWriteSize = Ar->Tell();
 			return Ar->Close() && !Ar->IsError();
 		}
+		OutWriteSize = 0;
 		return false;
 	}
 
@@ -1498,12 +1988,42 @@ private:
 
 		if (TUniquePtr<FArchive> Ar{IFileManager::Get().CreateFileReader(*Path, FILEREAD_Silent)})
 		{
-			const int64 TotalSize = Ar->TotalSize();
-			FUniqueBuffer MutableBuffer = FUniqueBuffer::Alloc(uint64(TotalSize));
-			Ar->Serialize(MutableBuffer.GetData(), TotalSize);
-			if (Ar->Close())
+			if (const int64 TotalSize = Ar->TotalSize(); TotalSize >= sizeof(FBlake3Hash))
+			{
+				const int64 DataSize = TotalSize - sizeof(FBlake3Hash);
+				FUniqueBuffer MutableBuffer = FUniqueBuffer::Alloc(uint64(DataSize));
+				Ar->Serialize(MutableBuffer.GetData(), DataSize);
+				FBlake3Hash SavedHash;
+				*Ar << SavedHash;
+				if (Ar->Close())
+				{
+					const FBlake3Hash Hash = FBlake3::HashBuffer(MutableBuffer);
+					if (Hash == SavedHash)
+					{
+						Buffer = MutableBuffer.MoveToShared();
+					}
+					else
+					{
+						UE_LOG(LogDerivedDataCache, Display,
+							TEXT("%s: File %s from '%.*s' is corrupted and has hash %s when %s is expected."),
+							*CachePath, *Path, Context.Len(), Context.GetData(),
+							*WriteToString<80>(Hash), *WriteToString<80>(SavedHash));
+						if (!bReadOnly)
+						{
+							IFileManager::Get().Delete(*Path, /*bRequireExists*/ false, /*bEvenReadOnly*/ false, /*bQuiet*/ true);
+						}
+					}
+				}
+			}
+			else
 			{
-				Buffer = MutableBuffer.MoveToShared();
+				UE_LOG(LogDerivedDataCache, Display,
+					TEXT("%s: File %s from '%.*s' is %" INT64_FMT " bytes when at least %" INT64_FMT " bytes are required."),
+					*CachePath, *Path, Context.Len(), Context.GetData(), TotalSize, int64(sizeof(FBlake3Hash)));
+				if (!bReadOnly)
+				{
+					IFileManager::Get().Delete(*Path, /*bRequireExists*/ false, /*bEvenReadOnly*/ false, /*bQuiet*/ true);
+				}
 			}
 		}
 
@@ -1533,9 +2053,9 @@ private:
 		{
 			return false;
 		}
-		if (bTouch || (!bReadOnly && (FDateTime::Now() - TimeStamp).GetDays() > (DaysToDeleteUnusedFiles / 4)))
+		if (bTouch || (!bReadOnly && (FDateTime::UtcNow() - TimeStamp).GetTotalDays() > (DaysToDeleteUnusedFiles / 4)))
 		{
-			IFileManager::Get().SetTimeStamp(*Path, FDateTime::Now());
+			IFileManager::Get().SetTimeStamp(*Path, FDateTime::UtcNow());
 		}
 		return true;
 	}
@@ -1558,18 +2078,21 @@ private:
 	FString	CachePath;
 	/** Class of this cache */
 	ESpeedClass SpeedClass;
-	/** If true, do not attempt to write to this cache **/
+	/** If true, we failed to write to this directory and it did not contain anything so we should not be used. */
+	bool		bDisabled;
+	/** If true, do not attempt to write to this cache. */
 	bool		bReadOnly;
-	
 	/** If true, CachedDataProbablyExists will update the file timestamps. */
 	bool		bTouch;
 	/** If true, allow transient data to be removed from the cache. */
 	bool		bPurgeTransient;
 	/** Age of file when it should be deleted from DDC cache. */
-	int32		DaysToDeleteUnusedFiles;
+	double		DaysToDeleteUnusedFiles;
 
-	/** If true, we failed to write to this directory and it did not contain anything so we should not be used **/
-	bool		bDisabled;
+	/** Maximum total size of compressed data stored within a record package with multiple attachments. */
+	uint64		MaxRecordSizeKB = 256;
+	/** Maximum total size of compressed data stored within a value package, or a record package with one attachment. */
+	uint64		MaxValueSizeKB = 1024;
 
 	/** Object used for synchronization via a scoped lock						*/
 	FCriticalSection SynchronizationObject;
@@ -1593,6 +2116,8 @@ private:
 	TSet<FName> DebugMissedKeys;
 	TSet<FCacheKey> DebugMissedCacheKeys;
 
+	TUniquePtr<FFileSystemCacheStoreMaintainer> Maintainer;
+
 	bool ShouldSimulateMiss(const TCHAR* InKey)
 	{
 		if (DebugOptions.RandomMissRate == 0 && DebugOptions.SimulateMissTypes.IsEmpty())
@@ -1658,4 +2183,4 @@ FDerivedDataBackendInterface* CreateFileSystemDerivedDataBackend(const TCHAR* Ca
 	return FileDDB;
 }
 
-} // UE::DerivedData::Backends
+} // UE::DerivedData::CacheStore::FileSystem
diff --git a/Engine/Source/Developer/DerivedDataCache/Private/HashingArchiveProxy.h b/Engine/Source/Developer/DerivedDataCache/Private/HashingArchiveProxy.h
new file mode 100644
--- /dev/null
+++ b/Engine/Source/Developer/DerivedDataCache/Private/HashingArchiveProxy.h
@@ -0,0 +1,38 @@
+// Copyright Epic Games, Inc. All Rights Reserved.
+
+#pragma once
+
+#include "Serialization/ArchiveProxy.h"
+#include "Templates/UnrealTemplate.h"
+
+namespace UE
+{
+
+template <typename HashBuilderType>
+class THashingArchiveProxy : public FArchiveProxy
+{
+public:
+	using FArchiveProxy::FArchiveProxy;
+
+	void Serialize(void* V, int64 Length) override
+	{
+		FArchiveProxy::Serialize(V, Length);
+		HashBuilder.Update(V, Length);
+	}
+
+	void Seek(int64 InPos) override
+	{
+		checkf(InPos == Tell(), TEXT("A hash cannot be computed when serialization relies on seeking."));
+		FArchiveProxy::Seek(InPos);
+	}
+
+	auto GetHash() const -> decltype(DeclVal<HashBuilderType>().Finalize())
+	{
+		return HashBuilder.Finalize();
+	}
+
+private:
+	HashBuilderType HashBuilder;
+};
+
+} // UE
diff --git a/Engine/Source/Developer/DerivedDataCache/Private/S3DerivedDataBackend.cpp b/Engine/Source/Developer/DerivedDataCache/Private/S3DerivedDataBackend.cpp
--- a/Engine/Source/Developer/DerivedDataCache/Private/S3DerivedDataBackend.cpp
+++ b/Engine/Source/Developer/DerivedDataCache/Private/S3DerivedDataBackend.cpp
@@ -16,6 +16,7 @@
 #include <openssl/ssl.h>
 #include <openssl/sha.h>
 #include <openssl/hmac.h>
+#include "Algo/AllOf.h"
 #include "Misc/Base64.h"
 #include "Misc/Paths.h"
 #include "Misc/SecureHash.h"
@@ -59,7 +60,16 @@ TRACE_DECLARE_INT_COUNTER(S3DDC_Get, TEXT("S3DDC Get"));
 TRACE_DECLARE_INT_COUNTER(S3DDC_GetHit, TEXT("S3DDC Get Hit"));
 TRACE_DECLARE_INT_COUNTER(S3DDC_BytesRecieved, TEXT("S3DDC Bytes Recieved"));
 
-FString BuildPathForCacheKey(const TCHAR* CacheKey);
+FString BuildPathForCacheKey(const TCHAR* CacheKey)
+{
+	FString Key = FString(CacheKey).ToUpper();
+	checkf(Algo::AllOf(Key, [](TCHAR C) { return FChar::IsAlnum(C) || FChar::IsUnderscore(C) || C == TEXT('$'); }),
+		TEXT("Invalid characters in cache key %s"), CacheKey);
+	uint32 Hash = FCrc::StrCrc_DEPRECATED(*Key);
+	// this creates a tree of 1000 directories
+	FString HashPath = FString::Printf(TEXT("%1d/%1d/%1d/"), (Hash / 100) % 10, (Hash / 10) % 10, Hash % 10);
+	return HashPath / Key + TEXT(".udd");
+}
 
 class FStringAnsi
 {
@@ -95,7 +105,7 @@ public:
 	void Append(const ANSICHAR* Start, const ANSICHAR* End)
 	{
 		Inner.RemoveAt(Inner.Num() - 1);
-		Inner.Append(Start, End - Start);
+		Inner.Append(Start, UE_PTRDIFF_TO_INT32(End - Start));
 		Inner.Add(0);
 	}
 
@@ -338,7 +348,7 @@ private:
 			{
 				for (const ANSICHAR* Char = Header->data; Char != Colon; Char++)
 				{
-					CanonicalHeaders.Append(tolower(*Char));
+					CanonicalHeaders.Append(FCharAnsi::ToLower(*Char));
 				}
 				CanonicalHeaders.Append(':');
 
@@ -368,7 +378,7 @@ private:
 				}
 				for (const ANSICHAR* Char = Header->data; Char != Colon; Char++)
 				{
-					SignedHeaders.Append(tolower(*Char));
+					SignedHeaders.Append(FCharAnsi::ToLower(*Char));
 				}
 			}
 		}
@@ -427,7 +437,7 @@ private:
 		IRequestCallback* Callback = (IRequestCallback*)Ptr;
 		if (Callback != nullptr)
 		{
-			return Callback->Update(CurrentDownloadSize, TotalDownloadSize)? 0 : 1;
+			return Callback->Update(IntCastChecked<int>(CurrentDownloadSize), IntCastChecked<int>(TotalDownloadSize)) ? 0 : 1;
 		}
 		return 0;
 	}
diff --git a/Engine/Source/Developer/DerivedDataCache/Public/DDCCleanup.h b/Engine/Source/Developer/DerivedDataCache/Public/DDCCleanup.h
--- a/Engine/Source/Developer/DerivedDataCache/Public/DDCCleanup.h
+++ b/Engine/Source/Developer/DerivedDataCache/Public/DDCCleanup.h
@@ -2,12 +2,6 @@
 
 #pragma once
 
-#include "CoreMinimal.h"
-#include "HAL/ThreadSafeCounter.h"
-#include "HAL/Runnable.h"
-
-struct FFilesystemInfo;
-
 class IDDCCleanup
 {
 public:
@@ -15,100 +9,3 @@ public:
 	virtual bool IsFinished() const = 0;
 	virtual void WaitBetweenDeletes(bool bWait) = 0;
 };
-
-/** 
- * DDC Filesystem Cache cleanup thread.
- */
-class FDDCCleanup final : public FRunnable, public IDDCCleanup
-{
-	/** Singleton instance */
-	static FDDCCleanup* Runnable;
-	
-	/** Thread to run the cleanup FRunnable on */
-	FRunnableThread* Thread;
-	/** > 0 if we've been asked to abort work in progress at the next opportunity */
-	FThreadSafeCounter StopTaskCounter;
-	/** List of filesystems to clean up */
-	TArray< TSharedPtr< struct FFilesystemInfo > > CleanupList;
-	/** Synchronization object */
-	FCriticalSection DataLock;
-	/** If true, work without giving up time to other threads */
-	bool bDontWaitBetweenDeletes;
-
-	/** Time before any deleting starts */
-	float TimeToWaitAfterInit;
-	/** Time to wait before starting deleting the next DDC directory files */
-	float TimeBetweenDeleteingDirectories;
-	/** Time to wait before deleting the next DDC file */
-	float TimeBetweenDeletingFiles;
-
-	/** Constructor */
-	FDDCCleanup();
-
-	/** Destructor */
-	virtual ~FDDCCleanup();
-
-	//~ Begin FRunnable Interface.
-	virtual bool Init();
-	virtual uint32 Run();
-	virtual void Stop();
-	virtual void Exit();
-	//~ End FRunnable Interface
-
-	/** Checks if there's been any Stop requests */
-	bool ShouldStop() const;
-
-	/**
-	 * Waits for a given amount of time periodically checking if there's been any Stop requests.
-	 *
-	 * @param InSeconds time in seconds to wait.
-	 * @param InSleepTime interval at which to check for Stop requests.
-	 */
-	void Wait( const float InSeconds, const float InSleepTime = 0.1f );
-
-	/**
-	 * Performs directory cleanup.
-	 *
-	 * @param FilesystemInfo Filesystem to cleanup.
-	 * @return true if the filesystem contained any directories to clean up, false otherwise.
-	 */
-	bool CleanupFilesystemDirectory( TSharedPtr< FFilesystemInfo > FilesystemInfo );
-
-	/** Makes sure this thread has stopped properly */
-	void EnsureCompletion();
-
-public:
-
-	/**
-	 * Adds DDC filesystem to clean up.
-	 *
-	 * @param InCachePath Filesystem path.
-	 * @param InDaysToDelete Number of days since last access time to consider a file as unused.
-	 * @param InMaxContinuousFileChecks Number of files to check without pausing.
-	 */
-	void AddFilesystem( FString& InCachePath, int32 InDaysToDelete, int32 InMaxNumFoldersToCheck, int32 InMaxContinuousFileChecks );
-
-	/** Gets DDC Cleanup singleton instance, crates one if it doesn't exist already */
-	static FDDCCleanup* Get();
-
-	/** Gets DDC Cleanup singleton instance */
-	static FDDCCleanup* GetNoInit()
-	{
-		return Runnable;
-	}
-
-	/** Shuts down DDC Cleanup thread. */
-	static void Shutdown();
-
-	/** Sets whether cleanup should give up time to other threads between deletes */
-	virtual void WaitBetweenDeletes(bool bWait) final
-	{
-		bDontWaitBetweenDeletes = !bWait;
-	}
-
-	/** Checks if the cleanup thread is done deleting files */
-	virtual bool IsFinished() const final
-	{
-		return !CleanupList.Num();
-	}
-};
diff --git a/Engine/Source/Developer/DerivedDataCache/Public/DerivedDataCache.h b/Engine/Source/Developer/DerivedDataCache/Public/DerivedDataCache.h
--- a/Engine/Source/Developer/DerivedDataCache/Public/DerivedDataCache.h
+++ b/Engine/Source/Developer/DerivedDataCache/Public/DerivedDataCache.h
@@ -16,6 +16,7 @@
 
 namespace UE::DerivedData { class FCacheRecord; }
 namespace UE::DerivedData { class FCacheRecordPolicy; }
+namespace UE::DerivedData { class ICacheStoreMaintainer; }
 namespace UE::DerivedData { class IRequestOwner; }
 namespace UE::DerivedData { struct FCacheChunkRequest; }
 namespace UE::DerivedData { struct FCacheGetChunkCompleteParams; }
@@ -316,6 +317,10 @@ public:
 	 * owner of the cache if future requests are meant to be avoided.
 	 */
 	virtual void CancelAll() = 0;
+
+	/** Returns the interface to the background cache store maintenance. */
+	virtual ICacheStoreMaintainer& GetMaintainer() = 0;
+
 };
 
 ///////////////////////////////////////////////////////////////////////////////////////////////////////////////////////
diff --git a/Engine/Source/Developer/DerivedDataCache/Public/DerivedDataCacheInterface.h b/Engine/Source/Developer/DerivedDataCache/Public/DerivedDataCacheInterface.h
--- a/Engine/Source/Developer/DerivedDataCache/Public/DerivedDataCacheInterface.h
+++ b/Engine/Source/Developer/DerivedDataCache/Public/DerivedDataCacheInterface.h
@@ -274,6 +274,7 @@ public:
 	/**
 	 * Retrieve the interface to the background cache cleanup.
 	 */
+	UE_DEPRECATED(5.0, "This has been replaced by UE::DerivedData::GetCache().GetMaintainer().")
 	virtual IDDCCleanup* GetCleanup() const = 0;
 
 	//--------------------
diff --git a/Engine/Source/Developer/DerivedDataCache/Public/DerivedDataCacheKey.h b/Engine/Source/Developer/DerivedDataCache/Public/DerivedDataCacheKey.h
--- a/Engine/Source/Developer/DerivedDataCache/Public/DerivedDataCacheKey.h
+++ b/Engine/Source/Developer/DerivedDataCache/Public/DerivedDataCacheKey.h
@@ -54,6 +54,9 @@ public:
 	/** Get the name of the cache bucket as a string. */
 	inline FAnsiStringView ToString() const;
 
+	/** Get the name of the cache bucket as a null-terminated string. */
+	inline const ANSICHAR* ToCString() const { return Name; }
+
 protected:
 	static constexpr int32 LengthOffset = -1;
 
diff --git a/Engine/Source/Developer/DerivedDataCache/Public/DerivedDataCacheMaintainer.h b/Engine/Source/Developer/DerivedDataCache/Public/DerivedDataCacheMaintainer.h
new file mode 100644
--- /dev/null
+++ b/Engine/Source/Developer/DerivedDataCache/Public/DerivedDataCacheMaintainer.h
@@ -0,0 +1,43 @@
+// Copyright Epic Games, Inc. All Rights Reserved.
+
+#pragma once
+
+#include "Features/IModularFeature.h"
+#include "UObject/NameTypes.h"
+
+namespace UE::DerivedData
+{
+
+/**
+ * A cache store may register an instance of this interface to take part in cache maintenance.
+ *
+ * As an example, a filesystem cache might register a maintainer that scans for and deletes stale
+ * files within its cache directory. Boosting the priority removes any sleep between file scans.
+ *
+ * An example consumer is the DDCCleanup commandlet which uses this to boost priority and wait on
+ * the completion of maintenance by cache stores.
+ */
+class ICacheStoreMaintainer : public IModularFeature
+{
+public:
+	static inline const FLazyName FeatureName{"CacheStoreMaintainer"};
+
+	/**
+	 * True when maintenance is not active.
+	 *
+	 * This must return true eventually because it is called to wait for maintenance.
+	 */
+	virtual bool IsIdle() const = 0;
+
+	/**
+	 * Boost the priority of the active maintenance operation.
+	 *
+	 * This is expected to remove any delays, or otherwise allow the active maintenance operation
+	 * to complete more quickly. This is called before waiting for maintenance to be idle.
+	 */
+	virtual void BoostPriority()
+	{
+	}
+};
+
+} // UE::DerivedData
diff --git a/Engine/Source/Editor/UnrealEd/Private/Commandlets/DDCCleanupCommandlet.cpp b/Engine/Source/Editor/UnrealEd/Private/Commandlets/DDCCleanupCommandlet.cpp
--- a/Engine/Source/Editor/UnrealEd/Private/Commandlets/DDCCleanupCommandlet.cpp
+++ b/Engine/Source/Editor/UnrealEd/Private/Commandlets/DDCCleanupCommandlet.cpp
@@ -1,22 +1,20 @@
 // Copyright Epic Games, Inc. All Rights Reserved.
 
 #include "Commandlets/DDCCleanupCommandlet.h"
-#include "DDCCleanup.h"
-#include "DerivedDataCacheInterface.h"
+#include "DerivedDataCache.h"
+#include "DerivedDataCacheMaintainer.h"
 #include "HAL/PlatformProcess.h"
 
 int32 UDDCCleanupCommandlet::Main(const FString& Params)
 {
-	if (IDDCCleanup* Cleanup = GetDerivedDataCacheRef().GetCleanup())
+	using namespace UE::DerivedData;
+	ICacheStoreMaintainer& Maintainer = GetCache().GetMaintainer();
+	Maintainer.BoostPriority();
+	while (!Maintainer.IsIdle())
 	{
-		Cleanup->WaitBetweenDeletes(false);
-		do
-		{
-			// Cleanup works from its own thread so just wait until it's done and flush logs once in a while
-			FPlatformProcess::SleepNoStats(1.0f);
-			GLog->Flush();
-		} while (!Cleanup->IsFinished());
+		// Maintenance works from dedicated threads. Wait for it to finish and flush logs occasionally.
+		FPlatformProcess::SleepNoStats(0.05f);
+		GLog->Flush();
 	}
-
 	return 0;
 }
