DDC: Reorganize the code in the file system cache store

No functional changes, only bringing the code style in line with the current coding standards since nearly the entire file has been modified for 5.0 anyway.

#rnx

diff --git a/Engine/Source/Developer/DerivedDataCache/Private/FileSystemDerivedDataBackend.cpp b/Engine/Source/Developer/DerivedDataCache/Private/FileSystemDerivedDataBackend.cpp
--- a/Engine/Source/Developer/DerivedDataCache/Private/FileSystemDerivedDataBackend.cpp
+++ b/Engine/Source/Developer/DerivedDataCache/Private/FileSystemDerivedDataBackend.cpp
@@ -22,6 +22,7 @@
 #include "Hash/xxhash.h"
 #include "HashingArchiveProxy.h"
 #include "Misc/CommandLine.h"
+#include "Misc/ConfigCacheIni.h"
 #include "Misc/CoreMisc.h"
 #include "Misc/FileHelper.h"
 #include "Misc/Guid.h"
@@ -69,7 +70,7 @@ static const TCHAR GContentDirectoryName[] = TEXT("Content");
 
 ///////////////////////////////////////////////////////////////////////////////////////////////////////////////////////
 
-FString BuildPathForCacheKey(const TCHAR* CacheKey)
+FString BuildPathForCacheKey(const TCHAR* const CacheKey)
 {
 	FString Key = FString(CacheKey).ToUpper();
 	checkf(Algo::AllOf(Key, [](TCHAR C) { return FChar::IsAlnum(C) || FChar::IsUnderscore(C) || C == TEXT('$'); }),
@@ -132,14 +133,14 @@ private:
 class FRandomStream
 {
 public:
-	inline explicit FRandomStream(uint32 Seed)
+	inline explicit FRandomStream(const uint32 Seed)
 		: Random(1103515245, 12345) // From ANSI C
 		, Value(Seed)
 	{
 	}
 
 	/** Returns a random value in [Min, Max). */
-	inline uint32 GetRandRange(uint32 Min, uint32 Max)
+	inline uint32 GetRandRange(const uint32 Min, const uint32 Max)
 	{
 		return Min + uint32((uint64(Max - Min) * Random.GetNext(Value)) >> 32);
 	}
@@ -222,15 +223,15 @@ private:
 	void CreateContentRoot();
 	void CreateBucketRoots();
 	void ScanHashRoot(uint32 RootIndex);
-	TStaticBitArray<256> ScanHashDirectory(FStringBuilderBase& Path);
+	TStaticBitArray<256> ScanHashDirectory(FStringBuilderBase& BasePath);
 
-	TStaticBitArray<10> ScanLegacyDirectory(FStringBuilderBase& Path);
+	TStaticBitArray<10> ScanLegacyDirectory(FStringBuilderBase& BasePath);
 	void CreateLegacyRoot();
 	void ScanLegacyRoot();
 
 	void ResetRoots();
 
-	void ProcessFile(const TCHAR* Path, const FFileStatData& StatData);
+	void ProcessFile(const TCHAR* Path, const FFileStatData& Stat);
 
 private:
 	struct FRoot;
@@ -265,7 +266,7 @@ private:
 
 struct FFileSystemCacheStoreMaintainer::FRoot
 {
-	inline FRoot(FStringView RootPath, FRandomStream& Stream)
+	inline FRoot(const FStringView RootPath, FRandomStream& Stream)
 		: Order(Stream)
 	{
 		Path.Append(RootPath);
@@ -298,7 +299,7 @@ struct FFileSystemCacheStoreMaintainer::FLegacyRoot
 
 FFileSystemCacheStoreMaintainer::FFileSystemCacheStoreMaintainer(
 	const FFileSystemCacheStoreMaintainerParams& InParams,
-	FStringView InCachePath)
+	const FStringView InCachePath)
 	: Params(InParams)
 	, CachePath(InCachePath)
 	, IdleEvent(EEventMode::ManualReset)
@@ -403,9 +404,9 @@ void FFileSystemCacheStoreMaintainer::CreateBucketRoots()
 {
 	TStringBuilder<256> BucketsPath;
 	FPathViews::Append(BucketsPath, CachePath, GBucketsDirectoryName);
-	FileManager.IterateDirectoryStat(*BucketsPath, [this](const TCHAR* Path, const FFileStatData& StatData) -> bool
+	FileManager.IterateDirectoryStat(*BucketsPath, [this](const TCHAR* Path, const FFileStatData& Stat) -> bool
 	{
-		if (StatData.bIsDirectory)
+		if (Stat.bIsDirectory)
 		{
 			Roots.Add(MakeUnique<FRoot>(Path, Random));
 		}
@@ -413,7 +414,7 @@ void FFileSystemCacheStoreMaintainer::CreateBucketRoots()
 	});
 }
 
-void FFileSystemCacheStoreMaintainer::ScanHashRoot(uint32 RootIndex)
+void FFileSystemCacheStoreMaintainer::ScanHashRoot(const uint32 RootIndex)
 {
 	FRoot& Root = *Roots[int32(RootIndex)];
 	const uint32 DirectoryIndex = Root.Order.GetNext();
@@ -426,7 +427,8 @@ void FFileSystemCacheStoreMaintainer::ScanHashRoot(uint32 RootIndex)
 		if ((DirectoryIndex == Root.Order.GetFirst()) ||
 			(bScanned && ++Root.DirectoryScanCount >= Params.MaxDirectoryScanCount))
 		{
-			UE_LOG(LogDerivedDataCache, VeryVerbose, TEXT("%s: Maintenance finished scanning %s."), *CachePath, *Root.Path);
+			UE_LOG(LogDerivedDataCache, VeryVerbose,
+				TEXT("%s: Maintenance finished scanning %s."), *CachePath, *Root.Path);
 			Roots.RemoveAt(int32(RootIndex));
 		}
 	};
@@ -456,25 +458,25 @@ void FFileSystemCacheStoreMaintainer::ScanHashRoot(uint32 RootIndex)
 
 	TStringBuilder<256> Path;
 	Path.Appendf(TEXT("%s/%02x/%02x"), *Root.Path, IndexLevel0, IndexLevel1);
-	FileManager.IterateDirectoryStat(*Path, [this](const TCHAR* const Path, const FFileStatData& StatData) -> bool
+	FileManager.IterateDirectoryStat(*Path, [this](const TCHAR* const Path, const FFileStatData& Stat) -> bool
 	{
-		ProcessFile(Path, StatData);
+		ProcessFile(Path, Stat);
 		return !bExit;
 	});
 
 	bScanned = true;
 }
 
-TStaticBitArray<256> FFileSystemCacheStoreMaintainer::ScanHashDirectory(FStringBuilderBase& Path)
+TStaticBitArray<256> FFileSystemCacheStoreMaintainer::ScanHashDirectory(FStringBuilderBase& BasePath)
 {
 	TStaticBitArray<256> Exists;
-	FileManager.IterateDirectoryStat(*Path, [this, &Exists](const TCHAR* ChildPath, const FFileStatData& StatData) -> bool
+	FileManager.IterateDirectoryStat(*BasePath, [this, &Exists](const TCHAR* Path, const FFileStatData& Stat) -> bool
 	{
-		FStringView ChildView = FPathViews::GetCleanFilename(ChildPath);
-		if (StatData.bIsDirectory && ChildView.Len() == 2 && Algo::AllOf(ChildView, FChar::IsHexDigit))
+		FStringView View = FPathViews::GetCleanFilename(Path);
+		if (Stat.bIsDirectory && View.Len() == 2 && Algo::AllOf(View, FChar::IsHexDigit))
 		{
 			uint8 Byte;
-			if (String::HexToBytes(ChildView, &Byte) == 1)
+			if (String::HexToBytes(View, &Byte) == 1)
 			{
 				Exists[Byte] = true;
 			}
@@ -484,16 +486,15 @@ TStaticBitArray<256> FFileSystemCacheStoreMaintainer::ScanHashDirectory(FStringB
 	return Exists;
 }
 
-TStaticBitArray<10> FFileSystemCacheStoreMaintainer::ScanLegacyDirectory(FStringBuilderBase& Path)
+TStaticBitArray<10> FFileSystemCacheStoreMaintainer::ScanLegacyDirectory(FStringBuilderBase& BasePath)
 {
 	TStaticBitArray<10> Exists;
-	const bool bDeleteFiles = CachePath.Len() < Path.Len();
-	FileManager.IterateDirectoryStat(*Path, [this, &Exists, bDeleteFiles](const TCHAR* ChildPath, const FFileStatData& StatData) -> bool
+	FileManager.IterateDirectoryStat(*BasePath, [this, &Exists](const TCHAR* Path, const FFileStatData& Stat) -> bool
 	{
-		FStringView ChildView = FPathViews::GetCleanFilename(ChildPath);
-		if (StatData.bIsDirectory && ChildView.Len() == 1 && Algo::AllOf(ChildView, FChar::IsDigit))
+		FStringView View = FPathViews::GetCleanFilename(Path);
+		if (Stat.bIsDirectory && View.Len() == 1 && Algo::AllOf(View, FChar::IsDigit))
 		{
-			Exists[FChar::ConvertCharDigitToInt(ChildView[0])] = true;
+			Exists[FChar::ConvertCharDigitToInt(View[0])] = true;
 		}
 		return !bExit;
 	});
@@ -561,9 +562,9 @@ void FFileSystemCacheStoreMaintainer::ScanLegacyRoot()
 
 	TStringBuilder<256> Path;
 	FPathViews::Append(Path, CachePath, IndexLevel0, IndexLevel1, IndexLevel2);
-	FileManager.IterateDirectoryStat(*Path, [this](const TCHAR* const Path, const FFileStatData& StatData) -> bool
+	FileManager.IterateDirectoryStat(*Path, [this](const TCHAR* const Path, const FFileStatData& Stat) -> bool
 	{
-		ProcessFile(Path, StatData);
+		ProcessFile(Path, Stat);
 		return !bExit;
 	});
 
@@ -576,28 +577,29 @@ void FFileSystemCacheStoreMaintainer::ResetRoots()
 	LegacyRoot.Reset();
 }
 
-void FFileSystemCacheStoreMaintainer::ProcessFile(const TCHAR* const Path, const FFileStatData& StatData)
+void FFileSystemCacheStoreMaintainer::ProcessFile(const TCHAR* const Path, const FFileStatData& Stat)
 {
-	if (StatData.bIsDirectory)
+	if (Stat.bIsDirectory)
 	{
 		return;
 	}
 
-	if (StatData.ModificationTime + Params.MaxFileAge < FDateTime::UtcNow())
+	const FDateTime Now = FDateTime::UtcNow();
+	if (Stat.ModificationTime + Params.MaxFileAge < Now && Stat.AccessTime + Params.MaxFileAge < Now)
 	{
 		++DeleteCount;
-		DeleteSize += StatData.FileSize > 0 ? uint64(StatData.FileSize) : 0;
+		DeleteSize += Stat.FileSize > 0 ? uint64(Stat.FileSize) : 0;
 		if (FileManager.Delete(Path, /*bRequireExists*/ false, /*bEvenReadOnly*/ false, /*bQuiet*/ true))
 		{
 			UE_LOG(LogDerivedDataCache, VeryVerbose,
 				TEXT("%s: Maintenance deleted file %s that was last modified at %s."),
-				*CachePath, Path, *StatData.ModificationTime.ToIso8601());
+				*CachePath, Path, *Stat.ModificationTime.ToIso8601());
 		}
 		else
 		{
 			UE_LOG(LogDerivedDataCache, Verbose,
 				TEXT("%s: Maintenance failed to delete file %s that was last modified at %s."),
-				*CachePath, Path, *StatData.ModificationTime.ToIso8601());
+				*CachePath, Path, *Stat.ModificationTime.ToIso8601());
 		}
 	}
 
@@ -618,1569 +620,1629 @@ void FFileSystemCacheStoreMaintainer::ProcessFile(const TCHAR* const Path, const
 
 ///////////////////////////////////////////////////////////////////////////////////////////////////////////////////////
 
-/** 
- * Cache server that uses the OS filesystem
- * The entire API should be callable from any thread (except the singleton can be assumed to be called at least once before concurrent access).
-**/
-class FFileSystemDerivedDataBackend : public FDerivedDataBackendInterface
+class FAccessLogWriter
+{
+public:
+	FAccessLogWriter(const TCHAR* FileName, const FString& CachePath);
+
+	void Append(const TCHAR* CacheKey, FStringView Path);
+	void Append(const FIoHash& RawHash, FStringView Path);
+	void Append(const FCacheKey& CacheKey, FStringView Path);
+
+private:
+	void AppendPath(FStringView Path);
+
+	TUniquePtr<FArchive> Archive;
+	FString BasePath;
+	FCriticalSection CriticalSection;
+	TSet<FString> CacheKeys;
+	TSet<FIoHash> ContentKeys;
+	TSet<FCacheKey> RecordKeys;
+};
+
+FAccessLogWriter::FAccessLogWriter(const TCHAR* const FileName, const FString& CachePath)
+	: Archive(IFileManager::Get().CreateFileWriter(FileName, FILEWRITE_AllowRead))
+	, BasePath(CachePath / TEXT(""))
+{
+}
+
+void FAccessLogWriter::Append(const TCHAR* const CacheKey, const FStringView Path)
+{
+	FScopeLock Lock(&CriticalSection);
+
+	bool bIsAlreadyInSet = false;
+	CacheKeys.FindOrAdd(FString(CacheKey), &bIsAlreadyInSet);
+	if (!bIsAlreadyInSet)
+	{
+		AppendPath(Path);
+	}
+}
+
+void FAccessLogWriter::Append(const FIoHash& RawHash, const FStringView Path)
+{
+	FScopeLock Lock(&CriticalSection);
+
+	bool bIsAlreadyInSet = false;
+	ContentKeys.FindOrAdd(RawHash, &bIsAlreadyInSet);
+	if (!bIsAlreadyInSet)
+	{
+		AppendPath(Path);
+	}
+}
+
+void FAccessLogWriter::Append(const FCacheKey& CacheKey, const FStringView Path)
+{
+	FScopeLock Lock(&CriticalSection);
+
+	bool bIsAlreadyInSet = false;
+	RecordKeys.FindOrAdd(CacheKey, &bIsAlreadyInSet);
+	if (!bIsAlreadyInSet)
+	{
+		AppendPath(Path);
+	}
+}
+
+void FAccessLogWriter::AppendPath(const FStringView Path)
+{
+	if (Path.StartsWith(BasePath))
+	{
+		const FTCHARToUTF8 PathUtf8(Path);
+		Archive->Serialize(const_cast<ANSICHAR*>(PathUtf8.Get()), PathUtf8.Length());
+		Archive->Serialize(const_cast<ANSICHAR*>(LINE_TERMINATOR_ANSI), sizeof(LINE_TERMINATOR_ANSI) - 1);
+	}
+}
+
+///////////////////////////////////////////////////////////////////////////////////////////////////////////////////////
+
+class FFileSystemCacheStore final : public FDerivedDataBackendInterface
 {
 public:
-	FFileSystemDerivedDataBackend(const TCHAR* InCachePath, const TCHAR* InParams, const TCHAR* InAccessLogFileName)
-		: CachePath(InCachePath)
-		, SpeedClass(ESpeedClass::Unknown)
-		, bDisabled(false)
-		, bReadOnly(false)
-		, bTouch(false)
-		, bPurgeTransient(false)
-		, DaysToDeleteUnusedFiles(15.0)
-		, TotalEstimatedBuildTime(0)
-	{
-		// If we find a platform that has more stringent limits, this needs to be rethought.
-		checkf(MAX_BACKEND_KEY_LENGTH + MAX_CACHE_DIR_LEN + MAX_BACKEND_NUMBERED_SUBFOLDER_LENGTH + MAX_CACHE_EXTENTION_LEN < FPlatformMisc::GetMaxPathLength(),
-			TEXT("Not enough room left for cache keys in max path."));
-
-		check(CachePath.Len());
-		FPaths::NormalizeFilename(CachePath);
-
-		// Params that override our instance defaults
-		FParse::Bool(InParams, TEXT("ReadOnly="), bReadOnly);
-		FParse::Bool(InParams, TEXT("Touch="), bTouch);
-		FParse::Bool(InParams, TEXT("PurgeTransient="), bPurgeTransient);
-		FParse::Value(InParams, TEXT("UnusedFileAge="), DaysToDeleteUnusedFiles);
-		FParse::Value(InParams, TEXT("MaxRecordSizeKB="), MaxRecordSizeKB);
-		FParse::Value(InParams, TEXT("MaxValueSizeKB="), MaxValueSizeKB);
-
-		// Flush the cache if requested.
-		bool bFlush = false;
-		if (!bReadOnly && FParse::Bool(InParams, TEXT("Flush="), bFlush) && bFlush)
+	FFileSystemCacheStore(const TCHAR* CachePath, const TCHAR* Params, const TCHAR* AccessLogPath);
+
+	inline bool IsUsable() const { return !bDisabled; }
+
+	bool RunSpeedTest(
+		double InSkipTestsIfSeeksExceedMS,
+		double& OutSeekTimeMS,
+		double& OutReadSpeedMBs,
+		double& OutWriteSpeedMBs) const;
+
+	// FDerivedDataBackendInterface Interface
+
+	FString GetName() const final { return CachePath; }
+	FString GetDisplayName() const final { return FString(TEXT("File System")); }
+	bool IsWritable() const final { return !bReadOnly && !bDisabled; }
+	ESpeedClass GetSpeedClass() const final { return SpeedClass; }
+
+	bool CachedDataProbablyExists(const TCHAR* CacheKey) final;
+	bool GetCachedData(const TCHAR* CacheKey, TArray<uint8>& Data) final;
+	EPutStatus PutCachedData(const TCHAR* CacheKey, TConstArrayView<uint8> Data, bool bPutEvenIfExists) final;
+	void RemoveCachedData(const TCHAR* CacheKey, bool bTransient) final;
+	TSharedRef<FDerivedDataCacheStatsNode> GatherUsageStats() const final;
+	bool TryToPrefetch(TConstArrayView<FString> CacheKeys) final;
+	bool WouldCache(const TCHAR* CacheKey, TConstArrayView<uint8> InData) final;
+	bool ApplyDebugOptions(FBackendDebugOptions& InOptions) final;
+
+	// ICacheStore Interface
+
+	void Put(
+		TConstArrayView<FCacheRecord> Records,
+		FStringView Context,
+		ECachePolicy Policy,
+		IRequestOwner& Owner,
+		FOnCachePutComplete&& OnComplete) final;
+	void Get(
+		TConstArrayView<FCacheKey> Keys,
+		FStringView Context,
+		FCacheRecordPolicy Policy,
+		IRequestOwner& Owner,
+		FOnCacheGetComplete&& OnComplete) final;
+	void GetChunks(
+		TConstArrayView<FCacheChunkRequest> Chunks,
+		FStringView Context,
+		IRequestOwner& Owner,
+		FOnCacheGetChunkComplete&& OnComplete) final;
+
+private:
+	uint64 MeasureCompressedCacheRecord(const FCacheRecord& Record) const;
+	uint64 MeasureRawCacheRecord(const FCacheRecord& Record) const;
+
+	bool PutCacheRecord(const FCacheRecord& Record, FStringView Context, ECachePolicy Policy);
+	bool PutCacheContent(const FCompressedBuffer& Content, const FStringView Context) const;
+
+	FOptionalCacheRecord GetCacheRecordOnly(
+		const FCacheKey& Key,
+		const FStringView Context,
+		const FCacheRecordPolicy& Policy);
+	FOptionalCacheRecord GetCacheRecord(
+		const FCacheKey& Key,
+		const FStringView Context,
+		const FCacheRecordPolicy& Policy,
+		EStatus& OutStatus);
+	void GetCacheContent(
+		const FCacheKey& Key,
+		const FStringView Context,
+		const ECachePolicy Policy,
+		const ECachePolicy SkipFlag,
+		FPayload& InOutPayload,
+		EStatus& InOutStatus) const;
+
+	void BuildCacheRecordPath(const FCacheKey& CacheKey, FStringBuilderBase& Path) const;
+	void BuildCacheContentPath(const FIoHash& RawHash, FStringBuilderBase& Path) const;
+
+	bool SaveFile(FStringBuilderBase& Path, FStringView Context, TFunctionRef<void (FArchive&)> WriteFunction) const;
+	bool SaveDataToFile(FStringBuilderBase& Path, TFunctionRef<void (FArchive&)> WriteFunction, int64& OutWriteSize) const;
+	FSharedBuffer LoadFile(FStringBuilderBase& Path, FStringView Context) const;
+
+	bool FileExists(FStringBuilderBase& Path) const;
+
+	bool ShouldSimulateMiss(const TCHAR* InKey);
+	bool ShouldSimulateMiss(const FCacheKey& Key);
+
+private:
+	/**
+	 * Thread-safe method to compute the filename from the cache key, currently just adds a path and an extension.
+	 *
+	 * @param	CacheKey	Alphanumeric+underscore key of this cache item
+	 * @return				filename built from the cache key
+	 */
+	FString BuildFilename(const TCHAR* const CacheKey) const
+	{
+		return CachePath / BuildPathForCacheKey(CacheKey);
+	}
+
+	/** Base path we are storing the cache files in. */
+	FString	CachePath;
+	/** Class of this cache */
+	ESpeedClass SpeedClass;
+	/** If true, we failed to write to this directory and it did not contain anything so we should not be used. */
+	bool		bDisabled;
+	/** If true, do not attempt to write to this cache. */
+	bool		bReadOnly;
+	/** If true, CachedDataProbablyExists will update the file timestamps. */
+	bool		bTouch;
+	/** If true, allow transient data to be removed from the cache. */
+	bool		bPurgeTransient;
+	/** Age of file when it should be deleted from DDC cache. */
+	double		DaysToDeleteUnusedFiles;
+
+	/** Maximum total size of compressed data stored within a record package with multiple attachments. */
+	uint64		MaxRecordSizeKB = 256;
+	/** Maximum total size of compressed data stored within a value package, or a record package with one attachment. */
+	uint64		MaxValueSizeKB = 1024;
+
+	/** Object used for synchronization via a scoped lock. */
+	FCriticalSection SynchronizationObject;
+
+	// DDCNotification metrics
+
+	/** Map of cache keys to miss times for generating timing deltas. */
+	TMap<FString, double> DDCNotificationCacheTimes;
+
+	/** The total estimated build time accumulated from cache miss/put deltas. */
+	double TotalEstimatedBuildTime;
+
+	/** Access log to write to */
+	TUniquePtr<FAccessLogWriter> AccessLogWriter;
+
+	/** Debug Options */
+	FBackendDebugOptions DebugOptions;
+
+	/** Keys we ignored due to miss rate settings */
+	FCriticalSection MissedKeysCS;
+	TSet<FName> DebugMissedKeys;
+	TSet<FCacheKey> DebugMissedCacheKeys;
+
+	FDerivedDataCacheUsageStats UsageStats;
+
+	TUniquePtr<FFileSystemCacheStoreMaintainer> Maintainer;
+};
+
+FFileSystemCacheStore::FFileSystemCacheStore(
+	const TCHAR* const InCachePath,
+	const TCHAR* const InParams,
+	const TCHAR* const InAccessLogPath)
+	: CachePath(InCachePath)
+	, SpeedClass(ESpeedClass::Unknown)
+	, bDisabled(false)
+	, bReadOnly(false)
+	, bTouch(false)
+	, bPurgeTransient(false)
+	, DaysToDeleteUnusedFiles(15.0)
+	, TotalEstimatedBuildTime(0)
+{
+	// If we find a platform that has more stringent limits, this needs to be rethought.
+	checkf(MAX_BACKEND_KEY_LENGTH + MAX_CACHE_DIR_LEN + MAX_BACKEND_NUMBERED_SUBFOLDER_LENGTH + MAX_CACHE_EXTENTION_LEN < FPlatformMisc::GetMaxPathLength(),
+		TEXT("Not enough room left for cache keys in max path."));
+
+	check(CachePath.Len());
+	FPaths::NormalizeFilename(CachePath);
+
+	// Params that override our instance defaults
+	FParse::Bool(InParams, TEXT("ReadOnly="), bReadOnly);
+	FParse::Bool(InParams, TEXT("Touch="), bTouch);
+	FParse::Bool(InParams, TEXT("PurgeTransient="), bPurgeTransient);
+	FParse::Value(InParams, TEXT("UnusedFileAge="), DaysToDeleteUnusedFiles);
+	FParse::Value(InParams, TEXT("MaxRecordSizeKB="), MaxRecordSizeKB);
+	FParse::Value(InParams, TEXT("MaxValueSizeKB="), MaxValueSizeKB);
+
+	// Flush the cache if requested.
+	bool bFlush = false;
+	if (!bReadOnly && FParse::Bool(InParams, TEXT("Flush="), bFlush) && bFlush)
+	{
+		IFileManager::Get().DeleteDirectory(*(CachePath / TEXT("")), /*bRequireExists*/ false, /*bTree*/ true);
+	}
+
+	// check latency and speed. Read values should always be valid
+	double ReadSpeedMBs = 0.0;
+	double WriteSpeedMBs = 0.0;
+	double SeekTimeMS = 0.0;
+
+	/* Speeds faster than this are considered local*/
+	const float ConsiderFastAtMS = 10;
+	/* Speeds faster than this are ok. Everything else is slow. This value can be overridden in the ini file */
+	float ConsiderSlowAtMS = 50;
+	FParse::Value(InParams, TEXT("ConsiderSlowAt="), ConsiderSlowAtMS);
+
+	// can skip the speed test so everything acts as local (e.g. 4.25 and earlier behavior).
+	bool SkipSpeedTest = !WITH_EDITOR || FParse::Param(FCommandLine::Get(), TEXT("DDCSkipSpeedTest"));
+	if (SkipSpeedTest)
+	{
+		ReadSpeedMBs = 999;
+		WriteSpeedMBs = 999;
+		SeekTimeMS = 0;
+		UE_LOG(LogDerivedDataCache, Log, TEXT("Skipping speed test to %s. Assuming local performance"), *CachePath);
+	}
+
+	if (!SkipSpeedTest && !RunSpeedTest(ConsiderSlowAtMS * 2, SeekTimeMS, ReadSpeedMBs, WriteSpeedMBs))
+	{
+		bDisabled = true;
+		UE_LOG(LogDerivedDataCache, Warning, TEXT("No read or write access to %s"), *CachePath);
+	}
+	else
+	{
+		bool bReadTestPassed = ReadSpeedMBs > 0.0;
+		bool bWriteTestPassed = WriteSpeedMBs > 0.0;
+
+		// if we failed writes mark this as read only
+		bReadOnly = bReadOnly || !bWriteTestPassed;
+
+		// classify and report on these times
+		if (SeekTimeMS < 1)
 		{
-			IFileManager::Get().DeleteDirectory(*(CachePath / TEXT("")), /*bRequireExists*/ false, /*bTree*/ true);
+			SpeedClass = ESpeedClass::Local;
 		}
-
-		// check latency and speed. Read values should always be valid
-		double ReadSpeedMBs = 0.0;
-		double WriteSpeedMBs = 0.0;
-		double SeekTimeMS = 0.0;
-
-		/* Speeds faster than this are considered local*/
-		const float ConsiderFastAtMS = 10;
-		/* Speeds faster than this are ok. Everything else is slow. This value can be overridden in the ini file */
-		float ConsiderSlowAtMS = 50;
-		FParse::Value(InParams, TEXT("ConsiderSlowAt="), ConsiderSlowAtMS);
-
-		// can skip the speed test so everything acts as local (e.g. 4.25 and earlier behavior). 
-		bool SkipSpeedTest = !WITH_EDITOR || FParse::Param(FCommandLine::Get(), TEXT("ddcskipspeedtest"));
-		if (SkipSpeedTest)
+		else if (SeekTimeMS <= ConsiderFastAtMS)
 		{
-			ReadSpeedMBs = 999;
-			WriteSpeedMBs = 999;
-			SeekTimeMS = 0;
-			UE_LOG(LogDerivedDataCache, Log, TEXT("Skipping speed test to %s. Assuming local performance"), *CachePath);
+			SpeedClass = ESpeedClass::Fast;
 		}
-
-		if (!SkipSpeedTest && !RunSpeedTest(ConsiderSlowAtMS * 2, SeekTimeMS, ReadSpeedMBs, WriteSpeedMBs))
+		else if (SeekTimeMS >= ConsiderSlowAtMS)
 		{
-			bDisabled = true;
-			UE_LOG(LogDerivedDataCache, Warning, TEXT("No read or write access to %s"), *CachePath);
+			SpeedClass = ESpeedClass::Slow;
 		}
 		else
 		{
-			bool bReadTestPassed = ReadSpeedMBs > 0.0;
-			bool bWriteTestPassed = WriteSpeedMBs > 0.0;
+			SpeedClass = ESpeedClass::Ok;
+		}
 
-			// if we failed writes mark this as read only
-			bReadOnly = bReadOnly || !bWriteTestPassed;
+		UE_LOG(LogDerivedDataCache, Display,
+			TEXT("Performance to %s: Latency=%.02fms. RandomReadSpeed=%.02fMBs, RandomWriteSpeed=%.02fMBs. ")
+			TEXT("Assigned SpeedClass '%s'"),
+			*CachePath, SeekTimeMS, ReadSpeedMBs, WriteSpeedMBs, LexToString(SpeedClass));
 
-			// classify and report on these times
-			if (SeekTimeMS < 1)
+		if (SpeedClass <= FDerivedDataBackendInterface::ESpeedClass::Slow && !bReadOnly)
+		{
+			if (GIsBuildMachine)
 			{
-				SpeedClass = ESpeedClass::Local;
+				UE_LOG(LogDerivedDataCache, Display, TEXT("Access to %s appears to be slow. ")
+					TEXT("'Touch' will be disabled and queries/writes will be limited."), *CachePath);
 			}
-			else if (SeekTimeMS <= ConsiderFastAtMS)
+			else
 			{
-				SpeedClass = ESpeedClass::Fast;
+				UE_LOG(LogDerivedDataCache, Warning, TEXT("Access to %s appears to be slow. ")
+					TEXT("'Touch' will be disabled and queries/writes will be limited."), *CachePath);
 			}
-			else if (SeekTimeMS >= ConsiderSlowAtMS)
+			bTouch = false;
+			//bReadOnly = true;
+		}
+
+		if (!bReadOnly)
+		{
+			if (FString(FCommandLine::Get()).Contains(TEXT("Run=DerivedDataCache")))
 			{
-				SpeedClass = ESpeedClass::Slow;
+				bTouch = true; // we always touch files when running the DDC commandlet
 			}
-			else
+
+			// The command line (-ddctouch) enables touch on all filesystem backends if specified.
+			bTouch = bTouch || FParse::Param(FCommandLine::Get(), TEXT("DDCTOUCH"));
+
+			if (bTouch)
 			{
-				SpeedClass = ESpeedClass::Ok;
+				UE_LOG(LogDerivedDataCache, Display, TEXT("Files in %s will be touched."), *CachePath);
 			}
 
-			UE_LOG(LogDerivedDataCache, Display, TEXT("Performance to %s: Latency=%.02fms. RandomReadSpeed=%.02fMBs, RandomWriteSpeed=%.02fMBs. Assigned SpeedClass '%s'"), 
-				*CachePath, SeekTimeMS, ReadSpeedMBs, WriteSpeedMBs, LexToString(SpeedClass));
+			bool bClean = false;
+			bool bDeleteUnused = true;
+			FParse::Bool(InParams, TEXT("Clean="), bClean);
+			FParse::Bool(InParams, TEXT("DeleteUnused="), bDeleteUnused);
+			bDeleteUnused = bDeleteUnused && !FParse::Param(FCommandLine::Get(), TEXT("NODDCCLEANUP"));
 
-			if (SpeedClass <= FDerivedDataBackendInterface::ESpeedClass::Slow && !bReadOnly)
+			if (bClean || bDeleteUnused)
 			{
-				if (GIsBuildMachine)
+				FFileSystemCacheStoreMaintainerParams MaintainerParams;
+				MaintainerParams.MaxFileAge = FTimespan::FromDays(DaysToDeleteUnusedFiles);
+				if (bDeleteUnused)
 				{
-					UE_LOG(LogDerivedDataCache, Display, TEXT("Access to %s appears to be slow. 'Touch' will be disabled and queries/writes will be limited."), *CachePath);
+					if (!FParse::Value(InParams, TEXT("MaxFileChecksPerSec="), MaintainerParams.MaxFileScanRate))
+					{
+						int32 MaxFileScanRate;
+						if (GConfig->GetInt(TEXT("DDCCleanup"), TEXT("MaxFileChecksPerSec"), MaxFileScanRate, GEngineIni))
+						{
+							MaintainerParams.MaxFileScanRate = uint32(MaxFileScanRate);
+						}
+					}
+					FParse::Value(InParams, TEXT("FoldersToClean="), MaintainerParams.MaxDirectoryScanCount);
 				}
 				else
 				{
-					UE_LOG(LogDerivedDataCache, Warning, TEXT("Access to %s appears to be slow. 'Touch' will be disabled and queries/writes will be limited."), *CachePath);
+					MaintainerParams.ScanFrequency = FTimespan::FromDays(500.0);
 				}
-				bTouch = false;
-				//bReadOnly = true;
-			}
-
-			if (!bReadOnly)
-			{
-				if (FString(FCommandLine::Get()).Contains(TEXT("Run=DerivedDataCache")))
+				double TimeToWaitAfterInit;
+				if (bClean)
 				{
-					bTouch = true; // we always touch files when running the DDC commandlet
+					MaintainerParams.TimeToWaitAfterInit = FTimespan::Zero();
 				}
-
-				// The command line (-ddctouch) enables touch on all filesystem backends if specified. 
-				bTouch = bTouch || FParse::Param(FCommandLine::Get(), TEXT("DDCTOUCH"));
-
-				if (bTouch)
+				else if (GConfig->GetDouble(TEXT("DDCCleanup"), TEXT("TimeToWaitAfterInit"), TimeToWaitAfterInit, GEngineIni))
 				{
-					UE_LOG(LogDerivedDataCache, Display, TEXT("Files in %s will be touched."), *CachePath);
+					MaintainerParams.TimeToWaitAfterInit = FTimespan::FromSeconds(TimeToWaitAfterInit);
 				}
 
-				bool bClean = false;
-				bool bDeleteUnused = true;
-				FParse::Bool(InParams, TEXT("Clean="), bClean);
-				FParse::Bool(InParams, TEXT("DeleteUnused="), bDeleteUnused);
-				bDeleteUnused = bDeleteUnused && !FParse::Param(FCommandLine::Get(), TEXT("NODDCCLEANUP"));
+				Maintainer = MakeUnique<FFileSystemCacheStoreMaintainer>(MaintainerParams, CachePath);
 
-				if (bClean || bDeleteUnused)
+				if (bClean)
 				{
-					FFileSystemCacheStoreMaintainerParams MaintainerParams;
-					MaintainerParams.MaxFileAge = FTimespan::FromDays(DaysToDeleteUnusedFiles);
-					if (bDeleteUnused)
-					{
-						if (!FParse::Value(InParams, TEXT("MaxFileChecksPerSec="), MaintainerParams.MaxFileScanRate))
-						{
-							int32 MaxFileScanRate;
-							if (GConfig->GetInt(TEXT("DDCCleanup"), TEXT("MaxFileChecksPerSec"), MaxFileScanRate, GEngineIni))
-							{
-								MaintainerParams.MaxFileScanRate = uint32(MaxFileScanRate);
-							}
-						}
-						FParse::Value(InParams, TEXT("FoldersToClean="), MaintainerParams.MaxDirectoryScanCount);
-					}
-					else
-					{
-						MaintainerParams.ScanFrequency = FTimespan::FromDays(500.0);
-					}
-					double TimeToWaitAfterInit;
-					if (bClean)
-					{
-						MaintainerParams.TimeToWaitAfterInit = FTimespan::Zero();
-					}
-					else if (GConfig->GetDouble(TEXT("DDCCleanup"), TEXT("TimeToWaitAfterInit"), TimeToWaitAfterInit, GEngineIni))
-					{
-						MaintainerParams.TimeToWaitAfterInit = FTimespan::FromSeconds(TimeToWaitAfterInit);
-					}
-
-					Maintainer = MakeUnique<FFileSystemCacheStoreMaintainer>(MaintainerParams, CachePath);
-
-					if (bClean)
-					{
-						Maintainer->BoostPriority();
-						Maintainer->WaitForIdle();
-					}
+					Maintainer->BoostPriority();
+					Maintainer->WaitForIdle();
 				}
 			}
-			
-			if (IsUsable() && InAccessLogFileName != nullptr && *InAccessLogFileName != 0)
-			{
-				AccessLogWriter.Reset(new FAccessLogWriter(InAccessLogFileName, CachePath));
-			}
 		}
-	}
-
-	bool RunSpeedTest(double InSkipTestsIfSeeksExceedMS, double& OutSeekTimeMS, double& OutReadSpeedMBs, double& OutWriteSpeedMBs) const
-	{
-		SCOPED_BOOT_TIMING("RunSpeedTest");
-
-		//  files of increasing size. Most DDC data falls within this range so we don't want to skew by reading 
-		// large amounts of data. Ultimately we care most about latency anyway.
-		const int FileSizes[] = { 4, 8, 16, 64, 128, 256 };
-		const int NumTestFolders = 2; //(0-9)
-		const int FileSizeCount = UE_ARRAY_COUNT(FileSizes);
-
-		bool bWriteTestPassed = true;
-		bool bReadTestPassed = true;
-		bool bTestDataExists = true;
-
-		double TotalSeekTime = 0;
-		double TotalReadTime = 0;
-		double TotalWriteTime = 0;
-		int TotalDataRead = 0;
-		int TotalDataWritten = 0;
 
-		const FString AbsoluteCachePath = IFileManager::Get().ConvertToAbsolutePathForExternalAppForRead(*CachePath);
-		if (AbsoluteCachePath.Len() > MAX_CACHE_DIR_LEN)
+		if (IsUsable() && InAccessLogPath != nullptr && *InAccessLogPath != 0)
 		{
-			const FText ErrorMessage = FText::Format(NSLOCTEXT("DerivedDataCache", "PathTooLong", "Cache path {0} is longer than {1} characters...please adjust [DerivedDataBackendGraph] paths to be shorter (this leaves more room for cache keys)."), FText::FromString(AbsoluteCachePath), FText::AsNumber(MAX_CACHE_DIR_LEN));
-			FMessageDialog::Open(EAppMsgType::Ok, ErrorMessage);
-			UE_LOG(LogDerivedDataCache, Fatal, TEXT("%s"), *ErrorMessage.ToString());
+			AccessLogWriter.Reset(new FAccessLogWriter(InAccessLogPath, CachePath));
 		}
+	}
+}
+
+bool FFileSystemCacheStore::RunSpeedTest(
+	const double InSkipTestsIfSeeksExceedMS,
+	double& OutSeekTimeMS,
+	double& OutReadSpeedMBs,
+	double& OutWriteSpeedMBs) const
+{
+	SCOPED_BOOT_TIMING("RunSpeedTest");
+
+	//  files of increasing size. Most DDC data falls within this range so we don't want to skew by reading
+	// large amounts of data. Ultimately we care most about latency anyway.
+	const int FileSizes[] = { 4, 8, 16, 64, 128, 256 };
+	const int NumTestFolders = 2; //(0-9)
+	const int FileSizeCount = UE_ARRAY_COUNT(FileSizes);
+
+	bool bWriteTestPassed = true;
+	bool bReadTestPassed = true;
+	bool bTestDataExists = true;
+
+	double TotalSeekTime = 0;
+	double TotalReadTime = 0;
+	double TotalWriteTime = 0;
+	int TotalDataRead = 0;
+	int TotalDataWritten = 0;
+
+	const FString AbsoluteCachePath = IFileManager::Get().ConvertToAbsolutePathForExternalAppForRead(*CachePath);
+	if (AbsoluteCachePath.Len() > MAX_CACHE_DIR_LEN)
+	{
+		const FText ErrorMessage = FText::Format(NSLOCTEXT("DerivedDataCache", "PathTooLong",
+			"Cache path {0} is longer than {1} characters... "
+			"please adjust [DerivedDataBackendGraph] paths to be shorter (this leaves more room for cache keys)."),
+			FText::FromString(AbsoluteCachePath), FText::AsNumber(MAX_CACHE_DIR_LEN));
+		FMessageDialog::Open(EAppMsgType::Ok, ErrorMessage);
+		UE_LOG(LogDerivedDataCache, Fatal, TEXT("%s"), *ErrorMessage.ToString());
+	}
 
-		TArray<FString> Paths;
-		TArray<FString> MissingFiles;
+	TArray<FString> Paths;
+	TArray<FString> MissingFiles;
 
-		MissingFiles.Reserve(NumTestFolders * FileSizeCount);
+	MissingFiles.Reserve(NumTestFolders * FileSizeCount);
 
-		const FString TestDataPath = FPaths::Combine(CachePath, TEXT("TestData"));
+	const FString TestDataPath = FPaths::Combine(CachePath, TEXT("TestData"));
 
-		// create an upfront map of paths to data size in bytes
-		// create the paths we'll use. <path>/0/TestData.dat, <path>/1/TestData.dat etc. If those files don't exist we'll
-		// create them which will likely give an invalid result when measuring them now but not in the future...
-		TMap<FString, int> TestFileEntries;
-		for (int iSize = 0; iSize < FileSizeCount; iSize++)
+	// create an upfront map of paths to data size in bytes
+	// create the paths we'll use. <path>/0/TestData.dat, <path>/1/TestData.dat etc. If those files don't exist we'll
+	// create them which will likely give an invalid result when measuring them now but not in the future...
+	TMap<FString, int> TestFileEntries;
+	for (int iSize = 0; iSize < FileSizeCount; iSize++)
+	{
+		// make sure we dont stat/read/write to consecuting files in folders
+		for (int iFolder = 0; iFolder < NumTestFolders; iFolder++)
 		{
-			// make sure we dont stat/read/write to consecuting files in folders
-			for (int iFolder = 0; iFolder < NumTestFolders; iFolder++)
-			{
-				int FileSizeKB = FileSizes[iSize];
-				FString Path = FPaths::Combine(CachePath, TEXT("TestData"), *FString::FromInt(iFolder), *FString::Printf(TEXT("TestData_%dkb.dat"), FileSizeKB));
-				TestFileEntries.Add(Path, FileSizeKB * 1024);
-			}
+			int FileSizeKB = FileSizes[iSize];
+			FString Path = FPaths::Combine(CachePath, TEXT("TestData"), *FString::FromInt(iFolder),
+				*FString::Printf(TEXT("TestData_%dkb.dat"), FileSizeKB));
+			TestFileEntries.Add(Path, FileSizeKB * 1024);
 		}
+	}
 
-		// measure latency by checking for the presence of all these files. We'll also track which don't exist..
-		const double StatStartTime = FPlatformTime::Seconds();
- 		for (auto& KV : TestFileEntries)
-		{
-			FFileStatData StatData = IFileManager::Get().GetStatData(*KV.Key);
+	// measure latency by checking for the presence of all these files. We'll also track which don't exist..
+	const double StatStartTime = FPlatformTime::Seconds();
+ 	for (auto& KV : TestFileEntries)
+	{
+		FFileStatData StatData = IFileManager::Get().GetStatData(*KV.Key);
 
-			if (!StatData.bIsValid || StatData.FileSize != KV.Value)
-			{
-				MissingFiles.Add(KV.Key);
-			}
+		if (!StatData.bIsValid || StatData.FileSize != KV.Value)
+		{
+			MissingFiles.Add(KV.Key);
 		}
+	}
 
-		// save total stat time
-		TotalSeekTime = (FPlatformTime::Seconds() - StatStartTime);
-
-		// calculate seek time here
-		OutSeekTimeMS = (TotalSeekTime / TestFileEntries.Num()) * 1000;
+	// save total stat time
+	TotalSeekTime = (FPlatformTime::Seconds() - StatStartTime);
 
-		UE_LOG(LogDerivedDataCache, Verbose, TEXT("Stat tests to %s took %.02f seconds"), *CachePath, TotalSeekTime);
+	// calculate seek time here
+	OutSeekTimeMS = (TotalSeekTime / TestFileEntries.Num()) * 1000;
 
-		// if seek times are very slow do a single read/write test just to confirm access
-		/*if (OutSeekTimeMS >= InSkipTestsIfSeeksExceedMS)
-		{
-			UE_LOG(LogDerivedDataCache, Warning, TEXT("Limiting read/write speed tests due to seek times of %.02f exceeding %.02fms. Values will be inaccurate."), OutSeekTimeMS, InSkipTestsIfSeeksExceedMS);
+	UE_LOG(LogDerivedDataCache, Verbose, TEXT("Stat tests to %s took %.02f seconds"), *CachePath, TotalSeekTime);
 
-			FString Path = TestFileEntries.begin()->Key;
-			int Size = TestFileEntries.begin()->Value;
+	// if seek times are very slow do a single read/write test just to confirm access
+	/*if (OutSeekTimeMS >= InSkipTestsIfSeeksExceedMS)
+	{
+		UE_LOG(LogDerivedDataCache, Warning,
+			TEXT("Limiting read/write speed tests due to seek times of %.02f exceeding %.02fms. ")
+			TEXT("Values will be inaccurate."), OutSeekTimeMS, InSkipTestsIfSeeksExceedMS);
 
-			TestFileEntries.Reset();
-			TestFileEntries.Add(Path, Size);		
-		}*/
-	
-		// create any files that were missing
-		if (!bReadOnly)
-		{
-			TArray<uint8> Data;
-			for (auto& File : MissingFiles)
-			{
-				const int DesiredSize = TestFileEntries[File];
-				Data.SetNumUninitialized(DesiredSize);
+		FString Path = TestFileEntries.begin()->Key;
+		int Size = TestFileEntries.begin()->Value;
 
-				if (!FFileHelper::SaveArrayToFile(Data, *File, &IFileManager::Get(), FILEWRITE_Silent))
-				{
-					// handle the case where something else may have created the path at the same time. This is less about multiple users
-					// and more about things like SCW's / UnrealPak that can spin up multiple instances at once
-					if (!IFileManager::Get().FileExists(*File))
-					{
-						uint32 ErrorCode = FPlatformMisc::GetLastError();
-						TCHAR ErrorBuffer[1024];
-						FPlatformMisc::GetSystemErrorMessage(ErrorBuffer, 1024, ErrorCode);
-						UE_LOG(LogDerivedDataCache, Warning, TEXT("Fail to create %s, derived data cache to this directory will be read only. WriteError: %u (%s)"), *File, ErrorCode, ErrorBuffer);
-						bTestDataExists = false;
-						bWriteTestPassed = false;
-						break;
-					}
-				}
-			}
-		}
+		TestFileEntries.Reset();
+		TestFileEntries.Add(Path, Size);
+	}*/
 
-		// now read all sizes from random folders
+	// create any files that were missing
+	if (!bReadOnly)
+	{
+		TArray<uint8> Data;
+		for (auto& File : MissingFiles)
 		{
-			const int ArraySize = UE_ARRAY_COUNT(FileSizes);
-			TArray<uint8> TempData;
-			TempData.Empty(FileSizes[ArraySize - 1] * 1024);
+			const int DesiredSize = TestFileEntries[File];
+			Data.SetNumUninitialized(DesiredSize);
 
-			const double ReadStartTime = FPlatformTime::Seconds();
-
-			for (auto& KV : TestFileEntries)
+			if (!FFileHelper::SaveArrayToFile(Data, *File, &IFileManager::Get(), FILEWRITE_Silent))
 			{
-				const int FileSize = KV.Value;
-				const FString& FilePath = KV.Key;
-
-				if (!FFileHelper::LoadFileToArray(TempData, *FilePath, FILEREAD_Silent))
+				// Handle the case where something else may have created the path at the same time.
+				// This is less about multiple users and more about things like SCW's / UnrealPak
+				// that can spin up multiple instances at once.
+				if (!IFileManager::Get().FileExists(*File))
 				{
 					uint32 ErrorCode = FPlatformMisc::GetLastError();
 					TCHAR ErrorBuffer[1024];
 					FPlatformMisc::GetSystemErrorMessage(ErrorBuffer, 1024, ErrorCode);
-					UE_LOG(LogDerivedDataCache, Warning, TEXT("Fail to read from %s, derived data cache will be disabled. ReadError: %u (%s)"), *FilePath, ErrorCode, ErrorBuffer);
-					bReadTestPassed = false;
+					UE_LOG(LogDerivedDataCache, Warning,
+						TEXT("Fail to create %s, derived data cache to this directory will be read only. ")
+						TEXT("WriteError: %u (%s)"), *File, ErrorCode, ErrorBuffer);
+					bTestDataExists = false;
+					bWriteTestPassed = false;
 					break;
 				}
-
-				TotalDataRead += TempData.Num();
 			}
+		}
+	}
 
-			TotalReadTime = FPlatformTime::Seconds() - ReadStartTime;
+	// now read all sizes from random folders
+	{
+		const int ArraySize = UE_ARRAY_COUNT(FileSizes);
+		TArray<uint8> TempData;
+		TempData.Empty(FileSizes[ArraySize - 1] * 1024);
 
-			UE_LOG(LogDerivedDataCache, Verbose, TEXT("Read tests %s on %s and took %.02f seconds"), bReadTestPassed ? TEXT("passed") : TEXT("failed"), *CachePath, TotalReadTime);
-		}
-		
-		// do write tests if or read tests passed and our seeks were below the cut-off 
-		if (bReadTestPassed && !bReadOnly)
+		const double ReadStartTime = FPlatformTime::Seconds();
+
+		for (auto& KV : TestFileEntries)
 		{
-			// do write tests but use a unique folder that is cleaned up afterwards
-			FString CustomPath = FPaths::Combine(CachePath, TEXT("TestData"), *FGuid::NewGuid().ToString());
+			const int FileSize = KV.Value;
+			const FString& FilePath = KV.Key;
+
+			if (!FFileHelper::LoadFileToArray(TempData, *FilePath, FILEREAD_Silent))
+			{
+				uint32 ErrorCode = FPlatformMisc::GetLastError();
+				TCHAR ErrorBuffer[1024];
+				FPlatformMisc::GetSystemErrorMessage(ErrorBuffer, 1024, ErrorCode);
+				UE_LOG(LogDerivedDataCache, Warning,
+					TEXT("Fail to read from %s, derived data cache will be disabled. ReadError: %u (%s)"),
+					*FilePath, ErrorCode, ErrorBuffer);
+				bReadTestPassed = false;
+				break;
+			}
 
-			const int ArraySize = UE_ARRAY_COUNT(FileSizes);
-			TArray<uint8> TempData;
-			TempData.Empty(FileSizes[ArraySize - 1] * 1024);
+			TotalDataRead += TempData.Num();
+		}
 
-			const double WriteStartTime = FPlatformTime::Seconds();
+		TotalReadTime = FPlatformTime::Seconds() - ReadStartTime;
 
-			for (auto& KV : TestFileEntries)
-			{
-				const int FileSize = KV.Value;
-				FString FilePath = KV.Key;
+		UE_LOG(LogDerivedDataCache, Verbose, TEXT("Read tests %s on %s and took %.02f seconds"),
+			bReadTestPassed ? TEXT("passed") : TEXT("failed"), *CachePath, TotalReadTime);
+	}
 
-				TempData.SetNumUninitialized(FileSize);
+	// do write tests if or read tests passed and our seeks were below the cut-off
+	if (bReadTestPassed && !bReadOnly)
+	{
+		// do write tests but use a unique folder that is cleaned up afterwards
+		FString CustomPath = FPaths::Combine(CachePath, TEXT("TestData"), *FGuid::NewGuid().ToString());
 
-				FilePath = FilePath.Replace(*CachePath, *CustomPath);
+		const int ArraySize = UE_ARRAY_COUNT(FileSizes);
+		TArray<uint8> TempData;
+		TempData.Empty(FileSizes[ArraySize - 1] * 1024);
 
-				if (!FFileHelper::SaveArrayToFile(TempData, *FilePath, &IFileManager::Get(), FILEWRITE_Silent))
-				{
-					uint32 ErrorCode = FPlatformMisc::GetLastError();
-					TCHAR ErrorBuffer[1024];
-					FPlatformMisc::GetSystemErrorMessage(ErrorBuffer, 1024, ErrorCode);
-					UE_LOG(LogDerivedDataCache, Warning, TEXT("Fail to write to %s, derived data cache will be disabled. ReadError: %u (%s)"), *FilePath, ErrorCode, ErrorBuffer);
-					bWriteTestPassed = false;
-					break;
-				}
+		const double WriteStartTime = FPlatformTime::Seconds();
 
-				TotalDataWritten += TempData.Num();
-			}
+		for (auto& KV : TestFileEntries)
+		{
+			const int FileSize = KV.Value;
+			FString FilePath = KV.Key;
 
-			TotalWriteTime = FPlatformTime::Seconds() - WriteStartTime;
+			TempData.SetNumUninitialized(FileSize);
 
-			UE_LOG(LogDerivedDataCache, Verbose, TEXT("write tests %s on %s and took %.02f seconds"), bWriteTestPassed ? TEXT("passed") : TEXT("failed"), *CachePath, TotalReadTime)
+			FilePath = FilePath.Replace(*CachePath, *CustomPath);
+
+			if (!FFileHelper::SaveArrayToFile(TempData, *FilePath, &IFileManager::Get(), FILEWRITE_Silent))
+			{
+				uint32 ErrorCode = FPlatformMisc::GetLastError();
+				TCHAR ErrorBuffer[1024];
+				FPlatformMisc::GetSystemErrorMessage(ErrorBuffer, 1024, ErrorCode);
+				UE_LOG(LogDerivedDataCache, Warning,
+					TEXT("Fail to write to %s, derived data cache will be disabled. ReadError: %u (%s)"),
+					*FilePath, ErrorCode, ErrorBuffer);
+				bWriteTestPassed = false;
+				break;
+			}
 
-			// remove the custom path but do it async as this can be slow on remote drives
-			AsyncTask(ENamedThreads::AnyThread, [CustomPath]() {
-				IFileManager::Get().DeleteDirectory(*CustomPath, false, true);
-			});
-	
-			// check latency and speed. Read values should always be valid
-			const double ReadSpeedMBs = (bReadTestPassed ? (TotalDataRead / TotalReadTime) : 0) / (1024 * 1024);
-			const double WriteSpeedMBs = (bWriteTestPassed ? (TotalDataWritten / TotalWriteTime) : 0) / (1024 * 1024);
-			const double SeekTimeMS = (TotalSeekTime / TestFileEntries.Num()) * 1000;
+			TotalDataWritten += TempData.Num();
 		}
 
-		const double TotalTestTime = FPlatformTime::Seconds() - StatStartTime;
+		TotalWriteTime = FPlatformTime::Seconds() - WriteStartTime;
 
-		UE_LOG(LogDerivedDataCache, Log, TEXT("Speed tests for %s took %.02f seconds"), *CachePath, TotalTestTime);
+		UE_LOG(LogDerivedDataCache, Verbose, TEXT("write tests %s on %s and took %.02f seconds"),
+			bWriteTestPassed ? TEXT("passed") : TEXT("failed"), *CachePath, TotalReadTime)
 
-		// check latency and speed. Read values should always be valid
-		OutReadSpeedMBs = (bReadTestPassed ? (TotalDataRead / TotalReadTime) : 0) / (1024 * 1024);
-		OutWriteSpeedMBs = (bWriteTestPassed ? (TotalDataWritten / TotalWriteTime) : 0) / (1024 * 1024);
+		// remove the custom path but do it async as this can be slow on remote drives
+		AsyncTask(ENamedThreads::AnyThread, [CustomPath]() {
+			IFileManager::Get().DeleteDirectory(*CustomPath, false, true);
+		});
 
-		return bWriteTestPassed || bReadTestPassed;
+		// check latency and speed. Read values should always be valid
+		const double ReadSpeedMBs = (bReadTestPassed ? (TotalDataRead / TotalReadTime) : 0) / (1024 * 1024);
+		const double WriteSpeedMBs = (bWriteTestPassed ? (TotalDataWritten / TotalWriteTime) : 0) / (1024 * 1024);
+		const double SeekTimeMS = (TotalSeekTime / TestFileEntries.Num()) * 1000;
 	}
 
-	/** Return a type for this interface */
-	virtual FString GetDisplayName() const override
-	{
-		return FString(TEXT("File System"));
-	}
+	const double TotalTestTime = FPlatformTime::Seconds() - StatStartTime;
 
+	UE_LOG(LogDerivedDataCache, Log, TEXT("Speed tests for %s took %.02f seconds"), *CachePath, TotalTestTime);
 
-	/** Return a name for this interface */
-	virtual FString GetName() const override 
-	{ 
-		return CachePath; 
-	}
+	// check latency and speed. Read values should always be valid
+	OutReadSpeedMBs = (bReadTestPassed ? (TotalDataRead / TotalReadTime) : 0) / (1024 * 1024);
+	OutWriteSpeedMBs = (bWriteTestPassed ? (TotalDataWritten / TotalWriteTime) : 0) / (1024 * 1024);
 
-	/** return true if the cache is usable **/
-	bool IsUsable() const
-	{
-		return !bDisabled;
-	}
+	return bWriteTestPassed || bReadTestPassed;
+}
 
-	/** return true if this cache is writable **/
-	virtual bool IsWritable() const override
-	{
-		return !bReadOnly && !bDisabled;
-	}
+bool FFileSystemCacheStore::CachedDataProbablyExists(const TCHAR* const CacheKey)
+{
+	TRACE_CPUPROFILER_EVENT_SCOPE(FileSystemDDC_Exist);
+	TRACE_COUNTER_INCREMENT(FileSystemDDC_Exist);
+	COOK_STAT(auto Timer = UsageStats.TimeProbablyExists());
+	check(IsUsable());
 
-	/** Returns a class of speed for this interface **/
-	virtual ESpeedClass GetSpeedClass() const override
+	if (ShouldSimulateMiss(CacheKey))
 	{
-		return SpeedClass;
+		return false;
 	}
 
-	class FAccessLogWriter
+	const FString Filename = BuildFilename(CacheKey);
+
+	const FFileStatData FileStat = IFileManager::Get().GetStatData(*Filename);
+
+	if (FileStat.bIsValid)
 	{
-	public:
-		FAccessLogWriter(const TCHAR* FileName, const FString& CachePath)
-			: Archive(IFileManager::Get().CreateFileWriter(FileName, FILEWRITE_AllowRead))
-			, BasePath(CachePath / TEXT(""))
-		{
-		}
+		const FDateTime TimeStamp = FileStat.ModificationTime;
 
-		void Append(const TCHAR* CacheKey, FStringView Path)
+		// Update file timestamp to prevent it from being deleted by maintenance.
+		if (bTouch || (!bReadOnly && (FDateTime::UtcNow() - TimeStamp).GetTotalDays() > (DaysToDeleteUnusedFiles / 4)))
 		{
-			FScopeLock Lock(&CriticalSection);
-
-			bool bIsAlreadyInSet = false;
-			CacheKeys.FindOrAdd(FString(CacheKey), &bIsAlreadyInSet);
-			if (!bIsAlreadyInSet)
-			{
-				AppendPath(Path);
-			}
+			IFileManager::Get().SetTimeStamp(*Filename, FDateTime::UtcNow());
 		}
 
-		void Append(const FIoHash& RawHash, FStringView Path)
+		if (AccessLogWriter.IsValid())
 		{
-			FScopeLock Lock(&CriticalSection);
-
-			bool bIsAlreadyInSet = false;
-			ContentKeys.FindOrAdd(RawHash, &bIsAlreadyInSet);
-			if (!bIsAlreadyInSet)
-			{
-				AppendPath(Path);
-			}
+			AccessLogWriter->Append(CacheKey, Filename);
 		}
 
-		void Append(const FCacheKey& CacheKey, FStringView Path)
-		{
-			FScopeLock Lock(&CriticalSection);
-
-			bool bIsAlreadyInSet = false;
-			RecordKeys.FindOrAdd(CacheKey, &bIsAlreadyInSet);
-			if (!bIsAlreadyInSet)
-			{
-				AppendPath(Path);
-			}
-		}
+		TRACE_COUNTER_INCREMENT(FileSystemDDC_ExistHit);
+		COOK_STAT(Timer.AddHit(0));
+	}
 
-	private:
-		void AppendPath(FStringView Path)
+	// If not using a shared cache, record a (probable) miss
+	if (!FileStat.bIsValid && !GetDerivedDataCacheRef().GetUsingSharedDDC())
+	{
+		// store a cache miss
+		FScopeLock ScopeLock(&SynchronizationObject);
+		if (!DDCNotificationCacheTimes.Contains(CacheKey))
 		{
-			if (Path.StartsWith(BasePath))
-			{
-				const FTCHARToUTF8 PathUtf8(Path);
-				Archive->Serialize(const_cast<ANSICHAR*>(PathUtf8.Get()), PathUtf8.Length());
-				Archive->Serialize(const_cast<ANSICHAR*>(LINE_TERMINATOR_ANSI), sizeof(LINE_TERMINATOR_ANSI) - 1);
-			}
+			DDCNotificationCacheTimes.Add(CacheKey, FPlatformTime::Seconds());
 		}
+	}
 
-		TUniquePtr<FArchive> Archive;
-		FString BasePath;
-		FCriticalSection CriticalSection;
-		TSet<FString> CacheKeys;
-		TSet<FIoHash> ContentKeys;
-		TSet<FCacheKey> RecordKeys;
-	};
-	
-	
-	/**
-	 * Synchronous test for the existence of a cache item
-	 *
-	 * @param	CacheKey	Alphanumeric+underscore key of this cache item
-	 * @return				true if the data probably will be found, this can't be guaranteed because of concurrency in the backends, corruption, etc
-	 */
-	virtual bool CachedDataProbablyExists(const TCHAR* CacheKey) override
-	{
-		TRACE_CPUPROFILER_EVENT_SCOPE(FileSystemDDC_Exist);
-		TRACE_COUNTER_INCREMENT(FileSystemDDC_Exist);
-		COOK_STAT(auto Timer = UsageStats.TimeProbablyExists());
-		check(IsUsable());
-
-		// if we're a slow device just say we have the data. It's faster to try and fail than it
-		// is to check and succeed.
-		
-		// todo(@agrant: Some types currently use ProbablyExists as a guarantee. Disabling this until those can be 
-		// addressed.
-		/*
-		if (GetSpeedClass() <= FDerivedDataBackendInterface::ESpeedClass::Slow)
-		{
-			return true;
-		}
-		*/
+	UE_LOG(LogDerivedDataCache, Verbose,
+		TEXT("%s CachedDataProbablyExists=%d for %s"), *GetName(), FileStat.bIsValid, CacheKey);
 
-		if (ShouldSimulateMiss(CacheKey))
-		{
-			return false;
-		}
+	return FileStat.bIsValid;
+}
 
-		FString Filename = BuildFilename(CacheKey);
+bool FFileSystemCacheStore::GetCachedData(const TCHAR* const CacheKey, TArray<uint8>& Data)
+{
+	TRACE_CPUPROFILER_EVENT_SCOPE(FileSystemDDC_Get);
+	TRACE_COUNTER_INCREMENT(FileSystemDDC_Get);
+	COOK_STAT(auto Timer = UsageStats.TimeGet());
+	check(IsUsable());
+	const FString Filename = BuildFilename(CacheKey);
+	const double StartTime = FPlatformTime::Seconds();
 
-		FFileStatData FileStat = IFileManager::Get().GetStatData(*Filename);
+	if (ShouldSimulateMiss(CacheKey))
+	{
+		return false;
+	}
 
-		if (FileStat.bIsValid)
+	if (FFileHelper::LoadFileToArray(Data,*Filename,FILEREAD_Silent))
+	{
+		if (SpeedClass >= FDerivedDataBackendInterface::ESpeedClass::Fast && (!bReadOnly || bTouch))
 		{
-			FDateTime TimeStamp = FileStat.ModificationTime;
-
-			// Update file timestamp to prevent it from being deleted by DDC Cleanup.
-			if (bTouch || 
-				 (!bReadOnly && (FDateTime::UtcNow() - TimeStamp).GetTotalDays() > (DaysToDeleteUnusedFiles / 4)))
+			const FFileStatData FileStat = IFileManager::Get().GetStatData(*Filename);
+			if (FileStat.bIsValid)
 			{
-				IFileManager::Get().SetTimeStamp(*Filename, FDateTime::UtcNow());
+				const FDateTime TimeStamp = FileStat.ModificationTime;
+				// Update file timestamp to prevent it from being deleted by DDC Cleanup.
+				if (bTouch ||
+					(!bReadOnly && (FDateTime::UtcNow() - TimeStamp).GetTotalDays() > (DaysToDeleteUnusedFiles / 4)))
+				{
+					IFileManager::Get().SetTimeStamp(*Filename, FDateTime::UtcNow());
+				}
 			}
+		}
 
-			if (AccessLogWriter.IsValid())
-			{
-				AccessLogWriter->Append(CacheKey, Filename);
-			}
+		const double ReadDuration = FPlatformTime::Seconds() - StartTime;
+		const double ReadSpeed = (Data.Num() / ReadDuration) / (1024.0 * 1024.0);
 
-			TRACE_COUNTER_INCREMENT(FileSystemDDC_ExistHit);
-			COOK_STAT(Timer.AddHit(0));
+		if(!GIsBuildMachine && ReadDuration > 5.0)
+		{
+			// Slower than 0.5MB/s?
+			UE_CLOG(ReadSpeed < 0.5, LogDerivedDataCache, Warning,
+				TEXT("%s is very slow (%.2fMB/s) when accessing %s, consider disabling it."),
+				*CachePath, ReadSpeed, *Filename);
 		}
 
-		// If not using a shared cache, record a (probable) miss
-		if (!FileStat.bIsValid && !GetDerivedDataCacheRef().GetUsingSharedDDC())
+		if (AccessLogWriter.IsValid())
 		{
-			// store a cache miss
-			FScopeLock ScopeLock(&SynchronizationObject);
-			if (!DDCNotificationCacheTimes.Contains(CacheKey))
-			{
-				DDCNotificationCacheTimes.Add(CacheKey, FPlatformTime::Seconds());
-			}
+			AccessLogWriter->Append(CacheKey, Filename);
 		}
 
-		UE_LOG(LogDerivedDataCache, Verbose, TEXT("%s CachedDataProbablyExists=%d for %s"), *GetName(), FileStat.bIsValid, CacheKey);
+		UE_LOG(LogDerivedDataCache, Verbose,
+			TEXT("%s: Cache hit on %s (%d bytes, %.02f secs, %.2fMB/s)"),
+			*GetName(), CacheKey, Data.Num(), ReadDuration, ReadSpeed);
+		TRACE_COUNTER_INCREMENT(FileSystemDDC_GetHit);
+		TRACE_COUNTER_ADD(FileSystemDDC_BytesRead, int64(Data.Num()));
+		COOK_STAT(Timer.AddHit(Data.Num()));
+		return true;
+	}
 
-		return FileStat.bIsValid;
-	}	
+	UE_LOG(LogDerivedDataCache, Verbose, TEXT("%s: Cache miss on %s"), *GetName(), CacheKey);
+	Data.Empty();
 
-	/**
-	 * Synchronous retrieve of a cache item
-	 *
-	 * @param	CacheKey	Alphanumeric+underscore key of this cache item
-	 * @param	OutData		Buffer to receive the results, if any were found
-	 * @return				true if any data was found, and in this case OutData is non-empty
-	 */
-	virtual bool GetCachedData(const TCHAR* CacheKey, TArray<uint8>& Data) override
+	// If not using a shared cache, record a miss
+	if (!GetDerivedDataCacheRef().GetUsingSharedDDC())
 	{
-		TRACE_CPUPROFILER_EVENT_SCOPE(FileSystemDDC_Get);
-		TRACE_COUNTER_INCREMENT(FileSystemDDC_Get);
-		COOK_STAT(auto Timer = UsageStats.TimeGet());
-		check(IsUsable());
-		FString Filename = BuildFilename(CacheKey);
-		double StartTime = FPlatformTime::Seconds();
-
-		if (ShouldSimulateMiss(CacheKey))
+		// store a cache miss
+		FScopeLock ScopeLock(&SynchronizationObject);
+		if (!DDCNotificationCacheTimes.Contains(CacheKey))
 		{
-			return false;
+			DDCNotificationCacheTimes.Add(CacheKey, FPlatformTime::Seconds());
 		}
+	}
 
-		if (FFileHelper::LoadFileToArray(Data,*Filename,FILEREAD_Silent))
-		{
-			if (SpeedClass >= FDerivedDataBackendInterface::ESpeedClass::Fast && (!bReadOnly || bTouch))
-			{
-				FFileStatData FileStat = IFileManager::Get().GetStatData(*Filename);
-				if (FileStat.bIsValid)
-				{
-					FDateTime TimeStamp = FileStat.ModificationTime;
-					// Update file timestamp to prevent it from being deleted by DDC Cleanup.
-					if (bTouch ||
-						(!bReadOnly && (FDateTime::UtcNow() - TimeStamp).GetTotalDays() > (DaysToDeleteUnusedFiles / 4)))
-					{
-						IFileManager::Get().SetTimeStamp(*Filename, FDateTime::UtcNow());
-					}
-				}
-			}
+	return false;
+}
 
-			double ReadDuration = FPlatformTime::Seconds() - StartTime;
-			double ReadSpeed = (Data.Num() / ReadDuration) / (1024.0 * 1024.0);
+bool FFileSystemCacheStore::WouldCache(const TCHAR* const CacheKey, const TConstArrayView<uint8> InData)
+{
+	return IsWritable() && !CachedDataProbablyExists(CacheKey);
+}
 
-			if(!GIsBuildMachine && ReadDuration > 5.0)
-			{				
-				// Slower than 0.5MB/s?
-				UE_CLOG(ReadSpeed < 0.5, LogDerivedDataCache, Warning, TEXT("%s is very slow (%.2fMB/s) when accessing %s, consider disabling it."), *CachePath, ReadSpeed, *Filename);
-			}
+FFileSystemCacheStore::EPutStatus FFileSystemCacheStore::PutCachedData(
+	const TCHAR* const CacheKey,
+	const TConstArrayView<uint8> Data,
+	const bool bPutEvenIfExists)
+{
+	TRACE_CPUPROFILER_EVENT_SCOPE(FileSystemDDC_Put);
+	TRACE_COUNTER_INCREMENT(FileSystemDDC_Put);
+	COOK_STAT(auto Timer = UsageStats.TimePut());
+	check(IsUsable());
 
-			if (AccessLogWriter.IsValid())
-			{
-				AccessLogWriter->Append(CacheKey, Filename);
-			}
+	if (IsWritable())
+	{
+		const FString Filename = BuildFilename(CacheKey);
 
-			UE_LOG(LogDerivedDataCache, Verbose, TEXT("%s: Cache hit on %s (%d bytes, %.02f secs, %.2fMB/s)"), *GetName(), CacheKey, Data.Num(), ReadDuration, ReadSpeed);
-			TRACE_COUNTER_INCREMENT(FileSystemDDC_GetHit);
-			TRACE_COUNTER_ADD(FileSystemDDC_BytesRead, int64(Data.Num()));
-			COOK_STAT(Timer.AddHit(Data.Num()));
-			return true;
+		if (AccessLogWriter.IsValid())
+		{
+			AccessLogWriter->Append(CacheKey, Filename);
 		}
 
-		UE_LOG(LogDerivedDataCache, Verbose, TEXT("%s: Cache miss on %s"), *GetName(), CacheKey);
-		Data.Empty();
-
-		// If not using a shared cache, record a miss
-		if (!GetDerivedDataCacheRef().GetUsingSharedDDC())
+		// don't put anything we pretended didn't exist
+		if (ShouldSimulateMiss(CacheKey))
 		{
-			// store a cache miss
-			FScopeLock ScopeLock(&SynchronizationObject);
-			if (!DDCNotificationCacheTimes.Contains(CacheKey))
-			{
-				DDCNotificationCacheTimes.Add(CacheKey, FPlatformTime::Seconds());
-			}
+			return EPutStatus::Skipped;
 		}
 
-		return false;
-	}
-
-	/**
-	 * Would we cache this? Say yes so long as we aren't read-only.
-	 */
-	bool WouldCache(const TCHAR* CacheKey, TArrayView<const uint8> InData) override
-	{
-		return IsWritable() && !CachedDataProbablyExists(CacheKey);
-	}
-
-	/**
-	 * Asynchronous, fire-and-forget placement of a cache item
-	 *
-	 * @param	CacheKey	Alphanumeric+underscore key of this cache item
-	 * @param	OutData		Buffer containing the data to cache, can be destroyed after the call returns, immediately
-	 * @param	bPutEvenIfExists	If true, then do not attempt skip the put even if CachedDataProbablyExists returns true
-	 */
-	virtual EPutStatus PutCachedData(const TCHAR* CacheKey, TArrayView<const uint8> Data, bool bPutEvenIfExists) override
-	{
-		TRACE_CPUPROFILER_EVENT_SCOPE(FileSystemDDC_Put);
-		TRACE_COUNTER_INCREMENT(FileSystemDDC_Put);
-		COOK_STAT(auto Timer = UsageStats.TimePut());
-		check(IsUsable());
+		EPutStatus Status = EPutStatus::NotCached;
 
-		if (IsWritable())
+		if (bPutEvenIfExists || !CachedDataProbablyExists(CacheKey))
 		{
-			FString Filename = BuildFilename(CacheKey);
-
-			if (AccessLogWriter.IsValid())
-			{
-				AccessLogWriter->Append(CacheKey, Filename);
-			}
-
-			// don't put anything we pretended didn't exist
-			if (ShouldSimulateMiss(CacheKey))
-			{
-				return EPutStatus::Skipped;
-			}
-
-			EPutStatus Status = EPutStatus::NotCached;
-
-			if (bPutEvenIfExists || !CachedDataProbablyExists(CacheKey))
+			TRACE_COUNTER_INCREMENT(FileSystemDDC_PutHit);
+			TRACE_COUNTER_ADD(FileSystemDDC_BytesWritten, int64(Data.Num()));
+			COOK_STAT(Timer.AddHit(Data.Num()));
+			check(Data.Num());
+			FString TempFilename(TEXT("temp."));
+			TempFilename += FGuid::NewGuid().ToString();
+			TempFilename = FPaths::GetPath(Filename) / TempFilename;
+			if (FFileHelper::SaveArrayToFile(Data, *TempFilename, &IFileManager::Get(), FILEWRITE_Silent))
 			{
-				TRACE_COUNTER_INCREMENT(FileSystemDDC_PutHit);
-				TRACE_COUNTER_ADD(FileSystemDDC_BytesWritten, int64(Data.Num()));
-				COOK_STAT(Timer.AddHit(Data.Num()));
-				check(Data.Num());
-				FString TempFilename(TEXT("temp.")); 
-				TempFilename += FGuid::NewGuid().ToString();
-				TempFilename = FPaths::GetPath(Filename) / TempFilename;
-				bool bResult;
-				{
-					bResult = FFileHelper::SaveArrayToFile(Data, *TempFilename, &IFileManager::Get(), FILEWRITE_Silent);
-				}
-				if (bResult)
+				if (IFileManager::Get().FileSize(*TempFilename) == Data.Num())
 				{
-					if (IFileManager::Get().FileSize(*TempFilename) == Data.Num())
+					bool DoMove = !CachedDataProbablyExists(CacheKey);
+					if (bPutEvenIfExists && !DoMove)
 					{
-						bool DoMove = !CachedDataProbablyExists(CacheKey);
-						if (bPutEvenIfExists && !DoMove)
-						{
-							DoMove = true;
-							RemoveCachedData(CacheKey, /*bTransient=*/ false);
-						}
-						if (DoMove) 
+						DoMove = true;
+						RemoveCachedData(CacheKey, /*bTransient=*/ false);
+					}
+					if (DoMove)
+					{
+						if (!IFileManager::Get().Move(*Filename, *TempFilename, true, true, false, true))
 						{
-							if (!IFileManager::Get().Move(*Filename, *TempFilename, true, true, false, true))
-							{
-								UE_LOG(LogDerivedDataCache, Log, TEXT("%s: Move collision, attempt at redundant update, OK %s."), *GetName(),*Filename);
-							}
-							else
-							{
-								UE_LOG(LogDerivedDataCache, Verbose, TEXT("%s: Successful cache put of %s to %s"),*GetName(), CacheKey, *Filename);
-								Status = EPutStatus::Cached;
-							}
+							UE_LOG(LogDerivedDataCache, Log,
+								TEXT("%s: Move collision, attempt at redundant update, OK %s."), *GetName(), *Filename);
 						}
 						else
 						{
+							UE_LOG(LogDerivedDataCache, Verbose,
+								TEXT("%s: Successful cache put of %s to %s"), *GetName(), CacheKey, *Filename);
 							Status = EPutStatus::Cached;
 						}
 					}
 					else
 					{
-						UE_LOG(LogDerivedDataCache, Warning, TEXT("%s: Temp file is short %s!"), *GetName(), *TempFilename);
+						Status = EPutStatus::Cached;
 					}
 				}
 				else
 				{
-					uint32 ErrorCode = FPlatformMisc::GetLastError();
-					TCHAR ErrorBuffer[1024];
-					FPlatformMisc::GetSystemErrorMessage(ErrorBuffer, 1024, ErrorCode);
-					UE_LOG(LogDerivedDataCache, Warning, TEXT("FFileSystemDerivedDataBackend: Could not write temp file %s! Error: %u (%s)"), *TempFilename, ErrorCode, ErrorBuffer);
-				}
-				// if everything worked, this is not necessary, but we will make every effort to avoid leaving junk in the cache
-				if (FPaths::FileExists(TempFilename))
-				{
-					IFileManager::Get().Delete(*TempFilename, false, false, true);
+					UE_LOG(LogDerivedDataCache, Warning, TEXT("%s: Temp file is short %s!"), *GetName(), *TempFilename);
 				}
 			}
 			else
 			{
-				COOK_STAT(Timer.AddMiss(Data.Num()));
-				UE_LOG(LogDerivedDataCache, Verbose, TEXT("%s skipping put to existing file %s"), *GetName(), CacheKey);
-				Status = EPutStatus::Cached;
+				const uint32 ErrorCode = FPlatformMisc::GetLastError();
+				TCHAR ErrorBuffer[1024];
+				FPlatformMisc::GetSystemErrorMessage(ErrorBuffer, 1024, ErrorCode);
+				UE_LOG(LogDerivedDataCache, Warning,
+					TEXT("FFileSystemCacheStore: Could not write temp file %s! Error: %u (%s)"),
+					*TempFilename, ErrorCode, ErrorBuffer);
+			}
+			// if everything worked, this is not necessary, but we will make every effort to avoid leaving junk in the cache
+			if (FPaths::FileExists(TempFilename))
+			{
+				IFileManager::Get().Delete(*TempFilename, false, false, true);
 			}
+		}
+		else
+		{
+			COOK_STAT(Timer.AddMiss(Data.Num()));
+			UE_LOG(LogDerivedDataCache, Verbose, TEXT("%s skipping put to existing file %s"), *GetName(), CacheKey);
+			Status = EPutStatus::Cached;
+		}
 
-			// If not using a shared cache, update estimated build time
-			if (!GetDerivedDataCacheRef().GetUsingSharedDDC())
+		// If not using a shared cache, update estimated build time
+		if (!GetDerivedDataCacheRef().GetUsingSharedDDC())
+		{
+			FScopeLock ScopeLock(&SynchronizationObject);
+
+			if (DDCNotificationCacheTimes.Contains(CacheKey))
 			{
-				FScopeLock ScopeLock(&SynchronizationObject);
+				// There isn't any way to get exact build times in the DDC code as custom asset processing and async are factors.
+				// So, estimate the asset build time based on the delta between the cache miss and the put
+				TotalEstimatedBuildTime += (FPlatformTime::Seconds() - DDCNotificationCacheTimes[CacheKey]);
+				DDCNotificationCacheTimes.Remove(CacheKey);
 
-				if (DDCNotificationCacheTimes.Contains(CacheKey))
+				// If more than 20 seconds has been spent building assets, send out a notification
+				if (TotalEstimatedBuildTime > 20.0f)
 				{
-					// There isn't any way to get exact build times in the DDC code as custom asset processing and async are factors.
-					// So, estimate the asset build time based on the delta between the cache miss and the put
-					TotalEstimatedBuildTime += (FPlatformTime::Seconds() - DDCNotificationCacheTimes[CacheKey]);
-					DDCNotificationCacheTimes.Remove(CacheKey);
+					// Send out a DDC put notification if we have any subscribers
+					FDerivedDataCacheInterface::FOnDDCNotification& DDCNotificationEvent =
+						GetDerivedDataCacheRef().GetDDCNotificationEvent();
 
-					// If more than 20 seconds has been spent building assets, send out a notification
-					if (TotalEstimatedBuildTime > 20.0f)
+					if (DDCNotificationEvent.IsBound())
 					{
-						// Send out a DDC put notification if we have any subscribers
-						FDerivedDataCacheInterface::FOnDDCNotification& DDCNotificationEvent = GetDerivedDataCacheRef().GetDDCNotificationEvent();
-
-						if (DDCNotificationEvent.IsBound())
-						{
-							TotalEstimatedBuildTime = 0.0f;
-
-							DECLARE_CYCLE_STAT(TEXT("FSimpleDelegateGraphTask.PutCachedData"), STAT_FSimpleDelegateGraphTask_DDCNotification, STATGROUP_TaskGraphTasks);
-
-							FSimpleDelegateGraphTask::CreateAndDispatchWhenReady(
-								FSimpleDelegateGraphTask::FDelegate::CreateLambda([DDCNotificationEvent]() {
-								DDCNotificationEvent.Broadcast(FDerivedDataCacheInterface::SharedDDCPerformanceNotification);
-							}),
-								GET_STATID(STAT_FSimpleDelegateGraphTask_DDCNotification),
-								nullptr,
-								ENamedThreads::GameThread);
-						}
+						TotalEstimatedBuildTime = 0.0f;
+
+						DECLARE_CYCLE_STAT(TEXT("FSimpleDelegateGraphTask.PutCachedData"),
+						STAT_FSimpleDelegateGraphTask_DDCNotification, STATGROUP_TaskGraphTasks);
+
+						FSimpleDelegateGraphTask::CreateAndDispatchWhenReady(
+							FSimpleDelegateGraphTask::FDelegate::CreateLambda([DDCNotificationEvent]() {
+							DDCNotificationEvent.Broadcast(FDerivedDataCacheInterface::SharedDDCPerformanceNotification);
+						}),
+							GET_STATID(STAT_FSimpleDelegateGraphTask_DDCNotification),
+							nullptr,
+							ENamedThreads::GameThread);
 					}
-
 				}
-
 			}
-
-			return Status;
-		}
-		else
-		{
-			UE_LOG(LogDerivedDataCache, Verbose, TEXT("%s is read only. Skipping put of %s"), *GetName(), CacheKey);
-			return EPutStatus::NotCached;
 		}
+
+		return Status;
 	}
+	else
+	{
+		UE_LOG(LogDerivedDataCache, Verbose, TEXT("%s is read only. Skipping put of %s"), *GetName(), CacheKey);
+		return EPutStatus::NotCached;
+	}
+}
 
-	void RemoveCachedData(const TCHAR* CacheKey, bool bTransient) override
+void FFileSystemCacheStore::RemoveCachedData(const TCHAR* const CacheKey, const bool bTransient)
+{
+	check(IsUsable());
+	if (IsWritable() && (!bTransient || bPurgeTransient))
 	{
-		check(IsUsable());
-		if (IsWritable() && (!bTransient || bPurgeTransient))
+		const FString Filename = BuildFilename(CacheKey);
+		if (bTransient)
 		{
-			FString Filename = BuildFilename(CacheKey);
-			if (bTransient)
-			{
-				UE_LOG(LogDerivedDataCache,Verbose,TEXT("Deleting transient cached data. Key=%s Filename=%s"),CacheKey,*Filename);
-			}
-			IFileManager::Get().Delete(*Filename, false, false, true);
+			UE_LOG(LogDerivedDataCache, Verbose,
+				TEXT("Deleting transient cached data. Key=%s Filename=%s"), CacheKey, *Filename);
 		}
+		IFileManager::Get().Delete(*Filename, false, false, true);
 	}
+}
 
-	virtual TSharedRef<FDerivedDataCacheStatsNode> GatherUsageStats() const override
-	{
-		TSharedRef<FDerivedDataCacheStatsNode> Usage = MakeShared<FDerivedDataCacheStatsNode>(this, FString::Printf(TEXT("%s.%s"), TEXT("FileSystem"), *CachePath));
-		Usage->Stats.Add(TEXT(""), UsageStats);
-
-		return Usage;
-	}
+TSharedRef<FDerivedDataCacheStatsNode> FFileSystemCacheStore::GatherUsageStats() const
+{
+	TSharedRef<FDerivedDataCacheStatsNode> Usage =
+		MakeShared<FDerivedDataCacheStatsNode>(this, FString::Printf(TEXT("%s.%s"), TEXT("FileSystem"), *CachePath));
+	Usage->Stats.Add(TEXT(""), UsageStats);
+	return Usage;
+}
 
-	bool TryToPrefetch(TConstArrayView<FString> CacheKeys) override
-	{
-		return CachedDataProbablyExistsBatch(CacheKeys).CountSetBits() == CacheKeys.Num();
-	}
+bool FFileSystemCacheStore::TryToPrefetch(const TConstArrayView<FString> CacheKeys)
+{
+	return CachedDataProbablyExistsBatch(CacheKeys).CountSetBits() == CacheKeys.Num();
+}
 
-	bool ApplyDebugOptions(FBackendDebugOptions& InOptions) override
-	{
-		DebugOptions = InOptions;
-		return true;
-	}
+bool FFileSystemCacheStore::ApplyDebugOptions(FBackendDebugOptions& InOptions)
+{
+	DebugOptions = InOptions;
+	return true;
+}
 
-	virtual void Put(
-		TConstArrayView<FCacheRecord> Records,
-		FStringView Context,
-		ECachePolicy Policy,
-		IRequestOwner& Owner,
-		FOnCachePutComplete&& OnComplete) override
+void FFileSystemCacheStore::Put(
+	const TConstArrayView<FCacheRecord> Records,
+	const FStringView Context,
+	const ECachePolicy Policy,
+	IRequestOwner& Owner,
+	FOnCachePutComplete&& OnComplete)
+{
+	for (const FCacheRecord& Record : Records)
 	{
-		for (const FCacheRecord& Record : Records)
+		TRACE_CPUPROFILER_EVENT_SCOPE(FileSystemDDC_Put);
+		TRACE_COUNTER_INCREMENT(FileSystemDDC_Put);
+		COOK_STAT(auto Timer = UsageStats.TimePut());
+		if (PutCacheRecord(Record, Context, Policy))
 		{
-			TRACE_CPUPROFILER_EVENT_SCOPE(FileSystemDDC_Put);
-			TRACE_COUNTER_INCREMENT(FileSystemDDC_Put);
-			COOK_STAT(auto Timer = UsageStats.TimePut());
-			if (PutCacheRecord(Record, Context, Policy))
+			UE_LOG(LogDerivedDataCache, Verbose, TEXT("%s: Cache put complete for %s from '%.*s'"),
+				*CachePath, *WriteToString<96>(Record.GetKey()), Context.Len(), Context.GetData());
+			TRACE_COUNTER_INCREMENT(FileSystemDDC_PutHit);
+			TRACE_COUNTER_ADD(FileSystemDDC_BytesWritten, MeasureCompressedCacheRecord(Record));
+			COOK_STAT(Timer.AddHit(MeasureRawCacheRecord(Record)));
+			if (OnComplete)
 			{
-				UE_LOG(LogDerivedDataCache, Verbose, TEXT("%s: Cache put complete for %s from '%.*s'"),
-					*CachePath, *WriteToString<96>(Record.GetKey()), Context.Len(), Context.GetData());
-				TRACE_COUNTER_INCREMENT(FileSystemDDC_PutHit);
-				TRACE_COUNTER_ADD(FileSystemDDC_BytesWritten, MeasureCompressedCacheRecord(Record));
-				COOK_STAT(Timer.AddHit(MeasureRawCacheRecord(Record)));
-				if (OnComplete)
-				{
-					OnComplete({Record.GetKey(), EStatus::Ok});
-				}
+				OnComplete({Record.GetKey(), EStatus::Ok});
 			}
-			else
+		}
+		else
+		{
+			COOK_STAT(Timer.AddMiss(MeasureRawCacheRecord(Record)));
+			if (OnComplete)
 			{
-				COOK_STAT(Timer.AddMiss(MeasureRawCacheRecord(Record)));
-				if (OnComplete)
-				{
-					OnComplete({Record.GetKey(), EStatus::Error});
-				}
+				OnComplete({Record.GetKey(), EStatus::Error});
 			}
 		}
 	}
+}
 
-	virtual void Get(
-		TConstArrayView<FCacheKey> Keys,
-		FStringView Context,
-		FCacheRecordPolicy Policy,
-		IRequestOwner& Owner,
-		FOnCacheGetComplete&& OnComplete) override
+void FFileSystemCacheStore::Get(
+	const TConstArrayView<FCacheKey> Keys,
+	const FStringView Context,
+	const FCacheRecordPolicy Policy,
+	IRequestOwner& Owner,
+	FOnCacheGetComplete&& OnComplete)
+{
+	for (const FCacheKey& Key : Keys)
 	{
-		for (const FCacheKey& Key : Keys)
+		TRACE_CPUPROFILER_EVENT_SCOPE(FileSystemDDC_Get);
+		TRACE_COUNTER_INCREMENT(FileSystemDDC_Get);
+		COOK_STAT(auto Timer = UsageStats.TimeGet());
+		EStatus Status = EStatus::Ok;
+		if (FOptionalCacheRecord Record = GetCacheRecord(Key, Context, Policy, Status))
 		{
-			TRACE_CPUPROFILER_EVENT_SCOPE(FileSystemDDC_Get);
-			TRACE_COUNTER_INCREMENT(FileSystemDDC_Get);
-			COOK_STAT(auto Timer = UsageStats.TimeGet());
-			EStatus Status = EStatus::Ok;
-			if (FOptionalCacheRecord Record = GetCacheRecord(Key, Context, Policy, Status))
+			UE_LOG(LogDerivedDataCache, Verbose, TEXT("%s: Cache hit for %s from '%.*s'"),
+				*CachePath, *WriteToString<96>(Key), Context.Len(), Context.GetData());
+			TRACE_COUNTER_INCREMENT(FileSystemDDC_GetHit);
+			TRACE_COUNTER_ADD(FileSystemDDC_BytesRead, MeasureCompressedCacheRecord(Record.Get()));
+			COOK_STAT(Timer.AddHit(MeasureRawCacheRecord(Record.Get())));
+			if (OnComplete)
 			{
-				UE_LOG(LogDerivedDataCache, Verbose, TEXT("%s: Cache hit for %s from '%.*s'"),
-					*CachePath, *WriteToString<96>(Key), Context.Len(), Context.GetData());
-				TRACE_COUNTER_INCREMENT(FileSystemDDC_GetHit);
-				TRACE_COUNTER_ADD(FileSystemDDC_BytesRead, MeasureCompressedCacheRecord(Record.Get()));
-				COOK_STAT(Timer.AddHit(MeasureRawCacheRecord(Record.Get())));
-				if (OnComplete)
-				{
-					OnComplete({MoveTemp(Record).Get(), Status});
-				}
+				OnComplete({MoveTemp(Record).Get(), Status});
 			}
-			else
+		}
+		else
+		{
+			if (OnComplete)
 			{
-				if (OnComplete)
-				{
-					OnComplete({FCacheRecordBuilder(Key).Build(), Status});
-				}
+				OnComplete({FCacheRecordBuilder(Key).Build(), Status});
 			}
 		}
 	}
+}
 
-	virtual void GetChunks(
-		TConstArrayView<FCacheChunkRequest> Chunks,
-		FStringView Context,
-		IRequestOwner& Owner,
-		FOnCacheGetChunkComplete&& OnComplete) override
-	{
-		TArray<FCacheChunkRequest, TInlineAllocator<16>> SortedChunks(Chunks);
-		SortedChunks.StableSort(TChunkLess());
+void FFileSystemCacheStore::GetChunks(
+	const TConstArrayView<FCacheChunkRequest> Chunks,
+	const FStringView Context,
+	IRequestOwner& Owner,
+	FOnCacheGetChunkComplete&& OnComplete)
+{
+	TArray<FCacheChunkRequest, TInlineAllocator<16>> SortedChunks(Chunks);
+	SortedChunks.StableSort(TChunkLess());
 
-		FOptionalCacheRecord Record;
-		for (const FCacheChunkRequest& Chunk : SortedChunks)
+	FOptionalCacheRecord Record;
+	for (const FCacheChunkRequest& Chunk : SortedChunks)
+	{
+		constexpr ECachePolicy SkipFlag = ECachePolicy::SkipValue;
+		const bool bExistsOnly = EnumHasAnyFlags(Chunk.Policy, SkipFlag);
+		TRACE_CPUPROFILER_EVENT_SCOPE(FileSystemDDC_Get);
+		TRACE_COUNTER_INCREMENT(FileSystemDDC_Get);
+		COOK_STAT(auto Timer = bExistsOnly ? UsageStats.TimeProbablyExists() : UsageStats.TimeGet());
+		if (!Record || Record.Get().GetKey() != Chunk.Key)
 		{
-			constexpr ECachePolicy SkipFlag = ECachePolicy::SkipValue;
-			const bool bExistsOnly = EnumHasAnyFlags(Chunk.Policy, SkipFlag);
-			TRACE_CPUPROFILER_EVENT_SCOPE(FileSystemDDC_Get);
-			TRACE_COUNTER_INCREMENT(FileSystemDDC_Get);
-			COOK_STAT(auto Timer = bExistsOnly ? UsageStats.TimeProbablyExists() : UsageStats.TimeGet());
-			if (!Record || Record.Get().GetKey() != Chunk.Key)
-			{
-				FCacheRecordPolicyBuilder PolicyBuilder(ECachePolicy::None);
-				const ECachePolicy RecordSkipFlags = bExistsOnly ? ECachePolicy::SkipData : ECachePolicy::None;
-				PolicyBuilder.AddPayloadPolicy(Chunk.Id, Chunk.Policy | RecordSkipFlags);
-				Record = GetCacheRecordOnly(Chunk.Key, Context, PolicyBuilder.Build());
-			}
-			if (Record)
+			FCacheRecordPolicyBuilder PolicyBuilder(ECachePolicy::None);
+			const ECachePolicy RecordSkipFlags = bExistsOnly ? ECachePolicy::SkipData : ECachePolicy::None;
+			PolicyBuilder.AddPayloadPolicy(Chunk.Id, Chunk.Policy | RecordSkipFlags);
+			Record = GetCacheRecordOnly(Chunk.Key, Context, PolicyBuilder.Build());
+		}
+		if (Record)
+		{
+			EStatus PayloadStatus = EStatus::Ok;
+			FPayload Payload = Record.Get().GetPayload(Chunk.Id);
+			GetCacheContent(Chunk.Key, Context, Chunk.Policy, SkipFlag, Payload, PayloadStatus);
+			if (Payload)
 			{
-				EStatus PayloadStatus = EStatus::Ok;
-				FPayload Payload = Record.Get().GetPayload(Chunk.Id);
-				GetCachePayload(Chunk.Key, Context, Chunk.Policy, SkipFlag, Payload, PayloadStatus);
-				if (Payload)
+				const uint64 RawSize = FMath::Min(Payload.GetRawSize(), Chunk.RawSize);
+				UE_LOG(LogDerivedDataCache, Verbose, TEXT("%s: Cache hit for %s from '%.*s'"),
+					*CachePath, *WriteToString<96>(Chunk.Key, '/', Chunk.Id), Context.Len(), Context.GetData());
+				TRACE_COUNTER_INCREMENT(FileSystemDDC_GetHit);
+				TRACE_COUNTER_ADD(FileSystemDDC_BytesRead, Payload.GetData().GetCompressedSize());
+				COOK_STAT(Timer.AddHit(Payload.HasData() ? RawSize : 0));
+				if (OnComplete)
 				{
-					const uint64 RawSize = FMath::Min(Payload.GetRawSize(), Chunk.RawSize);
-					UE_LOG(LogDerivedDataCache, Verbose, TEXT("%s: Cache hit for %s from '%.*s'"),
-						*CachePath, *WriteToString<96>(Chunk.Key, '/', Chunk.Id), Context.Len(), Context.GetData());
-					TRACE_COUNTER_INCREMENT(FileSystemDDC_GetHit);
-					TRACE_COUNTER_ADD(FileSystemDDC_BytesRead, Payload.GetData().GetCompressedSize());
-					COOK_STAT(Timer.AddHit(Payload.HasData() ? RawSize : 0));
-					if (OnComplete)
+					FUniqueBuffer Buffer;
+					if (Payload.HasData() && !bExistsOnly)
 					{
-						FUniqueBuffer Buffer;
-						if (Payload.HasData() && !bExistsOnly)
-						{
-							Buffer = FUniqueBuffer::Alloc(RawSize);
-							Payload.GetData().DecompressToComposite().CopyTo(Buffer, Chunk.RawOffset);
-						}
-						OnComplete({Chunk.Key, Chunk.Id, Chunk.RawOffset,
-							RawSize, Payload.GetRawHash(), Buffer.MoveToShared(), PayloadStatus});
+						Buffer = FUniqueBuffer::Alloc(RawSize);
+						Payload.GetData().DecompressToComposite().CopyTo(Buffer, Chunk.RawOffset);
 					}
-					continue;
+					OnComplete({Chunk.Key, Chunk.Id, Chunk.RawOffset,
+						RawSize, Payload.GetRawHash(), Buffer.MoveToShared(), PayloadStatus});
 				}
+				continue;
 			}
+		}
 
-			if (OnComplete)
-			{
-				OnComplete({Chunk.Key, Chunk.Id, Chunk.RawOffset, 0, {}, {}, EStatus::Error});
-			}
+		if (OnComplete)
+		{
+			OnComplete({Chunk.Key, Chunk.Id, Chunk.RawOffset, 0, {}, {}, EStatus::Error});
 		}
 	}
+}
 
-private:
-	uint64 MeasureCompressedCacheRecord(const FCacheRecord& Record) const
+uint64 FFileSystemCacheStore::MeasureCompressedCacheRecord(const FCacheRecord& Record) const
+{
+	return Record.GetMeta().GetSize() +
+		Record.GetValuePayload().GetData().GetCompressedSize() +
+		Algo::TransformAccumulate(Record.GetAttachmentPayloads(),
+			[](const FPayload& Payload) { return Payload.GetData().GetCompressedSize(); }, uint64(0));
+}
+
+uint64 FFileSystemCacheStore::MeasureRawCacheRecord(const FCacheRecord& Record) const
+{
+	return Record.GetMeta().GetSize() +
+		Record.GetValuePayload().GetData().GetRawSize() +
+		Algo::TransformAccumulate(Record.GetAttachmentPayloads(),
+			[](const FPayload& Payload) { return Payload.GetData().GetRawSize(); }, uint64(0));
+}
+
+bool FFileSystemCacheStore::PutCacheRecord(
+	const FCacheRecord& Record,
+	const FStringView Context,
+	const ECachePolicy Policy)
+{
+	if (!IsWritable())
 	{
-		return Record.GetMeta().GetSize() +
-			Record.GetValuePayload().GetData().GetCompressedSize() +
-			Algo::TransformAccumulate(Record.GetAttachmentPayloads(),
-				[](const FPayload& Payload) { return Payload.GetData().GetCompressedSize(); }, uint64(0));
+		UE_LOG(LogDerivedDataCache, VeryVerbose,
+			TEXT("%s: Skipped put of %s from '%.*s' because this cache store is not writable"),
+			*CachePath, *WriteToString<96>(Record.GetKey()), Context.Len(), Context.GetData());
+		return false;
 	}
 
-	uint64 MeasureRawCacheRecord(const FCacheRecord& Record) const
+	const FCacheKey& Key = Record.GetKey();
+
+	// Skip the request if storing to the cache is disabled.
+	if (!EnumHasAnyFlags(Policy, SpeedClass == ESpeedClass::Local ? ECachePolicy::StoreLocal : ECachePolicy::StoreRemote))
 	{
-		return Record.GetMeta().GetSize() +
-			Record.GetValuePayload().GetData().GetRawSize() +
-			Algo::TransformAccumulate(Record.GetAttachmentPayloads(),
-				[](const FPayload& Payload) { return Payload.GetData().GetRawSize(); }, uint64(0));
+		UE_LOG(LogDerivedDataCache, VeryVerbose, TEXT("%s: Skipped put of %s from '%.*s' due to cache policy"),
+			*CachePath, *WriteToString<96>(Key), Context.Len(), Context.GetData());
+		return false;
 	}
 
-	bool PutCacheRecord(const FCacheRecord& Record, FStringView Context, ECachePolicy Policy)
+	if (ShouldSimulateMiss(Key))
 	{
-		if (!IsWritable())
-		{
-			UE_LOG(LogDerivedDataCache, VeryVerbose,
-				TEXT("%s: Skipped put of %s from '%.*s' because this cache store is not writable"),
-				*CachePath, *WriteToString<96>(Record.GetKey()), Context.Len(), Context.GetData());
-			return false;
-		}
-
-		const FCacheKey& Key = Record.GetKey();
+		UE_LOG(LogDerivedDataCache, Verbose, TEXT("%s: Simulated miss for put of %s from '%.*s'"),
+			*CachePath, *WriteToString<96>(Key), Context.Len(), Context.GetData());
+		return false;
+	}
 
-		// Skip the request if storing to the cache is disabled.
-		if (!EnumHasAnyFlags(Policy, SpeedClass == ESpeedClass::Local ? ECachePolicy::StoreLocal : ECachePolicy::StoreRemote))
-		{
-			UE_LOG(LogDerivedDataCache, VeryVerbose, TEXT("%s: Skipped put of %s from '%.*s' due to cache policy"),
-				*CachePath, *WriteToString<96>(Key), Context.Len(), Context.GetData());
-			return false;
-		}
-
-		if (ShouldSimulateMiss(Key))
-		{
-			UE_LOG(LogDerivedDataCache, Verbose, TEXT("%s: Simulated miss for put of %s from '%.*s'"),
-				*CachePath, *WriteToString<96>(Key), Context.Len(), Context.GetData());
-			return false;
-		}
-
-		// Check if there is an existing record package.
-		bool bRecordExists;
-		FCbPackage ExistingPackage;
-		TStringBuilder<256> Path;
-		BuildCacheRecordPath(Key, Path);
-		if (EnumHasAllFlags(Policy, ECachePolicy::SkipValue | ECachePolicy::SkipAttachments))
-		{
-			bRecordExists = FileExists(Path);
-		}
-		else
-		{
-			FSharedBuffer Buffer = LoadFile(Path, Context);
-			FCbFieldIterator It = FCbFieldIterator::MakeRange(Buffer);
-			bRecordExists = ExistingPackage.TryLoad(It);
-		}
+	// Check if there is an existing record package.
+	bool bRecordExists;
+	FCbPackage ExistingPackage;
+	TStringBuilder<256> Path;
+	BuildCacheRecordPath(Key, Path);
+	if (EnumHasAllFlags(Policy, ECachePolicy::SkipValue | ECachePolicy::SkipAttachments))
+	{
+		bRecordExists = FileExists(Path);
+	}
+	else
+	{
+		FSharedBuffer Buffer = LoadFile(Path, Context);
+		FCbFieldIterator It = FCbFieldIterator::MakeRange(Buffer);
+		bRecordExists = ExistingPackage.TryLoad(It);
+	}
 
-		// Save the record to a package and remove attachments that will be stored externally.
-		FCbPackage Package = Record.Save();
-		TArray<FCompressedBuffer, TInlineAllocator<8>> ExternalContent;
-		if (ExistingPackage)
+	// Save the record to a package and remove attachments that will be stored externally.
+	FCbPackage Package = Record.Save();
+	TArray<FCompressedBuffer, TInlineAllocator<8>> ExternalContent;
+	if (ExistingPackage)
+	{
+		// Mirror the existing internal/external attachment storage.
+		TArray<FCompressedBuffer, TInlineAllocator<8>> AllContent;
+		Algo::Transform(Package.GetAttachments(), AllContent, &FCbAttachment::AsCompressedBinary);
+		for (FCompressedBuffer& Content : AllContent)
 		{
-			// Mirror the existing internal/external attachment storage.
-			TArray<FCompressedBuffer, TInlineAllocator<8>> AllContent;
-			Algo::Transform(Package.GetAttachments(), AllContent, &FCbAttachment::AsCompressedBinary);
-			for (FCompressedBuffer& Content : AllContent)
+			const FIoHash RawHash = Content.GetRawHash();
+			if (!ExistingPackage.FindAttachment(RawHash))
 			{
-				const FIoHash RawHash = Content.GetRawHash();
-				if (!ExistingPackage.FindAttachment(RawHash))
-				{
-					Package.RemoveAttachment(RawHash);
-					ExternalContent.Add(MoveTemp(Content));
-				}
+				Package.RemoveAttachment(RawHash);
+				ExternalContent.Add(MoveTemp(Content));
 			}
 		}
-		else
+	}
+	else
+	{
+		// Remove the largest attachments from the package until it fits within the size limits.
+		TArray<FCompressedBuffer, TInlineAllocator<8>> AllContent;
+		Algo::Transform(Package.GetAttachments(), AllContent, &FCbAttachment::AsCompressedBinary);
+		uint64 TotalSize = Algo::TransformAccumulate(AllContent, &FCompressedBuffer::GetCompressedSize, uint64(0));
+		const uint64 MaxSize = (AllContent.Num() == 1 ? MaxValueSizeKB : MaxRecordSizeKB) * 1024;
+		if (TotalSize > MaxSize)
 		{
-			// Remove the largest attachments from the package until it fits within the size limits.
-			TArray<FCompressedBuffer, TInlineAllocator<8>> AllContent;
-			Algo::Transform(Package.GetAttachments(), AllContent, &FCbAttachment::AsCompressedBinary);
-			uint64 TotalSize = Algo::TransformAccumulate(AllContent, &FCompressedBuffer::GetCompressedSize, uint64(0));
-			const uint64 MaxSize = (AllContent.Num() == 1 ? MaxValueSizeKB : MaxRecordSizeKB) * 1024;
-			if (TotalSize > MaxSize)
+			Algo::StableSortBy(AllContent, &FCompressedBuffer::GetCompressedSize, TGreater<>());
+			for (FCompressedBuffer& Content : AllContent)
 			{
-				Algo::StableSortBy(AllContent, &FCompressedBuffer::GetCompressedSize, TGreater<>());
-				for (FCompressedBuffer& Content : AllContent)
+				const uint64 CompressedSize = Content.GetCompressedSize();
+				Package.RemoveAttachment(Content.GetRawHash());
+				ExternalContent.Add(MoveTemp(Content));
+				TotalSize -= CompressedSize;
+				if (TotalSize <= MaxSize)
 				{
-					const uint64 CompressedSize = Content.GetCompressedSize();
-					Package.RemoveAttachment(Content.GetRawHash());
-					ExternalContent.Add(MoveTemp(Content));
-					TotalSize -= CompressedSize;
-					if (TotalSize <= MaxSize)
-					{
-						break;
-					}
+					break;
 				}
 			}
 		}
+	}
 
-		// Save the external content to storage.
-		for (FCompressedBuffer& Content : ExternalContent)
-		{
-			PutCacheContent(Content, Context);
-		}
+	// Save the external content to storage.
+	for (FCompressedBuffer& Content : ExternalContent)
+	{
+		PutCacheContent(Content, Context);
+	}
 
-		// Save the record package to storage.
-		if (!bRecordExists && !SaveFile(Path, Context, [&Package](FArchive& Ar) { Package.Save(Ar); }))
-		{
-			return false;
-		}
+	// Save the record package to storage.
+	if (!bRecordExists && !SaveFile(Path, Context, [&Package](FArchive& Ar) { Package.Save(Ar); }))
+	{
+		return false;
+	}
 
-		if (AccessLogWriter)
-		{
-			AccessLogWriter->Append(Key, Path);
-		}
+	if (AccessLogWriter)
+	{
+		AccessLogWriter->Append(Key, Path);
+	}
 
-		return true;
+	return true;
+}
+
+FOptionalCacheRecord FFileSystemCacheStore::GetCacheRecordOnly(
+	const FCacheKey& Key,
+	const FStringView Context,
+	const FCacheRecordPolicy& Policy)
+{
+	if (!IsUsable())
+	{
+		UE_LOG(LogDerivedDataCache, VeryVerbose,
+			TEXT("%s: Skipped get of %s from '%.*s' because this cache store is not available"),
+			*CachePath, *WriteToString<96>(Key), Context.Len(), Context.GetData());
+		return FOptionalCacheRecord();
 	}
 
-	FOptionalCacheRecord GetCacheRecordOnly(
-		const FCacheKey& Key,
-		const FStringView Context,
-		const FCacheRecordPolicy& Policy)
+	// Skip the request if querying the cache is disabled.
+	const ECachePolicy QueryPolicy = SpeedClass == ESpeedClass::Local
+		? ECachePolicy::QueryLocal : ECachePolicy::QueryRemote;
+	if (!EnumHasAnyFlags(Policy.GetRecordPolicy(), QueryPolicy))
 	{
-		if (!IsUsable())
-		{
-			UE_LOG(LogDerivedDataCache, VeryVerbose,
-				TEXT("%s: Skipped get of %s from '%.*s' because this cache store is not available"),
-				*CachePath, *WriteToString<96>(Key), Context.Len(), Context.GetData());
-			return FOptionalCacheRecord();
-		}
+		UE_LOG(LogDerivedDataCache, VeryVerbose, TEXT("%s: Skipped get of %s from '%.*s' due to cache policy"),
+			*CachePath, *WriteToString<96>(Key), Context.Len(), Context.GetData());
+		return FOptionalCacheRecord();
+	}
 
-		// Skip the request if querying the cache is disabled.
-		if (!EnumHasAnyFlags(Policy.GetRecordPolicy(), SpeedClass == ESpeedClass::Local ? ECachePolicy::QueryLocal : ECachePolicy::QueryRemote))
-		{
-			UE_LOG(LogDerivedDataCache, VeryVerbose, TEXT("%s: Skipped get of %s from '%.*s' due to cache policy"),
-				*CachePath, *WriteToString<96>(Key), Context.Len(), Context.GetData());
-			return FOptionalCacheRecord();
-		}
+	if (ShouldSimulateMiss(Key))
+	{
+		UE_LOG(LogDerivedDataCache, Verbose, TEXT("%s: Simulated miss for get of %s from '%.*s'"),
+			*CachePath, *WriteToString<96>(Key), Context.Len(), Context.GetData());
+		return FOptionalCacheRecord();
+	}
 
-		if (ShouldSimulateMiss(Key))
-		{
-			UE_LOG(LogDerivedDataCache, Verbose, TEXT("%s: Simulated miss for get of %s from '%.*s'"),
-				*CachePath, *WriteToString<96>(Key), Context.Len(), Context.GetData());
-			return FOptionalCacheRecord();
-		}
+	// Request the record from storage.
+	TStringBuilder<256> Path;
+	BuildCacheRecordPath(Key, Path);
+	FSharedBuffer Buffer = LoadFile(Path, Context);
+	if (Buffer.IsNull())
+	{
+		UE_LOG(LogDerivedDataCache, Verbose, TEXT("%s: Cache miss with missing record for %s from '%.*s'"),
+			*CachePath, *WriteToString<96>(Key), Context.Len(), Context.GetData());
+		return FOptionalCacheRecord();
+	}
 
-		// Request the record from storage.
-		TStringBuilder<256> Path;
-		BuildCacheRecordPath(Key, Path);
-		FSharedBuffer Buffer = LoadFile(Path, Context);
-		if (Buffer.IsNull())
+	// Delete the record from storage if it is invalid.
+	bool bDeleteCacheObject = true;
+	ON_SCOPE_EXIT
+	{
+		if (bDeleteCacheObject && !bReadOnly)
 		{
-			UE_LOG(LogDerivedDataCache, Verbose, TEXT("%s: Cache miss with missing record for %s from '%.*s'"),
-				*CachePath, *WriteToString<96>(Key), Context.Len(), Context.GetData());
-			return FOptionalCacheRecord();
+			IFileManager::Get().Delete(*Path, /*bRequireExists*/ false, /*bEvenReadOnly*/ false, /*bQuiet*/ true);
 		}
+	};
 
-		// Delete the record from storage if it is invalid.
-		bool bDeleteCacheObject = true;
-		ON_SCOPE_EXIT
-		{
-			if (bDeleteCacheObject && !bReadOnly)
-			{
-				IFileManager::Get().Delete(*Path, /*bRequireExists*/ false, /*bEvenReadOnly*/ false, /*bQuiet*/ true);
-			}
-		};
+	// Validate that the record can be read as a compact binary package without crashing.
+	if (ValidateCompactBinaryPackage(Buffer, ECbValidateMode::Default | ECbValidateMode::Package) != ECbValidateError::None)
+	{
+		UE_LOG(LogDerivedDataCache, Display, TEXT("%s: Cache miss with invalid package for %s from '%.*s'"),
+			*CachePath, *WriteToString<96>(Key), Context.Len(), Context.GetData());
+		return FOptionalCacheRecord();
+	}
 
-		// Validate that the record can be read as a compact binary package without crashing.
-		if (ValidateCompactBinaryPackage(Buffer, ECbValidateMode::Default | ECbValidateMode::Package) != ECbValidateError::None)
+	// Load the record from the package.
+	FOptionalCacheRecord Record;
+	{
+		FCbPackage Package;
+		if (FCbFieldIterator It = FCbFieldIterator::MakeRange(Buffer); !Package.TryLoad(It))
 		{
-			UE_LOG(LogDerivedDataCache, Display, TEXT("%s: Cache miss with invalid package for %s from '%.*s'"),
+			UE_LOG(LogDerivedDataCache, Display, TEXT("%s: Cache miss with package load failure for %s from '%.*s'"),
 				*CachePath, *WriteToString<96>(Key), Context.Len(), Context.GetData());
 			return FOptionalCacheRecord();
 		}
-
-		// Load the record from the package.
-		FOptionalCacheRecord Record;
+		Record = FCacheRecord::Load(Package);
+		if (Record.IsNull())
 		{
-			FCbPackage Package;
-			if (FCbFieldIterator It = FCbFieldIterator::MakeRange(Buffer); !Package.TryLoad(It))
-			{
-				UE_LOG(LogDerivedDataCache, Display, TEXT("%s: Cache miss with package load failure for %s from '%.*s'"),
-					*CachePath, *WriteToString<96>(Key), Context.Len(), Context.GetData());
-				return FOptionalCacheRecord();
-			}
-			Record = FCacheRecord::Load(Package);
-			if (Record.IsNull())
-			{
-				UE_LOG(LogDerivedDataCache, Display, TEXT("%s: Cache miss with record load failure for %s from '%.*s'"),
-					*CachePath, *WriteToString<96>(Key), Context.Len(), Context.GetData());
-				return FOptionalCacheRecord();
-			}
+			UE_LOG(LogDerivedDataCache, Display, TEXT("%s: Cache miss with record load failure for %s from '%.*s'"),
+				*CachePath, *WriteToString<96>(Key), Context.Len(), Context.GetData());
+			return FOptionalCacheRecord();
 		}
+	}
 
-		// Disable deletion now that the record is loaded and validated.
-		bDeleteCacheObject = false;
-
-		if (AccessLogWriter)
-		{
-			AccessLogWriter->Append(Key, Path);
-		}
+	// Disable deletion now that the record is loaded and validated.
+	bDeleteCacheObject = false;
 
-		return Record.Get();
+	if (AccessLogWriter)
+	{
+		AccessLogWriter->Append(Key, Path);
 	}
 
-	FOptionalCacheRecord GetCacheRecord(
-		const FCacheKey& Key,
-		const FStringView Context,
-		const FCacheRecordPolicy& Policy,
-		EStatus& OutStatus)
+	return Record.Get();
+}
+
+FOptionalCacheRecord FFileSystemCacheStore::GetCacheRecord(
+	const FCacheKey& Key,
+	const FStringView Context,
+	const FCacheRecordPolicy& Policy,
+	EStatus& OutStatus)
+{
+	FOptionalCacheRecord Record = GetCacheRecordOnly(Key, Context, Policy);
+	if (Record.IsNull())
 	{
-		FOptionalCacheRecord Record = GetCacheRecordOnly(Key, Context, Policy);
-		if (Record.IsNull())
-		{
-			OutStatus = EStatus::Error;
-			return Record;
-		}
+		OutStatus = EStatus::Error;
+		return Record;
+	}
 
-		OutStatus = EStatus::Ok;
+	OutStatus = EStatus::Ok;
 
-		FCacheRecordBuilder RecordBuilder(Key);
+	FCacheRecordBuilder RecordBuilder(Key);
 
-		if (!EnumHasAnyFlags(Policy.GetRecordPolicy(), ECachePolicy::SkipMeta))
-		{
-			RecordBuilder.SetMeta(FCbObject(Record.Get().GetMeta()));
-		}
+	if (!EnumHasAnyFlags(Policy.GetRecordPolicy(), ECachePolicy::SkipMeta))
+	{
+		RecordBuilder.SetMeta(FCbObject(Record.Get().GetMeta()));
+	}
 
-		if (FPayload Payload = Record.Get().GetValuePayload())
+	if (FPayload Payload = Record.Get().GetValuePayload())
+	{
+		const ECachePolicy PayloadPolicy = Policy.GetPayloadPolicy(Payload.GetId());
+		GetCacheContent(Key, Context, PayloadPolicy, ECachePolicy::SkipValue, Payload, OutStatus);
+		if (Payload.IsNull())
 		{
-			const ECachePolicy PayloadPolicy = Policy.GetPayloadPolicy(Payload.GetId());
-			GetCachePayload(Key, Context, PayloadPolicy, ECachePolicy::SkipValue, Payload, OutStatus);
-			if (Payload.IsNull())
-			{
-				return FOptionalCacheRecord();
-			}
-			RecordBuilder.SetValue(MoveTemp(Payload));
+			return FOptionalCacheRecord();
 		}
+		RecordBuilder.SetValue(MoveTemp(Payload));
+	}
 
-		for (FPayload Payload : Record.Get().GetAttachmentPayloads())
+	for (FPayload Payload : Record.Get().GetAttachmentPayloads())
+	{
+		const ECachePolicy PayloadPolicy = Policy.GetPayloadPolicy(Payload.GetId());
+		GetCacheContent(Key, Context, PayloadPolicy, ECachePolicy::SkipAttachments, Payload, OutStatus);
+		if (Payload.IsNull())
 		{
-			const ECachePolicy PayloadPolicy = Policy.GetPayloadPolicy(Payload.GetId());
-			GetCachePayload(Key, Context, PayloadPolicy, ECachePolicy::SkipAttachments, Payload, OutStatus);
-			if (Payload.IsNull())
-			{
-				return FOptionalCacheRecord();
-			}
-			RecordBuilder.AddAttachment(MoveTemp(Payload));
+			return FOptionalCacheRecord();
 		}
-
-		return RecordBuilder.Build();
+		RecordBuilder.AddAttachment(MoveTemp(Payload));
 	}
 
-	bool PutCacheContent(const FCompressedBuffer& Content, const FStringView Context) const
+	return RecordBuilder.Build();
+}
+
+bool FFileSystemCacheStore::PutCacheContent(const FCompressedBuffer& Content, const FStringView Context) const
+{
+	const FIoHash& RawHash = Content.GetRawHash();
+	TStringBuilder<256> Path;
+	BuildCacheContentPath(RawHash, Path);
+	if (!FileExists(Path))
 	{
-		const FIoHash& RawHash = Content.GetRawHash();
-		TStringBuilder<256> Path;
-		BuildCacheContentPath(RawHash, Path);
-		if (!FileExists(Path))
+		if (!SaveFile(Path, Context, [&Content](FArchive& Ar) { Ar << const_cast<FCompressedBuffer&>(Content); }))
 		{
-			if (!SaveFile(Path, Context, [&Content](FArchive& Ar) { Ar << const_cast<FCompressedBuffer&>(Content); }))
-			{
-				return false;
-			}
-		}
-		if (AccessLogWriter)
-		{
-			AccessLogWriter->Append(RawHash, Path);
+			return false;
 		}
-		return true;
 	}
+	if (AccessLogWriter)
+	{
+		AccessLogWriter->Append(RawHash, Path);
+	}
+	return true;
+}
 
-	void GetCachePayload(
-		const FCacheKey& Key,
-		const FStringView Context,
-		const ECachePolicy Policy,
-		const ECachePolicy SkipFlag,
-		FPayload& InOutPayload,
-		EStatus& InOutStatus) const
+void FFileSystemCacheStore::GetCacheContent(
+	const FCacheKey& Key,
+	const FStringView Context,
+	const ECachePolicy Policy,
+	const ECachePolicy SkipFlag,
+	FPayload& InOutPayload,
+	EStatus& InOutStatus) const
+{
+	const FIoHash& RawHash = InOutPayload.GetRawHash();
+
+	if (!EnumHasAnyFlags(Policy, ECachePolicy::Query) || (EnumHasAnyFlags(Policy, SkipFlag) && InOutPayload.HasData()))
 	{
-		const FIoHash& RawHash = InOutPayload.GetRawHash();
+		InOutPayload = FPayload(InOutPayload.GetId(), RawHash, InOutPayload.GetRawSize());
+		return;
+	}
 
-		if (!EnumHasAnyFlags(Policy, ECachePolicy::Query) || (EnumHasAnyFlags(Policy, SkipFlag) && InOutPayload.HasData()))
-		{
-			InOutPayload = FPayload(InOutPayload.GetId(), RawHash, InOutPayload.GetRawSize());
-			return;
-		}
+	if (InOutPayload.HasData())
+	{
+		return;
+	}
 
-		if (InOutPayload.HasData())
+	TStringBuilder<256> Path;
+	BuildCacheContentPath(RawHash, Path);
+	if (EnumHasAllFlags(Policy, SkipFlag))
+	{
+		if (FileExists(Path))
 		{
+			if (AccessLogWriter)
+			{
+				AccessLogWriter->Append(RawHash, Path);
+			}
 			return;
 		}
-
-		TStringBuilder<256> Path;
-		BuildCacheContentPath(RawHash, Path);
-		if (EnumHasAllFlags(Policy, SkipFlag))
+	}
+	else
+	{
+		if (FSharedBuffer CompressedData = LoadFile(Path, Context))
 		{
-			if (FileExists(Path))
+			if (FCompressedBuffer CompressedBuffer = FCompressedBuffer::FromCompressed(MoveTemp(CompressedData));
+				CompressedBuffer && CompressedBuffer.GetRawHash() == RawHash)
 			{
 				if (AccessLogWriter)
 				{
 					AccessLogWriter->Append(RawHash, Path);
 				}
+				InOutPayload = FPayload(InOutPayload.GetId(), MoveTemp(CompressedBuffer));
 				return;
 			}
-		}
-		else
-		{
-			if (FSharedBuffer CompressedData = LoadFile(Path, Context))
+			UE_LOG(LogDerivedDataCache, Display,
+				TEXT("%s: Cache miss with corrupted payload %s with hash %s for %s from '%.*s'"),
+				*CachePath, *WriteToString<16>(InOutPayload.GetId()), *WriteToString<48>(RawHash),
+				*WriteToString<96>(Key), Context.Len(), Context.GetData());
+			InOutStatus = EStatus::Error;
+			if (!EnumHasAnyFlags(Policy, ECachePolicy::PartialOnError))
 			{
-				if (FCompressedBuffer CompressedBuffer = FCompressedBuffer::FromCompressed(MoveTemp(CompressedData));
-					CompressedBuffer && CompressedBuffer.GetRawHash() == RawHash)
-				{
-					if (AccessLogWriter)
-					{
-						AccessLogWriter->Append(RawHash, Path);
-					}
-					InOutPayload = FPayload(InOutPayload.GetId(), MoveTemp(CompressedBuffer));
-					return;
-				}
-				UE_LOG(LogDerivedDataCache, Display,
-					TEXT("%s: Cache miss with corrupted payload %s with hash %s for %s from '%.*s'"),
-					*CachePath, *WriteToString<16>(InOutPayload.GetId()), *WriteToString<48>(RawHash),
-					*WriteToString<96>(Key), Context.Len(), Context.GetData());
-				InOutStatus = EStatus::Error;
-				if (!EnumHasAnyFlags(Policy, ECachePolicy::PartialOnError))
-				{
-					InOutPayload = FPayload::Null;
-				}
-				return;
+				InOutPayload = FPayload::Null;
 			}
-		}
-
-		UE_LOG(LogDerivedDataCache, Verbose,
-			TEXT("%s: Cache miss with missing payload %s with hash %s for %s from '%.*s'"),
-			*CachePath, *WriteToString<16>(InOutPayload.GetId()), *WriteToString<48>(RawHash), *WriteToString<96>(Key),
-			Context.Len(), Context.GetData());
-		InOutStatus = EStatus::Error;
-		if (!EnumHasAnyFlags(Policy, ECachePolicy::PartialOnError))
-		{
-			InOutPayload = FPayload::Null;
+			return;
 		}
 	}
 
-	void BuildCacheRecordPath(const FCacheKey& CacheKey, FStringBuilderBase& Path) const
+	UE_LOG(LogDerivedDataCache, Verbose,
+		TEXT("%s: Cache miss with missing payload %s with hash %s for %s from '%.*s'"),
+		*CachePath, *WriteToString<16>(InOutPayload.GetId()), *WriteToString<48>(RawHash), *WriteToString<96>(Key),
+		Context.Len(), Context.GetData());
+	InOutStatus = EStatus::Error;
+	if (!EnumHasAnyFlags(Policy, ECachePolicy::PartialOnError))
 	{
-		Path << CachePath << TEXT('/');
-		BuildPathForCacheRecord(CacheKey, Path);
+		InOutPayload = FPayload::Null;
 	}
+}
 
-	void BuildCacheContentPath(const FIoHash& RawHash, FStringBuilderBase& Path) const
-	{
-		Path << CachePath << TEXT('/');
-		BuildPathForCacheContent(RawHash, Path);
-	}
+void FFileSystemCacheStore::BuildCacheRecordPath(const FCacheKey& CacheKey, FStringBuilderBase& Path) const
+{
+	Path << CachePath << TEXT('/');
+	BuildPathForCacheRecord(CacheKey, Path);
+}
 
-	bool SaveFile(FStringBuilderBase& Path, FStringView Context, TFunctionRef<void (FArchive&)> WriteFunction) const
-	{
-		TStringBuilder<256> TempPath;
-		TempPath << FPathViews::GetPath(Path) << TEXT("/Temp.") << FGuid::NewGuid();
+void FFileSystemCacheStore::BuildCacheContentPath(const FIoHash& RawHash, FStringBuilderBase& Path) const
+{
+	Path << CachePath << TEXT('/');
+	BuildPathForCacheContent(RawHash, Path);
+}
 
-		ON_SCOPE_EXIT
-		{
-			IFileManager::Get().Delete(*TempPath, /*bRequireExists*/ false, /*bEvenReadOnly*/ false, /*bQuiet*/ true);
-		};
+bool FFileSystemCacheStore::SaveFile(
+	FStringBuilderBase& Path,
+	const FStringView Context,
+	TFunctionRef<void (FArchive&)> WriteFunction) const
+{
+	TStringBuilder<256> TempPath;
+	TempPath << FPathViews::GetPath(Path) << TEXT("/Temp.") << FGuid::NewGuid();
+
+	ON_SCOPE_EXIT
+	{
+		IFileManager::Get().Delete(*TempPath, /*bRequireExists*/ false, /*bEvenReadOnly*/ false, /*bQuiet*/ true);
+	};
 
-		int64 ExpectedSize = 0;
-		if (SaveDataToFile(TempPath, MoveTemp(WriteFunction), ExpectedSize))
+	int64 ExpectedSize = 0;
+	if (SaveDataToFile(TempPath, MoveTemp(WriteFunction), ExpectedSize))
+	{
+		if (IFileManager::Get().FileSize(*TempPath) == ExpectedSize)
 		{
-			if (IFileManager::Get().FileSize(*TempPath) == ExpectedSize)
+			if (IFileManager::Get().Move(*Path, *TempPath,
+				/*bReplace*/ false, /*bEvenIfReadOnly*/ false, /*bAttributes*/ false, /*bDoNotRetryOrError*/ true))
 			{
-				if (IFileManager::Get().Move(*Path, *TempPath, /*bReplace*/ false, /*bEvenIfReadOnly*/ false, /*bAttributes*/ false, /*bDoNotRetryOrError*/ true))
-				{
-					return true;
-				}
-				else
-				{
-					UE_LOG(LogDerivedDataCache, Log,
-						TEXT("%s: Move collision when writing file %s from '%.*s'"),
-						*CachePath, *Path, Context.Len(), Context.GetData());
-					return true;
-				}
+				return true;
 			}
 			else
 			{
-				UE_LOG(LogDerivedDataCache, Warning,
-					TEXT("%s: Failed to write to temp file %s when saving %s from '%.*s'. ")
-					TEXT("File is %" INT64_FMT " bytes when %" INT64_FMT " bytes are expected."),
-					*CachePath, *TempPath, *Path, Context.Len(), Context.GetData(),
-					IFileManager::Get().FileSize(*TempPath), ExpectedSize);
-				return false;
+				UE_LOG(LogDerivedDataCache, Log,
+					TEXT("%s: Move collision when writing file %s from '%.*s'"),
+					*CachePath, *Path, Context.Len(), Context.GetData());
+				return true;
 			}
 		}
 		else
 		{
 			UE_LOG(LogDerivedDataCache, Warning,
-				TEXT("%s: Failed to write to temp file %s when saving %s from '%.*s'. Error 0x%08x"),
-				*CachePath, *TempPath, *Path, Context.Len(), Context.GetData(), FPlatformMisc::GetLastError());
+				TEXT("%s: Failed to write to temp file %s when saving %s from '%.*s'. ")
+				TEXT("File is %" INT64_FMT " bytes when %" INT64_FMT " bytes are expected."),
+				*CachePath, *TempPath, *Path, Context.Len(), Context.GetData(),
+				IFileManager::Get().FileSize(*TempPath), ExpectedSize);
 			return false;
 		}
 	}
-
-	bool SaveDataToFile(
-		FStringBuilderBase& Path,
-		TFunctionRef<void (FArchive&)> WriteFunction,
-		int64& OutWriteSize) const
+	else
 	{
-		if (TUniquePtr<FArchive> Ar{IFileManager::Get().CreateFileWriter(*Path, FILEWRITE_Silent)})
-		{
-			THashingArchiveProxy<FBlake3> HashAr(*Ar);
-			WriteFunction(HashAr);
-			FBlake3Hash Hash = HashAr.GetHash();
-			*Ar << Hash;
-			OutWriteSize = Ar->Tell();
-			return Ar->Close() && !Ar->IsError();
-		}
-		OutWriteSize = 0;
+		UE_LOG(LogDerivedDataCache, Warning,
+			TEXT("%s: Failed to write to temp file %s when saving %s from '%.*s'. Error 0x%08x"),
+			*CachePath, *TempPath, *Path, Context.Len(), Context.GetData(), FPlatformMisc::GetLastError());
 		return false;
 	}
+}
 
-	FSharedBuffer LoadFile(FStringBuilderBase& Path, FStringView Context) const
-	{
-		check(IsUsable());
-		const double StartTime = FPlatformTime::Seconds();
+bool FFileSystemCacheStore::SaveDataToFile(
+	FStringBuilderBase& Path,
+	TFunctionRef<void (FArchive&)> WriteFunction,
+	int64& OutWriteSize) const
+{
+	if (TUniquePtr<FArchive> Ar{IFileManager::Get().CreateFileWriter(*Path, FILEWRITE_Silent)})
+	{
+		THashingArchiveProxy<FBlake3> HashAr(*Ar);
+		WriteFunction(HashAr);
+		FBlake3Hash Hash = HashAr.GetHash();
+		*Ar << Hash;
+		OutWriteSize = Ar->Tell();
+		return Ar->Close() && !Ar->IsError();
+	}
+	OutWriteSize = 0;
+	return false;
+}
 
-		FSharedBuffer Buffer;
+FSharedBuffer FFileSystemCacheStore::LoadFile(FStringBuilderBase& Path, const FStringView Context) const
+{
+	check(IsUsable());
+	const double StartTime = FPlatformTime::Seconds();
 
-		// Check for existence before reading because it may update the modification time and avoid the
-		// file being deleted by a cache cleanup thread or process.
-		if (!FileExists(Path))
-		{
-			return Buffer;
-		}
+	FSharedBuffer Buffer;
 
-		if (TUniquePtr<FArchive> Ar{IFileManager::Get().CreateFileReader(*Path, FILEREAD_Silent)})
+	// Check for existence before reading because it may update the modification time and avoid the
+	// file being deleted by a cache cleanup thread or process.
+	if (!FileExists(Path))
+	{
+		return Buffer;
+	}
+
+	if (TUniquePtr<FArchive> Ar{IFileManager::Get().CreateFileReader(*Path, FILEREAD_Silent)})
+	{
+		if (const int64 TotalSize = Ar->TotalSize(); TotalSize >= sizeof(FBlake3Hash))
 		{
-			if (const int64 TotalSize = Ar->TotalSize(); TotalSize >= sizeof(FBlake3Hash))
+			const int64 DataSize = TotalSize - sizeof(FBlake3Hash);
+			FUniqueBuffer MutableBuffer = FUniqueBuffer::Alloc(uint64(DataSize));
+			Ar->Serialize(MutableBuffer.GetData(), DataSize);
+			FBlake3Hash SavedHash;
+			*Ar << SavedHash;
+			if (Ar->Close())
 			{
-				const int64 DataSize = TotalSize - sizeof(FBlake3Hash);
-				FUniqueBuffer MutableBuffer = FUniqueBuffer::Alloc(uint64(DataSize));
-				Ar->Serialize(MutableBuffer.GetData(), DataSize);
-				FBlake3Hash SavedHash;
-				*Ar << SavedHash;
-				if (Ar->Close())
+				const FBlake3Hash Hash = FBlake3::HashBuffer(MutableBuffer);
+				if (Hash == SavedHash)
 				{
-					const FBlake3Hash Hash = FBlake3::HashBuffer(MutableBuffer);
-					if (Hash == SavedHash)
-					{
-						Buffer = MutableBuffer.MoveToShared();
-					}
-					else
-					{
-						UE_LOG(LogDerivedDataCache, Display,
-							TEXT("%s: File %s from '%.*s' is corrupted and has hash %s when %s is expected."),
-							*CachePath, *Path, Context.Len(), Context.GetData(),
-							*WriteToString<80>(Hash), *WriteToString<80>(SavedHash));
-						if (!bReadOnly)
-						{
-							IFileManager::Get().Delete(*Path, /*bRequireExists*/ false, /*bEvenReadOnly*/ false, /*bQuiet*/ true);
-						}
-					}
+					Buffer = MutableBuffer.MoveToShared();
 				}
-			}
-			else
-			{
-				UE_LOG(LogDerivedDataCache, Display,
-					TEXT("%s: File %s from '%.*s' is %" INT64_FMT " bytes when at least %" INT64_FMT " bytes are required."),
-					*CachePath, *Path, Context.Len(), Context.GetData(), TotalSize, int64(sizeof(FBlake3Hash)));
-				if (!bReadOnly)
+				else
 				{
-					IFileManager::Get().Delete(*Path, /*bRequireExists*/ false, /*bEvenReadOnly*/ false, /*bQuiet*/ true);
+					UE_LOG(LogDerivedDataCache, Display,
+						TEXT("%s: File %s from '%.*s' is corrupted and has hash %s when %s is expected."),
+						*CachePath, *Path, Context.Len(), Context.GetData(),
+						*WriteToString<80>(Hash), *WriteToString<80>(SavedHash));
+					if (!bReadOnly)
+					{
+						IFileManager::Get().Delete(*Path, /*bRequireExists*/ false, /*bEvenReadOnly*/ false, /*bQuiet*/ true);
+					}
 				}
 			}
 		}
-
-		const double ReadDuration = FPlatformTime::Seconds() - StartTime;
-		const double ReadSpeed = ReadDuration > 0.001 ? (Buffer.GetSize() / ReadDuration) / (1024.0 * 1024.0) : 0.0;
-
-		if (!GIsBuildMachine && ReadDuration > 5.0)
+		else
 		{
-			// Slower than 0.5 MiB/s?
-			UE_CLOG(ReadSpeed < 0.5, LogDerivedDataCache, Warning, TEXT("%s: Loading %s from '%.*s' is very slow (%.2f MiB/s); ")
-				TEXT("consider disabling this cache backend"), *CachePath, *Path, Context.Len(), Context.GetData(), ReadSpeed);
+			UE_LOG(LogDerivedDataCache, Display,
+				TEXT("%s: File %s from '%.*s' is %" INT64_FMT " bytes when at least %" INT64_FMT " bytes are required."),
+				*CachePath, *Path, Context.Len(), Context.GetData(), TotalSize, int64(sizeof(FBlake3Hash)));
+			if (!bReadOnly)
+			{
+				IFileManager::Get().Delete(*Path, /*bRequireExists*/ false, /*bEvenReadOnly*/ false, /*bQuiet*/ true);
+			}
 		}
+	}
 
-		if (Buffer)
-		{
-			UE_LOG(LogDerivedDataCache, VeryVerbose, TEXT("%s: Loaded %s from '%.*s' (%" UINT64_FMT " bytes, %.02f secs, %.2f MiB/s)"),
-				*CachePath, *Path, Context.Len(), Context.GetData(), Buffer.GetSize(), ReadDuration, ReadSpeed);
-		}
+	const double ReadDuration = FPlatformTime::Seconds() - StartTime;
+	const double ReadSpeed = ReadDuration > 0.001 ? (Buffer.GetSize() / ReadDuration) / (1024.0 * 1024.0) : 0.0;
 
-		return Buffer;
+	if (!GIsBuildMachine && ReadDuration > 5.0)
+	{
+		// Slower than 0.5 MiB/s?
+		UE_CLOG(ReadSpeed < 0.5, LogDerivedDataCache, Warning,
+			TEXT("%s: Loading %s from '%.*s' is very slow (%.2f MiB/s); consider disabling this cache backend"),
+			*CachePath, *Path, Context.Len(), Context.GetData(), ReadSpeed);
 	}
 
-	bool FileExists(FStringBuilderBase& Path) const
+	if (Buffer)
 	{
-		const FDateTime TimeStamp = IFileManager::Get().GetTimeStamp(*Path);
-		if (TimeStamp == FDateTime::MinValue())
-		{
-			return false;
-		}
-		if (bTouch || (!bReadOnly && (FDateTime::UtcNow() - TimeStamp).GetTotalDays() > (DaysToDeleteUnusedFiles / 4)))
-		{
-			IFileManager::Get().SetTimeStamp(*Path, FDateTime::UtcNow());
-		}
-		return true;
+		UE_LOG(LogDerivedDataCache, VeryVerbose,
+			TEXT("%s: Loaded %s from '%.*s' (%" UINT64_FMT " bytes, %.02f secs, %.2f MiB/s)"),
+			*CachePath, *Path, Context.Len(), Context.GetData(), Buffer.GetSize(), ReadDuration, ReadSpeed);
 	}
 
-private:
-	FDerivedDataCacheUsageStats UsageStats;
+	return Buffer;
+}
 
-	/**
-	 * Threadsafe method to compute the filename from the cachekey, currently just adds a path and an extension.
-	 *
-	 * @param	CacheKey	Alphanumeric+underscore key of this cache item
-	 * @return				filename built from the cache key
-	 */
-	FString BuildFilename(const TCHAR* CacheKey) const
+bool FFileSystemCacheStore::FileExists(FStringBuilderBase& Path) const
+{
+	const FDateTime TimeStamp = IFileManager::Get().GetTimeStamp(*Path);
+	if (TimeStamp == FDateTime::MinValue())
 	{
-		return CachePath / BuildPathForCacheKey(CacheKey);
+		return false;
 	}
-
-	/** Base path we are storing the cache files in. **/
-	FString	CachePath;
-	/** Class of this cache */
-	ESpeedClass SpeedClass;
-	/** If true, we failed to write to this directory and it did not contain anything so we should not be used. */
-	bool		bDisabled;
-	/** If true, do not attempt to write to this cache. */
-	bool		bReadOnly;
-	/** If true, CachedDataProbablyExists will update the file timestamps. */
-	bool		bTouch;
-	/** If true, allow transient data to be removed from the cache. */
-	bool		bPurgeTransient;
-	/** Age of file when it should be deleted from DDC cache. */
-	double		DaysToDeleteUnusedFiles;
-
-	/** Maximum total size of compressed data stored within a record package with multiple attachments. */
-	uint64		MaxRecordSizeKB = 256;
-	/** Maximum total size of compressed data stored within a value package, or a record package with one attachment. */
-	uint64		MaxValueSizeKB = 1024;
-
-	/** Object used for synchronization via a scoped lock						*/
-	FCriticalSection SynchronizationObject;
-
-	// DDCNotification metrics
-
-	/** Map of cache keys to miss times for generating timing deltas */
-	TMap<FString, double> DDCNotificationCacheTimes;
-
-	/** The total estimated build time accumulated from cache miss/put deltas */
-	double TotalEstimatedBuildTime;
-
-	/** Access log to write to */
-	TUniquePtr<FAccessLogWriter> AccessLogWriter;
-
-	/** Debug Options */
-	FBackendDebugOptions DebugOptions;
-
-	/** Keys we ignored due to miss rate settings */
-	FCriticalSection MissedKeysCS;
-	TSet<FName> DebugMissedKeys;
-	TSet<FCacheKey> DebugMissedCacheKeys;
-
-	TUniquePtr<FFileSystemCacheStoreMaintainer> Maintainer;
-
-	bool ShouldSimulateMiss(const TCHAR* InKey)
+	if (bTouch || (!bReadOnly && (FDateTime::UtcNow() - TimeStamp).GetTotalDays() > (DaysToDeleteUnusedFiles / 4)))
 	{
-		if (DebugOptions.RandomMissRate == 0 && DebugOptions.SimulateMissTypes.IsEmpty())
-		{
-			return false;
-		}
-
-		const FName Key(InKey);
-		const uint32 Hash = GetTypeHash(Key);
-
-		if (FScopeLock Lock(&MissedKeysCS); DebugMissedKeys.ContainsByHash(Hash, Key))
-		{
-			return true;
-		}
-
-		if (DebugOptions.ShouldSimulateMiss(InKey))
-		{
-			FScopeLock Lock(&MissedKeysCS);
-			UE_LOG(LogDerivedDataCache, Verbose, TEXT("Simulating miss in %s for %s"), *GetName(), InKey);
-			DebugMissedKeys.AddByHash(Hash, Key);
-			return true;
-		}
+		IFileManager::Get().SetTimeStamp(*Path, FDateTime::UtcNow());
+	}
+	return true;
+}
 
+bool FFileSystemCacheStore::ShouldSimulateMiss(const TCHAR* const InKey)
+{
+	if (DebugOptions.RandomMissRate == 0 && DebugOptions.SimulateMissTypes.IsEmpty())
+	{
 		return false;
 	}
 
-	bool ShouldSimulateMiss(const FCacheKey& Key)
-	{
-		if (DebugOptions.RandomMissRate == 0 && DebugOptions.SimulateMissTypes.IsEmpty())
-		{
-			return false;
-		}
+	const FName Key(InKey);
+	const uint32 Hash = GetTypeHash(Key);
 
-		const uint32 Hash = GetTypeHash(Key);
+	if (FScopeLock Lock(&MissedKeysCS); DebugMissedKeys.ContainsByHash(Hash, Key))
+	{
+		return true;
+	}
 
-		if (FScopeLock Lock(&MissedKeysCS); DebugMissedCacheKeys.ContainsByHash(Hash, Key))
-		{
-			return true;
-		}
+	if (DebugOptions.ShouldSimulateMiss(InKey))
+	{
+		FScopeLock Lock(&MissedKeysCS);
+		UE_LOG(LogDerivedDataCache, Verbose, TEXT("Simulating miss in %s for %s"), *GetName(), InKey);
+		DebugMissedKeys.AddByHash(Hash, Key);
+		return true;
+	}
 
-		if (DebugOptions.ShouldSimulateMiss(Key))
-		{
-			FScopeLock Lock(&MissedKeysCS);
-			DebugMissedCacheKeys.AddByHash(Hash, Key);
-			return true;
-		}
+	return false;
+}
 
+bool FFileSystemCacheStore::ShouldSimulateMiss(const FCacheKey& Key)
+{
+	if (DebugOptions.RandomMissRate == 0 && DebugOptions.SimulateMissTypes.IsEmpty())
+	{
 		return false;
 	}
 
-};
+	const uint32 Hash = GetTypeHash(Key);
 
-FDerivedDataBackendInterface* CreateFileSystemDerivedDataBackend(const TCHAR* CacheDirectory, const TCHAR* InParams, const TCHAR* InAccessLogFileName /*= nullptr*/)
-{
-	FFileSystemDerivedDataBackend* FileDDB = new FFileSystemDerivedDataBackend(CacheDirectory, InParams, InAccessLogFileName);
+	if (FScopeLock Lock(&MissedKeysCS); DebugMissedCacheKeys.ContainsByHash(Hash, Key))
+	{
+		return true;
+	}
 
-	if (!FileDDB->IsUsable())
+	if (DebugOptions.ShouldSimulateMiss(Key))
 	{
-		delete FileDDB;
-		FileDDB = NULL;
+		FScopeLock Lock(&MissedKeysCS);
+		DebugMissedCacheKeys.AddByHash(Hash, Key);
+		return true;
 	}
 
-	return FileDDB;
+	return false;
+}
+
+FDerivedDataBackendInterface* CreateFileSystemDerivedDataBackend(
+	const TCHAR* CachePath,
+	const TCHAR* Params,
+	const TCHAR* AccessLogPath)
+{
+	TUniquePtr<FFileSystemCacheStore> Store(new FFileSystemCacheStore(CachePath, Params, AccessLogPath));
+	return Store->IsUsable() ? Store.Release() : nullptr;
 }
 
 } // UE::DerivedData::CacheStore::FileSystem
