DDC: Fixed file system cache maintenance with -NoThreading

#rnx

diff --git a/Engine/Source/Developer/DerivedDataCache/Private/FileSystemDerivedDataBackend.cpp b/Engine/Source/Developer/DerivedDataCache/Private/FileSystemDerivedDataBackend.cpp
--- a/Engine/Source/Developer/DerivedDataCache/Private/FileSystemDerivedDataBackend.cpp
+++ b/Engine/Source/Developer/DerivedDataCache/Private/FileSystemDerivedDataBackend.cpp
@@ -200,7 +200,7 @@ struct FFileSystemCacheStoreMaintainerParams
 	uint32 MaxFileScanRate = MAX_uint32;
 	/** Limits the number of directories scanned in each cache bucket or content root. */
 	uint32 MaxDirectoryScanCount = MAX_uint32;
-	/** Minimum duration between the start of consecutive scans. */
+	/** Minimum duration between the start of consecutive scans. Use MaxValue to scan only once. */
 	FTimespan ScanFrequency = FTimespan::FromHours(1.0);
 	/** Time to wait after initialization before maintenance begins. */
 	FTimespan TimeToWaitAfterInit = FTimespan::FromMinutes(1.0);
@@ -217,6 +217,7 @@ public:
 	void BoostPriority() final;
 
 private:
+	void Tick();
 	void Loop();
 	void Scan();
 
@@ -240,10 +241,12 @@ private:
 	FFileSystemCacheStoreMaintainerParams Params;
 	/** Path to the root of the cache store. */
 	FString CachePath;
-	/** True when maintenance is expected to exit. */
-	bool bExit = false;
 	/** True when there is no active maintenance scan. */
 	bool bIdle = false;
+	/** True when maintenance is expected to exit as soon as possible. */
+	bool bExit = false;
+	/** True when maintenance is expected to exit at the end of the scan. */
+	bool bExitAfterScan = false;
 	/** Ignore the file scan rate for one maintenance scan. */
 	bool bIgnoreFileScanRate = false;
 
@@ -262,6 +265,8 @@ private:
 	TArray<TUniquePtr<FRoot>> Roots;
 	TUniquePtr<FLegacyRoot> LegacyRoot;
 	FRandomStream Random{uint32(RandFromGuid())};
+
+	static constexpr double MaxScanFrequencyDays = 365.0;
 };
 
 struct FFileSystemCacheStoreMaintainer::FRoot
@@ -302,11 +307,13 @@ FFileSystemCacheStoreMaintainer::FFileSystemCacheStoreMaintainer(
 	const FStringView InCachePath)
 	: Params(InParams)
 	, CachePath(InCachePath)
+	, bExitAfterScan(Params.ScanFrequency.GetTotalDays() > MaxScanFrequencyDays)
 	, IdleEvent(EEventMode::ManualReset)
 	, WaitEvent(EEventMode::AutoReset)
 	, Thread(
 		TEXT("FileSystemCacheStoreMaintainer"),
 		[this] { Loop(); },
+		[this] { Tick(); },
 		/*StackSize*/ 32 * 1024,
 		TPri_BelowNormal)
 {
@@ -327,6 +334,18 @@ void FFileSystemCacheStoreMaintainer::BoostPriority()
 	WaitEvent->Trigger();
 }
 
+void FFileSystemCacheStoreMaintainer::Tick()
+{
+	// Scan once and exit if the priority has been boosted.
+	if (bIgnoreFileScanRate)
+	{
+		bExitAfterScan = true;
+		Loop();
+	}
+	bIdle = true;
+	IdleEvent.Trigger();
+}
+
 void FFileSystemCacheStoreMaintainer::Loop()
 {
 	WaitEvent->Wait(Params.TimeToWaitAfterInit, /*bIgnoreThreadIdleStats*/ true);
@@ -348,19 +367,22 @@ void FFileSystemCacheStoreMaintainer::Loop()
 			TEXT("%s: Maintenance finished in %s and deleted %u file(s) with total size %" UINT64_FMT " MiB."),
 			*CachePath, *(ScanEnd - ScanStart).ToString(), DeleteCount, DeleteSize / 1024 / 1024);
 
-		if (bExit || Params.ScanFrequency.GetTotalDays() > 365.0)
+		if (bExit || bExitAfterScan)
 		{
 			break;
 		}
 
 		const FDateTime ScanTime = ScanStart + Params.ScanFrequency;
-		UE_LOG(LogDerivedDataCache, Verbose, TEXT("%s: Maintenance is paused until the next scan at %s."),
-			*CachePath, *ScanTime.ToString());
+		UE_CLOG(ScanEnd < ScanTime, LogDerivedDataCache, Verbose,
+			TEXT("%s: Maintenance is paused until the next scan at %s."), *CachePath, *ScanTime.ToString());
 		for (FDateTime Now = ScanEnd; !bExit && Now < ScanTime; Now = FDateTime::Now())
 		{
 			WaitEvent->Wait(ScanTime - Now, /*bIgnoreThreadIdleStats*/ true);
 		}
 	}
+
+	bIdle = true;
+	IdleEvent.Trigger();
 }
 
 void FFileSystemCacheStoreMaintainer::Scan()
@@ -985,7 +1007,7 @@ FFileSystemCacheStore::FFileSystemCacheStore(
 				}
 				else
 				{
-					MaintainerParams.ScanFrequency = FTimespan::FromDays(500.0);
+					MaintainerParams.ScanFrequency = FTimespan::MaxValue();
 				}
 				double TimeToWaitAfterInit;
 				if (bClean)
diff --git a/Engine/Source/Editor/UnrealEd/Private/Commandlets/DDCCleanupCommandlet.cpp b/Engine/Source/Editor/UnrealEd/Private/Commandlets/DDCCleanupCommandlet.cpp
--- a/Engine/Source/Editor/UnrealEd/Private/Commandlets/DDCCleanupCommandlet.cpp
+++ b/Engine/Source/Editor/UnrealEd/Private/Commandlets/DDCCleanupCommandlet.cpp
@@ -4,6 +4,7 @@
 #include "DerivedDataCache.h"
 #include "DerivedDataCacheMaintainer.h"
 #include "HAL/PlatformProcess.h"
+#include "HAL/ThreadManager.h"
 
 int32 UDDCCleanupCommandlet::Main(const FString& Params)
 {
@@ -12,7 +13,7 @@ int32 UDDCCleanupCommandlet::Main(const FString& Params)
 	Maintainer.BoostPriority();
 	while (!Maintainer.IsIdle())
 	{
-		// Maintenance works from dedicated threads. Wait for it to finish and flush logs occasionally.
+		FThreadManager::Get().Tick();
 		FPlatformProcess::SleepNoStats(0.05f);
 		GLog->Flush();
 	}
