DDC: Fixed the size calculation for partial payload reads in the file system cache store

#rnx

diff --git a/Engine/Source/Developer/DerivedDataCache/Private/FileSystemDerivedDataBackend.cpp b/Engine/Source/Developer/DerivedDataCache/Private/FileSystemDerivedDataBackend.cpp
--- a/Engine/Source/Developer/DerivedDataCache/Private/FileSystemDerivedDataBackend.cpp
+++ b/Engine/Source/Developer/DerivedDataCache/Private/FileSystemDerivedDataBackend.cpp
@@ -1659,7 +1659,8 @@ void FFileSystemCacheStore::GetChunks(
 			GetCacheContent(Chunk.Key, Context, Chunk.Policy, SkipFlag, Payload, PayloadStatus);
 			if (Payload)
 			{
-				const uint64 RawSize = FMath::Min(Payload.GetRawSize(), Chunk.RawSize);
+				const uint64 RawOffset = FMath::Min(Payload.GetRawSize(), Chunk.RawOffset);
+				const uint64 RawSize = FMath::Min(Payload.GetRawSize() - RawOffset, Chunk.RawSize);
 				UE_LOG(LogDerivedDataCache, Verbose, TEXT("%s: Cache hit for %s from '%.*s'"),
 					*CachePath, *WriteToString<96>(Chunk.Key, '/', Chunk.Id), Context.Len(), Context.GetData());
 				TRACE_COUNTER_INCREMENT(FileSystemDDC_GetHit);
@@ -1667,14 +1668,13 @@ void FFileSystemCacheStore::GetChunks(
 				COOK_STAT(Timer.AddHit(Payload.HasData() ? RawSize : 0));
 				if (OnComplete)
 				{
-					FUniqueBuffer Buffer;
+					FSharedBuffer Buffer;
 					if (Payload.HasData() && !bExistsOnly)
 					{
-						Buffer = FUniqueBuffer::Alloc(RawSize);
-						Payload.GetData().DecompressToComposite().CopyTo(Buffer, Chunk.RawOffset);
+						Buffer = Payload.GetData().Decompress(RawOffset, RawSize);
 					}
 					OnComplete({Chunk.Key, Chunk.Id, Chunk.RawOffset,
-						RawSize, Payload.GetRawHash(), Buffer.MoveToShared(), PayloadStatus});
+						RawSize, Payload.GetRawHash(), MoveTemp(Buffer), PayloadStatus});
 				}
 				continue;
 			}
@@ -1734,7 +1734,7 @@ bool FFileSystemCacheStore::PutCacheRecord(
 	}
 
 	// Check if there is an existing record package.
-	bool bRecordExists;
+	bool bRecordExists = false;
 	FCbPackage ExistingPackage;
 	TStringBuilder<256> Path;
 	BuildCacheRecordPath(Key, Path);
@@ -1742,9 +1742,8 @@ bool FFileSystemCacheStore::PutCacheRecord(
 	{
 		bRecordExists = FileExists(Path);
 	}
-	else
+	else if (FSharedBuffer Buffer = LoadFile(Path, Context))
 	{
-		FSharedBuffer Buffer = LoadFile(Path, Context);
 		FCbFieldIterator It = FCbFieldIterator::MakeRange(Buffer);
 		bRecordExists = ExistingPackage.TryLoad(It);
 	}
