DDC: Improved the cache fill behavior of the hierarchical backend for the structured cache

#rnx

diff --git a/Engine/Source/Developer/DerivedDataCache/Private/HierarchicalDerivedDataBackend.h b/Engine/Source/Developer/DerivedDataCache/Private/HierarchicalDerivedDataBackend.h
--- a/Engine/Source/Developer/DerivedDataCache/Private/HierarchicalDerivedDataBackend.h
+++ b/Engine/Source/Developer/DerivedDataCache/Private/HierarchicalDerivedDataBackend.h
@@ -551,42 +551,35 @@ public:
 
 		TArray<FCacheKey, TInlineAllocator<16>> RemainingKeys(Keys);
 
-		TSet<FCacheKey> KeysOk;
-
 		{
 			FReadScopeLock LockScope(Lock);
+			TSet<FCacheKey> KeysOk;
+
+			bool bHadLocalMiss = false;
+			bool bHadRemoteMiss = false;
 
 			for (int32 GetCacheIndex = 0; GetCacheIndex < InnerBackends.Num() && !RemainingKeys.IsEmpty(); ++GetCacheIndex)
 			{
-				// Remove SkipData flags when possibly filling other backends because only complete records can be written.
 				const bool bIsLocalGet = InnerBackends[GetCacheIndex]->GetSpeedClass() >= ESpeedClass::Fast;
-				const bool bFill =
-					(bIsLocalGet && bStoreLocalCopy) ||
-					(bIsLocalGet && bStoreRemote) ||
-					(!bIsLocalGet && bStoreLocal) ||
-					(!bIsLocalGet && bStoreRemote && bHasMultipleRemoteBackends);
 
 				// Block on this because backends in this hierarchy are not expected to be asynchronous.
 				FRequestOwner BlockingOwner(EPriority::Blocking);
 				InnerBackends[GetCacheIndex]->Get(RemainingKeys, Context, Policy, BlockingOwner,
-					[this, Context, GetCacheIndex, bStoreLocal, bStoreRemote, bStoreLocalCopy, bIsLocalGet, bFill, &Owner, &OnComplete, &KeysOk](FCacheGetCompleteParams&& Params)
+					[this, Context, GetCacheIndex, bStoreLocal, bStoreRemote, bStoreLocalCopy, bIsLocalGet, &Owner, &OnComplete, &KeysOk](FCacheGetCompleteParams&& Params)
 					{
 						if (Params.Status == EStatus::Ok)
 						{
-							if (bFill)
+							FRequestOwner AsyncOwner(FPlatformMath::Min(Owner.GetPriority(), EPriority::Highest));
+							FRequestBarrier AsyncBarrier(AsyncOwner);
+							AsyncOwner.KeepAlive();
+							for (int32 FillCacheIndex = 0; FillCacheIndex < InnerBackends.Num(); ++FillCacheIndex)
 							{
-								FRequestOwner AsyncOwner(FPlatformMath::Min(Owner.GetPriority(), EPriority::Highest));
-								FRequestBarrier AsyncBarrier(AsyncOwner);
-								AsyncOwner.KeepAlive();
-								for (int32 FillCacheIndex = 0; FillCacheIndex < InnerBackends.Num(); ++FillCacheIndex)
+								if (GetCacheIndex != FillCacheIndex)
 								{
-									if (GetCacheIndex != FillCacheIndex)
+									const bool bIsLocalFill = InnerBackends[FillCacheIndex]->GetSpeedClass() >= ESpeedClass::Fast;
+									if (bIsLocalFill ? bStoreLocal : bStoreRemote)
 									{
-										const bool bIsLocalFill = InnerBackends[FillCacheIndex]->GetSpeedClass() >= ESpeedClass::Fast;
-										if ((bIsLocalFill && bStoreLocal && bIsLocalGet <= bStoreLocalCopy) || (!bIsLocalFill && bStoreRemote))
-										{
-											AsyncPutInnerBackends[FillCacheIndex]->Put({Params.Record}, Context, ECachePolicy::Default, AsyncOwner);
-										}
+										AsyncPutInnerBackends[FillCacheIndex]->Put({Params.Record}, Context, ECachePolicy::Default, AsyncOwner);
 									}
 								}
 							}
@@ -601,6 +594,26 @@ public:
 				BlockingOwner.Wait();
 
 				RemainingKeys.RemoveAll([&KeysOk](const FCacheKey& Key) { return KeysOk.Contains(Key); });
+
+				bool& bHadMiss = bIsLocalGet ? bHadLocalMiss : bHadRemoteMiss;
+				if (!bHadMiss && !RemainingKeys.IsEmpty() && InnerBackends[GetCacheIndex]->IsWritable())
+				{
+					bHadMiss = true;
+					const auto ConvertPolicy = [bIsLocalGet](ECachePolicy P) -> ECachePolicy
+					{
+						if (EnumHasAnyFlags(P, bIsLocalGet ? ECachePolicy::StoreLocal : ECachePolicy::StoreRemote))
+						{
+							EnumRemoveFlags(P, ECachePolicy::SkipData);
+						}
+						return P;
+					};
+					FCacheRecordPolicyBuilder Builder(ConvertPolicy(Policy.GetDefaultPayloadPolicy()));
+					for (const FCachePayloadPolicy& PayloadPolicy : Policy.GetPayloadPolicies())
+					{
+						Builder.AddPayloadPolicy(PayloadPolicy.Id, ConvertPolicy(PayloadPolicy.Policy));
+					}
+					Policy = Builder.Build();
+				}
 			}
 		}
 
