DDC: Added structured cache support to the pak file cache store

#rnx

diff --git a/Engine/Source/Developer/DerivedDataCache/Private/DerivedDataBackends.cpp b/Engine/Source/Developer/DerivedDataCache/Private/DerivedDataBackends.cpp
--- a/Engine/Source/Developer/DerivedDataCache/Private/DerivedDataBackends.cpp
+++ b/Engine/Source/Developer/DerivedDataCache/Private/DerivedDataBackends.cpp
@@ -295,6 +295,8 @@ public:
 	 */
 	FDerivedDataBackendInterface* ParsePak( const TCHAR* NodeName, const TCHAR* Entry, const bool bWriting )
 	{
+		using namespace CacheStore::PakFile;
+
 		FDerivedDataBackendInterface* PakNode = NULL;
 		FString PakFilename;
 		FParse::Value( Entry, TEXT("Filename="), PakFilename );
@@ -1074,7 +1076,7 @@ public:
 
 				for(const FString& MergePakName : MergePakList)
 				{
-					FPakFileDerivedDataBackend ReadPak(*FPaths::Combine(*FPaths::GetPath(WritePakFilename), *MergePakName), false);
+					CacheStore::PakFile::FPakFileDerivedDataBackend ReadPak(*FPaths::Combine(*FPaths::GetPath(WritePakFilename), *MergePakName), false);
 					WritePakCache->MergeCache(&ReadPak);
 				}
 			}
@@ -1099,7 +1101,7 @@ public:
 							UE_LOG(LogDerivedDataCache, Error, TEXT("Could not delete the pak file %s to overwrite it with a new one."), *ReadPakFilename);
 						}
 					}
-					if (!FPakFileDerivedDataBackend::SortAndCopy(WritePakFilename, ReadPakFilename))
+					if (!CacheStore::PakFile::FPakFileDerivedDataBackend::SortAndCopy(WritePakFilename, ReadPakFilename))
 					{
 						UE_LOG(LogDerivedDataCache, Error, TEXT("Couldn't sort pak file (%s)"), *WritePakFilename);
 					}
@@ -1163,10 +1165,10 @@ public:
 	{
 		// Assumptions: there's at least one read-only pak backend in the hierarchy
 		// and its parent is a hierarchical backend.
-		FPakFileDerivedDataBackend* ReadPak = NULL;
+		CacheStore::PakFile::FPakFileDerivedDataBackend* ReadPak = NULL;
 		if (HierarchicalWrapper && FPlatformFileManager::Get().GetPlatformFile().FileExists(PakFilename))
 		{
-			ReadPak = new FPakFileDerivedDataBackend(PakFilename, false);
+			ReadPak = new CacheStore::PakFile::FPakFileDerivedDataBackend(PakFilename, false);
 
 			HierarchicalWrapper->AddInnerBackend(ReadPak);
 			CreatedBackends.Add(ReadPak);
@@ -1184,7 +1186,7 @@ public:
 	{
 		for (int PakIndex = 0; PakIndex < ReadPakCache.Num(); ++PakIndex)
 		{
-			FPakFileDerivedDataBackend* ReadPak = ReadPakCache[PakIndex];
+			CacheStore::PakFile::FPakFileDerivedDataBackend* ReadPak = ReadPakCache[PakIndex];
 			if (ReadPak->GetFilename() == PakFilename)
 			{
 				check(HierarchicalWrapper);
@@ -1261,12 +1263,12 @@ private:
 
 	/** Instances of backend interfaces which exist in only one copy */
 	FFileBackedDerivedDataBackend*	BootCache;
-	FPakFileDerivedDataBackend*		WritePakCache;
+	CacheStore::PakFile::FPakFileDerivedDataBackend* WritePakCache;
 	FDerivedDataBackendInterface*	AsyncPutWrapper;
 	FDerivedDataBackendInterface*	KeyLengthWrapper;
 	FHierarchicalDerivedDataBackend* HierarchicalWrapper;
 	/** Support for multiple read only pak files. */
-	TArray<FPakFileDerivedDataBackend*>		ReadPakCache;
+	TArray<CacheStore::PakFile::FPakFileDerivedDataBackend*> ReadPakCache;
 
 	/** List of directories used by the DDC */
 	TArray<FString> Directories;
diff --git a/Engine/Source/Developer/DerivedDataCache/Private/HashingArchiveProxy.h b/Engine/Source/Developer/DerivedDataCache/Private/HashingArchiveProxy.h
--- a/Engine/Source/Developer/DerivedDataCache/Private/HashingArchiveProxy.h
+++ b/Engine/Source/Developer/DerivedDataCache/Private/HashingArchiveProxy.h
@@ -17,7 +17,7 @@ public:
 	void Serialize(void* V, int64 Length) override
 	{
 		FArchiveProxy::Serialize(V, Length);
-		HashBuilder.Update(V, Length);
+		HashBuilder.Update(V, uint64(Length));
 	}
 
 	void Seek(int64 InPos) override
@@ -26,7 +26,7 @@ public:
 		FArchiveProxy::Seek(InPos);
 	}
 
-	auto GetHash() const -> decltype(DeclVal<HashBuilderType>().Finalize())
+	auto GetHash() -> decltype(DeclVal<HashBuilderType>().Finalize())
 	{
 		return HashBuilder.Finalize();
 	}
diff --git a/Engine/Source/Developer/DerivedDataCache/Private/PakFileDerivedDataBackend.cpp b/Engine/Source/Developer/DerivedDataCache/Private/PakFileDerivedDataBackend.cpp
--- a/Engine/Source/Developer/DerivedDataCache/Private/PakFileDerivedDataBackend.cpp
+++ b/Engine/Source/Developer/DerivedDataCache/Private/PakFileDerivedDataBackend.cpp
@@ -1,40 +1,61 @@
 // Copyright Epic Games, Inc. All Rights Reserved.
 
 #include "PakFileDerivedDataBackend.h"
-#include "Misc/Compression.h"
+
+#include "Algo/Accumulate.h"
+#include "Algo/StableSort.h"
+#include "Algo/Transform.h"
 #include "DerivedDataCacheRecord.h"
 #include "DerivedDataCacheUsageStats.h"
+#include "DerivedDataChunk.h"
 #include "DerivedDataPayload.h"
+#include "HAL/FileManager.h"
+#include "HAL/PlatformFile.h"
+#include "HAL/PlatformFileManager.h"
+#include "HashingArchiveProxy.h"
+#include "Misc/Compression.h"
+#include "Misc/FileHelper.h"
+#include "Misc/Paths.h"
+#include "Misc/PathViews.h"
+#include "Misc/ScopeRWLock.h"
+#include "Serialization/CompactBinary.h"
+#include "Serialization/CompactBinaryPackage.h"
+#include "Serialization/CompactBinaryValidation.h"
+#include "Serialization/MemoryReader.h"
+#include "Serialization/MemoryWriter.h"
+#include "Templates/Greater.h"
 
-namespace UE::DerivedData::Backends
+namespace UE::DerivedData::CacheStore::PakFile
 {
 
-FPakFileDerivedDataBackend::FPakFileDerivedDataBackend(const TCHAR* InFilename, bool bInWriting)
+FPakFileDerivedDataBackend::FPakFileDerivedDataBackend(const TCHAR* const InCachePath, const bool bInWriting)
 	: bWriting(bInWriting)
 	, bClosed(false)
-	, Filename(InFilename)
+	, CachePath(InCachePath)
 {
+	IPlatformFile& PlatformFile = FPlatformFileManager::Get().GetPlatformFile();
 	if (bWriting)
 	{
-		FileHandle.Reset(IFileManager::Get().CreateFileWriter(InFilename, FILEWRITE_NoReplaceExisting));
+		PlatformFile.CreateDirectoryTree(*FPaths::GetPath(CachePath));
+		FileHandle.Reset(PlatformFile.OpenWrite(*CachePath, /*bAppend*/ false, /*bAllowRead*/ true));
 		if (!FileHandle)
 		{
-			UE_LOG(LogDerivedDataCache, Fatal, TEXT("Pak cache could not be opened for writing %s."), InFilename);
+			UE_LOG(LogDerivedDataCache, Fatal, TEXT("%s: Failed to open pak cache for writing."), *CachePath);
 			bClosed = true;
 		}
 		else
 		{
-			UE_LOG(LogDerivedDataCache, Display, TEXT("Pak cache opened for writing %s."), InFilename);
+			UE_LOG(LogDerivedDataCache, Display, TEXT("%s: Opened pak cache for writing."), *CachePath);
 		}
 	}
 	else
 	{
-		FileHandle.Reset(IFileManager::Get().CreateFileReader(InFilename));
+		FileHandle.Reset(PlatformFile.OpenRead(*CachePath));
 		if (!FileHandle)
 		{
-			UE_LOG(LogDerivedDataCache, Warning, TEXT("Pak cache could not be opened for reading %s."), InFilename);
+			UE_LOG(LogDerivedDataCache, Warning, TEXT("%s: Failed to open pak cache for reading."), *CachePath);
 		}
-		else if (!LoadCache(InFilename))
+		else if (!LoadCache(*CachePath))
 		{
 			FileHandle.Reset();
 			CacheItems.Empty();
@@ -42,8 +63,8 @@ FPakFileDerivedDataBackend::FPakFileDerivedDataBackend(const TCHAR* InFilename,
 		}
 		else
 		{
-			UE_LOG(LogDerivedDataCache, Display, TEXT("Pak cache opened for reading %s (%" INT64_FMT " MiB)."),
-				InFilename, FileHandle->TotalSize() / 1024 / 1024);
+			UE_LOG(LogDerivedDataCache, Display, TEXT("%s: Opened pak cache for reading. (%" INT64_FMT " MiB)"),
+				*CachePath, FileHandle->Size() / 1024 / 1024);
 		}
 	}
 }
@@ -60,19 +81,19 @@ FString FPakFileDerivedDataBackend::GetDisplayName() const
 
 FString FPakFileDerivedDataBackend::GetName() const
 {
-	return Filename;
+	return CachePath;
 }
 
 void FPakFileDerivedDataBackend::Close()
 {
 	FDerivedDataBackend::Get().WaitForQuiescence();
-	FScopeLock ScopeLock(&SynchronizationObject);
 	if (!bClosed)
 	{
 		if (bWriting)
 		{
 			SaveCache();
 		}
+		FWriteScopeLock ScopeLock(SynchronizationObject);
 		FileHandle.Reset();
 		CacheItems.Empty();
 		bClosed = true;
@@ -98,7 +119,7 @@ bool FPakFileDerivedDataBackend::BackfillLowerCacheLevels() const
 bool FPakFileDerivedDataBackend::CachedDataProbablyExists(const TCHAR* CacheKey)
 {
 	COOK_STAT(auto Timer = UsageStats.TimeProbablyExists());
-	FScopeLock ScopeLock(&SynchronizationObject);
+	FReadScopeLock ScopeLock(SynchronizationObject);
 	bool Result = CacheItems.Contains(FString(CacheKey));
 	if (Result)
 	{
@@ -110,35 +131,45 @@ bool FPakFileDerivedDataBackend::CachedDataProbablyExists(const TCHAR* CacheKey)
 bool FPakFileDerivedDataBackend::GetCachedData(const TCHAR* CacheKey, TArray<uint8>& OutData)
 {
 	COOK_STAT(auto Timer = UsageStats.TimeGet());
-	if (bWriting || bClosed)
+	if (bClosed)
 	{
 		return false;
 	}
-	FScopeLock ScopeLock(&SynchronizationObject);
-	FCacheValue* Item = CacheItems.Find(FString(CacheKey));
-	if (Item)
+	FWriteScopeLock ScopeLock(SynchronizationObject);
+	if (FCacheValue* Item = CacheItems.Find(FString(CacheKey)))
 	{
 		check(FileHandle);
-		FileHandle->Seek(Item->Offset);
-		if (FileHandle->Tell() != Item->Offset)
+		ON_SCOPE_EXIT
+		{
+			if (bWriting)
+			{
+				FileHandle->SeekFromEnd();
+			}
+		};
+		if (Item->Size >= MAX_int32)
+		{
+			UE_LOG(LogDerivedDataCache, Warning, TEXT("%s: Pak file, %s exceeds 2 GiB limit."), *CachePath, CacheKey);
+		}
+		else if (!FileHandle->Seek(Item->Offset))
 		{
-			UE_LOG(LogDerivedDataCache, Warning, TEXT("Pak file, bad seek."));
+			UE_LOG(LogDerivedDataCache, Warning, TEXT("%s: Pak file, bad seek."), *CachePath);
 		}
 		else
 		{
 			check(Item->Size);
 			check(!OutData.Num());
-			check(FileHandle->IsLoading());
 			OutData.AddUninitialized(Item->Size);
-			FileHandle->Serialize(OutData.GetData(), int64(Item->Size));
-			uint32 TestCrc = FCrc::MemCrc_DEPRECATED(OutData.GetData(), Item->Size);
-			if (TestCrc != Item->Crc)
+			if (!FileHandle->Read(OutData.GetData(), int64(Item->Size)))
 			{
-				UE_LOG(LogDerivedDataCache, Warning, TEXT("Pak file, bad crc."));
+				UE_LOG(LogDerivedDataCache, Warning, TEXT("%s: Pak file, bad read."), *CachePath);
+			}
+			else if (uint32 TestCrc = FCrc::MemCrc_DEPRECATED(OutData.GetData(), Item->Size); TestCrc != Item->Crc)
+			{
+				UE_LOG(LogDerivedDataCache, Warning, TEXT("%s: Pak file, bad crc."), *CachePath);
 			}
 			else
 			{
-				UE_LOG(LogDerivedDataCache, Verbose, TEXT("FPakFileDerivedDataBackend: Cache hit on %s"), CacheKey);
+				UE_LOG(LogDerivedDataCache, Verbose, TEXT("%s: Cache hit on %s"), *CachePath, CacheKey);
 				check(OutData.Num());
 				COOK_STAT(Timer.AddHit(OutData.Num()));
 				return true;
@@ -147,7 +178,7 @@ bool FPakFileDerivedDataBackend::GetCachedData(const TCHAR* CacheKey, TArray<uin
 	}
 	else
 	{
-		UE_LOG(LogDerivedDataCache, Verbose, TEXT("FPakFileDerivedDataBackend: Miss on %s"), CacheKey);
+		UE_LOG(LogDerivedDataCache, Verbose, TEXT("%s: Cache miss on %s"), *CachePath, CacheKey);
 	}
 	OutData.Empty();
 	return false;
@@ -156,18 +187,17 @@ bool FPakFileDerivedDataBackend::GetCachedData(const TCHAR* CacheKey, TArray<uin
 FDerivedDataBackendInterface::EPutStatus FPakFileDerivedDataBackend::PutCachedData(const TCHAR* CacheKey, TArrayView<const uint8> InData, bool bPutEvenIfExists)
 {
 	COOK_STAT(auto Timer = UsageStats.TimePut());
-	if (!bWriting || bClosed)
+	if (!IsWritable())
 	{
 		return EPutStatus::NotCached;
 	}
 	{
-		FScopeLock ScopeLock(&SynchronizationObject);
+		FWriteScopeLock ScopeLock(SynchronizationObject);
 		FString Key(CacheKey);
 		TOptional<uint32> Crc;
 		check(InData.Num());
 		check(Key.Len());
 		check(FileHandle);
-		check(FileHandle->IsSaving());
 
 		if (bPutEvenIfExists)
 		{
@@ -185,14 +215,17 @@ FDerivedDataBackendInterface::EPutStatus FPakFileDerivedDataBackend::PutCachedDa
 					{
 						int64 Offset = FileHandle->Tell();
 						FileHandle->Seek(Item->Offset);
-						FileHandle->Serialize(const_cast<uint8*>(InData.GetData()), int64(InData.Num()));
+						FileHandle->Write(InData.GetData(), InData.Num());
 						Item->Crc = Crc.GetValue();
 						FileHandle->Seek(Offset);
 					}
 					return EPutStatus::Cached;
 				}
 
-				UE_LOG(LogDerivedDataCache, Warning, TEXT("FPakFileDerivedDataBackend: Repeated put of %s with different sized contents.  Multiple contents will be in the file, but only the last will be in the index.  This has wasted %" INT64_FMT " bytes in the file."), CacheKey, Item->Offset);
+				UE_LOG(LogDerivedDataCache, Warning,
+					TEXT("%s: Repeated put of %s with different sized contents. Multiple contents will be in the file, ")
+					TEXT("but only the last will be in the index. This has wasted %" INT64_FMT " bytes in the file."),
+					*CachePath, CacheKey, Item->Size);
 				CacheItems.Remove(Key);
 			}
 		}
@@ -202,7 +235,7 @@ FDerivedDataBackendInterface::EPutStatus FPakFileDerivedDataBackend::PutCachedDa
 		{
 			CacheItems.Empty();
 			FileHandle.Reset();
-			UE_LOG(LogDerivedDataCache, Fatal, TEXT("Could not write pak file...out of disk space?"));
+			UE_LOG(LogDerivedDataCache, Fatal, TEXT("%s: Could not write pak file... out of disk space?"), *CachePath);
 			return EPutStatus::NotCached;
 		}
 		else
@@ -212,10 +245,9 @@ FDerivedDataBackendInterface::EPutStatus FPakFileDerivedDataBackend::PutCachedDa
 			{
 				Crc = FCrc::MemCrc_DEPRECATED(InData.GetData(), InData.Num());
 			}
-			// NOTE: Gross that FArchive doesn't have a const version just for saving...
-			FileHandle->Serialize(const_cast<uint8*>(InData.GetData()), int64(InData.Num()));
-			UE_LOG(LogDerivedDataCache, Verbose, TEXT("FPakFileDerivedDataBackend: Put %s"), CacheKey);
-			CacheItems.Add(Key,FCacheValue(Offset, InData.Num(), Crc.GetValue()));
+			FileHandle->Write(InData.GetData(), InData.Num());
+			UE_LOG(LogDerivedDataCache, Verbose, TEXT("%s: Put %s"), *CachePath, CacheKey);
+			CacheItems.Add(Key, FCacheValue(Offset, InData.Num(), Crc.GetValue()));
 			return EPutStatus::Cached;
 		}
 	}
@@ -229,20 +261,19 @@ void FPakFileDerivedDataBackend::RemoveCachedData(const TCHAR* CacheKey, bool bT
 	}
 	// strangish. We can delete from a pak, but it only deletes the index 
 	// if this is a read cache, it will read it next time
-	// if this is a write cache, we wated space
-	FScopeLock ScopeLock(&SynchronizationObject);
+	// if this is a write cache, we wasted space
+	FWriteScopeLock ScopeLock(SynchronizationObject);
 	FString Key(CacheKey);
 	CacheItems.Remove(Key);
 }
 
 bool FPakFileDerivedDataBackend::SaveCache()
 {
-	FScopeLock ScopeLock(&SynchronizationObject);
+	FWriteScopeLock ScopeLock(SynchronizationObject);
 	check(FileHandle);
-	check(FileHandle->IsSaving());
 	int64 IndexOffset = FileHandle->Tell();
 	check(IndexOffset >= 0);
-	uint32 NumItems =  uint32(CacheItems.Num());
+	uint32 NumItems = uint32(CacheItems.Num());
 	check(IndexOffset > 0 || !NumItems);
 	TArray<uint8> IndexBuffer;
 	{
@@ -250,13 +281,14 @@ bool FPakFileDerivedDataBackend::SaveCache()
 		uint32 NumProcessed = 0;
 		for (TMap<FString, FCacheValue>::TIterator It(CacheItems); It; ++It )
 		{
-			check(It.Value().Offset >= 0 && It.Value().Offset < IndexOffset);
-			check(It.Value().Size);
+			FCacheValue& Value = It.Value();
 			check(It.Key().Len());
+			check(Value.Size);
+			check(Value.Offset >= 0 && Value.Offset < IndexOffset);
 			Saver << It.Key();
-			Saver << It.Value().Offset;
-			Saver << It.Value().Size;
-			Saver << It.Value().Crc;
+			Saver << Value.Offset;
+			Saver << Value.Size;
+			Saver << Value.Crc;
 			NumProcessed++;
 		}
 		check(NumProcessed == NumItems);
@@ -274,7 +306,7 @@ bool FPakFileDerivedDataBackend::SaveCache()
 	Saver.Serialize(IndexBuffer.GetData(), IndexBuffer.Num());
 	Saver << Magic;
 	Saver << IndexOffset;
-	FileHandle->Serialize(Buffer.GetData(), Buffer.Num());
+	FileHandle->Write(Buffer.GetData(), Buffer.Num());
 	CacheItems.Empty();
 	FileHandle.Reset();
 	bClosed = true;
@@ -284,12 +316,11 @@ bool FPakFileDerivedDataBackend::SaveCache()
 bool FPakFileDerivedDataBackend::LoadCache(const TCHAR* InFilename)
 {
 	check(FileHandle);
-	check(FileHandle->IsLoading());
-	int64 FileSize = FileHandle->TotalSize();
+	int64 FileSize = FileHandle->Size();
 	check(FileSize >= 0);
 	if (FileSize < sizeof(int64) + sizeof(uint32) * 5)
 	{
-		UE_LOG(LogDerivedDataCache, Error, TEXT("Pak cache was corrupted (short) %s."), InFilename);
+		UE_LOG(LogDerivedDataCache, Error, TEXT("%s: Pak cache was corrupted (short)."), InFilename);
 		return false;
 	}
 	int64 IndexOffset = -1;
@@ -301,19 +332,19 @@ bool FPakFileDerivedDataBackend::LoadCache(const TCHAR* InFilename)
 		Trailer = FileHandle->Tell();
 		if (Trailer != SeekPos)
 		{
-			UE_LOG(LogDerivedDataCache, Error, TEXT("Pak cache was corrupted (bad seek) %s."), InFilename);
+			UE_LOG(LogDerivedDataCache, Error, TEXT("%s: Pak cache was corrupted (bad seek)."), InFilename);
 			return false;
 		}
 		check(Trailer >= 0 && Trailer < FileSize);
 		Buffer.AddUninitialized(sizeof(int64) + sizeof(uint32));
-		FileHandle->Serialize(Buffer.GetData(), int64(sizeof(int64)+sizeof(uint32)));
+		FileHandle->Read(Buffer.GetData(), int64(sizeof(int64)+sizeof(uint32)));
 		FMemoryReader Loader(Buffer);
 		uint32 Magic = 0;
 		Loader << Magic;
 		Loader << IndexOffset;
 		if (Magic != PakCache_Magic || IndexOffset < 0 || IndexOffset + int64(sizeof(uint32) * 4) > Trailer)
 		{
-			UE_LOG(LogDerivedDataCache, Error, TEXT("Pak cache was corrupted (bad footer) %s."), InFilename);
+			UE_LOG(LogDerivedDataCache, Error, TEXT("%s: Pak cache was corrupted (bad footer)."), InFilename);
 			return false;
 		}
 	}
@@ -325,11 +356,11 @@ bool FPakFileDerivedDataBackend::LoadCache(const TCHAR* InFilename)
 		FileHandle->Seek(IndexOffset);
 		if (FileHandle->Tell() != IndexOffset)
 		{
-			UE_LOG(LogDerivedDataCache, Error, TEXT("Pak cache was corrupted (bad seek index) %s."), InFilename);
+			UE_LOG(LogDerivedDataCache, Error, TEXT("%s: Pak cache was corrupted (bad seek index)."), InFilename);
 			return false;
 		}
 		Buffer.AddUninitialized(sizeof(uint32) * 4);
-		FileHandle->Serialize(Buffer.GetData(), int64(sizeof(uint32)* 4));
+		FileHandle->Read(Buffer.GetData(), sizeof(uint32) * 4);
 		FMemoryReader Loader(Buffer);
 		uint32 Magic = 0;
 		Loader << Magic;
@@ -338,21 +369,21 @@ bool FPakFileDerivedDataBackend::LoadCache(const TCHAR* InFilename)
 		Loader << SizeIndex;
 		if (Magic != PakCache_Magic || (SizeIndex != 0 && NumIndex == 0) || (SizeIndex == 0 && NumIndex != 0)) 
 		{
-			UE_LOG(LogDerivedDataCache, Error, TEXT("Pak cache was corrupted (bad index header) %s."), InFilename);
+			UE_LOG(LogDerivedDataCache, Error, TEXT("%s: Pak cache was corrupted (bad index header)."), InFilename);
 			return false;
 		}
 		if (IndexOffset + sizeof(uint32) * 4 + SizeIndex != Trailer) 
 		{
-			UE_LOG(LogDerivedDataCache, Error, TEXT("Pak cache was corrupted (bad index size) %s."), InFilename);
+			UE_LOG(LogDerivedDataCache, Error, TEXT("%s: Pak cache was corrupted (bad index size)."), InFilename);
 			return false;
 		}
 	}
 	{
 		TArray<uint8> Buffer;
 		Buffer.AddUninitialized(SizeIndex);
-		FileHandle->Serialize(Buffer.GetData(), int64(SizeIndex));
+		FileHandle->Read(Buffer.GetData(), SizeIndex);
 		FMemoryReader Loader(Buffer);
-		while (Loader.Tell() < (int32)SizeIndex)
+		while (Loader.Tell() < SizeIndex)
 		{
 			FString Key;
 			int64 Offset;
@@ -364,14 +395,14 @@ bool FPakFileDerivedDataBackend::LoadCache(const TCHAR* InFilename)
 			Loader << Crc;
 			if (!Key.Len() || Offset < 0 || Offset >= IndexOffset || !Size)
 			{
-				UE_LOG(LogDerivedDataCache, Error, TEXT("Pak cache was corrupted (bad index entry) %s."), InFilename);
+				UE_LOG(LogDerivedDataCache, Error, TEXT("%s: Pak cache was corrupted (bad index entry)."), InFilename);
 				return false;
 			}
 			CacheItems.Add(Key, FCacheValue(Offset, Size, Crc));
 		}
 		if (CacheItems.Num() != NumIndex)
 		{
-			UE_LOG(LogDerivedDataCache, Error, TEXT("Pak cache was corrupted (bad index count) %s."), InFilename);
+			UE_LOG(LogDerivedDataCache, Error, TEXT("%s: Pak cache was corrupted (bad index count)."), InFilename);
 			return false;
 		}
 	}
@@ -445,59 +476,597 @@ bool FPakFileDerivedDataBackend::SortAndCopy(const FString &InputFilename, const
 
 TSharedRef<FDerivedDataCacheStatsNode> FPakFileDerivedDataBackend::GatherUsageStats() const
 {
-	TSharedRef<FDerivedDataCacheStatsNode> Usage = MakeShared<FDerivedDataCacheStatsNode>(this, FString::Printf(TEXT("%s.%s"), TEXT("PakFile"), *Filename));
+	TSharedRef<FDerivedDataCacheStatsNode> Usage = MakeShared<FDerivedDataCacheStatsNode>(this, FString::Printf(TEXT("%s.%s"), TEXT("PakFile"), *CachePath));
 	Usage->Stats.Add(TEXT(""), UsageStats);
-
 	return Usage;
 }
 
 void FPakFileDerivedDataBackend::Put(
-	TConstArrayView<FCacheRecord> Records,
-	FStringView Context,
-	ECachePolicy Policy,
+	const TConstArrayView<FCacheRecord> Records,
+	const FStringView Context,
+	const ECachePolicy Policy,
 	IRequestOwner& Owner,
 	FOnCachePutComplete&& OnComplete)
 {
-	if (OnComplete)
+	for (const FCacheRecord& Record : Records)
 	{
-		for (const FCacheRecord& Record : Records)
+		TRACE_CPUPROFILER_EVENT_SCOPE(PakFileDDC_Put);
+		COOK_STAT(auto Timer = UsageStats.TimePut());
+		if (PutCacheRecord(Record, Context, Policy))
 		{
-			OnComplete({Record.GetKey(), EStatus::Error});
+			UE_LOG(LogDerivedDataCache, Verbose, TEXT("%s: Cache put complete for %s from '%.*s'"),
+				*CachePath, *WriteToString<96>(Record.GetKey()), Context.Len(), Context.GetData());
+			COOK_STAT(Timer.AddHit(MeasureRawCacheRecord(Record)));
+			if (OnComplete)
+			{
+				OnComplete({Record.GetKey(), EStatus::Ok});
+			}
+		}
+		else
+		{
+			COOK_STAT(Timer.AddMiss(MeasureRawCacheRecord(Record)));
+			if (OnComplete)
+			{
+				OnComplete({Record.GetKey(), EStatus::Error});
+			}
 		}
 	}
 }
 
 void FPakFileDerivedDataBackend::Get(
-	TConstArrayView<FCacheKey> Keys,
-	FStringView Context,
-	FCacheRecordPolicy Policy,
+	const TConstArrayView<FCacheKey> Keys,
+	const FStringView Context,
+	const FCacheRecordPolicy Policy,
 	IRequestOwner& Owner,
 	FOnCacheGetComplete&& OnComplete)
 {
-	if (OnComplete)
+	for (const FCacheKey& Key : Keys)
 	{
-		for (const FCacheKey& Key : Keys)
+		TRACE_CPUPROFILER_EVENT_SCOPE(PakFileDDC_Get);
+		COOK_STAT(auto Timer = UsageStats.TimeGet());
+		EStatus Status = EStatus::Ok;
+		if (FOptionalCacheRecord Record = GetCacheRecord(Key, Context, Policy, Status))
+		{
+			UE_LOG(LogDerivedDataCache, Verbose, TEXT("%s: Cache hit for %s from '%.*s'"),
+				*CachePath, *WriteToString<96>(Key), Context.Len(), Context.GetData());
+			COOK_STAT(Timer.AddHit(MeasureRawCacheRecord(Record.Get())));
+			if (OnComplete)
+			{
+				OnComplete({MoveTemp(Record).Get(), Status});
+			}
+		}
+		else
 		{
-			OnComplete({FCacheRecordBuilder(Key).Build(), EStatus::Error});
+			if (OnComplete)
+			{
+				OnComplete({FCacheRecordBuilder(Key).Build(), Status});
+			}
 		}
 	}
 }
 
 void FPakFileDerivedDataBackend::GetChunks(
-	TConstArrayView<FCacheChunkRequest> Chunks,
-	FStringView Context,
+	const TConstArrayView<FCacheChunkRequest> Chunks,
+	const FStringView Context,
 	IRequestOwner& Owner,
 	FOnCacheGetChunkComplete&& OnComplete)
 {
-	if (OnComplete)
+	TArray<FCacheChunkRequest, TInlineAllocator<16>> SortedChunks(Chunks);
+	SortedChunks.StableSort(TChunkLess());
+
+	FOptionalCacheRecord Record;
+	for (const FCacheChunkRequest& Chunk : SortedChunks)
 	{
-		for (const FCacheChunkRequest& Chunk : Chunks)
+		constexpr ECachePolicy SkipFlag = ECachePolicy::SkipValue;
+		const bool bExistsOnly = EnumHasAnyFlags(Chunk.Policy, SkipFlag);
+		TRACE_CPUPROFILER_EVENT_SCOPE(PakFileDDC_Get);
+		COOK_STAT(auto Timer = bExistsOnly ? UsageStats.TimeProbablyExists() : UsageStats.TimeGet());
+		if (!Record || Record.Get().GetKey() != Chunk.Key)
+		{
+			FCacheRecordPolicyBuilder PolicyBuilder(ECachePolicy::None);
+			const ECachePolicy RecordSkipFlags = bExistsOnly ? ECachePolicy::SkipData : ECachePolicy::None;
+			PolicyBuilder.AddPayloadPolicy(Chunk.Id, Chunk.Policy | RecordSkipFlags);
+			Record = GetCacheRecordOnly(Chunk.Key, Context, PolicyBuilder.Build());
+		}
+		if (Record)
+		{
+			EStatus PayloadStatus = EStatus::Ok;
+			FPayload Payload = Record.Get().GetPayload(Chunk.Id);
+			GetCacheContent(Chunk.Key, Context, Chunk.Policy, SkipFlag, Payload, PayloadStatus);
+			if (Payload)
+			{
+				const uint64 RawOffset = FMath::Min(Payload.GetRawSize(), Chunk.RawOffset);
+				const uint64 RawSize = FMath::Min(Payload.GetRawSize() - RawOffset, Chunk.RawSize);
+				UE_LOG(LogDerivedDataCache, Verbose, TEXT("%s: Cache hit for %s from '%.*s'"),
+					*CachePath, *WriteToString<96>(Chunk.Key, '/', Chunk.Id), Context.Len(), Context.GetData());
+				COOK_STAT(Timer.AddHit(Payload.HasData() ? RawSize : 0));
+				if (OnComplete)
+				{
+					FSharedBuffer Buffer;
+					if (Payload.HasData() && !bExistsOnly)
+					{
+						Buffer = Payload.GetData().Decompress(RawOffset, RawSize);
+					}
+					OnComplete({Chunk.Key, Chunk.Id, Chunk.RawOffset,
+						RawSize, Payload.GetRawHash(), MoveTemp(Buffer), PayloadStatus});
+				}
+				continue;
+			}
+		}
+
+		if (OnComplete)
 		{
 			OnComplete({Chunk.Key, Chunk.Id, Chunk.RawOffset, 0, {}, {}, EStatus::Error});
 		}
 	}
 }
 
+uint64 FPakFileDerivedDataBackend::MeasureCompressedCacheRecord(const FCacheRecord& Record) const
+{
+	return Record.GetMeta().GetSize() +
+		Record.GetValuePayload().GetData().GetCompressedSize() +
+		Algo::TransformAccumulate(Record.GetAttachmentPayloads(),
+			[](const FPayload& Payload) { return Payload.GetData().GetCompressedSize(); }, uint64(0));
+}
+
+uint64 FPakFileDerivedDataBackend::MeasureRawCacheRecord(const FCacheRecord& Record) const
+{
+	return Record.GetMeta().GetSize() +
+		Record.GetValuePayload().GetData().GetRawSize() +
+		Algo::TransformAccumulate(Record.GetAttachmentPayloads(),
+			[](const FPayload& Payload) { return Payload.GetData().GetRawSize(); }, uint64(0));
+}
+
+bool FPakFileDerivedDataBackend::PutCacheRecord(
+	const FCacheRecord& Record,
+	const FStringView Context,
+	const ECachePolicy Policy)
+{
+	if (!IsWritable())
+	{
+		return false;
+	}
+
+	const FCacheKey& Key = Record.GetKey();
+
+	// Skip the request if storing to the cache is disabled.
+	if (!EnumHasAnyFlags(Policy, ECachePolicy::StoreLocal))
+	{
+		UE_LOG(LogDerivedDataCache, VeryVerbose, TEXT("%s: Skipped put of %s from '%.*s' due to cache policy"),
+			*CachePath, *WriteToString<96>(Key), Context.Len(), Context.GetData());
+		return false;
+	}
+
+	//if (ShouldSimulateMiss(Key))
+	//{
+	//	UE_LOG(LogDerivedDataCache, Verbose, TEXT("%s: Simulated miss for put of %s from '%.*s'"),
+	//		*CachePath, *WriteToString<96>(Key), Context.Len(), Context.GetData());
+	//	return false;
+	//}
+
+	// Check if there is an existing record package.
+	bool bRecordExists = false;
+	FCbPackage ExistingPackage;
+	TStringBuilder<256> Path;
+	FPathViews::Append(Path, TEXT("Buckets"), Key);
+	if (EnumHasAllFlags(Policy, ECachePolicy::SkipValue | ECachePolicy::SkipAttachments))
+	{
+		bRecordExists = FileExists(Path);
+	}
+	else if (FSharedBuffer Buffer = LoadFile(Path, Context))
+	{
+		FCbFieldIterator It = FCbFieldIterator::MakeRange(MoveTemp(Buffer));
+		bRecordExists = ExistingPackage.TryLoad(It);
+	}
+
+	// Save the record to a package and remove attachments that will be stored externally.
+	FCbPackage Package = Record.Save();
+	TArray<FCompressedBuffer, TInlineAllocator<8>> ExternalContent;
+	if (ExistingPackage)
+	{
+		// Mirror the existing internal/external attachment storage.
+		TArray<FCompressedBuffer, TInlineAllocator<8>> AllContent;
+		Algo::Transform(Package.GetAttachments(), AllContent, &FCbAttachment::AsCompressedBinary);
+		for (FCompressedBuffer& Content : AllContent)
+		{
+			const FIoHash RawHash = Content.GetRawHash();
+			if (!ExistingPackage.FindAttachment(RawHash))
+			{
+				Package.RemoveAttachment(RawHash);
+				ExternalContent.Add(MoveTemp(Content));
+			}
+		}
+	}
+	else
+	{
+		// Remove the largest attachments from the package until it fits within the size limits.
+		TArray<FCompressedBuffer, TInlineAllocator<8>> AllContent;
+		Algo::Transform(Package.GetAttachments(), AllContent, &FCbAttachment::AsCompressedBinary);
+		uint64 TotalSize = Algo::TransformAccumulate(AllContent, &FCompressedBuffer::GetCompressedSize, uint64(0));
+		const uint64 MaxSize = (AllContent.Num() == 1 ? MaxValueSizeKB : MaxRecordSizeKB) * 1024;
+		if (TotalSize > MaxSize)
+		{
+			Algo::StableSortBy(AllContent, &FCompressedBuffer::GetCompressedSize, TGreater<>());
+			for (FCompressedBuffer& Content : AllContent)
+			{
+				const uint64 CompressedSize = Content.GetCompressedSize();
+				Package.RemoveAttachment(Content.GetRawHash());
+				ExternalContent.Add(MoveTemp(Content));
+				TotalSize -= CompressedSize;
+				if (TotalSize <= MaxSize)
+				{
+					break;
+				}
+			}
+		}
+	}
+
+	// Save the external content to storage.
+	for (FCompressedBuffer& Content : ExternalContent)
+	{
+		PutCacheContent(Content, Context);
+	}
+
+	// Save the record package to storage.
+	if (!bRecordExists && !SaveFile(Path, Context, [&Package](FArchive& Ar) { Package.Save(Ar); }))
+	{
+		return false;
+	}
+
+	return true;
+}
+
+FOptionalCacheRecord FPakFileDerivedDataBackend::GetCacheRecordOnly(
+	const FCacheKey& Key,
+	const FStringView Context,
+	const FCacheRecordPolicy& Policy)
+{
+	if (bClosed)
+	{
+		UE_LOG(LogDerivedDataCache, VeryVerbose,
+			TEXT("%s: Skipped get of %s from '%.*s' because this cache store is not available"),
+			*CachePath, *WriteToString<96>(Key), Context.Len(), Context.GetData());
+		return FOptionalCacheRecord();
+	}
+
+	// Skip the request if querying the cache is disabled.
+	if (!EnumHasAnyFlags(Policy.GetRecordPolicy(), ECachePolicy::QueryLocal))
+	{
+		UE_LOG(LogDerivedDataCache, VeryVerbose, TEXT("%s: Skipped get of %s from '%.*s' due to cache policy"),
+			*CachePath, *WriteToString<96>(Key), Context.Len(), Context.GetData());
+		return FOptionalCacheRecord();
+	}
+
+	//if (ShouldSimulateMiss(Key))
+	//{
+	//	UE_LOG(LogDerivedDataCache, Verbose, TEXT("%s: Simulated miss for get of %s from '%.*s'"),
+	//		*CachePath, *WriteToString<96>(Key), Context.Len(), Context.GetData());
+	//	return FOptionalCacheRecord();
+	//}
+
+	// Request the record from storage.
+	TStringBuilder<256> Path;
+	FPathViews::Append(Path, TEXT("Buckets"), Key);
+	FSharedBuffer Buffer = LoadFile(Path, Context);
+	if (Buffer.IsNull())
+	{
+		UE_LOG(LogDerivedDataCache, Verbose, TEXT("%s: Cache miss with missing record for %s from '%.*s'"),
+			*CachePath, *WriteToString<96>(Key), Context.Len(), Context.GetData());
+		return FOptionalCacheRecord();
+	}
+
+	// Validate that the record can be read as a compact binary package without crashing.
+	if (ValidateCompactBinaryPackage(Buffer, ECbValidateMode::Default | ECbValidateMode::Package) != ECbValidateError::None)
+	{
+		UE_LOG(LogDerivedDataCache, Display, TEXT("%s: Cache miss with invalid package for %s from '%.*s'"),
+			*CachePath, *WriteToString<96>(Key), Context.Len(), Context.GetData());
+		return FOptionalCacheRecord();
+	}
+
+	// Load the record from the package.
+	FOptionalCacheRecord Record;
+	{
+		FCbPackage Package;
+		if (FCbFieldIterator It = FCbFieldIterator::MakeRange(Buffer); !Package.TryLoad(It))
+		{
+			UE_LOG(LogDerivedDataCache, Display, TEXT("%s: Cache miss with package load failure for %s from '%.*s'"),
+				*CachePath, *WriteToString<96>(Key), Context.Len(), Context.GetData());
+			return FOptionalCacheRecord();
+		}
+		Record = FCacheRecord::Load(Package);
+		if (Record.IsNull())
+		{
+			UE_LOG(LogDerivedDataCache, Display, TEXT("%s: Cache miss with record load failure for %s from '%.*s'"),
+				*CachePath, *WriteToString<96>(Key), Context.Len(), Context.GetData());
+			return FOptionalCacheRecord();
+		}
+	}
+
+	return Record.Get();
+}
+
+FOptionalCacheRecord FPakFileDerivedDataBackend::GetCacheRecord(
+	const FCacheKey& Key,
+	const FStringView Context,
+	const FCacheRecordPolicy& Policy,
+	EStatus& OutStatus)
+{
+	FOptionalCacheRecord Record = GetCacheRecordOnly(Key, Context, Policy);
+	if (Record.IsNull())
+	{
+		OutStatus = EStatus::Error;
+		return Record;
+	}
+
+	OutStatus = EStatus::Ok;
+
+	FCacheRecordBuilder RecordBuilder(Key);
+
+	if (!EnumHasAnyFlags(Policy.GetRecordPolicy(), ECachePolicy::SkipMeta))
+	{
+		RecordBuilder.SetMeta(FCbObject(Record.Get().GetMeta()));
+	}
+
+	if (FPayload Payload = Record.Get().GetValuePayload())
+	{
+		const ECachePolicy PayloadPolicy = Policy.GetPayloadPolicy(Payload.GetId());
+		GetCacheContent(Key, Context, PayloadPolicy, ECachePolicy::SkipValue, Payload, OutStatus);
+		if (Payload.IsNull())
+		{
+			return FOptionalCacheRecord();
+		}
+		RecordBuilder.SetValue(MoveTemp(Payload));
+	}
+
+	for (FPayload Payload : Record.Get().GetAttachmentPayloads())
+	{
+		const ECachePolicy PayloadPolicy = Policy.GetPayloadPolicy(Payload.GetId());
+		GetCacheContent(Key, Context, PayloadPolicy, ECachePolicy::SkipAttachments, Payload, OutStatus);
+		if (Payload.IsNull())
+		{
+			return FOptionalCacheRecord();
+		}
+		RecordBuilder.AddAttachment(MoveTemp(Payload));
+	}
+
+	return RecordBuilder.Build();
+}
+
+bool FPakFileDerivedDataBackend::PutCacheContent(const FCompressedBuffer& Content, const FStringView Context)
+{
+	const FIoHash& RawHash = Content.GetRawHash();
+	TStringBuilder<256> Path;
+	FPathViews::Append(Path, TEXT("Content"), RawHash);
+	if (!FileExists(Path))
+	{
+		if (!SaveFile(Path, Context, [&Content](FArchive& Ar) { Ar << const_cast<FCompressedBuffer&>(Content); }))
+		{
+			return false;
+		}
+	}
+	return true;
+}
+
+void FPakFileDerivedDataBackend::GetCacheContent(
+	const FCacheKey& Key,
+	const FStringView Context,
+	const ECachePolicy Policy,
+	const ECachePolicy SkipFlag,
+	FPayload& InOutPayload,
+	EStatus& InOutStatus)
+{
+	const FIoHash& RawHash = InOutPayload.GetRawHash();
+
+	if (!EnumHasAnyFlags(Policy, ECachePolicy::Query) || (EnumHasAnyFlags(Policy, SkipFlag) && InOutPayload.HasData()))
+	{
+		InOutPayload = FPayload(InOutPayload.GetId(), RawHash, InOutPayload.GetRawSize());
+		return;
+	}
+
+	if (InOutPayload.HasData())
+	{
+		return;
+	}
+
+	TStringBuilder<256> Path;
+	FPathViews::Append(Path, TEXT("Content"), RawHash);
+	if (EnumHasAllFlags(Policy, SkipFlag))
+	{
+		if (FileExists(Path))
+		{
+			return;
+		}
+	}
+	else
+	{
+		if (FSharedBuffer CompressedData = LoadFile(Path, Context))
+		{
+			if (FCompressedBuffer CompressedBuffer = FCompressedBuffer::FromCompressed(MoveTemp(CompressedData));
+				CompressedBuffer && CompressedBuffer.GetRawHash() == RawHash)
+			{
+				InOutPayload = FPayload(InOutPayload.GetId(), MoveTemp(CompressedBuffer));
+				return;
+			}
+			UE_LOG(LogDerivedDataCache, Display,
+				TEXT("%s: Cache miss with corrupted payload %s with hash %s for %s from '%.*s'"),
+				*CachePath, *WriteToString<16>(InOutPayload.GetId()), *WriteToString<48>(RawHash),
+				*WriteToString<96>(Key), Context.Len(), Context.GetData());
+			InOutStatus = EStatus::Error;
+			if (!EnumHasAnyFlags(Policy, ECachePolicy::PartialOnError))
+			{
+				InOutPayload = FPayload::Null;
+			}
+			return;
+		}
+	}
+
+	UE_LOG(LogDerivedDataCache, Verbose,
+		TEXT("%s: Cache miss with missing payload %s with hash %s for %s from '%.*s'"),
+		*CachePath, *WriteToString<16>(InOutPayload.GetId()), *WriteToString<48>(RawHash), *WriteToString<96>(Key),
+		Context.Len(), Context.GetData());
+	InOutStatus = EStatus::Error;
+	if (!EnumHasAnyFlags(Policy, ECachePolicy::PartialOnError))
+	{
+		InOutPayload = FPayload::Null;
+	}
+}
+
+class FCrcBuilder
+{
+public:
+	inline void Update(const void* Data, uint64 Size)
+	{
+		while (Size > 0)
+		{
+			const int32 CrcSize = int32(FMath::Min<uint64>(Size, MAX_int32));
+			Crc = FCrc::MemCrc_DEPRECATED(Data, CrcSize, Crc);
+			Size -= CrcSize;
+		}
+	}
+
+	inline uint32 Finalize()
+	{
+		return Crc;
+	}
+
+private:
+	uint32 Crc = 0;
+};
+
+class FPakWriterArchive final : public FArchive
+{
+public:
+	inline FPakWriterArchive(IFileHandle& InHandle, FStringView InPath)
+		: Handle(InHandle)
+		, Path(InPath)
+	{
+		SetIsSaving(true);
+		SetIsPersistent(true);
+	}
+
+	inline FString GetArchiveName() const final { return FString(Path); }
+	inline int64 TotalSize() final { return Handle.Size(); }
+	inline int64 Tell() final { unimplemented(); return 0; }
+	inline void Seek(int64 InPos) final { unimplemented(); }
+	inline void Flush() final { unimplemented(); }
+	inline bool Close() final { unimplemented(); return false; }
+
+	inline void Serialize(void* V, int64 Length) final
+	{
+		if (!Handle.Write(static_cast<uint8*>(V), Length))
+		{
+			SetError();
+		}
+	}
+
+private:
+	IFileHandle& Handle;
+	FStringView Path;
+};
+
+class FPakReaderArchive final : public FArchive
+{
+public:
+	inline FPakReaderArchive(IFileHandle& InHandle, FStringView InPath)
+		: Handle(InHandle)
+		, Path(InPath)
+	{
+		SetIsLoading(true);
+		SetIsPersistent(true);
+	}
+
+	inline FString GetArchiveName() const final { return FString(Path); }
+	inline int64 TotalSize() final { return Handle.Size(); }
+	inline int64 Tell() final { unimplemented(); return 0; }
+	inline void Seek(int64 InPos) final { unimplemented(); }
+	inline void Flush() final { unimplemented(); }
+	inline bool Close() final { unimplemented(); return false; }
+
+	inline void Serialize(void* V, int64 Length) final
+	{
+		if (!Handle.Read(static_cast<uint8*>(V), Length))
+		{
+			SetError();
+		}
+	}
+
+private:
+	IFileHandle& Handle;
+	FStringView Path;
+};
+
+bool FPakFileDerivedDataBackend::SaveFile(
+	const FStringView Path,
+	const FStringView Context,
+	TFunctionRef<void (FArchive&)> WriteFunction)
+{
+	FWriteScopeLock ScopeLock(SynchronizationObject);
+	check(FileHandle);
+	if (const int64 Offset = FileHandle->Tell(); Offset >= 0)
+	{
+		FPakWriterArchive Ar(*FileHandle, CachePath);
+		THashingArchiveProxy<FCrcBuilder> HashAr(Ar);
+		WriteFunction(HashAr);
+		if (const int64 EndOffset = FileHandle->Tell(); EndOffset >= Offset && !Ar.IsError())
+		{
+			FCacheValue& Item = CacheItems.Emplace(Path, FCacheValue(Offset, EndOffset - Offset, HashAr.GetHash()));
+			UE_LOG(LogDerivedDataCache, Log,
+				TEXT("%s: File %.*s from '%.*s' written with offset %" INT64_FMT ", size %" INT64_FMT", CRC 0x%08x."),
+				*CachePath, Path.Len(), Path.GetData(), Context.Len(), Context.GetData(), Item.Offset, Item.Size, Item.Crc);
+			return true;
+		}
+	}
+	return false;
+}
+
+FSharedBuffer FPakFileDerivedDataBackend::LoadFile(const FStringView Path, const FStringView Context)
+{
+	FWriteScopeLock ScopeLock(SynchronizationObject);
+	if (const FCacheValue* Item = CacheItems.FindByHash(GetTypeHash(Path), Path))
+	{
+		check(FileHandle);
+		ON_SCOPE_EXIT
+		{
+			if (bWriting)
+			{
+				FileHandle->SeekFromEnd();
+			}
+		};
+		check(Item->Size);
+		if (FileHandle->Seek(Item->Offset))
+		{
+			FPakReaderArchive Ar(*FileHandle, CachePath);
+			THashingArchiveProxy<FCrcBuilder> HashAr(Ar);
+			FUniqueBuffer MutableBuffer = FUniqueBuffer::Alloc(uint64(Item->Size));
+			HashAr.Serialize(MutableBuffer.GetData(), Item->Size);
+			if (Ar.IsError())
+			{
+				UE_LOG(LogDerivedDataCache, Display,
+					TEXT("%s: File %.*s from '%.*s' failed to read %" INT64_FMT " bytes."),
+					*CachePath, Path.Len(), Path.GetData(), Context.Len(), Context.GetData(), Item->Size);
+			}
+			else if (const uint32 TestCrc = HashAr.GetHash(); TestCrc != Item->Crc)
+			{
+				UE_LOG(LogDerivedDataCache, Display,
+					TEXT("%s: File %.*s from '%.*s' is corrupted and has CRC 0x%08x when 0x%08x is expected."),
+					*CachePath, Path.Len(), Path.GetData(), Context.Len(), Context.GetData(), TestCrc, Item->Crc);
+			}
+			else
+			{
+				return MutableBuffer.MoveToShared();
+			}
+		}
+	}
+	return FSharedBuffer();
+}
+
+bool FPakFileDerivedDataBackend::FileExists(const FStringView Path)
+{
+	FReadScopeLock ScopeLock(SynchronizationObject);
+	const uint32 PathHash = GetTypeHash(Path);
+	return CacheItems.ContainsByHash(PathHash, Path);
+}
+
 FCompressedPakFileDerivedDataBackend::FCompressedPakFileDerivedDataBackend(const TCHAR* InFilename, bool bInWriting)
 	: FPakFileDerivedDataBackend(InFilename, bInWriting)
 {
@@ -534,4 +1103,4 @@ bool FCompressedPakFileDerivedDataBackend::GetCachedData(const TCHAR* CacheKey,
 	return true;
 }
 
-} // UE::DerivedData::Backends
+} // UE::DerivedData::CacheStore::PakFile
diff --git a/Engine/Source/Developer/DerivedDataCache/Private/PakFileDerivedDataBackend.h b/Engine/Source/Developer/DerivedDataCache/Private/PakFileDerivedDataBackend.h
--- a/Engine/Source/Developer/DerivedDataCache/Private/PakFileDerivedDataBackend.h
+++ b/Engine/Source/Developer/DerivedDataCache/Private/PakFileDerivedDataBackend.h
@@ -2,21 +2,21 @@
 
 #pragma once
 
-#include "CoreMinimal.h"
-#include "HAL/FileManager.h"
-#include "Misc/Paths.h"
+#include "CoreTypes.h"
+#include "Containers/StringFwd.h"
 #include "DerivedDataBackendInterface.h"
-#include "ProfilingDebugging/CookStats.h"
 #include "DerivedDataCacheUsageStats.h"
-#include "Misc/ScopeLock.h"
-#include "Misc/FileHelper.h"
-#include "Serialization/MemoryReader.h"
-#include "Serialization/MemoryWriter.h"
+#include "HAL/CriticalSection.h"
+#include "ProfilingDebugging/CookStats.h"
 #include "Templates/UniquePtr.h"
 
-class Error;
+class FCompressedBuffer;
+class IFileHandle;
 
-namespace UE::DerivedData::Backends
+namespace UE::DerivedData { class FOptionalCacheRecord; }
+namespace UE::DerivedData { class FPayload; }
+
+namespace UE::DerivedData::CacheStore::PakFile
 {
 
 /** 
@@ -91,7 +91,7 @@ public:
 	
 	const FString& GetFilename() const
 	{
-		return Filename;
+		return CachePath;
 	}
 
 	static bool SortAndCopy(const FString &InputFilename, const FString &OutputFilename);
@@ -127,6 +127,34 @@ public:
 		IRequestOwner& Owner,
 		FOnCacheGetChunkComplete&& OnComplete) override;
 
+private:
+	uint64 MeasureCompressedCacheRecord(const FCacheRecord& Record) const;
+	uint64 MeasureRawCacheRecord(const FCacheRecord& Record) const;
+
+	bool PutCacheRecord(const FCacheRecord& Record, FStringView Context, ECachePolicy Policy);
+	bool PutCacheContent(const FCompressedBuffer& Content, const FStringView Context);
+
+	FOptionalCacheRecord GetCacheRecordOnly(
+		const FCacheKey& Key,
+		const FStringView Context,
+		const FCacheRecordPolicy& Policy);
+	FOptionalCacheRecord GetCacheRecord(
+		const FCacheKey& Key,
+		const FStringView Context,
+		const FCacheRecordPolicy& Policy,
+		EStatus& OutStatus);
+	void GetCacheContent(
+		const FCacheKey& Key,
+		const FStringView Context,
+		const ECachePolicy Policy,
+		const ECachePolicy SkipFlag,
+		FPayload& InOutPayload,
+		EStatus& InOutStatus);
+
+	bool SaveFile(FStringView Path, FStringView Context, TFunctionRef<void (FArchive&)> WriteFunction);
+	FSharedBuffer LoadFile(FStringView Path, FStringView Context);
+	bool FileExists(FStringView Path);
+
 private:
 	FDerivedDataCacheUsageStats UsageStats;
 
@@ -135,7 +163,7 @@ private:
 		int64 Offset;
 		int64 Size;
 		uint32 Crc;
-		FCacheValue(int64 InOffset, uint32 InSize, uint32 InCrc)
+		FCacheValue(int64 InOffset, int64 InSize, uint32 InCrc)
 			: Offset(InOffset)
 			, Size(InSize)
 			, Crc(InCrc)
@@ -147,14 +175,20 @@ private:
 	bool bWriting;
 	/** When set to true, we are a pak writer and we saved, so we shouldn't be used anymore. Also, a read cache that failed to open. */
 	bool bClosed;
-	/** Object used for synchronization via a scoped lock						*/
-	FCriticalSection	SynchronizationObject;
+	/** Object used for synchronization via scoped read or write locks. */
+	FRWLock SynchronizationObject;
 	/** Set of files that are being written to disk asynchronously. */
 	TMap<FString, FCacheValue> CacheItems;
 	/** File handle of pak. */
-	TUniquePtr<FArchive> FileHandle;
+	TUniquePtr<IFileHandle> FileHandle;
 	/** File name of pak. */
-	FString Filename;
+	FString CachePath;
+
+	/** Maximum total size of compressed data stored within a record package with multiple attachments. */
+	uint64 MaxRecordSizeKB = 256;
+	/** Maximum total size of compressed data stored within a value package, or a record package with one attachment. */
+	uint64 MaxValueSizeKB = 1024;
+
 	enum 
 	{
 		/** Magic number to use in header */
@@ -181,4 +215,4 @@ private:
 	static const ECompressionFlags CompressionFlags = COMPRESS_BiasMemory;
 };
 
-} // UE::DerivedData::Backends
+} // UE::DerivedData::CacheStore::PakFile
