DDC: Fixed the async put wrapper with -NoThreading

#rnx

diff --git a/Engine/Source/Developer/DerivedDataCache/Private/DerivedDataBackendAsyncPutWrapper.cpp b/Engine/Source/Developer/DerivedDataCache/Private/DerivedDataBackendAsyncPutWrapper.cpp
--- a/Engine/Source/Developer/DerivedDataCache/Private/DerivedDataBackendAsyncPutWrapper.cpp
+++ b/Engine/Source/Developer/DerivedDataCache/Private/DerivedDataBackendAsyncPutWrapper.cpp
@@ -463,7 +463,7 @@ void FDerivedDataBackendAsyncPutWrapper::Put(
 	IRequestOwner& Owner,
 	FOnCachePutComplete&& OnComplete)
 {
-	if (Owner.GetPriority() == EPriority::Blocking)
+	if (Owner.GetPriority() == EPriority::Blocking || !GDDCIOThreadPool)
 	{
 		InnerBackend->Put(Records, Context, Policy, Owner, MoveTemp(OnComplete));
 	}
@@ -497,7 +497,7 @@ void FDerivedDataBackendAsyncPutWrapper::Get(
 	IRequestOwner& Owner,
 	FOnCacheGetComplete&& OnComplete)
 {
-	if (Owner.GetPriority() == EPriority::Blocking)
+	if (Owner.GetPriority() == EPriority::Blocking || !GDDCIOThreadPool)
 	{
 		InnerBackend->Get(Keys, Context, Policy, Owner, MoveTemp(OnComplete));
 	}
@@ -530,7 +530,7 @@ void FDerivedDataBackendAsyncPutWrapper::GetChunks(
 	IRequestOwner& Owner,
 	FOnCacheGetChunkComplete&& OnComplete)
 {
-	if (Owner.GetPriority() == EPriority::Blocking)
+	if (Owner.GetPriority() == EPriority::Blocking || !GDDCIOThreadPool)
 	{
 		InnerBackend->GetChunks(Chunks, Context, Owner, MoveTemp(OnComplete));
 	}
