Optimise FZenDerivedDataBackend::IsRemote() by checking the once on backend creation and caching the result

diff --git a/Engine/Source/Developer/DerivedDataCache/Private/ZenDerivedDataBackend.cpp b/Engine/Source/Developer/DerivedDataCache/Private/ZenDerivedDataBackend.cpp
--- a/Engine/Source/Developer/DerivedDataCache/Private/ZenDerivedDataBackend.cpp
+++ b/Engine/Source/Developer/DerivedDataCache/Private/ZenDerivedDataBackend.cpp
@@ -69,6 +69,14 @@ FZenDerivedDataBackend::FZenDerivedDataBackend(
 	{
 		RequestPool = MakeUnique<Zen::FZenHttpRequestPool>(ZenService.GetInstance().GetURL(), 32);
 		bIsUsable = true;
+		bIsRemote = false;
+
+		Zen::FZenStats ZenStats;
+
+		if (ZenService.GetInstance().GetStats(ZenStats) == true)
+		{
+			 bIsRemote = ZenStats.UpstreamStats.EndPointStats.IsEmpty() == false;
+		}
 	}
 
 	GConfig->GetInt(TEXT("Zen"), TEXT("CacheRecordBatchSize"), CacheRecordBatchSize, GEngineIni);
@@ -91,14 +99,7 @@ FString FZenDerivedDataBackend::GetName() const
 
 bool FZenDerivedDataBackend::IsRemote() const
 {
-	Zen::FZenStats ZenStats;
-
-	if (ZenService.GetInstance().GetStats(ZenStats) == true)
-	{
-		return( ZenStats.UpstreamStats.EndPointStats.IsEmpty() == false );
-	}
-
-	return false;
+	return bIsRemote;
 }
 
 bool FZenDerivedDataBackend::IsServiceReady()
diff --git a/Engine/Source/Developer/DerivedDataCache/Private/ZenDerivedDataBackend.h b/Engine/Source/Developer/DerivedDataCache/Private/ZenDerivedDataBackend.h
--- a/Engine/Source/Developer/DerivedDataCache/Private/ZenDerivedDataBackend.h
+++ b/Engine/Source/Developer/DerivedDataCache/Private/ZenDerivedDataBackend.h
@@ -145,6 +145,7 @@ private:
 	mutable FDerivedDataCacheUsageStats UsageStats;
 	TUniquePtr<UE::Zen::FZenHttpRequestPool> RequestPool;
 	bool bIsUsable = false;
+	bool bIsRemote = false;
 	uint32 FailedLoginAttempts = 0;
 	uint32 MaxAttempts = 4;
 	int32 CacheRecordBatchSize = 8;
