DDC: Split build diagnostics into messages and logs

Messages are deterministic output recorded explicitly by the build function. Logs will be automatically captured from the output log and are not necessarily deterministic. The presence of logs in a build output automatically disables caching of the build output.

#rnx

diff --git a/Engine/Source/Developer/DerivedDataCache/Private/DerivedDataBuild.cpp b/Engine/Source/Developer/DerivedDataCache/Private/DerivedDataBuild.cpp
--- a/Engine/Source/Developer/DerivedDataCache/Private/DerivedDataBuild.cpp
+++ b/Engine/Source/Developer/DerivedDataCache/Private/DerivedDataBuild.cpp
@@ -119,7 +119,7 @@ private:
 	TUniquePtr<IBuildFunctionRegistry> FunctionRegistry{CreateBuildFunctionRegistry()};
 	TUniquePtr<IBuildWorkerRegistry> WorkerRegistry{CreateBuildWorkerRegistry()};
 	TUniquePtr<IBuildScheduler> DefaultScheduler{CreateBuildScheduler()};
-	const FGuid Version{TEXT("ed89db6a-335d-45d3-8599-39d3b09f7ad7")};
+	const FGuid Version{TEXT("02c9caed-123c-4fc3-b8c1-7816058abed5")};
 };
 
 ///////////////////////////////////////////////////////////////////////////////////////////////////////////////////////
diff --git a/Engine/Source/Developer/DerivedDataCache/Private/DerivedDataBuildJob.cpp b/Engine/Source/Developer/DerivedDataCache/Private/DerivedDataBuildJob.cpp
--- a/Engine/Source/Developer/DerivedDataCache/Private/DerivedDataBuildJob.cpp
+++ b/Engine/Source/Developer/DerivedDataCache/Private/DerivedDataBuildJob.cpp
@@ -172,7 +172,7 @@ private:
 	void SetOutputNoCheck(FBuildOutput&& Output, EBuildJobState NewState = EBuildJobState::CacheStore);
 
 	/** Terminate the job and send the error to the output complete callback. */
-	void CompleteWithError(FStringView Error);
+	void CompleteWithError(FUtf8StringView Error);
 
 	/** Advance to the new state, dispatching to the scheduler and invoking callbacks as appropriate. */
 	void AdvanceToState(EBuildJobState NewState);
@@ -397,25 +397,25 @@ void FBuildJob::CreateContext()
 	const IBuildFunction* Function = BuildSystem.GetFunctionRegistry().FindFunction(FunctionName);
 	if (BuildSystem.GetVersion() != Action.Get().GetBuildSystemVersion())
 	{
-		return CompleteWithError(WriteToString<192>(TEXT("Failed because the build system is version "_SV),
-			BuildSystem.GetVersion(), TEXT(" when version "_SV),
-			Action.Get().GetBuildSystemVersion(), TEXT(" is expected."_SV)));
+		return CompleteWithError(WriteToUtf8String<192>("Failed because the build system is version "_ASV,
+			BuildSystem.GetVersion(), " when version "_ASV,
+			Action.Get().GetBuildSystemVersion(), " is expected."_ASV));
 	}
 	else if (!Function)
 	{
-		return CompleteWithError(WriteToString<128>(TEXT("Failed because the function "_SV), FunctionName,
-			TEXT(" was not found."_SV)));
+		return CompleteWithError(WriteToUtf8String<128>("Failed because the function "_ASV, FunctionName,
+			" was not found."_ASV));
 	}
 	else if (!Function->GetVersion().IsValid())
 	{
-		return CompleteWithError(WriteToString<128>(TEXT("Failed because the function "_SV), FunctionName,
-			TEXT(" has a version of zero."_SV)));
+		return CompleteWithError(WriteToUtf8String<128>("Failed because the function "_ASV, FunctionName,
+			" has a version of zero."_ASV));
 	}
 	else if (Function->GetVersion() != Action.Get().GetFunctionVersion())
 	{
-		return CompleteWithError(WriteToString<192>(TEXT("Failed because the function "_SV), FunctionName,
-			TEXT(" is version "_SV), Function->GetVersion(), TEXT(" when version "_SV),
-			Action.Get().GetFunctionVersion(), TEXT(" is expected."_SV)));
+		return CompleteWithError(WriteToUtf8String<192>("Failed because the function "_ASV, FunctionName,
+			" is version "_ASV, Function->GetVersion(), " when version "_ASV,
+			Action.Get().GetFunctionVersion(), " is expected."_ASV));
 	}
 	else
 	{
@@ -430,6 +430,11 @@ void FBuildJob::CreateContext()
 		EnumAddFlags(BuildStatus, EBuildStatus::CacheKey);
 	}
 
+	if (OutputBuilder.HasError())
+	{
+		return AdvanceToState(EBuildJobState::Complete);
+	}
+
 	// Populate the scheduler params with the information that is available now.
 	FBuildSchedulerParams& SchedulerParams = Schedule->EditParameters();
 	SchedulerParams.Key = Action.Get().GetKey();
@@ -511,7 +516,7 @@ void FBuildJob::EndCacheQuery(FCacheGetCompleteParams&& Params)
 {
 	if (Params.Status == EStatus::Canceled)
 	{
-		return CompleteWithError(TEXT("Build was canceled."));
+		return CompleteWithError("Build was canceled."_ASV);
 	}
 	if (Params.Status == EStatus::Ok)
 	{
@@ -523,7 +528,7 @@ void FBuildJob::EndCacheQuery(FCacheGetCompleteParams&& Params)
 	}
 	if (!EnumHasAnyFlags(BuildPolicy.GetCombinedPolicy() & Context->GetBuildPolicyMask(), EBuildPolicy::Build))
 	{
-		return CompleteWithError(TEXT("Failed to fetch from the cache and build policy does not allow execution."));
+		return CompleteWithError("Failed to fetch from the cache and build policy does not allow execution."_ASV);
 	}
 	return AdvanceToState(EBuildJobState::ExecuteRemote);
 }
@@ -556,7 +561,8 @@ void FBuildJob::EnterCacheStore()
 		!EnumHasAnyFlags(Context->GetCachePolicyMask(), ECachePolicy::Store) ||
 		!EnumHasAnyFlags(Context->GetBuildPolicyMask() & BuildPolicy.GetCombinedPolicy(), EBuildPolicy::CacheStoreOnBuild) ||
 		EnumHasAnyFlags(BuildStatus, EBuildStatus::CacheQueryHit) ||
-		Output.Get().HasError())
+		Output.Get().HasError() ||
+		Output.Get().HasLogs())
 	{
 		return AdvanceToState(EBuildJobState::Complete);
 	}
@@ -595,11 +601,11 @@ void FBuildJob::EnterResolveKey()
 	}
 	if (DefinitionKey != FBuildKey::Empty)
 	{
-		return CompleteWithError(TEXT("Failed to resolve null key."_SV));
+		return CompleteWithError("Failed to resolve null key."_ASV);
 	}
 	if (!InputResolver)
 	{
-		return CompleteWithError(TEXT("Failed to resolve key due to null input resolver."_SV));
+		return CompleteWithError("Failed to resolve key due to null input resolver."_ASV);
 	}
 }
 
@@ -614,7 +620,7 @@ void FBuildJob::EndResolveKey(FBuildKeyResolvedParams&& Params)
 {
 	if (Params.Status == EStatus::Canceled)
 	{
-		return CompleteWithError(TEXT("Build was canceled."));
+		return CompleteWithError("Build was canceled."_ASV);
 	}
 	if (Params.Status == EStatus::Ok && Params.Definition)
 	{
@@ -622,7 +628,7 @@ void FBuildJob::EndResolveKey(FBuildKeyResolvedParams&& Params)
 	}
 	else
 	{
-		return CompleteWithError(WriteToString<128>(TEXT("Failed to resolve key "_SV), Params.Key, TEXT("."_SV)));
+		return CompleteWithError(WriteToUtf8String<128>("Failed to resolve key "_ASV, Params.Key, "."_ASV));
 	}
 }
 
@@ -636,7 +642,7 @@ void FBuildJob::EnterResolveInputMeta()
 	}
 	if (!InputResolver)
 	{
-		return CompleteWithError(TEXT("Failed to resolve input metadata due to null input resolver."_SV));
+		return CompleteWithError("Failed to resolve input metadata due to null input resolver."_ASV);
 	}
 }
 
@@ -651,7 +657,7 @@ void FBuildJob::EndResolveInputMeta(FBuildInputMetaResolvedParams&& Params)
 {
 	if (Params.Status == EStatus::Canceled)
 	{
-		return CompleteWithError(TEXT("Build was canceled."));
+		return CompleteWithError("Build was canceled."_ASV);
 	}
 	if (Params.Status == EStatus::Ok)
 	{
@@ -659,7 +665,7 @@ void FBuildJob::EndResolveInputMeta(FBuildInputMetaResolvedParams&& Params)
 	}
 	else
 	{
-		return CompleteWithError(TEXT("Failed to resolve input metadata."_SV));
+		return CompleteWithError("Failed to resolve input metadata."_ASV);
 	}
 }
 
@@ -709,7 +715,7 @@ void FBuildJob::EnterResolveInputData()
 
 	if (!InputResolver)
 	{
-		return CompleteWithError(TEXT("Failed to resolve input data due to null input resolver."_SV));
+		return CompleteWithError("Failed to resolve input data due to null input resolver."_ASV);
 	}
 }
 
@@ -734,7 +740,7 @@ void FBuildJob::EndResolveInputData(FBuildInputDataResolvedParams&& Params)
 {
 	if (Params.Status == EStatus::Canceled)
 	{
-		return CompleteWithError(TEXT("Build was canceled."));
+		return CompleteWithError("Build was canceled."_ASV);
 	}
 	if (Params.Status == EStatus::Ok)
 	{
@@ -756,7 +762,7 @@ void FBuildJob::EndResolveInputData(FBuildInputDataResolvedParams&& Params)
 	}
 	else
 	{
-		return CompleteWithError(TEXT("Failed to resolve input data."_SV));
+		return CompleteWithError("Failed to resolve input data."_ASV);
 	}
 }
 
@@ -815,7 +821,7 @@ void FBuildJob::EndExecuteRemote(FBuildWorkerActionCompleteParams&& Params)
 {
 	if (Params.Status == EStatus::Canceled)
 	{
-		return CompleteWithError(TEXT("Build was canceled."));
+		return CompleteWithError("Build was canceled."_ASV);
 	}
 	if (Params.Output)
 	{
@@ -874,13 +880,13 @@ void FBuildJob::EnterExecuteLocal()
 		}
 		else if (EnumHasAnyFlags(BuildStatus, EBuildStatus::BuildTryRemote))
 		{
-			return CompleteWithError(TEXT("Failed because build policy does not allow local execution, ")
-				TEXT("and remote execution failed to build."_SV));
+			return CompleteWithError("Failed because build policy does not allow local execution, "
+				"and remote execution failed to build."_ASV);
 		}
 		else
 		{
-			return CompleteWithError(TEXT("Failed because build policy does not allow local execution, ")
-				TEXT("and remote execution was not available."_SV));
+			return CompleteWithError("Failed because build policy does not allow local execution, "
+				"and remote execution was not available."_ASV);
 		}
 	}
 
@@ -893,9 +899,9 @@ void FBuildJob::EnterExecuteLocal()
 		}
 		else
 		{
-			CompleteWithError(WriteToString<256>(TEXT("Failed because input '"_SV), Key, TEXT("' with hash "_SV),
-				RawHash, TEXT(" ("_SV), RawSize, TEXT(" bytes) was resolved with hash "_SV), Buffer.GetRawHash(),
-				TEXT(" ("_SV), Buffer.GetRawSize(), TEXT(" bytes)."_SV)));
+			CompleteWithError(WriteToUtf8String<256>("Failed because input '"_ASV, Key, "' with hash "_ASV,
+				RawHash, " ("_ASV, RawSize, " bytes) was resolved with hash "_ASV, Buffer.GetRawHash(),
+				" ("_ASV, Buffer.GetRawSize(), " bytes)."_ASV));
 		}
 	});
 	Action.Reset();
@@ -1012,13 +1018,13 @@ void FBuildJob::SetOutputNoCheck(FBuildOutput&& InOutput, EBuildJobState NewStat
 
 ///////////////////////////////////////////////////////////////////////////////////////////////////////////////////////
 
-void FBuildJob::CompleteWithError(FStringView Error)
+void FBuildJob::CompleteWithError(FUtf8StringView Error)
 {
 	if (FWriteScopeLock WriteLock(Lock); Output || NextState == EBuildJobState::Complete)
 	{
 		return;
 	}
-	OutputBuilder.AddError(TEXT("LogDerivedDataBuild"_SV), Error);
+	OutputBuilder.AddLog({"LogDerivedDataBuild"_ASV, Error, EBuildOutputLogLevel::Error});
 	return SetOutputNoCheck(OutputBuilder.Build(), EBuildJobState::Complete);
 }
 
@@ -1038,7 +1044,7 @@ void FBuildJob::AdvanceToState(EBuildJobState NewState)
 			NewState = EBuildJobState::Complete;
 			if (Output.IsNull())
 			{
-				OutputBuilder.AddError(TEXT("LogDerivedDataBuild"_SV), TEXT("Build was canceled."_SV));
+				OutputBuilder.AddLog({"LogDerivedDataBuild"_ASV, "Build was canceled."_ASV, EBuildOutputLogLevel::Error});
 			}
 		}
 
diff --git a/Engine/Source/Developer/DerivedDataCache/Private/DerivedDataBuildJobContext.cpp b/Engine/Source/Developer/DerivedDataCache/Private/DerivedDataBuildJobContext.cpp
--- a/Engine/Source/Developer/DerivedDataCache/Private/DerivedDataBuildJobContext.cpp
+++ b/Engine/Source/Developer/DerivedDataCache/Private/DerivedDataBuildJobContext.cpp
@@ -2,6 +2,7 @@
 
 #include "DerivedDataBuildJobContext.h"
 
+#include "Containers/StringConv.h"
 #include "DerivedDataBuildJob.h"
 #include "DerivedDataBuildOutput.h"
 #include "DerivedDataBuildPrivate.h"
@@ -46,6 +47,21 @@ void FBuildJobContext::AddInput(FStringView Key, const FCompressedBuffer& Value)
 	Inputs.EmplaceByHash(GetTypeHash(Key), Key, Value);
 }
 
+void FBuildJobContext::AddError(FStringView Message)
+{
+	OutputBuilder.AddMessage({FTCHARToUTF8(Message), EBuildOutputMessageLevel::Error});
+}
+
+void FBuildJobContext::AddWarning(FStringView Message)
+{
+	OutputBuilder.AddMessage({FTCHARToUTF8(Message), EBuildOutputMessageLevel::Warning});
+}
+
+void FBuildJobContext::AddMessage(FStringView Message)
+{
+	OutputBuilder.AddMessage({FTCHARToUTF8(Message), EBuildOutputMessageLevel::Display});
+}
+
 void FBuildJobContext::ResetInputs()
 {
 	Constants.Empty();
@@ -78,7 +94,7 @@ FSharedBuffer FBuildJobContext::FindInput(FStringView Key) const
 				<< TEXT(" and raw size ") << Input->GetRawSize() << TEXT(" but has raw hash ") << RawHash
 				<< TEXT(" and raw size ") << Buffer.GetSize() << TEXT(" after decompression for build of '")
 				<< Job.GetName() << TEXT("' by ") << Job.GetFunction() << TEXT(".");
-			OutputBuilder.AddError(TEXT("LogDerivedDataBuild"_SV), Error);
+			OutputBuilder.AddLog({"LogDerivedDataBuild"_ASV, FTCHARToUTF8(Error), EBuildOutputLogLevel::Error});
 			UE_LOG(LogDerivedDataBuild, Error, TEXT("%.*s"), Error.Len(), Error.GetData());
 		}
 	}
diff --git a/Engine/Source/Developer/DerivedDataCache/Private/DerivedDataBuildJobContext.h b/Engine/Source/Developer/DerivedDataCache/Private/DerivedDataBuildJobContext.h
--- a/Engine/Source/Developer/DerivedDataCache/Private/DerivedDataBuildJobContext.h
+++ b/Engine/Source/Developer/DerivedDataCache/Private/DerivedDataBuildJobContext.h
@@ -48,6 +48,10 @@ public:
 	void AddConstant(FStringView Key, FCbObject&& Value);
 	void AddInput(FStringView Key, const FCompressedBuffer& Value);
 
+	void AddError(FStringView Message) final;
+	void AddWarning(FStringView Message) final;
+	void AddMessage(FStringView Message) final;
+
 	void ResetInputs();
 
 private:
diff --git a/Engine/Source/Developer/DerivedDataCache/Private/DerivedDataBuildOutput.cpp b/Engine/Source/Developer/DerivedDataCache/Private/DerivedDataBuildOutput.cpp
--- a/Engine/Source/Developer/DerivedDataCache/Private/DerivedDataBuildOutput.cpp
+++ b/Engine/Source/Developer/DerivedDataCache/Private/DerivedDataBuildOutput.cpp
@@ -2,12 +2,10 @@
 
 #include "DerivedDataBuildOutput.h"
 
-#include "Algo/AllOf.h"
 #include "Algo/BinarySearch.h"
-#include "Containers/StringConv.h"
+#include "Algo/Find.h"
 #include "Containers/StringView.h"
 #include "Containers/UnrealString.h"
-#include "DerivedDataBuildDefinition.h"
 #include "DerivedDataBuildKey.h"
 #include "DerivedDataBuildPrivate.h"
 #include "DerivedDataCacheRecord.h"
@@ -19,28 +17,62 @@
 namespace UE::DerivedData::Private
 {
 
-static FStringView LexToString(EBuildDiagnosticLevel Level)
+static FUtf8StringView LexToString(EBuildOutputMessageLevel Level)
 {
 	switch (Level)
 	{
-	default: return TEXT("Unknown"_SV);
-	case EBuildDiagnosticLevel::Error: return TEXT("Error"_SV);
-	case EBuildDiagnosticLevel::Warning: return TEXT("Warning"_SV);
+	case EBuildOutputMessageLevel::Error: return "Error"_U8SV;
+	case EBuildOutputMessageLevel::Warning: return "Warning"_U8SV;
+	case EBuildOutputMessageLevel::Display: return "Display"_U8SV;
+	default: return "Unknown"_U8SV;
 	}
 }
 
-static void LexFromString(EBuildDiagnosticLevel& OutLevel, FUtf8StringView String)
+static bool TryLexFromString(EBuildOutputMessageLevel& OutLevel, FUtf8StringView String)
 {
+	if (String.Equals("Error"_U8SV, ESearchCase::CaseSensitive))
+	{
+		OutLevel = EBuildOutputMessageLevel::Error;
+		return true;
+	}
 	if (String.Equals("Warning"_U8SV, ESearchCase::CaseSensitive))
 	{
-		OutLevel = EBuildDiagnosticLevel::Warning;
+		OutLevel = EBuildOutputMessageLevel::Warning;
+		return true;
 	}
-	else
+	if (String.Equals("Display"_U8SV, ESearchCase::CaseSensitive))
+	{
+		OutLevel = EBuildOutputMessageLevel::Display;
+		return true;
+	}
+	return false;
+}
+
+static FUtf8StringView LexToString(EBuildOutputLogLevel Level)
+{
+	switch (Level)
 	{
-		OutLevel = EBuildDiagnosticLevel::Error;
+	case EBuildOutputLogLevel::Error: return "Error"_U8SV;
+	case EBuildOutputLogLevel::Warning: return "Warning"_U8SV;
+	default: return "Unknown"_U8SV;
 	}
 }
 
+static bool TryLexFromString(EBuildOutputLogLevel& OutLevel, FUtf8StringView String)
+{
+	if (String.Equals("Error"_U8SV, ESearchCase::CaseSensitive))
+	{
+		OutLevel = EBuildOutputLogLevel::Error;
+		return true;
+	}
+	if (String.Equals("Warning"_U8SV, ESearchCase::CaseSensitive))
+	{
+		OutLevel = EBuildOutputLogLevel::Warning;
+		return true;
+	}
+	return false;
+}
+
 ///////////////////////////////////////////////////////////////////////////////////////////////////////////////////////
 
 class FBuildOutputBuilderInternal final : public IBuildOutputBuilderInternal
@@ -49,19 +81,19 @@ public:
 	explicit FBuildOutputBuilderInternal(FStringView InName, FStringView InFunction)
 		: Name(InName)
 		, Function(InFunction)
-		, bHasError(false)
-		, bHasDiagnostics(false)
 	{
 		checkf(!Name.IsEmpty(), TEXT("A build output requires a non-empty name."));
 		AssertValidBuildFunctionName(Function, Name);
-		DiagnosticsWriter.BeginArray();
+		MessageWriter.BeginArray();
+		LogWriter.BeginArray();
 	}
 
 	~FBuildOutputBuilderInternal() final = default;
 
-	void SetMeta(FCbObject&& InMeta) final { Meta = MoveTemp(InMeta); }
+	void SetMeta(FCbObject&& InMeta) final { Meta = MoveTemp(InMeta); Meta.MakeOwned(); }
 	void AddPayload(const FPayload& Payload) final;
-	void AddDiagnostic(const FBuildDiagnostic& Diagnostic) final;
+	void AddMessage(const FBuildOutputMessage& Message) final;
+	void AddLog(const FBuildOutputLog& Log) final;
 	bool HasError() const final { return bHasError; }
 	FBuildOutput Build() final;
 
@@ -69,18 +101,19 @@ public:
 	FString Function;
 	FCbObject Meta;
 	TArray<FPayload> Payloads;
-	FCbField Diagnostics;
-	FCbWriter DiagnosticsWriter;
-	bool bHasError;
-	bool bHasDiagnostics;
+	FCbWriter MessageWriter;
+	FCbWriter LogWriter;
+	bool bHasMessages = false;
+	bool bHasLogs = false;
+	bool bHasError = false;
 };
 
 class FBuildOutputInternal final : public IBuildOutputInternal
 {
 public:
-	explicit FBuildOutputInternal(FBuildOutputBuilderInternal&& Output);
-	explicit FBuildOutputInternal(FStringView Name, FStringView Function, const FCbObject& Output, bool& bOutIsValid);
-	explicit FBuildOutputInternal(FStringView Name, FStringView Function, const FCacheRecord& Output, bool& bOutIsValid);
+	FBuildOutputInternal(FString&& Name, FString&& Function, FCbObject&& Meta, FCbObject&& Output, TArray<FPayload>&& Payloads);
+	FBuildOutputInternal(FStringView Name, FStringView Function, const FCbObject& Output, bool& bOutIsValid);
+	FBuildOutputInternal(FStringView Name, FStringView Function, const FCacheRecord& Output, bool& bOutIsValid);
 
 	~FBuildOutputInternal() final = default;
 
@@ -91,14 +124,16 @@ public:
 
 	const FPayload& GetPayload(const FPayloadId& Id) const final;
 	TConstArrayView<FPayload> GetPayloads() const final { return Payloads; }
-	void IterateDiagnostics(TFunctionRef<void (const FBuildDiagnostic& Diagnostic)> Visitor) const final;
+	TConstArrayView<FBuildOutputMessage> GetMessages() const final { return Messages; }
+	TConstArrayView<FBuildOutputLog> GetLogs() const final { return Logs; }
+	bool HasLogs() const final { return !Logs.IsEmpty(); }
 	bool HasError() const final;
 
+	bool TryLoad();
+
 	void Save(FCbWriter& Writer) const final;
 	void Save(FCacheRecordBuilder& RecordBuilder) const final;
 
-	bool IsValid() const;
-
 	inline void AddRef() const final
 	{
 		ReferenceCount.fetch_add(1, std::memory_order_relaxed);
@@ -116,63 +151,55 @@ private:
 	FString Name;
 	FString Function;
 	FCbObject Meta;
+	FCbObject Output;
 	TArray<FPayload> Payloads;
-	FCbField Diagnostics;
+	TArray<FBuildOutputMessage> Messages;
+	TArray<FBuildOutputLog> Logs;
 	mutable std::atomic<uint32> ReferenceCount{0};
 };
 
 ///////////////////////////////////////////////////////////////////////////////////////////////////////////////////////
 
-FBuildOutputInternal::FBuildOutputInternal(FBuildOutputBuilderInternal&& InOutput)
-	: Name(MoveTemp(InOutput.Name))
-	, Function(MoveTemp(InOutput.Function))
-	, Meta(MoveTemp(InOutput.Meta))
-	, Payloads(MoveTemp(InOutput.Payloads))
-	, Diagnostics(MoveTemp(InOutput.Diagnostics))
+FBuildOutputInternal::FBuildOutputInternal(
+	FString&& InName,
+	FString&& InFunction,
+	FCbObject&& InMeta,
+	FCbObject&& InOutput,
+	TArray<FPayload>&& InPayloads)
+	: Name(MoveTemp(InName))
+	, Function(MoveTemp(InFunction))
+	, Meta(MoveTemp(InMeta))
+	, Output(MoveTemp(InOutput))
+	, Payloads(MoveTemp(InPayloads))
 {
-	if (InOutput.bHasError)
-	{
-		Payloads.Empty();
-	}
+	Meta.MakeOwned();
+	Output.MakeOwned();
+	TryLoad();
 }
 
 FBuildOutputInternal::FBuildOutputInternal(FStringView InName, FStringView InFunction, const FCbObject& InOutput, bool& bOutIsValid)
 	: Name(InName)
 	, Function(InFunction)
+	, Output(InOutput)
 {
-	bOutIsValid = false;
 	checkf(!Name.IsEmpty(), TEXT("A build output requires a non-empty name."));
 	AssertValidBuildFunctionName(Function, Name);
-	Meta = InOutput["Meta"_ASV].AsObject();
-	Meta.MakeOwned();
-	for (FCbFieldView Payload : InOutput["Payloads"_ASV])
-	{
-		const FPayloadId Id = Payload["Id"_ASV].AsObjectId();
-		const FIoHash& RawHash = Payload["RawHash"_ASV].AsAttachment();
-		if (Id.IsNull() || RawHash.IsZero() || !Payload["RawSize"_ASV].IsInteger())
-		{
-			return;
-		}
-		Payloads.Emplace(Id, RawHash, Payload["RawSize"_ASV].AsUInt64());
-	}
-	Diagnostics = InOutput["Diagnostics"_ASV];
-	Diagnostics.MakeOwned();
-	bOutIsValid = IsValid() && (!InOutput.AsView()["Meta"_ASV] || InOutput.AsView()["Meta"_ASV].IsObject());
+	Output.MakeOwned();
+	FCbField MetaField = InOutput["Meta"_ASV];
+	bOutIsValid = TryLoad() && (!MetaField || MetaField.IsObject());
+	Meta = MoveTemp(MetaField).AsObject();
 }
 
 FBuildOutputInternal::FBuildOutputInternal(FStringView InName, FStringView InFunction, const FCacheRecord& InOutput, bool& bOutIsValid)
 	: Name(InName)
 	, Function(InFunction)
 	, Meta(InOutput.GetMeta())
+	, Output(InOutput.GetValue())
 	, Payloads(InOutput.GetAttachmentPayloads())
 {
 	checkf(!Name.IsEmpty(), TEXT("A build output requires a non-empty name."));
 	AssertValidBuildFunctionName(Function, Name);
-	if (FSharedBuffer Buffer = InOutput.GetValue())
-	{
-		Diagnostics = FCbObject(MoveTemp(Buffer))["Diagnostics"_ASV];
-	}
-	bOutIsValid = IsValid();
+	bOutIsValid = TryLoad();
 }
 
 const FPayload& FBuildOutputInternal::GetPayload(const FPayloadId& Id) const
@@ -181,31 +208,73 @@ const FPayload& FBuildOutputInternal::GetPayload(const FPayloadId& Id) const
 	return Payloads.IsValidIndex(Index) ? Payloads[Index] : FPayload::Null;
 }
 
-void FBuildOutputInternal::IterateDiagnostics(TFunctionRef<void (const FBuildDiagnostic& Diagnostic)> Visitor) const
+bool FBuildOutputInternal::HasError() const
 {
-	for (FCbFieldView Diagnostic : Diagnostics.CreateViewIterator())
-	{
-		EBuildDiagnosticLevel Level;
-		LexFromString(Level, Diagnostic["Level"_ASV].AsString());
-		const FUtf8StringView Category = Diagnostic["Category"_ASV].AsString();
-		const FUtf8StringView Message = Diagnostic["Message"_ASV].AsString();
-		Visitor(FBuildDiagnostic{FUTF8ToTCHAR(Category), FUTF8ToTCHAR(Message), Level});
-	}
+	return Algo::FindBy(Messages, EBuildOutputMessageLevel::Error, &FBuildOutputMessage::Level) ||
+		Algo::FindBy(Logs, EBuildOutputLogLevel::Error, &FBuildOutputLog::Level);
 }
 
-bool FBuildOutputInternal::HasError() const
+bool FBuildOutputInternal::TryLoad()
 {
-	for (FCbFieldView Field : Diagnostics.CreateViewIterator())
+	const FCbObjectView OutputView = Output;
+
+	if (Payloads.IsEmpty())
 	{
-		const FCbObjectView Diagnostic = Field.AsObjectView();
-		EBuildDiagnosticLevel Level;
-		LexFromString(Level, Diagnostic["Level"_ASV].AsString());
-		if (Level == EBuildDiagnosticLevel::Error)
+		for (FCbFieldView Payload : OutputView["Payloads"_ASV])
 		{
-			return true;
+			const FPayloadId Id = Payload["Id"_ASV].AsObjectId();
+			const FIoHash& RawHash = Payload["RawHash"_ASV].AsAttachment();
+			const uint64 RawSize = Payload["RawSize"_ASV].AsUInt64(MAX_uint64);
+			if (Id.IsNull() || RawHash.IsZero() || RawSize == MAX_uint64)
+			{
+				return false;
+			}
+			Payloads.Emplace(Id, RawHash, RawSize);
 		}
 	}
-	return false;
+
+	if (FCbFieldView MessagesField = OutputView["Messages"_ASV])
+	{
+		if (!MessagesField.IsArray())
+		{
+			return false;
+		}
+		Messages.Reserve(MessagesField.AsArrayView().Num());
+		for (FCbFieldView MessageField : MessagesField)
+		{
+			const FUtf8StringView LevelName = MessageField["Level"_ASV].AsString();
+			const FUtf8StringView Message = MessageField["Message"_ASV].AsString();
+			EBuildOutputMessageLevel Level;
+			if (LevelName.IsEmpty() || Message.IsEmpty() || !TryLexFromString(Level, LevelName))
+			{
+				return false;
+			}
+			Messages.Add({Message, Level});
+		}
+	}
+
+	if (FCbFieldView LogsField = OutputView["Logs"_ASV])
+	{
+		if (!LogsField.IsArray())
+		{
+			return false;
+		}
+		Logs.Reserve(LogsField.AsArrayView().Num());
+		for (FCbFieldView LogField : LogsField)
+		{
+			const FUtf8StringView LevelName = LogField["Level"_ASV].AsString();
+			const FUtf8StringView Category = LogField["Category"_ASV].AsString();
+			const FUtf8StringView Message = LogField["Message"_ASV].AsString();
+			EBuildOutputLogLevel Level;
+			if (LevelName.IsEmpty() || Category.IsEmpty() || Message.IsEmpty() || !TryLexFromString(Level, LevelName))
+			{
+				return false;
+			}
+			Logs.Add({Category, Message, Level});
+		}
+	}
+
+	return true;
 }
 
 void FBuildOutputInternal::Save(FCbWriter& Writer) const
@@ -224,9 +293,13 @@ void FBuildOutputInternal::Save(FCbWriter& Writer) const
 		}
 		Writer.EndArray();
 	}
-	if (Diagnostics)
+	if (FCbField MessagesField = Output["Messages"_ASV])
+	{
+		Writer.AddField("Messages"_ASV, MessagesField);
+	}
+	if (FCbField LogsField = Output["Logs"_ASV])
 	{
-		Writer.AddField("Diagnostics"_ASV, Diagnostics);
+		Writer.AddField("Logs"_ASV, LogsField);
 	}
 	if (Meta)
 	{
@@ -238,13 +311,18 @@ void FBuildOutputInternal::Save(FCbWriter& Writer) const
 void FBuildOutputInternal::Save(FCacheRecordBuilder& RecordBuilder) const
 {
 	RecordBuilder.SetMeta(FCbObject(Meta));
-	if (Diagnostics)
+	if (!Messages.IsEmpty() || !Logs.IsEmpty())
 	{
-		TCbWriter<128> Value;
-		Value.BeginObject();
-		Value.AddField("Diagnostics"_ASV, Diagnostics);
-		Value.EndObject();
-		RecordBuilder.SetValue(Value.Save().GetBuffer());
+		TCbWriter<1024> Writer;
+		if (FCbField MessagesField = Output["Messages"_ASV])
+		{
+			Writer.AddField("Messages"_ASV, MessagesField);
+		}
+		if (FCbField LogsField = Output["Logs"_ASV])
+		{
+			Writer.AddField("Logs"_ASV, LogsField);
+		}
+		RecordBuilder.SetValue(Writer.Save().GetBuffer());
 	}
 	for (const FPayload& Payload : Payloads)
 	{
@@ -252,20 +330,6 @@ void FBuildOutputInternal::Save(FCacheRecordBuilder& RecordBuilder) const
 	}
 }
 
-bool FBuildOutputInternal::IsValid() const
-{
-	return IsValidBuildFunctionName(Function)
-		&& Algo::AllOf(Payloads, [](const FPayload& Payload) { return !Payload.GetRawHash().IsZero(); })
-		&& (!Diagnostics || Diagnostics.IsArray())
-		&& Algo::AllOf(Diagnostics.CreateViewIterator(), [](FCbFieldView Field)
-			{
-				return Field.IsObject()
-					&& Field["Level"_ASV].AsString().Len() > 0
-					&& Field["Category"_ASV].AsString().Len() > 0
-					&& Field["Message"_ASV].AsString().Len() > 0;
-			});
-}
-
 ///////////////////////////////////////////////////////////////////////////////////////////////////////////////////////
 
 void FBuildOutputBuilderInternal::AddPayload(const FPayload& Payload)
@@ -278,26 +342,54 @@ void FBuildOutputBuilderInternal::AddPayload(const FPayload& Payload)
 	Payloads.Insert(Payload, Index);
 }
 
-void FBuildOutputBuilderInternal::AddDiagnostic(const FBuildDiagnostic& Diagnostic)
+void FBuildOutputBuilderInternal::AddMessage(const FBuildOutputMessage& Message)
+{
+	bHasError |= Message.Level == EBuildOutputMessageLevel::Error;
+	bHasMessages = true;
+	MessageWriter.BeginObject();
+	MessageWriter.AddString("Level"_ASV, LexToString(Message.Level));
+	MessageWriter.AddString("Message"_ASV, Message.Message);
+	MessageWriter.EndObject();
+}
+
+void FBuildOutputBuilderInternal::AddLog(const FBuildOutputLog& Log)
 {
-	bHasError |= Diagnostic.Level == EBuildDiagnosticLevel::Error;
-	bHasDiagnostics = true;
-	DiagnosticsWriter.BeginObject();
-	DiagnosticsWriter.AddString("Level"_ASV, LexToString(Diagnostic.Level));
-	DiagnosticsWriter.AddString("Category"_ASV, Diagnostic.Category);
-	DiagnosticsWriter.AddString("Message"_ASV, Diagnostic.Message);
-	DiagnosticsWriter.EndObject();
+	bHasError |= Log.Level == EBuildOutputLogLevel::Error;
+	bHasLogs = true;
+	LogWriter.BeginObject();
+	LogWriter.AddString("Level"_ASV, LexToString(Log.Level));
+	LogWriter.AddString("Category"_ASV, Log.Category);
+	LogWriter.AddString("Message"_ASV, Log.Message);
+	LogWriter.EndObject();
 }
 
 FBuildOutput FBuildOutputBuilderInternal::Build()
 {
-	DiagnosticsWriter.EndArray();
-	if (bHasDiagnostics)
+	if (bHasError)
+	{
+		Payloads.Empty();
+	}
+	MessageWriter.EndArray();
+	LogWriter.EndArray();
+	FCbObject Output;
+	if (bHasMessages || bHasLogs)
 	{
-		Diagnostics = DiagnosticsWriter.Save();
+		TCbWriter<1024> Writer;
+		Writer.BeginObject();
+		if (bHasMessages)
+		{
+			Writer.AddArray("Messages"_ASV, MessageWriter.Save().AsArray());
+			MessageWriter.Reset();
+		}
+		if (bHasLogs)
+		{
+			Writer.AddArray("Logs"_ASV, LogWriter.Save().AsArray());
+			LogWriter.Reset();
+		}
+		Writer.EndObject();
+		Output = Writer.Save().AsObject();
 	}
-	DiagnosticsWriter.Reset();
-	return CreateBuildOutput(new FBuildOutputInternal(MoveTemp(*this)));
+	return CreateBuildOutput(new FBuildOutputInternal(MoveTemp(Name), MoveTemp(Function), MoveTemp(Meta), MoveTemp(Output), MoveTemp(Payloads)));
 }
 
 ///////////////////////////////////////////////////////////////////////////////////////////////////////////////////////
diff --git a/Engine/Source/Developer/DerivedDataCache/Public/DerivedDataBuildFunction.h b/Engine/Source/Developer/DerivedDataCache/Public/DerivedDataBuildFunction.h
--- a/Engine/Source/Developer/DerivedDataCache/Public/DerivedDataBuildFunction.h
+++ b/Engine/Source/Developer/DerivedDataCache/Public/DerivedDataBuildFunction.h
@@ -97,6 +97,15 @@ public:
 	virtual void AddPayload(const FPayloadId& Id, const FSharedBuffer& Buffer) = 0;
 	virtual void AddPayload(const FPayloadId& Id, const FCbObject& Object) = 0;
 
+	/** Adds the error to the build output. Must be deterministic and based only on inputs. */
+	virtual void AddError(FStringView Message) = 0;
+
+	/** Adds the warning to the build output. Must be deterministic and based only on inputs. */
+	virtual void AddWarning(FStringView Message) = 0;
+
+	/** Adds the message to the build output. Must be deterministic and based only on inputs. */
+	virtual void AddMessage(FStringView Message) = 0;
+
 	/** Overrides the cache policy mask used when writing this build in the cache. */
 	virtual void SetCachePolicyMask(ECachePolicy Policy) = 0;
 
@@ -134,6 +143,15 @@ public:
 	/** Returns the constant with the matching key, or an object with no fields if not found. */
 	virtual FCbObject FindConstant(FStringView Key) const = 0;
 
+	/** Adds the error to the build output. Must be deterministic and based only on inputs. */
+	virtual void AddError(FStringView Message) = 0;
+
+	/** Adds the warning to the build output. Must be deterministic and based only on inputs. */
+	virtual void AddWarning(FStringView Message) = 0;
+
+	/** Adds the message to the build output. Must be deterministic and based only on inputs. */
+	virtual void AddMessage(FStringView Message) = 0;
+
 	/** Overrides the cache bucket used when reading or writing this build in the cache. */
 	virtual void SetCacheBucket(FCacheBucket Bucket) = 0;
 
diff --git a/Engine/Source/Developer/DerivedDataCache/Public/DerivedDataBuildOutput.h b/Engine/Source/Developer/DerivedDataCache/Public/DerivedDataBuildOutput.h
--- a/Engine/Source/Developer/DerivedDataCache/Public/DerivedDataBuildOutput.h
+++ b/Engine/Source/Developer/DerivedDataCache/Public/DerivedDataBuildOutput.h
@@ -6,7 +6,6 @@
 #include "Containers/ArrayView.h"
 #include "Containers/StringView.h"
 #include "Misc/ScopeExit.h"
-#include "Templates/Function.h"
 #include "Templates/RefCounting.h"
 #include "Templates/UniquePtr.h"
 
@@ -21,7 +20,8 @@ namespace UE::DerivedData { class FCacheRecord; }
 namespace UE::DerivedData { class FCacheRecordBuilder; }
 namespace UE::DerivedData { class FOptionalBuildOutput; }
 namespace UE::DerivedData { class FPayload; }
-namespace UE::DerivedData { struct FBuildDiagnostic; }
+namespace UE::DerivedData { struct FBuildOutputLog; }
+namespace UE::DerivedData { struct FBuildOutputMessage; }
 namespace UE::DerivedData { struct FPayloadId; }
 
 namespace UE::DerivedData::Private
@@ -36,7 +36,9 @@ public:
 	virtual const FCbObject& GetMeta() const = 0;
 	virtual const FPayload& GetPayload(const FPayloadId& Id) const = 0;
 	virtual TConstArrayView<FPayload> GetPayloads() const = 0;
-	virtual void IterateDiagnostics(TFunctionRef<void (const FBuildDiagnostic& Diagnostic)> Visitor) const = 0;
+	virtual TConstArrayView<FBuildOutputMessage> GetMessages() const = 0;
+	virtual TConstArrayView<FBuildOutputLog> GetLogs() const = 0;
+	virtual bool HasLogs() const = 0;
 	virtual bool HasError() const = 0;
 	virtual void Save(FCbWriter& Writer) const = 0;
 	virtual void Save(FCacheRecordBuilder& RecordBuilder) const = 0;
@@ -52,7 +54,8 @@ public:
 	virtual ~IBuildOutputBuilderInternal() = default;
 	virtual void SetMeta(FCbObject&& Meta) = 0;
 	virtual void AddPayload(const FPayload& Payload) = 0;
-	virtual void AddDiagnostic(const FBuildDiagnostic& Diagnostic) = 0;
+	virtual void AddMessage(const FBuildOutputMessage& Message) = 0;
+	virtual void AddLog(const FBuildOutputLog& Log) = 0;
 	virtual bool HasError() const = 0;
 	virtual FBuildOutput Build() = 0;
 };
@@ -64,28 +67,43 @@ FBuildOutputBuilder CreateBuildOutputBuilder(IBuildOutputBuilderInternal* Output
 namespace UE::DerivedData
 {
 
-/** Level of severity for build diagnostics. */
-enum class EBuildDiagnosticLevel : uint8
+enum class EBuildOutputMessageLevel : uint8
 {
-	/** Errors always indicates a failure of the corresponding build. */
 	Error,
-	/** Warnings are expected to be actionable issues found while executing a build. */
 	Warning,
+	Display,
 };
 
-/** A build diagnostic is a message logged by a build. */
-struct FBuildDiagnostic
+/** A build output message is diagnostic output from a build function and must be deterministic. */
+struct FBuildOutputMessage
 {
-	/** The name of the category of the diagnostic. */
-	FStringView Category;
-	/** The message of the diagnostic. */
-	FStringView Message;
-	/** The level of severity of the diagnostic. */
-	EBuildDiagnosticLevel Level;
+	FUtf8StringView Message;
+	EBuildOutputMessageLevel Level;
+};
+
+enum class EBuildOutputLogLevel : uint8
+{
+	Error,
+	Warning,
 };
 
 /**
- * A build output is an immutable container of payloads and diagnostics produced by a build.
+ * A build output log is a log message captured from a build function.
+ *
+ * The build function may capture every log above a certain level of verbosity, which means these
+ * have no guarantee of being deterministic. The presence of any output logs will disable caching
+ * of the build output. To allow caching of build output with warnings or errors, replace the log
+ * statements with build messages that are added to the build context.
+ */
+struct FBuildOutputLog
+{
+	FUtf8StringView Category;
+	FUtf8StringView Message;
+	EBuildOutputLogLevel Level;
+};
+
+/**
+ * A build output is an immutable container of payloads, messages, and logs produced by a build.
  *
  * The output will not contain any payloads if it has any errors.
  *
@@ -109,13 +127,16 @@ public:
 	/** Returns the payloads in the output in order by ID. */
 	inline TConstArrayView<FPayload> GetPayloads() const { return Output->GetPayloads(); }
 
-	/** Visits every diagnostic in the order it was recorded. */
-	inline void IterateDiagnostics(TFunctionRef<void (const FBuildDiagnostic& Diagnostic)> Visitor) const
-	{
-		Output->IterateDiagnostics(Visitor);
-	}
+	/** Returns the messages in the order that they were recorded. */
+	inline TConstArrayView<FBuildOutputMessage> GetMessages() const { return Output->GetMessages(); }
+
+	/** Returns the logs in the order that they were recorded. */
+	inline TConstArrayView<FBuildOutputLog> GetLogs() const { return Output->GetLogs(); }
+
+	/** Returns whether the output has any logs. */
+	inline bool HasLogs() const { return Output->HasLogs(); }
 
-	/** Returns whether the output has any error diagnostics. */
+	/** Returns whether the output has any errors. */
 	inline bool HasError() const { return Output->HasError(); }
 
 	/** Saves the build output to a compact binary object with payloads as attachments. */
@@ -175,19 +196,19 @@ public:
 		OutputBuilder->AddPayload(Payload);
 	}
 
-	/** Add an error diagnostic to the output. */
-	inline void AddError(FStringView Category, FStringView Message)
+	/** Add a message to the output. */
+	inline void AddMessage(const FBuildOutputMessage& Message)
 	{
-		OutputBuilder->AddDiagnostic({Category, Message, EBuildDiagnosticLevel::Error});
+		OutputBuilder->AddMessage(Message);
 	}
 
-	/** Add a warning diagnostic to the output. */
-	inline void AddWarning(FStringView Category, FStringView Message)
+	/** Add a log to the output. */
+	inline void AddLog(const FBuildOutputLog& Log)
 	{
-		OutputBuilder->AddDiagnostic({Category, Message, EBuildDiagnosticLevel::Warning});
+		OutputBuilder->AddLog(Log);
 	}
 
-	/** Returns whether the output has any error diagnostics. */
+	/** Returns whether the output has any errors. */
 	inline bool HasError() const
 	{
 		return OutputBuilder->HasError();
diff --git a/Engine/Source/Editor/UnrealEd/Private/DerivedDataBuildRemoteExecutor.cpp b/Engine/Source/Editor/UnrealEd/Private/DerivedDataBuildRemoteExecutor.cpp
--- a/Engine/Source/Editor/UnrealEd/Private/DerivedDataBuildRemoteExecutor.cpp
+++ b/Engine/Source/Editor/UnrealEd/Private/DerivedDataBuildRemoteExecutor.cpp
@@ -992,17 +992,15 @@ void FRemoteBuildExecutionRequest::OnGetResultPackageComplete(const UE::Zen::FZe
 
 	FBuildOutputBuilder OutputBuilder = State.BuildSystem.CreateOutput(State.BuildAction.GetName(), State.BuildAction.GetFunction());
 
-	RemoteBuildOutput.Get().IterateDiagnostics([&OutputBuilder](const FBuildDiagnostic& Diagnostic)
-		{
-			if (Diagnostic.Level == EBuildDiagnosticLevel::Warning)
-			{
-				OutputBuilder.AddWarning(Diagnostic.Category, Diagnostic.Message);
-			}
-			else if (Diagnostic.Level == EBuildDiagnosticLevel::Error)
-			{
-				OutputBuilder.AddError(Diagnostic.Category, Diagnostic.Message);
-			}
-		});
+	for (const FBuildOutputMessage& Message : RemoteBuildOutput.Get().GetMessages())
+	{
+		OutputBuilder.AddMessage(Message);
+	}
+
+	for (const FBuildOutputLog& Log : RemoteBuildOutput.Get().GetLogs())
+	{
+		OutputBuilder.AddLog(Log);
+	}
 
 	for (const FPayload& Payload : RemoteBuildOutput.Get().GetPayloads())
 	{
diff --git a/Engine/Source/Programs/DerivedDataBuildWorker/Private/DerivedDataBuildWorker.cpp b/Engine/Source/Programs/DerivedDataBuildWorker/Private/DerivedDataBuildWorker.cpp
--- a/Engine/Source/Programs/DerivedDataBuildWorker/Private/DerivedDataBuildWorker.cpp
+++ b/Engine/Source/Programs/DerivedDataBuildWorker/Private/DerivedDataBuildWorker.cpp
@@ -251,42 +251,69 @@ bool FBuildWorkerProgram::Build()
 void FBuildWorkerProgram::BuildComplete(FBuildCompleteParams&& Params) const
 {
 	const FBuildOutput Output = MoveTemp(Params.Output);
+	const FStringView Function = Output.GetFunction();
+	const FStringView Name = Output.GetName();
+
+	for (const FBuildOutputMessage& Message : Output.GetMessages())
+	{
+		switch (Message.Level)
+		{
+		case EBuildOutputMessageLevel::Error:
+			UE_LOG(LogDerivedDataBuildWorker, Error, TEXT("%s (Build of '%.*s' by %.*s.)"),
+				*WriteToString<256>(Message.Message), Name.Len(), Name.GetData(), Function.Len(), Function.GetData());
+			break;
+		case EBuildOutputMessageLevel::Warning:
+			UE_LOG(LogDerivedDataBuildWorker, Warning, TEXT("%s (Build of '%.*s' by %.*s.)"),
+				*WriteToString<256>(Message.Message), Name.Len(), Name.GetData(), Function.Len(), Function.GetData());
+			break;
+		case EBuildOutputMessageLevel::Display:
+			UE_LOG(LogDerivedDataBuildWorker, Display, TEXT("%s (Build of '%.*s' by %.*s.)"),
+				*WriteToString<256>(Message.Message), Name.Len(), Name.GetData(), Function.Len(), Function.GetData());
+			break;
+		default:
+			checkNoEntry();
+			break;
+		}
+	}
 
 	if constexpr (!NO_LOGGING)
 	{
 		if (GWarn)
 		{
-			Output.IterateDiagnostics([](const FBuildDiagnostic& Diagnostic)
+			for (const FBuildOutputLog& Log : Output.GetLogs())
 			{
 				ELogVerbosity::Type Verbosity;
-				switch (Diagnostic.Level)
+				switch (Log.Level)
 				{
 				default:
-				case EBuildDiagnosticLevel::Error:   Verbosity = ELogVerbosity::Error;   break;
-				case EBuildDiagnosticLevel::Warning: Verbosity = ELogVerbosity::Warning; break;
+				case EBuildOutputLogLevel::Error:
+					Verbosity = ELogVerbosity::Error;
+					break;
+				case EBuildOutputLogLevel::Warning:
+					Verbosity = ELogVerbosity::Warning;
+					break;
 				}
-				GWarn->Log(FName(Diagnostic.Category), Verbosity, FString(Diagnostic.Message));
-			});
+				GWarn->Log(FName(Log.Category), Verbosity, FString::Printf(TEXT("%s (Build of '%.*s' by %.*s.)"),
+					*WriteToString<256>(Log.Message), Name.Len(), Name.GetData(), Function.Len(), Function.GetData()));
+			}
 		}
 	}
 
-	if (Output.HasError())
+	for (const FPayload& Payload : Output.GetPayloads())
 	{
-		UE_LOG(LogDerivedDataBuildWorker, Error,
-			TEXT("Errors in build caused storage of payloads to be skipped for build of '%s' by %s."),
-			*WriteToString<128>(Output.GetName()), *WriteToString<32>(Output.GetFunction()));
-	}
-	else
-	{
-		for (const FPayload& Payload : Output.GetPayloads())
+		if (Payload.HasData())
 		{
-			if (Payload.HasData())
+			if (TUniquePtr<FArchive> Ar = OpenOutput(Output.GetName(), Payload.GetRawHash()))
 			{
-				if (TUniquePtr<FArchive> Ar = OpenOutput(Output.GetName(), Payload.GetRawHash()))
+				*Ar << const_cast<FCompressedBuffer&>(Payload.GetData());
+				if (Ar->Close())
 				{
-					*Ar << const_cast<FCompressedBuffer&>(Payload.GetData());
+					continue;
 				}
 			}
+			UE_LOG(LogDerivedDataBuildWorker, Error,
+				TEXT("Failed to store build output %s for build of '%.*s' by %.*s."),
+				*WriteToString<48>(Payload.GetRawHash()), Name.Len(), Name.GetData(), Function.Len(), Function.GetData());
 		}
 	}
 
@@ -300,8 +327,8 @@ void FBuildWorkerProgram::BuildComplete(FBuildCompleteParams&& Params) const
 	else
 	{
 		UE_LOG(LogDerivedDataBuildWorker, Error,
-			TEXT("Failed to store build output to '%s' for build of '%s' by %s."),
-			*OutputPath, *WriteToString<128>(Output.GetName()), *WriteToString<32>(Output.GetFunction()));
+			TEXT("Failed to store build output to '%s' for build of '%.*s' by %.*s."),
+			*OutputPath, Name.Len(), Name.GetData(), Function.Len(), Function.GetData());
 	}
 }
 
diff --git a/Engine/Source/Runtime/Engine/Private/TextureDerivedDataTask.cpp b/Engine/Source/Runtime/Engine/Private/TextureDerivedDataTask.cpp
--- a/Engine/Source/Runtime/Engine/Private/TextureDerivedDataTask.cpp
+++ b/Engine/Source/Runtime/Engine/Private/TextureDerivedDataTask.cpp
@@ -1421,26 +1421,55 @@ private:
 	{
 		using namespace UE::DerivedData;
 
-		Output.IterateDiagnostics([](const FBuildDiagnostic& Diagnostic)
+		const FStringView Function = Output.GetFunction();
+		const FStringView Name = Output.GetName();
+
+		for (const FBuildOutputMessage& Message : Output.GetMessages())
 		{
-			if (Diagnostic.Level == EBuildDiagnosticLevel::Error)
+			switch (Message.Level)
 			{
-				UE_LOG(LogTexture, Warning, TEXT("[Build Error] %.*s: %.*s"),
-					Diagnostic.Category.Len(), Diagnostic.Category.GetData(),
-					Diagnostic.Message.Len(), Diagnostic.Message.GetData());
+			case EBuildOutputMessageLevel::Error:
+				UE_LOG(LogTexture, Warning, TEXT("[Error] %s (Build of '%.*s' by %.*s.)"),
+					*WriteToString<256>(Message.Message), Name.Len(), Name.GetData(), Function.Len(), Function.GetData());
+				break;
+			case EBuildOutputMessageLevel::Warning:
+				UE_LOG(LogTexture, Warning, TEXT("%s (Build of '%.*s' by %.*s.)"),
+					*WriteToString<256>(Message.Message), Name.Len(), Name.GetData(), Function.Len(), Function.GetData());
+				break;
+			case EBuildOutputMessageLevel::Display:
+				UE_LOG(LogTexture, Display, TEXT("%s (Build of '%.*s' by %.*s.)"),
+					*WriteToString<256>(Message.Message), Name.Len(), Name.GetData(), Function.Len(), Function.GetData());
+				break;
+			default:
+				checkNoEntry();
+				break;
 			}
-			else
+		}
+
+		for (const FBuildOutputLog& Log : Output.GetLogs())
+		{
+			switch (Log.Level)
 			{
-				UE_LOG(LogTexture, Warning, TEXT("[Build Warning] %.*s: %.*s"),
-					Diagnostic.Category.Len(), Diagnostic.Category.GetData(),
-					Diagnostic.Message.Len(), Diagnostic.Message.GetData());
+			case EBuildOutputLogLevel::Error:
+				UE_LOG(LogTexture, Warning, TEXT("[Error] %s: %s (Build of '%.*s' by %.*s.)"),
+					*WriteToString<64>(Log.Category), *WriteToString<256>(Log.Message),
+					Name.Len(), Name.GetData(), Function.Len(), Function.GetData());
+				break;
+			case EBuildOutputLogLevel::Warning:
+				UE_LOG(LogTexture, Warning, TEXT("%s: %s (Build of '%.*s' by %.*s.)"),
+					*WriteToString<64>(Log.Category), *WriteToString<256>(Log.Message),
+					Name.Len(), Name.GetData(), Function.Len(), Function.GetData());
+				break;
+			default:
+				checkNoEntry();
+				break;
 			}
-		});
+		}
 
 		if (Output.HasError())
 		{
-			UE_LOG(LogTexture, Warning, TEXT("Failed to build derived data for build of '%s' by %s."),
-				*WriteToString<128>(Output.GetName()), *WriteToString<32>(Output.GetFunction()));
+			UE_LOG(LogTexture, Warning, TEXT("Failed to build derived data for build of '%.*s' by %.*s."),
+				Name.Len(), Name.GetData(), Function.Len(), Function.GetData());
 			return;
 		}
 
