DDC: Optimized GetChunks in the file system cache store to read only the requested range from disk

Contains changes to other structured cache file access to stream data from archives instead of loading the whole file into memory. The record package will now be split into separate allocations on load, rather than the previous behavior where the whole package was kept in memory as long as any part of it was referenced.

#rnx

diff --git a/Engine/Source/Developer/DerivedDataCache/Private/FileSystemDerivedDataBackend.cpp b/Engine/Source/Developer/DerivedDataCache/Private/FileSystemDerivedDataBackend.cpp
--- a/Engine/Source/Developer/DerivedDataCache/Private/FileSystemDerivedDataBackend.cpp
+++ b/Engine/Source/Developer/DerivedDataCache/Private/FileSystemDerivedDataBackend.cpp
@@ -781,20 +781,27 @@ private:
 		const FStringView Context,
 		const FCacheRecordPolicy& Policy,
 		EStatus& OutStatus);
+	FPayload GetCacheContent(
+		const FCacheKey& Key,
+		const FPayload& Payload,
+		FStringView Context,
+		ECachePolicy Policy,
+		ECachePolicy SkipFlag) const;
 	void GetCacheContent(
 		const FCacheKey& Key,
-		const FStringView Context,
-		const ECachePolicy Policy,
-		const ECachePolicy SkipFlag,
-		FPayload& InOutPayload,
-		EStatus& InOutStatus) const;
+		const FPayload& Payload,
+		FStringView Context,
+		ECachePolicy Policy,
+		ECachePolicy SkipFlag,
+		FCompressedBufferReader& Reader,
+		TUniquePtr<FArchive>& OutArchive) const;
 
 	void BuildCacheRecordPath(const FCacheKey& CacheKey, FStringBuilderBase& Path) const;
 	void BuildCacheContentPath(const FIoHash& RawHash, FStringBuilderBase& Path) const;
 
 	bool SaveFile(FStringBuilderBase& Path, FStringView Context, TFunctionRef<void (FArchive&)> WriteFunction) const;
-	bool SaveDataToFile(FStringBuilderBase& Path, TFunctionRef<void (FArchive&)> WriteFunction, int64& OutWriteSize) const;
-	FSharedBuffer LoadFile(FStringBuilderBase& Path, FStringView Context) const;
+	bool LoadFile(FStringBuilderBase& Path, FStringView Context, TFunctionRef<void (FArchive&)> ReadFunction) const;
+	TUniquePtr<FArchive> OpenFile(FStringBuilderBase& Path, FStringView Context) const;
 
 	bool FileExists(FStringBuilderBase& Path) const;
 
@@ -1637,6 +1644,11 @@ void FFileSystemCacheStore::GetChunks(
 	TArray<FCacheChunkRequest, TInlineAllocator<16>> SortedChunks(Chunks);
 	SortedChunks.StableSort(TChunkLess());
 
+	FPayload Payload;
+	FCacheKey PayloadKey;
+	TUniquePtr<FArchive> PayloadAr;
+	FCompressedBufferReader PayloadReader;
+	EStatus PayloadStatus = EStatus::Error;
 	FOptionalCacheRecord Record;
 	for (const FCacheChunkRequest& Chunk : SortedChunks)
 	{
@@ -1645,39 +1657,48 @@ void FFileSystemCacheStore::GetChunks(
 		TRACE_CPUPROFILER_EVENT_SCOPE(FileSystemDDC_Get);
 		TRACE_COUNTER_INCREMENT(FileSystemDDC_Get);
 		COOK_STAT(auto Timer = bExistsOnly ? UsageStats.TimeProbablyExists() : UsageStats.TimeGet());
-		if (!Record || Record.Get().GetKey() != Chunk.Key)
+		if (!(Payload && PayloadKey == Chunk.Key && Payload.GetId() == Chunk.Id) || !bExistsOnly <= PayloadReader.HasSource())
 		{
-			FCacheRecordPolicyBuilder PolicyBuilder(ECachePolicy::None);
-			const ECachePolicy RecordSkipFlags = bExistsOnly ? ECachePolicy::SkipData : ECachePolicy::None;
-			PolicyBuilder.AddPayloadPolicy(Chunk.Id, Chunk.Policy | RecordSkipFlags);
-			Record = GetCacheRecordOnly(Chunk.Key, Context, PolicyBuilder.Build());
+			PayloadStatus = EStatus::Error;
+			PayloadReader.ResetSource();
+			PayloadAr.Reset();
+			Payload.Reset();
+			if (!(Record && Record.Get().GetKey() == Chunk.Key))
+			{
+				FCacheRecordPolicyBuilder PolicyBuilder(ECachePolicy::None);
+				const ECachePolicy RecordSkipFlags = bExistsOnly ? ECachePolicy::SkipData : ECachePolicy::None;
+				PolicyBuilder.AddPayloadPolicy(Chunk.Id, Chunk.Policy | RecordSkipFlags);
+				Record.Reset();
+				Record = GetCacheRecordOnly(Chunk.Key, Context, PolicyBuilder.Build());
+			}
+			if (Record)
+			{
+				Payload = Record.Get().GetPayload(Chunk.Id);
+				PayloadKey = Chunk.Key;
+				GetCacheContent(Chunk.Key, Payload, Context, Chunk.Policy, SkipFlag, PayloadReader, PayloadAr);
+			}
 		}
-		if (Record)
+		if (Payload)
 		{
-			EStatus PayloadStatus = EStatus::Ok;
-			FPayload Payload = Record.Get().GetPayload(Chunk.Id);
-			GetCacheContent(Chunk.Key, Context, Chunk.Policy, SkipFlag, Payload, PayloadStatus);
-			if (Payload)
+			const uint64 RawOffset = FMath::Min(Payload.GetRawSize(), Chunk.RawOffset);
+			const uint64 RawSize = FMath::Min(Payload.GetRawSize() - RawOffset, Chunk.RawSize);
+			UE_LOG(LogDerivedDataCache, Verbose, TEXT("%s: Cache hit for %s from '%.*s'"),
+				*CachePath, *WriteToString<96>(Chunk.Key, '/', Chunk.Id), Context.Len(), Context.GetData());
+			TRACE_COUNTER_INCREMENT(FileSystemDDC_GetHit);
+			TRACE_COUNTER_ADD(FileSystemDDC_BytesRead, !bExistsOnly ? RawSize : 0);
+			COOK_STAT(Timer.AddHit(!bExistsOnly ? RawSize : 0));
+			if (OnComplete)
 			{
-				const uint64 RawOffset = FMath::Min(Payload.GetRawSize(), Chunk.RawOffset);
-				const uint64 RawSize = FMath::Min(Payload.GetRawSize() - RawOffset, Chunk.RawSize);
-				UE_LOG(LogDerivedDataCache, Verbose, TEXT("%s: Cache hit for %s from '%.*s'"),
-					*CachePath, *WriteToString<96>(Chunk.Key, '/', Chunk.Id), Context.Len(), Context.GetData());
-				TRACE_COUNTER_INCREMENT(FileSystemDDC_GetHit);
-				TRACE_COUNTER_ADD(FileSystemDDC_BytesRead, Payload.GetData().GetCompressedSize());
-				COOK_STAT(Timer.AddHit(Payload.HasData() ? RawSize : 0));
-				if (OnComplete)
+				FSharedBuffer Buffer;
+				if (!bExistsOnly)
 				{
-					FSharedBuffer Buffer;
-					if (Payload.HasData() && !bExistsOnly)
-					{
-						Buffer = FCompressedBufferReader(Payload.GetData()).Decompress(RawOffset, RawSize);
-					}
-					OnComplete({Chunk.Key, Chunk.Id, Chunk.RawOffset,
-						RawSize, Payload.GetRawHash(), MoveTemp(Buffer), PayloadStatus});
+					Buffer = PayloadReader.Decompress(RawOffset, RawSize);
 				}
-				continue;
+				const EStatus ChunkStatus = bExistsOnly || Buffer.GetSize() == RawSize ? EStatus::Ok : EStatus::Error;
+				OnComplete({Chunk.Key, Chunk.Id, Chunk.RawOffset,
+					RawSize, Payload.GetRawHash(), MoveTemp(Buffer), ChunkStatus});
 			}
+			continue;
 		}
 
 		if (OnComplete)
@@ -1742,10 +1763,9 @@ bool FFileSystemCacheStore::PutCacheRecord(
 	{
 		bRecordExists = FileExists(Path);
 	}
-	else if (FSharedBuffer Buffer = LoadFile(Path, Context))
+	else
 	{
-		FCbFieldIterator It = FCbFieldIterator::MakeRange(Buffer);
-		bRecordExists = ExistingPackage.TryLoad(It);
+		bRecordExists = LoadFile(Path, Context, [&ExistingPackage](FArchive& Ar) { ExistingPackage.TryLoad(Ar); });
 	}
 
 	// Save the record to a package and remove attachments that will be stored externally.
@@ -1840,63 +1860,51 @@ FOptionalCacheRecord FFileSystemCacheStore::GetCacheRecordOnly(
 		return FOptionalCacheRecord();
 	}
 
-	// Request the record from storage.
 	TStringBuilder<256> Path;
 	BuildCacheRecordPath(Key, Path);
-	FSharedBuffer Buffer = LoadFile(Path, Context);
-	if (Buffer.IsNull())
-	{
-		UE_LOG(LogDerivedDataCache, Verbose, TEXT("%s: Cache miss with missing record for %s from '%.*s'"),
-			*CachePath, *WriteToString<96>(Key), Context.Len(), Context.GetData());
-		return FOptionalCacheRecord();
-	}
 
-	// Delete the record from storage if it is invalid.
-	bool bDeleteCacheObject = true;
+	bool bDeleteRecord = true;
 	ON_SCOPE_EXIT
 	{
-		if (bDeleteCacheObject && !bReadOnly)
+		if (bDeleteRecord && !bReadOnly)
 		{
 			IFileManager::Get().Delete(*Path, /*bRequireExists*/ false, /*bEvenReadOnly*/ false, /*bQuiet*/ true);
 		}
 	};
 
-	// Validate that the record can be read as a compact binary package without crashing.
-	if (ValidateCompactBinaryPackage(Buffer, ECbValidateMode::Default | ECbValidateMode::Package) != ECbValidateError::None)
-	{
-		UE_LOG(LogDerivedDataCache, Display, TEXT("%s: Cache miss with invalid package for %s from '%.*s'"),
-			*CachePath, *WriteToString<96>(Key), Context.Len(), Context.GetData());
-		return FOptionalCacheRecord();
-	}
-
-	// Load the record from the package.
 	FOptionalCacheRecord Record;
 	{
 		FCbPackage Package;
-		if (FCbFieldIterator It = FCbFieldIterator::MakeRange(Buffer); !Package.TryLoad(It))
+		if (!LoadFile(Path, Context, [&Package](FArchive& Ar) { Package.TryLoad(Ar); }))
 		{
-			UE_LOG(LogDerivedDataCache, Display, TEXT("%s: Cache miss with package load failure for %s from '%.*s'"),
+			UE_LOG(LogDerivedDataCache, Verbose, TEXT("%s: Cache miss with missing record for %s from '%.*s'"),
+				*CachePath, *WriteToString<96>(Key), Context.Len(), Context.GetData());
+			return Record;
+		}
+
+		if (ValidateCompactBinary(Package, ECbValidateMode::Default | ECbValidateMode::Package) != ECbValidateError::None)
+		{
+			UE_LOG(LogDerivedDataCache, Display, TEXT("%s: Cache miss with invalid package for %s from '%.*s'"),
 				*CachePath, *WriteToString<96>(Key), Context.Len(), Context.GetData());
 			return FOptionalCacheRecord();
 		}
+
 		Record = FCacheRecord::Load(Package);
 		if (Record.IsNull())
 		{
 			UE_LOG(LogDerivedDataCache, Display, TEXT("%s: Cache miss with record load failure for %s from '%.*s'"),
 				*CachePath, *WriteToString<96>(Key), Context.Len(), Context.GetData());
-			return FOptionalCacheRecord();
+			return Record;
 		}
 	}
 
-	// Disable deletion now that the record is loaded and validated.
-	bDeleteCacheObject = false;
-
 	if (AccessLogWriter)
 	{
 		AccessLogWriter->Append(Key, Path);
 	}
 
-	return Record.Get();
+	bDeleteRecord = false;
+	return Record;
 }
 
 FOptionalCacheRecord FFileSystemCacheStore::GetCacheRecord(
@@ -1924,23 +1932,39 @@ FOptionalCacheRecord FFileSystemCacheStore::GetCacheRecord(
 	if (FPayload Payload = Record.Get().GetValuePayload())
 	{
 		const ECachePolicy PayloadPolicy = Policy.GetPayloadPolicy(Payload.GetId());
-		GetCacheContent(Key, Context, PayloadPolicy, ECachePolicy::SkipValue, Payload, OutStatus);
-		if (Payload.IsNull())
+		if (FPayload Content = GetCacheContent(Key, Payload, Context, PayloadPolicy, ECachePolicy::SkipValue))
 		{
+			RecordBuilder.SetValue(MoveTemp(Content));
+		}
+		else if (EnumHasAnyFlags(PayloadPolicy, ECachePolicy::PartialOnError))
+		{
+			OutStatus = EStatus::Error;
+			RecordBuilder.SetValue(MoveTemp(Payload));
+		}
+		else
+		{
+			OutStatus = EStatus::Error;
 			return FOptionalCacheRecord();
 		}
-		RecordBuilder.SetValue(MoveTemp(Payload));
 	}
 
 	for (FPayload Payload : Record.Get().GetAttachmentPayloads())
 	{
 		const ECachePolicy PayloadPolicy = Policy.GetPayloadPolicy(Payload.GetId());
-		GetCacheContent(Key, Context, PayloadPolicy, ECachePolicy::SkipAttachments, Payload, OutStatus);
-		if (Payload.IsNull())
+		if (FPayload Content = GetCacheContent(Key, Payload, Context, PayloadPolicy, ECachePolicy::SkipAttachments))
+		{
+			RecordBuilder.AddAttachment(MoveTemp(Content));
+		}
+		else if (EnumHasAnyFlags(PayloadPolicy, ECachePolicy::PartialOnError))
+		{
+			OutStatus = EStatus::Error;
+			RecordBuilder.AddAttachment(MoveTemp(Payload));
+		}
+		else
 		{
+			OutStatus = EStatus::Error;
 			return FOptionalCacheRecord();
 		}
-		RecordBuilder.AddAttachment(MoveTemp(Payload));
 	}
 
 	return RecordBuilder.Build();
@@ -1951,12 +1975,9 @@ bool FFileSystemCacheStore::PutCacheContent(const FCompressedBuffer& Content, co
 	const FIoHash& RawHash = Content.GetRawHash();
 	TStringBuilder<256> Path;
 	BuildCacheContentPath(RawHash, Path);
-	if (!FileExists(Path))
+	if (!FileExists(Path) && !SaveFile(Path, Context, [&Content](FArchive& Ar) { Content.Save(Ar); }))
 	{
-		if (!SaveFile(Path, Context, [&Content](FArchive& Ar) { Content.Save(Ar); }))
-		{
-			return false;
-		}
+		return false;
 	}
 	if (AccessLogWriter)
 	{
@@ -1965,27 +1986,99 @@ bool FFileSystemCacheStore::PutCacheContent(const FCompressedBuffer& Content, co
 	return true;
 }
 
+FPayload FFileSystemCacheStore::GetCacheContent(
+	const FCacheKey& Key,
+	const FPayload& Payload,
+	const FStringView Context,
+	const ECachePolicy Policy,
+	const ECachePolicy SkipFlag) const
+{
+	if (!EnumHasAnyFlags(Policy, ECachePolicy::Query))
+	{
+		return FPayload(Payload.GetId(), Payload.GetRawHash(), Payload.GetRawSize());
+	}
+
+	if (Payload.HasData())
+	{
+		if (EnumHasAnyFlags(Policy, SkipFlag))
+		{
+			return FPayload(Payload.GetId(), Payload.GetRawHash(), Payload.GetRawSize());
+		}
+		else
+		{
+			return Payload;
+		}
+	}
+
+	const FIoHash& RawHash = Payload.GetRawHash();
+
+	TStringBuilder<256> Path;
+	BuildCacheContentPath(RawHash, Path);
+	if (EnumHasAnyFlags(Policy, SkipFlag))
+	{
+		if (FileExists(Path))
+		{
+			if (AccessLogWriter)
+			{
+				AccessLogWriter->Append(RawHash, Path);
+			}
+			return Payload;
+		}
+	}
+	else
+	{
+		FCompressedBuffer CompressedBuffer;
+		if (LoadFile(Path, Context, [&CompressedBuffer](FArchive& Ar) { CompressedBuffer = FCompressedBuffer::Load(Ar); }))
+		{
+			if (CompressedBuffer.GetRawHash() == RawHash)
+			{
+				if (AccessLogWriter)
+				{
+					AccessLogWriter->Append(RawHash, Path);
+				}
+				return FPayload(Payload.GetId(), MoveTemp(CompressedBuffer));
+			}
+			UE_LOG(LogDerivedDataCache, Display,
+				TEXT("%s: Cache miss with corrupted payload %s with hash %s for %s from '%.*s'"),
+				*CachePath, *WriteToString<16>(Payload.GetId()), *WriteToString<48>(RawHash),
+				*WriteToString<96>(Key), Context.Len(), Context.GetData());
+			return FPayload::Null;
+		}
+	}
+
+	UE_LOG(LogDerivedDataCache, Verbose,
+		TEXT("%s: Cache miss with missing payload %s with hash %s for %s from '%.*s'"),
+		*CachePath, *WriteToString<16>(Payload.GetId()), *WriteToString<48>(RawHash), *WriteToString<96>(Key),
+		Context.Len(), Context.GetData());
+	return FPayload::Null;
+}
+
 void FFileSystemCacheStore::GetCacheContent(
 	const FCacheKey& Key,
+	const FPayload& Payload,
 	const FStringView Context,
 	const ECachePolicy Policy,
 	const ECachePolicy SkipFlag,
-	FPayload& InOutPayload,
-	EStatus& InOutStatus) const
+	FCompressedBufferReader& Reader,
+	TUniquePtr<FArchive>& OutArchive) const
 {
-	const FIoHash& RawHash = InOutPayload.GetRawHash();
-
-	if (!EnumHasAnyFlags(Policy, ECachePolicy::Query) || (EnumHasAnyFlags(Policy, SkipFlag) && InOutPayload.HasData()))
+	if (!EnumHasAnyFlags(Policy, ECachePolicy::Query))
 	{
-		InOutPayload = FPayload(InOutPayload.GetId(), RawHash, InOutPayload.GetRawSize());
 		return;
 	}
 
-	if (InOutPayload.HasData())
+	if (Payload.HasData())
 	{
+		if (!EnumHasAnyFlags(Policy, SkipFlag))
+		{
+			Reader.SetSource(Payload.GetData());
+		}
+		OutArchive.Reset();
 		return;
 	}
 
+	const FIoHash& RawHash = Payload.GetRawHash();
+
 	TStringBuilder<256> Path;
 	BuildCacheContentPath(RawHash, Path);
 	if (EnumHasAllFlags(Policy, SkipFlag))
@@ -2001,40 +2094,32 @@ void FFileSystemCacheStore::GetCacheContent(
 	}
 	else
 	{
-		if (FSharedBuffer CompressedData = LoadFile(Path, Context))
+		OutArchive = OpenFile(Path, Context);
+		if (OutArchive)
 		{
-			if (FCompressedBuffer CompressedBuffer = FCompressedBuffer::FromCompressed(MoveTemp(CompressedData));
-				CompressedBuffer && CompressedBuffer.GetRawHash() == RawHash)
+			Reader.SetSource(*OutArchive);
+			if (Reader.GetRawHash() == RawHash)
 			{
 				if (AccessLogWriter)
 				{
 					AccessLogWriter->Append(RawHash, Path);
 				}
-				InOutPayload = FPayload(InOutPayload.GetId(), MoveTemp(CompressedBuffer));
 				return;
 			}
 			UE_LOG(LogDerivedDataCache, Display,
 				TEXT("%s: Cache miss with corrupted payload %s with hash %s for %s from '%.*s'"),
-				*CachePath, *WriteToString<16>(InOutPayload.GetId()), *WriteToString<48>(RawHash),
+				*CachePath, *WriteToString<16>(Payload.GetId()), *WriteToString<48>(RawHash),
 				*WriteToString<96>(Key), Context.Len(), Context.GetData());
-			InOutStatus = EStatus::Error;
-			if (!EnumHasAnyFlags(Policy, ECachePolicy::PartialOnError))
-			{
-				InOutPayload = FPayload::Null;
-			}
+			Reader.ResetSource();
+			OutArchive.Reset();
 			return;
 		}
 	}
 
 	UE_LOG(LogDerivedDataCache, Verbose,
 		TEXT("%s: Cache miss with missing payload %s with hash %s for %s from '%.*s'"),
-		*CachePath, *WriteToString<16>(InOutPayload.GetId()), *WriteToString<48>(RawHash), *WriteToString<96>(Key),
+		*CachePath, *WriteToString<16>(Payload.GetId()), *WriteToString<48>(RawHash), *WriteToString<96>(Key),
 		Context.Len(), Context.GetData());
-	InOutStatus = EStatus::Error;
-	if (!EnumHasAnyFlags(Policy, ECachePolicy::PartialOnError))
-	{
-		InOutPayload = FPayload::Null;
-	}
 }
 
 void FFileSystemCacheStore::BuildCacheRecordPath(const FCacheKey& CacheKey, FStringBuilderBase& Path) const
@@ -2052,7 +2137,7 @@ void FFileSystemCacheStore::BuildCacheContentPath(const FIoHash& RawHash, FStrin
 bool FFileSystemCacheStore::SaveFile(
 	FStringBuilderBase& Path,
 	const FStringView Context,
-	TFunctionRef<void (FArchive&)> WriteFunction) const
+	const TFunctionRef<void (FArchive&)> WriteFunction) const
 {
 	TStringBuilder<256> TempPath;
 	TempPath << FPathViews::GetPath(Path) << TEXT("/Temp.") << FGuid::NewGuid();
@@ -2062,118 +2147,87 @@ bool FFileSystemCacheStore::SaveFile(
 		IFileManager::Get().Delete(*TempPath, /*bRequireExists*/ false, /*bEvenReadOnly*/ false, /*bQuiet*/ true);
 	};
 
-	int64 ExpectedSize = 0;
-	if (SaveDataToFile(TempPath, MoveTemp(WriteFunction), ExpectedSize))
+	int64 WriteSize = 0;
+	bool bError = false;
+
+	if (TUniquePtr<FArchive> Ar{IFileManager::Get().CreateFileWriter(*TempPath, FILEWRITE_Silent)})
 	{
-		if (IFileManager::Get().FileSize(*TempPath) == ExpectedSize)
-		{
-			if (IFileManager::Get().Move(*Path, *TempPath,
-				/*bReplace*/ false, /*bEvenIfReadOnly*/ false, /*bAttributes*/ false, /*bDoNotRetryOrError*/ true))
-			{
-				return true;
-			}
-			else
-			{
-				UE_LOG(LogDerivedDataCache, Log,
-					TEXT("%s: Move collision when writing file %s from '%.*s'"),
-					*CachePath, *Path, Context.Len(), Context.GetData());
-				return true;
-			}
-		}
-		else
-		{
-			UE_LOG(LogDerivedDataCache, Warning,
-				TEXT("%s: Failed to write to temp file %s when saving %s from '%.*s'. ")
-				TEXT("File is %" INT64_FMT " bytes when %" INT64_FMT " bytes are expected."),
-				*CachePath, *TempPath, *Path, Context.Len(), Context.GetData(),
-				IFileManager::Get().FileSize(*TempPath), ExpectedSize);
-			return false;
-		}
+		THashingArchiveProxy<FBlake3> HashAr(*Ar);
+		WriteFunction(HashAr);
+		FBlake3Hash Hash = HashAr.GetHash();
+		*Ar << Hash;
+		WriteSize = HashAr.Tell();
+		bError = !HashAr.Close();
 	}
-	else
+
+	if (bError || WriteSize == 0)
 	{
 		UE_LOG(LogDerivedDataCache, Warning,
 			TEXT("%s: Failed to write to temp file %s when saving %s from '%.*s'. Error 0x%08x"),
 			*CachePath, *TempPath, *Path, Context.Len(), Context.GetData(), FPlatformMisc::GetLastError());
 		return false;
 	}
-}
 
-bool FFileSystemCacheStore::SaveDataToFile(
-	FStringBuilderBase& Path,
-	TFunctionRef<void (FArchive&)> WriteFunction,
-	int64& OutWriteSize) const
-{
-	if (TUniquePtr<FArchive> Ar{IFileManager::Get().CreateFileWriter(*Path, FILEWRITE_Silent)})
+	if (IFileManager::Get().FileSize(*TempPath) != WriteSize)
 	{
-		THashingArchiveProxy<FBlake3> HashAr(*Ar);
-		WriteFunction(HashAr);
-		FBlake3Hash Hash = HashAr.GetHash();
-		*Ar << Hash;
-		OutWriteSize = Ar->Tell();
-		return Ar->Close() && !Ar->IsError();
+		UE_LOG(LogDerivedDataCache, Warning,
+			TEXT("%s: Failed to write to temp file %s when saving %s from '%.*s'. ")
+			TEXT("File is %" INT64_FMT " bytes when %" INT64_FMT " bytes are expected."),
+			*CachePath, *TempPath, *Path, Context.Len(), Context.GetData(),
+			IFileManager::Get().FileSize(*TempPath), WriteSize);
+		return false;
 	}
-	OutWriteSize = 0;
-	return false;
+
+	if (!IFileManager::Get().Move(*Path, *TempPath,
+		/*bReplace*/ false, /*bEvenIfReadOnly*/ false, /*bAttributes*/ false, /*bDoNotRetryOrError*/ true))
+	{
+		UE_LOG(LogDerivedDataCache, Log,
+			TEXT("%s: Move collision when writing file %s from '%.*s'."),
+			*CachePath, *Path, Context.Len(), Context.GetData());
+	}
+
+	return true;
 }
 
-FSharedBuffer FFileSystemCacheStore::LoadFile(FStringBuilderBase& Path, const FStringView Context) const
+bool FFileSystemCacheStore::LoadFile(
+	FStringBuilderBase& Path,
+	const FStringView Context,
+	const TFunctionRef<void (FArchive& Ar)> ReadFunction) const
 {
 	check(IsUsable());
 	const double StartTime = FPlatformTime::Seconds();
 
-	FSharedBuffer Buffer;
+	int64 ReadSize = 0;
+	bool bError = false;
 
-	// Check for existence before reading because it may update the modification time and avoid the
-	// file being deleted by a cache cleanup thread or process.
-	if (!FileExists(Path))
-	{
-		return Buffer;
-	}
-
-	if (TUniquePtr<FArchive> Ar{IFileManager::Get().CreateFileReader(*Path, FILEREAD_Silent)})
+	if (TUniquePtr<FArchive> Ar = OpenFile(Path, Context))
 	{
-		if (const int64 TotalSize = Ar->TotalSize(); TotalSize >= sizeof(FBlake3Hash))
+		THashingArchiveProxy<FBlake3> HashAr(*Ar);
+		ReadFunction(HashAr);
+		const FBlake3Hash Hash = HashAr.GetHash();
+		FBlake3Hash SavedHash;
+		*Ar << SavedHash;
+		ReadSize = HashAr.Tell();
+		bError = !HashAr.Close();
+
+		if (bError)
 		{
-			const int64 DataSize = TotalSize - sizeof(FBlake3Hash);
-			FUniqueBuffer MutableBuffer = FUniqueBuffer::Alloc(uint64(DataSize));
-			Ar->Serialize(MutableBuffer.GetData(), DataSize);
-			FBlake3Hash SavedHash;
-			*Ar << SavedHash;
-			if (Ar->Close())
-			{
-				const FBlake3Hash Hash = FBlake3::HashBuffer(MutableBuffer);
-				if (Hash == SavedHash)
-				{
-					Buffer = MutableBuffer.MoveToShared();
-				}
-				else
-				{
-					UE_LOG(LogDerivedDataCache, Display,
-						TEXT("%s: File %s from '%.*s' is corrupted and has hash %s when %s is expected."),
-						*CachePath, *Path, Context.Len(), Context.GetData(),
-						*WriteToString<80>(Hash), *WriteToString<80>(SavedHash));
-					if (!bReadOnly)
-					{
-						IFileManager::Get().Delete(*Path, /*bRequireExists*/ false, /*bEvenReadOnly*/ false, /*bQuiet*/ true);
-					}
-				}
-			}
+			UE_LOG(LogDerivedDataCache, Display,
+				TEXT("%s: Failed to load file %s from '%.*s'."),
+				*CachePath, *Path, Context.Len(), Context.GetData());
 		}
-		else
+		else if (Hash != SavedHash)
 		{
+			bError = true;
 			UE_LOG(LogDerivedDataCache, Display,
-				TEXT("%s: File %s from '%.*s' is %" INT64_FMT " bytes when at least %" INT64_FMT " bytes are required."),
-				*CachePath, *Path, Context.Len(), Context.GetData(), TotalSize, int64(sizeof(FBlake3Hash)));
-			if (!bReadOnly)
-			{
-				IFileManager::Get().Delete(*Path, /*bRequireExists*/ false, /*bEvenReadOnly*/ false, /*bQuiet*/ true);
-			}
+				TEXT("%s: File %s from '%.*s' is corrupted and has hash %s when %s is expected."),
+				*CachePath, *Path, Context.Len(), Context.GetData(),
+				*WriteToString<80>(Hash), *WriteToString<80>(SavedHash));
 		}
 	}
 
 	const double ReadDuration = FPlatformTime::Seconds() - StartTime;
-	const double ReadSpeed = ReadDuration > 0.001 ? (Buffer.GetSize() / ReadDuration) / (1024.0 * 1024.0) : 0.0;
+	const double ReadSpeed = ReadDuration > 0.001 ? (ReadSize / ReadDuration) / (1024.0 * 1024.0) : 0.0;
 
 	if (!GIsBuildMachine && ReadDuration > 5.0)
 	{
@@ -2183,14 +2237,28 @@ FSharedBuffer FFileSystemCacheStore::LoadFile(FStringBuilderBase& Path, const FS
 			*CachePath, *Path, Context.Len(), Context.GetData(), ReadSpeed);
 	}
 
-	if (Buffer)
+	UE_LOG(LogDerivedDataCache, VeryVerbose,
+		TEXT("%s: Loaded %s from '%.*s' (%" INT64_FMT " bytes, %.02f secs, %.2f MiB/s)"),
+		*CachePath, *Path, Context.Len(), Context.GetData(), ReadSize, ReadDuration, ReadSpeed);
+
+	if (bError && !bReadOnly)
+	{
+		IFileManager::Get().Delete(*Path, /*bRequireExists*/ false, /*bEvenReadOnly*/ false, /*bQuiet*/ true);
+	}
+
+	return !bError && ReadSize > 0;
+}
+
+TUniquePtr<FArchive> FFileSystemCacheStore::OpenFile(FStringBuilderBase& Path, const FStringView Context) const
+{
+	// Check for existence before reading because it may update the modification time and avoid the
+	// file being deleted by a cache cleanup thread or process.
+	if (!FileExists(Path))
 	{
-		UE_LOG(LogDerivedDataCache, VeryVerbose,
-			TEXT("%s: Loaded %s from '%.*s' (%" UINT64_FMT " bytes, %.02f secs, %.2f MiB/s)"),
-			*CachePath, *Path, Context.Len(), Context.GetData(), Buffer.GetSize(), ReadDuration, ReadSpeed);
+		return nullptr;
 	}
 
-	return Buffer;
+	return TUniquePtr<FArchive>(IFileManager::Get().CreateFileReader(*Path, FILEREAD_Silent));
 }
 
 bool FFileSystemCacheStore::FileExists(FStringBuilderBase& Path) const
