DDC: Added a uint64 UserData to ICacheStore requests to make completion callbacks more efficient

Requests to Put and Get now specify the policy independently for each record or key, but the implementations have not yet been updated to take this into account.

The existing Put and Get API has been maintained temporarily to allow the conversion to happen separately, and will be removed after callers have been converted to use the new API.

#rnx

diff --git a/Engine/Source/Developer/DerivedDataCache/Private/DerivedDataBackendAsyncPutWrapper.cpp b/Engine/Source/Developer/DerivedDataCache/Private/DerivedDataBackendAsyncPutWrapper.cpp
--- a/Engine/Source/Developer/DerivedDataCache/Private/DerivedDataBackendAsyncPutWrapper.cpp
+++ b/Engine/Source/Developer/DerivedDataCache/Private/DerivedDataBackendAsyncPutWrapper.cpp
@@ -457,32 +457,31 @@ private:
 };
 
 void FDerivedDataBackendAsyncPutWrapper::Put(
-	TConstArrayView<FCacheRecord> Records,
-	FStringView Context,
-	ECachePolicy Policy,
+	const TConstArrayView<FCachePutRequest> Requests,
+	const FStringView Context,
 	IRequestOwner& Owner,
 	FOnCachePutComplete&& OnComplete)
 {
 	if (Owner.GetPriority() == EPriority::Blocking || !GDDCIOThreadPool)
 	{
-		InnerBackend->Put(Records, Context, Policy, Owner, MoveTemp(OnComplete));
+		InnerBackend->Put(Requests, Context, Owner, MoveTemp(OnComplete));
 	}
 	else
 	{
 		FDerivedDataAsyncWrapperRequest* Request = new FDerivedDataAsyncWrapperRequest(Owner,
-			[this, Records = TArray<FCacheRecord>(Records), Context = FString(Context), Policy, OnComplete = MoveTemp(OnComplete)](bool bCancel) mutable
+			[this, Requests = TArray<FCachePutRequest>(Requests), Context = FString(Context), OnComplete = MoveTemp(OnComplete)](bool bCancel) mutable
 			{
 				if (!bCancel)
 				{
 					FRequestOwner BlockingOwner(EPriority::Blocking);
-					InnerBackend->Put(Records, Context, Policy, BlockingOwner, MoveTemp(OnComplete));
+					InnerBackend->Put(Requests, Context, BlockingOwner, MoveTemp(OnComplete));
 					BlockingOwner.Wait();
 				}
 				else if (OnComplete)
 				{
-					for (const FCacheRecord& Record : Records)
+					for (const FCachePutRequest& Request : Requests)
 					{
-						OnComplete({Record.GetKey(), EStatus::Canceled});
+						OnComplete({Request.Record.GetKey(), Request.UserData, EStatus::Canceled});
 					}
 				}
 			});
@@ -491,32 +490,31 @@ void FDerivedDataBackendAsyncPutWrapper::Put(
 }
 
 void FDerivedDataBackendAsyncPutWrapper::Get(
-	TConstArrayView<FCacheKey> Keys,
-	FStringView Context,
-	FCacheRecordPolicy Policy,
+	const TConstArrayView<FCacheGetRequest> Requests,
+	const FStringView Context,
 	IRequestOwner& Owner,
 	FOnCacheGetComplete&& OnComplete)
 {
 	if (Owner.GetPriority() == EPriority::Blocking || !GDDCIOThreadPool)
 	{
-		InnerBackend->Get(Keys, Context, Policy, Owner, MoveTemp(OnComplete));
+		InnerBackend->Get(Requests, Context, Owner, MoveTemp(OnComplete));
 	}
 	else
 	{
 		FDerivedDataAsyncWrapperRequest* Request = new FDerivedDataAsyncWrapperRequest(Owner,
-			[this, Keys = TArray<FCacheKey>(Keys), Context = FString(Context), Policy, OnComplete = MoveTemp(OnComplete)](bool bCancel) mutable
+			[this, Requests = TArray<FCacheGetRequest>(Requests), Context = FString(Context), OnComplete = MoveTemp(OnComplete)](bool bCancel) mutable
 			{
 				if (!bCancel)
 				{
 					FRequestOwner BlockingOwner(EPriority::Blocking);
-					InnerBackend->Get(Keys, Context, Policy, BlockingOwner, MoveTemp(OnComplete));
+					InnerBackend->Get(Requests, Context, BlockingOwner, MoveTemp(OnComplete));
 					BlockingOwner.Wait();
 				}
 				else if (OnComplete)
 				{
-					for (const FCacheKey& Key : Keys)
+					for (const FCacheGetRequest& Request : Requests)
 					{
-						OnComplete({FCacheRecordBuilder(Key).Build(), EStatus::Canceled});
+						OnComplete({FCacheRecordBuilder(Request.Key).Build(), Request.UserData, EStatus::Canceled});
 					}
 				}
 			});
@@ -549,7 +547,7 @@ void FDerivedDataBackendAsyncPutWrapper::GetChunks(
 				{
 					for (const FCacheChunkRequest& Chunk : Chunks)
 					{
-						OnComplete({Chunk.Key, Chunk.Id, Chunk.RawOffset, 0, {}, {}, EStatus::Canceled});
+						OnComplete({Chunk.Key, Chunk.Id, Chunk.RawOffset, 0, {}, {}, Chunk.UserData, EStatus::Canceled});
 					}
 				}
 			});
diff --git a/Engine/Source/Developer/DerivedDataCache/Private/DerivedDataBackendAsyncPutWrapper.h b/Engine/Source/Developer/DerivedDataCache/Private/DerivedDataBackendAsyncPutWrapper.h
--- a/Engine/Source/Developer/DerivedDataCache/Private/DerivedDataBackendAsyncPutWrapper.h
+++ b/Engine/Source/Developer/DerivedDataCache/Private/DerivedDataBackendAsyncPutWrapper.h
@@ -145,16 +145,14 @@ public:
 	virtual bool ApplyDebugOptions(FBackendDebugOptions& InOptions) override;
 
 	virtual void Put(
-		TConstArrayView<FCacheRecord> Records,
+		TConstArrayView<FCachePutRequest> Requests,
 		FStringView Context,
-		ECachePolicy Policy,
 		IRequestOwner& Owner,
 		FOnCachePutComplete&& OnComplete) override;
 
 	virtual void Get(
-		TConstArrayView<FCacheKey> Keys,
+		TConstArrayView<FCacheGetRequest> Requests,
 		FStringView Context,
-		FCacheRecordPolicy Policy,
 		IRequestOwner& Owner,
 		FOnCacheGetComplete&& OnComplete) override;
 
diff --git a/Engine/Source/Developer/DerivedDataCache/Private/DerivedDataBackendCorruptionWrapper.h b/Engine/Source/Developer/DerivedDataCache/Private/DerivedDataBackendCorruptionWrapper.h
--- a/Engine/Source/Developer/DerivedDataCache/Private/DerivedDataBackendCorruptionWrapper.h
+++ b/Engine/Source/Developer/DerivedDataCache/Private/DerivedDataBackendCorruptionWrapper.h
@@ -258,23 +258,21 @@ public:
 	}
 
 	virtual void Put(
-		TConstArrayView<FCacheRecord> Records,
+		TConstArrayView<FCachePutRequest> Requests,
 		FStringView Context,
-		ECachePolicy Policy,
 		IRequestOwner& Owner,
 		FOnCachePutComplete&& OnComplete) override
 	{
-		return InnerBackend->Put(Records, Context, Policy, Owner, MoveTemp(OnComplete));
+		return InnerBackend->Put(Requests, Context, Owner, MoveTemp(OnComplete));
 	}
 
 	virtual void Get(
-		TConstArrayView<FCacheKey> Keys,
+		TConstArrayView<FCacheGetRequest> Requests,
 		FStringView Context,
-		FCacheRecordPolicy Policy,
 		IRequestOwner& Owner,
 		FOnCacheGetComplete&& OnComplete) override
 	{
-		return InnerBackend->Get(Keys, Context, Policy, Owner, MoveTemp(OnComplete));
+		return InnerBackend->Get(Requests, Context, Owner, MoveTemp(OnComplete));
 	}
 
 	virtual void GetChunks(
diff --git a/Engine/Source/Developer/DerivedDataCache/Private/DerivedDataBackendThrottleWrapper.h b/Engine/Source/Developer/DerivedDataCache/Private/DerivedDataBackendThrottleWrapper.h
--- a/Engine/Source/Developer/DerivedDataCache/Private/DerivedDataBackendThrottleWrapper.h
+++ b/Engine/Source/Developer/DerivedDataCache/Private/DerivedDataBackendThrottleWrapper.h
@@ -178,9 +178,8 @@ public:
 	}
 
 	virtual void Put(
-		TConstArrayView<FCacheRecord> Records,
+		TConstArrayView<FCachePutRequest> Requests,
 		FStringView Context,
-		ECachePolicy Policy,
 		IRequestOwner& Owner,
 		FOnCachePutComplete&& OnComplete) override
 	{
@@ -190,13 +189,13 @@ public:
 			uint64 Size;
 		};
 		TArray<FRecordSize, TInlineAllocator<1>> RecordSizes;
-		RecordSizes.Reserve(Records.Num());
-		Algo::Transform(Records, RecordSizes, [](const FCacheRecord& Record) -> FRecordSize
+		RecordSizes.Reserve(Requests.Num());
+		Algo::Transform(Requests, RecordSizes, [](const FCachePutRequest& Request) -> FRecordSize
 		{
-			return {Record.GetKey(), Private::GetCacheRecordCompressedSize(Record)};
+			return {Request.Record.GetKey(), Private::GetCacheRecordCompressedSize(Request.Record)};
 		});
 
-		InnerBackend->Put(Records, Context, Policy, Owner,
+		InnerBackend->Put(Requests, Context, Owner,
 			[this, RecordSizes = MoveTemp(RecordSizes), State = EnterThrottlingScope(), OnComplete = MoveTemp(OnComplete)](FCachePutCompleteParams&& Params)
 			{
 				const FRecordSize* Size = Algo::FindBy(RecordSizes, Params.Key, &FRecordSize::Key);
@@ -209,13 +208,12 @@ public:
 	}
 
 	virtual void Get(
-		TConstArrayView<FCacheKey> Keys,
+		TConstArrayView<FCacheGetRequest> Requests,
 		FStringView Context,
-		FCacheRecordPolicy Policy,
 		IRequestOwner& Owner,
 		FOnCacheGetComplete&& OnComplete) override
 	{
-		InnerBackend->Get(Keys, Context, Policy, Owner,
+		InnerBackend->Get(Requests, Context, Owner,
 			[this, State = EnterThrottlingScope(), OnComplete = MoveTemp(OnComplete)](FCacheGetCompleteParams&& Params)
 			{
 				CloseThrottlingScope(State, FThrottlingState(this, Private::GetCacheRecordCompressedSize(Params.Record)));
diff --git a/Engine/Source/Developer/DerivedDataCache/Private/DerivedDataBackendVerifyWrapper.h b/Engine/Source/Developer/DerivedDataCache/Private/DerivedDataBackendVerifyWrapper.h
--- a/Engine/Source/Developer/DerivedDataCache/Private/DerivedDataBackendVerifyWrapper.h
+++ b/Engine/Source/Developer/DerivedDataCache/Private/DerivedDataBackendVerifyWrapper.h
@@ -174,23 +174,21 @@ public:
 	}
 
 	virtual void Put(
-		TConstArrayView<FCacheRecord> Records,
+		TConstArrayView<FCachePutRequest> Requests,
 		FStringView Context,
-		ECachePolicy Policy,
 		IRequestOwner& Owner,
 		FOnCachePutComplete&& OnComplete) override
 	{
-		return InnerBackend->Put(Records, Context, Policy, Owner, MoveTemp(OnComplete));
+		return InnerBackend->Put(Requests, Context, Owner, MoveTemp(OnComplete));
 	}
 
 	virtual void Get(
-		TConstArrayView<FCacheKey> Keys,
+		TConstArrayView<FCacheGetRequest> Requests,
 		FStringView Context,
-		FCacheRecordPolicy Policy,
 		IRequestOwner& Owner,
 		FOnCacheGetComplete&& OnComplete) override
 	{
-		return InnerBackend->Get(Keys, Context, Policy, Owner, MoveTemp(OnComplete));
+		return InnerBackend->Get(Requests, Context, Owner, MoveTemp(OnComplete));
 	}
 
 	virtual void GetChunks(
diff --git a/Engine/Source/Developer/DerivedDataCache/Private/DerivedDataCache.cpp b/Engine/Source/Developer/DerivedDataCache/Private/DerivedDataCache.cpp
--- a/Engine/Source/Developer/DerivedDataCache/Private/DerivedDataCache.cpp
+++ b/Engine/Source/Developer/DerivedDataCache/Private/DerivedDataCache.cpp
@@ -261,6 +261,38 @@ FCacheRecordPolicy FCacheRecordPolicyBuilder::Build()
 	return Policy;
 }
 
+void ICache::Put(
+	const TConstArrayView<FCacheRecord> Records,
+	const FStringView Context,
+	const ECachePolicy Policy,
+	IRequestOwner& Owner,
+	FOnCachePutComplete&& OnComplete)
+{
+	TArray<FCachePutRequest, TInlineAllocator<16>> Requests;
+	Requests.Reserve(Records.Num());
+	for (const FCacheRecord& Record : Records)
+	{
+		Requests.Add({Record, Policy});
+	}
+	return Put(Requests, Context, Owner, MoveTemp(OnComplete));
+}
+
+void ICache::Get(
+	const TConstArrayView<FCacheKey> Keys,
+	const FStringView Context,
+	const FCacheRecordPolicy Policy,
+	IRequestOwner& Owner,
+	FOnCacheGetComplete&& OnComplete)
+{
+	TArray<FCacheGetRequest, TInlineAllocator<16>> Requests;
+	Requests.Reserve(Keys.Num());
+	for (const FCacheKey& Key : Keys)
+	{
+		Requests.Add({Key, Policy});
+	}
+	return Get(Requests, Context, Owner, MoveTemp(OnComplete));
+}
+
 } // UE::DerivedData
 
 namespace UE::DerivedData::Private
@@ -807,26 +839,24 @@ private:
 	FOnDDCNotification DDCNotificationEvent;
 
 public:
-	// ICache Interface
+	// ICacheStore Interface
 
 	void Put(
-		TConstArrayView<FCacheRecord> Records,
+		TConstArrayView<FCachePutRequest> Requests,
 		FStringView Context,
-		ECachePolicy Policy,
 		IRequestOwner& Owner,
 		FOnCachePutComplete&& OnComplete) final
 	{
-		return FDerivedDataBackend::Get().GetRoot().Put(Records, Context, Policy, Owner, MoveTemp(OnComplete));
+		return FDerivedDataBackend::Get().GetRoot().Put(Requests, Context, Owner, MoveTemp(OnComplete));
 	}
 
 	void Get(
-		TConstArrayView<FCacheKey> Keys,
+		TConstArrayView<FCacheGetRequest> Requests,
 		FStringView Context,
-		FCacheRecordPolicy Policy,
 		IRequestOwner& Owner,
 		FOnCacheGetComplete&& OnComplete) final
 	{
-		return FDerivedDataBackend::Get().GetRoot().Get(Keys, Context, Policy, Owner, MoveTemp(OnComplete));
+		return FDerivedDataBackend::Get().GetRoot().Get(Requests, Context, Owner, MoveTemp(OnComplete));
 	}
 
 	void GetChunks(
@@ -838,6 +868,8 @@ public:
 		return FDerivedDataBackend::Get().GetRoot().GetChunks(Chunks, Context, Owner, MoveTemp(OnComplete));
 	}
 
+	// ICache Interface
+
 	void CancelAll() final
 	{
 	}
diff --git a/Engine/Source/Developer/DerivedDataCache/Private/DerivedDataLimitKeyLengthWrapper.h b/Engine/Source/Developer/DerivedDataCache/Private/DerivedDataLimitKeyLengthWrapper.h
--- a/Engine/Source/Developer/DerivedDataCache/Private/DerivedDataLimitKeyLengthWrapper.h
+++ b/Engine/Source/Developer/DerivedDataCache/Private/DerivedDataLimitKeyLengthWrapper.h
@@ -255,23 +255,21 @@ public:
 	}
 
 	virtual void Put(
-		TConstArrayView<FCacheRecord> Records,
+		TConstArrayView<FCachePutRequest> Requests,
 		FStringView Context,
-		ECachePolicy Policy,
 		IRequestOwner& Owner,
 		FOnCachePutComplete&& OnComplete) override
 	{
-		return InnerBackend->Put(Records, Context, Policy, Owner, MoveTemp(OnComplete));
+		return InnerBackend->Put(Requests, Context, Owner, MoveTemp(OnComplete));
 	}
 
 	virtual void Get(
-		TConstArrayView<FCacheKey> Keys,
+		TConstArrayView<FCacheGetRequest> Requests,
 		FStringView Context,
-		FCacheRecordPolicy Policy,
 		IRequestOwner& Owner,
 		FOnCacheGetComplete&& OnComplete) override
 	{
-		return InnerBackend->Get(Keys, Context, Policy, Owner, MoveTemp(OnComplete));
+		return InnerBackend->Get(Requests, Context, Owner, MoveTemp(OnComplete));
 	}
 
 	virtual void GetChunks(
diff --git a/Engine/Source/Developer/DerivedDataCache/Private/FileSystemDerivedDataBackend.cpp b/Engine/Source/Developer/DerivedDataCache/Private/FileSystemDerivedDataBackend.cpp
--- a/Engine/Source/Developer/DerivedDataCache/Private/FileSystemDerivedDataBackend.cpp
+++ b/Engine/Source/Developer/DerivedDataCache/Private/FileSystemDerivedDataBackend.cpp
@@ -748,15 +748,13 @@ public:
 	// ICacheStore Interface
 
 	void Put(
-		TConstArrayView<FCacheRecord> Records,
+		TConstArrayView<FCachePutRequest> Requests,
 		FStringView Context,
-		ECachePolicy Policy,
 		IRequestOwner& Owner,
 		FOnCachePutComplete&& OnComplete) final;
 	void Get(
-		TConstArrayView<FCacheKey> Keys,
+		TConstArrayView<FCacheGetRequest> Requests,
 		FStringView Context,
-		FCacheRecordPolicy Policy,
 		IRequestOwner& Owner,
 		FOnCacheGetComplete&& OnComplete) final;
 	void GetChunks(
@@ -769,7 +767,7 @@ private:
 	uint64 MeasureCompressedCacheRecord(const FCacheRecord& Record) const;
 	uint64 MeasureRawCacheRecord(const FCacheRecord& Record) const;
 
-	bool PutCacheRecord(const FCacheRecord& Record, FStringView Context, ECachePolicy Policy);
+	bool PutCacheRecord(const FCacheRecord& Record, FStringView Context, const FCacheRecordPolicy& Policy);
 	bool PutCacheContent(const FCompressedBuffer& Content, const FStringView Context) const;
 
 	FOptionalCacheRecord GetCacheRecordOnly(
@@ -1566,18 +1564,18 @@ bool FFileSystemCacheStore::ApplyDebugOptions(FBackendDebugOptions& InOptions)
 }
 
 void FFileSystemCacheStore::Put(
-	const TConstArrayView<FCacheRecord> Records,
+	const TConstArrayView<FCachePutRequest> Requests,
 	const FStringView Context,
-	const ECachePolicy Policy,
 	IRequestOwner& Owner,
 	FOnCachePutComplete&& OnComplete)
 {
-	for (const FCacheRecord& Record : Records)
+	for (const FCachePutRequest& Request : Requests)
 	{
+		const FCacheRecord& Record = Request.Record;
 		TRACE_CPUPROFILER_EVENT_SCOPE(FileSystemDDC_Put);
 		TRACE_COUNTER_INCREMENT(FileSystemDDC_Put);
 		COOK_STAT(auto Timer = UsageStats.TimePut());
-		if (PutCacheRecord(Record, Context, Policy))
+		if (PutCacheRecord(Record, Context, Request.Policy))
 		{
 			UE_LOG(LogDerivedDataCache, Verbose, TEXT("%s: Cache put complete for %s from '%.*s'"),
 				*CachePath, *WriteToString<96>(Record.GetKey()), Context.Len(), Context.GetData());
@@ -1586,7 +1584,7 @@ void FFileSystemCacheStore::Put(
 			COOK_STAT(Timer.AddHit(MeasureRawCacheRecord(Record)));
 			if (OnComplete)
 			{
-				OnComplete({Record.GetKey(), EStatus::Ok});
+				OnComplete({Record.GetKey(), Request.UserData, EStatus::Ok});
 			}
 		}
 		else
@@ -1594,42 +1592,41 @@ void FFileSystemCacheStore::Put(
 			COOK_STAT(Timer.AddMiss(MeasureRawCacheRecord(Record)));
 			if (OnComplete)
 			{
-				OnComplete({Record.GetKey(), EStatus::Error});
+				OnComplete({Record.GetKey(), Request.UserData, EStatus::Error});
 			}
 		}
 	}
 }
 
 void FFileSystemCacheStore::Get(
-	const TConstArrayView<FCacheKey> Keys,
+	const TConstArrayView<FCacheGetRequest> Requests,
 	const FStringView Context,
-	const FCacheRecordPolicy Policy,
 	IRequestOwner& Owner,
 	FOnCacheGetComplete&& OnComplete)
 {
-	for (const FCacheKey& Key : Keys)
+	for (const FCacheGetRequest& Request : Requests)
 	{
 		TRACE_CPUPROFILER_EVENT_SCOPE(FileSystemDDC_Get);
 		TRACE_COUNTER_INCREMENT(FileSystemDDC_Get);
 		COOK_STAT(auto Timer = UsageStats.TimeGet());
 		EStatus Status = EStatus::Ok;
-		if (FOptionalCacheRecord Record = GetCacheRecord(Key, Context, Policy, Status))
+		if (FOptionalCacheRecord Record = GetCacheRecord(Request.Key, Context, Request.Policy, Status))
 		{
 			UE_LOG(LogDerivedDataCache, Verbose, TEXT("%s: Cache hit for %s from '%.*s'"),
-				*CachePath, *WriteToString<96>(Key), Context.Len(), Context.GetData());
+				*CachePath, *WriteToString<96>(Request.Key), Context.Len(), Context.GetData());
 			TRACE_COUNTER_INCREMENT(FileSystemDDC_GetHit);
 			TRACE_COUNTER_ADD(FileSystemDDC_BytesRead, MeasureCompressedCacheRecord(Record.Get()));
 			COOK_STAT(Timer.AddHit(MeasureRawCacheRecord(Record.Get())));
 			if (OnComplete)
 			{
-				OnComplete({MoveTemp(Record).Get(), Status});
+				OnComplete({MoveTemp(Record).Get(), Request.UserData, Status});
 			}
 		}
 		else
 		{
 			if (OnComplete)
 			{
-				OnComplete({FCacheRecordBuilder(Key).Build(), Status});
+				OnComplete({FCacheRecordBuilder(Request.Key).Build(), Request.UserData, Status});
 			}
 		}
 	}
@@ -1696,14 +1693,14 @@ void FFileSystemCacheStore::GetChunks(
 				}
 				const EStatus ChunkStatus = bExistsOnly || Buffer.GetSize() == RawSize ? EStatus::Ok : EStatus::Error;
 				OnComplete({Chunk.Key, Chunk.Id, Chunk.RawOffset,
-					RawSize, Payload.GetRawHash(), MoveTemp(Buffer), ChunkStatus});
+					RawSize, Payload.GetRawHash(), MoveTemp(Buffer), Chunk.UserData, ChunkStatus});
 			}
 			continue;
 		}
 
 		if (OnComplete)
 		{
-			OnComplete({Chunk.Key, Chunk.Id, Chunk.RawOffset, 0, {}, {}, EStatus::Error});
+			OnComplete({Chunk.Key, Chunk.Id, Chunk.RawOffset, 0, {}, {}, Chunk.UserData, EStatus::Error});
 		}
 	}
 }
@@ -1727,7 +1724,7 @@ uint64 FFileSystemCacheStore::MeasureRawCacheRecord(const FCacheRecord& Record)
 bool FFileSystemCacheStore::PutCacheRecord(
 	const FCacheRecord& Record,
 	const FStringView Context,
-	const ECachePolicy Policy)
+	const FCacheRecordPolicy& Policy)
 {
 	if (!IsWritable())
 	{
@@ -1738,9 +1735,13 @@ bool FFileSystemCacheStore::PutCacheRecord(
 	}
 
 	const FCacheKey& Key = Record.GetKey();
+	const ECachePolicy CombinedPayloadPolicy = Algo::TransformAccumulate(
+		Policy.GetPayloadPolicies(), &FCachePayloadPolicy::Policy, Policy.GetDefaultPayloadPolicy(), UE_PROJECTION(operator|));
+	const ECachePolicy CombinedPolicy = Policy.GetRecordPolicy() | CombinedPayloadPolicy;
 
 	// Skip the request if storing to the cache is disabled.
-	if (!EnumHasAnyFlags(Policy, SpeedClass == ESpeedClass::Local ? ECachePolicy::StoreLocal : ECachePolicy::StoreRemote))
+	const ECachePolicy StoreFlag = SpeedClass == ESpeedClass::Local ? ECachePolicy::StoreLocal : ECachePolicy::StoreRemote;
+	if (!EnumHasAnyFlags(CombinedPolicy, StoreFlag))
 	{
 		UE_LOG(LogDerivedDataCache, VeryVerbose, TEXT("%s: Skipped put of %s from '%.*s' due to cache policy"),
 			*CachePath, *WriteToString<96>(Key), Context.Len(), Context.GetData());
@@ -1759,7 +1760,7 @@ bool FFileSystemCacheStore::PutCacheRecord(
 	FCbPackage ExistingPackage;
 	TStringBuilder<256> Path;
 	BuildCacheRecordPath(Key, Path);
-	if (EnumHasAllFlags(Policy, ECachePolicy::SkipValue | ECachePolicy::SkipAttachments))
+	if (EnumHasAllFlags(CombinedPayloadPolicy, ECachePolicy::SkipValue | ECachePolicy::SkipAttachments))
 	{
 		bRecordExists = FileExists(Path);
 	}
diff --git a/Engine/Source/Developer/DerivedDataCache/Private/HierarchicalDerivedDataBackend.h b/Engine/Source/Developer/DerivedDataCache/Private/HierarchicalDerivedDataBackend.h
--- a/Engine/Source/Developer/DerivedDataCache/Private/HierarchicalDerivedDataBackend.h
+++ b/Engine/Source/Developer/DerivedDataCache/Private/HierarchicalDerivedDataBackend.h
@@ -478,9 +478,8 @@ public:
 	}
 
 	virtual void Put(
-		TConstArrayView<FCacheRecord> Records,
+		TConstArrayView<FCachePutRequest> Requests,
 		FStringView Context,
-		ECachePolicy Policy,
 		IRequestOwner& Owner,
 		FOnCachePutComplete&& OnComplete) override
 	{
@@ -488,7 +487,7 @@ public:
 		FRequestBarrier AsyncBarrier(AsyncOwner);
 		AsyncOwner.KeepAlive();
 
-		TSet<FCacheKey> RecordsOk;
+		TSet<FCacheKey> RequestsOk;
 
 		{
 			FReadScopeLock LockScope(Lock);
@@ -498,16 +497,16 @@ public:
 				if (InnerBackends[PutCacheIndex]->IsWritable())
 				{
 					// Every record must put synchronously before switching to async calls.
-					if (RecordsOk.Num() < Records.Num())
+					if (RequestsOk.Num() < Requests.Num())
 					{
 						FRequestOwner BlockingOwner(EPriority::Blocking);
-						InnerBackends[PutCacheIndex]->Put(Records, Context, Policy, BlockingOwner,
-							[&OnComplete, &RecordsOk](FCachePutCompleteParams&& Params)
+						InnerBackends[PutCacheIndex]->Put(Requests, Context, BlockingOwner,
+							[&OnComplete, &RequestsOk](FCachePutCompleteParams&& Params)
 							{
 								if (Params.Status == EStatus::Ok)
 								{
 									bool bIsAlreadyInSet = false;
-									RecordsOk.FindOrAdd(Params.Key, &bIsAlreadyInSet);
+									RequestsOk.FindOrAdd(Params.Key, &bIsAlreadyInSet);
 									if (OnComplete && !bIsAlreadyInSet)
 									{
 										OnComplete(MoveTemp(Params));
@@ -518,54 +517,44 @@ public:
 					}
 					else
 					{
-						AsyncPutInnerBackends[PutCacheIndex]->Put(Records, Context, Policy, AsyncOwner);
+						AsyncPutInnerBackends[PutCacheIndex]->Put(Requests, Context, AsyncOwner);
 					}
 				}
 			}
 		}
 
-		if (OnComplete && RecordsOk.Num() < Records.Num())
+		if (OnComplete && RequestsOk.Num() < Requests.Num())
 		{
-			for (const FCacheRecord& Record : Records)
+			for (const FCachePutRequest& Request : Requests)
 			{
-				if (!RecordsOk.Contains(Record.GetKey()))
+				if (!RequestsOk.Contains(Request.Record.GetKey()))
 				{
-					OnComplete({Record.GetKey(), EStatus::Error});
+					OnComplete({Request.Record.GetKey(), Request.UserData, EStatus::Error});
 				}
 			}
 		}
 	}
 
 	virtual void Get(
-		TConstArrayView<FCacheKey> Keys,
+		TConstArrayView<FCacheGetRequest> Requests,
 		FStringView Context,
-		FCacheRecordPolicy Policy,
 		IRequestOwner& Owner,
 		FOnCacheGetComplete&& OnComplete) override
 	{
-		const bool bQueryLocal = bHasLocalBackends && EnumHasAnyFlags(Policy.GetRecordPolicy(), ECachePolicy::QueryLocal);
-		const bool bStoreLocal = bHasWritableLocalBackends && EnumHasAnyFlags(Policy.GetRecordPolicy(), ECachePolicy::StoreLocal);
-		const bool bQueryRemote = bHasRemoteBackends && EnumHasAnyFlags(Policy.GetRecordPolicy(), ECachePolicy::QueryRemote);
-		const bool bStoreRemote = bHasWritableRemoteBackends && EnumHasAnyFlags(Policy.GetRecordPolicy(), ECachePolicy::StoreRemote);
-		const bool bStoreLocalCopy = bStoreLocal && bHasMultipleLocalBackends;
-
-		TArray<FCacheKey, TInlineAllocator<16>> RemainingKeys(Keys);
+		TArray<FCacheGetRequest, TInlineAllocator<16>> RemainingRequests(Requests);
 
 		{
 			FReadScopeLock LockScope(Lock);
 			TSet<FCacheKey> KeysOk;
 
-			bool bHadLocalMiss = false;
-			bool bHadRemoteMiss = false;
+			bool bHadMiss = false;
 
-			for (int32 GetCacheIndex = 0; GetCacheIndex < InnerBackends.Num() && !RemainingKeys.IsEmpty(); ++GetCacheIndex)
+			for (int32 GetCacheIndex = 0; GetCacheIndex < InnerBackends.Num() && !RemainingRequests.IsEmpty(); ++GetCacheIndex)
 			{
-				const bool bIsLocalGet = InnerBackends[GetCacheIndex]->GetSpeedClass() >= ESpeedClass::Fast;
-
 				// Block on this because backends in this hierarchy are not expected to be asynchronous.
 				FRequestOwner BlockingOwner(EPriority::Blocking);
-				InnerBackends[GetCacheIndex]->Get(RemainingKeys, Context, Policy, BlockingOwner,
-					[this, Context, GetCacheIndex, bStoreLocal, bStoreRemote, bStoreLocalCopy, bIsLocalGet, &Owner, &OnComplete, &KeysOk](FCacheGetCompleteParams&& Params)
+				InnerBackends[GetCacheIndex]->Get(RemainingRequests, Context, BlockingOwner,
+					[this, Context, GetCacheIndex, &Owner, &OnComplete, &KeysOk](FCacheGetCompleteParams&& Params)
 					{
 						if (Params.Status == EStatus::Ok)
 						{
@@ -576,11 +565,7 @@ public:
 							{
 								if (GetCacheIndex != FillCacheIndex)
 								{
-									const bool bIsLocalFill = InnerBackends[FillCacheIndex]->GetSpeedClass() >= ESpeedClass::Fast;
-									if (bIsLocalFill ? bStoreLocal : bStoreRemote)
-									{
-										AsyncPutInnerBackends[FillCacheIndex]->Put({Params.Record}, Context, ECachePolicy::Default, AsyncOwner);
-									}
+									AsyncPutInnerBackends[FillCacheIndex]->Put({{Params.Record, ECachePolicy::Default}}, Context, AsyncOwner);
 								}
 							}
 
@@ -593,35 +578,37 @@ public:
 					});
 				BlockingOwner.Wait();
 
-				RemainingKeys.RemoveAll([&KeysOk](const FCacheKey& Key) { return KeysOk.Contains(Key); });
+				RemainingRequests.RemoveAll([&KeysOk](const FCacheGetRequest& Request) { return KeysOk.Contains(Request.Key); });
 
-				bool& bHadMiss = bIsLocalGet ? bHadLocalMiss : bHadRemoteMiss;
-				if (!bHadMiss && !RemainingKeys.IsEmpty() && InnerBackends[GetCacheIndex]->IsWritable())
+				if (!bHadMiss && !RemainingRequests.IsEmpty() && InnerBackends[GetCacheIndex]->IsWritable())
 				{
 					bHadMiss = true;
-					const auto ConvertPolicy = [bIsLocalGet](ECachePolicy P) -> ECachePolicy
+					const auto ConvertPolicy = [](ECachePolicy P) -> ECachePolicy
 					{
-						if (EnumHasAnyFlags(P, bIsLocalGet ? ECachePolicy::StoreLocal : ECachePolicy::StoreRemote))
+						if (EnumHasAnyFlags(P, ECachePolicy::Store))
 						{
 							EnumRemoveFlags(P, ECachePolicy::SkipData);
 						}
 						return P;
 					};
-					FCacheRecordPolicyBuilder Builder(ConvertPolicy(Policy.GetDefaultPayloadPolicy()));
-					for (const FCachePayloadPolicy& PayloadPolicy : Policy.GetPayloadPolicies())
+					for (FCacheGetRequest& Request : RemainingRequests)
 					{
-						Builder.AddPayloadPolicy(PayloadPolicy.Id, ConvertPolicy(PayloadPolicy.Policy));
+						FCacheRecordPolicyBuilder Builder(ConvertPolicy(Request.Policy.GetDefaultPayloadPolicy()));
+						for (const FCachePayloadPolicy& PayloadPolicy : Request.Policy.GetPayloadPolicies())
+						{
+							Builder.AddPayloadPolicy(PayloadPolicy.Id, ConvertPolicy(PayloadPolicy.Policy));
+						}
+						Request.Policy = Builder.Build();
 					}
-					Policy = Builder.Build();
 				}
 			}
 		}
 
 		if (OnComplete)
 		{
-			for (const FCacheKey& Key : RemainingKeys)
+			for (const FCacheGetRequest& Request : RemainingRequests)
 			{
-				OnComplete({FCacheRecordBuilder(Key).Build(), EStatus::Error});
+				OnComplete({FCacheRecordBuilder(Request.Key).Build(), Request.UserData, EStatus::Error});
 			}
 		}
 	}
@@ -671,7 +658,7 @@ public:
 		{
 			for (const FCacheChunkRequest& Chunk : RemainingChunks)
 			{
-				OnComplete({Chunk.Key, Chunk.Id, Chunk.RawOffset, 0, {}, {}, EStatus::Error});
+				OnComplete({Chunk.Key, Chunk.Id, Chunk.RawOffset, 0, {}, {}, Chunk.UserData, EStatus::Error});
 			}
 		}
 	}
diff --git a/Engine/Source/Developer/DerivedDataCache/Private/HttpDerivedDataBackend.cpp b/Engine/Source/Developer/DerivedDataCache/Private/HttpDerivedDataBackend.cpp
--- a/Engine/Source/Developer/DerivedDataCache/Private/HttpDerivedDataBackend.cpp
+++ b/Engine/Source/Developer/DerivedDataCache/Private/HttpDerivedDataBackend.cpp
@@ -2532,40 +2532,38 @@ TSharedRef<FDerivedDataCacheStatsNode> FHttpDerivedDataBackend::GatherUsageStats
 }
 
 void FHttpDerivedDataBackend::Put(
-	TConstArrayView<FCacheRecord> Records,
-	FStringView Context,
-	ECachePolicy Policy,
+	const TConstArrayView<FCachePutRequest> Requests,
+	const FStringView Context,
 	IRequestOwner& Owner,
 	FOnCachePutComplete&& OnComplete)
 {
 	if (OnComplete)
 	{
-		for (const FCacheRecord& Record : Records)
+		for (const FCachePutRequest& Request : Requests)
 		{
-			OnComplete({Record.GetKey(), EStatus::Error});
+			OnComplete({Request.Record.GetKey(), Request.UserData, EStatus::Error});
 		}
 	}
 }
 
 void FHttpDerivedDataBackend::Get(
-	TConstArrayView<FCacheKey> Keys,
-	FStringView Context,
-	FCacheRecordPolicy Policy,
+	const TConstArrayView<FCacheGetRequest> Requests,
+	const FStringView Context,
 	IRequestOwner& Owner,
 	FOnCacheGetComplete&& OnComplete)
 {
 	if (OnComplete)
 	{
-		for (const FCacheKey& Key : Keys)
+		for (const FCacheGetRequest& Request : Requests)
 		{
-			OnComplete({FCacheRecordBuilder(Key).Build(), EStatus::Error});
+			OnComplete({FCacheRecordBuilder(Request.Key).Build(), Request.UserData, EStatus::Error});
 		}
 	}
 }
 
 void FHttpDerivedDataBackend::GetChunks(
-	TConstArrayView<FCacheChunkRequest> Chunks,
-	FStringView Context,
+	const TConstArrayView<FCacheChunkRequest> Chunks,
+	const FStringView Context,
 	IRequestOwner& Owner,
 	FOnCacheGetChunkComplete&& OnComplete)
 {
@@ -2573,7 +2571,7 @@ void FHttpDerivedDataBackend::GetChunks(
 	{
 		for (const FCacheChunkRequest& Chunk : Chunks)
 		{
-			OnComplete({Chunk.Key, Chunk.Id, Chunk.RawOffset, 0, {}, {}, EStatus::Error});
+			OnComplete({Chunk.Key, Chunk.Id, Chunk.RawOffset, 0, {}, {}, Chunk.UserData, EStatus::Error});
 		}
 	}
 }
diff --git a/Engine/Source/Developer/DerivedDataCache/Private/HttpDerivedDataBackend.h b/Engine/Source/Developer/DerivedDataCache/Private/HttpDerivedDataBackend.h
--- a/Engine/Source/Developer/DerivedDataCache/Private/HttpDerivedDataBackend.h
+++ b/Engine/Source/Developer/DerivedDataCache/Private/HttpDerivedDataBackend.h
@@ -77,16 +77,14 @@ public:
 	void SetSpeedClass(ESpeedClass InSpeedClass) { SpeedClass = InSpeedClass; }
 
 	virtual void Put(
-		TConstArrayView<FCacheRecord> Records,
+		TConstArrayView<FCachePutRequest> Requests,
 		FStringView Context,
-		ECachePolicy Policy,
 		IRequestOwner& Owner,
 		FOnCachePutComplete&& OnComplete) override;
 
 	virtual void Get(
-		TConstArrayView<FCacheKey> Keys,
+		TConstArrayView<FCacheGetRequest> Requests,
 		FStringView Context,
-		FCacheRecordPolicy Policy,
 		IRequestOwner& Owner,
 		FOnCacheGetComplete&& OnComplete) override;
 
diff --git a/Engine/Source/Developer/DerivedDataCache/Private/MemoryDerivedDataBackend.cpp b/Engine/Source/Developer/DerivedDataCache/Private/MemoryDerivedDataBackend.cpp
--- a/Engine/Source/Developer/DerivedDataCache/Private/MemoryDerivedDataBackend.cpp
+++ b/Engine/Source/Developer/DerivedDataCache/Private/MemoryDerivedDataBackend.cpp
@@ -467,21 +467,21 @@ int64 FMemoryDerivedDataBackend::CalcSerializedCacheRecordSize(const FCacheRecor
 }
 
 void FMemoryDerivedDataBackend::Put(
-	TConstArrayView<FCacheRecord> Records,
-	FStringView Context,
-	ECachePolicy Policy,
+	const TConstArrayView<FCachePutRequest> Requests,
+	const FStringView Context,
 	IRequestOwner& Owner,
 	FOnCachePutComplete&& OnComplete)
 {
-	for (const FCacheRecord& Record : Records)
+	for (const FCachePutRequest& Request : Requests)
 	{
+		const FCacheRecord& Record = Request.Record;
 		const FCacheKey& Key = Record.GetKey();
 		EStatus Status = EStatus::Error;
 		ON_SCOPE_EXIT
 		{
 			if (OnComplete)
 			{
-				OnComplete({Key, Status});
+				OnComplete({Key, Request.UserData, Status});
 			}
 		};
 
@@ -541,14 +541,15 @@ void FMemoryDerivedDataBackend::Put(
 }
 
 void FMemoryDerivedDataBackend::Get(
-	TConstArrayView<FCacheKey> Keys,
-	FStringView Context,
-	FCacheRecordPolicy Policy,
+	const TConstArrayView<FCacheGetRequest> Requests,
+	const FStringView Context,
 	IRequestOwner& Owner,
 	FOnCacheGetComplete&& OnComplete)
 {
-	for (const FCacheKey& Key : Keys)
+	for (const FCacheGetRequest& Request : Requests)
 	{
+		const FCacheKey& Key = Request.Key;
+		const FCacheRecordPolicy& Policy = Request.Policy;
 		const bool bExistsOnly = EnumHasAllFlags(Policy.GetRecordPolicy(), ECachePolicy::SkipData);
 		COOK_STAT(auto Timer = bExistsOnly ? UsageStats.TimeProbablyExists() : UsageStats.TimeGet());
 		FOptionalCacheRecord Record;
@@ -595,22 +596,22 @@ void FMemoryDerivedDataBackend::Get(
 			COOK_STAT(Timer.AddHit(CalcRawCacheRecordSize(Record.Get())));
 			if (OnComplete)
 			{
-				OnComplete({MoveTemp(Record).Get(), Status});
+				OnComplete({MoveTemp(Record).Get(), Request.UserData, Status});
 			}
 		}
 		else
 		{
 			if (OnComplete)
 			{
-				OnComplete({FCacheRecordBuilder(Key).Build(), EStatus::Error});
+				OnComplete({FCacheRecordBuilder(Key).Build(), Request.UserData, EStatus::Error});
 			}
 		}
 	}
 }
 
 void FMemoryDerivedDataBackend::GetChunks(
-	TConstArrayView<FCacheChunkRequest> Chunks,
-	FStringView Context,
+	const TConstArrayView<FCacheChunkRequest> Chunks,
+	const FStringView Context,
 	IRequestOwner& Owner,
 	FOnCacheGetChunkComplete&& OnComplete)
 {
@@ -651,14 +652,14 @@ void FMemoryDerivedDataBackend::GetChunks(
 				}
 				const EStatus Status = bExistsOnly || Buffer ? EStatus::Ok : EStatus::Error;
 				OnComplete({Chunk.Key, Chunk.Id, Chunk.RawOffset,
-					RawSize, Payload.GetRawHash(), MoveTemp(Buffer), Status});
+					RawSize, Payload.GetRawHash(), MoveTemp(Buffer), Chunk.UserData, Status});
 			}
 		}
 		else
 		{
 			if (OnComplete)
 			{
-				OnComplete({Chunk.Key, Chunk.Id, Chunk.RawOffset, 0, {}, {}, EStatus::Error});
+				OnComplete({Chunk.Key, Chunk.Id, Chunk.RawOffset, 0, {}, {}, Chunk.UserData, EStatus::Error});
 			}
 		}
 	}
diff --git a/Engine/Source/Developer/DerivedDataCache/Private/MemoryDerivedDataBackend.h b/Engine/Source/Developer/DerivedDataCache/Private/MemoryDerivedDataBackend.h
--- a/Engine/Source/Developer/DerivedDataCache/Private/MemoryDerivedDataBackend.h
+++ b/Engine/Source/Developer/DerivedDataCache/Private/MemoryDerivedDataBackend.h
@@ -98,16 +98,14 @@ public:
 	bool ApplyDebugOptions(FBackendDebugOptions& InOptions) override;
 
 	virtual void Put(
-		TConstArrayView<FCacheRecord> Records,
+		TConstArrayView<FCachePutRequest> Requests,
 		FStringView Context,
-		ECachePolicy Policy,
 		IRequestOwner& Owner,
 		FOnCachePutComplete&& OnComplete) override;
 
 	virtual void Get(
-		TConstArrayView<FCacheKey> Keys,
+		TConstArrayView<FCacheGetRequest> Requests,
 		FStringView Context,
-		FCacheRecordPolicy Policy,
 		IRequestOwner& Owner,
 		FOnCacheGetComplete&& OnComplete) override;
 
diff --git a/Engine/Source/Developer/DerivedDataCache/Private/PakFileDerivedDataBackend.cpp b/Engine/Source/Developer/DerivedDataCache/Private/PakFileDerivedDataBackend.cpp
--- a/Engine/Source/Developer/DerivedDataCache/Private/PakFileDerivedDataBackend.cpp
+++ b/Engine/Source/Developer/DerivedDataCache/Private/PakFileDerivedDataBackend.cpp
@@ -482,24 +482,24 @@ TSharedRef<FDerivedDataCacheStatsNode> FPakFileDerivedDataBackend::GatherUsageSt
 }
 
 void FPakFileDerivedDataBackend::Put(
-	const TConstArrayView<FCacheRecord> Records,
+	const TConstArrayView<FCachePutRequest> Requests,
 	const FStringView Context,
-	const ECachePolicy Policy,
 	IRequestOwner& Owner,
 	FOnCachePutComplete&& OnComplete)
 {
-	for (const FCacheRecord& Record : Records)
+	for (const FCachePutRequest& Request : Requests)
 	{
+		const FCacheRecord& Record = Request.Record;
 		TRACE_CPUPROFILER_EVENT_SCOPE(PakFileDDC_Put);
 		COOK_STAT(auto Timer = UsageStats.TimePut());
-		if (PutCacheRecord(Record, Context, Policy))
+		if (PutCacheRecord(Record, Context, Request.Policy))
 		{
 			UE_LOG(LogDerivedDataCache, Verbose, TEXT("%s: Cache put complete for %s from '%.*s'"),
 				*CachePath, *WriteToString<96>(Record.GetKey()), Context.Len(), Context.GetData());
 			COOK_STAT(Timer.AddHit(MeasureRawCacheRecord(Record)));
 			if (OnComplete)
 			{
-				OnComplete({Record.GetKey(), EStatus::Ok});
+				OnComplete({Record.GetKey(), Request.UserData, EStatus::Ok});
 			}
 		}
 		else
@@ -507,39 +507,38 @@ void FPakFileDerivedDataBackend::Put(
 			COOK_STAT(Timer.AddMiss(MeasureRawCacheRecord(Record)));
 			if (OnComplete)
 			{
-				OnComplete({Record.GetKey(), EStatus::Error});
+				OnComplete({Record.GetKey(), Request.UserData, EStatus::Error});
 			}
 		}
 	}
 }
 
 void FPakFileDerivedDataBackend::Get(
-	const TConstArrayView<FCacheKey> Keys,
+	const TConstArrayView<FCacheGetRequest> Requests,
 	const FStringView Context,
-	const FCacheRecordPolicy Policy,
 	IRequestOwner& Owner,
 	FOnCacheGetComplete&& OnComplete)
 {
-	for (const FCacheKey& Key : Keys)
+	for (const FCacheGetRequest& Request : Requests)
 	{
 		TRACE_CPUPROFILER_EVENT_SCOPE(PakFileDDC_Get);
 		COOK_STAT(auto Timer = UsageStats.TimeGet());
 		EStatus Status = EStatus::Ok;
-		if (FOptionalCacheRecord Record = GetCacheRecord(Key, Context, Policy, Status))
+		if (FOptionalCacheRecord Record = GetCacheRecord(Request.Key, Context, Request.Policy, Status))
 		{
 			UE_LOG(LogDerivedDataCache, Verbose, TEXT("%s: Cache hit for %s from '%.*s'"),
-				*CachePath, *WriteToString<96>(Key), Context.Len(), Context.GetData());
+				*CachePath, *WriteToString<96>(Request.Key), Context.Len(), Context.GetData());
 			COOK_STAT(Timer.AddHit(MeasureRawCacheRecord(Record.Get())));
 			if (OnComplete)
 			{
-				OnComplete({MoveTemp(Record).Get(), Status});
+				OnComplete({MoveTemp(Record).Get(), Request.UserData, Status});
 			}
 		}
 		else
 		{
 			if (OnComplete)
 			{
-				OnComplete({FCacheRecordBuilder(Key).Build(), Status});
+				OnComplete({FCacheRecordBuilder(Request.Key).Build(), Request.UserData, Status});
 			}
 		}
 	}
@@ -590,7 +589,7 @@ void FPakFileDerivedDataBackend::GetChunks(
 						Buffer = Reader.Decompress(RawOffset, RawSize);
 					}
 					OnComplete({Chunk.Key, Chunk.Id, Chunk.RawOffset,
-						RawSize, Payload.GetRawHash(), MoveTemp(Buffer), PayloadStatus});
+						RawSize, Payload.GetRawHash(), MoveTemp(Buffer), Chunk.UserData, PayloadStatus});
 				}
 				continue;
 			}
@@ -598,7 +597,7 @@ void FPakFileDerivedDataBackend::GetChunks(
 
 		if (OnComplete)
 		{
-			OnComplete({Chunk.Key, Chunk.Id, Chunk.RawOffset, 0, {}, {}, EStatus::Error});
+			OnComplete({Chunk.Key, Chunk.Id, Chunk.RawOffset, 0, {}, {}, Chunk.UserData, EStatus::Error});
 		}
 	}
 }
@@ -622,7 +621,7 @@ uint64 FPakFileDerivedDataBackend::MeasureRawCacheRecord(const FCacheRecord& Rec
 bool FPakFileDerivedDataBackend::PutCacheRecord(
 	const FCacheRecord& Record,
 	const FStringView Context,
-	const ECachePolicy Policy)
+	const FCacheRecordPolicy& Policy)
 {
 	if (!IsWritable())
 	{
@@ -632,7 +631,7 @@ bool FPakFileDerivedDataBackend::PutCacheRecord(
 	const FCacheKey& Key = Record.GetKey();
 
 	// Skip the request if storing to the cache is disabled.
-	if (!EnumHasAnyFlags(Policy, ECachePolicy::StoreLocal))
+	if (!EnumHasAnyFlags(Policy.GetRecordPolicy(), ECachePolicy::StoreLocal))
 	{
 		UE_LOG(LogDerivedDataCache, VeryVerbose, TEXT("%s: Skipped put of %s from '%.*s' due to cache policy"),
 			*CachePath, *WriteToString<96>(Key), Context.Len(), Context.GetData());
@@ -651,7 +650,9 @@ bool FPakFileDerivedDataBackend::PutCacheRecord(
 	FCbPackage ExistingPackage;
 	TStringBuilder<256> Path;
 	FPathViews::Append(Path, TEXT("Buckets"), Key);
-	if (EnumHasAllFlags(Policy, ECachePolicy::SkipValue | ECachePolicy::SkipAttachments))
+	const ECachePolicy CombinedPayloadPolicy = Algo::TransformAccumulate(
+		Policy.GetPayloadPolicies(), &FCachePayloadPolicy::Policy, Policy.GetDefaultPayloadPolicy(), UE_PROJECTION(operator|));
+	if (EnumHasAllFlags(CombinedPayloadPolicy, ECachePolicy::SkipValue | ECachePolicy::SkipAttachments))
 	{
 		bRecordExists = FileExists(Path);
 	}
diff --git a/Engine/Source/Developer/DerivedDataCache/Private/PakFileDerivedDataBackend.h b/Engine/Source/Developer/DerivedDataCache/Private/PakFileDerivedDataBackend.h
--- a/Engine/Source/Developer/DerivedDataCache/Private/PakFileDerivedDataBackend.h
+++ b/Engine/Source/Developer/DerivedDataCache/Private/PakFileDerivedDataBackend.h
@@ -108,16 +108,14 @@ public:
 	virtual bool ApplyDebugOptions(FBackendDebugOptions& InOptions) override { return false; }
 
 	virtual void Put(
-		TConstArrayView<FCacheRecord> Records,
+		TConstArrayView<FCachePutRequest> Requests,
 		FStringView Context,
-		ECachePolicy Policy,
 		IRequestOwner& Owner,
 		FOnCachePutComplete&& OnComplete) override;
 
 	virtual void Get(
-		TConstArrayView<FCacheKey> Keys,
+		TConstArrayView<FCacheGetRequest> Requests,
 		FStringView Context,
-		FCacheRecordPolicy Policy,
 		IRequestOwner& Owner,
 		FOnCacheGetComplete&& OnComplete) override;
 
@@ -131,7 +129,7 @@ private:
 	uint64 MeasureCompressedCacheRecord(const FCacheRecord& Record) const;
 	uint64 MeasureRawCacheRecord(const FCacheRecord& Record) const;
 
-	bool PutCacheRecord(const FCacheRecord& Record, FStringView Context, ECachePolicy Policy);
+	bool PutCacheRecord(const FCacheRecord& Record, FStringView Context, const FCacheRecordPolicy& Policy);
 	bool PutCacheContent(const FCompressedBuffer& Content, const FStringView Context);
 
 	FOptionalCacheRecord GetCacheRecordOnly(
diff --git a/Engine/Source/Developer/DerivedDataCache/Private/S3DerivedDataBackend.cpp b/Engine/Source/Developer/DerivedDataCache/Private/S3DerivedDataBackend.cpp
--- a/Engine/Source/Developer/DerivedDataCache/Private/S3DerivedDataBackend.cpp
+++ b/Engine/Source/Developer/DerivedDataCache/Private/S3DerivedDataBackend.cpp
@@ -1202,40 +1202,38 @@ bool FS3DerivedDataBackend::ShouldSimulateMiss(const TCHAR* InKey)
 }
 
 void FS3DerivedDataBackend::Put(
-	TConstArrayView<FCacheRecord> Records,
-	FStringView Context,
-	ECachePolicy Policy,
+	const TConstArrayView<FCachePutRequest> Requests,
+	const FStringView Context,
 	IRequestOwner& Owner,
 	FOnCachePutComplete&& OnComplete)
 {
 	if (OnComplete)
 	{
-		for (const FCacheRecord& Record : Records)
+		for (const FCachePutRequest& Request : Requests)
 		{
-			OnComplete({Record.GetKey(), EStatus::Error});
+			OnComplete({Request.Record.GetKey(), Request.UserData, EStatus::Error});
 		}
 	}
 }
 
 void FS3DerivedDataBackend::Get(
-	TConstArrayView<FCacheKey> Keys,
-	FStringView Context,
-	FCacheRecordPolicy Policy,
+	const TConstArrayView<FCacheGetRequest> Requests,
+	const FStringView Context,
 	IRequestOwner& Owner,
 	FOnCacheGetComplete&& OnComplete)
 {
 	if (OnComplete)
 	{
-		for (const FCacheKey& Key : Keys)
+		for (const FCacheGetRequest& Request : Requests)
 		{
-			OnComplete({FCacheRecordBuilder(Key).Build(), EStatus::Error});
+			OnComplete({FCacheRecordBuilder(Request.Key).Build(), Request.UserData, EStatus::Error});
 		}
 	}
 }
 
 void FS3DerivedDataBackend::GetChunks(
-	TConstArrayView<FCacheChunkRequest> Chunks,
-	FStringView Context,
+	const TConstArrayView<FCacheChunkRequest> Chunks,
+	const FStringView Context,
 	IRequestOwner& Owner,
 	FOnCacheGetChunkComplete&& OnComplete)
 {
@@ -1243,7 +1241,7 @@ void FS3DerivedDataBackend::GetChunks(
 	{
 		for (const FCacheChunkRequest& Chunk : Chunks)
 		{
-			OnComplete({Chunk.Key, Chunk.Id, Chunk.RawOffset, 0, {}, {}, EStatus::Error});
+			OnComplete({Chunk.Key, Chunk.Id, Chunk.RawOffset, 0, {}, {}, Chunk.UserData, EStatus::Error});
 		}
 	}
 }
diff --git a/Engine/Source/Developer/DerivedDataCache/Private/S3DerivedDataBackend.h b/Engine/Source/Developer/DerivedDataCache/Private/S3DerivedDataBackend.h
--- a/Engine/Source/Developer/DerivedDataCache/Private/S3DerivedDataBackend.h
+++ b/Engine/Source/Developer/DerivedDataCache/Private/S3DerivedDataBackend.h
@@ -64,16 +64,14 @@ public:
 	bool ApplyDebugOptions(FBackendDebugOptions& InOptions) override;
 
 	virtual void Put(
-		TConstArrayView<FCacheRecord> Records,
+		TConstArrayView<FCachePutRequest> Requests,
 		FStringView Context,
-		ECachePolicy Policy,
 		IRequestOwner& Owner,
 		FOnCachePutComplete&& OnComplete) override;
 
 	virtual void Get(
-		TConstArrayView<FCacheKey> Keys,
+		TConstArrayView<FCacheGetRequest> Requests,
 		FStringView Context,
-		FCacheRecordPolicy Policy,
 		IRequestOwner& Owner,
 		FOnCacheGetComplete&& OnComplete) override;
 
diff --git a/Engine/Source/Developer/DerivedDataCache/Private/ZenDerivedDataBackend.cpp b/Engine/Source/Developer/DerivedDataCache/Private/ZenDerivedDataBackend.cpp
--- a/Engine/Source/Developer/DerivedDataCache/Private/ZenDerivedDataBackend.cpp
+++ b/Engine/Source/Developer/DerivedDataCache/Private/ZenDerivedDataBackend.cpp
@@ -553,17 +553,17 @@ bool FZenDerivedDataBackend::ShouldSimulateMiss(const FCacheKey& Key)
 }
 
 void FZenDerivedDataBackend::Put(
-	TConstArrayView<FCacheRecord> Records,
-	FStringView Context,
-	ECachePolicy Policy,
+	const TConstArrayView<FCachePutRequest> Requests,
+	const FStringView Context,
 	IRequestOwner& Owner,
 	FOnCachePutComplete&& OnComplete)
 {
 	TRACE_CPUPROFILER_EVENT_SCOPE(ZenDDC::PutCachedRecord);
 
-	for (const FCacheRecord& Record : Records)
+	for (const FCachePutRequest& Request : Requests)
 	{
 		COOK_STAT(auto Timer = UsageStats.TimePut());
+		const FCacheRecord& Record = Request.Record;
 		bool bResult;
 		if (ShouldSimulateMiss(Record.GetKey()))
 		{
@@ -572,7 +572,7 @@ void FZenDerivedDataBackend::Put(
 			bResult = false;
 		}
 
-		bResult = PutCacheRecord(Record, Context, Policy);
+		bResult = PutCacheRecord(Record, Context, Request.Policy);
 
 		if (bResult)
 		{
@@ -581,7 +581,7 @@ void FZenDerivedDataBackend::Put(
 			COOK_STAT(Timer.AddHit(MeasureCacheRecord(Record)));
 			if (OnComplete)
 			{
-				OnComplete({ Record.GetKey(), EStatus::Ok });
+				OnComplete({ Record.GetKey(), Request.UserData, EStatus::Ok });
 			}
 		}
 		else
@@ -589,26 +589,25 @@ void FZenDerivedDataBackend::Put(
 			COOK_STAT(Timer.AddMiss(MeasureCacheRecord(Record)));
 			if (OnComplete)
 			{
-				OnComplete({ Record.GetKey(), EStatus::Error });
+				OnComplete({ Record.GetKey(), Request.UserData, EStatus::Error });
 			}
 		}
 	}
 }
 
 void FZenDerivedDataBackend::Get(
-	TConstArrayView<FCacheKey> Keys,
-	FStringView Context,
-	FCacheRecordPolicy Policy,
+	const TConstArrayView<FCacheGetRequest> Requests,
+	const FStringView Context,
 	IRequestOwner& Owner,
 	FOnCacheGetComplete&& OnComplete)
 {
 	TRACE_CPUPROFILER_EVENT_SCOPE(ZenDDC::GetCacheRecord);
-	TRACE_COUNTER_ADD(ZenDDC_CacheRecordRequestCount, int64(Keys.Num()));
+	TRACE_COUNTER_ADD(ZenDDC_CacheRecordRequestCount, int64(Requests.Num()));
 
 	int32 TotalCompleted = 0;
 
-	ForEachBatch(CacheRecordBatchSize, Keys.Num(),
-		[this, &Keys, &Context, &Policy, &Owner, &OnComplete, &TotalCompleted](int32 BatchFirst, int32 BatchLast)
+	ForEachBatch(CacheRecordBatchSize, Requests.Num(),
+		[this, &Requests, &Context, &Owner, &OnComplete, &TotalCompleted](int32 BatchFirst, int32 BatchLast)
 	{
 		TRACE_COUNTER_ADD(ZenDDC_Get, int64(BatchLast) - int64(BatchFirst) + 1);
 		COOK_STAT(auto Timer = UsageStats.TimeGet());
@@ -622,7 +621,7 @@ void FZenDerivedDataBackend::Get(
 				BatchRequest.BeginArray("CacheKeys"_ASV);
 				for (int32 KeyIndex = BatchFirst; KeyIndex <= BatchLast; KeyIndex++)
 				{
-					const FCacheKey& Key = Keys[KeyIndex];
+					const FCacheKey& Key = Requests[KeyIndex].Key;
 
 					BatchRequest.BeginObject();
 					BatchRequest << "Bucket"_ASV << Key.Bucket.ToString();
@@ -631,6 +630,8 @@ void FZenDerivedDataBackend::Get(
 				}
 				BatchRequest.EndArray();
 
+				// TODO: The policy needs to be sent with each key.
+				const FCacheRecordPolicy& Policy = Requests[BatchFirst].Policy;
 				BatchRequest.BeginObject("Policy"_ASV);
 				{
 					BatchRequest << "RecordPolicy"_ASV << static_cast<uint32>(Policy.GetRecordPolicy());
@@ -671,7 +672,8 @@ void FZenDerivedDataBackend::Get(
 			int32 KeyIndex = BatchFirst;
 			for (FCbFieldView RecordView : ResponseObj["Result"_ASV])
 			{
-				const FCacheKey& Key = Keys[KeyIndex++];
+				const FCacheGetRequest& Request = Requests[KeyIndex++];
+				const FCacheKey& Key = Request.Key;
 				
 				FOptionalCacheRecord Record;
 
@@ -700,7 +702,7 @@ void FZenDerivedDataBackend::Get(
 
 					if (OnComplete)
 					{
-						OnComplete({ MoveTemp(Record).Get(), EStatus::Ok });
+						OnComplete({ MoveTemp(Record).Get(), Request.UserData, EStatus::Ok });
 					}
 				}
 				else
@@ -710,7 +712,7 @@ void FZenDerivedDataBackend::Get(
 					
 					if (OnComplete)
 					{
-						OnComplete({ FCacheRecordBuilder(Key).Build(), EStatus::Error });
+						OnComplete({ FCacheRecordBuilder(Key).Build(), Request.UserData, EStatus::Error });
 					}
 				}
 				
@@ -721,14 +723,15 @@ void FZenDerivedDataBackend::Get(
 		{
 			for (int32 KeyIndex = BatchFirst; KeyIndex <= BatchLast; KeyIndex++)
 			{
-				const FCacheKey& Key = Keys[KeyIndex];
+				const FCacheGetRequest& Request = Requests[KeyIndex];
+				const FCacheKey& Key = Request.Key;
 
 				UE_LOG(LogDerivedDataCache, Verbose, TEXT("%s: Cache miss for '%s' from '%.*s'"),
 					*GetName(), *WriteToString<96>(Key), Context.Len(), Context.GetData());
 					
 				if (OnComplete)
 				{
-					OnComplete({ FCacheRecordBuilder(Key).Build(), EStatus::Error });
+					OnComplete({ FCacheRecordBuilder(Key).Build(), Request.UserData, EStatus::Error });
 				}
 				
 				TotalCompleted++;
@@ -736,8 +739,8 @@ void FZenDerivedDataBackend::Get(
 		}
 	});
 	
-	UE_CLOG(TotalCompleted != Keys.Num(), LogDerivedDataCache, Warning, TEXT("Only '%d/%d' cache record request(s) completed"), TotalCompleted, Keys.Num());
-	TRACE_COUNTER_SUBTRACT(ZenDDC_CacheRecordRequestCount, int64(Keys.Num()));
+	UE_CLOG(TotalCompleted != Requests.Num(), LogDerivedDataCache, Warning, TEXT("Only '%d/%d' cache record request(s) completed"), TotalCompleted, Requests.Num());
+	TRACE_COUNTER_SUBTRACT(ZenDDC_CacheRecordRequestCount, int64(Requests.Num()));
 }
 
 void FZenDerivedDataBackend::GetChunks(
@@ -815,7 +818,7 @@ void FZenDerivedDataBackend::GetChunks(
 					
 					if (OnComplete)
 					{
-						OnComplete({ Chunk.Key, Chunk.Id, Chunk.RawOffset, 0, {}, {}, EStatus::Error });
+						OnComplete({ Chunk.Key, Chunk.Id, Chunk.RawOffset, 0, {}, {}, Chunk.UserData, EStatus::Error });
 					}
 				}
 				else if (const FCbAttachment* Attachment = BatchResponse.FindAttachment(HashView.AsHash()))
@@ -833,7 +836,7 @@ void FZenDerivedDataBackend::GetChunks(
 					
 					if (OnComplete)
 					{
-						OnComplete({Chunk.Key, Chunk.Id, Chunk.RawOffset, RawSize, CompressedBuffer.GetRawHash(), MoveTemp(Buffer), EStatus::Ok});
+						OnComplete({Chunk.Key, Chunk.Id, Chunk.RawOffset, RawSize, CompressedBuffer.GetRawHash(), MoveTemp(Buffer), Chunk.UserData, EStatus::Ok});
 					}
 				}
 				else
@@ -843,7 +846,7 @@ void FZenDerivedDataBackend::GetChunks(
 					
 					if (OnComplete)
 					{
-						OnComplete({ Chunk.Key, Chunk.Id, Chunk.RawOffset, 0, {}, {}, EStatus::Error });
+						OnComplete({ Chunk.Key, Chunk.Id, Chunk.RawOffset, 0, {}, {}, Chunk.UserData, EStatus::Error });
 					}
 				}
 
@@ -861,7 +864,7 @@ void FZenDerivedDataBackend::GetChunks(
 				
 				if (OnComplete)
 				{
-					OnComplete({ Chunk.Key, Chunk.Id, Chunk.RawOffset, 0, {}, {}, EStatus::Error });
+					OnComplete({ Chunk.Key, Chunk.Id, Chunk.RawOffset, 0, {}, {}, Chunk.UserData, EStatus::Error });
 				}
 				
 				TotalCompleted++;
@@ -873,7 +876,7 @@ void FZenDerivedDataBackend::GetChunks(
 	TRACE_COUNTER_SUBTRACT(ZenDDC_ChunkRequestCount, int64(Chunks.Num()));
 }
 
-bool FZenDerivedDataBackend::PutCacheRecord(const FCacheRecord& Record, FStringView Context, ECachePolicy Policy)
+bool FZenDerivedDataBackend::PutCacheRecord(const FCacheRecord& Record, FStringView Context, const FCacheRecordPolicy& Policy)
 {
 	const FCacheKey& Key = Record.GetKey();
 	FCbPackage Package = Record.Save();
@@ -882,7 +885,7 @@ bool FZenDerivedDataBackend::PutCacheRecord(const FCacheRecord& Record, FStringV
 	FCompositeBuffer Buffer = FCompositeBuffer(FSharedBuffer::MakeView(Ar.GetData(), Ar.Num()));
 	TStringBuilder<256> Uri;
 	AppendZenUri(Record.GetKey(), Uri);
-	AppendPolicyQueryString(Policy, Uri);
+	AppendPolicyQueryString(Policy.GetRecordPolicy(), Uri);
 	if (PutZenData(Uri.ToString(), Buffer, Zen::EContentType::CbPackage)
 		!= FDerivedDataBackendInterface::EPutStatus::Cached)
 	{
diff --git a/Engine/Source/Developer/DerivedDataCache/Private/ZenDerivedDataBackend.h b/Engine/Source/Developer/DerivedDataCache/Private/ZenDerivedDataBackend.h
--- a/Engine/Source/Developer/DerivedDataCache/Private/ZenDerivedDataBackend.h
+++ b/Engine/Source/Developer/DerivedDataCache/Private/ZenDerivedDataBackend.h
@@ -93,16 +93,14 @@ public:
 	// ICacheStore
 
 	virtual void Put(
-		TConstArrayView<FCacheRecord> Records,
+		TConstArrayView<FCachePutRequest> Requests,
 		FStringView Context,
-		ECachePolicy Policy,
 		IRequestOwner& Owner,
 		FOnCachePutComplete&& OnComplete = FOnCachePutComplete()) override;
 
 	virtual void Get(
-		TConstArrayView<FCacheKey> Keys,
+		TConstArrayView<FCacheGetRequest> Requests,
 		FStringView Context,
-		FCacheRecordPolicy Policy,
 		IRequestOwner& Owner,
 		FOnCacheGetComplete&& OnComplete) override;
 
@@ -125,7 +123,7 @@ private:
 	FDerivedDataBackendInterface::EPutStatus PutZenData(const TCHAR* Uri, const FCompositeBuffer& InData, Zen::EContentType ContentType);
 	EGetResult GetZenData(const FCacheKey& Key, ECachePolicy CachePolicy, FCbPackage& OutPackage) const;
 
-	bool PutCacheRecord(const FCacheRecord& Record, FStringView Context, ECachePolicy Policy);
+	bool PutCacheRecord(const FCacheRecord& Record, FStringView Context, const FCacheRecordPolicy& Policy);
 
 	bool IsServiceReady();
 	static FString MakeLegacyZenKey(const TCHAR* CacheKey);
diff --git a/Engine/Source/Developer/DerivedDataCache/Public/DerivedDataCache.h b/Engine/Source/Developer/DerivedDataCache/Public/DerivedDataCache.h
--- a/Engine/Source/Developer/DerivedDataCache/Public/DerivedDataCache.h
+++ b/Engine/Source/Developer/DerivedDataCache/Public/DerivedDataCache.h
@@ -5,6 +5,7 @@
 #include "CoreTypes.h"
 #include "Containers/StringFwd.h"
 #include "DerivedDataCacheKey.h"
+#include "DerivedDataCacheRecord.h"
 #include "DerivedDataPayloadId.h"
 #include "DerivedDataRequestTypes.h"
 #include "Math/NumericLimits.h"
@@ -14,14 +15,14 @@
 
 #define UE_API DERIVEDDATACACHE_API
 
-namespace UE::DerivedData { class FCacheRecord; }
-namespace UE::DerivedData { class FCacheRecordPolicy; }
 namespace UE::DerivedData { class ICacheStoreMaintainer; }
 namespace UE::DerivedData { class IRequestOwner; }
 namespace UE::DerivedData { struct FCacheChunkRequest; }
 namespace UE::DerivedData { struct FCacheGetChunkCompleteParams; }
 namespace UE::DerivedData { struct FCacheGetCompleteParams; }
+namespace UE::DerivedData { struct FCacheGetRequest; }
 namespace UE::DerivedData { struct FCachePutCompleteParams; }
+namespace UE::DerivedData { struct FCachePutRequest; }
 namespace UE::DerivedData::Private { class ICacheRecordPolicyShared; }
 
 namespace UE::DerivedData
@@ -235,16 +236,14 @@ public:
 	 * must reference an existing copy of the payload in the store, if available, and must otherwise
 	 * be stored as a partial record that can attempt recovery of missing payloads when fetched.
 	 *
-	 * @param Records      The cache records to store. Must have a key.
+	 * @param Requests     Requests with the cache records to store. Records must have a key.
 	 * @param Context      A description of the request. An object path is typically sufficient.
-	 * @param Policy       Flags to control the behavior of the request. See ECachePolicy.
 	 * @param Owner        The owner to execute the request within.
 	 * @param OnComplete   A callback invoked for every record as it completes or is canceled.
 	 */
 	virtual void Put(
-		TConstArrayView<FCacheRecord> Records,
+		TConstArrayView<FCachePutRequest> Requests,
 		FStringView Context,
-		ECachePolicy Policy,
 		IRequestOwner& Owner,
 		FOnCachePutComplete&& OnComplete = FOnCachePutComplete()) = 0;
 
@@ -261,16 +260,14 @@ public:
 	 * cache store must produce a partial record with the available payloads when the policy flag
 	 * PartialOnError is set for the missing payloads.
 	 *
-	 * @param Keys         The keys identifying the cache records to query.
+	 * @param Requests     Requests with the keys identifying the cache records to query.
 	 * @param Context      A description of the request. An object path is typically sufficient.
-	 * @param Policy       Flags to control the behavior of the request. See FCacheRecordPolicy.
 	 * @param Owner        The owner to execute the request within.
 	 * @param OnComplete   A callback invoked for every key as it completes or is canceled.
 	 */
 	virtual void Get(
-		TConstArrayView<FCacheKey> Keys,
+		TConstArrayView<FCacheGetRequest> Requests,
 		FStringView Context,
-		FCacheRecordPolicy Policy,
 		IRequestOwner& Owner,
 		FOnCacheGetComplete&& OnComplete) = 0;
 
@@ -308,6 +305,59 @@ class ICache : public ICacheStore
 public:
 	virtual ~ICache() = default;
 
+	using ICacheStore::Put;
+	using ICacheStore::Get;
+
+	/**
+	 * Asynchronous request to put cache records according to the policy.
+	 *
+	 * The callback will always be called for every key, and may be called from an arbitrary thread.
+	 * Records may finish storing in any order, and from multiple threads concurrently.
+	 *
+	 * A cache store is free to interpret a record containing only a key as a request to delete that
+	 * record from the store. Records may contain payloads that do not have data, and these payloads
+	 * must reference an existing copy of the payload in the store, if available, and must otherwise
+	 * be stored as a partial record that can attempt recovery of missing payloads when fetched.
+	 *
+	 * @param Records      The cache records to store. Must have a key.
+	 * @param Context      A description of the request. An object path is typically sufficient.
+	 * @param Policy       Flags to control the behavior of the request. See ECachePolicy.
+	 * @param Owner        The owner to execute the request within.
+	 * @param OnComplete   A callback invoked for every record as it completes or is canceled.
+	 */
+	UE_API void Put(
+		TConstArrayView<FCacheRecord> Records,
+		FStringView Context,
+		ECachePolicy Policy,
+		IRequestOwner& Owner,
+		FOnCachePutComplete&& OnComplete = FOnCachePutComplete());
+
+	/**
+	 * Asynchronous request to get cache records according to the policy.
+	 *
+	 * The callback will always be called for every key, and may be called from an arbitrary thread.
+	 * Records may become available in any order, and from multiple threads concurrently.
+	 *
+	 * Records may propagate into other cache stores, in accordance with the policy. A propagated
+	 * record may be a partial record, with some payloads missing data, depending on the policy.
+	 *
+	 * When payloads are required by the policy, but not available, the status must be Error. The
+	 * cache store must produce a partial record with the available payloads when the policy flag
+	 * PartialOnError is set for the missing payloads.
+	 *
+	 * @param Keys         The keys identifying the cache records to query.
+	 * @param Context      A description of the request. An object path is typically sufficient.
+	 * @param Policy       Flags to control the behavior of the request. See FCacheRecordPolicy.
+	 * @param Owner        The owner to execute the request within.
+	 * @param OnComplete   A callback invoked for every key as it completes or is canceled.
+	 */
+	UE_API void Get(
+		TConstArrayView<FCacheKey> Keys,
+		FStringView Context,
+		FCacheRecordPolicy Policy,
+		IRequestOwner& Owner,
+		FOnCacheGetComplete&& OnComplete);
+
 	/**
 	 * Cancel all queued and active cache requests and invoke their callbacks.
 	 *
@@ -325,16 +375,45 @@ public:
 
 ///////////////////////////////////////////////////////////////////////////////////////////////////////////////////////
 
+/** Parameters to request to put a cache record. */
+struct FCachePutRequest
+{
+	/** The cache record to put. */
+	FCacheRecord Record;
+
+	/** Flags to control the behavior of the request. See FCacheRecordPolicy. */
+	FCacheRecordPolicy Policy;
+
+	/** A value that will be returned in the completion callback. */
+	uint64 UserData = 0;
+};
+
 /** Parameters for the completion callback for cache put requests. */
 struct FCachePutCompleteParams
 {
 	/** Key for the part of the put request that completed or was canceled. */
 	FCacheKey Key;
 
+	/** A copy of the value from the request. */
+	uint64 UserData = 0;
+
 	/** Status of the cache request. */
 	EStatus Status = EStatus::Error;
 };
 
+/** Parameters to request a cache record. */
+struct FCacheGetRequest
+{
+	/** The key identifying the cache record to fetch. */
+	FCacheKey Key;
+
+	/** Flags to control the behavior of the request. See FCacheRecordPolicy. */
+	FCacheRecordPolicy Policy;
+
+	/** A value that will be returned in the completion callback. */
+	uint64 UserData = 0;
+};
+
 /** Parameters for the completion callback for cache get requests. */
 struct FCacheGetCompleteParams
 {
@@ -348,6 +427,9 @@ struct FCacheGetCompleteParams
 	 */
 	FCacheRecord&& Record;
 
+	/** A copy of the value from the request. */
+	uint64 UserData = 0;
+
 	/** Status of the cache request. */
 	EStatus Status = EStatus::Error;
 };
@@ -367,8 +449,14 @@ struct FCacheChunkRequest
 	/** The maximum number of raw bytes of the payload to fetch, starting from the offset. */
 	uint64 RawSize = MAX_uint64;
 
+	/** The raw hash of the entire payload to fetch, if available, otherwise zero. */
+	FIoHash RawHash;
+
 	/** Flags to control the behavior of the request. See ECachePolicy. */
 	ECachePolicy Policy = ECachePolicy::Default;
+
+	/** A value that will be returned in the completion callback. */
+	uint64 UserData = 0;
 };
 
 /** Parameters for the completion callback for cache chunk requests. */
@@ -392,6 +480,9 @@ struct FCacheGetChunkCompleteParams
 	/** Data for the subset of the payload that was fetched when Status is Ok, otherwise null. */
 	FSharedBuffer&& RawData;
 
+	/** A copy of the value from the request. */
+	uint64 UserData = 0;
+
 	/** Status of the cache request. */
 	EStatus Status = EStatus::Error;
 };
