DDC: Exposed payload compression block size through the cache record and build function interfaces

#rnx

diff --git a/Engine/Source/Developer/DerivedDataCache/Private/DerivedDataBuildJobContext.cpp b/Engine/Source/Developer/DerivedDataCache/Private/DerivedDataBuildJobContext.cpp
--- a/Engine/Source/Developer/DerivedDataCache/Private/DerivedDataBuildJobContext.cpp
+++ b/Engine/Source/Developer/DerivedDataCache/Private/DerivedDataBuildJobContext.cpp
@@ -9,7 +9,6 @@
 #include "DerivedDataBuildTypes.h"
 #include "DerivedDataCache.h"
 #include "DerivedDataPayload.h"
-#include "DerivedDataPayloadPrivate.h"
 #include "DerivedDataRequestOwner.h"
 #include "Hash/Blake3.h"
 #include "Misc/StringBuilder.h"
@@ -111,20 +110,19 @@ void FBuildJobContext::AddPayload(const FPayloadId& Id, const FCompressedBuffer&
 	AddPayload(FPayload(Id, Buffer));
 }
 
-void FBuildJobContext::AddPayload(const FPayloadId& Id, const FCompositeBuffer& Buffer)
+void FBuildJobContext::AddPayload(const FPayloadId& Id, const FCompositeBuffer& Buffer, const uint64 BlockSize)
 {
-	AddPayload(FPayload(Id, FCompressedBuffer::Compress(Buffer, GDefaultCompressor, GDefaultCompressionLevel)));
+	AddPayload(FPayload(Id, FPayload::Compress(Buffer, BlockSize)));
 }
 
-void FBuildJobContext::AddPayload(const FPayloadId& Id, const FSharedBuffer& Buffer)
+void FBuildJobContext::AddPayload(const FPayloadId& Id, const FSharedBuffer& Buffer, const uint64 BlockSize)
 {
-	AddPayload(FPayload(Id, FCompressedBuffer::Compress(Buffer, GDefaultCompressor, GDefaultCompressionLevel)));
+	AddPayload(FPayload(Id, FPayload::Compress(Buffer, BlockSize)));
 }
 
 void FBuildJobContext::AddPayload(const FPayloadId& Id, const FCbObject& Object)
 {
-	const FCompositeBuffer& Buffer = Object.GetBuffer();
-	AddPayload(FPayload(Id, FCompressedBuffer::Compress(Buffer, GDefaultCompressor, GDefaultCompressionLevel)));
+	AddPayload(FPayload(Id, FPayload::Compress(Object.GetBuffer())));
 }
 
 void FBuildJobContext::BeginBuild(IRequestOwner& InOwner, TUniqueFunction<void ()>&& InOnEndBuild)
diff --git a/Engine/Source/Developer/DerivedDataCache/Private/DerivedDataBuildJobContext.h b/Engine/Source/Developer/DerivedDataCache/Private/DerivedDataBuildJobContext.h
--- a/Engine/Source/Developer/DerivedDataCache/Private/DerivedDataBuildJobContext.h
+++ b/Engine/Source/Developer/DerivedDataCache/Private/DerivedDataBuildJobContext.h
@@ -60,8 +60,8 @@ private:
 
 	void AddPayload(const FPayload& Payload) final;
 	void AddPayload(const FPayloadId& Id, const FCompressedBuffer& Buffer) final;
-	void AddPayload(const FPayloadId& Id, const FCompositeBuffer& Buffer) final;
-	void AddPayload(const FPayloadId& Id, const FSharedBuffer& Buffer) final;
+	void AddPayload(const FPayloadId& Id, const FCompositeBuffer& Buffer, uint64 BlockSize) final;
+	void AddPayload(const FPayloadId& Id, const FSharedBuffer& Buffer, uint64 BlockSize) final;
 	void AddPayload(const FPayloadId& Id, const FCbObject& Object) final;
 
 	void EndBuild();
diff --git a/Engine/Source/Developer/DerivedDataCache/Private/DerivedDataCacheRecord.cpp b/Engine/Source/Developer/DerivedDataCache/Private/DerivedDataCacheRecord.cpp
--- a/Engine/Source/Developer/DerivedDataCache/Private/DerivedDataCacheRecord.cpp
+++ b/Engine/Source/Developer/DerivedDataCache/Private/DerivedDataCacheRecord.cpp
@@ -6,7 +6,6 @@
 #include "DerivedDataCacheKey.h"
 #include "DerivedDataCachePrivate.h"
 #include "DerivedDataPayload.h"
-#include "DerivedDataPayloadPrivate.h"
 #include "HAL/CriticalSection.h"
 #include "Misc/ScopeExit.h"
 #include "Misc/ScopeRWLock.h"
@@ -29,13 +28,13 @@ public:
 
 	void SetMeta(FCbObject&& Meta) final;
 
-	FPayloadId SetValue(const FCompositeBuffer& Buffer, const FPayloadId& Id) final;
-	FPayloadId SetValue(const FSharedBuffer& Buffer, const FPayloadId& Id) final;
-	FPayloadId SetValue(const FPayload& Payload) final;
+	void SetValue(const FCompositeBuffer& Buffer, const FPayloadId& Id, uint64 BlockSize) final;
+	void SetValue(const FSharedBuffer& Buffer, const FPayloadId& Id, uint64 BlockSize) final;
+	void SetValue(const FPayload& Payload) final;
 
-	FPayloadId AddAttachment(const FCompositeBuffer& Buffer, const FPayloadId& Id) final;
-	FPayloadId AddAttachment(const FSharedBuffer& Buffer, const FPayloadId& Id) final;
-	FPayloadId AddAttachment(const FPayload& Payload) final;
+	void AddAttachment(const FCompositeBuffer& Buffer, const FPayloadId& Id, uint64 BlockSize) final;
+	void AddAttachment(const FSharedBuffer& Buffer, const FPayloadId& Id, uint64 BlockSize) final;
+	void AddAttachment(const FPayload& Payload) final;
 
 	FCacheRecord Build() final;
 	void BuildAsync(IRequestOwner& Owner, FOnCacheRecordComplete&& OnComplete) final;
@@ -196,49 +195,47 @@ void FCacheRecordBuilderInternal::SetMeta(FCbObject&& InMeta)
 	Meta.MakeOwned();
 }
 
-FPayloadId FCacheRecordBuilderInternal::SetValue(const FCompositeBuffer& Buffer, const FPayloadId& Id)
+void FCacheRecordBuilderInternal::SetValue(const FCompositeBuffer& Buffer, const FPayloadId& Id, const uint64 BlockSize)
 {
-	FCompressedBuffer Compressed = FCompressedBuffer::Compress(Buffer, GDefaultCompressor, GDefaultCompressionLevel);
+	FCompressedBuffer Compressed = FPayload::Compress(Buffer, BlockSize);
 	const FPayloadId ValueId = GetOrCreatePayloadId(Id, Compressed);
 	return SetValue(FPayload(ValueId, MoveTemp(Compressed)));
 }
 
-FPayloadId FCacheRecordBuilderInternal::SetValue(const FSharedBuffer& Buffer, const FPayloadId& Id)
+void FCacheRecordBuilderInternal::SetValue(const FSharedBuffer& Buffer, const FPayloadId& Id, const uint64 BlockSize)
 {
-	FCompressedBuffer Compressed = FCompressedBuffer::Compress(Buffer, GDefaultCompressor, GDefaultCompressionLevel);
+	FCompressedBuffer Compressed = FPayload::Compress(Buffer, BlockSize);
 	const FPayloadId ValueId = GetOrCreatePayloadId(Id, Compressed);
 	return SetValue(FPayload(ValueId, MoveTemp(Compressed)));
 }
 
-FPayloadId FCacheRecordBuilderInternal::SetValue(const FPayload& Payload)
+void FCacheRecordBuilderInternal::SetValue(const FPayload& Payload)
 {
 	checkf(Payload, TEXT("Failed to set value on %s because the payload is null."), *WriteToString<96>(Key));
-	const FPayloadId& Id = Payload.GetId();
 	checkf(Value.IsNull(),
 		TEXT("Failed to set value on %s with ID %s because it has an existing value with ID %s."),
-		*WriteToString<96>(Key), *WriteToString<32>(Id), *WriteToString<32>(Value.GetId()));
-	checkf(Algo::BinarySearchBy(Attachments, Id, &FPayload::GetId) == INDEX_NONE,
+		*WriteToString<96>(Key), *WriteToString<32>(Payload.GetId()), *WriteToString<32>(Value.GetId()));
+	checkf(Algo::BinarySearchBy(Attachments, Payload.GetId(), &FPayload::GetId) == INDEX_NONE,
 		TEXT("Failed to set value on %s with ID %s because it has an existing attachment with that ID."),
-		*WriteToString<96>(Key), *WriteToString<32>(Id));
+		*WriteToString<96>(Key), *WriteToString<32>(Payload.GetId()));
 	Value = Payload;
-	return Id;
 }
 
-FPayloadId FCacheRecordBuilderInternal::AddAttachment(const FCompositeBuffer& Buffer, const FPayloadId& Id)
+void FCacheRecordBuilderInternal::AddAttachment(const FCompositeBuffer& Buffer, const FPayloadId& Id, const uint64 BlockSize)
 {
-	FCompressedBuffer Compressed = FCompressedBuffer::Compress(Buffer, GDefaultCompressor, GDefaultCompressionLevel);
+	FCompressedBuffer Compressed = FPayload::Compress(Buffer, BlockSize);
 	const FPayloadId AttachmentId = GetOrCreatePayloadId(Id, Compressed);
 	return AddAttachment(FPayload(AttachmentId, MoveTemp(Compressed)));
 }
 
-FPayloadId FCacheRecordBuilderInternal::AddAttachment(const FSharedBuffer& Buffer, const FPayloadId& Id)
+void FCacheRecordBuilderInternal::AddAttachment(const FSharedBuffer& Buffer, const FPayloadId& Id, const uint64 BlockSize)
 {
-	FCompressedBuffer Compressed = FCompressedBuffer::Compress(Buffer, GDefaultCompressor, GDefaultCompressionLevel);
+	FCompressedBuffer Compressed = FPayload::Compress(Buffer, BlockSize);
 	const FPayloadId AttachmentId = GetOrCreatePayloadId(Id, Compressed);
 	return AddAttachment(FPayload(AttachmentId, MoveTemp(Compressed)));
 }
 
-FPayloadId FCacheRecordBuilderInternal::AddAttachment(const FPayload& Payload)
+void FCacheRecordBuilderInternal::AddAttachment(const FPayload& Payload)
 {
 	checkf(Payload, TEXT("Failed to add attachment on %s because the payload is null."), *WriteToString<96>(Key));
 	const FPayloadId& Id = Payload.GetId();
@@ -247,7 +244,6 @@ FPayloadId FCacheRecordBuilderInternal::AddAttachment(const FPayload& Payload)
 		TEXT("Failed to add attachment on %s with ID %s because it has an existing attachment or value with that ID."),
 		*WriteToString<96>(Key), *WriteToString<32>(Id));
 	Attachments.Insert(Payload, Index);
-	return Id;
 }
 
 FCacheRecord FCacheRecordBuilderInternal::Build()
diff --git a/Engine/Source/Developer/DerivedDataCache/Private/DerivedDataPayload.cpp b/Engine/Source/Developer/DerivedDataCache/Private/DerivedDataPayload.cpp
--- a/Engine/Source/Developer/DerivedDataCache/Private/DerivedDataPayload.cpp
+++ b/Engine/Source/Developer/DerivedDataCache/Private/DerivedDataPayload.cpp
@@ -2,6 +2,7 @@
 
 #include "DerivedDataPayload.h"
 
+#include "Compression/OodleDataCompression.h"
 #include "Containers/StringConv.h"
 #include "Hash/xxhash.h"
 #include "Serialization/CompactBinary.h"
@@ -10,6 +11,9 @@
 namespace UE::DerivedData
 {
 
+constexpr ECompressedBufferCompressor GDefaultCompressor = ECompressedBufferCompressor::Mermaid;
+constexpr ECompressedBufferCompressionLevel GDefaultCompressionLevel = ECompressedBufferCompressionLevel::VeryFast;
+
 FPayloadId::FPayloadId(const FMemoryView Id)
 {
 	checkf(Id.GetSize() == sizeof(ByteArray),
@@ -33,7 +37,7 @@ FPayloadId FPayloadId::FromHash(const FIoHash& Hash)
 	return FPayloadId(MakeMemoryView(Hash.GetBytes()).Left(sizeof(ByteArray)));
 }
 
-FPayloadId FPayloadId::FromName(const FAnsiStringView Name)
+FPayloadId FPayloadId::FromName(const FUtf8StringView Name)
 {
 	checkf(!Name.IsEmpty(), TEXT("FPayloadId requires a non-empty name."));
 	uint8 HashBytes[16];
@@ -46,4 +50,14 @@ FPayloadId FPayloadId::FromName(const FWideStringView Name)
 	return FPayloadId::FromName(FTCHARToUTF8(Name));
 }
 
+FCompressedBuffer FPayload::Compress(const FCompositeBuffer& Buffer, const uint64 BlockSize)
+{
+	return FCompressedBuffer::Compress(Buffer, GDefaultCompressor, GDefaultCompressionLevel, BlockSize);
+}
+
+FCompressedBuffer FPayload::Compress(const FSharedBuffer& Buffer, const uint64 BlockSize)
+{
+	return FCompressedBuffer::Compress(Buffer, GDefaultCompressor, GDefaultCompressionLevel, BlockSize);
+}
+
 } // UE::DerivedData
diff --git a/Engine/Source/Developer/DerivedDataCache/Private/DerivedDataPayloadPrivate.h b/Engine/Source/Developer/DerivedDataCache/Private/DerivedDataPayloadPrivate.h
deleted file mode 100644
diff --git a/Engine/Source/Developer/DerivedDataCache/Public/DerivedDataBuildFunction.h b/Engine/Source/Developer/DerivedDataCache/Public/DerivedDataBuildFunction.h
--- a/Engine/Source/Developer/DerivedDataCache/Public/DerivedDataBuildFunction.h
+++ b/Engine/Source/Developer/DerivedDataCache/Public/DerivedDataBuildFunction.h
@@ -93,8 +93,8 @@ public:
 	/** Adds a payload to the build output. Must have a non-null buffer and a unique ID. */
 	virtual void AddPayload(const FPayload& Payload) = 0;
 	virtual void AddPayload(const FPayloadId& Id, const FCompressedBuffer& Buffer) = 0;
-	virtual void AddPayload(const FPayloadId& Id, const FCompositeBuffer& Buffer) = 0;
-	virtual void AddPayload(const FPayloadId& Id, const FSharedBuffer& Buffer) = 0;
+	virtual void AddPayload(const FPayloadId& Id, const FCompositeBuffer& Buffer, uint64 BlockSize = 0) = 0;
+	virtual void AddPayload(const FPayloadId& Id, const FSharedBuffer& Buffer, uint64 BlockSize = 0) = 0;
 	virtual void AddPayload(const FPayloadId& Id, const FCbObject& Object) = 0;
 
 	/** Adds the error to the build output. Must be deterministic and based only on inputs. */
diff --git a/Engine/Source/Developer/DerivedDataCache/Public/DerivedDataCacheRecord.h b/Engine/Source/Developer/DerivedDataCache/Public/DerivedDataCacheRecord.h
--- a/Engine/Source/Developer/DerivedDataCache/Public/DerivedDataCacheRecord.h
+++ b/Engine/Source/Developer/DerivedDataCache/Public/DerivedDataCacheRecord.h
@@ -50,12 +50,12 @@ class ICacheRecordBuilderInternal
 public:
 	virtual ~ICacheRecordBuilderInternal() = default;
 	virtual void SetMeta(FCbObject&& Meta) = 0;
-	virtual FPayloadId SetValue(const FCompositeBuffer& Buffer, const FPayloadId& Id) = 0;
-	virtual FPayloadId SetValue(const FSharedBuffer& Buffer, const FPayloadId& Id) = 0;
-	virtual FPayloadId SetValue(const FPayload& Payload) = 0;
-	virtual FPayloadId AddAttachment(const FCompositeBuffer& Buffer, const FPayloadId& Id) = 0;
-	virtual FPayloadId AddAttachment(const FSharedBuffer& Buffer, const FPayloadId& Id) = 0;
-	virtual FPayloadId AddAttachment(const FPayload& Payload) = 0;
+	virtual void SetValue(const FCompositeBuffer& Buffer, const FPayloadId& Id, uint64 BlockSize) = 0;
+	virtual void SetValue(const FSharedBuffer& Buffer, const FPayloadId& Id, uint64 BlockSize) = 0;
+	virtual void SetValue(const FPayload& Payload) = 0;
+	virtual void AddAttachment(const FCompositeBuffer& Buffer, const FPayloadId& Id, uint64 BlockSize) = 0;
+	virtual void AddAttachment(const FSharedBuffer& Buffer, const FPayloadId& Id, uint64 BlockSize) = 0;
+	virtual void AddAttachment(const FPayload& Payload) = 0;
 	virtual FCacheRecord Build() = 0;
 	virtual void BuildAsync(IRequestOwner& Owner, FOnCacheRecordComplete&& OnComplete) = 0;
 };
@@ -161,18 +161,19 @@ public:
 	/**
 	 * Set the value for the cache record.
 	 *
-	 * @param Buffer   The value, which is compressed by the builder, and cloned if not owned.
-	 * @param Id       An ID for the value that is unique within this cache record. When omitted,
-	 *                 the buffer will be hashed to create an ID.
+	 * @param Buffer      The value, which is compressed by the builder, and cloned if not owned.
+	 * @param Id          An ID for the value that is unique within this cache record.
+	 *                    When omitted, the hash of the buffer will be used as the ID.
+	 * @param BlockSize   The power-of-two block size to encode raw data in. 0 is default.
 	 * @return The ID that was provided or created.
 	 */
-	inline FPayloadId SetValue(const FCompositeBuffer& Buffer, const FPayloadId& Id = FPayloadId())
+	inline void SetValue(const FCompositeBuffer& Buffer, const FPayloadId& Id = FPayloadId(), const uint64 BlockSize = 0)
 	{
-		return RecordBuilder->SetValue(Buffer, Id);
+		return RecordBuilder->SetValue(Buffer, Id, BlockSize);
 	}
-	inline FPayloadId SetValue(const FSharedBuffer& Buffer, const FPayloadId& Id = FPayloadId())
+	inline void SetValue(const FSharedBuffer& Buffer, const FPayloadId& Id = FPayloadId(), const uint64 BlockSize = 0)
 	{
-		return RecordBuilder->SetValue(Buffer, Id);
+		return RecordBuilder->SetValue(Buffer, Id, BlockSize);
 	}
 
 	/**
@@ -181,7 +182,7 @@ public:
 	 * @param Payload   The payload, which must have data unless it is known to be in the cache.
 	 * @return The ID that was provided. Unique within the scope of the cache record.
 	 */
-	inline FPayloadId SetValue(const FPayload& Payload)
+	inline void SetValue(const FPayload& Payload)
 	{
 		return RecordBuilder->SetValue(Payload);
 	}
@@ -189,18 +190,19 @@ public:
 	/**
 	 * Add an attachment to the cache record.
 	 *
-	 * @param Buffer   The attachment, which is compressed by the builder, and cloned if not owned.
-	 * @param Id       An ID for the attachment that is unique within this cache record. When omitted,
-	 *                 the buffer will be hashed to create an ID.
+	 * @param Buffer      The attachment, which is compressed by the builder, and cloned if not owned.
+	 * @param Id          An ID for the attachment that is unique within this cache record.
+	 *                    When omitted, the hash of the buffer will be used as the ID.
+	 * @param BlockSize   The power-of-two block size to encode raw data in. 0 is default.
 	 * @return The ID that was provided or created.
 	 */
-	inline FPayloadId AddAttachment(const FCompositeBuffer& Buffer, const FPayloadId& Id = FPayloadId())
+	inline void AddAttachment(const FCompositeBuffer& Buffer, const FPayloadId& Id = FPayloadId(), const uint64 BlockSize = 0)
 	{
-		return RecordBuilder->AddAttachment(Buffer, Id);
+		return RecordBuilder->AddAttachment(Buffer, Id, BlockSize);
 	}
-	inline FPayloadId AddAttachment(const FSharedBuffer& Buffer, const FPayloadId& Id = FPayloadId())
+	inline void AddAttachment(const FSharedBuffer& Buffer, const FPayloadId& Id = FPayloadId(), const uint64 BlockSize = 0)
 	{
-		return RecordBuilder->AddAttachment(Buffer, Id);
+		return RecordBuilder->AddAttachment(Buffer, Id, BlockSize);
 	}
 
 	/**
@@ -209,7 +211,7 @@ public:
 	 * @param Payload   The payload, which must have data unless it is known to be in the cache.
 	 * @return The ID that was provided. Unique within the scope of the cache record.
 	 */
-	inline FPayloadId AddAttachment(const FPayload& Payload)
+	inline void AddAttachment(const FPayload& Payload)
 	{
 		return RecordBuilder->AddAttachment(Payload);
 	}
diff --git a/Engine/Source/Developer/DerivedDataCache/Public/DerivedDataPayload.h b/Engine/Source/Developer/DerivedDataCache/Public/DerivedDataPayload.h
--- a/Engine/Source/Developer/DerivedDataCache/Public/DerivedDataPayload.h
+++ b/Engine/Source/Developer/DerivedDataCache/Public/DerivedDataPayload.h
@@ -7,6 +7,11 @@
 #include "DerivedDataPayloadId.h"
 #include "IO/IoHash.h"
 
+#define UE_API DERIVEDDATACACHE_API
+
+class FCompositeBuffer;
+class FSharedBuffer;
+
 namespace UE::DerivedData
 {
 
@@ -19,6 +24,15 @@ namespace UE::DerivedData
 class FPayload
 {
 public:
+	/**
+	 * Compress the buffer using the default payload compression parameters.
+	 *
+	 * @param RawData     The raw data to be compressed.
+	 * @param BlockSize   The power-of-two block size to encode raw data in. 0 is default.
+	 */
+	UE_API static FCompressedBuffer Compress(const FCompositeBuffer& RawData, uint64 BlockSize = 0);
+	UE_API static FCompressedBuffer Compress(const FSharedBuffer& RawData, uint64 BlockSize = 0);
+
 	/** Construct a null payload. */
 	FPayload() = default;
 
@@ -118,3 +132,5 @@ inline bool operator<(const FPayload& A, const FPayload& B)
 }
 
 } // UE::DerivedData
+
+#undef UE_API
diff --git a/Engine/Source/Developer/DerivedDataCache/Public/DerivedDataPayloadId.h b/Engine/Source/Developer/DerivedDataCache/Public/DerivedDataPayloadId.h
--- a/Engine/Source/Developer/DerivedDataCache/Public/DerivedDataPayloadId.h
+++ b/Engine/Source/Developer/DerivedDataCache/Public/DerivedDataPayloadId.h
@@ -40,7 +40,7 @@ public:
 	[[nodiscard]] UE_API static FPayloadId FromHash(const FIoHash& Hash);
 
 	/** Construct an ID from a non-empty name. */
-	[[nodiscard]] UE_API static FPayloadId FromName(FAnsiStringView Name);
+	[[nodiscard]] UE_API static FPayloadId FromName(FUtf8StringView Name);
 	[[nodiscard]] UE_API static FPayloadId FromName(FWideStringView Name);
 
 	/** Returns a reference to the raw byte array for the ID. */
