DDC: Removed accessors for raw data from FCacheRecord

Use GetValue().GetData().Decompress() in place of GetValue(), or the equivalent for GetAttachment(Id). The removed accessors kept a cache of the decompressed payload and required synchronization to be safe to call from multiple threads.

#rnx

diff --git a/Engine/Source/Developer/DerivedDataCache/Private/DerivedDataBuildOutput.cpp b/Engine/Source/Developer/DerivedDataCache/Private/DerivedDataBuildOutput.cpp
--- a/Engine/Source/Developer/DerivedDataCache/Private/DerivedDataBuildOutput.cpp
+++ b/Engine/Source/Developer/DerivedDataCache/Private/DerivedDataBuildOutput.cpp
@@ -194,7 +194,7 @@ FBuildOutputInternal::FBuildOutputInternal(FStringView InName, FStringView InFun
 	: Name(InName)
 	, Function(InFunction)
 	, Meta(InOutput.GetMeta())
-	, Output(InOutput.GetValue())
+	, Output(InOutput.GetValue().GetData().Decompress())
 	, Payloads(InOutput.GetAttachmentPayloads())
 {
 	checkf(!Name.IsEmpty(), TEXT("A build output requires a non-empty name."));
diff --git a/Engine/Source/Developer/DerivedDataCache/Private/DerivedDataCacheRecord.cpp b/Engine/Source/Developer/DerivedDataCache/Private/DerivedDataCacheRecord.cpp
--- a/Engine/Source/Developer/DerivedDataCache/Private/DerivedDataCacheRecord.cpp
+++ b/Engine/Source/Developer/DerivedDataCache/Private/DerivedDataCacheRecord.cpp
@@ -6,9 +6,7 @@
 #include "DerivedDataCacheKey.h"
 #include "DerivedDataCachePrivate.h"
 #include "DerivedDataPayload.h"
-#include "HAL/CriticalSection.h"
 #include "Misc/ScopeExit.h"
-#include "Misc/ScopeRWLock.h"
 #include "Misc/StringBuilder.h"
 #include "Serialization/CompactBinary.h"
 #include "Serialization/CompactBinaryPackage.h"
@@ -57,14 +55,9 @@ public:
 
 	const FCacheKey& GetKey() const final;
 	const FCbObject& GetMeta() const final;
-
-	FSharedBuffer GetValue() const final;
-	const FPayload& GetValuePayload() const final;
-
-	FSharedBuffer GetAttachment(const FPayloadId& Id) const final;
-	const FPayload& GetAttachmentPayload(const FPayloadId& Id) const final;
-	TConstArrayView<FPayload> GetAttachmentPayloads() const final;
-
+	const FPayload& GetValue() const final;
+	const FPayload& GetAttachment(const FPayloadId& Id) const final;
+	TConstArrayView<FPayload> GetAttachments() const final;
 	const FPayload& GetPayload(const FPayloadId& Id) const final;
 
 	inline void AddRef() const final
@@ -85,9 +78,6 @@ private:
 	FCbObject Meta;
 	FPayload Value;
 	TArray<FPayload> Attachments;
-	mutable FRWLock CacheLock;
-	mutable FSharedBuffer ValueCache;
-	mutable TArray<FSharedBuffer> AttachmentsCache;
 	mutable std::atomic<uint32> ReferenceCount{0};
 };
 
@@ -118,68 +108,25 @@ const FCbObject& FCacheRecordInternal::GetMeta() const
 	return Meta;
 }
 
-FSharedBuffer FCacheRecordInternal::GetValue() const
-{
-	if (Value.IsNull())
-	{
-		return FSharedBuffer();
-	}
-	if (FReadScopeLock Lock(CacheLock); ValueCache)
-	{
-		return ValueCache;
-	}
-	FWriteScopeLock Lock(CacheLock);
-	if (ValueCache)
-	{
-		return ValueCache;
-	}
-	ValueCache = Value.GetData().Decompress();
-	return ValueCache;
-}
-
-const FPayload& FCacheRecordInternal::GetValuePayload() const
+const FPayload& FCacheRecordInternal::GetValue() const
 {
 	return Value;
 }
 
-FSharedBuffer FCacheRecordInternal::GetAttachment(const FPayloadId& Id) const
-{
-	const int32 Index = Algo::BinarySearchBy(Attachments, Id, &FPayload::GetId);
-	if (Attachments.IsValidIndex(Index))
-	{
-		if (FReadScopeLock Lock(CacheLock); AttachmentsCache.IsValidIndex(Index) && AttachmentsCache[Index])
-		{
-			return AttachmentsCache[Index];
-		}
-		FWriteScopeLock Lock(CacheLock);
-		if (AttachmentsCache.IsEmpty())
-		{
-			AttachmentsCache.SetNum(Attachments.Num());
-		}
-		FSharedBuffer& Cache = AttachmentsCache[Index];
-		if (Cache.IsNull())
-		{
-			Cache = Attachments[Index].GetData().Decompress();
-		}
-		return Cache;
-	}
-	return FSharedBuffer();
-}
-
-const FPayload& FCacheRecordInternal::GetAttachmentPayload(const FPayloadId& Id) const
+const FPayload& FCacheRecordInternal::GetAttachment(const FPayloadId& Id) const
 {
 	const int32 Index = Algo::BinarySearchBy(Attachments, Id, &FPayload::GetId);
 	return Attachments.IsValidIndex(Index) ? Attachments[Index] : FPayload::Null;
 }
 
-TConstArrayView<FPayload> FCacheRecordInternal::GetAttachmentPayloads() const
+TConstArrayView<FPayload> FCacheRecordInternal::GetAttachments() const
 {
 	return Attachments;
 }
 
 const FPayload& FCacheRecordInternal::GetPayload(const FPayloadId& Id) const
 {
-	return Value.GetId() == Id ? Value : GetAttachmentPayload(Id);
+	return Value.GetId() == Id ? Value : GetAttachment(Id);
 }
 
 ///////////////////////////////////////////////////////////////////////////////////////////////////////////////////////
diff --git a/Engine/Source/Developer/DerivedDataCache/Public/DerivedDataCacheRecord.h b/Engine/Source/Developer/DerivedDataCache/Public/DerivedDataCacheRecord.h
--- a/Engine/Source/Developer/DerivedDataCache/Public/DerivedDataCacheRecord.h
+++ b/Engine/Source/Developer/DerivedDataCache/Public/DerivedDataCacheRecord.h
@@ -4,7 +4,7 @@
 
 #include "DerivedDataCacheKey.h"
 #include "DerivedDataPayloadId.h"
-#include "Memory/SharedBuffer.h"
+#include "Memory/MemoryFwd.h"
 #include "Misc/ScopeExit.h"
 #include "Templates/Function.h"
 #include "Templates/RefCounting.h"
@@ -33,11 +33,9 @@ public:
 	virtual ~ICacheRecordInternal() = default;
 	virtual const FCacheKey& GetKey() const = 0;
 	virtual const FCbObject& GetMeta() const = 0;
-	virtual FSharedBuffer GetValue() const = 0;
-	virtual const FPayload& GetValuePayload() const = 0;
-	virtual FSharedBuffer GetAttachment(const FPayloadId& Id) const = 0;
-	virtual const FPayload& GetAttachmentPayload(const FPayloadId& Id) const = 0;
-	virtual TConstArrayView<FPayload> GetAttachmentPayloads() const = 0;
+	virtual const FPayload& GetValue() const = 0;
+	virtual const FPayload& GetAttachment(const FPayloadId& Id) const = 0;
+	virtual TConstArrayView<FPayload> GetAttachments() const = 0;
 	virtual const FPayload& GetPayload(const FPayloadId& Id) const = 0;
 	virtual void AddRef() const = 0;
 	virtual void Release() const = 0;
@@ -91,20 +89,21 @@ public:
 	/** Returns the metadata. Null when requested with ECachePolicy::SkipMeta. */
 	inline const FCbObject& GetMeta() const { return Record->GetMeta(); }
 
-	/** Returns the value. Null if no value or requested with ECachePolicy::SkipValue. */
-	inline FSharedBuffer GetValue() const { return Record->GetValue(); }
+	/** Returns the value. Data is null if skipped by the policy on request. */
+	inline const FPayload& GetValue() const { return Record->GetValue(); }
 
-	/** Returns the value payload. Null if no value. Data is null if value was skipped. */
-	inline const FPayload& GetValuePayload() const { return Record->GetValuePayload(); }
+	/** Returns the value payload. Null if no value. Data is null if skipped by the policy on request. */
+	inline const FPayload& GetValuePayload() const { return Record->GetValue(); }
 
-	/** Returns the attachment matching the ID. Null if no match or requested with ECachePolicy::SkipAttachments. */
-	inline FSharedBuffer GetAttachment(const FPayloadId& Id) const { return Record->GetAttachment(Id); }
+	/** Returns the attachment matching the ID. Data is null if skipped by the policy on request. */
+	inline const FPayload& GetAttachment(const FPayloadId& Id) const { return Record->GetAttachment(Id); }
 
-	/** Returns the attachment payload matching the ID. Null if no match. Data is null if attachments were skipped. */
-	inline const FPayload& GetAttachmentPayload(const FPayloadId& Id) const { return Record->GetAttachmentPayload(Id); }
+	/** Returns the attachment payload matching the ID. Null if no match. Data is null if skipped by the policy on request. */
+	inline const FPayload& GetAttachmentPayload(const FPayloadId& Id) const { return Record->GetAttachment(Id); }
 
-	/** Returns a view of the attachments. Always available, but data may be skipped. */
-	inline TConstArrayView<FPayload> GetAttachmentPayloads() const { return Record->GetAttachmentPayloads(); }
+	/** Returns a view of the attachments. Always available, but data may be skipped if skipped. */
+	inline TConstArrayView<FPayload> GetAttachments() const { return Record->GetAttachments(); }
+	inline TConstArrayView<FPayload> GetAttachmentPayloads() const { return Record->GetAttachments(); }
 
 	/** Returns the payload matching the ID, whether value or attachment. Null if no match. Data is null if skipped. */
 	inline const FPayload& GetPayload(const FPayloadId& Id) const { return Record->GetPayload(Id); }
diff --git a/Engine/Source/Editor/UnrealEd/Private/EditorDomain/EditorDomainUtils.cpp b/Engine/Source/Editor/UnrealEd/Private/EditorDomain/EditorDomainUtils.cpp
--- a/Engine/Source/Editor/UnrealEd/Private/EditorDomain/EditorDomainUtils.cpp
+++ b/Engine/Source/Editor/UnrealEd/Private/EditorDomain/EditorDomainUtils.cpp
@@ -11,6 +11,7 @@
 #include "DerivedDataCache.h"
 #include "DerivedDataCacheKey.h"
 #include "DerivedDataCacheRecord.h"
+#include "DerivedDataPayload.h"
 #include "DerivedDataRequestOwner.h"
 #include "Editor.h"
 #include "HAL/FileManager.h"
@@ -1313,7 +1314,7 @@ void GetBulkDataList(FName PackageName, UE::DerivedData::IRequestOwner& Owner, T
 		[InnerCallback = MoveTemp(Callback)](FCacheGetCompleteParams&& Params)
 		{
 			bool bOk = Params.Status == EStatus::Ok;
-			InnerCallback(bOk ? Params.Record.GetValue() : FSharedBuffer());
+			InnerCallback(bOk ? Params.Record.GetValue().GetData().Decompress() : FSharedBuffer());
 		});
 }
 
@@ -1384,7 +1385,7 @@ void GetBulkDataPayloadId(FName PackageName, const FGuid& BulkDataId, UE::Derive
 		[InnerCallback = MoveTemp(Callback)](FCacheGetCompleteParams&& Params)
 	{
 		bool bOk = Params.Status == EStatus::Ok;
-		InnerCallback(bOk ? Params.Record.GetValue() : FSharedBuffer());
+		InnerCallback(bOk ? Params.Record.GetValue().GetData().Decompress() : FSharedBuffer());
 	});
 }
 
diff --git a/Engine/Source/Editor/UnrealEd/Private/EditorDomain/EditorDomainUtils.h b/Engine/Source/Editor/UnrealEd/Private/EditorDomain/EditorDomainUtils.h
--- a/Engine/Source/Editor/UnrealEd/Private/EditorDomain/EditorDomainUtils.h
+++ b/Engine/Source/Editor/UnrealEd/Private/EditorDomain/EditorDomainUtils.h
@@ -10,6 +10,7 @@
 #include "DerivedDataCache.h"
 #include "EditorDomain/EditorDomain.h"
 #include "HAL/CriticalSection.h"
+#include "Memory/MemoryFwd.h"
 #include "Misc/ConfigCacheIni.h"
 #include "Misc/EnumClassFlags.h"
 #include "Serialization/CustomVersion.h"
@@ -21,7 +22,6 @@ DECLARE_LOG_CATEGORY_EXTERN(LogEditorDomainSave, Log, All);
 class FAssetPackageData;
 class FBlake3;
 class FPackagePath;
-class FSharedBuffer;
 class IAssetRegistry;
 class UPackage;
 
