DDC: Added FPayload::RemoveData() to simplify converting payloads to references

#rnx

diff --git a/Engine/Source/Developer/DerivedDataCache/Private/FileSystemDerivedDataBackend.cpp b/Engine/Source/Developer/DerivedDataCache/Private/FileSystemDerivedDataBackend.cpp
--- a/Engine/Source/Developer/DerivedDataCache/Private/FileSystemDerivedDataBackend.cpp
+++ b/Engine/Source/Developer/DerivedDataCache/Private/FileSystemDerivedDataBackend.cpp
@@ -1996,19 +1996,12 @@ FPayload FFileSystemCacheStore::GetCacheContent(
 {
 	if (!EnumHasAnyFlags(Policy, ECachePolicy::Query))
 	{
-		return FPayload(Payload.GetId(), Payload.GetRawHash(), Payload.GetRawSize());
+		return Payload.RemoveData();
 	}
 
 	if (Payload.HasData())
 	{
-		if (EnumHasAnyFlags(Policy, SkipFlag))
-		{
-			return FPayload(Payload.GetId(), Payload.GetRawHash(), Payload.GetRawSize());
-		}
-		else
-		{
-			return Payload;
-		}
+		return EnumHasAnyFlags(Policy, SkipFlag) ? Payload.RemoveData() : Payload;
 	}
 
 	const FIoHash& RawHash = Payload.GetRawHash();
diff --git a/Engine/Source/Developer/DerivedDataCache/Private/PakFileDerivedDataBackend.cpp b/Engine/Source/Developer/DerivedDataCache/Private/PakFileDerivedDataBackend.cpp
--- a/Engine/Source/Developer/DerivedDataCache/Private/PakFileDerivedDataBackend.cpp
+++ b/Engine/Source/Developer/DerivedDataCache/Private/PakFileDerivedDataBackend.cpp
@@ -858,11 +858,9 @@ void FPakFileDerivedDataBackend::GetCacheContent(
 	FPayload& InOutPayload,
 	EStatus& InOutStatus)
 {
-	const FIoHash& RawHash = InOutPayload.GetRawHash();
-
 	if (!EnumHasAnyFlags(Policy, ECachePolicy::Query) || (EnumHasAnyFlags(Policy, SkipFlag) && InOutPayload.HasData()))
 	{
-		InOutPayload = FPayload(InOutPayload.GetId(), RawHash, InOutPayload.GetRawSize());
+		InOutPayload = InOutPayload.RemoveData();
 		return;
 	}
 
@@ -871,6 +869,8 @@ void FPakFileDerivedDataBackend::GetCacheContent(
 		return;
 	}
 
+	const FIoHash& RawHash = InOutPayload.GetRawHash();
+
 	TStringBuilder<256> Path;
 	FPathViews::Append(Path, TEXT("Content"), RawHash);
 	if (EnumHasAllFlags(Policy, SkipFlag))
diff --git a/Engine/Source/Developer/DerivedDataCache/Public/DerivedDataPayload.h b/Engine/Source/Developer/DerivedDataCache/Public/DerivedDataPayload.h
--- a/Engine/Source/Developer/DerivedDataCache/Public/DerivedDataPayload.h
+++ b/Engine/Source/Developer/DerivedDataCache/Public/DerivedDataPayload.h
@@ -6,12 +6,10 @@
 #include "Compression/CompressedBuffer.h"
 #include "DerivedDataPayloadId.h"
 #include "IO/IoHash.h"
+#include "Memory/MemoryFwd.h"
 
 #define UE_API DERIVEDDATACACHE_API
 
-class FCompositeBuffer;
-class FSharedBuffer;
-
 namespace UE::DerivedData
 {
 
@@ -61,6 +59,9 @@ public:
 	/** Whether the compressed buffer for the payload is available. */
 	inline bool HasData() const { return !!Data; }
 
+	/** Create a copy of the payload with the data removed. */
+	[[nodiscard]] inline FPayload RemoveData() const { return FPayload(Id, RawHash, RawSize); }
+
 	/** Whether this is null. */
 	inline bool IsNull() const { return Id.IsNull(); }
 	/** Whether this is not null. */
