Fix DDC should not sanitize $ characters because it makes $ characters
SanitizeCacheKey and ValidateCacheKey should check the same condition

diff --git a/Engine/Source/Developer/DerivedDataCache/Private/DerivedDataCache.cpp b/Engine/Source/Developer/DerivedDataCache/Private/DerivedDataCache.cpp
--- a/Engine/Source/Developer/DerivedDataCache/Private/DerivedDataCache.cpp
+++ b/Engine/Source/Developer/DerivedDataCache/Private/DerivedDataCache.cpp
@@ -824,7 +824,7 @@ private:
 
 	static void ValidateCacheKey(const TCHAR* CacheKey)
 	{
-		checkf(Algo::AllOf(FStringView(CacheKey), [](TCHAR C) { return FChar::IsAlnum(C) || FChar::IsUnderscore(C) || C == TEXT('$'); }),
+		checkf(Algo::AllOf(FStringView(CacheKey), IsValidCacheChar),
 			TEXT("Invalid characters in cache key %s. Use SanitizeCacheKey or BuildCacheKey to create valid keys."), CacheKey);
 	}
 
diff --git a/Engine/Source/Developer/DerivedDataCache/Public/DerivedDataCacheInterface.h b/Engine/Source/Developer/DerivedDataCache/Public/DerivedDataCacheInterface.h
--- a/Engine/Source/Developer/DerivedDataCache/Public/DerivedDataCacheInterface.h
+++ b/Engine/Source/Developer/DerivedDataCache/Public/DerivedDataCacheInterface.h
@@ -65,6 +65,14 @@ public:
 	//--------------------------
 	// Low Level Static Helpers
 	//--------------------------
+	
+	/**
+	 * Returns true if character is valid in a DDC cache key without escaping 
+	**/
+	static bool IsValidCacheChar(const TCHAR C)
+	{
+		return FChar::IsAlnum(C) || FChar::IsUnderscore(C) || C == TEXT('$');
+	}
 
 	/** 
 	 * Static function to make sure a cache key contains only legal characters by using an escape
@@ -80,7 +88,7 @@ public:
 
 		for (int32 i = 0; i < Input.Len(); i++)
 		{
-			if (FChar::IsAlnum(Input[i]) || FChar::IsUnderscore(Input[i]))
+			if ( IsValidCacheChar(Input[i]) )
 			{
 				NumValid++;
 			}
