DDC: Fold the corruption wrapper into the file system cache store

#rnx

diff --git a/Engine/Source/Developer/DerivedDataCache/Private/DerivedDataBackendCorruptionWrapper.h b/Engine/Source/Developer/DerivedDataCache/Private/DerivedDataBackendCorruptionWrapper.h
--- a/Engine/Source/Developer/DerivedDataCache/Private/DerivedDataBackendCorruptionWrapper.h
+++ b/Engine/Source/Developer/DerivedDataCache/Private/DerivedDataBackendCorruptionWrapper.h
@@ -6,65 +6,80 @@
 #include "DerivedDataBackendInterface.h"
 #include "ProfilingDebugging/CookStats.h"
 #include "DerivedDataCacheUsageStats.h"
+#include "Memory/MemoryView.h"
 
-namespace UE::DerivedData::Backends
+namespace UE::DerivedData::CacheStore
 {
 
-/** 
+/**
  * Helper class for placing a footer at the end of of a cache file.
  * No effort is made to byte-swap this as we assume local format.
-**/
+ */
 struct FDerivedDataTrailer
 {
-	enum
-	{
-		/** Arbitrary number used to identify corruption **/
-		MagicConstant = 0x1e873d89
-	};
-
-	/** Arbitrary number used to identify corruption **/
-	uint32 Magic;
-	/** Version of the backend, for future use **/
-	uint32 Version;
-	/** CRC of the payload, used to detect corruption **/
-	uint32 CRCofPayload;
-	/** Size of the payload, used to detect corruption **/
-	uint32 SizeOfPayload;
-
-	/** Constructor that zeros everything. This is never valid because even the magic number is wrong **/
-	FDerivedDataTrailer() 
-		: Magic(0)
-		, Version(0)
-		, CRCofPayload(0)
-		, SizeOfPayload(0)
-	{
-	}
+	/** Arbitrary number used to identify corruption */
+	static constexpr inline uint32 MagicConstant = 0x1e873d89;
 
-	/**
-		* Constructor that sets the data properly for a given buffer of data
-		* @param	Data	Buffer of data to setup this trailer for
-		*/
-	FDerivedDataTrailer(const TArray<uint8>& Data)
+	/** Arbitrary number used to identify corruption */
+	uint32 Magic = 0;
+	/** Version of the trailer */
+	uint32 Version = 0;
+	/** CRC of the payload, used to detect corruption */
+	uint32 CRCofPayload = 0;
+	/** Size of the payload, used to detect corruption */
+	uint32 SizeOfPayload = 0;
+
+	FDerivedDataTrailer() = default;
+
+	explicit FDerivedDataTrailer(const FMemoryView Data)
 		: Magic(MagicConstant)
 		, Version(1)
-		, CRCofPayload(FCrc::MemCrc_DEPRECATED(Data.GetData(), Data.Num()))
-		, SizeOfPayload(Data.Num())
+		, CRCofPayload(FCrc::MemCrc_DEPRECATED(Data.GetData(), IntCastChecked<int32>(Data.GetSize())))
+		, SizeOfPayload(uint32(Data.GetSize()))
 	{
 	}
-	/**
-		* Comparison operator
-		* @param	Other	trailer to compare "this" to
-		* @return			true if all elements of the trailer match exactly
-		*/
-	bool operator==(const FDerivedDataTrailer& Other)
+
+	bool operator==(const FDerivedDataTrailer& Other) const
 	{
-		return Magic == Other.Magic &&
-			Version == Other.Version &&
-			CRCofPayload == Other.CRCofPayload &&
-			SizeOfPayload == Other.SizeOfPayload;
+		return Magic == Other.Magic
+			&& Version == Other.Version
+			&& CRCofPayload == Other.CRCofPayload
+			&& SizeOfPayload == Other.SizeOfPayload;
 	}
 };
 
+class FCorruptionWrapper
+{
+public:
+	static bool ReadTrailer(TArray<uint8>& Data, const TCHAR* CacheName, const TCHAR* CacheKey)
+	{
+		if (Data.Num() < sizeof(FDerivedDataTrailer))
+		{
+			UE_LOG(LogDerivedDataCache, Display, TEXT("%s: Corrupted file (short), ignoring and deleting %s."), CacheName, CacheKey);
+			return false;
+		}
+
+		FDerivedDataTrailer Trailer;
+		FMemory::Memcpy(&Trailer, &Data[Data.Num() - sizeof(FDerivedDataTrailer)], sizeof(FDerivedDataTrailer));
+		Data.RemoveAt(Data.Num() - sizeof(FDerivedDataTrailer), sizeof(FDerivedDataTrailer), /*bAllowShrinking*/ false);
+		FDerivedDataTrailer RecomputedTrailer(MakeMemoryView(Data));
+		if (!(Trailer == RecomputedTrailer))
+		{
+			UE_LOG(LogDerivedDataCache, Display, TEXT("%s: Corrupted file, ignoring and deleting %s."), CacheName, CacheKey);
+			return false;
+		}
+
+		return true;
+	}
+
+	static void WriteTrailer(FArchive& Ar, TConstArrayView<uint8> Data, const TCHAR* CacheName, const TCHAR* CacheKey)
+	{
+		checkf(Ar.TotalSize() + sizeof(FDerivedDataTrailer) <= MAX_int32,
+			TEXT("%s: Appending the trailer makes the data exceed 2 GiB for %s"), CacheName, CacheKey);
+		FDerivedDataTrailer Trailer(MakeMemoryView(Data));
+		Ar.Serialize(&Trailer, sizeof(FDerivedDataTrailer));
+	}
+};
 
 /** 
  * A backend wrapper that adds a footer to the data to check the CRC, length, etc.
@@ -165,7 +180,7 @@ public:
 				FDerivedDataTrailer Trailer;
 				FMemory::Memcpy(&Trailer,&OutData[OutData.Num() - sizeof(FDerivedDataTrailer)], sizeof(FDerivedDataTrailer));
 				OutData.RemoveAt(OutData.Num() - sizeof(FDerivedDataTrailer),sizeof(FDerivedDataTrailer), false);
-				FDerivedDataTrailer RecomputedTrailer(OutData);
+				FDerivedDataTrailer RecomputedTrailer(MakeMemoryView(OutData));
 				if (Trailer == RecomputedTrailer)
 				{
 					UE_LOG(LogDerivedDataCache, VeryVerbose, TEXT("FDerivedDataBackendCorruptionWrapper: cache hit, footer is ok %s"),CacheKey);
@@ -214,7 +229,7 @@ public:
 		Data.Reset( InData.Num() + sizeof(FDerivedDataTrailer) );
 		Data.AddUninitialized(InData.Num());
 		FMemory::Memcpy(&Data[0], &InData[0], InData.Num());
-		FDerivedDataTrailer Trailer(Data);
+		FDerivedDataTrailer Trailer(MakeMemoryView(Data));
 		Data.AddUninitialized(sizeof(FDerivedDataTrailer));
 		FMemory::Memcpy(&Data[Data.Num() - sizeof(FDerivedDataTrailer)], &Trailer, sizeof(FDerivedDataTrailer));
 		return InnerBackend->PutCachedData(CacheKey, Data, bPutEvenIfExists);
@@ -291,4 +306,4 @@ private:
 	FDerivedDataBackendInterface* InnerBackend;
 };
 
-} // UE::DerivedData::Backends
+} // UE::DerivedData::CacheStore
diff --git a/Engine/Source/Developer/DerivedDataCache/Private/DerivedDataBackends.cpp b/Engine/Source/Developer/DerivedDataCache/Private/DerivedDataBackends.cpp
--- a/Engine/Source/Developer/DerivedDataCache/Private/DerivedDataBackends.cpp
+++ b/Engine/Source/Developer/DerivedDataCache/Private/DerivedDataBackends.cpp
@@ -634,7 +634,7 @@ public:
 				FString WriteAccessLog;
 				bool bPromptIfMissing = false;
 
-				FParse::Value( Entry, TEXT("WriteAccessLog="), WriteAccessLog );		
+				FParse::Value(Entry, TEXT("WriteAccessLog="), WriteAccessLog);
 				FParse::Bool(Entry, TEXT("PromptIfMissing="), bPromptIfMissing);
 
 				if (!bShared || IFileManager::Get().DirectoryExists(*Path))
@@ -645,8 +645,8 @@ public:
 				if (InnerFileSystem)
 				{
 					bUsingSharedDDC |= bShared;
-	
-					DataCache = new FDerivedDataBackendCorruptionWrapper(InnerFileSystem);
+
+					DataCache = InnerFileSystem;
 					UE_LOG(LogDerivedDataCache, Log, TEXT("Using %s data cache path %s: %s"), NodeName, *Path, !InnerFileSystem->IsWritable() ? TEXT("ReadOnly") : TEXT("Writable"));
 					Directories.AddUnique(Path);
 				}
@@ -667,8 +667,6 @@ public:
 			} while (RetryOnFailure);
 		}
 
-
-
 		return DataCache;
 	}
 
@@ -738,7 +736,7 @@ public:
 
 		// Insert the backend corruption wrapper. Since the filesystem already uses this, and we're recycling the data with the trailer intact, we need to use it for the S3 cache too.
 		FS3DerivedDataBackend* Backend = new FS3DerivedDataBackend(*ManifestPath, *BaseUrl, *Region, *CanaryObjectKey, *CachePath);
-		return new FDerivedDataBackendCorruptionWrapper(Backend);
+		return new CacheStore::FDerivedDataBackendCorruptionWrapper(Backend);
 #else
 		UE_LOG(LogDerivedDataCache, Log, TEXT("S3 backend is not supported on the current platform."));
 		return nullptr;
diff --git a/Engine/Source/Developer/DerivedDataCache/Private/FileSystemDerivedDataBackend.cpp b/Engine/Source/Developer/DerivedDataCache/Private/FileSystemDerivedDataBackend.cpp
--- a/Engine/Source/Developer/DerivedDataCache/Private/FileSystemDerivedDataBackend.cpp
+++ b/Engine/Source/Developer/DerivedDataCache/Private/FileSystemDerivedDataBackend.cpp
@@ -7,6 +7,7 @@
 #include "Async/Async.h"
 #include "Async/TaskGraphInterfaces.h"
 #include "Containers/StaticBitArray.h"
+#include "DerivedDataBackendCorruptionWrapper.h"
 #include "DerivedDataBackendInterface.h"
 #include "DerivedDataCacheInterface.h"
 #include "DerivedDataCacheMaintainer.h"
@@ -70,15 +71,19 @@ static const TCHAR GContentDirectoryName[] = TEXT("Content");
 
 ///////////////////////////////////////////////////////////////////////////////////////////////////////////////////////
 
-FString BuildPathForCacheKey(const TCHAR* const CacheKey)
+void BuildPathForLegacyCache(const TCHAR* const CacheKey, FStringBuilderBase& Path)
 {
-	FString Key = FString(CacheKey).ToUpper();
-	checkf(Algo::AllOf(Key, [](TCHAR C) { return FChar::IsAlnum(C) || FChar::IsUnderscore(C) || C == TEXT('$'); }),
+	TStringBuilder<256> Key;
+	Key << CacheKey;
+	for (TCHAR& Char : MakeArrayView(Key))
+	{
+		Char = FChar::ToUpper(Char);
+	}
+	checkf(Algo::AllOf(MakeArrayView(Key),
+		[](TCHAR C) { return FChar::IsAlnum(C) || FChar::IsUnderscore(C) || C == TEXT('$'); }),
 		TEXT("Invalid characters in cache key %s"), CacheKey);
-	uint32 Hash = FCrc::StrCrc_DEPRECATED(*Key);
-	// this creates a tree of 1000 directories
-	FString HashPath = FString::Printf(TEXT("%1d/%1d/%1d/"), (Hash / 100) % 10, (Hash / 10) % 10, Hash % 10);
-	return HashPath / Key + TEXT(".udd");
+	const uint32 Hash = FCrc::StrCrc_DEPRECATED(Key.Len(), Key.GetData());
+	Path.Appendf(TEXT("%1d/%1d/%1d/%s.udd"), (Hash / 100) % 10, (Hash / 10) % 10, Hash % 10, *Key);
 }
 
 void BuildPathForCacheRecord(const FCacheKey& CacheKey, FStringBuilderBase& Path)
@@ -797,6 +802,8 @@ private:
 	void BuildCacheRecordPath(const FCacheKey& CacheKey, FStringBuilderBase& Path) const;
 	void BuildCacheContentPath(const FIoHash& RawHash, FStringBuilderBase& Path) const;
 
+	bool SaveFileWithHash(FStringBuilderBase& Path, FStringView Context, TFunctionRef<void (FArchive&)> WriteFunction) const;
+	bool LoadFileWithHash(FStringBuilderBase& Path, FStringView Context, TFunctionRef<void (FArchive&)> ReadFunction) const;
 	bool SaveFile(FStringBuilderBase& Path, FStringView Context, TFunctionRef<void (FArchive&)> WriteFunction) const;
 	bool LoadFile(FStringBuilderBase& Path, FStringView Context, TFunctionRef<void (FArchive&)> ReadFunction) const;
 	TUniquePtr<FArchive> OpenFile(FStringBuilderBase& Path, FStringView Context) const;
@@ -807,15 +814,10 @@ private:
 	bool ShouldSimulateMiss(const FCacheKey& Key);
 
 private:
-	/**
-	 * Thread-safe method to compute the filename from the cache key, currently just adds a path and an extension.
-	 *
-	 * @param	CacheKey	Alphanumeric+underscore key of this cache item
-	 * @return				filename built from the cache key
-	 */
-	FString BuildFilename(const TCHAR* const CacheKey) const
+	void BuildCacheLegacyPath(const TCHAR* const CacheKey, FStringBuilderBase& Path) const
 	{
-		return CachePath / BuildPathForCacheKey(CacheKey);
+		Path << CachePath << TEXT('/');
+		BuildPathForLegacyCache(CacheKey, Path);
 	}
 
 	/** Base path we are storing the cache files in. */
@@ -1265,38 +1267,35 @@ bool FFileSystemCacheStore::CachedDataProbablyExists(const TCHAR* const CacheKey
 	TRACE_CPUPROFILER_EVENT_SCOPE(FileSystemDDC_Exist);
 	TRACE_COUNTER_INCREMENT(FileSystemDDC_Exist);
 	COOK_STAT(auto Timer = UsageStats.TimeProbablyExists());
-	check(IsUsable());
 
-	if (ShouldSimulateMiss(CacheKey))
+	if (!IsUsable())
 	{
+		UE_LOG(LogDerivedDataCache, VeryVerbose,
+			TEXT("%s: Skipped exists check of %s because this cache store is not available"), *CachePath, CacheKey);
 		return false;
 	}
 
-	const FString Filename = BuildFilename(CacheKey);
-
-	const FFileStatData FileStat = IFileManager::Get().GetStatData(*Filename);
-
-	if (FileStat.bIsValid)
+	if (ShouldSimulateMiss(CacheKey))
 	{
-		const FDateTime TimeStamp = FileStat.ModificationTime;
-
-		// Update file timestamp to prevent it from being deleted by maintenance.
-		if (bTouch || (!bReadOnly && (FDateTime::UtcNow() - TimeStamp).GetTotalDays() > (DaysToDeleteUnusedFiles / 4)))
-		{
-			IFileManager::Get().SetTimeStamp(*Filename, FDateTime::UtcNow());
-		}
+		UE_LOG(LogDerivedDataCache, Verbose, TEXT("%s: Simulated miss for exists check of %s"), *CachePath, CacheKey);
+		return false;
+	}
 
-		if (AccessLogWriter.IsValid())
+	TStringBuilder<256> LegacyPath;
+	BuildCacheLegacyPath(CacheKey, LegacyPath);
+	const bool bExists = FileExists(LegacyPath);
+	if (bExists)
+	{
+		if (AccessLogWriter)
 		{
-			AccessLogWriter->Append(CacheKey, Filename);
+			AccessLogWriter->Append(CacheKey, LegacyPath);
 		}
 
 		TRACE_COUNTER_INCREMENT(FileSystemDDC_ExistHit);
 		COOK_STAT(Timer.AddHit(0));
 	}
-
 	// If not using a shared cache, record a (probable) miss
-	if (!FileStat.bIsValid && !GetDerivedDataCacheRef().GetUsingSharedDDC())
+	else if (!GetDerivedDataCacheRef().GetUsingSharedDDC())
 	{
 		// store a cache miss
 		FScopeLock ScopeLock(&SynchronizationObject);
@@ -1307,9 +1306,9 @@ bool FFileSystemCacheStore::CachedDataProbablyExists(const TCHAR* const CacheKey
 	}
 
 	UE_LOG(LogDerivedDataCache, Verbose,
-		TEXT("%s CachedDataProbablyExists=%d for %s"), *GetName(), FileStat.bIsValid, CacheKey);
+		TEXT("%s: CachedDataProbablyExists=%d for %s"), *CachePath, int32(bExists), CacheKey);
 
-	return FileStat.bIsValid;
+	return bExists;
 }
 
 bool FFileSystemCacheStore::GetCachedData(const TCHAR* const CacheKey, TArray<uint8>& Data)
@@ -1317,58 +1316,62 @@ bool FFileSystemCacheStore::GetCachedData(const TCHAR* const CacheKey, TArray<ui
 	TRACE_CPUPROFILER_EVENT_SCOPE(FileSystemDDC_Get);
 	TRACE_COUNTER_INCREMENT(FileSystemDDC_Get);
 	COOK_STAT(auto Timer = UsageStats.TimeGet());
-	check(IsUsable());
-	const FString Filename = BuildFilename(CacheKey);
+
 	const double StartTime = FPlatformTime::Seconds();
 
+	if (!IsUsable())
+	{
+		UE_LOG(LogDerivedDataCache, VeryVerbose,
+			TEXT("%s: Skipped get of %s because this cache store is not available"), *CachePath, CacheKey);
+		return false;
+	}
+
 	if (ShouldSimulateMiss(CacheKey))
 	{
+		UE_LOG(LogDerivedDataCache, Verbose, TEXT("%s: Simulated miss for get of %s"), *CachePath, CacheKey);
 		return false;
 	}
 
-	if (FFileHelper::LoadFileToArray(Data,*Filename,FILEREAD_Silent))
+	const auto ReadData = [this, CacheKey, &Data](FArchive& Ar)
 	{
-		if (SpeedClass >= FDerivedDataBackendInterface::ESpeedClass::Fast && (!bReadOnly || bTouch))
+		const int64 Size = Ar.TotalSize();
+		if (Size > MAX_int32)
 		{
-			const FFileStatData FileStat = IFileManager::Get().GetStatData(*Filename);
-			if (FileStat.bIsValid)
-			{
-				const FDateTime TimeStamp = FileStat.ModificationTime;
-				// Update file timestamp to prevent it from being deleted by DDC Cleanup.
-				if (bTouch ||
-					(!bReadOnly && (FDateTime::UtcNow() - TimeStamp).GetTotalDays() > (DaysToDeleteUnusedFiles / 4)))
-				{
-					IFileManager::Get().SetTimeStamp(*Filename, FDateTime::UtcNow());
-				}
-			}
+			Ar.SetError();
+			return;
 		}
 
-		const double ReadDuration = FPlatformTime::Seconds() - StartTime;
-		const double ReadSpeed = (Data.Num() / ReadDuration) / (1024.0 * 1024.0);
+		Data.Reset(int32(Size));
+		Data.AddUninitialized(int32(Size));
+		Ar.Serialize(Data.GetData(), Size);
 
-		if(!GIsBuildMachine && ReadDuration > 5.0)
+		if (!FCorruptionWrapper::ReadTrailer(Data, *CachePath, CacheKey))
 		{
-			// Slower than 0.5MB/s?
-			UE_CLOG(ReadSpeed < 0.5, LogDerivedDataCache, Warning,
-				TEXT("%s is very slow (%.2fMB/s) when accessing %s, consider disabling it."),
-				*CachePath, ReadSpeed, *Filename);
+			Ar.SetError();
 		}
+	};
 
-		if (AccessLogWriter.IsValid())
+	TStringBuilder<256> LegacyPath;
+	BuildCacheLegacyPath(CacheKey, LegacyPath);
+	if (LoadFile(LegacyPath, CacheKey, ReadData))
+	{
+		if (AccessLogWriter)
 		{
-			AccessLogWriter->Append(CacheKey, Filename);
+			AccessLogWriter->Append(CacheKey, LegacyPath);
 		}
 
+		const double ReadDuration = FPlatformTime::Seconds() - StartTime;
+		const double ReadSpeed = ReadDuration > 0.001 ? (Data.Num() / ReadDuration) / (1024.0 * 1024.0) : 0.0;
 		UE_LOG(LogDerivedDataCache, Verbose,
 			TEXT("%s: Cache hit on %s (%d bytes, %.02f secs, %.2fMB/s)"),
-			*GetName(), CacheKey, Data.Num(), ReadDuration, ReadSpeed);
+			*CachePath, CacheKey, Data.Num(), ReadDuration, ReadSpeed);
 		TRACE_COUNTER_INCREMENT(FileSystemDDC_GetHit);
 		TRACE_COUNTER_ADD(FileSystemDDC_BytesRead, int64(Data.Num()));
 		COOK_STAT(Timer.AddHit(Data.Num()));
 		return true;
 	}
 
-	UE_LOG(LogDerivedDataCache, Verbose, TEXT("%s: Cache miss on %s"), *GetName(), CacheKey);
+	UE_LOG(LogDerivedDataCache, Verbose, TEXT("%s: Cache miss on %s"), *CachePath, CacheKey);
 	Data.Empty();
 
 	// If not using a shared cache, record a miss
@@ -1398,135 +1401,103 @@ FFileSystemCacheStore::EPutStatus FFileSystemCacheStore::PutCachedData(
 	TRACE_CPUPROFILER_EVENT_SCOPE(FileSystemDDC_Put);
 	TRACE_COUNTER_INCREMENT(FileSystemDDC_Put);
 	COOK_STAT(auto Timer = UsageStats.TimePut());
-	check(IsUsable());
 
-	if (IsWritable())
+	checkf(Data.Num(), TEXT("%s: Failed to put %s due to empty data"), *CachePath, CacheKey);
+
+	if (!IsUsable())
+	{
+		UE_LOG(LogDerivedDataCache, VeryVerbose,
+			TEXT("%s: Skipped put of %s because this cache store is not available"), *CachePath, CacheKey);
+		return EPutStatus::NotCached;
+	}
+
+	if (!IsWritable())
+	{
+		UE_LOG(LogDerivedDataCache, VeryVerbose,
+			TEXT("%s: Skipped put of %s because this cache store is read-only"), *CachePath, CacheKey);
+		return EPutStatus::NotCached;
+	}
+
+	if (ShouldSimulateMiss(CacheKey))
 	{
-		const FString Filename = BuildFilename(CacheKey);
+		UE_LOG(LogDerivedDataCache, Verbose, TEXT("%s: Simulated miss for put of %s"), *CachePath, CacheKey);
+		return EPutStatus::Skipped;
+	}
 
-		if (AccessLogWriter.IsValid())
-		{
-			AccessLogWriter->Append(CacheKey, Filename);
-		}
+	TStringBuilder<256> LegacyPath;
+	BuildCacheLegacyPath(CacheKey, LegacyPath);
 
-		// don't put anything we pretended didn't exist
-		if (ShouldSimulateMiss(CacheKey))
-		{
-			return EPutStatus::Skipped;
-		}
+	if (AccessLogWriter)
+	{
+		AccessLogWriter->Append(CacheKey, LegacyPath);
+	}
+
+	EPutStatus Status = EPutStatus::NotCached;
 
-		EPutStatus Status = EPutStatus::NotCached;
+	if (bPutEvenIfExists || !FileExists(LegacyPath))
+	{
+		const auto WriteData = [this, CacheKey, &Data](FArchive& Ar)
+		{
+			Ar.Serialize(const_cast<uint8*>(Data.GetData()), Data.Num());
+			FCorruptionWrapper::WriteTrailer(Ar, Data, *CachePath, CacheKey);
+		};
 
-		if (bPutEvenIfExists || !CachedDataProbablyExists(CacheKey))
+		if (SaveFile(LegacyPath, CacheKey, WriteData))
 		{
+			UE_LOG(LogDerivedDataCache, Verbose,
+				TEXT("%s: Successful cache put of %s to %s"), *CachePath, CacheKey, *LegacyPath);
 			TRACE_COUNTER_INCREMENT(FileSystemDDC_PutHit);
 			TRACE_COUNTER_ADD(FileSystemDDC_BytesWritten, int64(Data.Num()));
 			COOK_STAT(Timer.AddHit(Data.Num()));
-			check(Data.Num());
-			FString TempFilename(TEXT("temp."));
-			TempFilename += FGuid::NewGuid().ToString();
-			TempFilename = FPaths::GetPath(Filename) / TempFilename;
-			if (FFileHelper::SaveArrayToFile(Data, *TempFilename, &IFileManager::Get(), FILEWRITE_Silent))
-			{
-				if (IFileManager::Get().FileSize(*TempFilename) == Data.Num())
-				{
-					bool DoMove = !CachedDataProbablyExists(CacheKey);
-					if (bPutEvenIfExists && !DoMove)
-					{
-						DoMove = true;
-						RemoveCachedData(CacheKey, /*bTransient=*/ false);
-					}
-					if (DoMove)
-					{
-						if (!IFileManager::Get().Move(*Filename, *TempFilename, true, true, false, true))
-						{
-							UE_LOG(LogDerivedDataCache, Log,
-								TEXT("%s: Move collision, attempt at redundant update, OK %s."), *GetName(), *Filename);
-						}
-						else
-						{
-							UE_LOG(LogDerivedDataCache, Verbose,
-								TEXT("%s: Successful cache put of %s to %s"), *GetName(), CacheKey, *Filename);
-							Status = EPutStatus::Cached;
-						}
-					}
-					else
-					{
-						Status = EPutStatus::Cached;
-					}
-				}
-				else
-				{
-					UE_LOG(LogDerivedDataCache, Warning, TEXT("%s: Temp file is short %s!"), *GetName(), *TempFilename);
-				}
-			}
-			else
-			{
-				const uint32 ErrorCode = FPlatformMisc::GetLastError();
-				TCHAR ErrorBuffer[1024];
-				FPlatformMisc::GetSystemErrorMessage(ErrorBuffer, 1024, ErrorCode);
-				UE_LOG(LogDerivedDataCache, Warning,
-					TEXT("FFileSystemCacheStore: Could not write temp file %s! Error: %u (%s)"),
-					*TempFilename, ErrorCode, ErrorBuffer);
-			}
-			// if everything worked, this is not necessary, but we will make every effort to avoid leaving junk in the cache
-			if (FPaths::FileExists(TempFilename))
-			{
-				IFileManager::Get().Delete(*TempFilename, false, false, true);
-			}
-		}
-		else
-		{
-			COOK_STAT(Timer.AddMiss(Data.Num()));
-			UE_LOG(LogDerivedDataCache, Verbose, TEXT("%s skipping put to existing file %s"), *GetName(), CacheKey);
 			Status = EPutStatus::Cached;
 		}
+	}
+	else
+	{
+		COOK_STAT(Timer.AddMiss(Data.Num()));
+		UE_LOG(LogDerivedDataCache, Verbose, TEXT("%s: Skipping put of %s to existing file"), *CachePath, CacheKey);
+		Status = EPutStatus::Cached;
+	}
 
-		// If not using a shared cache, update estimated build time
-		if (!GetDerivedDataCacheRef().GetUsingSharedDDC())
+	// If not using a shared cache, update estimated build time
+	if (!GetDerivedDataCacheRef().GetUsingSharedDDC())
+	{
+		FScopeLock ScopeLock(&SynchronizationObject);
+
+		if (DDCNotificationCacheTimes.Contains(CacheKey))
 		{
-			FScopeLock ScopeLock(&SynchronizationObject);
+			// There isn't any way to get exact build times in the DDC code as custom asset processing and async are factors.
+			// So, estimate the asset build time based on the delta between the cache miss and the put
+			TotalEstimatedBuildTime += (FPlatformTime::Seconds() - DDCNotificationCacheTimes[CacheKey]);
+			DDCNotificationCacheTimes.Remove(CacheKey);
 
-			if (DDCNotificationCacheTimes.Contains(CacheKey))
+			// If more than 20 seconds has been spent building assets, send out a notification
+			if (TotalEstimatedBuildTime > 20.0f)
 			{
-				// There isn't any way to get exact build times in the DDC code as custom asset processing and async are factors.
-				// So, estimate the asset build time based on the delta between the cache miss and the put
-				TotalEstimatedBuildTime += (FPlatformTime::Seconds() - DDCNotificationCacheTimes[CacheKey]);
-				DDCNotificationCacheTimes.Remove(CacheKey);
+				// Send out a DDC put notification if we have any subscribers
+				FDerivedDataCacheInterface::FOnDDCNotification& DDCNotificationEvent =
+					GetDerivedDataCacheRef().GetDDCNotificationEvent();
 
-				// If more than 20 seconds has been spent building assets, send out a notification
-				if (TotalEstimatedBuildTime > 20.0f)
+				if (DDCNotificationEvent.IsBound())
 				{
-					// Send out a DDC put notification if we have any subscribers
-					FDerivedDataCacheInterface::FOnDDCNotification& DDCNotificationEvent =
-						GetDerivedDataCacheRef().GetDDCNotificationEvent();
-
-					if (DDCNotificationEvent.IsBound())
-					{
-						TotalEstimatedBuildTime = 0.0f;
-
-						DECLARE_CYCLE_STAT(TEXT("FSimpleDelegateGraphTask.PutCachedData"),
-						STAT_FSimpleDelegateGraphTask_DDCNotification, STATGROUP_TaskGraphTasks);
-
-						FSimpleDelegateGraphTask::CreateAndDispatchWhenReady(
-							FSimpleDelegateGraphTask::FDelegate::CreateLambda([DDCNotificationEvent]() {
-							DDCNotificationEvent.Broadcast(FDerivedDataCacheInterface::SharedDDCPerformanceNotification);
-						}),
-							GET_STATID(STAT_FSimpleDelegateGraphTask_DDCNotification),
-							nullptr,
-							ENamedThreads::GameThread);
-					}
+					TotalEstimatedBuildTime = 0.0f;
+
+					DECLARE_CYCLE_STAT(TEXT("FSimpleDelegateGraphTask.PutCachedData"),
+					STAT_FSimpleDelegateGraphTask_DDCNotification, STATGROUP_TaskGraphTasks);
+
+					FSimpleDelegateGraphTask::CreateAndDispatchWhenReady(
+						FSimpleDelegateGraphTask::FDelegate::CreateLambda([DDCNotificationEvent]() {
+						DDCNotificationEvent.Broadcast(FDerivedDataCacheInterface::SharedDDCPerformanceNotification);
+					}),
+						GET_STATID(STAT_FSimpleDelegateGraphTask_DDCNotification),
+						nullptr,
+						ENamedThreads::GameThread);
 				}
 			}
 		}
-
-		return Status;
-	}
-	else
-	{
-		UE_LOG(LogDerivedDataCache, Verbose, TEXT("%s is read only. Skipping put of %s"), *GetName(), CacheKey);
-		return EPutStatus::NotCached;
 	}
+
+	return Status;
 }
 
 void FFileSystemCacheStore::RemoveCachedData(const TCHAR* const CacheKey, const bool bTransient)
@@ -1534,13 +1505,14 @@ void FFileSystemCacheStore::RemoveCachedData(const TCHAR* const CacheKey, const
 	check(IsUsable());
 	if (IsWritable() && (!bTransient || bPurgeTransient))
 	{
-		const FString Filename = BuildFilename(CacheKey);
+		TStringBuilder<256> LegacyPath;
+		BuildCacheLegacyPath(CacheKey, LegacyPath);
 		if (bTransient)
 		{
 			UE_LOG(LogDerivedDataCache, Verbose,
-				TEXT("Deleting transient cached data. Key=%s Filename=%s"), CacheKey, *Filename);
+				TEXT("%s: Deleting cached data for %s from %s"), *CachePath, CacheKey, *LegacyPath);
 		}
-		IFileManager::Get().Delete(*Filename, false, false, true);
+		IFileManager::Get().Delete(*LegacyPath, /*bRequireExists*/ false, /*bEvenReadOnly*/ false, /*bQuiet*/ true);
 	}
 }
 
@@ -1729,7 +1701,7 @@ bool FFileSystemCacheStore::PutCacheRecord(
 	if (!IsWritable())
 	{
 		UE_LOG(LogDerivedDataCache, VeryVerbose,
-			TEXT("%s: Skipped put of %s from '%.*s' because this cache store is not writable"),
+			TEXT("%s: Skipped put of %s from '%.*s' because this cache store is read-only"),
 			*CachePath, *WriteToString<96>(Record.GetKey()), Context.Len(), Context.GetData());
 		return false;
 	}
@@ -1766,7 +1738,7 @@ bool FFileSystemCacheStore::PutCacheRecord(
 	}
 	else
 	{
-		bRecordExists = LoadFile(Path, Context, [&ExistingPackage](FArchive& Ar) { ExistingPackage.TryLoad(Ar); });
+		bRecordExists = LoadFileWithHash(Path, Context, [&ExistingPackage](FArchive& Ar) { ExistingPackage.TryLoad(Ar); });
 	}
 
 	// Save the record to a package and remove attachments that will be stored externally.
@@ -1818,7 +1790,7 @@ bool FFileSystemCacheStore::PutCacheRecord(
 	}
 
 	// Save the record package to storage.
-	if (!bRecordExists && !SaveFile(Path, Context, [&Package](FArchive& Ar) { Package.Save(Ar); }))
+	if (!bRecordExists && !SaveFileWithHash(Path, Context, [&Package](FArchive& Ar) { Package.Save(Ar); }))
 	{
 		return false;
 	}
@@ -1876,7 +1848,7 @@ FOptionalCacheRecord FFileSystemCacheStore::GetCacheRecordOnly(
 	FOptionalCacheRecord Record;
 	{
 		FCbPackage Package;
-		if (!LoadFile(Path, Context, [&Package](FArchive& Ar) { Package.TryLoad(Ar); }))
+		if (!LoadFileWithHash(Path, Context, [&Package](FArchive& Ar) { Package.TryLoad(Ar); }))
 		{
 			UE_LOG(LogDerivedDataCache, Verbose, TEXT("%s: Cache miss with missing record for %s from '%.*s'"),
 				*CachePath, *WriteToString<96>(Key), Context.Len(), Context.GetData());
@@ -1976,7 +1948,7 @@ bool FFileSystemCacheStore::PutCacheContent(const FCompressedBuffer& Content, co
 	const FIoHash& RawHash = Content.GetRawHash();
 	TStringBuilder<256> Path;
 	BuildCacheContentPath(RawHash, Path);
-	if (!FileExists(Path) && !SaveFile(Path, Context, [&Content](FArchive& Ar) { Content.Save(Ar); }))
+	if (!FileExists(Path) && !SaveFileWithHash(Path, Context, [&Content](FArchive& Ar) { Content.Save(Ar); }))
 	{
 		return false;
 	}
@@ -2022,7 +1994,7 @@ FPayload FFileSystemCacheStore::GetCacheContent(
 	else
 	{
 		FCompressedBuffer CompressedBuffer;
-		if (LoadFile(Path, Context, [&CompressedBuffer](FArchive& Ar) { CompressedBuffer = FCompressedBuffer::Load(Ar); }))
+		if (LoadFileWithHash(Path, Context, [&CompressedBuffer](FArchive& Ar) { CompressedBuffer = FCompressedBuffer::Load(Ar); }))
 		{
 			if (CompressedBuffer.GetRawHash() == RawHash)
 			{
@@ -2128,6 +2100,44 @@ void FFileSystemCacheStore::BuildCacheContentPath(const FIoHash& RawHash, FStrin
 	BuildPathForCacheContent(RawHash, Path);
 }
 
+bool FFileSystemCacheStore::SaveFileWithHash(
+	FStringBuilderBase& Path,
+	const FStringView Context,
+	const TFunctionRef<void (FArchive&)> WriteFunction) const
+{
+	return SaveFile(Path, Context, [&WriteFunction](FArchive& Ar)
+	{
+		THashingArchiveProxy<FBlake3> HashAr(Ar);
+		WriteFunction(HashAr);
+		FBlake3Hash Hash = HashAr.GetHash();
+		Ar << Hash;
+	});
+}
+
+bool FFileSystemCacheStore::LoadFileWithHash(
+	FStringBuilderBase& Path,
+	const FStringView Context,
+	const TFunctionRef<void (FArchive& Ar)> ReadFunction) const
+{
+	return LoadFile(Path, Context, [this, &Path, &Context, &ReadFunction](FArchive& Ar)
+	{
+		THashingArchiveProxy<FBlake3> HashAr(Ar);
+		ReadFunction(HashAr);
+		const FBlake3Hash Hash = HashAr.GetHash();
+		FBlake3Hash SavedHash;
+		Ar << SavedHash;
+
+		if (Hash != SavedHash && !Ar.IsError())
+		{
+			Ar.SetError();
+			UE_LOG(LogDerivedDataCache, Display,
+				TEXT("%s: File %s from '%.*s' is corrupted and has hash %s when %s is expected."),
+				*CachePath, *Path, Context.Len(), Context.GetData(),
+				*WriteToString<80>(Hash), *WriteToString<80>(SavedHash));
+		}
+	});
+}
+
 bool FFileSystemCacheStore::SaveFile(
 	FStringBuilderBase& Path,
 	const FStringView Context,
@@ -2146,12 +2156,9 @@ bool FFileSystemCacheStore::SaveFile(
 
 	if (TUniquePtr<FArchive> Ar{IFileManager::Get().CreateFileWriter(*TempPath, FILEWRITE_Silent)})
 	{
-		THashingArchiveProxy<FBlake3> HashAr(*Ar);
-		WriteFunction(HashAr);
-		FBlake3Hash Hash = HashAr.GetHash();
-		*Ar << Hash;
-		WriteSize = HashAr.Tell();
-		bError = !HashAr.Close();
+		WriteFunction(*Ar);
+		WriteSize = Ar->Tell();
+		bError = !Ar->Close();
 	}
 
 	if (bError || WriteSize == 0)
@@ -2196,13 +2203,9 @@ bool FFileSystemCacheStore::LoadFile(
 
 	if (TUniquePtr<FArchive> Ar = OpenFile(Path, Context))
 	{
-		THashingArchiveProxy<FBlake3> HashAr(*Ar);
-		ReadFunction(HashAr);
-		const FBlake3Hash Hash = HashAr.GetHash();
-		FBlake3Hash SavedHash;
-		*Ar << SavedHash;
-		ReadSize = HashAr.Tell();
-		bError = !HashAr.Close();
+		ReadFunction(*Ar);
+		ReadSize = Ar->Tell();
+		bError = !Ar->Close();
 
 		if (bError)
 		{
@@ -2210,14 +2213,6 @@ bool FFileSystemCacheStore::LoadFile(
 				TEXT("%s: Failed to load file %s from '%.*s'."),
 				*CachePath, *Path, Context.Len(), Context.GetData());
 		}
-		else if (Hash != SavedHash)
-		{
-			bError = true;
-			UE_LOG(LogDerivedDataCache, Display,
-				TEXT("%s: File %s from '%.*s' is corrupted and has hash %s when %s is expected."),
-				*CachePath, *Path, Context.Len(), Context.GetData(),
-				*WriteToString<80>(Hash), *WriteToString<80>(SavedHash));
-		}
 	}
 
 	const double ReadDuration = FPlatformTime::Seconds() - StartTime;
@@ -2287,7 +2282,6 @@ bool FFileSystemCacheStore::ShouldSimulateMiss(const TCHAR* const InKey)
 	if (DebugOptions.ShouldSimulateMiss(InKey))
 	{
 		FScopeLock Lock(&MissedKeysCS);
-		UE_LOG(LogDerivedDataCache, Verbose, TEXT("Simulating miss in %s for %s"), *GetName(), InKey);
 		DebugMissedKeys.AddByHash(Hash, Key);
 		return true;
 	}
