DDC: Fixed the file access log to write paths relative to the cache directory

#rnx

diff --git a/Engine/Source/Developer/DerivedDataCache/Private/FileSystemDerivedDataBackend.cpp b/Engine/Source/Developer/DerivedDataCache/Private/FileSystemDerivedDataBackend.cpp
--- a/Engine/Source/Developer/DerivedDataCache/Private/FileSystemDerivedDataBackend.cpp
+++ b/Engine/Source/Developer/DerivedDataCache/Private/FileSystemDerivedDataBackend.cpp
@@ -713,7 +713,7 @@ void FAccessLogWriter::AppendPath(const FStringView Path)
 {
 	if (Path.StartsWith(BasePath))
 	{
-		const FTCHARToUTF8 PathUtf8(Path);
+		const FTCHARToUTF8 PathUtf8(Path.RightChop(BasePath.Len()));
 		Archive->Serialize(const_cast<ANSICHAR*>(PathUtf8.Get()), PathUtf8.Length());
 		Archive->Serialize(const_cast<ANSICHAR*>(LINE_TERMINATOR_ANSI), sizeof(LINE_TERMINATOR_ANSI) - 1);
 	}
