DDC: Folded the corruption wrapper into the S3 cache

#rnx

diff --git a/Engine/Source/Developer/DerivedDataCache/Private/DerivedDataBackendCorruptionWrapper.h b/Engine/Source/Developer/DerivedDataCache/Private/DerivedDataBackendCorruptionWrapper.h
--- a/Engine/Source/Developer/DerivedDataCache/Private/DerivedDataBackendCorruptionWrapper.h
+++ b/Engine/Source/Developer/DerivedDataCache/Private/DerivedDataBackendCorruptionWrapper.h
@@ -2,11 +2,13 @@
 
 #pragma once
 
-#include "CoreMinimal.h"
+#include "CoreTypes.h"
+#include "Containers/Array.h"
 #include "DerivedDataBackendInterface.h"
-#include "ProfilingDebugging/CookStats.h"
-#include "DerivedDataCacheUsageStats.h"
+#include "HAL/UnrealMemory.h"
 #include "Memory/MemoryView.h"
+#include "Misc/Crc.h"
+#include "Serialization/Archive.h"
 
 namespace UE::DerivedData::CacheStore
 {
@@ -81,229 +83,4 @@ public:
 	}
 };
 
-/** 
- * A backend wrapper that adds a footer to the data to check the CRC, length, etc.
-**/
-class FDerivedDataBackendCorruptionWrapper : public FDerivedDataBackendInterface
-{
-public:
-
-	/**
-	 * Constructor
-	 *
-	 * @param	InInnerBackend	Backend to use for storage, my responsibilities are about corruption
-	 */
-	FDerivedDataBackendCorruptionWrapper(FDerivedDataBackendInterface* InInnerBackend)
-		: InnerBackend(InInnerBackend)
-	{
-		check(InnerBackend);
-	}
-
-	/** Return a name for this interface */
-	virtual FString GetDisplayName() const override
-	{
-		return FString(TEXT("CorruptionWrapper"));
-	}
-
-	/** Return a name for this interface */
-	virtual FString GetName() const override
-	{
-		return FString::Printf(TEXT("CorruptionWrapper (%s)"), *InnerBackend->GetName());
-	}
-
-	/** return true if this cache is writable **/
-	virtual bool IsWritable() const override
-	{
-		return InnerBackend->IsWritable();
-	}
-
-	/** Returns a class of speed for this interface **/
-	virtual ESpeedClass GetSpeedClass() const override
-	{
-		return InnerBackend->GetSpeedClass();
-	}
-
-	/**
-	 * Synchronous test for the existence of a cache item
-	 *
-	 * @param	CacheKey	Alphanumeric+underscore key of this cache item
-	 * @return				true if the data probably will be found, this can't be guaranteed because of concurrency in the backends, corruption, etc
-	 */
-	virtual bool CachedDataProbablyExists(const TCHAR* CacheKey) override
-	{
-		COOK_STAT(auto Timer = UsageStats.TimeProbablyExists());
-		bool Result = InnerBackend->CachedDataProbablyExists(CacheKey);
-		if (Result)
-		{
-			COOK_STAT(Timer.AddHit(0));
-		}
-		return Result;
-	}
-
-	/**
-	 * Synchronous test for the existence of multiple cache items
-	 *
-	 * @param	CacheKeys	Alphanumeric+underscore key of the cache items
-	 * @return				A bit array with bits indicating whether the data for the corresponding key will probably be found
-	 */
-	virtual TBitArray<> CachedDataProbablyExistsBatch(TConstArrayView<FString> CacheKeys) override
-	{
-		COOK_STAT(auto Timer = UsageStats.TimeProbablyExists());
-		TBitArray<> Result = InnerBackend->CachedDataProbablyExistsBatch(CacheKeys);
-		if (Result.CountSetBits() == CacheKeys.Num())
-		{
-			COOK_STAT(Timer.AddHit(0));
-		}
-		return Result;
-	}
-
-	/**
-	 * Synchronous retrieve of a cache item
-	 *
-	 * @param	CacheKey	Alphanumeric+underscore key of this cache item
-	 * @param	OutData		Buffer to receive the results, if any were found
-	 * @return				true if any data was found, and in this case OutData is non-empty
-	 */
-	virtual bool GetCachedData(const TCHAR* CacheKey, TArray<uint8>& OutData) override
-	{
-		COOK_STAT(auto Timer = UsageStats.TimeGet());
-		bool bOk = InnerBackend->GetCachedData(CacheKey, OutData);
-		if (bOk)
-		{
-			if (OutData.Num() < sizeof(FDerivedDataTrailer))
-			{
-				UE_LOG(LogDerivedDataCache, Display, TEXT("FDerivedDataBackendCorruptionWrapper: Corrupted file (short), ignoring and deleting %s."),CacheKey);
-				bOk	= false;
-			}
-			else
-			{
-				FDerivedDataTrailer Trailer;
-				FMemory::Memcpy(&Trailer,&OutData[OutData.Num() - sizeof(FDerivedDataTrailer)], sizeof(FDerivedDataTrailer));
-				OutData.RemoveAt(OutData.Num() - sizeof(FDerivedDataTrailer),sizeof(FDerivedDataTrailer), false);
-				FDerivedDataTrailer RecomputedTrailer(MakeMemoryView(OutData));
-				if (Trailer == RecomputedTrailer)
-				{
-					UE_LOG(LogDerivedDataCache, VeryVerbose, TEXT("FDerivedDataBackendCorruptionWrapper: cache hit, footer is ok %s"),CacheKey);
-				}
-				else
-				{
-					UE_LOG(LogDerivedDataCache, Display, TEXT("FDerivedDataBackendCorruptionWrapper: Corrupted file, ignoring and deleting %s."),CacheKey);
-					bOk	= false;
-				}
-			}
-			if (!bOk)
-			{
-				// _we_ detected corruption, so _we_ will force a flush of the corrupted data
-				InnerBackend->RemoveCachedData(CacheKey, /*bTransient=*/ false);
-			}
-		}
-		if (!bOk)
-		{
-			OutData.Empty();
-		}
-		else
-		{
-			COOK_STAT(Timer.AddHit(OutData.Num()));
-		}
-		return bOk;
-	}
-	/**
-	 * Asynchronous, fire-and-forget placement of a cache item
-	 *
-	 * @param	CacheKey	Alphanumeric+underscore key of this cache item
-	 * @param	InData		Buffer containing the data to cache, can be destroyed after the call returns, immediately
-	 * @param	bPutEvenIfExists	If true, then do not attempt skip the put even if CachedDataProbablyExists returns true
-	 */
-	virtual EPutStatus PutCachedData(const TCHAR* CacheKey, TArrayView<const uint8> InData, bool bPutEvenIfExists) override
-	{
-		COOK_STAT(auto Timer = UsageStats.TimePut());
-		if (!InnerBackend->IsWritable())
-		{
-			return EPutStatus::NotCached; // no point in continuing down the chain
-		}
-		COOK_STAT(Timer.AddHit(InData.Num()));
-		// Get rid of the double copy!
-		checkf(int64(InData.Num()) + sizeof(FDerivedDataTrailer) <= INT32_MAX,
-			TEXT("FDerivedDataBackendCorruptionWrapper: appending the trailer makes the data exceed 2 GiB for %s"), CacheKey);
-		TArray<uint8> Data;
-		Data.Reset( InData.Num() + sizeof(FDerivedDataTrailer) );
-		Data.AddUninitialized(InData.Num());
-		FMemory::Memcpy(&Data[0], &InData[0], InData.Num());
-		FDerivedDataTrailer Trailer(MakeMemoryView(Data));
-		Data.AddUninitialized(sizeof(FDerivedDataTrailer));
-		FMemory::Memcpy(&Data[Data.Num() - sizeof(FDerivedDataTrailer)], &Trailer, sizeof(FDerivedDataTrailer));
-		return InnerBackend->PutCachedData(CacheKey, Data, bPutEvenIfExists);
-	}
-	
-	virtual void RemoveCachedData(const TCHAR* CacheKey, bool bTransient) override
-	{
-		if (!InnerBackend->IsWritable())
-		{
-			return; // no point in continuing down the chain
-		}
-		return InnerBackend->RemoveCachedData(CacheKey, bTransient);
-	}
-
-	virtual TSharedRef<FDerivedDataCacheStatsNode> GatherUsageStats() const override
-	{
-		TSharedRef<FDerivedDataCacheStatsNode> Usage = MakeShared<FDerivedDataCacheStatsNode>(this, TEXT("CorruptionWrapper"));
-		Usage->Stats.Add(TEXT(""), UsageStats);
-
-		if (InnerBackend)
-		{
-			Usage->Children.Add(InnerBackend->GatherUsageStats());
-		}
-
-		return Usage;
-	}
-
-	virtual bool TryToPrefetch(TConstArrayView<FString> CacheKeys) override
-	{
-		return InnerBackend->TryToPrefetch(CacheKeys);
-	}
-
-	virtual bool WouldCache(const TCHAR* CacheKey, TArrayView<const uint8> InData) override
-	{
-		return InnerBackend->WouldCache(CacheKey, InData);
-	}
-
-	virtual bool ApplyDebugOptions(FBackendDebugOptions& InOptions) 
-	{ 
-		return InnerBackend->ApplyDebugOptions(InOptions);
-	}
-
-	virtual void Put(
-		TConstArrayView<FCachePutRequest> Requests,
-		FStringView Context,
-		IRequestOwner& Owner,
-		FOnCachePutComplete&& OnComplete) override
-	{
-		return InnerBackend->Put(Requests, Context, Owner, MoveTemp(OnComplete));
-	}
-
-	virtual void Get(
-		TConstArrayView<FCacheGetRequest> Requests,
-		FStringView Context,
-		IRequestOwner& Owner,
-		FOnCacheGetComplete&& OnComplete) override
-	{
-		return InnerBackend->Get(Requests, Context, Owner, MoveTemp(OnComplete));
-	}
-
-	virtual void GetChunks(
-		TConstArrayView<FCacheChunkRequest> Chunks,
-		FStringView Context,
-		IRequestOwner& Owner,
-		FOnCacheGetChunkComplete&& OnComplete) override
-	{
-		return InnerBackend->GetChunks(Chunks, Context, Owner, MoveTemp(OnComplete));
-	}
-
-private:
-	FDerivedDataCacheUsageStats UsageStats;
-
-	/** Backend to use for storage, my responsibilities are about corruption **/
-	FDerivedDataBackendInterface* InnerBackend;
-};
-
 } // UE::DerivedData::CacheStore
diff --git a/Engine/Source/Developer/DerivedDataCache/Private/DerivedDataBackends.cpp b/Engine/Source/Developer/DerivedDataCache/Private/DerivedDataBackends.cpp
--- a/Engine/Source/Developer/DerivedDataCache/Private/DerivedDataBackends.cpp
+++ b/Engine/Source/Developer/DerivedDataCache/Private/DerivedDataBackends.cpp
@@ -19,7 +19,6 @@
 #include "ZenDerivedDataBackend.h"
 #include "HierarchicalDerivedDataBackend.h"
 #include "DerivedDataLimitKeyLengthWrapper.h"
-#include "DerivedDataBackendCorruptionWrapper.h"
 #include "DerivedDataBackendThrottleWrapper.h"
 #include "DerivedDataBackendVerifyWrapper.h"
 #include "Misc/EngineBuildSettings.h"
@@ -734,9 +733,7 @@ public:
 			}
 		}
 
-		// Insert the backend corruption wrapper. Since the filesystem already uses this, and we're recycling the data with the trailer intact, we need to use it for the S3 cache too.
-		FS3DerivedDataBackend* Backend = new FS3DerivedDataBackend(*ManifestPath, *BaseUrl, *Region, *CanaryObjectKey, *CachePath);
-		return new CacheStore::FDerivedDataBackendCorruptionWrapper(Backend);
+		return new CacheStore::S3::FS3DerivedDataBackend(*ManifestPath, *BaseUrl, *Region, *CanaryObjectKey, *CachePath);
 #else
 		UE_LOG(LogDerivedDataCache, Log, TEXT("S3 backend is not supported on the current platform."));
 		return nullptr;
diff --git a/Engine/Source/Developer/DerivedDataCache/Private/S3DerivedDataBackend.cpp b/Engine/Source/Developer/DerivedDataCache/Private/S3DerivedDataBackend.cpp
--- a/Engine/Source/Developer/DerivedDataCache/Private/S3DerivedDataBackend.cpp
+++ b/Engine/Source/Developer/DerivedDataCache/Private/S3DerivedDataBackend.cpp
@@ -25,6 +25,7 @@
 #include "Misc/OutputDeviceRedirector.h"
 #include "Async/ParallelFor.h"
 #include "Serialization/MemoryReader.h"
+#include "DerivedDataBackendCorruptionWrapper.h"
 #include "DerivedDataCacheRecord.h"
 #include "DerivedDataPayload.h"
 #include "Dom/JsonObject.h"
@@ -51,7 +52,7 @@
 #define S3DDC_MAX_ATTEMPTS 4
 #define S3DDC_MAX_BUFFER_RESERVE 104857600u
 
-namespace UE::DerivedData::Backends
+namespace UE::DerivedData::CacheStore::S3
 {
 
 TRACE_DECLARE_INT_COUNTER(S3DDC_Exist, TEXT("S3DDC Exist"));
@@ -959,7 +960,10 @@ bool FS3DerivedDataBackend::GetCachedData(const TCHAR* CacheKey, TArray<uint8>&
 			OutData.SetNum(BundleEntry->Length);
 			Reader->Seek(BundleEntry->Offset);
 			Reader->Serialize(OutData.GetData(), BundleEntry->Length);
-			return true;
+			if (FCorruptionWrapper::ReadTrailer(OutData, *BaseUrl, CacheKey))
+			{
+				return true;
+			}
 		}
 	}
 
@@ -1264,6 +1268,6 @@ void FS3DerivedDataBackend::GetChunks(
 	}
 }
 
-} // UE::DerivedData::Backends
+} // UE::DerivedData::CacheStore::S3
 
 #endif
diff --git a/Engine/Source/Developer/DerivedDataCache/Private/S3DerivedDataBackend.h b/Engine/Source/Developer/DerivedDataCache/Private/S3DerivedDataBackend.h
--- a/Engine/Source/Developer/DerivedDataCache/Private/S3DerivedDataBackend.h
+++ b/Engine/Source/Developer/DerivedDataCache/Private/S3DerivedDataBackend.h
@@ -13,7 +13,7 @@
 
 #if WITH_S3_DDC_BACKEND
 
-namespace UE::DerivedData::Backends
+namespace UE::DerivedData::CacheStore::S3
 {
 
 /**
@@ -118,6 +118,6 @@ private:
 	TSet<FName> DebugMissedKeys;
 };
 
-} // UE::DerivedData::Backends
+} // UE::DerivedData::CacheStore::S3
 
 #endif
