DDC: Moved FDerivedDataBackendInterface from Public to Private

The interface was in the public folder before 5.0, but was not usable from there in a meaningful way.

#rnx

diff --git a/Engine/Source/Developer/DerivedDataCache/Private/DerivedDataBackendAsyncPutWrapper.cpp b/Engine/Source/Developer/DerivedDataCache/Private/DerivedDataBackendAsyncPutWrapper.cpp
--- a/Engine/Source/Developer/DerivedDataCache/Private/DerivedDataBackendAsyncPutWrapper.cpp
+++ b/Engine/Source/Developer/DerivedDataCache/Private/DerivedDataBackendAsyncPutWrapper.cpp
@@ -330,7 +330,8 @@ void FDerivedDataBackendAsyncPutWrapper::RemoveCachedData(const TCHAR* CacheKey,
 
 TSharedRef<FDerivedDataCacheStatsNode> FDerivedDataBackendAsyncPutWrapper::GatherUsageStats() const
 {
-	TSharedRef<FDerivedDataCacheStatsNode> Usage = MakeShared<FDerivedDataCacheStatsNode>(this, TEXT("AsyncPutWrapper"));
+	TSharedRef<FDerivedDataCacheStatsNode> Usage =
+		MakeShared<FDerivedDataCacheStatsNode>(TEXT("AsyncPutWrapper"), TEXT(""), InnerBackend->GetSpeedClass() == ESpeedClass::Local);
 	Usage->Stats.Add(TEXT("AsyncPut"), UsageStats);
 	Usage->Stats.Add(TEXT("AsyncPutSync"), PutSyncUsageStats);
 
diff --git a/Engine/Source/Developer/DerivedDataCache/Private/DerivedDataBackendAsyncPutWrapper.h b/Engine/Source/Developer/DerivedDataCache/Private/DerivedDataBackendAsyncPutWrapper.h
--- a/Engine/Source/Developer/DerivedDataCache/Private/DerivedDataBackendAsyncPutWrapper.h
+++ b/Engine/Source/Developer/DerivedDataCache/Private/DerivedDataBackendAsyncPutWrapper.h
@@ -66,9 +66,6 @@ public:
 	 */
 	FDerivedDataBackendAsyncPutWrapper(FDerivedDataBackendInterface* InInnerBackend, bool bCacheInFlightPuts);
 
-	/** Return a name for this interface */
-	virtual FString GetDisplayName() const override { return FString(TEXT("AsyncPutWrapper")); }
-
 	/** Return a name for this interface */
 	virtual FString GetName() const override
 	{
@@ -78,12 +75,6 @@ public:
 	/** return true if this cache is writable **/
 	virtual bool IsWritable() const override;
 
-	/** This is a wrapepr type **/
-	virtual bool IsWrapper() const override
-	{
-		return true;
-	}
-
 	/** Returns a class of speed for this interface **/
 	virtual ESpeedClass GetSpeedClass() const override;
 
diff --git a/Engine/Source/Developer/DerivedDataCache/Private/DerivedDataBackendInterface.h b/Engine/Source/Developer/DerivedDataCache/Private/DerivedDataBackendInterface.h
new file mode 100644
--- /dev/null
+++ b/Engine/Source/Developer/DerivedDataCache/Private/DerivedDataBackendInterface.h
@@ -0,0 +1,286 @@
+// Copyright Epic Games, Inc. All Rights Reserved.
+
+
+#pragma once
+
+#include "CoreTypes.h"
+#include "Containers/BitArray.h"
+#include "DerivedDataCache.h"
+#include "Stats/Stats.h"
+
+class FDerivedDataCacheUsageStats;
+class FDerivedDataCacheStatsNode;
+
+namespace UE::DerivedData { struct FCacheKey; }
+
+DECLARE_LOG_CATEGORY_EXTERN(LogDerivedDataCache, Log, All);
+
+DECLARE_DWORD_ACCUMULATOR_STAT_EXTERN(TEXT("Num Gets"),STAT_DDC_NumGets,STATGROUP_DDC, );
+DECLARE_DWORD_ACCUMULATOR_STAT_EXTERN(TEXT("Num Puts"),STAT_DDC_NumPuts,STATGROUP_DDC, );
+DECLARE_DWORD_ACCUMULATOR_STAT_EXTERN(TEXT("Num Build"),STAT_DDC_NumBuilds,STATGROUP_DDC, );
+DECLARE_DWORD_ACCUMULATOR_STAT_EXTERN(TEXT("Num Exists"),STAT_DDC_NumExist,STATGROUP_DDC, );
+DECLARE_FLOAT_ACCUMULATOR_STAT_EXTERN(TEXT("Sync Get Time"),STAT_DDC_SyncGetTime,STATGROUP_DDC, );
+DECLARE_FLOAT_ACCUMULATOR_STAT_EXTERN(TEXT("ASync Wait Time"),STAT_DDC_ASyncWaitTime,STATGROUP_DDC, );
+DECLARE_FLOAT_ACCUMULATOR_STAT_EXTERN(TEXT("Sync Put Time"),STAT_DDC_PutTime,STATGROUP_DDC, );
+DECLARE_FLOAT_ACCUMULATOR_STAT_EXTERN(TEXT("Sync Build Time"),STAT_DDC_SyncBuildTime,STATGROUP_DDC, );
+DECLARE_FLOAT_ACCUMULATOR_STAT_EXTERN(TEXT("Exists Time"),STAT_DDC_ExistTime,STATGROUP_DDC, );
+
+
+
+/**
+ * Interface for cache server backends.
+ * The entire API should be callable from any thread (except the singleton can be assumed to be called at least once before concurrent access).
+ */
+class FDerivedDataBackendInterface : public UE::DerivedData::ICacheStore
+{
+public:
+
+	/** Status of a put operation. */
+	enum class EPutStatus
+	{
+		/** The put is executing asynchronously. */
+		Executing,
+		/** The put completed synchronously and the data was not cached. */
+		NotCached,
+		/** The put completed synchronously and the data was cached. */
+		Cached,
+		/** The put was skipped and should not be retried. */
+		Skipped,
+	};
+
+	/**
+	 * Speed classes. Higher values are faster so > / < comparisons make sense.
+	 */
+	enum class ESpeedClass
+	{
+		Unknown,		/* Don't know yet*/
+		Slow,			/* Slow, likely a remote drive. Some benefit but handle with care */
+		Ok,				/* Ok but not great.  */
+		Fast,			/* Fast but seek times still have an impact */
+		Local			/* Little to no impact from seek times and extremely fast reads */
+	};
+
+	/** Debug options that can be applied to backends to simulate different behavior */
+	struct FBackendDebugOptions
+	{
+		/** Percentage of requests that should result in random misses */
+		int					RandomMissRate;
+
+		/** Apply behavior of this speed class */
+		ESpeedClass			SpeedClass;
+
+		/** Types of DDC entries that should always be a miss */
+		TArray<FString>		SimulateMissTypes;
+
+		FBackendDebugOptions()
+			: RandomMissRate(0)
+			, SpeedClass(ESpeedClass::Unknown)
+		{
+		}
+
+		/** Fill in the provided structure based on the name of the node (e.g. 'shared') and the provided token stream */
+		static bool ParseFromTokens(FBackendDebugOptions& OutOptions, const TCHAR* InNodeName, const TCHAR* InTokens);
+
+		/**
+		 * Returns true if, according to the properties of this struct, the provided key should be treated as a miss.
+		 * Implementing that miss and accounting for any behavior impact (e.g. skipping a subsequent put) is left to
+		 * each backend.
+		 */
+		bool ShouldSimulateMiss(const TCHAR* InCacheKey);
+		bool ShouldSimulateMiss(const UE::DerivedData::FCacheKey& InCacheKey);
+	};
+
+	virtual ~FDerivedDataBackendInterface()
+	{
+	}
+
+	/** Return a name for this interface */
+	virtual FString GetName() const = 0;
+
+	/** return true if this cache is writable */
+	virtual bool IsWritable() const = 0;
+
+	/**
+	 * Returns true if hits on this cache should propagate to lower cache level. Typically false for a PAK file.
+	 * Caution! This generally isn't propagated, so the thing that returns false must be a direct child of the hierarchical cache.
+	 */
+	virtual bool BackfillLowerCacheLevels() const
+	{
+		return true;
+	}
+
+	/**
+	 * Returns a class of speed for this interface
+	 */
+	virtual ESpeedClass GetSpeedClass() const = 0;
+
+	/**
+	 * Synchronous test for the existence of a cache item
+	 *
+	 * @param	CacheKey	Alphanumeric+underscore key of this cache item
+	 * @return				true if the data probably will be found, this can't be guaranteed because of concurrency in the backends, corruption, etc
+	 */
+	virtual bool CachedDataProbablyExists(const TCHAR* CacheKey)=0;
+
+	/**
+	 * Synchronous test for the existence of multiple cache items
+	 *
+	 * @param	CacheKeys	Alphanumeric+underscore key of the cache items
+	 * @return				A bit array with bits indicating whether the data for the corresponding key will probably be found
+	 */
+	virtual TBitArray<> CachedDataProbablyExistsBatch(TConstArrayView<FString> CacheKeys)
+	{
+		TBitArray<> Result;
+		Result.Reserve(CacheKeys.Num());
+		for (const FString& Key : CacheKeys)
+		{
+			Result.Add(CachedDataProbablyExists(*Key));
+		}
+		return Result;
+	}
+
+	/**
+	 * Synchronous retrieve of a cache item
+	 *
+	 * @param	CacheKey	Alphanumeric+underscore key of this cache item
+	 * @param	OutData		Buffer to receive the results, if any were found
+	 * @return				true if any data was found, and in this case OutData is non-empty
+	 */
+	virtual bool GetCachedData(const TCHAR* CacheKey, TArray<uint8>& OutData) = 0;
+
+	/**
+	 * Asynchronous, fire-and-forget placement of a cache item
+	 *
+	 * @param	CacheKey			Alphanumeric+underscore key of this cache item
+	 * @param	InData				Buffer containing the data to cache, can be destroyed after the call returns, immediately
+	 * @param	bPutEvenIfExists	If true, then do not attempt skip the put even if CachedDataProbablyExists returns true
+	 */
+	virtual EPutStatus PutCachedData(const TCHAR* CacheKey, TArrayView<const uint8> InData, bool bPutEvenIfExists) = 0;
+
+	/**
+	 * Remove data from cache (used in the event that corruption is detected at a higher level and possibly house keeping)
+	 *
+	 * @param	CacheKey	Alphanumeric+underscore key of this cache item
+	 * @param	bTransient	true if the data is transient and it is up to the backend to decide when and if to remove cached data.
+	 */
+	virtual void RemoveCachedData(const TCHAR* CacheKey, bool bTransient)=0;
+
+	/**
+	 * Retrieve usage stats for this backend. If the backend holds inner backends, this is expected to be passed down recursively.
+	 */
+	virtual TSharedRef<FDerivedDataCacheStatsNode> GatherUsageStats() const = 0;
+
+	/**
+	 * Synchronous attempt to make sure the cached data will be available as optimally as possible.
+	 *
+	 * @param	CacheKeys	Alphanumeric+underscore keys of the cache items
+	 * @return				true if the data will probably be found in a fast backend on a future request.
+	 */
+	virtual bool TryToPrefetch(TConstArrayView<FString> CacheKeys) = 0;
+
+	/**
+	 * Allows the DDC backend to determine if it wants to cache the provided data. Reasons for returning false could be a slow connection,
+	 * a file size limit, etc.
+	 */
+	virtual bool WouldCache(const TCHAR* CacheKey, TArrayView<const uint8> InData) = 0;
+
+	/**
+	 * Ask a backend to apply debug behavior to simulate different conditions. Backends that don't support these options should return
+	 * false which will result in a warning if an attempt is made to apply these options.
+	 */
+	virtual bool ApplyDebugOptions(FBackendDebugOptions& InOptions) = 0;
+};
+
+class FDerivedDataBackend
+{
+public:
+	/**
+	 * Singleton to retrieve the GLOBAL backend
+	 *
+	 * @return Reference to the global cache backend
+	 */
+	static FDerivedDataBackend& Get();
+
+	/**
+	 * Singleton to retrieve the root cache
+	 * @return Reference to the global cache root
+	 */
+	virtual FDerivedDataBackendInterface& GetRoot() = 0;
+
+	//--------------------
+	// System Interface, copied from FDerivedDataCacheInterface
+	//--------------------
+
+	virtual void NotifyBootComplete() = 0;
+	virtual void AddToAsyncCompletionCounter(int32 Addend) = 0;
+	virtual bool AnyAsyncRequestsRemaining() = 0;
+	virtual bool IsShuttingDown() = 0;
+	virtual void WaitForQuiescence(bool bShutdown = false) = 0;
+	virtual void GetDirectories(TArray<FString>& OutResults) = 0;
+	virtual bool GetUsingSharedDDC() const = 0;
+	virtual const TCHAR* GetGraphName() const = 0;
+	virtual const TCHAR* GetDefaultGraphName() const = 0;
+
+	/**
+	 * Mounts a read-only pak file.
+	 *
+	 * @param PakFilename Pak filename
+	 */
+	virtual FDerivedDataBackendInterface* MountPakFile(const TCHAR* PakFilename) = 0;
+
+	/**
+	 * Unmounts a read-only pak file.
+	 *
+	 * @param PakFilename Pak filename
+	 */
+	virtual bool UnmountPakFile(const TCHAR* PakFilename) = 0;
+
+	/**
+	 *  Gather the usage of the DDC hierarchically.
+	 */
+	virtual TSharedRef<FDerivedDataCacheStatsNode> GatherUsageStats() const = 0;
+};
+
+/** Lexical conversions from and to enums */
+
+inline const TCHAR* LexToString(FDerivedDataBackendInterface::ESpeedClass SpeedClass)
+{
+	switch (SpeedClass)
+	{
+	case FDerivedDataBackendInterface::ESpeedClass::Unknown:
+		return TEXT("Unknown");
+	case FDerivedDataBackendInterface::ESpeedClass::Slow:
+		return TEXT("Slow");
+	case FDerivedDataBackendInterface::ESpeedClass::Ok:
+		return TEXT("Ok");
+	case FDerivedDataBackendInterface::ESpeedClass::Fast:
+		return TEXT("Fast");
+	case FDerivedDataBackendInterface::ESpeedClass::Local:
+		return TEXT("Local");
+	}
+
+	checkNoEntry();
+	return TEXT("Unknown value! (Update LexToString!)");
+}
+
+inline void LexFromString(FDerivedDataBackendInterface::ESpeedClass& OutValue, const TCHAR* Buffer)
+{
+	OutValue = FDerivedDataBackendInterface::ESpeedClass::Unknown;
+
+	if (FCString::Stricmp(Buffer, TEXT("Slow")) == 0)
+	{
+		OutValue = FDerivedDataBackendInterface::ESpeedClass::Slow;
+	}
+	else if (FCString::Stricmp(Buffer, TEXT("Ok")) == 0)
+	{
+		OutValue = FDerivedDataBackendInterface::ESpeedClass::Ok;
+	}
+	else if (FCString::Stricmp(Buffer, TEXT("Fast")) == 0)
+	{
+		OutValue = FDerivedDataBackendInterface::ESpeedClass::Fast;
+	}
+	else if (FCString::Stricmp(Buffer, TEXT("Local")) == 0)
+	{
+		OutValue = FDerivedDataBackendInterface::ESpeedClass::Local;
+	}
+}
diff --git a/Engine/Source/Developer/DerivedDataCache/Private/DerivedDataBackendThrottleWrapper.h b/Engine/Source/Developer/DerivedDataCache/Private/DerivedDataBackendThrottleWrapper.h
--- a/Engine/Source/Developer/DerivedDataCache/Private/DerivedDataBackendThrottleWrapper.h
+++ b/Engine/Source/Developer/DerivedDataCache/Private/DerivedDataBackendThrottleWrapper.h
@@ -9,6 +9,7 @@
 #include "DerivedDataCachePrivate.h"
 #include "DerivedDataCacheRecord.h"
 #include "DerivedDataCacheUsageStats.h"
+#include "Memory/SharedBuffer.h"
 #include <atomic>
 
 namespace UE::DerivedData::Backends
@@ -60,12 +61,6 @@ public:
 		check(InnerBackend);
 	}
 
-	/** Return a type for this interface */
-	virtual FString GetDisplayName() const override
-	{
-		return FString(TEXT("ThrottleWrapper"));
-	}
-
 	/** Return a name for this interface */
 	virtual FString GetName() const override
 	{
@@ -78,12 +73,6 @@ public:
 		return InnerBackend->IsWritable();
 	}
 
-	/** This is a wrapepr type **/
-	virtual bool IsWrapper() const override
-	{
-		return true;
-	}
-
 	/** Returns a class of speed for this interface **/
 	virtual ESpeedClass GetSpeedClass() const override
 	{
@@ -153,10 +142,9 @@ public:
 
 	virtual TSharedRef<FDerivedDataCacheStatsNode> GatherUsageStats() const override
 	{
-		TSharedRef<FDerivedDataCacheStatsNode> Inner = InnerBackend->GatherUsageStats();
-		TSharedRef<FDerivedDataCacheStatsNode> Usage = MakeShared<FDerivedDataCacheStatsNode>(this, TEXT("ThrottleWrapper"));
-		Usage->Children.Add(Inner);
-
+		TSharedRef<FDerivedDataCacheStatsNode> Usage =
+			MakeShared<FDerivedDataCacheStatsNode>(TEXT("ThrottleWrapper"), TEXT(""), InnerBackend->GetSpeedClass() == ESpeedClass::Local);
+		Usage->Children.Add(InnerBackend->GatherUsageStats());
 		return Usage;
 	}
 
diff --git a/Engine/Source/Developer/DerivedDataCache/Private/DerivedDataBackendVerifyWrapper.h b/Engine/Source/Developer/DerivedDataCache/Private/DerivedDataBackendVerifyWrapper.h
--- a/Engine/Source/Developer/DerivedDataCache/Private/DerivedDataBackendVerifyWrapper.h
+++ b/Engine/Source/Developer/DerivedDataCache/Private/DerivedDataBackendVerifyWrapper.h
@@ -35,12 +35,6 @@ public:
 		check(InnerBackend);
 	}
 
-	/** Return a type for this interface */
-	virtual FString GetDisplayName() const override
-	{
-		return FString(TEXT("VerifyWrapper"));
-	}
-
 	/** Return a name for this interface */
 	virtual FString GetName() const override
 	{
@@ -52,11 +46,6 @@ public:
 		return true;
 	}
 
-	virtual bool IsWrapper() const override 
-	{ 
-		return true; 
-	}
-
 	virtual ESpeedClass GetSpeedClass() const override
 	{
 		return ESpeedClass::Local;
@@ -147,7 +136,8 @@ public:
 
 	virtual TSharedRef<FDerivedDataCacheStatsNode> GatherUsageStats() const override
 	{
-		TSharedRef<FDerivedDataCacheStatsNode> Usage = MakeShared<FDerivedDataCacheStatsNode>(this, TEXT("VerifyWrapper"));
+		TSharedRef<FDerivedDataCacheStatsNode> Usage =
+			MakeShared<FDerivedDataCacheStatsNode>(TEXT("VerifyWrapper"), TEXT(""), InnerBackend->GetSpeedClass() == ESpeedClass::Local);
 		Usage->Stats.Add(TEXT(""), UsageStats);
 
 		if (InnerBackend)
diff --git a/Engine/Source/Developer/DerivedDataCache/Private/DerivedDataBackends.cpp b/Engine/Source/Developer/DerivedDataCache/Private/DerivedDataBackends.cpp
--- a/Engine/Source/Developer/DerivedDataCache/Private/DerivedDataBackends.cpp
+++ b/Engine/Source/Developer/DerivedDataCache/Private/DerivedDataBackends.cpp
@@ -1207,7 +1207,7 @@ public:
 			return RootCache->GatherUsageStats();
 		}
 
-		return MakeShared<FDerivedDataCacheStatsNode>(nullptr, TEXT(""));
+		return MakeShared<FDerivedDataCacheStatsNode>(TEXT(""), TEXT(""), /*bIsLocal*/ true);
 	}
 
 private:
diff --git a/Engine/Source/Developer/DerivedDataCache/Private/DerivedDataLimitKeyLengthWrapper.h b/Engine/Source/Developer/DerivedDataCache/Private/DerivedDataLimitKeyLengthWrapper.h
--- a/Engine/Source/Developer/DerivedDataCache/Private/DerivedDataLimitKeyLengthWrapper.h
+++ b/Engine/Source/Developer/DerivedDataCache/Private/DerivedDataLimitKeyLengthWrapper.h
@@ -30,12 +30,6 @@ public:
 		check(InnerBackend);
 	}
 
-	/** Return a type for this interface */
-	virtual FString GetDisplayName() const override
-	{
-		return FString(TEXT("LimitKeyLengthWrapper"));
-	}
-
 	/** Return a name for this interface */
 	virtual FString GetName() const override 
 	{ 
@@ -48,12 +42,6 @@ public:
 		return InnerBackend->IsWritable();
 	}
 
-	/** This is a wrapepr type **/
-	virtual bool IsWrapper() const override
-	{
-		return true;
-	}
-
 	/** Returns a class of speed for this interface **/
 	virtual ESpeedClass GetSpeedClass() const override
 	{
@@ -243,7 +231,8 @@ public:
 
 	virtual TSharedRef<FDerivedDataCacheStatsNode> GatherUsageStats() const override
 	{
-		TSharedRef<FDerivedDataCacheStatsNode> Usage = MakeShared<FDerivedDataCacheStatsNode>(this, TEXT("LimitKeyLength"));
+		TSharedRef<FDerivedDataCacheStatsNode> Usage =
+			MakeShared<FDerivedDataCacheStatsNode>(TEXT("LimitKeyLength"), TEXT(""), InnerBackend->GetSpeedClass() == ESpeedClass::Local);
 		Usage->Stats.Add(TEXT(""), UsageStats);
 
 		if (InnerBackend)
diff --git a/Engine/Source/Developer/DerivedDataCache/Private/FileBackedDerivedDataBackend.h b/Engine/Source/Developer/DerivedDataCache/Private/FileBackedDerivedDataBackend.h
--- a/Engine/Source/Developer/DerivedDataCache/Private/FileBackedDerivedDataBackend.h
+++ b/Engine/Source/Developer/DerivedDataCache/Private/FileBackedDerivedDataBackend.h
@@ -2,19 +2,14 @@
 
 #pragma once
 
-#include "CoreMinimal.h"
 #include "DerivedDataBackendInterface.h"
 
-
-/** 
+/**
  * A file backed backend. This is used for local caching of data where more flexibility is needed than loose files.
-**/
+ */
 class FFileBackedDerivedDataBackend : public FDerivedDataBackendInterface
 {
 public:
-
-	virtual FString GetDisplayName() const override { return FString(TEXT("File")); }
-	
 	/**
 	 * Save the cache to disk
 	 * @param	Filename	Filename to save
@@ -33,8 +28,4 @@ public:
 	 * Disables use of this cache and allows the underlying implementation to free up memory
 	 */
 	virtual void Disable() = 0;
-
-private:
-
 };
-
diff --git a/Engine/Source/Developer/DerivedDataCache/Private/FileSystemDerivedDataBackend.cpp b/Engine/Source/Developer/DerivedDataCache/Private/FileSystemDerivedDataBackend.cpp
--- a/Engine/Source/Developer/DerivedDataCache/Private/FileSystemDerivedDataBackend.cpp
+++ b/Engine/Source/Developer/DerivedDataCache/Private/FileSystemDerivedDataBackend.cpp
@@ -737,7 +737,6 @@ public:
 	// FDerivedDataBackendInterface Interface
 
 	FString GetName() const final { return CachePath; }
-	FString GetDisplayName() const final { return FString(TEXT("File System")); }
 	bool IsWritable() const final { return !bReadOnly && !bDisabled; }
 	ESpeedClass GetSpeedClass() const final { return SpeedClass; }
 
@@ -1519,7 +1518,7 @@ void FFileSystemCacheStore::RemoveCachedData(const TCHAR* const CacheKey, const
 TSharedRef<FDerivedDataCacheStatsNode> FFileSystemCacheStore::GatherUsageStats() const
 {
 	TSharedRef<FDerivedDataCacheStatsNode> Usage =
-		MakeShared<FDerivedDataCacheStatsNode>(this, FString::Printf(TEXT("%s.%s"), TEXT("FileSystem"), *CachePath));
+		MakeShared<FDerivedDataCacheStatsNode>(TEXT("File System"), *CachePath, SpeedClass == ESpeedClass::Local);
 	Usage->Stats.Add(TEXT(""), UsageStats);
 	return Usage;
 }
diff --git a/Engine/Source/Developer/DerivedDataCache/Private/HierarchicalDerivedDataBackend.h b/Engine/Source/Developer/DerivedDataCache/Private/HierarchicalDerivedDataBackend.h
--- a/Engine/Source/Developer/DerivedDataCache/Private/HierarchicalDerivedDataBackend.h
+++ b/Engine/Source/Developer/DerivedDataCache/Private/HierarchicalDerivedDataBackend.h
@@ -48,12 +48,6 @@ public:
 		UpdateAsyncInnerBackends();
 	}
 
-	/** Return a type for this interface */
-	virtual FString GetDisplayName() const override
-	{
-		return TEXT("Hierarchical");
-	}
-
 	/** Return a name for this interface */
 	virtual FString GetName() const override
 	{
@@ -463,7 +457,8 @@ public:
 
 	virtual TSharedRef<FDerivedDataCacheStatsNode> GatherUsageStats() const override
 	{
-		TSharedRef<FDerivedDataCacheStatsNode> Usage = MakeShared<FDerivedDataCacheStatsNode>(this, TEXT("Hierarchical"));
+		TSharedRef<FDerivedDataCacheStatsNode> Usage =
+			MakeShared<FDerivedDataCacheStatsNode>(TEXT("Hierarchical"), TEXT(""), /*bIsLocal*/ true);
 		Usage->Stats.Add(TEXT(""), UsageStats);
 
 		FReadScopeLock LockScope(Lock);
diff --git a/Engine/Source/Developer/DerivedDataCache/Private/HttpDerivedDataBackend.cpp b/Engine/Source/Developer/DerivedDataCache/Private/HttpDerivedDataBackend.cpp
--- a/Engine/Source/Developer/DerivedDataCache/Private/HttpDerivedDataBackend.cpp
+++ b/Engine/Source/Developer/DerivedDataCache/Private/HttpDerivedDataBackend.cpp
@@ -1980,11 +1980,6 @@ FString FHttpDerivedDataBackend::GetName() const
 	return Domain;
 }
 
-FString FHttpDerivedDataBackend::GetDisplayName() const
-{
-	return TEXT("Horde Storage");
-}
-
 bool FHttpDerivedDataBackend::TryToPrefetch(TConstArrayView<FString> CacheKeys)
 {
 	return CachedDataProbablyExistsBatch(CacheKeys).CountSetBits() == CacheKeys.Num();
@@ -2525,9 +2520,9 @@ void FHttpDerivedDataBackend::RemoveCachedData(const TCHAR* CacheKey, bool bTran
 
 TSharedRef<FDerivedDataCacheStatsNode> FHttpDerivedDataBackend::GatherUsageStats() const
 {
-	TSharedRef<FDerivedDataCacheStatsNode> Usage = MakeShared<FDerivedDataCacheStatsNode>(this, FString::Printf(TEXT("%s @ %s (%s)"), TEXT("HTTP"), *Domain, *Namespace));
+	TSharedRef<FDerivedDataCacheStatsNode> Usage =
+		MakeShared<FDerivedDataCacheStatsNode>(TEXT("Horde Storage"), FString::Printf(TEXT("%s (%s)"), *Domain, *Namespace), /*bIsLocal*/ false);
 	Usage->Stats.Add(TEXT(""), UsageStats);
-
 	return Usage;
 }
 
diff --git a/Engine/Source/Developer/DerivedDataCache/Private/HttpDerivedDataBackend.h b/Engine/Source/Developer/DerivedDataCache/Private/HttpDerivedDataBackend.h
--- a/Engine/Source/Developer/DerivedDataCache/Private/HttpDerivedDataBackend.h
+++ b/Engine/Source/Developer/DerivedDataCache/Private/HttpDerivedDataBackend.h
@@ -52,8 +52,6 @@ public:
 	 */
 	bool IsUsable() const { return bIsUsable; }
 
-	bool IsRemote() const { return true; }
-
 	/** return true if this cache is writable **/
 	virtual bool IsWritable() const override
 	{
@@ -68,7 +66,6 @@ public:
 	virtual TSharedRef<FDerivedDataCacheStatsNode> GatherUsageStats() const override;
 	
 	virtual FString GetName() const override;
-	virtual FString GetDisplayName() const override;
 	virtual bool TryToPrefetch(TConstArrayView<FString> CacheKeys) override;
 	virtual bool WouldCache(const TCHAR* CacheKey, TArrayView<const uint8> InData) override;
 	virtual ESpeedClass GetSpeedClass() const override;
diff --git a/Engine/Source/Developer/DerivedDataCache/Private/MemoryDerivedDataBackend.cpp b/Engine/Source/Developer/DerivedDataCache/Private/MemoryDerivedDataBackend.cpp
--- a/Engine/Source/Developer/DerivedDataCache/Private/MemoryDerivedDataBackend.cpp
+++ b/Engine/Source/Developer/DerivedDataCache/Private/MemoryDerivedDataBackend.cpp
@@ -383,9 +383,9 @@ void FMemoryDerivedDataBackend::Disable()
 
 TSharedRef<FDerivedDataCacheStatsNode> FMemoryDerivedDataBackend::GatherUsageStats() const
 {
-	TSharedRef<FDerivedDataCacheStatsNode> Usage = MakeShared<FDerivedDataCacheStatsNode>(this, FString::Printf(TEXT("%s.%s"), TEXT("MemoryBackend"), *CacheFilename));
+	TSharedRef<FDerivedDataCacheStatsNode> Usage =
+		MakeShared<FDerivedDataCacheStatsNode>(CacheFilename.IsEmpty() ? TEXT("Memory") : TEXT("Boot"), CacheFilename, /*bIsLocal*/ true);
 	Usage->Stats.Add(TEXT(""), UsageStats);
-
 	return Usage;
 }
 
diff --git a/Engine/Source/Developer/DerivedDataCache/Private/MemoryDerivedDataBackend.h b/Engine/Source/Developer/DerivedDataCache/Private/MemoryDerivedDataBackend.h
--- a/Engine/Source/Developer/DerivedDataCache/Private/MemoryDerivedDataBackend.h
+++ b/Engine/Source/Developer/DerivedDataCache/Private/MemoryDerivedDataBackend.h
@@ -24,9 +24,6 @@ public:
 	explicit FMemoryDerivedDataBackend(const TCHAR* InName, int64 InMaxCacheSize = -1, bool bCanBeDisabled = false);
 	~FMemoryDerivedDataBackend();
 
-	/** Return a name for this interface */
-	virtual FString GetDisplayName() const override { return FString(TEXT("Memory")); }
-
 	/** Return a name for this interface */
 	virtual FString GetName() const override { return Name; }
 
diff --git a/Engine/Source/Developer/DerivedDataCache/Private/PakFileDerivedDataBackend.cpp b/Engine/Source/Developer/DerivedDataCache/Private/PakFileDerivedDataBackend.cpp
--- a/Engine/Source/Developer/DerivedDataCache/Private/PakFileDerivedDataBackend.cpp
+++ b/Engine/Source/Developer/DerivedDataCache/Private/PakFileDerivedDataBackend.cpp
@@ -74,11 +74,6 @@ FPakFileDerivedDataBackend::~FPakFileDerivedDataBackend()
 	Close();
 }
 
-FString FPakFileDerivedDataBackend::GetDisplayName() const
-{
-	return FString(TEXT("PakFile"));
-}
-
 FString FPakFileDerivedDataBackend::GetName() const
 {
 	return CachePath;
@@ -476,7 +471,8 @@ bool FPakFileDerivedDataBackend::SortAndCopy(const FString &InputFilename, const
 
 TSharedRef<FDerivedDataCacheStatsNode> FPakFileDerivedDataBackend::GatherUsageStats() const
 {
-	TSharedRef<FDerivedDataCacheStatsNode> Usage = MakeShared<FDerivedDataCacheStatsNode>(this, FString::Printf(TEXT("%s.%s"), TEXT("PakFile"), *CachePath));
+	TSharedRef<FDerivedDataCacheStatsNode> Usage =
+		MakeShared<FDerivedDataCacheStatsNode>(TEXT("PakFile"), CachePath, /*bIsLocal*/ true);
 	Usage->Stats.Add(TEXT(""), UsageStats);
 	return Usage;
 }
diff --git a/Engine/Source/Developer/DerivedDataCache/Private/PakFileDerivedDataBackend.h b/Engine/Source/Developer/DerivedDataCache/Private/PakFileDerivedDataBackend.h
--- a/Engine/Source/Developer/DerivedDataCache/Private/PakFileDerivedDataBackend.h
+++ b/Engine/Source/Developer/DerivedDataCache/Private/PakFileDerivedDataBackend.h
@@ -30,9 +30,6 @@ public:
 
 	void Close();
 
-	/** Return a type for this interface */
-	virtual FString GetDisplayName() const override;
-
 	/** Return a name for this interface */
 	virtual FString GetName() const override;
 
diff --git a/Engine/Source/Developer/DerivedDataCache/Private/S3DerivedDataBackend.cpp b/Engine/Source/Developer/DerivedDataCache/Private/S3DerivedDataBackend.cpp
--- a/Engine/Source/Developer/DerivedDataCache/Private/S3DerivedDataBackend.cpp
+++ b/Engine/Source/Developer/DerivedDataCache/Private/S3DerivedDataBackend.cpp
@@ -984,17 +984,12 @@ void FS3DerivedDataBackend::RemoveCachedData(const TCHAR* CacheKey, bool bTransi
 
 TSharedRef<FDerivedDataCacheStatsNode> FS3DerivedDataBackend::GatherUsageStats() const
 {
-	TSharedRef<FDerivedDataCacheStatsNode> Usage = MakeShared<FDerivedDataCacheStatsNode>(this, FString::Printf(TEXT("%s @ %s"), TEXT("S3"), *BaseUrl));
+	TSharedRef<FDerivedDataCacheStatsNode> Usage =
+		MakeShared<FDerivedDataCacheStatsNode>(TEXT("S3"), BaseUrl, /*bIsLocal*/ false);
 	Usage->Stats.Add(TEXT(""), UsageStats);
-
 	return Usage;
 }
 
-FString FS3DerivedDataBackend::GetDisplayName() const
-{
-	return FString(TEXT("S3"));
-}
-
 FString FS3DerivedDataBackend::GetName() const
 {
 	return BaseUrl;
diff --git a/Engine/Source/Developer/DerivedDataCache/Private/S3DerivedDataBackend.h b/Engine/Source/Developer/DerivedDataCache/Private/S3DerivedDataBackend.h
--- a/Engine/Source/Developer/DerivedDataCache/Private/S3DerivedDataBackend.h
+++ b/Engine/Source/Developer/DerivedDataCache/Private/S3DerivedDataBackend.h
@@ -43,9 +43,6 @@ public:
 	/* S3 Cache cannot be written to*/
 	bool IsWritable() const override { return false; }
 
-	/* S3 Cache is always remote */
-	bool IsRemote() const override { return true; }
-
 	/* S3 Cache does not try to write back to lower caches (e.g. Shared DDC) */
 	bool BackfillLowerCacheLevels() const override { return false; }
 
@@ -56,7 +53,6 @@ public:
 	virtual TSharedRef<FDerivedDataCacheStatsNode> GatherUsageStats() const override;
 
 	FString GetName() const override;
-	FString GetDisplayName() const override;
 	ESpeedClass GetSpeedClass() const override;
 	bool TryToPrefetch(TConstArrayView<FString> CacheKeys) override;
 	bool WouldCache(const TCHAR* CacheKey, TArrayView<const uint8> InData) override;
diff --git a/Engine/Source/Developer/DerivedDataCache/Private/ZenDerivedDataBackend.cpp b/Engine/Source/Developer/DerivedDataCache/Private/ZenDerivedDataBackend.cpp
--- a/Engine/Source/Developer/DerivedDataCache/Private/ZenDerivedDataBackend.cpp
+++ b/Engine/Source/Developer/DerivedDataCache/Private/ZenDerivedDataBackend.cpp
@@ -88,21 +88,11 @@ FZenDerivedDataBackend::~FZenDerivedDataBackend()
 {
 }
 
-FString FZenDerivedDataBackend::GetDisplayName() const
-{
-	return FString(TEXT("Zen"));
-}
-
 FString FZenDerivedDataBackend::GetName() const
 {
 	return ZenService.GetInstance().GetURL();
 }
 
-bool FZenDerivedDataBackend::IsRemote() const
-{
-	return bIsRemote;
-}
-
 bool FZenDerivedDataBackend::IsServiceReady()
 {
 	return ZenService.GetInstance().IsServiceReady();
@@ -512,7 +502,7 @@ TSharedRef<FDerivedDataCacheStatsNode> FZenDerivedDataBackend::GatherUsageStats(
 #endif
 
 	TSharedRef<FDerivedDataCacheStatsNode> LocalNode =
-		MakeShared<FDerivedDataCacheStatsNode>(this, FString::Printf(TEXT("%s.%s"), TEXT("ZenDDC"), *GetName()));
+		MakeShared<FDerivedDataCacheStatsNode>(TEXT("Zen"), ZenService.GetInstance().GetURL(), /*bIsLocal*/ true);
 	LocalNode->Stats.Add(TEXT(""), LocalStats);
 
 	if (ZenStats.UpstreamStats.EndPointStats.IsEmpty())
@@ -521,11 +511,11 @@ TSharedRef<FDerivedDataCacheStatsNode> FZenDerivedDataBackend::GatherUsageStats(
 	}
 
 	TSharedRef<FDerivedDataCacheStatsNode> RemoteNode =
-		MakeShared<FDerivedDataCacheStatsNode>(this, ZenStats.UpstreamStats.EndPointStats[0].Url);
+		MakeShared<FDerivedDataCacheStatsNode>(ZenStats.UpstreamStats.EndPointStats[0].Name, ZenStats.UpstreamStats.EndPointStats[0].Url, /*bIsLocal*/ false);
 	RemoteNode->Stats.Add(TEXT(""), RemoteStats);
 
 	TSharedRef<FDerivedDataCacheStatsNode> GroupNode =
-		MakeShared<FDerivedDataCacheStatsNode>(this, TEXT("Zen Group"));
+		MakeShared<FDerivedDataCacheStatsNode>(TEXT("Zen Group"), TEXT(""), /*bIsLocal*/ true);
 	GroupNode->Children.Add(LocalNode);
 	GroupNode->Children.Add(RemoteNode);
 	return GroupNode;
diff --git a/Engine/Source/Developer/DerivedDataCache/Private/ZenDerivedDataBackend.h b/Engine/Source/Developer/DerivedDataCache/Private/ZenDerivedDataBackend.h
--- a/Engine/Source/Developer/DerivedDataCache/Private/ZenDerivedDataBackend.h
+++ b/Engine/Source/Developer/DerivedDataCache/Private/ZenDerivedDataBackend.h
@@ -85,8 +85,6 @@ public:
 	virtual void RemoveCachedData(const TCHAR* CacheKey, bool bTransient) override;
 	
 	virtual FString GetName() const override;
-	virtual FString GetDisplayName() const override;
-	virtual bool IsRemote() const override;
 	virtual bool WouldCache(const TCHAR* CacheKey, TArrayView<const uint8> InData) override;
 	virtual bool ApplyDebugOptions(FBackendDebugOptions& InOptions) override;
 
diff --git a/Engine/Source/Developer/DerivedDataCache/Public/DerivedDataBackendInterface.h b/Engine/Source/Developer/DerivedDataCache/Public/DerivedDataBackendInterface.h
deleted file mode 100644
diff --git a/Engine/Source/Developer/DerivedDataCache/Public/DerivedDataCacheUsageStats.h b/Engine/Source/Developer/DerivedDataCache/Public/DerivedDataCacheUsageStats.h
--- a/Engine/Source/Developer/DerivedDataCache/Public/DerivedDataCacheUsageStats.h
+++ b/Engine/Source/Developer/DerivedDataCache/Public/DerivedDataCacheUsageStats.h
@@ -103,24 +103,25 @@ public:
 #endif
 };
 
-class FDerivedDataBackendInterface;
-
 /**
  *  Hierarchical usage stats for the DDC nodes.
  */
 class FDerivedDataCacheStatsNode : public TSharedFromThis<FDerivedDataCacheStatsNode>
 {
 public:
-	FDerivedDataCacheStatsNode(const FDerivedDataBackendInterface* InBackendInterface, const FString& InCacheName)
-		: BackendInterface(InBackendInterface)
+	FDerivedDataCacheStatsNode(const FString& InCacheType, const FString& InCacheName, bool bInIsLocal)
+		: CacheType(InCacheType)
 		, CacheName(InCacheName)
+		, bIsLocal(bInIsLocal)
 	{
 	}
 
-	const FDerivedDataBackendInterface* GetBackendInterface() const { return BackendInterface; }
+	const FString& GetCacheType() const { return CacheType; }
 
 	const FString& GetCacheName() const { return CacheName; }
 
+	bool IsLocal() const { return bIsLocal; }
+
 	TMap<FString, FDerivedDataCacheUsageStats> ToLegacyUsageMap() const
 	{
 		TMap<FString, FDerivedDataCacheUsageStats> UsageStats;
@@ -168,8 +169,9 @@ public:
 	TArray<TSharedRef<FDerivedDataCacheStatsNode>> Children;
 
 protected:
-	const FDerivedDataBackendInterface* BackendInterface;
+	FString CacheType;
 	FString CacheName;
+	bool bIsLocal;
 };
 
 struct FDerivedDataCacheResourceStat
@@ -254,4 +256,3 @@ struct FDerivedDataCacheSummaryStats
 {
 	COOK_STAT(TArray<FDerivedDataCacheSummaryStat> Stats);
 };
-
diff --git a/Engine/Source/Editor/DerivedDataEditor/Private/DerivedDataInformation.cpp b/Engine/Source/Editor/DerivedDataEditor/Private/DerivedDataInformation.cpp
--- a/Engine/Source/Editor/DerivedDataEditor/Private/DerivedDataInformation.cpp
+++ b/Engine/Source/Editor/DerivedDataEditor/Private/DerivedDataInformation.cpp
@@ -4,7 +4,6 @@
 #include "SDerivedDataStatusBar.h"
 #include "DerivedDataCacheUsageStats.h"
 #include "DerivedDataCacheInterface.h"
-#include "DerivedDataBackendInterface.h"
 #include "Settings/EditorProjectSettings.h"
 #include "Settings/EditorSettings.h"
 #include "ToolMenus.h"
@@ -18,29 +17,32 @@ bool FDerivedDataInformation::bIsUploading = false;
 FText FDerivedDataInformation::RemoteCacheWarningMessage;
 ERemoteCacheState FDerivedDataInformation::RemoteCacheState= ERemoteCacheState::Unavailable;
 
-double FDerivedDataInformation::GetCacheActivitySizeBytes(bool bGet, bool bLocal)
+static TArray<TSharedRef<const FDerivedDataCacheStatsNode>> GetCacheUsageStats()
 {
-	PRAGMA_DISABLE_DEPRECATION_WARNINGS
-		TSharedRef<FDerivedDataCacheStatsNode> RootUsage = GetDerivedDataCache()->GatherUsageStats();
-	PRAGMA_ENABLE_DEPRECATION_WARNINGS
-		TArray<TSharedRef<const FDerivedDataCacheStatsNode>> LeafUsageStats;
-	RootUsage->ForEachDescendant([&LeafUsageStats](TSharedRef<const FDerivedDataCacheStatsNode> Node) {
-		if (Node->Children.Num() == 0)
+	PRAGMA_DISABLE_DEPRECATION_WARNINGS;
+	TSharedRef<FDerivedDataCacheStatsNode> RootUsage = GetDerivedDataCacheRef().GatherUsageStats();
+	PRAGMA_ENABLE_DEPRECATION_WARNINGS;
+	TArray<TSharedRef<const FDerivedDataCacheStatsNode>> LeafUsageStats;
+	RootUsage->ForEachDescendant([&LeafUsageStats](TSharedRef<const FDerivedDataCacheStatsNode> Node)
+	{
+		if (Node->Children.IsEmpty())
 		{
 			LeafUsageStats.Add(Node);
 		}
-		});
+	});
+	return LeafUsageStats;
+}
 
+double FDerivedDataInformation::GetCacheActivitySizeBytes(bool bGet, bool bLocal)
+{
 	int64 TotalBytes = 0;
 
-	for (int32 Index = 0; Index < LeafUsageStats.Num(); Index++)
+	for (const TSharedRef<const FDerivedDataCacheStatsNode>& Usage : GetCacheUsageStats())
 	{
-		const FDerivedDataBackendInterface* Backend = LeafUsageStats[Index]->GetBackendInterface();
-
-		if (Backend->IsRemote() == bLocal)
+		if (Usage->IsLocal() != bLocal)
+		{
 			continue;
-
-		TSharedRef<FDerivedDataCacheStatsNode> Usage = Backend->GatherUsageStats();
+		}
 
 		for (const auto& KVP : Usage->Stats)
 		{
@@ -63,27 +65,14 @@ double FDerivedDataInformation::GetCacheActivitySizeBytes(bool bGet, bool bLocal
 
 double FDerivedDataInformation::GetCacheActivityTimeSeconds(bool bGet, bool bLocal)
 {
-	PRAGMA_DISABLE_DEPRECATION_WARNINGS
-		TSharedRef<FDerivedDataCacheStatsNode> RootUsage = GetDerivedDataCache()->GatherUsageStats();
-	PRAGMA_ENABLE_DEPRECATION_WARNINGS
-		TArray<TSharedRef<const FDerivedDataCacheStatsNode>> LeafUsageStats;
-	RootUsage->ForEachDescendant([&LeafUsageStats](TSharedRef<const FDerivedDataCacheStatsNode> Node) {
-		if (Node->Children.Num() == 0)
-		{
-			LeafUsageStats.Add(Node);
-		}
-		});
-
 	int64 TotalCycles = 0;
 
-	for (int32 Index = 0; Index < LeafUsageStats.Num(); Index++)
+	for (const TSharedRef<const FDerivedDataCacheStatsNode>& Usage : GetCacheUsageStats())
 	{
-		const FDerivedDataBackendInterface* Backend = LeafUsageStats[Index]->GetBackendInterface();
-
-		if (Backend->IsRemote() == bLocal)
+		if (Usage->IsLocal() != bLocal)
+		{
 			continue;
-
-		TSharedRef<FDerivedDataCacheStatsNode> Usage = Backend->GatherUsageStats();
+		}
 
 		for (const auto& KVP : Usage->Stats)
 		{
@@ -113,73 +102,37 @@ double FDerivedDataInformation::GetCacheActivityTimeSeconds(bool bGet, bool bLoc
 
 bool FDerivedDataInformation::GetHasRemoteCache()
 {
-	PRAGMA_DISABLE_DEPRECATION_WARNINGS
-		TSharedRef<FDerivedDataCacheStatsNode> RootUsage = GetDerivedDataCache()->GatherUsageStats();
-	PRAGMA_ENABLE_DEPRECATION_WARNINGS
-		TArray<TSharedRef<const FDerivedDataCacheStatsNode>> LeafUsageStats;
-	RootUsage->ForEachDescendant([&LeafUsageStats](TSharedRef<const FDerivedDataCacheStatsNode> Node) {
-		if (Node->Children.Num() == 0)
-		{
-			LeafUsageStats.Add(Node);
-		}
-		});
-
-	for (int32 Index = 0; Index < LeafUsageStats.Num(); Index++)
+	for (const TSharedRef<const FDerivedDataCacheStatsNode>& Usage : GetCacheUsageStats())
 	{
-		const FDerivedDataBackendInterface* Backend = LeafUsageStats[Index]->GetBackendInterface();
-
-		if (Backend->IsRemote())
+		if (!Usage->IsLocal())
+		{
 			return true;
+		}
 	}
-
 	return false;
 }
 
 bool FDerivedDataInformation::GetHasZenCache()
 {
-	PRAGMA_DISABLE_DEPRECATION_WARNINGS
-		TSharedRef<FDerivedDataCacheStatsNode> RootUsage = GetDerivedDataCache()->GatherUsageStats();
-	PRAGMA_ENABLE_DEPRECATION_WARNINGS
-		TArray<TSharedRef<const FDerivedDataCacheStatsNode>> LeafUsageStats;
-	RootUsage->ForEachDescendant([&LeafUsageStats](TSharedRef<const FDerivedDataCacheStatsNode> Node) {
-		if (Node->Children.Num() == 0)
-		{
-			LeafUsageStats.Add(Node);
-		}
-		});
-
-	for (int32 Index = 0; Index < LeafUsageStats.Num(); Index++)
+	for (const TSharedRef<const FDerivedDataCacheStatsNode>& Usage : GetCacheUsageStats())
 	{
-		const FDerivedDataBackendInterface* Backend = LeafUsageStats[Index]->GetBackendInterface();
-
-		if (Backend->GetDisplayName().Equals("Zen"))
+		if (Usage->GetCacheType().Equals(TEXT("Zen")))
+		{
 			return true;
+		}
 	}
-
 	return false;
 }
 
 bool FDerivedDataInformation::GetHasHordeStorageCache()
 {
-	PRAGMA_DISABLE_DEPRECATION_WARNINGS
-		TSharedRef<FDerivedDataCacheStatsNode> RootUsage = GetDerivedDataCache()->GatherUsageStats();
-	PRAGMA_ENABLE_DEPRECATION_WARNINGS
-		TArray<TSharedRef<const FDerivedDataCacheStatsNode>> LeafUsageStats;
-	RootUsage->ForEachDescendant([&LeafUsageStats](TSharedRef<const FDerivedDataCacheStatsNode> Node) {
-		if (Node->Children.Num() == 0)
-		{
-			LeafUsageStats.Add(Node);
-		}
-		});
-
-	for (int32 Index = 0; Index < LeafUsageStats.Num(); Index++)
+	for (const TSharedRef<const FDerivedDataCacheStatsNode>& Usage : GetCacheUsageStats())
 	{
-		const FDerivedDataBackendInterface* Backend = LeafUsageStats[Index]->GetBackendInterface();
-
-		if (Backend->GetDisplayName().Equals("Horde Storage"))
+		if (Usage->GetCacheType().Equals(TEXT("Horde Storage")))
+		{
 			return true;
+		}
 	}
-
 	return false;
 }
 
@@ -187,13 +140,13 @@ void FDerivedDataInformation::UpdateRemoteCacheState()
 {
 	RemoteCacheState = ERemoteCacheState::Unavailable;
 
-	if ( GetHasRemoteCache() )
+	if (GetHasRemoteCache())
 	{
 		const double OldLastGetTime = LastGetTime;
 		const double OldLastPutTime = LastPutTime;
 
-		LastGetTime = FDerivedDataInformation::GetCacheActivityTimeSeconds(true, false);
-		LastPutTime = FDerivedDataInformation::GetCacheActivityTimeSeconds(false, false);
+		LastGetTime = FDerivedDataInformation::GetCacheActivityTimeSeconds(/*bGet*/ true, /*bLocal*/ false);
+		LastPutTime = FDerivedDataInformation::GetCacheActivityTimeSeconds(/*bGet*/ false, /*bLocal*/ false);
 
 		if (OldLastGetTime != 0.0 && OldLastPutTime != 0.0)
 		{
@@ -211,15 +164,11 @@ void FDerivedDataInformation::UpdateRemoteCacheState()
 		}
 	}
 
-	const UDDCProjectSettings* DDCProjectSettings = GetDefault<UDDCProjectSettings>();
-
-	if (DDCProjectSettings && DDCProjectSettings->EnableWarnings)
+	if (const UDDCProjectSettings* DDCProjectSettings = GetDefault<UDDCProjectSettings>(); DDCProjectSettings && DDCProjectSettings->EnableWarnings)
 	{
-		const UEditorSettings* EditorSettings = GetDefault<UEditorSettings>();
-
-		if (EditorSettings)
+		if (const UEditorSettings* EditorSettings = GetDefault<UEditorSettings>())
 		{
-			if (DDCProjectSettings->RecommendEveryoneUseHordeStorage && GetHasHordeStorageCache()==false && (FCString::Stricmp(GetDerivedDataCache()->GetGraphName(), TEXT("NoJupiter"))!=0))
+			if (DDCProjectSettings->RecommendEveryoneUseHordeStorage && !GetHasHordeStorageCache() && (FCString::Stricmp(GetDerivedDataCache()->GetGraphName(), TEXT("NoJupiter")) != 0))
 			{
 				RemoteCacheState = ERemoteCacheState::Warning;
 				RemoteCacheWarningMessage = FText(LOCTEXT("HordeStorageWarning", "It is recommended that you use a DDC graph that supports Horde Storage. Please check any -ddc commandline overrides."));
@@ -234,7 +183,7 @@ void FDerivedDataInformation::UpdateRemoteCacheState()
 				RemoteCacheState = ERemoteCacheState::Warning;
 				RemoteCacheWarningMessage = FText(LOCTEXT("GlobalLocalDDCPathWarning", "It is recommended that you set up a valid Global Shared DDC Path"));
 			}
-			else if (DDCProjectSettings->RecommendEveryoneEnableS3DDC && EditorSettings->bEnableS3DDC == false)
+			else if (DDCProjectSettings->RecommendEveryoneEnableS3DDC && !EditorSettings->bEnableS3DDC)
 			{
 				RemoteCacheState = ERemoteCacheState::Warning;
 				RemoteCacheWarningMessage = FText(LOCTEXT("AWSS3CacheEnabledWarning", "It is recommended that you enable the AWS S3 Cache"));
@@ -245,7 +194,6 @@ void FDerivedDataInformation::UpdateRemoteCacheState()
 				RemoteCacheWarningMessage = FText(LOCTEXT("S3GloblaLocalPathdWarning", "It is recommended that you set up a valid Global Local S3 DDC Path"));
 			}
 		}
-		
 	}
 }
 
diff --git a/Engine/Source/Editor/DerivedDataEditor/Private/SDerivedDataDialogs.cpp b/Engine/Source/Editor/DerivedDataEditor/Private/SDerivedDataDialogs.cpp
--- a/Engine/Source/Editor/DerivedDataEditor/Private/SDerivedDataDialogs.cpp
+++ b/Engine/Source/Editor/DerivedDataEditor/Private/SDerivedDataDialogs.cpp
@@ -2,7 +2,6 @@
 
 #include "SDerivedDataDialogs.h"
 #include "Algo/Sort.h"
-#include "DerivedDataBackendInterface.h"
 #include "DerivedDataCacheInterface.h"
 #include "DerivedDataCacheUsageStats.h"
 #include "DerivedDataInformation.h"
@@ -600,9 +599,7 @@ TSharedRef<SWidget> SDerivedDataCacheStatisticsDialog::GetGridPanel()
 
 	for (TSharedRef<const FDerivedDataCacheStatsNode> Node : LeafUsageStats)
 	{
-		const FDerivedDataBackendInterface* Backend = Node->GetBackendInterface();
-
-		if (Backend->GetDisplayName().Equals("Memory") || Backend->IsWrapper())
+		if (Node->GetCacheType().Equals(TEXT("Memory")))
 		{
 			continue;
 		}
@@ -632,14 +629,14 @@ TSharedRef<SWidget> SDerivedDataCacheStatisticsDialog::GetGridPanel()
 		[
 			SNew(STextBlock)
 			.Margin(DefaultMarginFirstColumn)
-			.Text(FText::FromString(Backend->GetDisplayName()))
+			.Text(FText::FromString(Node->GetCacheType()))
 		];
 
 		Panel->AddSlot(1, Row)
 		[
 			SNew(STextBlock)
 			.Margin(DefaultMargin)
-			.Text((Backend->GetSpeedClass() == FDerivedDataBackendInterface::ESpeedClass::Local) ? LOCTEXT("Local", "Local") : LOCTEXT("Remote", "Remote"))
+			.Text(Node->IsLocal() ? LOCTEXT("Local", "Local") : LOCTEXT("Remote", "Remote"))
 		];
 
 		Panel->AddSlot(2, Row)
@@ -670,7 +667,7 @@ TSharedRef<SWidget> SDerivedDataCacheStatisticsDialog::GetGridPanel()
 		[
 			SNew(STextBlock)
 			.Margin(DefaultMargin)
-			.Text(FText::FromString(Backend->GetName()))
+			.Text(FText::FromString(Node->GetCacheName()))
 		];
 
 		Row++;
diff --git a/Engine/Source/Editor/DerivedDataEditor/Private/SDerivedDataStatusBar.cpp b/Engine/Source/Editor/DerivedDataEditor/Private/SDerivedDataStatusBar.cpp
--- a/Engine/Source/Editor/DerivedDataEditor/Private/SDerivedDataStatusBar.cpp
+++ b/Engine/Source/Editor/DerivedDataEditor/Private/SDerivedDataStatusBar.cpp
@@ -7,7 +7,6 @@
 #include "SDerivedDataCacheSettings.h"
 #include "DerivedDataCacheUsageStats.h"
 #include "DerivedDataCacheInterface.h"
-#include "DerivedDataBackendInterface.h"
 #include "Widgets/SBoxPanel.h"
 #include "Widgets/SToolTip.h"
 #include "Widgets/Text/STextBlock.h"
