DDC: Fixed file system cache store put stats

#rnx

diff --git a/Engine/Source/Developer/DerivedDataCache/Private/FileSystemDerivedDataBackend.cpp b/Engine/Source/Developer/DerivedDataCache/Private/FileSystemDerivedDataBackend.cpp
--- a/Engine/Source/Developer/DerivedDataCache/Private/FileSystemDerivedDataBackend.cpp
+++ b/Engine/Source/Developer/DerivedDataCache/Private/FileSystemDerivedDataBackend.cpp
@@ -771,8 +771,8 @@ private:
 	uint64 MeasureCompressedCacheRecord(const FCacheRecord& Record) const;
 	uint64 MeasureRawCacheRecord(const FCacheRecord& Record) const;
 
-	bool PutCacheRecord(const FCacheRecord& Record, FStringView Context, const FCacheRecordPolicy& Policy);
-	bool PutCacheContent(const FCompressedBuffer& Content, const FStringView Context) const;
+	bool PutCacheRecord(const FCacheRecord& Record, FStringView Context, const FCacheRecordPolicy& Policy, uint64& OutWriteSize);
+	bool PutCacheContent(const FCompressedBuffer& Content, const FStringView Context, uint64& OutWriteSize) const;
 
 	FOptionalCacheRecord GetCacheRecordOnly(
 		const FCacheKey& Key,
@@ -1453,7 +1453,6 @@ FFileSystemCacheStore::EPutStatus FFileSystemCacheStore::PutCachedData(
 	}
 	else
 	{
-		COOK_STAT(Timer.AddMiss(Data.Num()));
 		UE_LOG(LogDerivedDataCache, Verbose, TEXT("%s: Skipping put of %s to existing file"), *CachePath, CacheKey);
 		Status = EPutStatus::Cached;
 	}
@@ -1546,13 +1545,14 @@ void FFileSystemCacheStore::Put(
 		TRACE_CPUPROFILER_EVENT_SCOPE(FileSystemDDC_Put);
 		TRACE_COUNTER_INCREMENT(FileSystemDDC_Put);
 		COOK_STAT(auto Timer = UsageStats.TimePut());
-		if (PutCacheRecord(Record, Context, Request.Policy))
+		uint64 WriteSize = 0;
+		if (PutCacheRecord(Record, Context, Request.Policy, WriteSize))
 		{
 			UE_LOG(LogDerivedDataCache, Verbose, TEXT("%s: Cache put complete for %s from '%.*s'"),
 				*CachePath, *WriteToString<96>(Record.GetKey()), Context.Len(), Context.GetData());
 			TRACE_COUNTER_INCREMENT(FileSystemDDC_PutHit);
 			TRACE_COUNTER_ADD(FileSystemDDC_BytesWritten, MeasureCompressedCacheRecord(Record));
-			COOK_STAT(Timer.AddHit(MeasureRawCacheRecord(Record)));
+			COOK_STAT(if (WriteSize) { Timer.AddHit(WriteSize); });
 			if (OnComplete)
 			{
 				OnComplete({Record.GetKey(), Request.UserData, EStatus::Ok});
@@ -1560,7 +1560,6 @@ void FFileSystemCacheStore::Put(
 		}
 		else
 		{
-			COOK_STAT(Timer.AddMiss(MeasureRawCacheRecord(Record)));
 			if (OnComplete)
 			{
 				OnComplete({Record.GetKey(), Request.UserData, EStatus::Error});
@@ -1695,7 +1694,8 @@ uint64 FFileSystemCacheStore::MeasureRawCacheRecord(const FCacheRecord& Record)
 bool FFileSystemCacheStore::PutCacheRecord(
 	const FCacheRecord& Record,
 	const FStringView Context,
-	const FCacheRecordPolicy& Policy)
+	const FCacheRecordPolicy& Policy,
+	uint64& OutWriteSize)
 {
 	if (!IsWritable())
 	{
@@ -1785,11 +1785,14 @@ bool FFileSystemCacheStore::PutCacheRecord(
 	// Save the external content to storage.
 	for (FCompressedBuffer& Content : ExternalContent)
 	{
-		PutCacheContent(Content, Context);
+		uint64 WriteSize = 0;
+		PutCacheContent(Content, Context, OutWriteSize);
+		OutWriteSize += WriteSize;
 	}
 
 	// Save the record package to storage.
-	if (!bRecordExists && !SaveFileWithHash(Path, Context, [&Package](FArchive& Ar) { Package.Save(Ar); }))
+	const auto WriteRecord = [&](FArchive& Ar) { Package.Save(Ar); OutWriteSize += uint64(Ar.TotalSize()); };
+	if (!bRecordExists && !SaveFileWithHash(Path, Context, WriteRecord))
 	{
 		return false;
 	}
@@ -1942,12 +1945,16 @@ FOptionalCacheRecord FFileSystemCacheStore::GetCacheRecord(
 	return RecordBuilder.Build();
 }
 
-bool FFileSystemCacheStore::PutCacheContent(const FCompressedBuffer& Content, const FStringView Context) const
+bool FFileSystemCacheStore::PutCacheContent(
+	const FCompressedBuffer& Content,
+	const FStringView Context,
+	uint64& OutWriteSize) const
 {
 	const FIoHash& RawHash = Content.GetRawHash();
 	TStringBuilder<256> Path;
 	BuildCacheContentPath(RawHash, Path);
-	if (!FileExists(Path) && !SaveFileWithHash(Path, Context, [&Content](FArchive& Ar) { Content.Save(Ar); }))
+	const auto WriteContent = [&](FArchive& Ar) { Content.Save(Ar); OutWriteSize += uint64(Ar.TotalSize()); };
+	if (!FileExists(Path) && !SaveFileWithHash(Path, Context, WriteContent))
 	{
 		return false;
 	}
