DDC: Fixed static analysis warnings about null checks

#rnx

diff --git a/Engine/Source/Developer/DerivedDataCache/Private/DerivedDataBackendAsyncPutWrapper.cpp b/Engine/Source/Developer/DerivedDataCache/Private/DerivedDataBackendAsyncPutWrapper.cpp
--- a/Engine/Source/Developer/DerivedDataCache/Private/DerivedDataBackendAsyncPutWrapper.cpp
+++ b/Engine/Source/Developer/DerivedDataCache/Private/DerivedDataBackendAsyncPutWrapper.cpp
@@ -335,10 +335,7 @@ TSharedRef<FDerivedDataCacheStatsNode> FDerivedDataBackendAsyncPutWrapper::Gathe
 	Usage->Stats.Add(TEXT("AsyncPut"), UsageStats);
 	Usage->Stats.Add(TEXT("AsyncPutSync"), PutSyncUsageStats);
 
-	if (InnerBackend)
-	{
-		Usage->Children.Add(InnerBackend->GatherUsageStats());
-	}
+	Usage->Children.Add(InnerBackend->GatherUsageStats());
 	if (InflightCache)
 	{
 		Usage->Children.Add(InflightCache->GatherUsageStats());
diff --git a/Engine/Source/Developer/DerivedDataCache/Private/DerivedDataBackendVerifyWrapper.h b/Engine/Source/Developer/DerivedDataCache/Private/DerivedDataBackendVerifyWrapper.h
--- a/Engine/Source/Developer/DerivedDataCache/Private/DerivedDataBackendVerifyWrapper.h
+++ b/Engine/Source/Developer/DerivedDataCache/Private/DerivedDataBackendVerifyWrapper.h
@@ -139,12 +139,7 @@ public:
 		TSharedRef<FDerivedDataCacheStatsNode> Usage =
 			MakeShared<FDerivedDataCacheStatsNode>(TEXT("VerifyWrapper"), TEXT(""), InnerBackend->GetSpeedClass() == ESpeedClass::Local);
 		Usage->Stats.Add(TEXT(""), UsageStats);
-
-		if (InnerBackend)
-		{
-			Usage->Children.Add(InnerBackend->GatherUsageStats());
-		}
-
+		Usage->Children.Add(InnerBackend->GatherUsageStats());
 		return Usage;
 	}
 
diff --git a/Engine/Source/Developer/DerivedDataCache/Private/DerivedDataLimitKeyLengthWrapper.h b/Engine/Source/Developer/DerivedDataCache/Private/DerivedDataLimitKeyLengthWrapper.h
--- a/Engine/Source/Developer/DerivedDataCache/Private/DerivedDataLimitKeyLengthWrapper.h
+++ b/Engine/Source/Developer/DerivedDataCache/Private/DerivedDataLimitKeyLengthWrapper.h
@@ -234,12 +234,7 @@ public:
 		TSharedRef<FDerivedDataCacheStatsNode> Usage =
 			MakeShared<FDerivedDataCacheStatsNode>(TEXT("LimitKeyLength"), TEXT(""), InnerBackend->GetSpeedClass() == ESpeedClass::Local);
 		Usage->Stats.Add(TEXT(""), UsageStats);
-
-		if (InnerBackend)
-		{
-			Usage->Children.Add(InnerBackend->GatherUsageStats());
-		}
-
+		Usage->Children.Add(InnerBackend->GatherUsageStats());
 		return Usage;
 	}
 
