DDC: Fixed the logged cache summary stats

Stats from the root node must be a superset of stats from the rest of the nodes. The key length limiter is always the root node at present, and has been updated to include stats for the structured cache.

#rnx

diff --git a/Engine/Source/Developer/DerivedDataCache/Private/DerivedDataCache.cpp b/Engine/Source/Developer/DerivedDataCache/Private/DerivedDataCache.cpp
--- a/Engine/Source/Developer/DerivedDataCache/Private/DerivedDataCache.cpp
+++ b/Engine/Source/Developer/DerivedDataCache/Private/DerivedDataCache.cpp
@@ -54,121 +54,97 @@ namespace DerivedDataCacheCookStats
 	// See https://developercommunity.visualstudio.com/content/problem/576913/c6244-regression-in-new-lambda-processorpermissive.html
 	static void AddCookStats(FCookStatsManager::AddStatFuncRef AddStat)
 	{
-		PRAGMA_DISABLE_DEPRECATION_WARNINGS
-		TSharedRef<FDerivedDataCacheStatsNode> DDCUsage = GetDerivedDataCacheRef().GatherUsageStats();
-		PRAGMA_ENABLE_DEPRECATION_WARNINGS
-		TMap<FString, FDerivedDataCacheUsageStats> DDCStats = DDCUsage->ToLegacyUsageMap();
+		PRAGMA_DISABLE_DEPRECATION_WARNINGS;
+		TSharedRef<FDerivedDataCacheStatsNode> RootNode = GetDerivedDataCacheRef().GatherUsageStats();
+		PRAGMA_ENABLE_DEPRECATION_WARNINGS;
+
 		{
 			const FString StatName(TEXT("DDC.Usage"));
-			for (const auto& UsageStatPair : DDCStats)
+			for (const auto& UsageStatPair : RootNode->ToLegacyUsageMap())
 			{
 				UsageStatPair.Value.LogStats(AddStat, StatName, UsageStatPair.Key);
 			}
 		}
 
+		TArray<TSharedRef<const FDerivedDataCacheStatsNode>> Nodes;
+		RootNode->ForEachDescendant([&Nodes](TSharedRef<const FDerivedDataCacheStatsNode> Node)
+		{
+			if (Node->Children.IsEmpty())
+			{
+				Nodes.Add(Node);
+			}
+		});
+
 		// Now lets add some summary data to that applies some crazy knowledge of how we set up our DDC. The goal 
 		// is to print out the global hit rate, and the hit rate of the local and shared DDC.
 		// This is done by adding up the total get/miss calls the root node receives.
 		// Then we find the FileSystem nodes that correspond to the local and shared cache using some hacky logic to detect a "network drive".
 		// If the DDC graph ever contains more than one local or remote filesystem, this will only find one of them.
 		{
-			TArray<FString, TInlineAllocator<20>> Keys;
-			DDCStats.GenerateKeyArray(Keys);
-			FString* RootKey = Keys.FindByPredicate([](const FString& Key) {return Key.StartsWith(TEXT(" 0:")); });
-			// look for a Filesystem DDC that doesn't have a UNC path. Ugly, yeah, but we only cook on PC at the moment.
-			FString* LocalDDCKey = Keys.FindByPredicate([](const FString& Key) {return Key.Contains(TEXT(": FileSystem.")) && !Key.Contains(TEXT("//")); });
-			// look for a UNC path
-			FString* SharedDDCKey = Keys.FindByPredicate([](const FString& Key) {return Key.Contains(TEXT(": FileSystem.//")); });
-			// look for a Cloud path
-			FString* CloudDDCKey = Keys.FindByPredicate([](const FString& Key) {return Key.Contains(TEXT(": HTTP")); });
-			// look for a Zen Path
-			FString* ZenDDCKey = Keys.FindByPredicate([](const FString& Key) {return Key.Contains(TEXT(": Zen")); });
-
-			if (RootKey)
+			const TSharedRef<const FDerivedDataCacheStatsNode>* LocalNode = Nodes.FindByPredicate([](TSharedRef<const FDerivedDataCacheStatsNode> Node) { return Node->GetCacheType() == TEXT("File System") && Node->IsLocal(); });
+			const TSharedRef<const FDerivedDataCacheStatsNode>* SharedNode = Nodes.FindByPredicate([](TSharedRef<const FDerivedDataCacheStatsNode> Node) { return Node->GetCacheType() == TEXT("File System") && !Node->IsLocal(); });
+			const TSharedRef<const FDerivedDataCacheStatsNode>* CloudNode = Nodes.FindByPredicate([](TSharedRef<const FDerivedDataCacheStatsNode> Node) { return Node->GetCacheType() == TEXT("Horde Storage"); });
+			const TSharedRef<const FDerivedDataCacheStatsNode>* ZenLocalNode = Nodes.FindByPredicate([](TSharedRef<const FDerivedDataCacheStatsNode> Node) { return Node->GetCacheType() == TEXT("Zen") && Node->IsLocal(); });
+			const TSharedRef<const FDerivedDataCacheStatsNode>* ZenRemoteNode = Nodes.FindByPredicate([](TSharedRef<const FDerivedDataCacheStatsNode> Node) { return (Node->GetCacheType() == TEXT("Zen") || Node->GetCacheType() == TEXT("Horde")) && !Node->IsLocal(); });
+
+			const FDerivedDataCacheUsageStats& RootStats = RootNode->Stats.CreateConstIterator().Value();
+			const int64 TotalGetHits = RootStats.GetStats.GetAccumulatedValueAnyThread(FCookStats::CallStats::EHitOrMiss::Hit, FCookStats::CallStats::EStatType::Counter);
+			const int64 TotalGetMisses = RootStats.GetStats.GetAccumulatedValueAnyThread(FCookStats::CallStats::EHitOrMiss::Miss, FCookStats::CallStats::EStatType::Counter);
+			const int64 TotalGets = TotalGetHits + TotalGetMisses;
+
+			int64 LocalHits = 0;
+			if (LocalNode)
 			{
-				const FDerivedDataCacheUsageStats& RootStats = DDCStats[*RootKey];
-				int64 TotalGetHits =
-					RootStats.GetStats.GetAccumulatedValue(FCookStats::CallStats::EHitOrMiss::Hit, FCookStats::CallStats::EStatType::Counter, true) +
-					RootStats.GetStats.GetAccumulatedValue(FCookStats::CallStats::EHitOrMiss::Hit, FCookStats::CallStats::EStatType::Counter, false);
-				int64 TotalGetMisses =
-					RootStats.GetStats.GetAccumulatedValue(FCookStats::CallStats::EHitOrMiss::Miss, FCookStats::CallStats::EStatType::Counter, true) +
-					RootStats.GetStats.GetAccumulatedValue(FCookStats::CallStats::EHitOrMiss::Miss, FCookStats::CallStats::EStatType::Counter, false);
-				int64 TotalGets = TotalGetHits + TotalGetMisses;
-
-				int64 LocalHits = 0;
-				if (LocalDDCKey)
-				{
-					const FDerivedDataCacheUsageStats& LocalDDCStats = DDCStats[*LocalDDCKey];
-					LocalHits =
-						LocalDDCStats.GetStats.GetAccumulatedValue(FCookStats::CallStats::EHitOrMiss::Hit, FCookStats::CallStats::EStatType::Counter, true) +
-						LocalDDCStats.GetStats.GetAccumulatedValue(FCookStats::CallStats::EHitOrMiss::Hit, FCookStats::CallStats::EStatType::Counter, false);
-				}
-				int64 SharedHits = 0;
-				if (SharedDDCKey)
-				{
-					// The shared DDC is only queried if the local one misses (or there isn't one). So it's hit rate is technically 
-					const FDerivedDataCacheUsageStats& SharedDDCStats = DDCStats[*SharedDDCKey];
-					SharedHits =
-						SharedDDCStats.GetStats.GetAccumulatedValue(FCookStats::CallStats::EHitOrMiss::Hit, FCookStats::CallStats::EStatType::Counter, true) +
-						SharedDDCStats.GetStats.GetAccumulatedValue(FCookStats::CallStats::EHitOrMiss::Hit, FCookStats::CallStats::EStatType::Counter, false);
-				}
-				int64 CloudHits = 0;
-				if (CloudDDCKey)
-				{
-					const FDerivedDataCacheUsageStats& CloudDDCStats = DDCStats[*CloudDDCKey];
-					CloudHits =
-						CloudDDCStats.GetStats.GetAccumulatedValue(FCookStats::CallStats::EHitOrMiss::Hit, FCookStats::CallStats::EStatType::Counter, true) +
-						CloudDDCStats.GetStats.GetAccumulatedValue(FCookStats::CallStats::EHitOrMiss::Hit, FCookStats::CallStats::EStatType::Counter, false);
-				}
-
-				int64 TotalPutHits =
-					RootStats.PutStats.GetAccumulatedValue(FCookStats::CallStats::EHitOrMiss::Hit, FCookStats::CallStats::EStatType::Counter, true) +
-					RootStats.PutStats.GetAccumulatedValue(FCookStats::CallStats::EHitOrMiss::Hit, FCookStats::CallStats::EStatType::Counter, false);
-				int64 TotalPutMisses =
-					RootStats.PutStats.GetAccumulatedValue(FCookStats::CallStats::EHitOrMiss::Miss, FCookStats::CallStats::EStatType::Counter, true) +
-					RootStats.PutStats.GetAccumulatedValue(FCookStats::CallStats::EHitOrMiss::Miss, FCookStats::CallStats::EStatType::Counter, false);
-				int64 TotalPuts = TotalPutHits + TotalPutMisses;
-
-#if UE_WITH_ZEN
-				if (ZenDDCKey)
-				{
-					LocalDDCKey = ZenDDCKey;
-					
-					UE::Zen::FZenStats ZenStats;
-					
-					if (UE::Zen::GetDefaultServiceInstance().GetStats(ZenStats))
-					{				
-						TotalGetHits = ZenStats.CacheStats.Hits;
-						TotalGetMisses = ZenStats.CacheStats.Misses;
-						TotalGets = TotalGetHits+ TotalGetMisses;
-						LocalHits = ZenStats.CacheStats.Hits - ZenStats.CacheStats.UpstreamHits;
-						
-						SharedDDCKey = ZenStats.UpstreamStats.EndPointStats.IsEmpty() ? nullptr : ZenDDCKey;	
-						SharedHits = ZenStats.CacheStats.UpstreamHits;
-					}
-				}
-#endif // UE_WITH_ZEN
-
-				AddStat(TEXT("DDC.Summary"), FCookStatsManager::CreateKeyValueArray(
-					TEXT("BackEnd"), FDerivedDataBackend::Get().GetGraphName(),
-					TEXT("HasLocalCache"), LocalDDCKey != nullptr,
-					TEXT("HasSharedCache"), SharedDDCKey!=nullptr,
-					TEXT("HasCloudCache"), CloudDDCKey !=nullptr,
-					TEXT("HasZenCache"), ZenDDCKey != nullptr,
-					TEXT("TotalGetHits"), TotalGetHits,
-					TEXT("TotalGets"), TotalGets,
-					TEXT("TotalGetHitPct"), SafeDivide(TotalGetHits, TotalGets),
-					TEXT("LocalGetHitPct"), SafeDivide(LocalHits, TotalGets),
-					TEXT("SharedGetHitPct"), SafeDivide(SharedHits, TotalGets),
-					TEXT("CloudGetHitPct"), SafeDivide(CloudHits, TotalGets),
-					TEXT("OtherGetHitPct"), SafeDivide((TotalGetHits - LocalHits - SharedHits), TotalGets),
-					TEXT("GetMissPct"), SafeDivide(TotalGetMisses, TotalGets),
-					TEXT("TotalPutHits"), TotalPutHits,
-					TEXT("TotalPuts"), TotalPuts,
-					TEXT("TotalPutHitPct"), SafeDivide(TotalPutHits, TotalPuts),
-					TEXT("PutMissPct"), SafeDivide(TotalPutMisses, TotalPuts)
-					));
+				const FDerivedDataCacheUsageStats& Stats = (*LocalNode)->Stats.CreateConstIterator().Value();
+				LocalHits += Stats.GetStats.GetAccumulatedValueAnyThread(FCookStats::CallStats::EHitOrMiss::Hit, FCookStats::CallStats::EStatType::Counter);
 			}
+			if (ZenLocalNode)
+			{
+				const FDerivedDataCacheUsageStats& Stats = (*ZenLocalNode)->Stats.CreateConstIterator().Value();
+				LocalHits += Stats.GetStats.GetAccumulatedValueAnyThread(FCookStats::CallStats::EHitOrMiss::Hit, FCookStats::CallStats::EStatType::Counter);
+			}
+			int64 SharedHits = 0;
+			if (SharedNode)
+			{
+				// The shared DDC is only queried if the local one misses (or there isn't one). So it's hit rate is technically 
+				const FDerivedDataCacheUsageStats& Stats = (*SharedNode)->Stats.CreateConstIterator().Value();
+				SharedHits += Stats.GetStats.GetAccumulatedValueAnyThread(FCookStats::CallStats::EHitOrMiss::Hit, FCookStats::CallStats::EStatType::Counter);
+			}
+			if (ZenRemoteNode)
+			{
+				const FDerivedDataCacheUsageStats& Stats = (*ZenRemoteNode)->Stats.CreateConstIterator().Value();
+				SharedHits += Stats.GetStats.GetAccumulatedValueAnyThread(FCookStats::CallStats::EHitOrMiss::Hit, FCookStats::CallStats::EStatType::Counter);
+			}
+			int64 CloudHits = 0;
+			if (CloudNode)
+			{
+				const FDerivedDataCacheUsageStats& Stats = (*CloudNode)->Stats.CreateConstIterator().Value();
+				CloudHits += Stats.GetStats.GetAccumulatedValueAnyThread(FCookStats::CallStats::EHitOrMiss::Hit, FCookStats::CallStats::EStatType::Counter);
+			}
+
+			const int64 TotalPutHits = RootStats.PutStats.GetAccumulatedValueAnyThread(FCookStats::CallStats::EHitOrMiss::Hit, FCookStats::CallStats::EStatType::Counter);
+			const int64 TotalPutMisses = RootStats.PutStats.GetAccumulatedValueAnyThread(FCookStats::CallStats::EHitOrMiss::Miss, FCookStats::CallStats::EStatType::Counter);
+			const int64 TotalPuts = TotalPutHits + TotalPutMisses;
+
+			AddStat(TEXT("DDC.Summary"), FCookStatsManager::CreateKeyValueArray(
+				TEXT("BackEnd"), FDerivedDataBackend::Get().GetGraphName(),
+				TEXT("HasLocalCache"), LocalNode || ZenLocalNode,
+				TEXT("HasSharedCache"), SharedNode || ZenRemoteNode,
+				TEXT("HasCloudCache"), !!CloudNode,
+				TEXT("HasZenCache"), ZenLocalNode || ZenRemoteNode,
+				TEXT("TotalGetHits"), TotalGetHits,
+				TEXT("TotalGets"), TotalGets,
+				TEXT("TotalGetHitPct"), SafeDivide(TotalGetHits, TotalGets),
+				TEXT("LocalGetHitPct"), SafeDivide(LocalHits, TotalGets),
+				TEXT("SharedGetHitPct"), SafeDivide(SharedHits, TotalGets),
+				TEXT("CloudGetHitPct"), SafeDivide(CloudHits, TotalGets),
+				TEXT("OtherGetHitPct"), SafeDivide((TotalGetHits - LocalHits - SharedHits - CloudHits), TotalGets),
+				TEXT("GetMissPct"), SafeDivide(TotalGetMisses, TotalGets),
+				TEXT("TotalPutHits"), TotalPutHits,
+				TEXT("TotalPuts"), TotalPuts,
+				TEXT("TotalPutHitPct"), SafeDivide(TotalPutHits, TotalPuts),
+				TEXT("PutMissPct"), SafeDivide(TotalPutMisses, TotalPuts)
+				));
 		}
 	}
 
diff --git a/Engine/Source/Developer/DerivedDataCache/Private/DerivedDataLimitKeyLengthWrapper.h b/Engine/Source/Developer/DerivedDataCache/Private/DerivedDataLimitKeyLengthWrapper.h
--- a/Engine/Source/Developer/DerivedDataCache/Private/DerivedDataLimitKeyLengthWrapper.h
+++ b/Engine/Source/Developer/DerivedDataCache/Private/DerivedDataLimitKeyLengthWrapper.h
@@ -244,7 +244,22 @@ public:
 		IRequestOwner& Owner,
 		FOnCachePutComplete&& OnComplete) override
 	{
+	#if ENABLE_COOK_STATS
+		return InnerBackend->Put(Requests, Context, Owner,
+			[this, OnComplete = MoveTemp(OnComplete)](FCachePutCompleteParams&& Params)
+			{
+				if (Params.Status == EStatus::Ok)
+				{
+					UsageStats.TimePut().AddHit(0);
+				}
+				if (OnComplete)
+				{
+					OnComplete(MoveTemp(Params));
+				}
+			});
+	#else
 		return InnerBackend->Put(Requests, Context, Owner, MoveTemp(OnComplete));
+	#endif
 	}
 
 	virtual void Get(
@@ -253,7 +268,22 @@ public:
 		IRequestOwner& Owner,
 		FOnCacheGetComplete&& OnComplete) override
 	{
+	#if ENABLE_COOK_STATS
+		return InnerBackend->Get(Requests, Context, Owner,
+			[this, OnComplete = MoveTemp(OnComplete)](FCacheGetCompleteParams&& Params)
+			{
+				if (Params.Status == EStatus::Ok)
+				{
+					UsageStats.TimeGet().AddHit(0);
+				}
+				if (OnComplete)
+				{
+					OnComplete(MoveTemp(Params));
+				}
+			});
+	#else
 		return InnerBackend->Get(Requests, Context, Owner, MoveTemp(OnComplete));
+	#endif
 	}
 
 	virtual void GetChunks(
@@ -262,7 +292,22 @@ public:
 		IRequestOwner& Owner,
 		FOnCacheGetChunkComplete&& OnComplete) override
 	{
+	#if ENABLE_COOK_STATS
+		return InnerBackend->GetChunks(Chunks, Context, Owner,
+			[this, OnComplete = MoveTemp(OnComplete)](FCacheGetChunkCompleteParams&& Params)
+			{
+				if (Params.Status == EStatus::Ok)
+				{
+					UsageStats.TimeGet().AddHit(0);
+				}
+				if (OnComplete)
+				{
+					OnComplete(MoveTemp(Params));
+				}
+			});
+	#else
 		return InnerBackend->GetChunks(Chunks, Context, Owner, MoveTemp(OnComplete));
+	#endif
 	}
 
 private:
