DDC: Moved the ICacheStore Context params to be a Name member of the request types and return them in the completion params

#rnx

diff --git a/Engine/Source/Developer/DerivedDataCache/Private/DerivedDataBackendAsyncPutWrapper.cpp b/Engine/Source/Developer/DerivedDataCache/Private/DerivedDataBackendAsyncPutWrapper.cpp
--- a/Engine/Source/Developer/DerivedDataCache/Private/DerivedDataBackendAsyncPutWrapper.cpp
+++ b/Engine/Source/Developer/DerivedDataCache/Private/DerivedDataBackendAsyncPutWrapper.cpp
@@ -456,30 +456,29 @@ private:
 
 void FDerivedDataBackendAsyncPutWrapper::Put(
 	const TConstArrayView<FCachePutRequest> Requests,
-	const FStringView Context,
 	IRequestOwner& Owner,
 	FOnCachePutComplete&& OnComplete)
 {
 	if (Owner.GetPriority() == EPriority::Blocking || !GDDCIOThreadPool)
 	{
-		InnerBackend->Put(Requests, Context, Owner, MoveTemp(OnComplete));
+		InnerBackend->Put(Requests, Owner, MoveTemp(OnComplete));
 	}
 	else
 	{
 		FDerivedDataAsyncWrapperRequest* Request = new FDerivedDataAsyncWrapperRequest(Owner,
-			[this, Requests = TArray<FCachePutRequest>(Requests), Context = FString(Context), OnComplete = MoveTemp(OnComplete)](bool bCancel) mutable
+			[this, Requests = TArray<FCachePutRequest>(Requests), OnComplete = MoveTemp(OnComplete)](bool bCancel) mutable
 			{
 				if (!bCancel)
 				{
 					FRequestOwner BlockingOwner(EPriority::Blocking);
-					InnerBackend->Put(Requests, Context, BlockingOwner, MoveTemp(OnComplete));
+					InnerBackend->Put(Requests, BlockingOwner, MoveTemp(OnComplete));
 					BlockingOwner.Wait();
 				}
 				else if (OnComplete)
 				{
 					for (const FCachePutRequest& Request : Requests)
 					{
-						OnComplete({Request.Record.GetKey(), Request.UserData, EStatus::Canceled});
+						OnComplete({Request.Name, Request.Record.GetKey(), Request.UserData, EStatus::Canceled});
 					}
 				}
 			});
@@ -489,30 +488,29 @@ void FDerivedDataBackendAsyncPutWrapper::Put(
 
 void FDerivedDataBackendAsyncPutWrapper::Get(
 	const TConstArrayView<FCacheGetRequest> Requests,
-	const FStringView Context,
 	IRequestOwner& Owner,
 	FOnCacheGetComplete&& OnComplete)
 {
 	if (Owner.GetPriority() == EPriority::Blocking || !GDDCIOThreadPool)
 	{
-		InnerBackend->Get(Requests, Context, Owner, MoveTemp(OnComplete));
+		InnerBackend->Get(Requests, Owner, MoveTemp(OnComplete));
 	}
 	else
 	{
 		FDerivedDataAsyncWrapperRequest* Request = new FDerivedDataAsyncWrapperRequest(Owner,
-			[this, Requests = TArray<FCacheGetRequest>(Requests), Context = FString(Context), OnComplete = MoveTemp(OnComplete)](bool bCancel) mutable
+			[this, Requests = TArray<FCacheGetRequest>(Requests), OnComplete = MoveTemp(OnComplete)](bool bCancel) mutable
 			{
 				if (!bCancel)
 				{
 					FRequestOwner BlockingOwner(EPriority::Blocking);
-					InnerBackend->Get(Requests, Context, BlockingOwner, MoveTemp(OnComplete));
+					InnerBackend->Get(Requests, BlockingOwner, MoveTemp(OnComplete));
 					BlockingOwner.Wait();
 				}
 				else if (OnComplete)
 				{
 					for (const FCacheGetRequest& Request : Requests)
 					{
-						OnComplete({FCacheRecordBuilder(Request.Key).Build(), Request.UserData, EStatus::Canceled});
+						OnComplete({Request.Name, FCacheRecordBuilder(Request.Key).Build(), Request.UserData, EStatus::Canceled});
 					}
 				}
 			});
@@ -521,31 +519,30 @@ void FDerivedDataBackendAsyncPutWrapper::Get(
 }
 
 void FDerivedDataBackendAsyncPutWrapper::GetChunks(
-	TConstArrayView<FCacheChunkRequest> Chunks,
-	FStringView Context,
+	TConstArrayView<FCacheChunkRequest> Requests,
 	IRequestOwner& Owner,
-	FOnCacheGetChunkComplete&& OnComplete)
+	FOnCacheChunkComplete&& OnComplete)
 {
 	if (Owner.GetPriority() == EPriority::Blocking || !GDDCIOThreadPool)
 	{
-		InnerBackend->GetChunks(Chunks, Context, Owner, MoveTemp(OnComplete));
+		InnerBackend->GetChunks(Requests, Owner, MoveTemp(OnComplete));
 	}
 	else
 	{
 		FDerivedDataAsyncWrapperRequest* Request = new FDerivedDataAsyncWrapperRequest(Owner,
-			[this, Chunks = TArray<FCacheChunkRequest>(Chunks), Context = FString(Context), OnComplete = MoveTemp(OnComplete)](bool bCancel) mutable
+			[this, Requests = TArray<FCacheChunkRequest>(Requests), OnComplete = MoveTemp(OnComplete)](bool bCancel) mutable
 			{
 				if (!bCancel)
 				{
 					FRequestOwner BlockingOwner(EPriority::Blocking);
-					InnerBackend->GetChunks(Chunks, Context, BlockingOwner, MoveTemp(OnComplete));
+					InnerBackend->GetChunks(Requests, BlockingOwner, MoveTemp(OnComplete));
 					BlockingOwner.Wait();
 				}
 				else if (OnComplete)
 				{
-					for (const FCacheChunkRequest& Chunk : Chunks)
+					for (const FCacheChunkRequest& Request : Requests)
 					{
-						OnComplete({Chunk.Key, Chunk.Id, Chunk.RawOffset, 0, {}, {}, Chunk.UserData, EStatus::Canceled});
+						OnComplete({Request.Name, Request.Key, Request.Id, Request.RawOffset, 0, {}, {}, Request.UserData, EStatus::Canceled});
 					}
 				}
 			});
diff --git a/Engine/Source/Developer/DerivedDataCache/Private/DerivedDataBackendAsyncPutWrapper.h b/Engine/Source/Developer/DerivedDataCache/Private/DerivedDataBackendAsyncPutWrapper.h
--- a/Engine/Source/Developer/DerivedDataCache/Private/DerivedDataBackendAsyncPutWrapper.h
+++ b/Engine/Source/Developer/DerivedDataCache/Private/DerivedDataBackendAsyncPutWrapper.h
@@ -137,21 +137,18 @@ public:
 
 	virtual void Put(
 		TConstArrayView<FCachePutRequest> Requests,
-		FStringView Context,
 		IRequestOwner& Owner,
 		FOnCachePutComplete&& OnComplete) override;
 
 	virtual void Get(
 		TConstArrayView<FCacheGetRequest> Requests,
-		FStringView Context,
 		IRequestOwner& Owner,
 		FOnCacheGetComplete&& OnComplete) override;
 
 	virtual void GetChunks(
-		TConstArrayView<FCacheChunkRequest> Chunks,
-		FStringView Context,
+		TConstArrayView<FCacheChunkRequest> Requests,
 		IRequestOwner& Owner,
-		FOnCacheGetChunkComplete&& OnComplete) override;
+		FOnCacheChunkComplete&& OnComplete) override;
 
 private:
 	FDerivedDataCacheUsageStats UsageStats;
diff --git a/Engine/Source/Developer/DerivedDataCache/Private/DerivedDataBackendThrottleWrapper.h b/Engine/Source/Developer/DerivedDataCache/Private/DerivedDataBackendThrottleWrapper.h
--- a/Engine/Source/Developer/DerivedDataCache/Private/DerivedDataBackendThrottleWrapper.h
+++ b/Engine/Source/Developer/DerivedDataCache/Private/DerivedDataBackendThrottleWrapper.h
@@ -167,7 +167,6 @@ public:
 
 	virtual void Put(
 		TConstArrayView<FCachePutRequest> Requests,
-		FStringView Context,
 		IRequestOwner& Owner,
 		FOnCachePutComplete&& OnComplete) override
 	{
@@ -183,7 +182,7 @@ public:
 			return {Request.Record.GetKey(), Private::GetCacheRecordCompressedSize(Request.Record)};
 		});
 
-		InnerBackend->Put(Requests, Context, Owner,
+		InnerBackend->Put(Requests, Owner,
 			[this, RecordSizes = MoveTemp(RecordSizes), State = EnterThrottlingScope(), OnComplete = MoveTemp(OnComplete)](FCachePutCompleteParams&& Params)
 			{
 				const FRecordSize* Size = Algo::FindBy(RecordSizes, Params.Key, &FRecordSize::Key);
@@ -197,11 +196,10 @@ public:
 
 	virtual void Get(
 		TConstArrayView<FCacheGetRequest> Requests,
-		FStringView Context,
 		IRequestOwner& Owner,
 		FOnCacheGetComplete&& OnComplete) override
 	{
-		InnerBackend->Get(Requests, Context, Owner,
+		InnerBackend->Get(Requests, Owner,
 			[this, State = EnterThrottlingScope(), OnComplete = MoveTemp(OnComplete)](FCacheGetCompleteParams&& Params)
 			{
 				CloseThrottlingScope(State, FThrottlingState(this, Private::GetCacheRecordCompressedSize(Params.Record)));
@@ -213,13 +211,12 @@ public:
 	}
 
 	virtual void GetChunks(
-		TConstArrayView<FCacheChunkRequest> Chunks,
-		FStringView Context,
+		TConstArrayView<FCacheChunkRequest> Requests,
 		IRequestOwner& Owner,
-		FOnCacheGetChunkComplete&& OnComplete) override
+		FOnCacheChunkComplete&& OnComplete) override
 	{
-		InnerBackend->GetChunks(Chunks, Context, Owner,
-			[this, State = EnterThrottlingScope(), OnComplete = MoveTemp(OnComplete)](FCacheGetChunkCompleteParams&& Params)
+		InnerBackend->GetChunks(Requests, Owner,
+			[this, State = EnterThrottlingScope(), OnComplete = MoveTemp(OnComplete)](FCacheChunkCompleteParams&& Params)
 			{
 				CloseThrottlingScope(State, FThrottlingState(this, Params.RawData.GetSize()));
 				if (OnComplete)
diff --git a/Engine/Source/Developer/DerivedDataCache/Private/DerivedDataBackendVerifyWrapper.h b/Engine/Source/Developer/DerivedDataCache/Private/DerivedDataBackendVerifyWrapper.h
--- a/Engine/Source/Developer/DerivedDataCache/Private/DerivedDataBackendVerifyWrapper.h
+++ b/Engine/Source/Developer/DerivedDataCache/Private/DerivedDataBackendVerifyWrapper.h
@@ -160,29 +160,26 @@ public:
 
 	virtual void Put(
 		TConstArrayView<FCachePutRequest> Requests,
-		FStringView Context,
 		IRequestOwner& Owner,
 		FOnCachePutComplete&& OnComplete) override
 	{
-		return InnerBackend->Put(Requests, Context, Owner, MoveTemp(OnComplete));
+		return InnerBackend->Put(Requests, Owner, MoveTemp(OnComplete));
 	}
 
 	virtual void Get(
 		TConstArrayView<FCacheGetRequest> Requests,
-		FStringView Context,
 		IRequestOwner& Owner,
 		FOnCacheGetComplete&& OnComplete) override
 	{
-		return InnerBackend->Get(Requests, Context, Owner, MoveTemp(OnComplete));
+		return InnerBackend->Get(Requests, Owner, MoveTemp(OnComplete));
 	}
 
 	virtual void GetChunks(
-		TConstArrayView<FCacheChunkRequest> Chunks,
-		FStringView Context,
+		TConstArrayView<FCacheChunkRequest> Requests,
 		IRequestOwner& Owner,
-		FOnCacheGetChunkComplete&& OnComplete) override
+		FOnCacheChunkComplete&& OnComplete) override
 	{
-		return InnerBackend->GetChunks(Chunks, Context, Owner, MoveTemp(OnComplete));
+		return InnerBackend->GetChunks(Requests, Owner, MoveTemp(OnComplete));
 	}
 
 private:
diff --git a/Engine/Source/Developer/DerivedDataCache/Private/DerivedDataCache.cpp b/Engine/Source/Developer/DerivedDataCache/Private/DerivedDataCache.cpp
--- a/Engine/Source/Developer/DerivedDataCache/Private/DerivedDataCache.cpp
+++ b/Engine/Source/Developer/DerivedDataCache/Private/DerivedDataCache.cpp
@@ -244,13 +244,14 @@ void ICache::Put(
 	IRequestOwner& Owner,
 	FOnCachePutComplete&& OnComplete)
 {
+	const FCacheRequestName Name(Context);
 	TArray<FCachePutRequest, TInlineAllocator<16>> Requests;
 	Requests.Reserve(Records.Num());
 	for (const FCacheRecord& Record : Records)
 	{
-		Requests.Add({Record, Policy});
+		Requests.Add({Name, Record, Policy});
 	}
-	return Put(Requests, Context, Owner, MoveTemp(OnComplete));
+	return Put(Requests, Owner, MoveTemp(OnComplete));
 }
 
 void ICache::Get(
@@ -260,13 +261,14 @@ void ICache::Get(
 	IRequestOwner& Owner,
 	FOnCacheGetComplete&& OnComplete)
 {
+	const FCacheRequestName Name(Context);
 	TArray<FCacheGetRequest, TInlineAllocator<16>> Requests;
 	Requests.Reserve(Keys.Num());
 	for (const FCacheKey& Key : Keys)
 	{
-		Requests.Add({Key, Policy});
+		Requests.Add({Name, Key, Policy});
 	}
-	return Get(Requests, Context, Owner, MoveTemp(OnComplete));
+	return Get(Requests, Owner, MoveTemp(OnComplete));
 }
 
 } // UE::DerivedData
@@ -819,29 +821,26 @@ public:
 
 	void Put(
 		TConstArrayView<FCachePutRequest> Requests,
-		FStringView Context,
 		IRequestOwner& Owner,
 		FOnCachePutComplete&& OnComplete) final
 	{
-		return FDerivedDataBackend::Get().GetRoot().Put(Requests, Context, Owner, MoveTemp(OnComplete));
+		return FDerivedDataBackend::Get().GetRoot().Put(Requests, Owner, MoveTemp(OnComplete));
 	}
 
 	void Get(
 		TConstArrayView<FCacheGetRequest> Requests,
-		FStringView Context,
 		IRequestOwner& Owner,
 		FOnCacheGetComplete&& OnComplete) final
 	{
-		return FDerivedDataBackend::Get().GetRoot().Get(Requests, Context, Owner, MoveTemp(OnComplete));
+		return FDerivedDataBackend::Get().GetRoot().Get(Requests, Owner, MoveTemp(OnComplete));
 	}
 
 	void GetChunks(
-		TConstArrayView<FCacheChunkRequest> Chunks,
-		FStringView Context,
+		TConstArrayView<FCacheChunkRequest> Requests,
 		IRequestOwner& Owner,
-		FOnCacheGetChunkComplete&& OnComplete) final
+		FOnCacheChunkComplete&& OnComplete) final
 	{
-		return FDerivedDataBackend::Get().GetRoot().GetChunks(Chunks, Context, Owner, MoveTemp(OnComplete));
+		return FDerivedDataBackend::Get().GetRoot().GetChunks(Requests, Owner, MoveTemp(OnComplete));
 	}
 
 	// ICache Interface
diff --git a/Engine/Source/Developer/DerivedDataCache/Private/DerivedDataLimitKeyLengthWrapper.h b/Engine/Source/Developer/DerivedDataCache/Private/DerivedDataLimitKeyLengthWrapper.h
--- a/Engine/Source/Developer/DerivedDataCache/Private/DerivedDataLimitKeyLengthWrapper.h
+++ b/Engine/Source/Developer/DerivedDataCache/Private/DerivedDataLimitKeyLengthWrapper.h
@@ -240,12 +240,11 @@ public:
 
 	virtual void Put(
 		TConstArrayView<FCachePutRequest> Requests,
-		FStringView Context,
 		IRequestOwner& Owner,
 		FOnCachePutComplete&& OnComplete) override
 	{
 	#if ENABLE_COOK_STATS
-		return InnerBackend->Put(Requests, Context, Owner,
+		return InnerBackend->Put(Requests, Owner,
 			[this, OnComplete = MoveTemp(OnComplete)](FCachePutCompleteParams&& Params)
 			{
 				if (Params.Status == EStatus::Ok)
@@ -258,18 +257,17 @@ public:
 				}
 			});
 	#else
-		return InnerBackend->Put(Requests, Context, Owner, MoveTemp(OnComplete));
+		return InnerBackend->Put(Requests, Owner, MoveTemp(OnComplete));
 	#endif
 	}
 
 	virtual void Get(
 		TConstArrayView<FCacheGetRequest> Requests,
-		FStringView Context,
 		IRequestOwner& Owner,
 		FOnCacheGetComplete&& OnComplete) override
 	{
 	#if ENABLE_COOK_STATS
-		return InnerBackend->Get(Requests, Context, Owner,
+		return InnerBackend->Get(Requests, Owner,
 			[this, OnComplete = MoveTemp(OnComplete)](FCacheGetCompleteParams&& Params)
 			{
 				if (Params.Status == EStatus::Ok)
@@ -282,19 +280,18 @@ public:
 				}
 			});
 	#else
-		return InnerBackend->Get(Requests, Context, Owner, MoveTemp(OnComplete));
+		return InnerBackend->Get(Requests, Owner, MoveTemp(OnComplete));
 	#endif
 	}
 
 	virtual void GetChunks(
-		TConstArrayView<FCacheChunkRequest> Chunks,
-		FStringView Context,
+		TConstArrayView<FCacheChunkRequest> Requests,
 		IRequestOwner& Owner,
-		FOnCacheGetChunkComplete&& OnComplete) override
+		FOnCacheChunkComplete&& OnComplete) override
 	{
 	#if ENABLE_COOK_STATS
-		return InnerBackend->GetChunks(Chunks, Context, Owner,
-			[this, OnComplete = MoveTemp(OnComplete)](FCacheGetChunkCompleteParams&& Params)
+		return InnerBackend->GetChunks(Requests, Owner,
+			[this, OnComplete = MoveTemp(OnComplete)](FCacheChunkCompleteParams&& Params)
 			{
 				if (Params.Status == EStatus::Ok)
 				{
@@ -306,7 +303,7 @@ public:
 				}
 			});
 	#else
-		return InnerBackend->GetChunks(Chunks, Context, Owner, MoveTemp(OnComplete));
+		return InnerBackend->GetChunks(Requests, Owner, MoveTemp(OnComplete));
 	#endif
 	}
 
diff --git a/Engine/Source/Developer/DerivedDataCache/Private/FileSystemDerivedDataBackend.cpp b/Engine/Source/Developer/DerivedDataCache/Private/FileSystemDerivedDataBackend.cpp
--- a/Engine/Source/Developer/DerivedDataCache/Private/FileSystemDerivedDataBackend.cpp
+++ b/Engine/Source/Developer/DerivedDataCache/Private/FileSystemDerivedDataBackend.cpp
@@ -753,46 +753,43 @@ public:
 
 	void Put(
 		TConstArrayView<FCachePutRequest> Requests,
-		FStringView Context,
 		IRequestOwner& Owner,
 		FOnCachePutComplete&& OnComplete) final;
 	void Get(
 		TConstArrayView<FCacheGetRequest> Requests,
-		FStringView Context,
 		IRequestOwner& Owner,
 		FOnCacheGetComplete&& OnComplete) final;
 	void GetChunks(
-		TConstArrayView<FCacheChunkRequest> Chunks,
-		FStringView Context,
+		TConstArrayView<FCacheChunkRequest> Requests,
 		IRequestOwner& Owner,
-		FOnCacheGetChunkComplete&& OnComplete) final;
+		FOnCacheChunkComplete&& OnComplete) final;
 
 private:
 	uint64 MeasureCompressedCacheRecord(const FCacheRecord& Record) const;
 	uint64 MeasureRawCacheRecord(const FCacheRecord& Record) const;
 
-	bool PutCacheRecord(const FCacheRecord& Record, FStringView Context, const FCacheRecordPolicy& Policy, uint64& OutWriteSize);
-	bool PutCacheContent(const FCompressedBuffer& Content, const FStringView Context, uint64& OutWriteSize) const;
+	bool PutCacheRecord(const FCacheRecord& Record, FStringView RecordName, const FCacheRecordPolicy& Policy, uint64& OutWriteSize);
+	bool PutCacheContent(const FCompressedBuffer& Content, const FStringView ContentName, uint64& OutWriteSize) const;
 
 	FOptionalCacheRecord GetCacheRecordOnly(
 		const FCacheKey& Key,
-		const FStringView Context,
+		const FStringView RecordName,
 		const FCacheRecordPolicy& Policy);
 	FOptionalCacheRecord GetCacheRecord(
 		const FCacheKey& Key,
-		const FStringView Context,
+		const FStringView RecordName,
 		const FCacheRecordPolicy& Policy,
 		EStatus& OutStatus);
 	FPayload GetCacheContent(
 		const FCacheKey& Key,
 		const FPayload& Payload,
-		FStringView Context,
+		FStringView ContentName,
 		ECachePolicy Policy,
 		ECachePolicy SkipFlag) const;
 	void GetCacheContent(
 		const FCacheKey& Key,
 		const FPayload& Payload,
-		FStringView Context,
+		FStringView ContentName,
 		ECachePolicy Policy,
 		ECachePolicy SkipFlag,
 		FCompressedBufferReader& Reader,
@@ -801,11 +798,11 @@ private:
 	void BuildCacheRecordPath(const FCacheKey& CacheKey, FStringBuilderBase& Path) const;
 	void BuildCacheContentPath(const FIoHash& RawHash, FStringBuilderBase& Path) const;
 
-	bool SaveFileWithHash(FStringBuilderBase& Path, FStringView Context, TFunctionRef<void (FArchive&)> WriteFunction) const;
-	bool LoadFileWithHash(FStringBuilderBase& Path, FStringView Context, TFunctionRef<void (FArchive&)> ReadFunction) const;
-	bool SaveFile(FStringBuilderBase& Path, FStringView Context, TFunctionRef<void (FArchive&)> WriteFunction) const;
-	bool LoadFile(FStringBuilderBase& Path, FStringView Context, TFunctionRef<void (FArchive&)> ReadFunction) const;
-	TUniquePtr<FArchive> OpenFile(FStringBuilderBase& Path, FStringView Context) const;
+	bool SaveFileWithHash(FStringBuilderBase& Path, FStringView DebugName, TFunctionRef<void (FArchive&)> WriteFunction) const;
+	bool LoadFileWithHash(FStringBuilderBase& Path, FStringView DebugName, TFunctionRef<void (FArchive&)> ReadFunction) const;
+	bool SaveFile(FStringBuilderBase& Path, FStringView DebugName, TFunctionRef<void (FArchive&)> WriteFunction) const;
+	bool LoadFile(FStringBuilderBase& Path, FStringView DebugName, TFunctionRef<void (FArchive&)> ReadFunction) const;
+	TUniquePtr<FArchive> OpenFile(FStringBuilderBase& Path, FStringView DebugName) const;
 
 	bool FileExists(FStringBuilderBase& Path) const;
 
@@ -1535,7 +1532,6 @@ bool FFileSystemCacheStore::ApplyDebugOptions(FBackendDebugOptions& InOptions)
 
 void FFileSystemCacheStore::Put(
 	const TConstArrayView<FCachePutRequest> Requests,
-	const FStringView Context,
 	IRequestOwner& Owner,
 	FOnCachePutComplete&& OnComplete)
 {
@@ -1546,23 +1542,23 @@ void FFileSystemCacheStore::Put(
 		TRACE_COUNTER_INCREMENT(FileSystemDDC_Put);
 		COOK_STAT(auto Timer = UsageStats.TimePut());
 		uint64 WriteSize = 0;
-		if (PutCacheRecord(Record, Context, Request.Policy, WriteSize))
+		if (PutCacheRecord(Record, Request.Name, Request.Policy, WriteSize))
 		{
-			UE_LOG(LogDerivedDataCache, Verbose, TEXT("%s: Cache put complete for %s from '%.*s'"),
-				*CachePath, *WriteToString<96>(Record.GetKey()), Context.Len(), Context.GetData());
+			UE_LOG(LogDerivedDataCache, Verbose, TEXT("%s: Cache put complete for %s from '%s'"),
+				*CachePath, *WriteToString<96>(Record.GetKey()), *Request.Name);
 			TRACE_COUNTER_INCREMENT(FileSystemDDC_PutHit);
 			TRACE_COUNTER_ADD(FileSystemDDC_BytesWritten, MeasureCompressedCacheRecord(Record));
 			COOK_STAT(if (WriteSize) { Timer.AddHit(WriteSize); });
 			if (OnComplete)
 			{
-				OnComplete({Record.GetKey(), Request.UserData, EStatus::Ok});
+				OnComplete({Request.Name, Record.GetKey(), Request.UserData, EStatus::Ok});
 			}
 		}
 		else
 		{
 			if (OnComplete)
 			{
-				OnComplete({Record.GetKey(), Request.UserData, EStatus::Error});
+				OnComplete({Request.Name, Record.GetKey(), Request.UserData, EStatus::Error});
 			}
 		}
 	}
@@ -1570,7 +1566,6 @@ void FFileSystemCacheStore::Put(
 
 void FFileSystemCacheStore::Get(
 	const TConstArrayView<FCacheGetRequest> Requests,
-	const FStringView Context,
 	IRequestOwner& Owner,
 	FOnCacheGetComplete&& OnComplete)
 {
@@ -1580,36 +1575,35 @@ void FFileSystemCacheStore::Get(
 		TRACE_COUNTER_INCREMENT(FileSystemDDC_Get);
 		COOK_STAT(auto Timer = UsageStats.TimeGet());
 		EStatus Status = EStatus::Ok;
-		if (FOptionalCacheRecord Record = GetCacheRecord(Request.Key, Context, Request.Policy, Status))
+		if (FOptionalCacheRecord Record = GetCacheRecord(Request.Key, Request.Name, Request.Policy, Status))
 		{
-			UE_LOG(LogDerivedDataCache, Verbose, TEXT("%s: Cache hit for %s from '%.*s'"),
-				*CachePath, *WriteToString<96>(Request.Key), Context.Len(), Context.GetData());
+			UE_LOG(LogDerivedDataCache, Verbose, TEXT("%s: Cache hit for %s from '%s'"),
+				*CachePath, *WriteToString<96>(Request.Key), *Request.Name);
 			TRACE_COUNTER_INCREMENT(FileSystemDDC_GetHit);
 			TRACE_COUNTER_ADD(FileSystemDDC_BytesRead, MeasureCompressedCacheRecord(Record.Get()));
 			COOK_STAT(Timer.AddHit(MeasureRawCacheRecord(Record.Get())));
 			if (OnComplete)
 			{
-				OnComplete({MoveTemp(Record).Get(), Request.UserData, Status});
+				OnComplete({Request.Name, MoveTemp(Record).Get(), Request.UserData, Status});
 			}
 		}
 		else
 		{
 			if (OnComplete)
 			{
-				OnComplete({FCacheRecordBuilder(Request.Key).Build(), Request.UserData, Status});
+				OnComplete({Request.Name, FCacheRecordBuilder(Request.Key).Build(), Request.UserData, Status});
 			}
 		}
 	}
 }
 
 void FFileSystemCacheStore::GetChunks(
-	const TConstArrayView<FCacheChunkRequest> Chunks,
-	const FStringView Context,
+	const TConstArrayView<FCacheChunkRequest> Requests,
 	IRequestOwner& Owner,
-	FOnCacheGetChunkComplete&& OnComplete)
+	FOnCacheChunkComplete&& OnComplete)
 {
-	TArray<FCacheChunkRequest, TInlineAllocator<16>> SortedChunks(Chunks);
-	SortedChunks.StableSort(TChunkLess());
+	TArray<FCacheChunkRequest, TInlineAllocator<16>> SortedRequests(Requests);
+	SortedRequests.StableSort(TChunkLess());
 
 	FPayload Payload;
 	FCacheKey PayloadKey;
@@ -1617,40 +1611,40 @@ void FFileSystemCacheStore::GetChunks(
 	FCompressedBufferReader PayloadReader;
 	EStatus PayloadStatus = EStatus::Error;
 	FOptionalCacheRecord Record;
-	for (const FCacheChunkRequest& Chunk : SortedChunks)
+	for (const FCacheChunkRequest& Request : SortedRequests)
 	{
 		constexpr ECachePolicy SkipFlag = ECachePolicy::SkipValue;
-		const bool bExistsOnly = EnumHasAnyFlags(Chunk.Policy, SkipFlag);
+		const bool bExistsOnly = EnumHasAnyFlags(Request.Policy, SkipFlag);
 		TRACE_CPUPROFILER_EVENT_SCOPE(FileSystemDDC_Get);
 		TRACE_COUNTER_INCREMENT(FileSystemDDC_Get);
 		COOK_STAT(auto Timer = bExistsOnly ? UsageStats.TimeProbablyExists() : UsageStats.TimeGet());
-		if (!(Payload && PayloadKey == Chunk.Key && Payload.GetId() == Chunk.Id) || !bExistsOnly <= PayloadReader.HasSource())
+		if (!(Payload && PayloadKey == Request.Key && Payload.GetId() == Request.Id) || !bExistsOnly <= PayloadReader.HasSource())
 		{
 			PayloadStatus = EStatus::Error;
 			PayloadReader.ResetSource();
 			PayloadAr.Reset();
 			Payload.Reset();
-			if (!(Record && Record.Get().GetKey() == Chunk.Key))
+			if (!(Record && Record.Get().GetKey() == Request.Key))
 			{
 				FCacheRecordPolicyBuilder PolicyBuilder(ECachePolicy::None);
 				const ECachePolicy RecordSkipFlags = bExistsOnly ? ECachePolicy::SkipData : ECachePolicy::None;
-				PolicyBuilder.AddPayloadPolicy(Chunk.Id, Chunk.Policy | RecordSkipFlags);
+				PolicyBuilder.AddPayloadPolicy(Request.Id, Request.Policy | RecordSkipFlags);
 				Record.Reset();
-				Record = GetCacheRecordOnly(Chunk.Key, Context, PolicyBuilder.Build());
+				Record = GetCacheRecordOnly(Request.Key, Request.Name, PolicyBuilder.Build());
 			}
 			if (Record)
 			{
-				Payload = Record.Get().GetPayload(Chunk.Id);
-				PayloadKey = Chunk.Key;
-				GetCacheContent(Chunk.Key, Payload, Context, Chunk.Policy, SkipFlag, PayloadReader, PayloadAr);
+				Payload = Record.Get().GetPayload(Request.Id);
+				PayloadKey = Request.Key;
+				GetCacheContent(Request.Key, Payload, Request.Name, Request.Policy, SkipFlag, PayloadReader, PayloadAr);
 			}
 		}
 		if (Payload)
 		{
-			const uint64 RawOffset = FMath::Min(Payload.GetRawSize(), Chunk.RawOffset);
-			const uint64 RawSize = FMath::Min(Payload.GetRawSize() - RawOffset, Chunk.RawSize);
-			UE_LOG(LogDerivedDataCache, Verbose, TEXT("%s: Cache hit for %s from '%.*s'"),
-				*CachePath, *WriteToString<96>(Chunk.Key, '/', Chunk.Id), Context.Len(), Context.GetData());
+			const uint64 RawOffset = FMath::Min(Payload.GetRawSize(), Request.RawOffset);
+			const uint64 RawSize = FMath::Min(Payload.GetRawSize() - RawOffset, Request.RawSize);
+			UE_LOG(LogDerivedDataCache, Verbose, TEXT("%s: Cache hit for %s from '%s'"),
+				*CachePath, *WriteToString<96>(Request.Key, '/', Request.Id), *Request.Name);
 			TRACE_COUNTER_INCREMENT(FileSystemDDC_GetHit);
 			TRACE_COUNTER_ADD(FileSystemDDC_BytesRead, !bExistsOnly ? RawSize : 0);
 			COOK_STAT(Timer.AddHit(!bExistsOnly ? RawSize : 0));
@@ -1662,15 +1656,15 @@ void FFileSystemCacheStore::GetChunks(
 					Buffer = PayloadReader.Decompress(RawOffset, RawSize);
 				}
 				const EStatus ChunkStatus = bExistsOnly || Buffer.GetSize() == RawSize ? EStatus::Ok : EStatus::Error;
-				OnComplete({Chunk.Key, Chunk.Id, Chunk.RawOffset,
-					RawSize, Payload.GetRawHash(), MoveTemp(Buffer), Chunk.UserData, ChunkStatus});
+				OnComplete({Request.Name, Request.Key, Request.Id, Request.RawOffset,
+					RawSize, Payload.GetRawHash(), MoveTemp(Buffer), Request.UserData, ChunkStatus});
 			}
 			continue;
 		}
 
 		if (OnComplete)
 		{
-			OnComplete({Chunk.Key, Chunk.Id, Chunk.RawOffset, 0, {}, {}, Chunk.UserData, EStatus::Error});
+			OnComplete({Request.Name, Request.Key, Request.Id, Request.RawOffset, 0, {}, {}, Request.UserData, EStatus::Error});
 		}
 	}
 }
@@ -1693,7 +1687,7 @@ uint64 FFileSystemCacheStore::MeasureRawCacheRecord(const FCacheRecord& Record)
 
 bool FFileSystemCacheStore::PutCacheRecord(
 	const FCacheRecord& Record,
-	const FStringView Context,
+	const FStringView RecordName,
 	const FCacheRecordPolicy& Policy,
 	uint64& OutWriteSize)
 {
@@ -1701,7 +1695,7 @@ bool FFileSystemCacheStore::PutCacheRecord(
 	{
 		UE_LOG(LogDerivedDataCache, VeryVerbose,
 			TEXT("%s: Skipped put of %s from '%.*s' because this cache store is read-only"),
-			*CachePath, *WriteToString<96>(Record.GetKey()), Context.Len(), Context.GetData());
+			*CachePath, *WriteToString<96>(Record.GetKey()), RecordName.Len(), RecordName.GetData());
 		return false;
 	}
 
@@ -1715,14 +1709,14 @@ bool FFileSystemCacheStore::PutCacheRecord(
 	if (!EnumHasAnyFlags(CombinedPolicy, StoreFlag))
 	{
 		UE_LOG(LogDerivedDataCache, VeryVerbose, TEXT("%s: Skipped put of %s from '%.*s' due to cache policy"),
-			*CachePath, *WriteToString<96>(Key), Context.Len(), Context.GetData());
+			*CachePath, *WriteToString<96>(Key), RecordName.Len(), RecordName.GetData());
 		return false;
 	}
 
 	if (ShouldSimulateMiss(Key))
 	{
 		UE_LOG(LogDerivedDataCache, Verbose, TEXT("%s: Simulated miss for put of %s from '%.*s'"),
-			*CachePath, *WriteToString<96>(Key), Context.Len(), Context.GetData());
+			*CachePath, *WriteToString<96>(Key), RecordName.Len(), RecordName.GetData());
 		return false;
 	}
 
@@ -1737,7 +1731,7 @@ bool FFileSystemCacheStore::PutCacheRecord(
 	}
 	else
 	{
-		bRecordExists = LoadFileWithHash(Path, Context, [&ExistingPackage](FArchive& Ar) { ExistingPackage.TryLoad(Ar); });
+		bRecordExists = LoadFileWithHash(Path, RecordName, [&ExistingPackage](FArchive& Ar) { ExistingPackage.TryLoad(Ar); });
 	}
 
 	// Save the record to a package and remove attachments that will be stored externally.
@@ -1786,13 +1780,13 @@ bool FFileSystemCacheStore::PutCacheRecord(
 	for (FCompressedBuffer& Content : ExternalContent)
 	{
 		uint64 WriteSize = 0;
-		PutCacheContent(Content, Context, OutWriteSize);
+		PutCacheContent(Content, RecordName, OutWriteSize);
 		OutWriteSize += WriteSize;
 	}
 
 	// Save the record package to storage.
 	const auto WriteRecord = [&](FArchive& Ar) { Package.Save(Ar); OutWriteSize += uint64(Ar.TotalSize()); };
-	if (!bRecordExists && !SaveFileWithHash(Path, Context, WriteRecord))
+	if (!bRecordExists && !SaveFileWithHash(Path, RecordName, WriteRecord))
 	{
 		return false;
 	}
@@ -1807,14 +1801,14 @@ bool FFileSystemCacheStore::PutCacheRecord(
 
 FOptionalCacheRecord FFileSystemCacheStore::GetCacheRecordOnly(
 	const FCacheKey& Key,
-	const FStringView Context,
+	const FStringView RecordName,
 	const FCacheRecordPolicy& Policy)
 {
 	if (!IsUsable())
 	{
 		UE_LOG(LogDerivedDataCache, VeryVerbose,
 			TEXT("%s: Skipped get of %s from '%.*s' because this cache store is not available"),
-			*CachePath, *WriteToString<96>(Key), Context.Len(), Context.GetData());
+			*CachePath, *WriteToString<96>(Key), RecordName.Len(), RecordName.GetData());
 		return FOptionalCacheRecord();
 	}
 
@@ -1824,14 +1818,14 @@ FOptionalCacheRecord FFileSystemCacheStore::GetCacheRecordOnly(
 	if (!EnumHasAnyFlags(Policy.GetRecordPolicy(), QueryPolicy))
 	{
 		UE_LOG(LogDerivedDataCache, VeryVerbose, TEXT("%s: Skipped get of %s from '%.*s' due to cache policy"),
-			*CachePath, *WriteToString<96>(Key), Context.Len(), Context.GetData());
+			*CachePath, *WriteToString<96>(Key), RecordName.Len(), RecordName.GetData());
 		return FOptionalCacheRecord();
 	}
 
 	if (ShouldSimulateMiss(Key))
 	{
 		UE_LOG(LogDerivedDataCache, Verbose, TEXT("%s: Simulated miss for get of %s from '%.*s'"),
-			*CachePath, *WriteToString<96>(Key), Context.Len(), Context.GetData());
+			*CachePath, *WriteToString<96>(Key), RecordName.Len(), RecordName.GetData());
 		return FOptionalCacheRecord();
 	}
 
@@ -1850,17 +1844,17 @@ FOptionalCacheRecord FFileSystemCacheStore::GetCacheRecordOnly(
 	FOptionalCacheRecord Record;
 	{
 		FCbPackage Package;
-		if (!LoadFileWithHash(Path, Context, [&Package](FArchive& Ar) { Package.TryLoad(Ar); }))
+		if (!LoadFileWithHash(Path, RecordName, [&Package](FArchive& Ar) { Package.TryLoad(Ar); }))
 		{
 			UE_LOG(LogDerivedDataCache, Verbose, TEXT("%s: Cache miss with missing record for %s from '%.*s'"),
-				*CachePath, *WriteToString<96>(Key), Context.Len(), Context.GetData());
+				*CachePath, *WriteToString<96>(Key), RecordName.Len(), RecordName.GetData());
 			return Record;
 		}
 
 		if (ValidateCompactBinary(Package, ECbValidateMode::Default | ECbValidateMode::Package) != ECbValidateError::None)
 		{
 			UE_LOG(LogDerivedDataCache, Display, TEXT("%s: Cache miss with invalid package for %s from '%.*s'"),
-				*CachePath, *WriteToString<96>(Key), Context.Len(), Context.GetData());
+				*CachePath, *WriteToString<96>(Key), RecordName.Len(), RecordName.GetData());
 			return FOptionalCacheRecord();
 		}
 
@@ -1868,7 +1862,7 @@ FOptionalCacheRecord FFileSystemCacheStore::GetCacheRecordOnly(
 		if (Record.IsNull())
 		{
 			UE_LOG(LogDerivedDataCache, Display, TEXT("%s: Cache miss with record load failure for %s from '%.*s'"),
-				*CachePath, *WriteToString<96>(Key), Context.Len(), Context.GetData());
+				*CachePath, *WriteToString<96>(Key), RecordName.Len(), RecordName.GetData());
 			return Record;
 		}
 	}
@@ -1884,11 +1878,11 @@ FOptionalCacheRecord FFileSystemCacheStore::GetCacheRecordOnly(
 
 FOptionalCacheRecord FFileSystemCacheStore::GetCacheRecord(
 	const FCacheKey& Key,
-	const FStringView Context,
+	const FStringView RecordName,
 	const FCacheRecordPolicy& Policy,
 	EStatus& OutStatus)
 {
-	FOptionalCacheRecord Record = GetCacheRecordOnly(Key, Context, Policy);
+	FOptionalCacheRecord Record = GetCacheRecordOnly(Key, RecordName, Policy);
 	if (Record.IsNull())
 	{
 		OutStatus = EStatus::Error;
@@ -1907,7 +1901,7 @@ FOptionalCacheRecord FFileSystemCacheStore::GetCacheRecord(
 	if (FPayload Payload = Record.Get().GetValuePayload())
 	{
 		const ECachePolicy PayloadPolicy = Policy.GetPayloadPolicy(Payload.GetId());
-		if (FPayload Content = GetCacheContent(Key, Payload, Context, PayloadPolicy, ECachePolicy::SkipValue))
+		if (FPayload Content = GetCacheContent(Key, Payload, RecordName, PayloadPolicy, ECachePolicy::SkipValue))
 		{
 			RecordBuilder.SetValue(MoveTemp(Content));
 		}
@@ -1926,7 +1920,7 @@ FOptionalCacheRecord FFileSystemCacheStore::GetCacheRecord(
 	for (FPayload Payload : Record.Get().GetAttachmentPayloads())
 	{
 		const ECachePolicy PayloadPolicy = Policy.GetPayloadPolicy(Payload.GetId());
-		if (FPayload Content = GetCacheContent(Key, Payload, Context, PayloadPolicy, ECachePolicy::SkipAttachments))
+		if (FPayload Content = GetCacheContent(Key, Payload, RecordName, PayloadPolicy, ECachePolicy::SkipAttachments))
 		{
 			RecordBuilder.AddAttachment(MoveTemp(Content));
 		}
@@ -1947,14 +1941,14 @@ FOptionalCacheRecord FFileSystemCacheStore::GetCacheRecord(
 
 bool FFileSystemCacheStore::PutCacheContent(
 	const FCompressedBuffer& Content,
-	const FStringView Context,
+	const FStringView ContentName,
 	uint64& OutWriteSize) const
 {
 	const FIoHash& RawHash = Content.GetRawHash();
 	TStringBuilder<256> Path;
 	BuildCacheContentPath(RawHash, Path);
 	const auto WriteContent = [&](FArchive& Ar) { Content.Save(Ar); OutWriteSize += uint64(Ar.TotalSize()); };
-	if (!FileExists(Path) && !SaveFileWithHash(Path, Context, WriteContent))
+	if (!FileExists(Path) && !SaveFileWithHash(Path, ContentName, WriteContent))
 	{
 		return false;
 	}
@@ -1968,7 +1962,7 @@ bool FFileSystemCacheStore::PutCacheContent(
 FPayload FFileSystemCacheStore::GetCacheContent(
 	const FCacheKey& Key,
 	const FPayload& Payload,
-	const FStringView Context,
+	const FStringView ContentName,
 	const ECachePolicy Policy,
 	const ECachePolicy SkipFlag) const
 {
@@ -2000,7 +1994,7 @@ FPayload FFileSystemCacheStore::GetCacheContent(
 	else
 	{
 		FCompressedBuffer CompressedBuffer;
-		if (LoadFileWithHash(Path, Context, [&CompressedBuffer](FArchive& Ar) { CompressedBuffer = FCompressedBuffer::Load(Ar); }))
+		if (LoadFileWithHash(Path, ContentName, [&CompressedBuffer](FArchive& Ar) { CompressedBuffer = FCompressedBuffer::Load(Ar); }))
 		{
 			if (CompressedBuffer.GetRawHash() == RawHash)
 			{
@@ -2013,7 +2007,7 @@ FPayload FFileSystemCacheStore::GetCacheContent(
 			UE_LOG(LogDerivedDataCache, Display,
 				TEXT("%s: Cache miss with corrupted payload %s with hash %s for %s from '%.*s'"),
 				*CachePath, *WriteToString<16>(Payload.GetId()), *WriteToString<48>(RawHash),
-				*WriteToString<96>(Key), Context.Len(), Context.GetData());
+				*WriteToString<96>(Key), ContentName.Len(), ContentName.GetData());
 			return FPayload::Null;
 		}
 	}
@@ -2021,14 +2015,14 @@ FPayload FFileSystemCacheStore::GetCacheContent(
 	UE_LOG(LogDerivedDataCache, Verbose,
 		TEXT("%s: Cache miss with missing payload %s with hash %s for %s from '%.*s'"),
 		*CachePath, *WriteToString<16>(Payload.GetId()), *WriteToString<48>(RawHash), *WriteToString<96>(Key),
-		Context.Len(), Context.GetData());
+		ContentName.Len(), ContentName.GetData());
 	return FPayload::Null;
 }
 
 void FFileSystemCacheStore::GetCacheContent(
 	const FCacheKey& Key,
 	const FPayload& Payload,
-	const FStringView Context,
+	const FStringView ContentName,
 	const ECachePolicy Policy,
 	const ECachePolicy SkipFlag,
 	FCompressedBufferReader& Reader,
@@ -2066,7 +2060,7 @@ void FFileSystemCacheStore::GetCacheContent(
 	}
 	else
 	{
-		OutArchive = OpenFile(Path, Context);
+		OutArchive = OpenFile(Path, ContentName);
 		if (OutArchive)
 		{
 			Reader.SetSource(*OutArchive);
@@ -2081,7 +2075,7 @@ void FFileSystemCacheStore::GetCacheContent(
 			UE_LOG(LogDerivedDataCache, Display,
 				TEXT("%s: Cache miss with corrupted payload %s with hash %s for %s from '%.*s'"),
 				*CachePath, *WriteToString<16>(Payload.GetId()), *WriteToString<48>(RawHash),
-				*WriteToString<96>(Key), Context.Len(), Context.GetData());
+				*WriteToString<96>(Key), ContentName.Len(), ContentName.GetData());
 			Reader.ResetSource();
 			OutArchive.Reset();
 			return;
@@ -2091,7 +2085,7 @@ void FFileSystemCacheStore::GetCacheContent(
 	UE_LOG(LogDerivedDataCache, Verbose,
 		TEXT("%s: Cache miss with missing payload %s with hash %s for %s from '%.*s'"),
 		*CachePath, *WriteToString<16>(Payload.GetId()), *WriteToString<48>(RawHash), *WriteToString<96>(Key),
-		Context.Len(), Context.GetData());
+		ContentName.Len(), ContentName.GetData());
 }
 
 void FFileSystemCacheStore::BuildCacheRecordPath(const FCacheKey& CacheKey, FStringBuilderBase& Path) const
@@ -2108,10 +2102,10 @@ void FFileSystemCacheStore::BuildCacheContentPath(const FIoHash& RawHash, FStrin
 
 bool FFileSystemCacheStore::SaveFileWithHash(
 	FStringBuilderBase& Path,
-	const FStringView Context,
+	const FStringView DebugName,
 	const TFunctionRef<void (FArchive&)> WriteFunction) const
 {
-	return SaveFile(Path, Context, [&WriteFunction](FArchive& Ar)
+	return SaveFile(Path, DebugName, [&WriteFunction](FArchive& Ar)
 	{
 		THashingArchiveProxy<FBlake3> HashAr(Ar);
 		WriteFunction(HashAr);
@@ -2122,10 +2116,10 @@ bool FFileSystemCacheStore::SaveFileWithHash(
 
 bool FFileSystemCacheStore::LoadFileWithHash(
 	FStringBuilderBase& Path,
-	const FStringView Context,
+	const FStringView DebugName,
 	const TFunctionRef<void (FArchive& Ar)> ReadFunction) const
 {
-	return LoadFile(Path, Context, [this, &Path, &Context, &ReadFunction](FArchive& Ar)
+	return LoadFile(Path, DebugName, [this, &Path, &DebugName, &ReadFunction](FArchive& Ar)
 	{
 		THashingArchiveProxy<FBlake3> HashAr(Ar);
 		ReadFunction(HashAr);
@@ -2138,7 +2132,7 @@ bool FFileSystemCacheStore::LoadFileWithHash(
 			Ar.SetError();
 			UE_LOG(LogDerivedDataCache, Display,
 				TEXT("%s: File %s from '%.*s' is corrupted and has hash %s when %s is expected."),
-				*CachePath, *Path, Context.Len(), Context.GetData(),
+				*CachePath, *Path, DebugName.Len(), DebugName.GetData(),
 				*WriteToString<80>(Hash), *WriteToString<80>(SavedHash));
 		}
 	});
@@ -2146,7 +2140,7 @@ bool FFileSystemCacheStore::LoadFileWithHash(
 
 bool FFileSystemCacheStore::SaveFile(
 	FStringBuilderBase& Path,
-	const FStringView Context,
+	const FStringView DebugName,
 	const TFunctionRef<void (FArchive&)> WriteFunction) const
 {
 	TStringBuilder<256> TempPath;
@@ -2171,7 +2165,7 @@ bool FFileSystemCacheStore::SaveFile(
 	{
 		UE_LOG(LogDerivedDataCache, Warning,
 			TEXT("%s: Failed to write to temp file %s when saving %s from '%.*s'. Error 0x%08x"),
-			*CachePath, *TempPath, *Path, Context.Len(), Context.GetData(), FPlatformMisc::GetLastError());
+			*CachePath, *TempPath, *Path, DebugName.Len(), DebugName.GetData(), FPlatformMisc::GetLastError());
 		return false;
 	}
 
@@ -2180,7 +2174,7 @@ bool FFileSystemCacheStore::SaveFile(
 		UE_LOG(LogDerivedDataCache, Warning,
 			TEXT("%s: Failed to write to temp file %s when saving %s from '%.*s'. ")
 			TEXT("File is %" INT64_FMT " bytes when %" INT64_FMT " bytes are expected."),
-			*CachePath, *TempPath, *Path, Context.Len(), Context.GetData(),
+			*CachePath, *TempPath, *Path, DebugName.Len(), DebugName.GetData(),
 			IFileManager::Get().FileSize(*TempPath), WriteSize);
 		return false;
 	}
@@ -2190,7 +2184,7 @@ bool FFileSystemCacheStore::SaveFile(
 	{
 		UE_LOG(LogDerivedDataCache, Log,
 			TEXT("%s: Move collision when writing file %s from '%.*s'."),
-			*CachePath, *Path, Context.Len(), Context.GetData());
+			*CachePath, *Path, DebugName.Len(), DebugName.GetData());
 	}
 
 	return true;
@@ -2198,7 +2192,7 @@ bool FFileSystemCacheStore::SaveFile(
 
 bool FFileSystemCacheStore::LoadFile(
 	FStringBuilderBase& Path,
-	const FStringView Context,
+	const FStringView DebugName,
 	const TFunctionRef<void (FArchive& Ar)> ReadFunction) const
 {
 	check(IsUsable());
@@ -2207,7 +2201,7 @@ bool FFileSystemCacheStore::LoadFile(
 	int64 ReadSize = 0;
 	bool bError = false;
 
-	if (TUniquePtr<FArchive> Ar = OpenFile(Path, Context))
+	if (TUniquePtr<FArchive> Ar = OpenFile(Path, DebugName))
 	{
 		ReadFunction(*Ar);
 		ReadSize = Ar->Tell();
@@ -2217,7 +2211,7 @@ bool FFileSystemCacheStore::LoadFile(
 		{
 			UE_LOG(LogDerivedDataCache, Display,
 				TEXT("%s: Failed to load file %s from '%.*s'."),
-				*CachePath, *Path, Context.Len(), Context.GetData());
+				*CachePath, *Path, DebugName.Len(), DebugName.GetData());
 		}
 	}
 
@@ -2229,12 +2223,12 @@ bool FFileSystemCacheStore::LoadFile(
 		// Slower than 0.5 MiB/s?
 		UE_CLOG(ReadSpeed < 0.5, LogDerivedDataCache, Warning,
 			TEXT("%s: Loading %s from '%.*s' is very slow (%.2f MiB/s); consider disabling this cache backend"),
-			*CachePath, *Path, Context.Len(), Context.GetData(), ReadSpeed);
+			*CachePath, *Path, DebugName.Len(), DebugName.GetData(), ReadSpeed);
 	}
 
 	UE_LOG(LogDerivedDataCache, VeryVerbose,
 		TEXT("%s: Loaded %s from '%.*s' (%" INT64_FMT " bytes, %.02f secs, %.2f MiB/s)"),
-		*CachePath, *Path, Context.Len(), Context.GetData(), ReadSize, ReadDuration, ReadSpeed);
+		*CachePath, *Path, DebugName.Len(), DebugName.GetData(), ReadSize, ReadDuration, ReadSpeed);
 
 	if (bError && !bReadOnly)
 	{
@@ -2244,7 +2238,7 @@ bool FFileSystemCacheStore::LoadFile(
 	return !bError && ReadSize > 0;
 }
 
-TUniquePtr<FArchive> FFileSystemCacheStore::OpenFile(FStringBuilderBase& Path, const FStringView Context) const
+TUniquePtr<FArchive> FFileSystemCacheStore::OpenFile(FStringBuilderBase& Path, const FStringView DebugName) const
 {
 	// Check for existence before reading because it may update the modification time and avoid the
 	// file being deleted by a cache cleanup thread or process.
diff --git a/Engine/Source/Developer/DerivedDataCache/Private/HierarchicalDerivedDataBackend.h b/Engine/Source/Developer/DerivedDataCache/Private/HierarchicalDerivedDataBackend.h
--- a/Engine/Source/Developer/DerivedDataCache/Private/HierarchicalDerivedDataBackend.h
+++ b/Engine/Source/Developer/DerivedDataCache/Private/HierarchicalDerivedDataBackend.h
@@ -474,7 +474,6 @@ public:
 
 	virtual void Put(
 		TConstArrayView<FCachePutRequest> Requests,
-		FStringView Context,
 		IRequestOwner& Owner,
 		FOnCachePutComplete&& OnComplete) override
 	{
@@ -495,7 +494,7 @@ public:
 					if (RequestsOk.Num() < Requests.Num())
 					{
 						FRequestOwner BlockingOwner(EPriority::Blocking);
-						InnerBackends[PutCacheIndex]->Put(Requests, Context, BlockingOwner,
+						InnerBackends[PutCacheIndex]->Put(Requests, BlockingOwner,
 							[&OnComplete, &RequestsOk](FCachePutCompleteParams&& Params)
 							{
 								if (Params.Status == EStatus::Ok)
@@ -512,7 +511,7 @@ public:
 					}
 					else
 					{
-						AsyncPutInnerBackends[PutCacheIndex]->Put(Requests, Context, AsyncOwner);
+						AsyncPutInnerBackends[PutCacheIndex]->Put(Requests, AsyncOwner);
 					}
 				}
 			}
@@ -524,7 +523,7 @@ public:
 			{
 				if (!RequestsOk.Contains(Request.Record.GetKey()))
 				{
-					OnComplete({Request.Record.GetKey(), Request.UserData, EStatus::Error});
+					OnComplete({Request.Name, Request.Record.GetKey(), Request.UserData, EStatus::Error});
 				}
 			}
 		}
@@ -532,7 +531,6 @@ public:
 
 	virtual void Get(
 		TConstArrayView<FCacheGetRequest> Requests,
-		FStringView Context,
 		IRequestOwner& Owner,
 		FOnCacheGetComplete&& OnComplete) override
 	{
@@ -548,8 +546,8 @@ public:
 			{
 				// Block on this because backends in this hierarchy are not expected to be asynchronous.
 				FRequestOwner BlockingOwner(EPriority::Blocking);
-				InnerBackends[GetCacheIndex]->Get(RemainingRequests, Context, BlockingOwner,
-					[this, Context, GetCacheIndex, &Owner, &OnComplete, &KeysOk](FCacheGetCompleteParams&& Params)
+				InnerBackends[GetCacheIndex]->Get(RemainingRequests, BlockingOwner,
+					[this, GetCacheIndex, &Owner, &OnComplete, &KeysOk](FCacheGetCompleteParams&& Params)
 					{
 						if (Params.Status == EStatus::Ok)
 						{
@@ -560,7 +558,7 @@ public:
 							{
 								if (GetCacheIndex != FillCacheIndex)
 								{
-									AsyncPutInnerBackends[FillCacheIndex]->Put({{Params.Record, ECachePolicy::Default}}, Context, AsyncOwner);
+									AsyncPutInnerBackends[FillCacheIndex]->Put({{Params.Name, Params.Record, ECachePolicy::Default}}, AsyncOwner);
 								}
 							}
 
@@ -603,41 +601,40 @@ public:
 		{
 			for (const FCacheGetRequest& Request : RemainingRequests)
 			{
-				OnComplete({FCacheRecordBuilder(Request.Key).Build(), Request.UserData, EStatus::Error});
+				OnComplete({Request.Name, FCacheRecordBuilder(Request.Key).Build(), Request.UserData, EStatus::Error});
 			}
 		}
 	}
 
 	virtual void GetChunks(
-		TConstArrayView<FCacheChunkRequest> Chunks,
-		FStringView Context,
+		TConstArrayView<FCacheChunkRequest> Requests,
 		IRequestOwner& Owner,
-		FOnCacheGetChunkComplete&& OnComplete) override
+		FOnCacheChunkComplete&& OnComplete) override
 	{
-		TArray<FCacheChunkRequest, TInlineAllocator<16>> RemainingChunks(Chunks);
+		TArray<FCacheChunkRequest, TInlineAllocator<16>> RemainingRequests(Requests);
 
 		{
 			FReadScopeLock LockScope(Lock);
 
 			for (FDerivedDataBackendInterface* InnerBackend : InnerBackends)
 			{
-				if (RemainingChunks.IsEmpty())
+				if (RemainingRequests.IsEmpty())
 				{
 					return;
 				}
-				RemainingChunks.StableSort(TChunkLess());
+				RemainingRequests.StableSort(TChunkLess());
 				TArray<FCacheChunkRequest, TInlineAllocator<16>> ErrorChunks;
 				FRequestOwner BlockingOwner(EPriority::Blocking);
-				InnerBackend->GetChunks(RemainingChunks, Context, BlockingOwner,
-					[&OnComplete, &RemainingChunks, &ErrorChunks](FCacheGetChunkCompleteParams&& Params)
+				InnerBackend->GetChunks(RemainingRequests, BlockingOwner,
+					[&OnComplete, &RemainingRequests, &ErrorChunks](FCacheChunkCompleteParams&& Params)
 					{
 						if (Params.Status == EStatus::Error)
 						{
-							const int32 ChunkIndex = Algo::BinarySearch(RemainingChunks, Params, TChunkLess());
+							const int32 ChunkIndex = Algo::BinarySearch(RemainingRequests, Params, TChunkLess());
 							checkf(ChunkIndex != INDEX_NONE, TEXT("Failed to find remaining chunk %s ")
 								TEXT(" with raw offset %") UINT64_FMT TEXT("."),
 								*WriteToString<96>(Params.Key, '/', Params.Id), Params.RawOffset);
-							ErrorChunks.Add(RemainingChunks[ChunkIndex]);
+							ErrorChunks.Add(RemainingRequests[ChunkIndex]);
 						}
 						else if (OnComplete)
 						{
@@ -645,15 +642,15 @@ public:
 						}
 					});
 				BlockingOwner.Wait();
-				RemainingChunks = MoveTemp(ErrorChunks);
+				RemainingRequests = MoveTemp(ErrorChunks);
 			}
 		}
 
 		if (OnComplete)
 		{
-			for (const FCacheChunkRequest& Chunk : RemainingChunks)
+			for (const FCacheChunkRequest& Request : RemainingRequests)
 			{
-				OnComplete({Chunk.Key, Chunk.Id, Chunk.RawOffset, 0, {}, {}, Chunk.UserData, EStatus::Error});
+				OnComplete({Request.Name, Request.Key, Request.Id, Request.RawOffset, 0, {}, {}, Request.UserData, EStatus::Error});
 			}
 		}
 	}
diff --git a/Engine/Source/Developer/DerivedDataCache/Private/HttpDerivedDataBackend.cpp b/Engine/Source/Developer/DerivedDataCache/Private/HttpDerivedDataBackend.cpp
--- a/Engine/Source/Developer/DerivedDataCache/Private/HttpDerivedDataBackend.cpp
+++ b/Engine/Source/Developer/DerivedDataCache/Private/HttpDerivedDataBackend.cpp
@@ -2528,7 +2528,6 @@ TSharedRef<FDerivedDataCacheStatsNode> FHttpDerivedDataBackend::GatherUsageStats
 
 void FHttpDerivedDataBackend::Put(
 	const TConstArrayView<FCachePutRequest> Requests,
-	const FStringView Context,
 	IRequestOwner& Owner,
 	FOnCachePutComplete&& OnComplete)
 {
@@ -2536,14 +2535,13 @@ void FHttpDerivedDataBackend::Put(
 	{
 		for (const FCachePutRequest& Request : Requests)
 		{
-			OnComplete({Request.Record.GetKey(), Request.UserData, EStatus::Error});
+			OnComplete({Request.Name, Request.Record.GetKey(), Request.UserData, EStatus::Error});
 		}
 	}
 }
 
 void FHttpDerivedDataBackend::Get(
 	const TConstArrayView<FCacheGetRequest> Requests,
-	const FStringView Context,
 	IRequestOwner& Owner,
 	FOnCacheGetComplete&& OnComplete)
 {
@@ -2551,22 +2549,21 @@ void FHttpDerivedDataBackend::Get(
 	{
 		for (const FCacheGetRequest& Request : Requests)
 		{
-			OnComplete({FCacheRecordBuilder(Request.Key).Build(), Request.UserData, EStatus::Error});
+			OnComplete({Request.Name, FCacheRecordBuilder(Request.Key).Build(), Request.UserData, EStatus::Error});
 		}
 	}
 }
 
 void FHttpDerivedDataBackend::GetChunks(
-	const TConstArrayView<FCacheChunkRequest> Chunks,
-	const FStringView Context,
+	const TConstArrayView<FCacheChunkRequest> Requests,
 	IRequestOwner& Owner,
-	FOnCacheGetChunkComplete&& OnComplete)
+	FOnCacheChunkComplete&& OnComplete)
 {
 	if (OnComplete)
 	{
-		for (const FCacheChunkRequest& Chunk : Chunks)
+		for (const FCacheChunkRequest& Request : Requests)
 		{
-			OnComplete({Chunk.Key, Chunk.Id, Chunk.RawOffset, 0, {}, {}, Chunk.UserData, EStatus::Error});
+			OnComplete({Request.Name, Request.Key, Request.Id, Request.RawOffset, 0, {}, {}, Request.UserData, EStatus::Error});
 		}
 	}
 }
diff --git a/Engine/Source/Developer/DerivedDataCache/Private/HttpDerivedDataBackend.h b/Engine/Source/Developer/DerivedDataCache/Private/HttpDerivedDataBackend.h
--- a/Engine/Source/Developer/DerivedDataCache/Private/HttpDerivedDataBackend.h
+++ b/Engine/Source/Developer/DerivedDataCache/Private/HttpDerivedDataBackend.h
@@ -75,21 +75,18 @@ public:
 
 	virtual void Put(
 		TConstArrayView<FCachePutRequest> Requests,
-		FStringView Context,
 		IRequestOwner& Owner,
 		FOnCachePutComplete&& OnComplete) override;
 
 	virtual void Get(
 		TConstArrayView<FCacheGetRequest> Requests,
-		FStringView Context,
 		IRequestOwner& Owner,
 		FOnCacheGetComplete&& OnComplete) override;
 
 	virtual void GetChunks(
-		TConstArrayView<FCacheChunkRequest> Chunks,
-		FStringView Context,
+		TConstArrayView<FCacheChunkRequest> Requests,
 		IRequestOwner& Owner,
-		FOnCacheGetChunkComplete&& OnComplete) override;
+		FOnCacheChunkComplete&& OnComplete) override;
 
 	static FHttpDerivedDataBackend* GetAny()
 	{
diff --git a/Engine/Source/Developer/DerivedDataCache/Private/MemoryDerivedDataBackend.cpp b/Engine/Source/Developer/DerivedDataCache/Private/MemoryDerivedDataBackend.cpp
--- a/Engine/Source/Developer/DerivedDataCache/Private/MemoryDerivedDataBackend.cpp
+++ b/Engine/Source/Developer/DerivedDataCache/Private/MemoryDerivedDataBackend.cpp
@@ -468,7 +468,6 @@ int64 FMemoryDerivedDataBackend::CalcSerializedCacheRecordSize(const FCacheRecor
 
 void FMemoryDerivedDataBackend::Put(
 	const TConstArrayView<FCachePutRequest> Requests,
-	const FStringView Context,
 	IRequestOwner& Owner,
 	FOnCachePutComplete&& OnComplete)
 {
@@ -481,14 +480,14 @@ void FMemoryDerivedDataBackend::Put(
 		{
 			if (OnComplete)
 			{
-				OnComplete({Key, Request.UserData, Status});
+				OnComplete({Request.Name, Key, Request.UserData, Status});
 			}
 		};
 
 		if (ShouldSimulateMiss(Key))
 		{
-			UE_LOG(LogDerivedDataCache, Verbose, TEXT("%s: Simulated miss for put of %s from '%.*s'"),
-				*GetName(), *WriteToString<96>(Key), Context.Len(), Context.GetData());
+			UE_LOG(LogDerivedDataCache, Verbose, TEXT("%s: Simulated miss for put of %s from '%s'"),
+				*GetName(), *WriteToString<96>(Key), *Request.Name);
 			continue;
 		}
 
@@ -509,6 +508,7 @@ void FMemoryDerivedDataBackend::Put(
 			}
 			if (const FCacheRecord* Existing = CacheRecords.Find(Key))
 			{
+				CacheRecords.Remove(Key);
 				CurrentCacheSize -= CalcSerializedCacheRecordSize(*Existing);
 				bMaxSizeExceeded = false;
 			}
@@ -542,7 +542,6 @@ void FMemoryDerivedDataBackend::Put(
 
 void FMemoryDerivedDataBackend::Get(
 	const TConstArrayView<FCacheGetRequest> Requests,
-	const FStringView Context,
 	IRequestOwner& Owner,
 	FOnCacheGetComplete&& OnComplete)
 {
@@ -556,8 +555,8 @@ void FMemoryDerivedDataBackend::Get(
 		EStatus Status = EStatus::Error;
 		if (ShouldSimulateMiss(Key))
 		{
-			UE_LOG(LogDerivedDataCache, Verbose, TEXT("%s: Simulated miss for get of %s from '%.*s'"),
-				*GetName(), *WriteToString<96>(Key), Context.Len(), Context.GetData());
+			UE_LOG(LogDerivedDataCache, Verbose, TEXT("%s: Simulated miss for get of %s from '%s'"),
+				*GetName(), *WriteToString<96>(Key), *Request.Name);
 		}
 		else if (FWriteScopeLock ScopeLock(SynchronizationObject); const FCacheRecord* CacheRecord = CacheRecords.Find(Key))
 		{
@@ -596,70 +595,70 @@ void FMemoryDerivedDataBackend::Get(
 			COOK_STAT(Timer.AddHit(CalcRawCacheRecordSize(Record.Get())));
 			if (OnComplete)
 			{
-				OnComplete({MoveTemp(Record).Get(), Request.UserData, Status});
+				OnComplete({Request.Name, MoveTemp(Record).Get(), Request.UserData, Status});
 			}
 		}
 		else
 		{
 			if (OnComplete)
 			{
-				OnComplete({FCacheRecordBuilder(Key).Build(), Request.UserData, EStatus::Error});
+				OnComplete({Request.Name, FCacheRecordBuilder(Key).Build(), Request.UserData, EStatus::Error});
 			}
 		}
 	}
 }
 
 void FMemoryDerivedDataBackend::GetChunks(
-	const TConstArrayView<FCacheChunkRequest> Chunks,
-	const FStringView Context,
+	const TConstArrayView<FCacheChunkRequest> Requests,
 	IRequestOwner& Owner,
-	FOnCacheGetChunkComplete&& OnComplete)
+	FOnCacheChunkComplete&& OnComplete)
 {
 	FPayload Payload;
 	FCacheKey PayloadKey;
 	FCompressedBufferReader Reader;
-	for (const FCacheChunkRequest& Chunk : Chunks)
+	for (const FCacheChunkRequest& Request : Requests)
 	{
-		const bool bExistsOnly = EnumHasAnyFlags(Chunk.Policy, ECachePolicy::SkipValue);
+		const bool bExistsOnly = EnumHasAnyFlags(Request.Policy, ECachePolicy::SkipValue);
 		COOK_STAT(auto Timer = bExistsOnly ? UsageStats.TimeProbablyExists() : UsageStats.TimeGet());
-		if (ShouldSimulateMiss(Chunk.Key))
+		if (ShouldSimulateMiss(Request.Key))
 		{
-			UE_LOG(LogDerivedDataCache, Verbose, TEXT("%s: Simulated miss for get of %s from '%.*s'"),
-				*GetName(), *WriteToString<96>(Chunk.Key, '/', Chunk.Id), Context.Len(), Context.GetData());
+			UE_LOG(LogDerivedDataCache, Verbose, TEXT("%s: Simulated miss for get of %s from '%s'"),
+				*GetName(), *WriteToString<96>(Request.Key, '/', Request.Id), *Request.Name);
 		}
-		else if (PayloadKey == Chunk.Key && Payload.GetId() == Chunk.Id)
+		else if (PayloadKey == Request.Key && Payload.GetId() == Request.Id)
 		{
 			// Payload matches the request.
 		}
-		else if (FWriteScopeLock ScopeLock(SynchronizationObject); const FCacheRecord* Record = CacheRecords.Find(Chunk.Key))
+		else if (FWriteScopeLock ScopeLock(SynchronizationObject); const FCacheRecord* Record = CacheRecords.Find(Request.Key))
 		{
 			Reader.ResetSource();
 			Payload.Reset();
-			Payload = Record->GetAttachmentPayload(Chunk.Id);
-			PayloadKey = Chunk.Key;
+			Payload = Record->GetAttachmentPayload(Request.Id);
+			PayloadKey = Request.Key;
 			Reader.SetSource(Payload.GetData());
 		}
-		if (Payload && Chunk.RawOffset <= Payload.GetRawSize())
+		if (Payload && Request.RawOffset <= Payload.GetRawSize())
 		{
-			const uint64 RawSize = FMath::Min(Payload.GetRawSize() - Chunk.RawOffset, Chunk.RawSize);
+			const uint64 RawSize = FMath::Min(Payload.GetRawSize() - Request.RawOffset, Request.RawSize);
 			COOK_STAT(Timer.AddHit(RawSize));
 			if (OnComplete)
 			{
 				FSharedBuffer Buffer;
 				if (Payload.HasData() && !bExistsOnly)
 				{
-					Buffer = Reader.Decompress(Chunk.RawOffset, RawSize);
+					Buffer = Reader.Decompress(Request.RawOffset, RawSize);
 				}
 				const EStatus Status = bExistsOnly || Buffer ? EStatus::Ok : EStatus::Error;
-				OnComplete({Chunk.Key, Chunk.Id, Chunk.RawOffset,
-					RawSize, Payload.GetRawHash(), MoveTemp(Buffer), Chunk.UserData, Status});
+				OnComplete({Request.Name, Request.Key, Request.Id, Request.RawOffset,
+					RawSize, Payload.GetRawHash(), MoveTemp(Buffer), Request.UserData, Status});
 			}
 		}
 		else
 		{
 			if (OnComplete)
 			{
-				OnComplete({Chunk.Key, Chunk.Id, Chunk.RawOffset, 0, {}, {}, Chunk.UserData, EStatus::Error});
+				OnComplete({Request.Name, Request.Key, Request.Id, Request.RawOffset,
+					0, {}, {}, Request.UserData, EStatus::Error});
 			}
 		}
 	}
diff --git a/Engine/Source/Developer/DerivedDataCache/Private/MemoryDerivedDataBackend.h b/Engine/Source/Developer/DerivedDataCache/Private/MemoryDerivedDataBackend.h
--- a/Engine/Source/Developer/DerivedDataCache/Private/MemoryDerivedDataBackend.h
+++ b/Engine/Source/Developer/DerivedDataCache/Private/MemoryDerivedDataBackend.h
@@ -96,21 +96,18 @@ public:
 
 	virtual void Put(
 		TConstArrayView<FCachePutRequest> Requests,
-		FStringView Context,
 		IRequestOwner& Owner,
 		FOnCachePutComplete&& OnComplete) override;
 
 	virtual void Get(
 		TConstArrayView<FCacheGetRequest> Requests,
-		FStringView Context,
 		IRequestOwner& Owner,
 		FOnCacheGetComplete&& OnComplete) override;
 
 	virtual void GetChunks(
-		TConstArrayView<FCacheChunkRequest> Chunks,
-		FStringView Context,
+		TConstArrayView<FCacheChunkRequest> Requests,
 		IRequestOwner& Owner,
-		FOnCacheGetChunkComplete&& OnComplete) override;
+		FOnCacheChunkComplete&& OnComplete) override;
 
 private:
 	/** Name of the cache file loaded (if any). */
diff --git a/Engine/Source/Developer/DerivedDataCache/Private/PakFileDerivedDataBackend.cpp b/Engine/Source/Developer/DerivedDataCache/Private/PakFileDerivedDataBackend.cpp
--- a/Engine/Source/Developer/DerivedDataCache/Private/PakFileDerivedDataBackend.cpp
+++ b/Engine/Source/Developer/DerivedDataCache/Private/PakFileDerivedDataBackend.cpp
@@ -479,7 +479,6 @@ TSharedRef<FDerivedDataCacheStatsNode> FPakFileDerivedDataBackend::GatherUsageSt
 
 void FPakFileDerivedDataBackend::Put(
 	const TConstArrayView<FCachePutRequest> Requests,
-	const FStringView Context,
 	IRequestOwner& Owner,
 	FOnCachePutComplete&& OnComplete)
 {
@@ -488,14 +487,14 @@ void FPakFileDerivedDataBackend::Put(
 		const FCacheRecord& Record = Request.Record;
 		TRACE_CPUPROFILER_EVENT_SCOPE(PakFileDDC_Put);
 		COOK_STAT(auto Timer = UsageStats.TimePut());
-		if (PutCacheRecord(Record, Context, Request.Policy))
+		if (PutCacheRecord(Record, Request.Name, Request.Policy))
 		{
-			UE_LOG(LogDerivedDataCache, Verbose, TEXT("%s: Cache put complete for %s from '%.*s'"),
-				*CachePath, *WriteToString<96>(Record.GetKey()), Context.Len(), Context.GetData());
+			UE_LOG(LogDerivedDataCache, Verbose, TEXT("%s: Cache put complete for %s from '%s'"),
+				*CachePath, *WriteToString<96>(Record.GetKey()), *Request.Name);
 			COOK_STAT(Timer.AddHit(MeasureRawCacheRecord(Record)));
 			if (OnComplete)
 			{
-				OnComplete({Record.GetKey(), Request.UserData, EStatus::Ok});
+				OnComplete({Request.Name, Record.GetKey(), Request.UserData, EStatus::Ok});
 			}
 		}
 		else
@@ -503,7 +502,7 @@ void FPakFileDerivedDataBackend::Put(
 			COOK_STAT(Timer.AddMiss(MeasureRawCacheRecord(Record)));
 			if (OnComplete)
 			{
-				OnComplete({Record.GetKey(), Request.UserData, EStatus::Error});
+				OnComplete({Request.Name, Record.GetKey(), Request.UserData, EStatus::Error});
 			}
 		}
 	}
@@ -511,7 +510,6 @@ void FPakFileDerivedDataBackend::Put(
 
 void FPakFileDerivedDataBackend::Get(
 	const TConstArrayView<FCacheGetRequest> Requests,
-	const FStringView Context,
 	IRequestOwner& Owner,
 	FOnCacheGetComplete&& OnComplete)
 {
@@ -520,61 +518,60 @@ void FPakFileDerivedDataBackend::Get(
 		TRACE_CPUPROFILER_EVENT_SCOPE(PakFileDDC_Get);
 		COOK_STAT(auto Timer = UsageStats.TimeGet());
 		EStatus Status = EStatus::Ok;
-		if (FOptionalCacheRecord Record = GetCacheRecord(Request.Key, Context, Request.Policy, Status))
+		if (FOptionalCacheRecord Record = GetCacheRecord(Request.Key, Request.Name, Request.Policy, Status))
 		{
-			UE_LOG(LogDerivedDataCache, Verbose, TEXT("%s: Cache hit for %s from '%.*s'"),
-				*CachePath, *WriteToString<96>(Request.Key), Context.Len(), Context.GetData());
+			UE_LOG(LogDerivedDataCache, Verbose, TEXT("%s: Cache hit for %s from '%s'"),
+				*CachePath, *WriteToString<96>(Request.Key), *Request.Name);
 			COOK_STAT(Timer.AddHit(MeasureRawCacheRecord(Record.Get())));
 			if (OnComplete)
 			{
-				OnComplete({MoveTemp(Record).Get(), Request.UserData, Status});
+				OnComplete({Request.Name, MoveTemp(Record).Get(), Request.UserData, Status});
 			}
 		}
 		else
 		{
 			if (OnComplete)
 			{
-				OnComplete({FCacheRecordBuilder(Request.Key).Build(), Request.UserData, Status});
+				OnComplete({Request.Name, FCacheRecordBuilder(Request.Key).Build(), Request.UserData, Status});
 			}
 		}
 	}
 }
 
 void FPakFileDerivedDataBackend::GetChunks(
-	const TConstArrayView<FCacheChunkRequest> Chunks,
-	const FStringView Context,
+	const TConstArrayView<FCacheChunkRequest> Requests,
 	IRequestOwner& Owner,
-	FOnCacheGetChunkComplete&& OnComplete)
+	FOnCacheChunkComplete&& OnComplete)
 {
-	TArray<FCacheChunkRequest, TInlineAllocator<16>> SortedChunks(Chunks);
-	SortedChunks.StableSort(TChunkLess());
+	TArray<FCacheChunkRequest, TInlineAllocator<16>> SortedRequests(Requests);
+	SortedRequests.StableSort(TChunkLess());
 
 	FOptionalCacheRecord Record;
 	FCompressedBufferReader Reader;
-	for (const FCacheChunkRequest& Chunk : SortedChunks)
+	for (const FCacheChunkRequest& Request : SortedRequests)
 	{
 		constexpr ECachePolicy SkipFlag = ECachePolicy::SkipValue;
-		const bool bExistsOnly = EnumHasAnyFlags(Chunk.Policy, SkipFlag);
+		const bool bExistsOnly = EnumHasAnyFlags(Request.Policy, SkipFlag);
 		TRACE_CPUPROFILER_EVENT_SCOPE(PakFileDDC_Get);
 		COOK_STAT(auto Timer = bExistsOnly ? UsageStats.TimeProbablyExists() : UsageStats.TimeGet());
-		if (!Record || Record.Get().GetKey() != Chunk.Key)
+		if (!Record || Record.Get().GetKey() != Request.Key)
 		{
 			FCacheRecordPolicyBuilder PolicyBuilder(ECachePolicy::None);
 			const ECachePolicy RecordSkipFlags = bExistsOnly ? ECachePolicy::SkipData : ECachePolicy::None;
-			PolicyBuilder.AddPayloadPolicy(Chunk.Id, Chunk.Policy | RecordSkipFlags);
-			Record = GetCacheRecordOnly(Chunk.Key, Context, PolicyBuilder.Build());
+			PolicyBuilder.AddPayloadPolicy(Request.Id, Request.Policy | RecordSkipFlags);
+			Record = GetCacheRecordOnly(Request.Key, Request.Name, PolicyBuilder.Build());
 		}
 		if (Record)
 		{
 			EStatus PayloadStatus = EStatus::Ok;
-			FPayload Payload = Record.Get().GetPayload(Chunk.Id);
-			GetCacheContent(Chunk.Key, Context, Chunk.Policy, SkipFlag, Payload, PayloadStatus);
+			FPayload Payload = Record.Get().GetPayload(Request.Id);
+			GetCacheContent(Request.Key, Request.Name, Request.Policy, SkipFlag, Payload, PayloadStatus);
 			if (Payload)
 			{
-				const uint64 RawOffset = FMath::Min(Payload.GetRawSize(), Chunk.RawOffset);
-				const uint64 RawSize = FMath::Min(Payload.GetRawSize() - RawOffset, Chunk.RawSize);
-				UE_LOG(LogDerivedDataCache, Verbose, TEXT("%s: Cache hit for %s from '%.*s'"),
-					*CachePath, *WriteToString<96>(Chunk.Key, '/', Chunk.Id), Context.Len(), Context.GetData());
+				const uint64 RawOffset = FMath::Min(Payload.GetRawSize(), Request.RawOffset);
+				const uint64 RawSize = FMath::Min(Payload.GetRawSize() - RawOffset, Request.RawSize);
+				UE_LOG(LogDerivedDataCache, Verbose, TEXT("%s: Cache hit for %s from '%s'"),
+					*CachePath, *WriteToString<96>(Request.Key, '/', Request.Id), *Request.Name);
 				COOK_STAT(Timer.AddHit(Payload.HasData() ? RawSize : 0));
 				if (OnComplete)
 				{
@@ -584,8 +581,8 @@ void FPakFileDerivedDataBackend::GetChunks(
 						FCompressedBufferReaderSourceScope Source(Reader, Payload.GetData());
 						Buffer = Reader.Decompress(RawOffset, RawSize);
 					}
-					OnComplete({Chunk.Key, Chunk.Id, Chunk.RawOffset,
-						RawSize, Payload.GetRawHash(), MoveTemp(Buffer), Chunk.UserData, PayloadStatus});
+					OnComplete({Request.Name, Request.Key, Request.Id, Request.RawOffset,
+						RawSize, Payload.GetRawHash(), MoveTemp(Buffer), Request.UserData, PayloadStatus});
 				}
 				continue;
 			}
@@ -593,7 +590,8 @@ void FPakFileDerivedDataBackend::GetChunks(
 
 		if (OnComplete)
 		{
-			OnComplete({Chunk.Key, Chunk.Id, Chunk.RawOffset, 0, {}, {}, Chunk.UserData, EStatus::Error});
+			OnComplete({Request.Name, Request.Key, Request.Id, Request.RawOffset,
+				0, {}, {}, Request.UserData, EStatus::Error});
 		}
 	}
 }
@@ -616,7 +614,7 @@ uint64 FPakFileDerivedDataBackend::MeasureRawCacheRecord(const FCacheRecord& Rec
 
 bool FPakFileDerivedDataBackend::PutCacheRecord(
 	const FCacheRecord& Record,
-	const FStringView Context,
+	const FStringView RecordName,
 	const FCacheRecordPolicy& Policy)
 {
 	if (!IsWritable())
@@ -630,14 +628,14 @@ bool FPakFileDerivedDataBackend::PutCacheRecord(
 	if (!EnumHasAnyFlags(Policy.GetRecordPolicy(), ECachePolicy::StoreLocal))
 	{
 		UE_LOG(LogDerivedDataCache, VeryVerbose, TEXT("%s: Skipped put of %s from '%.*s' due to cache policy"),
-			*CachePath, *WriteToString<96>(Key), Context.Len(), Context.GetData());
+			*CachePath, *WriteToString<96>(Key), RecordName.Len(), RecordName.GetData());
 		return false;
 	}
 
 	//if (ShouldSimulateMiss(Key))
 	//{
 	//	UE_LOG(LogDerivedDataCache, Verbose, TEXT("%s: Simulated miss for put of %s from '%.*s'"),
-	//		*CachePath, *WriteToString<96>(Key), Context.Len(), Context.GetData());
+	//		*CachePath, *WriteToString<96>(Key), RecordName.Len(), RecordName.GetData());
 	//	return false;
 	//}
 
@@ -652,7 +650,7 @@ bool FPakFileDerivedDataBackend::PutCacheRecord(
 	{
 		bRecordExists = FileExists(Path);
 	}
-	else if (FSharedBuffer Buffer = LoadFile(Path, Context))
+	else if (FSharedBuffer Buffer = LoadFile(Path, RecordName))
 	{
 		FCbFieldIterator It = FCbFieldIterator::MakeRange(MoveTemp(Buffer));
 		bRecordExists = ExistingPackage.TryLoad(It);
@@ -703,11 +701,11 @@ bool FPakFileDerivedDataBackend::PutCacheRecord(
 	// Save the external content to storage.
 	for (FCompressedBuffer& Content : ExternalContent)
 	{
-		PutCacheContent(Content, Context);
+		PutCacheContent(Content, RecordName);
 	}
 
 	// Save the record package to storage.
-	if (!bRecordExists && !SaveFile(Path, Context, [&Package](FArchive& Ar) { Package.Save(Ar); }))
+	if (!bRecordExists && !SaveFile(Path, RecordName, [&Package](FArchive& Ar) { Package.Save(Ar); }))
 	{
 		return false;
 	}
@@ -717,14 +715,14 @@ bool FPakFileDerivedDataBackend::PutCacheRecord(
 
 FOptionalCacheRecord FPakFileDerivedDataBackend::GetCacheRecordOnly(
 	const FCacheKey& Key,
-	const FStringView Context,
+	const FStringView RecordName,
 	const FCacheRecordPolicy& Policy)
 {
 	if (bClosed)
 	{
 		UE_LOG(LogDerivedDataCache, VeryVerbose,
 			TEXT("%s: Skipped get of %s from '%.*s' because this cache store is not available"),
-			*CachePath, *WriteToString<96>(Key), Context.Len(), Context.GetData());
+			*CachePath, *WriteToString<96>(Key), RecordName.Len(), RecordName.GetData());
 		return FOptionalCacheRecord();
 	}
 
@@ -732,25 +730,25 @@ FOptionalCacheRecord FPakFileDerivedDataBackend::GetCacheRecordOnly(
 	if (!EnumHasAnyFlags(Policy.GetRecordPolicy(), ECachePolicy::QueryLocal))
 	{
 		UE_LOG(LogDerivedDataCache, VeryVerbose, TEXT("%s: Skipped get of %s from '%.*s' due to cache policy"),
-			*CachePath, *WriteToString<96>(Key), Context.Len(), Context.GetData());
+			*CachePath, *WriteToString<96>(Key), RecordName.Len(), RecordName.GetData());
 		return FOptionalCacheRecord();
 	}
 
 	//if (ShouldSimulateMiss(Key))
 	//{
 	//	UE_LOG(LogDerivedDataCache, Verbose, TEXT("%s: Simulated miss for get of %s from '%.*s'"),
-	//		*CachePath, *WriteToString<96>(Key), Context.Len(), Context.GetData());
+	//		*CachePath, *WriteToString<96>(Key), RecordName.Len(), RecordName.GetData());
 	//	return FOptionalCacheRecord();
 	//}
 
 	// Request the record from storage.
 	TStringBuilder<256> Path;
 	FPathViews::Append(Path, TEXT("Buckets"), Key);
-	FSharedBuffer Buffer = LoadFile(Path, Context);
+	FSharedBuffer Buffer = LoadFile(Path, RecordName);
 	if (Buffer.IsNull())
 	{
 		UE_LOG(LogDerivedDataCache, Verbose, TEXT("%s: Cache miss with missing record for %s from '%.*s'"),
-			*CachePath, *WriteToString<96>(Key), Context.Len(), Context.GetData());
+			*CachePath, *WriteToString<96>(Key), RecordName.Len(), RecordName.GetData());
 		return FOptionalCacheRecord();
 	}
 
@@ -758,7 +756,7 @@ FOptionalCacheRecord FPakFileDerivedDataBackend::GetCacheRecordOnly(
 	if (ValidateCompactBinaryPackage(Buffer, ECbValidateMode::Default | ECbValidateMode::Package) != ECbValidateError::None)
 	{
 		UE_LOG(LogDerivedDataCache, Display, TEXT("%s: Cache miss with invalid package for %s from '%.*s'"),
-			*CachePath, *WriteToString<96>(Key), Context.Len(), Context.GetData());
+			*CachePath, *WriteToString<96>(Key), RecordName.Len(), RecordName.GetData());
 		return FOptionalCacheRecord();
 	}
 
@@ -769,14 +767,14 @@ FOptionalCacheRecord FPakFileDerivedDataBackend::GetCacheRecordOnly(
 		if (FCbFieldIterator It = FCbFieldIterator::MakeRange(Buffer); !Package.TryLoad(It))
 		{
 			UE_LOG(LogDerivedDataCache, Display, TEXT("%s: Cache miss with package load failure for %s from '%.*s'"),
-				*CachePath, *WriteToString<96>(Key), Context.Len(), Context.GetData());
+				*CachePath, *WriteToString<96>(Key), RecordName.Len(), RecordName.GetData());
 			return FOptionalCacheRecord();
 		}
 		Record = FCacheRecord::Load(Package);
 		if (Record.IsNull())
 		{
 			UE_LOG(LogDerivedDataCache, Display, TEXT("%s: Cache miss with record load failure for %s from '%.*s'"),
-				*CachePath, *WriteToString<96>(Key), Context.Len(), Context.GetData());
+				*CachePath, *WriteToString<96>(Key), RecordName.Len(), RecordName.GetData());
 			return FOptionalCacheRecord();
 		}
 	}
@@ -786,11 +784,11 @@ FOptionalCacheRecord FPakFileDerivedDataBackend::GetCacheRecordOnly(
 
 FOptionalCacheRecord FPakFileDerivedDataBackend::GetCacheRecord(
 	const FCacheKey& Key,
-	const FStringView Context,
+	const FStringView RecordName,
 	const FCacheRecordPolicy& Policy,
 	EStatus& OutStatus)
 {
-	FOptionalCacheRecord Record = GetCacheRecordOnly(Key, Context, Policy);
+	FOptionalCacheRecord Record = GetCacheRecordOnly(Key, RecordName, Policy);
 	if (Record.IsNull())
 	{
 		OutStatus = EStatus::Error;
@@ -809,7 +807,7 @@ FOptionalCacheRecord FPakFileDerivedDataBackend::GetCacheRecord(
 	if (FPayload Payload = Record.Get().GetValuePayload())
 	{
 		const ECachePolicy PayloadPolicy = Policy.GetPayloadPolicy(Payload.GetId());
-		GetCacheContent(Key, Context, PayloadPolicy, ECachePolicy::SkipValue, Payload, OutStatus);
+		GetCacheContent(Key, RecordName, PayloadPolicy, ECachePolicy::SkipValue, Payload, OutStatus);
 		if (Payload.IsNull())
 		{
 			return FOptionalCacheRecord();
@@ -820,7 +818,7 @@ FOptionalCacheRecord FPakFileDerivedDataBackend::GetCacheRecord(
 	for (FPayload Payload : Record.Get().GetAttachmentPayloads())
 	{
 		const ECachePolicy PayloadPolicy = Policy.GetPayloadPolicy(Payload.GetId());
-		GetCacheContent(Key, Context, PayloadPolicy, ECachePolicy::SkipAttachments, Payload, OutStatus);
+		GetCacheContent(Key, RecordName, PayloadPolicy, ECachePolicy::SkipAttachments, Payload, OutStatus);
 		if (Payload.IsNull())
 		{
 			return FOptionalCacheRecord();
@@ -831,14 +829,14 @@ FOptionalCacheRecord FPakFileDerivedDataBackend::GetCacheRecord(
 	return RecordBuilder.Build();
 }
 
-bool FPakFileDerivedDataBackend::PutCacheContent(const FCompressedBuffer& Content, const FStringView Context)
+bool FPakFileDerivedDataBackend::PutCacheContent(const FCompressedBuffer& Content, const FStringView ContentName)
 {
 	const FIoHash& RawHash = Content.GetRawHash();
 	TStringBuilder<256> Path;
 	FPathViews::Append(Path, TEXT("Content"), RawHash);
 	if (!FileExists(Path))
 	{
-		if (!SaveFile(Path, Context, [&Content](FArchive& Ar) { Content.Save(Ar); }))
+		if (!SaveFile(Path, ContentName, [&Content](FArchive& Ar) { Content.Save(Ar); }))
 		{
 			return false;
 		}
@@ -848,7 +846,7 @@ bool FPakFileDerivedDataBackend::PutCacheContent(const FCompressedBuffer& Conten
 
 void FPakFileDerivedDataBackend::GetCacheContent(
 	const FCacheKey& Key,
-	const FStringView Context,
+	const FStringView ContentName,
 	const ECachePolicy Policy,
 	const ECachePolicy SkipFlag,
 	FPayload& InOutPayload,
@@ -878,7 +876,7 @@ void FPakFileDerivedDataBackend::GetCacheContent(
 	}
 	else
 	{
-		if (FSharedBuffer CompressedData = LoadFile(Path, Context))
+		if (FSharedBuffer CompressedData = LoadFile(Path, ContentName))
 		{
 			if (FCompressedBuffer CompressedBuffer = FCompressedBuffer::FromCompressed(MoveTemp(CompressedData));
 				CompressedBuffer && CompressedBuffer.GetRawHash() == RawHash)
@@ -889,7 +887,7 @@ void FPakFileDerivedDataBackend::GetCacheContent(
 			UE_LOG(LogDerivedDataCache, Display,
 				TEXT("%s: Cache miss with corrupted payload %s with hash %s for %s from '%.*s'"),
 				*CachePath, *WriteToString<16>(InOutPayload.GetId()), *WriteToString<48>(RawHash),
-				*WriteToString<96>(Key), Context.Len(), Context.GetData());
+				*WriteToString<96>(Key), ContentName.Len(), ContentName.GetData());
 			InOutStatus = EStatus::Error;
 			if (!EnumHasAnyFlags(Policy, ECachePolicy::PartialOnError))
 			{
@@ -902,7 +900,7 @@ void FPakFileDerivedDataBackend::GetCacheContent(
 	UE_LOG(LogDerivedDataCache, Verbose,
 		TEXT("%s: Cache miss with missing payload %s with hash %s for %s from '%.*s'"),
 		*CachePath, *WriteToString<16>(InOutPayload.GetId()), *WriteToString<48>(RawHash), *WriteToString<96>(Key),
-		Context.Len(), Context.GetData());
+		ContentName.Len(), ContentName.GetData());
 	InOutStatus = EStatus::Error;
 	if (!EnumHasAnyFlags(Policy, ECachePolicy::PartialOnError))
 	{
@@ -996,7 +994,7 @@ private:
 
 bool FPakFileDerivedDataBackend::SaveFile(
 	const FStringView Path,
-	const FStringView Context,
+	const FStringView DebugName,
 	TFunctionRef<void (FArchive&)> WriteFunction)
 {
 	FWriteScopeLock ScopeLock(SynchronizationObject);
@@ -1011,14 +1009,14 @@ bool FPakFileDerivedDataBackend::SaveFile(
 			FCacheValue& Item = CacheItems.Emplace(Path, FCacheValue(Offset, EndOffset - Offset, HashAr.GetHash()));
 			UE_LOG(LogDerivedDataCache, Log,
 				TEXT("%s: File %.*s from '%.*s' written with offset %" INT64_FMT ", size %" INT64_FMT", CRC 0x%08x."),
-				*CachePath, Path.Len(), Path.GetData(), Context.Len(), Context.GetData(), Item.Offset, Item.Size, Item.Crc);
+				*CachePath, Path.Len(), Path.GetData(), DebugName.Len(), DebugName.GetData(), Item.Offset, Item.Size, Item.Crc);
 			return true;
 		}
 	}
 	return false;
 }
 
-FSharedBuffer FPakFileDerivedDataBackend::LoadFile(const FStringView Path, const FStringView Context)
+FSharedBuffer FPakFileDerivedDataBackend::LoadFile(const FStringView Path, const FStringView DebugName)
 {
 	FWriteScopeLock ScopeLock(SynchronizationObject);
 	if (const FCacheValue* Item = CacheItems.FindByHash(GetTypeHash(Path), Path))
@@ -1042,13 +1040,13 @@ FSharedBuffer FPakFileDerivedDataBackend::LoadFile(const FStringView Path, const
 			{
 				UE_LOG(LogDerivedDataCache, Display,
 					TEXT("%s: File %.*s from '%.*s' failed to read %" INT64_FMT " bytes."),
-					*CachePath, Path.Len(), Path.GetData(), Context.Len(), Context.GetData(), Item->Size);
+					*CachePath, Path.Len(), Path.GetData(), DebugName.Len(), DebugName.GetData(), Item->Size);
 			}
 			else if (const uint32 TestCrc = HashAr.GetHash(); TestCrc != Item->Crc)
 			{
 				UE_LOG(LogDerivedDataCache, Display,
 					TEXT("%s: File %.*s from '%.*s' is corrupted and has CRC 0x%08x when 0x%08x is expected."),
-					*CachePath, Path.Len(), Path.GetData(), Context.Len(), Context.GetData(), TestCrc, Item->Crc);
+					*CachePath, Path.Len(), Path.GetData(), DebugName.Len(), DebugName.GetData(), TestCrc, Item->Crc);
 			}
 			else
 			{
diff --git a/Engine/Source/Developer/DerivedDataCache/Private/PakFileDerivedDataBackend.h b/Engine/Source/Developer/DerivedDataCache/Private/PakFileDerivedDataBackend.h
--- a/Engine/Source/Developer/DerivedDataCache/Private/PakFileDerivedDataBackend.h
+++ b/Engine/Source/Developer/DerivedDataCache/Private/PakFileDerivedDataBackend.h
@@ -106,48 +106,45 @@ public:
 
 	virtual void Put(
 		TConstArrayView<FCachePutRequest> Requests,
-		FStringView Context,
 		IRequestOwner& Owner,
 		FOnCachePutComplete&& OnComplete) override;
 
 	virtual void Get(
 		TConstArrayView<FCacheGetRequest> Requests,
-		FStringView Context,
 		IRequestOwner& Owner,
 		FOnCacheGetComplete&& OnComplete) override;
 
 	virtual void GetChunks(
-		TConstArrayView<FCacheChunkRequest> Chunks,
-		FStringView Context,
+		TConstArrayView<FCacheChunkRequest> Requests,
 		IRequestOwner& Owner,
-		FOnCacheGetChunkComplete&& OnComplete) override;
+		FOnCacheChunkComplete&& OnComplete) override;
 
 private:
 	uint64 MeasureCompressedCacheRecord(const FCacheRecord& Record) const;
 	uint64 MeasureRawCacheRecord(const FCacheRecord& Record) const;
 
-	bool PutCacheRecord(const FCacheRecord& Record, FStringView Context, const FCacheRecordPolicy& Policy);
-	bool PutCacheContent(const FCompressedBuffer& Content, const FStringView Context);
+	bool PutCacheRecord(const FCacheRecord& Record, FStringView RecordName, const FCacheRecordPolicy& Policy);
+	bool PutCacheContent(const FCompressedBuffer& Content, const FStringView ContentName);
 
 	FOptionalCacheRecord GetCacheRecordOnly(
 		const FCacheKey& Key,
-		const FStringView Context,
+		const FStringView RecordName,
 		const FCacheRecordPolicy& Policy);
 	FOptionalCacheRecord GetCacheRecord(
 		const FCacheKey& Key,
-		const FStringView Context,
+		const FStringView RecordName,
 		const FCacheRecordPolicy& Policy,
 		EStatus& OutStatus);
 	void GetCacheContent(
 		const FCacheKey& Key,
-		const FStringView Context,
+		const FStringView ContentName,
 		const ECachePolicy Policy,
 		const ECachePolicy SkipFlag,
 		FPayload& InOutPayload,
 		EStatus& InOutStatus);
 
-	bool SaveFile(FStringView Path, FStringView Context, TFunctionRef<void (FArchive&)> WriteFunction);
-	FSharedBuffer LoadFile(FStringView Path, FStringView Context);
+	bool SaveFile(FStringView Path, FStringView DebugName, TFunctionRef<void (FArchive&)> WriteFunction);
+	FSharedBuffer LoadFile(FStringView Path, FStringView DebugName);
 	bool FileExists(FStringView Path);
 
 private:
diff --git a/Engine/Source/Developer/DerivedDataCache/Private/S3DerivedDataBackend.cpp b/Engine/Source/Developer/DerivedDataCache/Private/S3DerivedDataBackend.cpp
--- a/Engine/Source/Developer/DerivedDataCache/Private/S3DerivedDataBackend.cpp
+++ b/Engine/Source/Developer/DerivedDataCache/Private/S3DerivedDataBackend.cpp
@@ -1220,7 +1220,6 @@ bool FS3DerivedDataBackend::ShouldSimulateMiss(const TCHAR* InKey)
 
 void FS3DerivedDataBackend::Put(
 	const TConstArrayView<FCachePutRequest> Requests,
-	const FStringView Context,
 	IRequestOwner& Owner,
 	FOnCachePutComplete&& OnComplete)
 {
@@ -1228,14 +1227,13 @@ void FS3DerivedDataBackend::Put(
 	{
 		for (const FCachePutRequest& Request : Requests)
 		{
-			OnComplete({Request.Record.GetKey(), Request.UserData, EStatus::Error});
+			OnComplete({Request.Name, Request.Record.GetKey(), Request.UserData, EStatus::Error});
 		}
 	}
 }
 
 void FS3DerivedDataBackend::Get(
 	const TConstArrayView<FCacheGetRequest> Requests,
-	const FStringView Context,
 	IRequestOwner& Owner,
 	FOnCacheGetComplete&& OnComplete)
 {
@@ -1243,22 +1241,21 @@ void FS3DerivedDataBackend::Get(
 	{
 		for (const FCacheGetRequest& Request : Requests)
 		{
-			OnComplete({FCacheRecordBuilder(Request.Key).Build(), Request.UserData, EStatus::Error});
+			OnComplete({Request.Name, FCacheRecordBuilder(Request.Key).Build(), Request.UserData, EStatus::Error});
 		}
 	}
 }
 
 void FS3DerivedDataBackend::GetChunks(
-	const TConstArrayView<FCacheChunkRequest> Chunks,
-	const FStringView Context,
+	const TConstArrayView<FCacheChunkRequest> Requests,
 	IRequestOwner& Owner,
-	FOnCacheGetChunkComplete&& OnComplete)
+	FOnCacheChunkComplete&& OnComplete)
 {
 	if (OnComplete)
 	{
-		for (const FCacheChunkRequest& Chunk : Chunks)
+		for (const FCacheChunkRequest& Request : Requests)
 		{
-			OnComplete({Chunk.Key, Chunk.Id, Chunk.RawOffset, 0, {}, {}, Chunk.UserData, EStatus::Error});
+			OnComplete({Request.Name, Request.Key, Request.Id, Request.RawOffset, 0, {}, {}, Request.UserData, EStatus::Error});
 		}
 	}
 }
diff --git a/Engine/Source/Developer/DerivedDataCache/Private/S3DerivedDataBackend.h b/Engine/Source/Developer/DerivedDataCache/Private/S3DerivedDataBackend.h
--- a/Engine/Source/Developer/DerivedDataCache/Private/S3DerivedDataBackend.h
+++ b/Engine/Source/Developer/DerivedDataCache/Private/S3DerivedDataBackend.h
@@ -61,21 +61,18 @@ public:
 
 	virtual void Put(
 		TConstArrayView<FCachePutRequest> Requests,
-		FStringView Context,
 		IRequestOwner& Owner,
 		FOnCachePutComplete&& OnComplete) override;
 
 	virtual void Get(
 		TConstArrayView<FCacheGetRequest> Requests,
-		FStringView Context,
 		IRequestOwner& Owner,
 		FOnCacheGetComplete&& OnComplete) override;
 
 	virtual void GetChunks(
-		TConstArrayView<FCacheChunkRequest> Chunks,
-		FStringView Context,
+		TConstArrayView<FCacheChunkRequest> Requests,
 		IRequestOwner& Owner,
-		FOnCacheGetChunkComplete&& OnComplete) override;
+		FOnCacheChunkComplete&& OnComplete) override;
 
 private:
 	struct FBundle;
diff --git a/Engine/Source/Developer/DerivedDataCache/Private/ZenDerivedDataBackend.cpp b/Engine/Source/Developer/DerivedDataCache/Private/ZenDerivedDataBackend.cpp
--- a/Engine/Source/Developer/DerivedDataCache/Private/ZenDerivedDataBackend.cpp
+++ b/Engine/Source/Developer/DerivedDataCache/Private/ZenDerivedDataBackend.cpp
@@ -589,7 +589,6 @@ bool FZenDerivedDataBackend::ShouldSimulateMiss(const FCacheKey& Key)
 
 void FZenDerivedDataBackend::Put(
 	const TConstArrayView<FCachePutRequest> Requests,
-	const FStringView Context,
 	IRequestOwner& Owner,
 	FOnCachePutComplete&& OnComplete)
 {
@@ -602,21 +601,21 @@ void FZenDerivedDataBackend::Put(
 		bool bResult;
 		if (ShouldSimulateMiss(Record.GetKey()))
 		{
-			UE_LOG(LogDerivedDataCache, Verbose, TEXT("%s: Simulated miss for put of %s from '%.*s'"),
-				*GetName(), *WriteToString<96>(Record.GetKey()), Context.Len(), Context.GetData());
+			UE_LOG(LogDerivedDataCache, Verbose, TEXT("%s: Simulated miss for put of %s from '%s'"),
+				*GetName(), *WriteToString<96>(Record.GetKey()), *Request.Name);
 			bResult = false;
 		}
 
-		bResult = PutCacheRecord(Record, Context, Request.Policy);
+		bResult = PutCacheRecord(Record, Request.Policy);
 
 		if (bResult)
 		{
-			UE_LOG(LogDerivedDataCache, Verbose, TEXT("%s: Cache put complete for %s from '%.*s'"),
-				*GetName(), *WriteToString<96>(Record.GetKey()), Context.Len(), Context.GetData());
+			UE_LOG(LogDerivedDataCache, Verbose, TEXT("%s: Cache put complete for %s from '%s'"),
+				*GetName(), *WriteToString<96>(Record.GetKey()), *Request.Name);
 			COOK_STAT(Timer.AddHit(MeasureCacheRecord(Record)));
 			if (OnComplete)
 			{
-				OnComplete({ Record.GetKey(), Request.UserData, EStatus::Ok });
+				OnComplete({ Request.Name, Record.GetKey(), Request.UserData, EStatus::Ok });
 			}
 		}
 		else
@@ -624,7 +623,7 @@ void FZenDerivedDataBackend::Put(
 			COOK_STAT(Timer.AddMiss(MeasureCacheRecord(Record)));
 			if (OnComplete)
 			{
-				OnComplete({ Record.GetKey(), Request.UserData, EStatus::Error });
+				OnComplete({ Request.Name, Record.GetKey(), Request.UserData, EStatus::Error });
 			}
 		}
 	}
@@ -632,7 +631,6 @@ void FZenDerivedDataBackend::Put(
 
 void FZenDerivedDataBackend::Get(
 	const TConstArrayView<FCacheGetRequest> Requests,
-	const FStringView Context,
 	IRequestOwner& Owner,
 	FOnCacheGetComplete&& OnComplete)
 {
@@ -642,7 +640,7 @@ void FZenDerivedDataBackend::Get(
 	int32 TotalCompleted = 0;
 
 	ForEachBatch(CacheRecordBatchSize, Requests.Num(),
-		[this, &Requests, &Context, &Owner, &OnComplete, &TotalCompleted](int32 BatchFirst, int32 BatchLast)
+		[this, &Requests, &Owner, &OnComplete, &TotalCompleted](int32 BatchFirst, int32 BatchLast)
 	{
 		TRACE_COUNTER_ADD(ZenDDC_Get, int64(BatchLast) - int64(BatchFirst) + 1);
 		COOK_STAT(auto Timer = UsageStats.TimeGet());
@@ -716,8 +714,8 @@ void FZenDerivedDataBackend::Get(
 				{
 					if (ShouldSimulateMiss(Key))
 					{
-						UE_LOG(LogDerivedDataCache, Verbose, TEXT("%s: Simulated miss for put of '%s' from '%.*s'"),
-							*GetName(), *WriteToString<96>(Key), Context.Len(), Context.GetData());
+						UE_LOG(LogDerivedDataCache, Verbose, TEXT("%s: Simulated miss for put of '%s' from '%s'"),
+							*GetName(), *WriteToString<96>(Key), *Request.Name);
 					}
 					else
 					{
@@ -727,8 +725,8 @@ void FZenDerivedDataBackend::Get(
 
 				if (Record)
 				{
-					UE_LOG(LogDerivedDataCache, Verbose, TEXT("%s: Cache hit for '%s' from '%.*s'"),
-						*GetName(), *WriteToString<96>(Key), Context.Len(), Context.GetData());
+					UE_LOG(LogDerivedDataCache, Verbose, TEXT("%s: Cache hit for '%s' from '%s'"),
+						*GetName(), *WriteToString<96>(Key), *Request.Name);
 					
 					TRACE_COUNTER_ADD(ZenDDC_GetHit, int64(1));
 					int64 ReceivedSize = MeasureCacheRecord(Record.Get());
@@ -737,17 +735,17 @@ void FZenDerivedDataBackend::Get(
 
 					if (OnComplete)
 					{
-						OnComplete({ MoveTemp(Record).Get(), Request.UserData, EStatus::Ok });
+						OnComplete({ Request.Name, MoveTemp(Record).Get(), Request.UserData, EStatus::Ok });
 					}
 				}
 				else
 				{
-					UE_LOG(LogDerivedDataCache, Verbose, TEXT("%s: Cache miss for '%s' from '%.*s'"),
-						*GetName(), *WriteToString<96>(Key), Context.Len(), Context.GetData());
+					UE_LOG(LogDerivedDataCache, Verbose, TEXT("%s: Cache miss for '%s' from '%s'"),
+						*GetName(), *WriteToString<96>(Key), *Request.Name);
 					
 					if (OnComplete)
 					{
-						OnComplete({ FCacheRecordBuilder(Key).Build(), Request.UserData, EStatus::Error });
+						OnComplete({ Request.Name, FCacheRecordBuilder(Key).Build(), Request.UserData, EStatus::Error });
 					}
 				}
 				
@@ -761,12 +759,12 @@ void FZenDerivedDataBackend::Get(
 				const FCacheGetRequest& Request = Requests[KeyIndex];
 				const FCacheKey& Key = Request.Key;
 
-				UE_LOG(LogDerivedDataCache, Verbose, TEXT("%s: Cache miss for '%s' from '%.*s'"),
-					*GetName(), *WriteToString<96>(Key), Context.Len(), Context.GetData());
+				UE_LOG(LogDerivedDataCache, Verbose, TEXT("%s: Cache miss for '%s' from '%s'"),
+					*GetName(), *WriteToString<96>(Key), *Request.Name);
 					
 				if (OnComplete)
 				{
-					OnComplete({ FCacheRecordBuilder(Key).Build(), Request.UserData, EStatus::Error });
+					OnComplete({ Request.Name, FCacheRecordBuilder(Key).Build(), Request.UserData, EStatus::Error });
 				}
 				
 				TotalCompleted++;
@@ -779,21 +777,20 @@ void FZenDerivedDataBackend::Get(
 }
 
 void FZenDerivedDataBackend::GetChunks(
-	TConstArrayView<FCacheChunkRequest> Chunks,
-	FStringView Context,
+	TConstArrayView<FCacheChunkRequest> Requests,
 	IRequestOwner& Owner,
-	FOnCacheGetChunkComplete&& OnComplete)
+	FOnCacheChunkComplete&& OnComplete)
 {
 	TRACE_CPUPROFILER_EVENT_SCOPE(ZenDDC::GetChunks);
-	TRACE_COUNTER_ADD(ZenDDC_ChunkRequestCount, int64(Chunks.Num()));
+	TRACE_COUNTER_ADD(ZenDDC_ChunkRequestCount, int64(Requests.Num()));
 
-	TArray<FCacheChunkRequest, TInlineAllocator<16>> SortedChunks(Chunks);
-	SortedChunks.StableSort(TChunkLess());
+	TArray<FCacheChunkRequest, TInlineAllocator<16>> SortedRequests(Requests);
+	SortedRequests.StableSort(TChunkLess());
 
 	int32 TotalCompleted = 0;
 
-	ForEachBatch(CacheChunksBatchSize, SortedChunks.Num(),
-		[this, &SortedChunks, &Context, &Owner, &OnComplete, &TotalCompleted](int32 BatchFirst, int32 BatchLast)
+	ForEachBatch(CacheChunksBatchSize, SortedRequests.Num(),
+		[this, &SortedRequests, &Owner, &OnComplete, &TotalCompleted](int32 BatchFirst, int32 BatchLast)
 	{
 		TRACE_COUNTER_ADD(ZenDDC_Get, int64(BatchLast) - int64(BatchFirst) + 1);
 		COOK_STAT(auto Timer = UsageStats.TimeGet());
@@ -807,19 +804,19 @@ void FZenDerivedDataBackend::GetChunks(
 				BatchRequest.BeginArray("ChunkRequests"_ASV);
 				for (int32 ChunkIndex = BatchFirst; ChunkIndex <= BatchLast; ChunkIndex++)
 				{
-					const FCacheChunkRequest& Chunk = SortedChunks[ChunkIndex];
+					const FCacheChunkRequest& Request = SortedRequests[ChunkIndex];
 					
 					BatchRequest.BeginObject();
 					
 					BatchRequest.BeginObject("Key"_ASV);
-					BatchRequest << "Bucket"_ASV << Chunk.Key.Bucket.ToString();
-					BatchRequest << "Hash"_ASV << Chunk.Key.Hash;
+					BatchRequest << "Bucket"_ASV << Request.Key.Bucket.ToString();
+					BatchRequest << "Hash"_ASV << Request.Key.Hash;
 					BatchRequest.EndObject();
 
-					BatchRequest.AddObjectId("PayloadId"_ASV, Chunk.Id);
-					BatchRequest << "RawOffset"_ASV << Chunk.RawOffset;
-					BatchRequest << "RawSize"_ASV << Chunk.RawSize;
-					BatchRequest << "Policy"_ASV << static_cast<uint32>(Chunk.Policy);
+					BatchRequest.AddObjectId("PayloadId"_ASV, Request.Id);
+					BatchRequest << "RawOffset"_ASV << Request.RawOffset;
+					BatchRequest << "RawSize"_ASV << Request.RawSize;
+					BatchRequest << "Policy"_ASV << static_cast<uint32>(Request.Policy);
 
 					BatchRequest.EndObject();
 				}
@@ -844,25 +841,25 @@ void FZenDerivedDataBackend::GetChunks(
 			int32 ChunkIndex = BatchFirst;
 			for (FCbFieldView HashView : ResponseObj["Result"_ASV])
 			{
-				const FCacheChunkRequest& Chunk = SortedChunks[ChunkIndex++];
+				const FCacheChunkRequest& Request = SortedRequests[ChunkIndex++];
 
-				if (ShouldSimulateMiss(Chunk.Key))
+				if (ShouldSimulateMiss(Request.Key))
 				{
-					UE_LOG(LogDerivedDataCache, Verbose, TEXT("%s: Simulated miss for get of '%s' from '%.*s'"),
-						*GetName(), *WriteToString<96>(Chunk.Key ,'/', Chunk.Id), Context.Len(), Context.GetData());
+					UE_LOG(LogDerivedDataCache, Verbose, TEXT("%s: Simulated miss for get of '%s' from '%s'"),
+						*GetName(), *WriteToString<96>(Request.Key ,'/', Request.Id), *Request.Name);
 					
 					if (OnComplete)
 					{
-						OnComplete({ Chunk.Key, Chunk.Id, Chunk.RawOffset, 0, {}, {}, Chunk.UserData, EStatus::Error });
+						OnComplete({ Request.Name, Request.Key, Request.Id, Request.RawOffset, 0, {}, {}, Request.UserData, EStatus::Error });
 					}
 				}
 				else if (const FCbAttachment* Attachment = BatchResponse.FindAttachment(HashView.AsHash()))
 				{
 					const FCompressedBuffer& CompressedBuffer = Attachment->AsCompressedBinary();
-					FSharedBuffer Buffer = FCompressedBufferReader(CompressedBuffer).Decompress(Chunk.RawOffset, Chunk.RawSize);
+					FSharedBuffer Buffer = FCompressedBufferReader(CompressedBuffer).Decompress(Request.RawOffset, Request.RawSize);
 					
-					UE_LOG(LogDerivedDataCache, Verbose, TEXT("%s: Cache hit for '%s' from '%.*s'"),
-						*GetName(), *WriteToString<96>(Chunk.Key, '/', Chunk.Id), Context.Len(), Context.GetData());
+					UE_LOG(LogDerivedDataCache, Verbose, TEXT("%s: Cache hit for '%s' from '%s'"),
+						*GetName(), *WriteToString<96>(Request.Key, '/', Request.Id), *Request.Name);
 
 					const uint64 RawSize = Buffer.GetSize();
 					TRACE_COUNTER_ADD(ZenDDC_GetHit, int64(1));
@@ -871,17 +868,17 @@ void FZenDerivedDataBackend::GetChunks(
 					
 					if (OnComplete)
 					{
-						OnComplete({Chunk.Key, Chunk.Id, Chunk.RawOffset, RawSize, CompressedBuffer.GetRawHash(), MoveTemp(Buffer), Chunk.UserData, EStatus::Ok});
+						OnComplete({ Request.Name, Request.Key, Request.Id, Request.RawOffset, RawSize, CompressedBuffer.GetRawHash(), MoveTemp(Buffer), Request.UserData, EStatus::Ok });
 					}
 				}
 				else
 				{
-					UE_LOG(LogDerivedDataCache, Display, TEXT("%s: Cache miss with missing payload '%s' for '%s' from '%.*s'"),
-						*GetName(), *WriteToString<16>(Chunk.Id), *WriteToString<96>(Chunk.Key), Context.Len(), Context.GetData());
+					UE_LOG(LogDerivedDataCache, Display, TEXT("%s: Cache miss with missing payload '%s' for '%s' from '%s'"),
+						*GetName(), *WriteToString<16>(Request.Id), *WriteToString<96>(Request.Key), *Request.Name);
 					
 					if (OnComplete)
 					{
-						OnComplete({ Chunk.Key, Chunk.Id, Chunk.RawOffset, 0, {}, {}, Chunk.UserData, EStatus::Error });
+						OnComplete({ Request.Name, Request.Key, Request.Id, Request.RawOffset, 0, {}, {}, Request.UserData, EStatus::Error });
 					}
 				}
 
@@ -892,14 +889,14 @@ void FZenDerivedDataBackend::GetChunks(
 		{
 			for (int32 ChunkIndex = BatchFirst; ChunkIndex <= BatchLast; ChunkIndex++)
 			{
-				const FCacheChunkRequest& Chunk = SortedChunks[ChunkIndex];
+				const FCacheChunkRequest& Request = SortedRequests[ChunkIndex];
 				
-				UE_LOG(LogDerivedDataCache, Display, TEXT("%s: Cache miss with missing payload '%s' for '%s' from '%.*s'"),
-					*GetName(), *WriteToString<16>(Chunk.Id), *WriteToString<96>(Chunk.Key), Context.Len(), Context.GetData());
+				UE_LOG(LogDerivedDataCache, Display, TEXT("%s: Cache miss with missing payload '%s' for '%s' from '%s'"),
+					*GetName(), *WriteToString<16>(Request.Id), *WriteToString<96>(Request.Key), *Request.Name);
 				
 				if (OnComplete)
 				{
-					OnComplete({ Chunk.Key, Chunk.Id, Chunk.RawOffset, 0, {}, {}, Chunk.UserData, EStatus::Error });
+					OnComplete({ Request.Name, Request.Key, Request.Id, Request.RawOffset, 0, {}, {}, Request.UserData, EStatus::Error });
 				}
 				
 				TotalCompleted++;
@@ -907,11 +904,11 @@ void FZenDerivedDataBackend::GetChunks(
 		}
 	});
 
-	UE_CLOG(TotalCompleted != SortedChunks.Num(), LogDerivedDataCache, Warning, TEXT("Only '%d/%d' cache chunk request(s) completed"), TotalCompleted, SortedChunks.Num());
-	TRACE_COUNTER_SUBTRACT(ZenDDC_ChunkRequestCount, int64(Chunks.Num()));
+	UE_CLOG(TotalCompleted != SortedRequests.Num(), LogDerivedDataCache, Warning, TEXT("Only '%d/%d' cache chunk request(s) completed"), TotalCompleted, SortedRequests.Num());
+	TRACE_COUNTER_SUBTRACT(ZenDDC_ChunkRequestCount, int64(Requests.Num()));
 }
 
-bool FZenDerivedDataBackend::PutCacheRecord(const FCacheRecord& Record, FStringView Context, const FCacheRecordPolicy& Policy)
+bool FZenDerivedDataBackend::PutCacheRecord(const FCacheRecord& Record, const FCacheRecordPolicy& Policy)
 {
 	const FCacheKey& Key = Record.GetKey();
 	FCbPackage Package = Record.Save();
diff --git a/Engine/Source/Developer/DerivedDataCache/Private/ZenDerivedDataBackend.h b/Engine/Source/Developer/DerivedDataCache/Private/ZenDerivedDataBackend.h
--- a/Engine/Source/Developer/DerivedDataCache/Private/ZenDerivedDataBackend.h
+++ b/Engine/Source/Developer/DerivedDataCache/Private/ZenDerivedDataBackend.h
@@ -92,21 +92,18 @@ public:
 
 	virtual void Put(
 		TConstArrayView<FCachePutRequest> Requests,
-		FStringView Context,
 		IRequestOwner& Owner,
 		FOnCachePutComplete&& OnComplete = FOnCachePutComplete()) override;
 
 	virtual void Get(
 		TConstArrayView<FCacheGetRequest> Requests,
-		FStringView Context,
 		IRequestOwner& Owner,
 		FOnCacheGetComplete&& OnComplete) override;
 
 	virtual void GetChunks(
-		TConstArrayView<FCacheChunkRequest> Chunks,
-		FStringView Context,
+		TConstArrayView<FCacheChunkRequest> Requests,
 		IRequestOwner& Owner,
-		FOnCacheGetChunkComplete&& OnComplete) override;
+		FOnCacheChunkComplete&& OnComplete) override;
 
 private:
 	enum class EGetResult
@@ -121,7 +118,7 @@ private:
 	FDerivedDataBackendInterface::EPutStatus PutZenData(const TCHAR* Uri, const FCompositeBuffer& InData, Zen::EContentType ContentType);
 	EGetResult GetZenData(const FCacheKey& Key, ECachePolicy CachePolicy, FCbPackage& OutPackage) const;
 
-	bool PutCacheRecord(const FCacheRecord& Record, FStringView Context, const FCacheRecordPolicy& Policy);
+	bool PutCacheRecord(const FCacheRecord& Record, const FCacheRecordPolicy& Policy);
 
 	bool IsServiceReady();
 	static FString MakeLegacyZenKey(const TCHAR* CacheKey);
diff --git a/Engine/Source/Developer/DerivedDataCache/Public/DerivedDataCache.h b/Engine/Source/Developer/DerivedDataCache/Public/DerivedDataCache.h
--- a/Engine/Source/Developer/DerivedDataCache/Public/DerivedDataCache.h
+++ b/Engine/Source/Developer/DerivedDataCache/Public/DerivedDataCache.h
@@ -4,6 +4,7 @@
 
 #include "CoreTypes.h"
 #include "Containers/StringFwd.h"
+#include "Containers/UnrealString.h" // DDC-TODO: Remove!
 #include "DerivedDataCacheKey.h"
 #include "DerivedDataCacheRecord.h"
 #include "DerivedDataPayloadId.h"
@@ -17,8 +18,8 @@
 
 namespace UE::DerivedData { class ICacheStoreMaintainer; }
 namespace UE::DerivedData { class IRequestOwner; }
+namespace UE::DerivedData { struct FCacheChunkCompleteParams; }
 namespace UE::DerivedData { struct FCacheChunkRequest; }
-namespace UE::DerivedData { struct FCacheGetChunkCompleteParams; }
 namespace UE::DerivedData { struct FCacheGetCompleteParams; }
 namespace UE::DerivedData { struct FCacheGetRequest; }
 namespace UE::DerivedData { struct FCachePutCompleteParams; }
@@ -28,9 +29,11 @@ namespace UE::DerivedData::Private { class ICacheRecordPolicyShared; }
 namespace UE::DerivedData
 {
 
+using FCacheRequestName = FString;
+
 using FOnCachePutComplete = TUniqueFunction<void (FCachePutCompleteParams&& Params)>;
 using FOnCacheGetComplete = TUniqueFunction<void (FCacheGetCompleteParams&& Params)>;
-using FOnCacheGetChunkComplete = TUniqueFunction<void (FCacheGetChunkCompleteParams&& Params)>;
+using FOnCacheChunkComplete = TUniqueFunction<void (FCacheChunkCompleteParams&& Params)>;
 
 ///////////////////////////////////////////////////////////////////////////////////////////////////////////////////////
 
@@ -237,13 +240,11 @@ public:
 	 * be stored as a partial record that can attempt recovery of missing payloads when fetched.
 	 *
 	 * @param Requests     Requests with the cache records to store. Records must have a key.
-	 * @param Context      A description of the request. An object path is typically sufficient.
 	 * @param Owner        The owner to execute the request within.
 	 * @param OnComplete   A callback invoked for every record as it completes or is canceled.
 	 */
 	virtual void Put(
 		TConstArrayView<FCachePutRequest> Requests,
-		FStringView Context,
 		IRequestOwner& Owner,
 		FOnCachePutComplete&& OnComplete = FOnCachePutComplete()) = 0;
 
@@ -261,13 +262,11 @@ public:
 	 * PartialOnError is set for the missing payloads.
 	 *
 	 * @param Requests     Requests with the keys identifying the cache records to query.
-	 * @param Context      A description of the request. An object path is typically sufficient.
 	 * @param Owner        The owner to execute the request within.
 	 * @param OnComplete   A callback invoked for every key as it completes or is canceled.
 	 */
 	virtual void Get(
 		TConstArrayView<FCacheGetRequest> Requests,
-		FStringView Context,
 		IRequestOwner& Owner,
 		FOnCacheGetComplete&& OnComplete) = 0;
 
@@ -281,16 +280,14 @@ public:
 	 * the requested records. Requests for whole payloads may propagate those payloads into other
 	 * cache stores, in accordance with the policy, through a put of a partial record.
 	 *
-	 * @param Chunks       The keys, IDs, offsets, and sizes of the chunks to query.
-	 * @param Context      A description of the request. An object path is typically sufficient.
+	 * @param Requests     The keys, IDs, offsets, and sizes of the chunks to query.
 	 * @param Owner        The owner to execute the request within.
 	 * @param OnComplete   A callback invoked for every chunk as it completes or is canceled.
 	 */
 	virtual void GetChunks(
-		TConstArrayView<FCacheChunkRequest> Chunks,
-		FStringView Context,
+		TConstArrayView<FCacheChunkRequest> Requests,
 		IRequestOwner& Owner,
-		FOnCacheGetChunkComplete&& OnComplete) = 0;
+		FOnCacheChunkComplete&& OnComplete) = 0;
 };
 
 ///////////////////////////////////////////////////////////////////////////////////////////////////////////////////////
@@ -370,7 +367,6 @@ public:
 
 	/** Returns the interface to the background cache store maintenance. */
 	virtual ICacheStoreMaintainer& GetMaintainer() = 0;
-
 };
 
 ///////////////////////////////////////////////////////////////////////////////////////////////////////////////////////
@@ -378,6 +374,9 @@ public:
 /** Parameters to request to put a cache record. */
 struct FCachePutRequest
 {
+	/** A name to identify this request for logging and profiling. An object path is typically sufficient. */
+	FCacheRequestName Name;
+
 	/** The cache record to put. */
 	FCacheRecord Record;
 
@@ -391,6 +390,9 @@ struct FCachePutRequest
 /** Parameters for the completion callback for cache put requests. */
 struct FCachePutCompleteParams
 {
+	/** A copy of the name from the request. */
+	FCacheRequestName Name;
+
 	/** Key for the part of the put request that completed or was canceled. */
 	FCacheKey Key;
 
@@ -404,6 +406,9 @@ struct FCachePutCompleteParams
 /** Parameters to request a cache record. */
 struct FCacheGetRequest
 {
+	/** A name to identify this request for logging and profiling. An object path is typically sufficient. */
+	FCacheRequestName Name;
+
 	/** The key identifying the cache record to fetch. */
 	FCacheKey Key;
 
@@ -417,6 +422,9 @@ struct FCacheGetRequest
 /** Parameters for the completion callback for cache get requests. */
 struct FCacheGetCompleteParams
 {
+	/** A copy of the name from the request. */
+	FCacheRequestName Name;
+
 	/**
 	 * Record for the part of the get request that completed or was canceled.
 	 *
@@ -437,6 +445,9 @@ struct FCacheGetCompleteParams
 /** Parameters to request a chunk, which is a subset of a payload, from a cache record. */
 struct FCacheChunkRequest
 {
+	/** A name to identify this request for logging and profiling. An object path is typically sufficient. */
+	FCacheRequestName Name;
+
 	/** The key identifying the cache record to fetch the payload from. */
 	FCacheKey Key;
 
@@ -460,8 +471,11 @@ struct FCacheChunkRequest
 };
 
 /** Parameters for the completion callback for cache chunk requests. */
-struct FCacheGetChunkCompleteParams
+struct FCacheChunkCompleteParams
 {
+	/** A copy of the name from the request. */
+	FCacheRequestName Name;
+
 	/** Key from the chunk request that completed or was canceled. */
 	FCacheKey Key;
 
diff --git a/Engine/Source/Editor/UnrealEd/Private/EditorDomain/EditorDomainArchive.cpp b/Engine/Source/Editor/UnrealEd/Private/EditorDomain/EditorDomainArchive.cpp
--- a/Engine/Source/Editor/UnrealEd/Private/EditorDomain/EditorDomainArchive.cpp
+++ b/Engine/Source/Editor/UnrealEd/Private/EditorDomain/EditorDomainArchive.cpp
@@ -407,10 +407,10 @@ void FEditorDomainPackageSegments::SendSegmentRequest(FSegment& Segment)
 
 	// Note that Segment.RequestOwner is Interface-only and so we can write it outside the lock
 	Segment.RequestOwner.Emplace(EPriority::Normal);
-	FCacheChunkRequest SegmentChunk{UE::EditorDomain::GetEditorDomainPackageKey(EditorDomainHash), Segment.PayloadId};
+	FCacheChunkRequest SegmentChunk{PackagePath.GetDebugName(), UE::EditorDomain::GetEditorDomainPackageKey(EditorDomainHash), Segment.PayloadId};
 	SegmentChunk.Policy = ECachePolicy::Local;
-	Cache.GetChunks({SegmentChunk}, PackagePath.GetDebugName(), *Segment.RequestOwner,
-		[this, &Segment](FCacheGetChunkCompleteParams&& Params)
+	Cache.GetChunks({SegmentChunk}, *Segment.RequestOwner,
+		[this, &Segment](FCacheChunkCompleteParams&& Params)
 		{
 			if (AsyncSource == ESource::Closed)
 			{
diff --git a/Engine/Source/Runtime/Engine/Private/Streaming/Texture2DMipDataProvider_DDC.cpp b/Engine/Source/Runtime/Engine/Private/Streaming/Texture2DMipDataProvider_DDC.cpp
--- a/Engine/Source/Runtime/Engine/Private/Streaming/Texture2DMipDataProvider_DDC.cpp
+++ b/Engine/Source/Runtime/Engine/Private/Streaming/Texture2DMipDataProvider_DDC.cpp
@@ -58,34 +58,37 @@ void FTexture2DMipDataProvider_DDC::Init(const FTextureUpdateContext& Context, c
 				using namespace UE::DerivedData;
 				TArray<FCacheChunkRequest> MipKeys;
 
+				TStringBuilder<256> MipNameBuilder;
+				Context.Texture->GetPathName(nullptr, MipNameBuilder);
+				const int32 TextureNameLen = MipNameBuilder.Len();
+
 				const FCacheKey& Key = *PlatformData->DerivedDataKey.Get<UE::DerivedData::FCacheKeyProxy>().AsCacheKey();
 				for (int32 MipIndex = PendingFirstLODIdx; MipIndex < CurrentFirstLODIdx; ++MipIndex)
 				{
 					const FTexture2DMipMap& MipMap = *Context.MipsView[MipIndex];
 					if (MipMap.IsPagedToDerivedData())
 					{
-						MipKeys.Add({Key, FTexturePlatformData::MakeMipId(MipIndex + LODBias)});
+						FCacheChunkRequest& Request = MipKeys.AddDefaulted_GetRef();
+						MipNameBuilder.Appendf(TEXT(" [MIP 0]"), MipIndex + LODBias);
+						Request.Name = MipNameBuilder;
+						Request.Key = Key;
+						Request.Id = FTexturePlatformData::MakeMipId(MipIndex + LODBias);
+						Request.UserData = MipIndex;
+						MipNameBuilder.RemoveSuffix(MipNameBuilder.Len() - TextureNameLen);
 					}
 				}
 
 				if (MipKeys.Num())
 				{
-					GetCache().GetChunks(MipKeys, Context.Texture->GetPathName(), DDCRequestOwner,
-						[this, &Context, LODBias](FCacheGetChunkCompleteParams&& Params)
+					GetCache().GetChunks(MipKeys, DDCRequestOwner, [this](FCacheChunkCompleteParams&& Params)
+					{
+						if (Params.Status == EStatus::Ok)
 						{
-							if (Params.Status == EStatus::Ok)
-							{
-								for (int32 MipIndex = PendingFirstLODIdx; MipIndex < CurrentFirstLODIdx; ++MipIndex)
-								{
-									if (Params.Id == FTexturePlatformData::MakeMipId(MipIndex + LODBias))
-									{
-										check(!DDCBuffers[MipIndex]);
-										DDCBuffers[MipIndex] = MoveTemp(Params.RawData);
-										break;
-									}
-								}
-							}
-						});
+							const int32 MipIndex = int32(Params.UserData);
+							check(!DDCBuffers[MipIndex]);
+							DDCBuffers[MipIndex] = MoveTemp(Params.RawData);
+						}
+					});
 				}
 			}
 			else
diff --git a/Engine/Source/Runtime/Engine/Private/Streaming/Texture2DStreamIn_DDC.cpp b/Engine/Source/Runtime/Engine/Private/Streaming/Texture2DStreamIn_DDC.cpp
--- a/Engine/Source/Runtime/Engine/Private/Streaming/Texture2DStreamIn_DDC.cpp
+++ b/Engine/Source/Runtime/Engine/Private/Streaming/Texture2DStreamIn_DDC.cpp
@@ -161,34 +161,37 @@ void FTexture2DStreamIn_DDC::DoCreateAsyncDDCRequests(const FContext& Context)
 			using namespace UE::DerivedData;
 			TArray<FCacheChunkRequest> MipKeys;
 
+			TStringBuilder<256> MipNameBuilder;
+			Context.Texture->GetPathName(nullptr, MipNameBuilder);
+			const int32 TextureNameLen = MipNameBuilder.Len();
+
 			const FCacheKey& Key = *PlatformData->DerivedDataKey.Get<UE::DerivedData::FCacheKeyProxy>().AsCacheKey();
 			for (int32 MipIndex = PendingFirstLODIdx; MipIndex < CurrentFirstLODIdx && !IsCancelled(); ++MipIndex)
 			{
 				const FTexture2DMipMap& MipMap = *Context.MipsView[MipIndex];
 				if (MipMap.IsPagedToDerivedData())
 				{
-					MipKeys.Add({Key, FTexturePlatformData::MakeMipId(MipIndex + LODBias)});
+					FCacheChunkRequest& Request = MipKeys.AddDefaulted_GetRef();
+					MipNameBuilder.Appendf(TEXT(" [MIP 0]"), MipIndex + LODBias);
+					Request.Name = MipNameBuilder;
+					Request.Key = Key;
+					Request.Id = FTexturePlatformData::MakeMipId(MipIndex + LODBias);
+					Request.UserData = MipIndex;
+					MipNameBuilder.RemoveSuffix(MipNameBuilder.Len() - TextureNameLen);
 				}
 			}
 
 			if (MipKeys.Num())
 			{
-				GetCache().GetChunks(MipKeys, Context.Texture->GetPathName(), DDCRequestOwner,
-					[this, LODBias](FCacheGetChunkCompleteParams&& Params)
+				GetCache().GetChunks(MipKeys, DDCRequestOwner, [this](FCacheChunkCompleteParams&& Params)
+				{
+					if (Params.Status == EStatus::Ok)
 					{
-						if (Params.Status == EStatus::Ok)
-						{
-							for (int32 MipIndex = PendingFirstLODIdx; MipIndex < CurrentFirstLODIdx; ++MipIndex)
-							{
-								if (Params.Id == FTexturePlatformData::MakeMipId(MipIndex + LODBias))
-								{
-									check(!DDCBuffers[MipIndex]);
-									DDCBuffers[MipIndex] = MoveTemp(Params.RawData);
-									break;
-								}
-							}
-						}
-					});
+						const int32 MipIndex = int32(Params.UserData);
+						check(!DDCBuffers[MipIndex]);
+						DDCBuffers[MipIndex] = MoveTemp(Params.RawData);
+					}
+				});
 			}
 		}
 		else
diff --git a/Engine/Source/Runtime/Engine/Private/TextureDerivedData.cpp b/Engine/Source/Runtime/Engine/Private/TextureDerivedData.cpp
--- a/Engine/Source/Runtime/Engine/Private/TextureDerivedData.cpp
+++ b/Engine/Source/Runtime/Engine/Private/TextureDerivedData.cpp
@@ -1287,7 +1287,13 @@ static bool BeginLoadDerivedMips(FTexturePlatformData& PlatformData, int32 First
 			const FTexture2DMipMap& Mip = Mips[MipIndex];
 			if (Mip.IsPagedToDerivedData())
 			{
-				MipKeys.Add({Key, FTexturePlatformData::MakeMipId(MipIndex)});
+				TStringBuilder<256> MipNameBuilder;
+				MipNameBuilder.Append(DebugContext).Appendf(TEXT(" [MIP 0]"), MipIndex);
+				FCacheChunkRequest& Request = MipKeys.AddDefaulted_GetRef();
+				Request.Name = MipNameBuilder;
+				Request.Key = Key;
+				Request.Id = FTexturePlatformData::MakeMipId(MipIndex);
+				Request.UserData = MipIndex;
 			}
 		}
 
@@ -1296,19 +1302,11 @@ static bool BeginLoadDerivedMips(FTexturePlatformData& PlatformData, int32 First
 			check(Callback);
 			bool bMiss = false;
 			FRequestOwner RequestOwner(EPriority::Blocking);
-			GetCache().GetChunks(MipKeys, DebugContext, RequestOwner,
-				[&Mips, Callback = MoveTemp(Callback), FirstMipToLoad, &bMiss](FCacheGetChunkCompleteParams&& Params)
+			GetCache().GetChunks(MipKeys, RequestOwner, [Callback = MoveTemp(Callback), &bMiss](FCacheChunkCompleteParams&& Params)
 			{
 				if (Params.Status == EStatus::Ok)
 				{
-					for (int32 MipIndex = FirstMipToLoad; MipIndex < Mips.Num(); ++MipIndex)
-					{
-						if (Params.Id == FTexturePlatformData::MakeMipId(MipIndex))
-						{
-							Callback(MipIndex, MoveTemp(Params.RawData));
-							break;
-						}
-					}
+					Callback(int32(Params.UserData), MoveTemp(Params.RawData));
 				}
 				else
 				{
