DDC: Added FSharedString to minimize copies of cache request names

#rnx

diff --git a/Engine/Source/Developer/DerivedDataCache/Public/DerivedDataCache.h b/Engine/Source/Developer/DerivedDataCache/Public/DerivedDataCache.h
--- a/Engine/Source/Developer/DerivedDataCache/Public/DerivedDataCache.h
+++ b/Engine/Source/Developer/DerivedDataCache/Public/DerivedDataCache.h
@@ -4,11 +4,11 @@
 
 #include "CoreTypes.h"
 #include "Containers/StringFwd.h"
-#include "Containers/UnrealString.h" // DDC-TODO: Remove!
 #include "DerivedDataCacheKey.h"
 #include "DerivedDataCacheRecord.h"
 #include "DerivedDataPayloadId.h"
 #include "DerivedDataRequestTypes.h"
+#include "DerivedDataSharedString.h"
 #include "Math/NumericLimits.h"
 #include "Misc/EnumClassFlags.h"
 #include "Templates/Function.h"
@@ -29,7 +29,7 @@ namespace UE::DerivedData::Private { class ICacheRecordPolicyShared; }
 namespace UE::DerivedData
 {
 
-using FCacheRequestName = FString;
+using FCacheRequestName = FSharedString;
 
 using FOnCachePutComplete = TUniqueFunction<void (FCachePutCompleteParams&& Params)>;
 using FOnCacheGetComplete = TUniqueFunction<void (FCacheGetCompleteParams&& Params)>;
diff --git a/Engine/Source/Developer/DerivedDataCache/Public/DerivedDataSharedString.h b/Engine/Source/Developer/DerivedDataCache/Public/DerivedDataSharedString.h
new file mode 100644
--- /dev/null
+++ b/Engine/Source/Developer/DerivedDataCache/Public/DerivedDataSharedString.h
@@ -0,0 +1,172 @@
+// Copyright Epic Games, Inc. All Rights Reserved.
+
+#pragma once
+
+#include "CoreTypes.h"
+#include "Containers/StringView.h"
+#include "HAL/UnrealMemory.h"
+#include "Traits/ElementType.h"
+#include "Traits/IsContiguousContainer.h"
+#include <atomic>
+
+namespace UE::DerivedData
+{
+
+/**
+ * A reference to an immutable, shared, reference-counted string.
+ *
+ * Prefer TStringView when there is a clear single point of ownership with a longer lifetime than
+ * the references to the string. This type is meant for cases where a shared string does not have
+ * an obvious owner or where the lifetime is not easy to manage.
+ *
+ * The string is stored as a pointer to the start of the null-terminated string which is preceded
+ * by a 4-byte reference count and a 4-byte size.
+ */
+template <typename CharType>
+class TSharedString
+{
+public:
+	TSharedString() = default;
+	TSharedString(TSharedString&& String);
+	TSharedString(const TSharedString& String);
+	TSharedString& operator=(TSharedString&& String);
+	TSharedString& operator=(const TSharedString& String);
+
+	TSharedString(TStringView<CharType> String);
+	TSharedString& operator=(TStringView<CharType> String);
+
+	~TSharedString() { Reset(); }
+
+	inline void Reset() { Release(Chars); Chars = nullptr; }
+
+	[[nodiscard]] inline bool IsEmpty() const { return !Chars; }
+	[[nodiscard]] inline int32 Len() const { return Chars ? reinterpret_cast<const int32*>(Chars)[-1] : 0; }
+	[[nodiscard]] inline const CharType* operator*() const { return Chars ? Chars : &NullChar; }
+
+private:
+	static inline void AddRef(CharType* Chars);
+	static inline void Release(CharType* Chars);
+
+	CharType* Chars = nullptr;
+
+	static constexpr inline CharType NullChar{};
+};
+
+template <typename CharType>
+constexpr inline const CharType* GetData(const TSharedString<CharType> String)
+{
+	return *String;
+}
+
+template <typename CharType>
+constexpr inline auto GetNum(const TSharedString<CharType> String)
+{
+	return String.Len();
+}
+
+using FSharedString = TSharedString<TCHAR>;
+using FAnsiSharedString = TSharedString<ANSICHAR>;
+using FWideSharedString = TSharedString<WIDECHAR>;
+using FUtf8SharedString = TSharedString<UTF8CHAR>;
+
+template <typename CharType>
+inline TSharedString<CharType>::TSharedString(TSharedString&& String)
+	: Chars(String.Chars)
+{
+	String.Chars = nullptr;
+}
+
+template <typename CharType>
+TSharedString<CharType>::TSharedString(const TSharedString& String)
+	: Chars(String.Chars)
+{
+	AddRef(Chars);
+}
+
+template <typename CharType>
+TSharedString<CharType>& TSharedString<CharType>::operator=(TSharedString&& String)
+{
+	if (this != &String)
+	{
+		Release(Chars);
+		Chars = String.Chars;
+		String.Chars = nullptr;
+	}
+	return *this;
+}
+
+template <typename CharType>
+TSharedString<CharType>& TSharedString<CharType>::operator=(const TSharedString& String)
+{
+	CharType* OldChars = Chars;
+	CharType* NewChars = String.Chars;
+	AddRef(NewChars);
+	Release(OldChars);
+	Chars = NewChars;
+	return *this;
+}
+
+template <typename CharType>
+TSharedString<CharType>::TSharedString(const TStringView<CharType> String)
+{
+	if (const int32 Length = String.Len())
+	{
+		static_assert(alignof(int32) <= MIN_ALIGNMENT);
+		static_assert(alignof(int32) >= alignof(CharType));
+		static_assert(sizeof(int32) == sizeof(std::atomic<int32>));
+		const SIZE_T Size = sizeof(int32) + sizeof(int32) + sizeof(CharType) + sizeof(CharType) * Length;
+		int32* const Header = static_cast<int32*>(FMemory::Malloc(Size));
+		new(&Header[0]) std::atomic<int32>(1); // Ref Count
+		new(&Header[1]) int32(Length);
+		Chars = reinterpret_cast<CharType*>(&Header[2]);
+		String.CopyString(Chars, Length);
+		Chars[Length] = CharType(0);
+	}
+}
+
+template <typename CharType>
+TSharedString<CharType>& TSharedString<CharType>::operator=(const TStringView<CharType> String)
+{
+	return *this = TSharedString(String);
+}
+
+template <typename CharType>
+void TSharedString<CharType>::AddRef(CharType* const MaybeNullChars)
+{
+	if (MaybeNullChars)
+	{
+		static_assert(sizeof(int32) == sizeof(std::atomic<int32>));
+		int32* const Header = reinterpret_cast<int32*>(MaybeNullChars) - 2;
+		std::atomic<int32>& RefCount = reinterpret_cast<std::atomic<int32>&>(Header[0]);
+		RefCount.fetch_add(1, std::memory_order_relaxed);
+	}
+}
+
+template <typename CharType>
+void TSharedString<CharType>::Release(CharType* const MaybeNullChars)
+{
+	if (MaybeNullChars)
+	{
+		static_assert(sizeof(int32) == sizeof(std::atomic<int32>));
+		int32* const Header = reinterpret_cast<int32*>(MaybeNullChars) - 2;
+		std::atomic<int32>& RefCount = reinterpret_cast<std::atomic<int32>&>(Header[0]);
+		if (RefCount.fetch_sub(1, std::memory_order_relaxed) == 1)
+		{
+			FMemory::Free(Header);
+		}
+	}
+}
+
+} // UE::DerivedData
+
+template <typename CharType>
+struct TIsContiguousContainer<UE::DerivedData::TSharedString<CharType>>
+{
+	static constexpr bool Value = true;
+};
+
+template <typename CharType>
+struct TElementType<UE::DerivedData::TSharedString<CharType>>
+{
+	using Type = CharType;
+};
diff --git a/Engine/Source/Editor/UnrealEd/Private/EditorDomain/EditorDomainArchive.cpp b/Engine/Source/Editor/UnrealEd/Private/EditorDomain/EditorDomainArchive.cpp
--- a/Engine/Source/Editor/UnrealEd/Private/EditorDomain/EditorDomainArchive.cpp
+++ b/Engine/Source/Editor/UnrealEd/Private/EditorDomain/EditorDomainArchive.cpp
@@ -407,7 +407,7 @@ void FEditorDomainPackageSegments::SendSegmentRequest(FSegment& Segment)
 
 	// Note that Segment.RequestOwner is Interface-only and so we can write it outside the lock
 	Segment.RequestOwner.Emplace(EPriority::Normal);
-	FCacheChunkRequest SegmentChunk{PackagePath.GetDebugName(), UE::EditorDomain::GetEditorDomainPackageKey(EditorDomainHash), Segment.PayloadId};
+	FCacheChunkRequest SegmentChunk{{PackagePath.GetDebugName()}, UE::EditorDomain::GetEditorDomainPackageKey(EditorDomainHash), Segment.PayloadId};
 	SegmentChunk.Policy = ECachePolicy::Local;
 	Cache.GetChunks({SegmentChunk}, *Segment.RequestOwner,
 		[this, &Segment](FCacheChunkCompleteParams&& Params)
