DDC: Changed FCacheRecord::Load to take the record as an object instead of a view

#rnx

diff --git a/Engine/Source/Developer/DerivedDataCache/Private/DerivedDataCacheRecord.cpp b/Engine/Source/Developer/DerivedDataCache/Private/DerivedDataCacheRecord.cpp
--- a/Engine/Source/Developer/DerivedDataCache/Private/DerivedDataCacheRecord.cpp
+++ b/Engine/Source/Developer/DerivedDataCache/Private/DerivedDataCacheRecord.cpp
@@ -270,10 +270,12 @@ FCbPackage FCacheRecord::Save() const
 	return Package;
 }
 
-FOptionalCacheRecord FCacheRecord::Load(const FCbPackage& Package, FCbObjectView RecordObject)
+FOptionalCacheRecord FCacheRecord::Load(const FCbPackage& Attachments, const FCbObject& Object)
 {
+	const FCbObjectView ObjectView = Object;
+
 	FCacheKey Key;
-	FCbObjectView KeyObject = RecordObject["Key"_ASV].AsObjectView();
+	FCbObjectView KeyObject = ObjectView["Key"_ASV].AsObjectView();
 	auto TrySetBucketName = [](FUtf8StringView Name, FCacheKey& Key)
 	{
 		if (Private::IsValidCacheBucketName(Name))
@@ -291,9 +293,9 @@ FOptionalCacheRecord FCacheRecord::Load(const FCbPackage& Package, FCbObjectView
 
 	FCacheRecordBuilder Builder(Key);
 
-	Builder.SetMeta(FCbObject::Clone(RecordObject["Meta"_ASV].AsObjectView()));
+	Builder.SetMeta(Object["Meta"_ASV].AsObject());
 
-	auto LoadPayload = [&Package](const FCbObjectView& PayloadObject)
+	auto LoadPayload = [&Attachments](const FCbObjectView& PayloadObject)
 	{
 		const FPayloadId PayloadId = PayloadObject["Id"_ASV].AsObjectId();
 		if (PayloadId.IsNull())
@@ -301,7 +303,7 @@ FOptionalCacheRecord FCacheRecord::Load(const FCbPackage& Package, FCbObjectView
 			return FPayload();
 		}
 		const FIoHash RawHash = PayloadObject["RawHash"_ASV].AsHash();
-		if (const FCbAttachment* Attachment = Package.FindAttachment(RawHash))
+		if (const FCbAttachment* Attachment = Attachments.FindAttachment(RawHash))
 		{
 			if (const FCompressedBuffer& Compressed = Attachment->AsCompressedBinary())
 			{
@@ -319,7 +321,7 @@ FOptionalCacheRecord FCacheRecord::Load(const FCbPackage& Package, FCbObjectView
 		}
 	};
 
-	if (FCbObjectView ValueObject = RecordObject["Value"_ASV].AsObjectView())
+	if (FCbObjectView ValueObject = ObjectView["Value"_ASV].AsObjectView())
 	{
 		FPayload Value = LoadPayload(ValueObject);
 		if (!Value)
@@ -329,7 +331,7 @@ FOptionalCacheRecord FCacheRecord::Load(const FCbPackage& Package, FCbObjectView
 		Builder.SetValue(Value);
 	}
 
-	for (FCbFieldView AttachmentField : RecordObject["Attachments"_ASV])
+	for (FCbFieldView AttachmentField : ObjectView["Attachments"_ASV])
 	{
 		FPayload Attachment = LoadPayload(AttachmentField.AsObjectView());
 		if (!Attachment)
diff --git a/Engine/Source/Developer/DerivedDataCache/Private/ZenDerivedDataBackend.cpp b/Engine/Source/Developer/DerivedDataCache/Private/ZenDerivedDataBackend.cpp
--- a/Engine/Source/Developer/DerivedDataCache/Private/ZenDerivedDataBackend.cpp
+++ b/Engine/Source/Developer/DerivedDataCache/Private/ZenDerivedDataBackend.cpp
@@ -703,14 +703,14 @@ void FZenDerivedDataBackend::Get(
 			const FCbObject& ResponseObj = BatchResponse.GetObject();
 			
 			int32 KeyIndex = BatchFirst;
-			for (FCbFieldView RecordView : ResponseObj["Result"_ASV])
+			for (FCbField RecordField : ResponseObj["Result"_ASV])
 			{
 				const FCacheGetRequest& Request = Requests[KeyIndex++];
 				const FCacheKey& Key = Request.Key;
 				
 				FOptionalCacheRecord Record;
 
-				if (!RecordView.IsNull())
+				if (!RecordField.IsNull())
 				{
 					if (ShouldSimulateMiss(Key))
 					{
@@ -719,7 +719,7 @@ void FZenDerivedDataBackend::Get(
 					}
 					else
 					{
-						Record = FCacheRecord::Load(BatchResponse, RecordView.AsObjectView()); 
+						Record = FCacheRecord::Load(BatchResponse, RecordField.AsObject());
 					}
 				}
 
diff --git a/Engine/Source/Developer/DerivedDataCache/Public/DerivedDataCacheRecord.h b/Engine/Source/Developer/DerivedDataCache/Public/DerivedDataCacheRecord.h
--- a/Engine/Source/Developer/DerivedDataCache/Public/DerivedDataCacheRecord.h
+++ b/Engine/Source/Developer/DerivedDataCache/Public/DerivedDataCacheRecord.h
@@ -14,7 +14,6 @@
 
 class FCbObject;
 class FCbPackage;
-class FCbObjectView;
 
 namespace UE::DerivedData { class FCacheRecord; }
 namespace UE::DerivedData { class FOptionalCacheRecord; }
@@ -113,7 +112,7 @@ public:
 
 	/** Load a cache record from a compact binary package. Null on error. */
 	UE_API static FOptionalCacheRecord Load(const FCbPackage& Package);
-	UE_API static FOptionalCacheRecord Load(const FCbPackage& Package, FCbObjectView RecordObject);
+	UE_API static FOptionalCacheRecord Load(const FCbPackage& Attachments, const FCbObject& Object);
 
 private:
 	friend class FOptionalCacheRecord;
