DDC: Split FPayload into separate FValue and FValueId types

A payload was conceptually a value with an ID. That has been formalized by removing the ID from the payload and having separate FValue and FValueId types. This separation cleans up the API in a few areas, and provides a more natural path to providing a basic key/value API.

#rnx

diff --git a/Engine/Source/Developer/DerivedDataCache/Private/DerivedDataBackendAsyncPutWrapper.cpp b/Engine/Source/Developer/DerivedDataCache/Private/DerivedDataBackendAsyncPutWrapper.cpp
--- a/Engine/Source/Developer/DerivedDataCache/Private/DerivedDataBackendAsyncPutWrapper.cpp
+++ b/Engine/Source/Developer/DerivedDataCache/Private/DerivedDataBackendAsyncPutWrapper.cpp
@@ -2,11 +2,12 @@
 
 #include "DerivedDataBackendAsyncPutWrapper.h"
 
-#include "DerivedDataPayload.h"
 #include "DerivedDataRequest.h"
 #include "DerivedDataRequestOwner.h"
+#include "DerivedDataValueId.h"
 #include "Experimental/Async/LazyEvent.h"
 #include "MemoryDerivedDataBackend.h"
+#include "Memory/SharedBuffer.h"
 #include "Tasks/Task.h"
 
 namespace UE::DerivedData::Backends
diff --git a/Engine/Source/Developer/DerivedDataCache/Private/DerivedDataBuild.cpp b/Engine/Source/Developer/DerivedDataCache/Private/DerivedDataBuild.cpp
--- a/Engine/Source/Developer/DerivedDataCache/Private/DerivedDataBuild.cpp
+++ b/Engine/Source/Developer/DerivedDataCache/Private/DerivedDataBuild.cpp
@@ -38,7 +38,7 @@ DEFINE_LOG_CATEGORY(LogDerivedDataBuild);
  * - Key->InputBuffer
  * - From FBuildInputsBuilder via IBuild::CreateInputs()
  * FBuildOutput:
- * - Metadata, Payloads[], Diagnostics[] (Level, Category, Message)
+ * - Metadata, Values[], Messages[] (Message, Level), Logs[] (Level, Category, Message)
  * - From FBuildOutputBuilder via IBuild::CreateOutput()
  * - Serializes to/from FCbObject and FCacheRecord
  * FBuildKey:
@@ -152,48 +152,48 @@ public:
 		}
 	}
 
-	inline TConstArrayView<FBuildPayloadPolicy> GetPayloadPolicies() const final
+	inline TConstArrayView<FBuildValuePolicy> GetValuePolicies() const final
 	{
-		return Payloads;
+		return Values;
 	}
 
-	inline void AddPayloadPolicy(const FBuildPayloadPolicy& Policy) final
+	inline void AddValuePolicy(const FBuildValuePolicy& Policy) final
 	{
-		Payloads.Add(Policy);
+		Values.Add(Policy);
 	}
 
 	inline void Build() final
 	{
-		Algo::SortBy(Payloads, &FBuildPayloadPolicy::Id);
+		Algo::SortBy(Values, &FBuildValuePolicy::Id);
 	}
 
 private:
-	TArray<FBuildPayloadPolicy, TInlineAllocator<14>> Payloads;
+	TArray<FBuildValuePolicy, TInlineAllocator<14>> Values;
 	mutable std::atomic<uint32> ReferenceCount{0};
 };
 
-EBuildPolicy FBuildPolicy::GetPayloadPolicy(const FPayloadId& Id) const
+EBuildPolicy FBuildPolicy::GetValuePolicy(const FValueId& Id) const
 {
 	if (Shared)
 	{
-		if (TConstArrayView<FBuildPayloadPolicy> Payloads = Shared->GetPayloadPolicies(); !Payloads.IsEmpty())
+		if (TConstArrayView<FBuildValuePolicy> Values = Shared->GetValuePolicies(); !Values.IsEmpty())
 		{
-			if (int32 Index = Algo::BinarySearchBy(Payloads, Id, &FBuildPayloadPolicy::Id); Index != INDEX_NONE)
+			if (int32 Index = Algo::BinarySearchBy(Values, Id, &FBuildValuePolicy::Id); Index != INDEX_NONE)
 			{
-				return Payloads[Index].Policy;
+				return Values[Index].Policy;
 			}
 		}
 	}
-	return DefaultPayloadPolicy;
+	return DefaultValuePolicy;
 }
 
-void FBuildPolicyBuilder::AddPayloadPolicy(const FBuildPayloadPolicy& Policy)
+void FBuildPolicyBuilder::AddValuePolicy(const FBuildValuePolicy& Policy)
 {
 	if (!Shared)
 	{
 		Shared = new Private::FBuildPolicyShared;
 	}
-	Shared->AddPayloadPolicy(Policy);
+	Shared->AddValuePolicy(Policy);
 }
 
 FBuildPolicy FBuildPolicyBuilder::Build()
@@ -203,8 +203,8 @@ FBuildPolicy FBuildPolicyBuilder::Build()
 	{
 		Shared->Build();
 		const auto PolicyOr = [](EBuildPolicy A, EBuildPolicy B) { return A | (B & EBuildPolicy::Default); };
-		const TConstArrayView<FBuildPayloadPolicy> Payloads = Shared->GetPayloadPolicies();
-		Policy.CombinedPolicy = Algo::TransformAccumulate(Payloads, &FBuildPayloadPolicy::Policy, BasePolicy, PolicyOr);
+		const TConstArrayView<FBuildValuePolicy> Values = Shared->GetValuePolicies();
+		Policy.CombinedPolicy = Algo::TransformAccumulate(Values, &FBuildValuePolicy::Policy, BasePolicy, PolicyOr);
 		Policy.Shared = MoveTemp(Shared);
 	}
 	return Policy;
diff --git a/Engine/Source/Developer/DerivedDataCache/Private/DerivedDataBuildDefinition.cpp b/Engine/Source/Developer/DerivedDataCache/Private/DerivedDataBuildDefinition.cpp
--- a/Engine/Source/Developer/DerivedDataCache/Private/DerivedDataBuildDefinition.cpp
+++ b/Engine/Source/Developer/DerivedDataCache/Private/DerivedDataBuildDefinition.cpp
@@ -36,9 +36,9 @@ public:
 		Add<FCbObject>(Key, Value);
 	}
 
-	void AddInputBuild(FStringView Key, const FBuildPayloadKey& PayloadKey) final
+	void AddInputBuild(FStringView Key, const FBuildValueKey& ValueKey) final
 	{
-		Add<FBuildPayloadKey>(Key, PayloadKey);
+		Add<FBuildValueKey>(Key, ValueKey);
 	}
 
 	void AddInputBulkData(FStringView Key, const FGuid& BulkDataId) final
@@ -58,7 +58,7 @@ public:
 
 	FBuildDefinition Build() final;
 
-	using InputType = TVariant<FCbObject, FBuildPayloadKey, FGuid, FString, FIoHash>;
+	using InputType = TVariant<FCbObject, FBuildValueKey, FGuid, FString, FIoHash>;
 
 	FString Name;
 	FString Function;
@@ -95,7 +95,7 @@ public:
 	bool HasInputs() const final;
 
 	void IterateConstants(TFunctionRef<void (FStringView Key, FCbObject&& Value)> Visitor) const final;
-	void IterateInputBuilds(TFunctionRef<void (FStringView Key, const FBuildPayloadKey& PayloadKey)> Visitor) const final;
+	void IterateInputBuilds(TFunctionRef<void (FStringView Key, const FBuildValueKey& ValueKey)> Visitor) const final;
 	void IterateInputBulkData(TFunctionRef<void (FStringView Key, const FGuid& BulkDataId)> Visitor) const final;
 	void IterateInputFiles(TFunctionRef<void (FStringView Key, FStringView Path)> Visitor) const final;
 	void IterateInputHashes(TFunctionRef<void (FStringView Key, const FIoHash& RawHash)> Visitor) const final;
@@ -143,7 +143,7 @@ FBuildDefinitionInternal::FBuildDefinitionInternal(FBuildDefinitionBuilderIntern
 		{
 			bHasConstants = true;
 		}
-		else if (Pair.Value.IsType<FBuildPayloadKey>())
+		else if (Pair.Value.IsType<FBuildValueKey>())
 		{
 			bHasBuilds = true;
 		}
@@ -190,12 +190,12 @@ FBuildDefinitionInternal::FBuildDefinitionInternal(FBuildDefinitionBuilderIntern
 		Writer.BeginObject("Builds"_ASV);
 		for (const TPair<FString, FBuildDefinitionBuilderInternal::InputType>& Pair : DefinitionBuilder.Inputs)
 		{
-			if (Pair.Value.IsType<FBuildPayloadKey>())
+			if (Pair.Value.IsType<FBuildValueKey>())
 			{
-				const FBuildPayloadKey& PayloadKey = Pair.Value.Get<FBuildPayloadKey>();
+				const FBuildValueKey& ValueKey = Pair.Value.Get<FBuildValueKey>();
 				Writer.BeginObject(FTCHARToUTF8(Pair.Key));
-				Writer.AddHash("Build"_ASV, PayloadKey.BuildKey.Hash);
-				Writer.AddObjectId("Payload"_ASV, PayloadKey.Id);
+				Writer.AddHash("Build"_ASV, ValueKey.BuildKey.Hash);
+				Writer.AddObjectId("Payload"_ASV, ValueKey.Id);
 				Writer.EndObject();
 			}
 		}
@@ -295,13 +295,13 @@ void FBuildDefinitionInternal::IterateConstants(TFunctionRef<void (FStringView K
 	}
 }
 
-void FBuildDefinitionInternal::IterateInputBuilds(TFunctionRef<void (FStringView Key, const FBuildPayloadKey& PayloadKey)> Visitor) const
+void FBuildDefinitionInternal::IterateInputBuilds(TFunctionRef<void (FStringView Key, const FBuildValueKey& ValueKey)> Visitor) const
 {
 	for (FCbFieldView Field : Definition.AsView()["Inputs"_ASV]["Builds"_ASV])
 	{
 		const FBuildKey BuildKey{Field["Build"_ASV].AsHash()};
-		const FPayloadId Id = Field["Payload"_ASV].AsObjectId();
-		Visitor(FUTF8ToTCHAR(Field.GetName()), FBuildPayloadKey{BuildKey, Id});
+		const FValueId Id = Field["Payload"_ASV].AsObjectId();
+		Visitor(FUTF8ToTCHAR(Field.GetName()), FBuildValueKey{BuildKey, Id});
 	}
 }
 
diff --git a/Engine/Source/Developer/DerivedDataCache/Private/DerivedDataBuildJob.cpp b/Engine/Source/Developer/DerivedDataCache/Private/DerivedDataBuildJob.cpp
--- a/Engine/Source/Developer/DerivedDataCache/Private/DerivedDataBuildJob.cpp
+++ b/Engine/Source/Developer/DerivedDataCache/Private/DerivedDataBuildJob.cpp
@@ -21,7 +21,6 @@
 #include "DerivedDataBuildWorkerRegistry.h"
 #include "DerivedDataCache.h"
 #include "DerivedDataCacheRecord.h"
-#include "DerivedDataPayload.h"
 #include "DerivedDataRequestOwner.h"
 #include "DerivedDataSharedString.h"
 #include "HAL/CriticalSection.h"
@@ -487,10 +486,10 @@ static ECachePolicy MakeCacheQueryPolicy(EBuildPolicy BuildPolicy, const FBuildJ
 
 static FCacheRecordPolicy MakeCacheRecordQueryPolicy(const FBuildPolicy& BuildPolicy, const FBuildJobContext& Context)
 {
-	FCacheRecordPolicyBuilder Builder(MakeCacheQueryPolicy(BuildPolicy.GetDefaultPayloadPolicy(), Context));
-	for (const FBuildPayloadPolicy& PayloadPolicy : BuildPolicy.GetPayloadPolicies())
+	FCacheRecordPolicyBuilder Builder(MakeCacheQueryPolicy(BuildPolicy.GetDefaultValuePolicy(), Context));
+	for (const FBuildValuePolicy& ValuePolicy : BuildPolicy.GetValuePolicies())
 	{
-		Builder.AddPayloadPolicy(PayloadPolicy.Id, MakeCacheQueryPolicy(PayloadPolicy.Policy, Context));
+		Builder.AddValuePolicy(ValuePolicy.Id, MakeCacheQueryPolicy(ValuePolicy.Policy, Context));
 	}
 	return Builder.Build();
 }
@@ -806,10 +805,10 @@ void FBuildJob::BeginExecuteRemote()
 	}
 	else
 	{
-		FBuildPolicyBuilder PolicyBuilder(BuildPolicy.GetDefaultPayloadPolicy() & Mask);
-		for (const FBuildPayloadPolicy& Payload : BuildPolicy.GetPayloadPolicies())
+		FBuildPolicyBuilder PolicyBuilder(BuildPolicy.GetDefaultValuePolicy() & Mask);
+		for (const FBuildValuePolicy& Value : BuildPolicy.GetValuePolicies())
 		{
-			PolicyBuilder.AddPayloadPolicy({Payload.Id, Payload.Policy & Mask});
+			PolicyBuilder.AddValuePolicy({Value.Id, Value.Policy & Mask});
 		}
 		RemoteBuildPolicy = PolicyBuilder.Build();
 	}
diff --git a/Engine/Source/Developer/DerivedDataCache/Private/DerivedDataBuildJobContext.cpp b/Engine/Source/Developer/DerivedDataCache/Private/DerivedDataBuildJobContext.cpp
--- a/Engine/Source/Developer/DerivedDataCache/Private/DerivedDataBuildJobContext.cpp
+++ b/Engine/Source/Developer/DerivedDataCache/Private/DerivedDataBuildJobContext.cpp
@@ -8,8 +8,8 @@
 #include "DerivedDataBuildPrivate.h"
 #include "DerivedDataBuildTypes.h"
 #include "DerivedDataCache.h"
-#include "DerivedDataPayload.h"
 #include "DerivedDataRequestOwner.h"
+#include "DerivedDataValue.h"
 #include "Hash/Blake3.h"
 #include "Misc/StringBuilder.h"
 #include "UObject/NameTypes.h"
@@ -100,29 +100,29 @@ FSharedBuffer FBuildJobContext::FindInput(FStringView Key) const
 	return FSharedBuffer();
 }
 
-void FBuildJobContext::AddPayload(const FPayload& Payload)
+void FBuildJobContext::AddValue(const FValueId& Id, const FValue& Value)
 {
-	OutputBuilder.AddPayload(Payload);
+	OutputBuilder.AddValue(Id, Value);
 }
 
-void FBuildJobContext::AddPayload(const FPayloadId& Id, const FCompressedBuffer& Buffer)
+void FBuildJobContext::AddValue(const FValueId& Id, const FCompressedBuffer& Buffer)
 {
-	AddPayload(FPayload(Id, Buffer));
+	AddValue(Id, FValue(Buffer));
 }
 
-void FBuildJobContext::AddPayload(const FPayloadId& Id, const FCompositeBuffer& Buffer, const uint64 BlockSize)
+void FBuildJobContext::AddValue(const FValueId& Id, const FCompositeBuffer& Buffer, const uint64 BlockSize)
 {
-	AddPayload(FPayload(Id, FPayload::Compress(Buffer, BlockSize)));
+	AddValue(Id, FValue::Compress(Buffer, BlockSize));
 }
 
-void FBuildJobContext::AddPayload(const FPayloadId& Id, const FSharedBuffer& Buffer, const uint64 BlockSize)
+void FBuildJobContext::AddValue(const FValueId& Id, const FSharedBuffer& Buffer, const uint64 BlockSize)
 {
-	AddPayload(FPayload(Id, FPayload::Compress(Buffer, BlockSize)));
+	AddValue(Id, FValue::Compress(Buffer, BlockSize));
 }
 
-void FBuildJobContext::AddPayload(const FPayloadId& Id, const FCbObject& Object)
+void FBuildJobContext::AddValue(const FValueId& Id, const FCbObject& Object)
 {
-	AddPayload(FPayload(Id, FPayload::Compress(Object.GetBuffer())));
+	AddValue(Id, FValue::Compress(Object.GetBuffer()));
 }
 
 void FBuildJobContext::BeginBuild(IRequestOwner& InOwner, TUniqueFunction<void ()>&& InOnEndBuild)
diff --git a/Engine/Source/Developer/DerivedDataCache/Private/DerivedDataBuildJobContext.h b/Engine/Source/Developer/DerivedDataCache/Private/DerivedDataBuildJobContext.h
--- a/Engine/Source/Developer/DerivedDataCache/Private/DerivedDataBuildJobContext.h
+++ b/Engine/Source/Developer/DerivedDataCache/Private/DerivedDataBuildJobContext.h
@@ -15,7 +15,6 @@
 #include "Templates/Function.h"
 
 namespace UE::DerivedData { class FBuildOutputBuilder; }
-namespace UE::DerivedData { class FPayload; }
 namespace UE::DerivedData { class IBuildJob; }
 namespace UE::DerivedData { class IRequestOwner; }
 namespace UE::DerivedData { enum class EBuildPolicy : uint32; }
@@ -58,11 +57,11 @@ private:
 	FCbObject FindConstant(FStringView Key) const final;
 	FSharedBuffer FindInput(FStringView Key) const final;
 
-	void AddPayload(const FPayload& Payload) final;
-	void AddPayload(const FPayloadId& Id, const FCompressedBuffer& Buffer) final;
-	void AddPayload(const FPayloadId& Id, const FCompositeBuffer& Buffer, uint64 BlockSize) final;
-	void AddPayload(const FPayloadId& Id, const FSharedBuffer& Buffer, uint64 BlockSize) final;
-	void AddPayload(const FPayloadId& Id, const FCbObject& Object) final;
+	void AddValue(const FValueId& Id, const FValue& Value) final;
+	void AddValue(const FValueId& Id, const FCompressedBuffer& Buffer) final;
+	void AddValue(const FValueId& Id, const FCompositeBuffer& Buffer, uint64 BlockSize) final;
+	void AddValue(const FValueId& Id, const FSharedBuffer& Buffer, uint64 BlockSize) final;
+	void AddValue(const FValueId& Id, const FCbObject& Object) final;
 
 	void EndBuild();
 	void BeginAsyncBuild() final;
diff --git a/Engine/Source/Developer/DerivedDataCache/Private/DerivedDataBuildOutput.cpp b/Engine/Source/Developer/DerivedDataCache/Private/DerivedDataBuildOutput.cpp
--- a/Engine/Source/Developer/DerivedDataCache/Private/DerivedDataBuildOutput.cpp
+++ b/Engine/Source/Developer/DerivedDataCache/Private/DerivedDataBuildOutput.cpp
@@ -9,7 +9,7 @@
 #include "DerivedDataBuildKey.h"
 #include "DerivedDataBuildPrivate.h"
 #include "DerivedDataCacheRecord.h"
-#include "DerivedDataPayload.h"
+#include "DerivedDataValue.h"
 #include "Misc/StringBuilder.h"
 #include "Serialization/CompactBinary.h"
 #include "Serialization/CompactBinaryWriter.h"
@@ -91,7 +91,7 @@ public:
 	~FBuildOutputBuilderInternal() final = default;
 
 	void SetMeta(FCbObject&& InMeta) final { Meta = MoveTemp(InMeta); Meta.MakeOwned(); }
-	void AddPayload(const FPayload& Payload) final;
+	void AddValue(const FValueId& Id, const FValue& Value) final;
 	void AddMessage(const FBuildOutputMessage& Message) final;
 	void AddLog(const FBuildOutputLog& Log) final;
 	bool HasError() const final { return bHasError; }
@@ -100,7 +100,7 @@ public:
 	FString Name;
 	FString Function;
 	FCbObject Meta;
-	TArray<FPayload> Payloads;
+	TArray<FValueWithId> Values;
 	FCbWriter MessageWriter;
 	FCbWriter LogWriter;
 	bool bHasMessages = false;
@@ -111,7 +111,7 @@ public:
 class FBuildOutputInternal final : public IBuildOutputInternal
 {
 public:
-	FBuildOutputInternal(FString&& Name, FString&& Function, FCbObject&& Meta, FCbObject&& Output, TArray<FPayload>&& Payloads);
+	FBuildOutputInternal(FString&& Name, FString&& Function, FCbObject&& Meta, FCbObject&& Output, TArray<FValueWithId>&& Values);
 	FBuildOutputInternal(FStringView Name, FStringView Function, const FCbObject& Output, bool& bOutIsValid);
 	FBuildOutputInternal(FStringView Name, FStringView Function, const FCacheRecord& Output, bool& bOutIsValid);
 
@@ -122,8 +122,8 @@ public:
 
 	const FCbObject& GetMeta() const final { return Meta; }
 
-	const FPayload& GetPayload(const FPayloadId& Id) const final;
-	TConstArrayView<FPayload> GetPayloads() const final { return Payloads; }
+	const FValueWithId& GetValue(const FValueId& Id) const final;
+	TConstArrayView<FValueWithId> GetValues() const final { return Values; }
 	TConstArrayView<FBuildOutputMessage> GetMessages() const final { return Messages; }
 	TConstArrayView<FBuildOutputLog> GetLogs() const final { return Logs; }
 	bool HasLogs() const final { return !Logs.IsEmpty(); }
@@ -152,7 +152,7 @@ private:
 	FString Function;
 	FCbObject Meta;
 	FCbObject Output;
-	TArray<FPayload> Payloads;
+	TArray<FValueWithId> Values;
 	TArray<FBuildOutputMessage> Messages;
 	TArray<FBuildOutputLog> Logs;
 	mutable std::atomic<uint32> ReferenceCount{0};
@@ -165,12 +165,12 @@ FBuildOutputInternal::FBuildOutputInternal(
 	FString&& InFunction,
 	FCbObject&& InMeta,
 	FCbObject&& InOutput,
-	TArray<FPayload>&& InPayloads)
+	TArray<FValueWithId>&& InValues)
 	: Name(MoveTemp(InName))
 	, Function(MoveTemp(InFunction))
 	, Meta(MoveTemp(InMeta))
 	, Output(MoveTemp(InOutput))
-	, Payloads(MoveTemp(InPayloads))
+	, Values(MoveTemp(InValues))
 {
 	Meta.MakeOwned();
 	Output.MakeOwned();
@@ -195,17 +195,17 @@ FBuildOutputInternal::FBuildOutputInternal(FStringView InName, FStringView InFun
 	, Function(InFunction)
 	, Meta(InOutput.GetMeta())
 	, Output(InOutput.GetValue().GetData().Decompress())
-	, Payloads(InOutput.GetAttachmentPayloads())
+	, Values(InOutput.GetAttachments())
 {
 	checkf(!Name.IsEmpty(), TEXT("A build output requires a non-empty name."));
 	AssertValidBuildFunctionName(Function, Name);
 	bOutIsValid = TryLoad();
 }
 
-const FPayload& FBuildOutputInternal::GetPayload(const FPayloadId& Id) const
+const FValueWithId& FBuildOutputInternal::GetValue(const FValueId& Id) const
 {
-	const int32 Index = Algo::BinarySearchBy(Payloads, Id, &FPayload::GetId);
-	return Payloads.IsValidIndex(Index) ? Payloads[Index] : FPayload::Null;
+	const int32 Index = Algo::BinarySearchBy(Values, Id, &FValueWithId::GetId);
+	return Values.IsValidIndex(Index) ? Values[Index] : FValueWithId::Null;
 }
 
 bool FBuildOutputInternal::HasError() const
@@ -218,18 +218,18 @@ bool FBuildOutputInternal::TryLoad()
 {
 	const FCbObjectView OutputView = Output;
 
-	if (Payloads.IsEmpty())
+	if (Values.IsEmpty())
 	{
-		for (FCbFieldView Payload : OutputView["Payloads"_ASV])
+		for (FCbFieldView Value : OutputView["Payloads"_ASV])
 		{
-			const FPayloadId Id = Payload["Id"_ASV].AsObjectId();
-			const FIoHash& RawHash = Payload["RawHash"_ASV].AsAttachment();
-			const uint64 RawSize = Payload["RawSize"_ASV].AsUInt64(MAX_uint64);
+			const FValueId Id = Value["Id"_ASV].AsObjectId();
+			const FIoHash& RawHash = Value["RawHash"_ASV].AsAttachment();
+			const uint64 RawSize = Value["RawSize"_ASV].AsUInt64(MAX_uint64);
 			if (Id.IsNull() || RawHash.IsZero() || RawSize == MAX_uint64)
 			{
 				return false;
 			}
-			Payloads.Emplace(Id, RawHash, RawSize);
+			Values.Emplace(Id, RawHash, RawSize);
 		}
 	}
 
@@ -280,15 +280,15 @@ bool FBuildOutputInternal::TryLoad()
 void FBuildOutputInternal::Save(FCbWriter& Writer) const
 {
 	Writer.BeginObject();
-	if (!Payloads.IsEmpty())
+	if (!Values.IsEmpty())
 	{
 		Writer.BeginArray("Payloads"_ASV);
-		for (const FPayload& Payload : Payloads)
+		for (const FValueWithId& Value : Values)
 		{
 			Writer.BeginObject();
-			Writer.AddObjectId("Id"_ASV, Payload.GetId());
-			Writer.AddBinaryAttachment("RawHash"_ASV, Payload.GetRawHash());
-			Writer.AddInteger("RawSize"_ASV, Payload.GetRawSize());
+			Writer.AddObjectId("Id"_ASV, Value.GetId());
+			Writer.AddBinaryAttachment("RawHash"_ASV, Value.GetRawHash());
+			Writer.AddInteger("RawSize"_ASV, Value.GetRawSize());
 			Writer.EndObject();
 		}
 		Writer.EndArray();
@@ -324,22 +324,21 @@ void FBuildOutputInternal::Save(FCacheRecordBuilder& RecordBuilder) const
 		}
 		RecordBuilder.SetValue(Writer.Save().GetBuffer());
 	}
-	for (const FPayload& Payload : Payloads)
+	for (const FValueWithId& Value : Values)
 	{
-		RecordBuilder.AddAttachment(Payload);
+		RecordBuilder.AddAttachment(Value);
 	}
 }
 
 ///////////////////////////////////////////////////////////////////////////////////////////////////////////////////////
 
-void FBuildOutputBuilderInternal::AddPayload(const FPayload& Payload)
+void FBuildOutputBuilderInternal::AddValue(const FValueId& Id, const FValue& Value)
 {
-	checkf(Payload, TEXT("Null payload added in output for build of '%s' by %s."), *Name, *Function);
-	const FPayloadId& Id = Payload.GetId();
-	const int32 Index = Algo::LowerBoundBy(Payloads, Id, &FPayload::GetId);
-	checkf(!(Payloads.IsValidIndex(Index) && Payloads[Index].GetId() == Id),
-		TEXT("Duplicate ID %s used by payload for build of '%s' by %s."), *WriteToString<32>(Id), *Name, *Function);
-	Payloads.Insert(Payload, Index);
+	checkf(Id.IsValid(), TEXT("Null value added in output for build of '%s' by %s."), *Name, *Function);
+	const int32 Index = Algo::LowerBoundBy(Values, Id, &FValueWithId::GetId);
+	checkf(!(Values.IsValidIndex(Index) && Values[Index].GetId() == Id),
+		TEXT("Duplicate ID %s used by value for build of '%s' by %s."), *WriteToString<32>(Id), *Name, *Function);
+	Values.Insert(FValueWithId(Id, Value), Index);
 }
 
 void FBuildOutputBuilderInternal::AddMessage(const FBuildOutputMessage& Message)
@@ -367,7 +366,7 @@ FBuildOutput FBuildOutputBuilderInternal::Build()
 {
 	if (bHasError)
 	{
-		Payloads.Empty();
+		Values.Empty();
 	}
 	MessageWriter.EndArray();
 	LogWriter.EndArray();
@@ -389,7 +388,7 @@ FBuildOutput FBuildOutputBuilderInternal::Build()
 		Writer.EndObject();
 		Output = Writer.Save().AsObject();
 	}
-	return CreateBuildOutput(new FBuildOutputInternal(MoveTemp(Name), MoveTemp(Function), MoveTemp(Meta), MoveTemp(Output), MoveTemp(Payloads)));
+	return CreateBuildOutput(new FBuildOutputInternal(MoveTemp(Name), MoveTemp(Function), MoveTemp(Meta), MoveTemp(Output), MoveTemp(Values)));
 }
 
 ///////////////////////////////////////////////////////////////////////////////////////////////////////////////////////
diff --git a/Engine/Source/Developer/DerivedDataCache/Private/DerivedDataBuildSession.cpp b/Engine/Source/Developer/DerivedDataCache/Private/DerivedDataBuildSession.cpp
--- a/Engine/Source/Developer/DerivedDataCache/Private/DerivedDataBuildSession.cpp
+++ b/Engine/Source/Developer/DerivedDataCache/Private/DerivedDataBuildSession.cpp
@@ -8,7 +8,7 @@
 #include "DerivedDataBuildJob.h"
 #include "DerivedDataBuildOutput.h"
 #include "DerivedDataBuildPrivate.h"
-#include "DerivedDataPayload.h"
+#include "DerivedDataValue.h"
 
 namespace UE::DerivedData::Private
 {
@@ -46,11 +46,11 @@ public:
 		IRequestOwner& Owner,
 		FOnBuildComplete&& OnComplete) final;
 
-	void BuildPayload(
-		const FBuildPayloadKey& Payload,
+	void BuildValue(
+		const FBuildValueKey& Value,
 		EBuildPolicy Policy,
 		IRequestOwner& Owner,
-		FOnBuildPayloadComplete&& OnComplete) final;
+		FOnBuildValueComplete&& OnComplete) final;
 
 	FString Name;
 	ICache& Cache;
@@ -83,33 +83,33 @@ void FBuildSessionInternal::Build(
 		OnComplete ? MoveTemp(OnComplete) : [](FBuildCompleteParams&&){});
 }
 
-void FBuildSessionInternal::BuildPayload(
-	const FBuildPayloadKey& PayloadKey,
+void FBuildSessionInternal::BuildValue(
+	const FBuildValueKey& ValueKey,
 	EBuildPolicy Policy,
 	IRequestOwner& Owner,
-	FOnBuildPayloadComplete&& OnComplete)
+	FOnBuildValueComplete&& OnComplete)
 {
-	// This requests the entire output to get one payload. It will be optimized later to request only one payload.
+	// This requests the entire output to get one value. It will be optimized later to request only one value.
 	FOnBuildComplete OnJobComplete;
 	if (OnComplete)
 	{
-		OnJobComplete = [PayloadKey, OnComplete = MoveTemp(OnComplete)](FBuildCompleteParams&& Params)
+		OnJobComplete = [ValueKey, OnComplete = MoveTemp(OnComplete)](FBuildCompleteParams&& Params)
 		{
-			FPayload Payload;
+			FValueWithId Value;
 			EStatus Status = Params.Status;
 			if (Status == EStatus::Ok)
 			{
-				Payload = Params.Output.GetPayload(PayloadKey.Id);
-				Status = Payload ? EStatus::Ok : EStatus::Error;
+				Value = Params.Output.GetValue(ValueKey.Id);
+				Status = Value ? EStatus::Ok : EStatus::Error;
 			}
-			if (!Payload)
+			if (!Value)
 			{
-				Payload = FPayload(PayloadKey.Id);
+				Value = FValueWithId(ValueKey.Id);
 			}
-			OnComplete({PayloadKey.BuildKey, MoveTemp(Payload), Status});
+			OnComplete({ValueKey.BuildKey, MoveTemp(Value), Status});
 		};
 	}
-	CreateBuildJob({Cache, BuildSystem, Scheduler, InputResolver, Owner}, PayloadKey.BuildKey, Policy, MoveTemp(OnJobComplete));
+	CreateBuildJob({Cache, BuildSystem, Scheduler, InputResolver, Owner}, ValueKey.BuildKey, Policy, MoveTemp(OnJobComplete));
 }
 
 ///////////////////////////////////////////////////////////////////////////////////////////////////////////////////////
diff --git a/Engine/Source/Developer/DerivedDataCache/Private/DerivedDataCache.cpp b/Engine/Source/Developer/DerivedDataCache/Private/DerivedDataCache.cpp
--- a/Engine/Source/Developer/DerivedDataCache/Private/DerivedDataCache.cpp
+++ b/Engine/Source/Developer/DerivedDataCache/Private/DerivedDataCache.cpp
@@ -179,48 +179,48 @@ public:
 		}
 	}
 
-	inline TConstArrayView<FCachePayloadPolicy> GetPayloadPolicies() const final
+	inline TConstArrayView<FCacheValuePolicy> GetValuePolicies() const final
 	{
-		return Payloads;
+		return Values;
 	}
 
-	inline void AddPayloadPolicy(const FCachePayloadPolicy& Policy) final
+	inline void AddValuePolicy(const FCacheValuePolicy& Policy) final
 	{
-		Payloads.Add(Policy);
+		Values.Add(Policy);
 	}
 
 	inline void Build() final
 	{
-		Algo::SortBy(Payloads, &FCachePayloadPolicy::Id);
+		Algo::SortBy(Values, &FCacheValuePolicy::Id);
 	}
 
 private:
-	TArray<FCachePayloadPolicy, TInlineAllocator<14>> Payloads;
+	TArray<FCacheValuePolicy, TInlineAllocator<14>> Values;
 	mutable std::atomic<uint32> ReferenceCount{0};
 };
 
-ECachePolicy FCacheRecordPolicy::GetPayloadPolicy(const FPayloadId& Id) const
+ECachePolicy FCacheRecordPolicy::GetValuePolicy(const FValueId& Id) const
 {
 	if (Shared)
 	{
-		if (TConstArrayView<FCachePayloadPolicy> Payloads = Shared->GetPayloadPolicies(); !Payloads.IsEmpty())
+		if (TConstArrayView<FCacheValuePolicy> Values = Shared->GetValuePolicies(); !Values.IsEmpty())
 		{
-			if (int32 Index = Algo::BinarySearchBy(Payloads, Id, &FCachePayloadPolicy::Id); Index != INDEX_NONE)
+			if (int32 Index = Algo::BinarySearchBy(Values, Id, &FCacheValuePolicy::Id); Index != INDEX_NONE)
 			{
-				return Payloads[Index].Policy;
+				return Values[Index].Policy;
 			}
 		}
 	}
-	return DefaultPayloadPolicy;
+	return DefaultValuePolicy;
 }
 
-void FCacheRecordPolicyBuilder::AddPayloadPolicy(const FCachePayloadPolicy& Policy)
+void FCacheRecordPolicyBuilder::AddValuePolicy(const FCacheValuePolicy& Policy)
 {
 	if (!Shared)
 	{
 		Shared = new Private::FCacheRecordPolicyShared;
 	}
-	Shared->AddPayloadPolicy(Policy);
+	Shared->AddValuePolicy(Policy);
 }
 
 FCacheRecordPolicy FCacheRecordPolicyBuilder::Build()
@@ -230,8 +230,8 @@ FCacheRecordPolicy FCacheRecordPolicyBuilder::Build()
 	{
 		Shared->Build();
 		const auto PolicyOr = [](ECachePolicy A, ECachePolicy B) { return A | (B & ~ECachePolicy::SkipData); };
-		const TConstArrayView<FCachePayloadPolicy> Payloads = Shared->GetPayloadPolicies();
-		Policy.RecordPolicy = Algo::TransformAccumulate(Payloads, &FCachePayloadPolicy::Policy, BasePolicy, PolicyOr);
+		const TConstArrayView<FCacheValuePolicy> Values = Shared->GetValuePolicies();
+		Policy.RecordPolicy = Algo::TransformAccumulate(Values, &FCacheValuePolicy::Policy, BasePolicy, PolicyOr);
 		Policy.Shared = MoveTemp(Shared);
 	}
 	return Policy;
diff --git a/Engine/Source/Developer/DerivedDataCache/Private/DerivedDataCacheRecord.cpp b/Engine/Source/Developer/DerivedDataCache/Private/DerivedDataCacheRecord.cpp
--- a/Engine/Source/Developer/DerivedDataCache/Private/DerivedDataCacheRecord.cpp
+++ b/Engine/Source/Developer/DerivedDataCache/Private/DerivedDataCacheRecord.cpp
@@ -5,7 +5,7 @@
 #include "Algo/Accumulate.h"
 #include "DerivedDataCacheKey.h"
 #include "DerivedDataCachePrivate.h"
-#include "DerivedDataPayload.h"
+#include "DerivedDataValue.h"
 #include "Misc/ScopeExit.h"
 #include "Misc/StringBuilder.h"
 #include "Serialization/CompactBinary.h"
@@ -26,21 +26,23 @@ public:
 
 	void SetMeta(FCbObject&& Meta) final;
 
-	void SetValue(const FCompositeBuffer& Buffer, const FPayloadId& Id, uint64 BlockSize) final;
-	void SetValue(const FSharedBuffer& Buffer, const FPayloadId& Id, uint64 BlockSize) final;
-	void SetValue(const FPayload& Payload) final;
+	void SetValue(const FCompositeBuffer& Buffer, const FValueId& Id, uint64 BlockSize) final;
+	void SetValue(const FSharedBuffer& Buffer, const FValueId& Id, uint64 BlockSize) final;
+	void SetValue(const FValue& Value, const FValueId& Id) final;
+	void SetValue(const FValueWithId& Value) final;
 
-	void AddAttachment(const FCompositeBuffer& Buffer, const FPayloadId& Id, uint64 BlockSize) final;
-	void AddAttachment(const FSharedBuffer& Buffer, const FPayloadId& Id, uint64 BlockSize) final;
-	void AddAttachment(const FPayload& Payload) final;
+	void AddAttachment(const FCompositeBuffer& Buffer, const FValueId& Id, uint64 BlockSize) final;
+	void AddAttachment(const FSharedBuffer& Buffer, const FValueId& Id, uint64 BlockSize) final;
+	void AddAttachment(const FValue& Value, const FValueId& Id) final;
+	void AddAttachment(const FValueWithId& Value) final;
 
 	FCacheRecord Build() final;
 	void BuildAsync(IRequestOwner& Owner, FOnCacheRecordComplete&& OnComplete) final;
 
 	FCacheKey Key;
 	FCbObject Meta;
-	FPayload Value;
-	TArray<FPayload> Attachments;
+	FValueWithId Value;
+	TArray<FValueWithId> Attachments;
 };
 
 ///////////////////////////////////////////////////////////////////////////////////////////////////////////////////////
@@ -55,10 +57,10 @@ public:
 
 	const FCacheKey& GetKey() const final;
 	const FCbObject& GetMeta() const final;
-	const FPayload& GetValue() const final;
-	const FPayload& GetAttachment(const FPayloadId& Id) const final;
-	TConstArrayView<FPayload> GetAttachments() const final;
-	const FPayload& GetPayload(const FPayloadId& Id) const final;
+	const FValueWithId& GetValue() const final;
+	const FValueWithId& GetAttachment(const FValueId& Id) const final;
+	TConstArrayView<FValueWithId> GetAttachments() const final;
+	const FValueWithId& GetValue(const FValueId& Id) const final;
 
 	inline void AddRef() const final
 	{
@@ -76,16 +78,16 @@ public:
 private:
 	FCacheKey Key;
 	FCbObject Meta;
-	FPayload Value;
-	TArray<FPayload> Attachments;
+	FValueWithId Value;
+	TArray<FValueWithId> Attachments;
 	mutable std::atomic<uint32> ReferenceCount{0};
 };
 
 ///////////////////////////////////////////////////////////////////////////////////////////////////////////////////////
 
-static FPayloadId GetOrCreatePayloadId(const FPayloadId& Id, const FCompressedBuffer& Buffer)
+static FValueId GetOrCreateValueId(const FValueId& Id, const FValue& Value)
 {
-	return Id.IsValid() ? Id : FPayloadId::FromHash(Buffer.GetRawHash());
+	return Id.IsValid() ? Id : FValueId::FromHash(Value.GetRawHash());
 }
 
 ///////////////////////////////////////////////////////////////////////////////////////////////////////////////////////
@@ -108,23 +110,23 @@ const FCbObject& FCacheRecordInternal::GetMeta() const
 	return Meta;
 }
 
-const FPayload& FCacheRecordInternal::GetValue() const
+const FValueWithId& FCacheRecordInternal::GetValue() const
 {
 	return Value;
 }
 
-const FPayload& FCacheRecordInternal::GetAttachment(const FPayloadId& Id) const
+const FValueWithId& FCacheRecordInternal::GetAttachment(const FValueId& Id) const
 {
-	const int32 Index = Algo::BinarySearchBy(Attachments, Id, &FPayload::GetId);
-	return Attachments.IsValidIndex(Index) ? Attachments[Index] : FPayload::Null;
+	const int32 Index = Algo::BinarySearchBy(Attachments, Id, &FValueWithId::GetId);
+	return Attachments.IsValidIndex(Index) ? Attachments[Index] : FValueWithId::Null;
 }
 
-TConstArrayView<FPayload> FCacheRecordInternal::GetAttachments() const
+TConstArrayView<FValueWithId> FCacheRecordInternal::GetAttachments() const
 {
 	return Attachments;
 }
 
-const FPayload& FCacheRecordInternal::GetPayload(const FPayloadId& Id) const
+const FValueWithId& FCacheRecordInternal::GetValue(const FValueId& Id) const
 {
 	return Value.GetId() == Id ? Value : GetAttachment(Id);
 }
@@ -142,55 +144,61 @@ void FCacheRecordBuilderInternal::SetMeta(FCbObject&& InMeta)
 	Meta.MakeOwned();
 }
 
-void FCacheRecordBuilderInternal::SetValue(const FCompositeBuffer& Buffer, const FPayloadId& Id, const uint64 BlockSize)
+void FCacheRecordBuilderInternal::SetValue(const FCompositeBuffer& Buffer, const FValueId& Id, const uint64 BlockSize)
 {
-	FCompressedBuffer Compressed = FPayload::Compress(Buffer, BlockSize);
-	const FPayloadId ValueId = GetOrCreatePayloadId(Id, Compressed);
-	return SetValue(FPayload(ValueId, MoveTemp(Compressed)));
+	return SetValue(FValue::Compress(Buffer, BlockSize), Id);
 }
 
-void FCacheRecordBuilderInternal::SetValue(const FSharedBuffer& Buffer, const FPayloadId& Id, const uint64 BlockSize)
+void FCacheRecordBuilderInternal::SetValue(const FSharedBuffer& Buffer, const FValueId& Id, const uint64 BlockSize)
 {
-	FCompressedBuffer Compressed = FPayload::Compress(Buffer, BlockSize);
-	const FPayloadId ValueId = GetOrCreatePayloadId(Id, Compressed);
-	return SetValue(FPayload(ValueId, MoveTemp(Compressed)));
+	return SetValue(FValue::Compress(Buffer, BlockSize), Id);
 }
 
-void FCacheRecordBuilderInternal::SetValue(const FPayload& Payload)
+void FCacheRecordBuilderInternal::SetValue(const FValue& InValue, const FValueId& Id)
 {
-	checkf(Payload, TEXT("Failed to set value on %s because the payload is null."), *WriteToString<96>(Key));
+	const FValueId ValueId = GetOrCreateValueId(Id, InValue);
 	checkf(Value.IsNull(),
 		TEXT("Failed to set value on %s with ID %s because it has an existing value with ID %s."),
-		*WriteToString<96>(Key), *WriteToString<32>(Payload.GetId()), *WriteToString<32>(Value.GetId()));
-	checkf(Algo::BinarySearchBy(Attachments, Payload.GetId(), &FPayload::GetId) == INDEX_NONE,
+		*WriteToString<96>(Key), *WriteToString<32>(ValueId), *WriteToString<32>(Value.GetId()));
+	checkf(Algo::BinarySearchBy(Attachments, ValueId, &FValueWithId::GetId) == INDEX_NONE,
 		TEXT("Failed to set value on %s with ID %s because it has an existing attachment with that ID."),
-		*WriteToString<96>(Key), *WriteToString<32>(Payload.GetId()));
-	Value = Payload;
+		*WriteToString<96>(Key), *WriteToString<32>(ValueId));
+	Value = FValueWithId(ValueId, InValue);
 }
 
-void FCacheRecordBuilderInternal::AddAttachment(const FCompositeBuffer& Buffer, const FPayloadId& Id, const uint64 BlockSize)
+void FCacheRecordBuilderInternal::SetValue(const FValueWithId& InValue)
 {
-	FCompressedBuffer Compressed = FPayload::Compress(Buffer, BlockSize);
-	const FPayloadId AttachmentId = GetOrCreatePayloadId(Id, Compressed);
-	return AddAttachment(FPayload(AttachmentId, MoveTemp(Compressed)));
+	checkf(InValue, TEXT("Failed to set value on %s because the value is null."), *WriteToString<96>(Key));
+	SetValue(InValue, InValue.GetId());
 }
 
-void FCacheRecordBuilderInternal::AddAttachment(const FSharedBuffer& Buffer, const FPayloadId& Id, const uint64 BlockSize)
+void FCacheRecordBuilderInternal::AddAttachment(const FCompositeBuffer& Buffer, const FValueId& Id, const uint64 BlockSize)
 {
-	FCompressedBuffer Compressed = FPayload::Compress(Buffer, BlockSize);
-	const FPayloadId AttachmentId = GetOrCreatePayloadId(Id, Compressed);
-	return AddAttachment(FPayload(AttachmentId, MoveTemp(Compressed)));
+	return AddAttachment(FValue::Compress(Buffer, BlockSize), Id);
 }
 
-void FCacheRecordBuilderInternal::AddAttachment(const FPayload& Payload)
+void FCacheRecordBuilderInternal::AddAttachment(const FSharedBuffer& Buffer, const FValueId& Id, const uint64 BlockSize)
 {
-	checkf(Payload, TEXT("Failed to add attachment on %s because the payload is null."), *WriteToString<96>(Key));
-	const FPayloadId& Id = Payload.GetId();
-	const int32 Index = Algo::LowerBoundBy(Attachments, Id, &FPayload::GetId);
+	return AddAttachment(FValue::Compress(Buffer, BlockSize), Id);
+}
+
+void FCacheRecordBuilderInternal::AddAttachment(const FValue& InValue, const FValueId& Id)
+{
+	const FValueId ValueId = GetOrCreateValueId(Id, InValue);
+	checkf(Value.IsNull(),
+		TEXT("Failed to set value on %s with ID %s because it has an existing value with ID %s."),
+		*WriteToString<96>(Key), *WriteToString<32>(ValueId), *WriteToString<32>(Value.GetId()));
+	const int32 Index = Algo::LowerBoundBy(Attachments, Id, &FValueWithId::GetId);
 	checkf(!(Attachments.IsValidIndex(Index) && Attachments[Index].GetId() == Id) && Value.GetId() != Id,
 		TEXT("Failed to add attachment on %s with ID %s because it has an existing attachment or value with that ID."),
 		*WriteToString<96>(Key), *WriteToString<32>(Id));
-	Attachments.Insert(Payload, Index);
+	Attachments.Insert(FValueWithId(ValueId, InValue), Index);
+}
+
+void FCacheRecordBuilderInternal::AddAttachment(const FValueWithId& InValue)
+{
+	checkf(InValue, TEXT("Failed to add attachment on %s because the value is null."), *WriteToString<96>(Key));
+	AddAttachment(InValue, InValue.GetId());
 }
 
 FCacheRecord FCacheRecordBuilderInternal::Build()
@@ -237,30 +245,30 @@ FCbPackage FCacheRecord::Save() const
 	{
 		Writer.AddObject("Meta"_ASV, Meta);
 	}
-	auto SavePayload = [&Package, &Writer](const FPayload& Payload)
+	auto SaveValue = [&Package, &Writer](const FValueWithId& ValueWithId)
 	{
-		if (Payload.HasData())
+		if (ValueWithId.HasData())
 		{
-			Package.AddAttachment(FCbAttachment(Payload.GetData()));
+			Package.AddAttachment(FCbAttachment(ValueWithId.GetData()));
 		}
 		Writer.BeginObject();
-		Writer.AddObjectId("Id"_ASV, Payload.GetId());
-		Writer.AddBinaryAttachment("RawHash"_ASV, Payload.GetRawHash());
-		Writer.AddInteger("RawSize"_ASV, Payload.GetRawSize());
+		Writer.AddObjectId("Id"_ASV, ValueWithId.GetId());
+		Writer.AddBinaryAttachment("RawHash"_ASV, ValueWithId.GetRawHash());
+		Writer.AddInteger("RawSize"_ASV, ValueWithId.GetRawSize());
 		Writer.EndObject();
 	};
-	if (const FPayload& Value = GetValuePayload())
+	if (const FValueWithId& Value = GetValue())
 	{
 		Writer.SetName("Value"_ASV);
-		SavePayload(Value);
+		SaveValue(Value);
 	}
-	TConstArrayView<FPayload> Attachments = GetAttachmentPayloads();
+	TConstArrayView<FValueWithId> Attachments = GetAttachments();
 	if (!Attachments.IsEmpty())
 	{
 		Writer.BeginArray("Attachments"_ASV);
-		for (const FPayload& Attachment : Attachments)
+		for (const FValueWithId& Attachment : Attachments)
 		{
-			SavePayload(Attachment);
+			SaveValue(Attachment);
 		}
 		Writer.EndArray();
 	}
@@ -295,35 +303,35 @@ FOptionalCacheRecord FCacheRecord::Load(const FCbPackage& Attachments, const FCb
 
 	Builder.SetMeta(Object["Meta"_ASV].AsObject());
 
-	auto LoadPayload = [&Attachments](const FCbObjectView& PayloadObject)
+	auto LoadValue = [&Attachments](const FCbObjectView& ValueObject)
 	{
-		const FPayloadId PayloadId = PayloadObject["Id"_ASV].AsObjectId();
-		if (PayloadId.IsNull())
+		const FValueId Id = ValueObject["Id"_ASV].AsObjectId();
+		if (Id.IsNull())
 		{
-			return FPayload();
+			return FValueWithId();
 		}
-		const FIoHash RawHash = PayloadObject["RawHash"_ASV].AsHash();
+		const FIoHash RawHash = ValueObject["RawHash"_ASV].AsHash();
 		if (const FCbAttachment* Attachment = Attachments.FindAttachment(RawHash))
 		{
 			if (const FCompressedBuffer& Compressed = Attachment->AsCompressedBinary())
 			{
-				return FPayload(PayloadId, Compressed);
+				return FValueWithId(Id, Compressed);
 			}
 		}
-		const uint64 RawSize = PayloadObject["RawSize"_ASV].AsUInt64(MAX_uint64);
+		const uint64 RawSize = ValueObject["RawSize"_ASV].AsUInt64(MAX_uint64);
 		if (!RawHash.IsZero() && RawSize != MAX_uint64)
 		{
-			return FPayload(PayloadId, RawHash, RawSize);
+			return FValueWithId(Id, RawHash, RawSize);
 		}
 		else
 		{
-			return FPayload();
+			return FValueWithId();
 		}
 	};
 
 	if (FCbObjectView ValueObject = ObjectView["Value"_ASV].AsObjectView())
 	{
-		FPayload Value = LoadPayload(ValueObject);
+		FValueWithId Value = LoadValue(ValueObject);
 		if (!Value)
 		{
 			return FOptionalCacheRecord();
@@ -333,7 +341,7 @@ FOptionalCacheRecord FCacheRecord::Load(const FCbPackage& Attachments, const FCb
 
 	for (FCbFieldView AttachmentField : ObjectView["Attachments"_ASV])
 	{
-		FPayload Attachment = LoadPayload(AttachmentField.AsObjectView());
+		FValueWithId Attachment = LoadValue(AttachmentField.AsObjectView());
 		if (!Attachment)
 		{
 			return FOptionalCacheRecord();
@@ -363,22 +371,22 @@ namespace UE::DerivedData::Private
 
 uint64 GetCacheRecordCompressedSize(const FCacheRecord& Record)
 {
-	const uint64 ValueSize = Record.GetValuePayload().GetData().GetCompressedSize();
-	return int64(Algo::TransformAccumulate(Record.GetAttachmentPayloads(),
-		[](const FPayload& Payload) { return Payload.GetData().GetCompressedSize(); }, ValueSize));
+	const uint64 ValueSize = Record.GetValue().GetData().GetCompressedSize();
+	return int64(Algo::TransformAccumulate(Record.GetAttachments(),
+		[](const FValueWithId& Value) { return Value.GetData().GetCompressedSize(); }, ValueSize));
 }
 
 uint64 GetCacheRecordTotalRawSize(const FCacheRecord& Record)
 {
-	const uint64 ValueSize = Record.GetValuePayload().GetRawSize();
-	return int64(Algo::TransformAccumulate(Record.GetAttachmentPayloads(), &FPayload::GetRawSize, ValueSize));
+	const uint64 ValueSize = Record.GetValue().GetRawSize();
+	return int64(Algo::TransformAccumulate(Record.GetAttachments(), &FValueWithId::GetRawSize, ValueSize));
 }
 
 uint64 GetCacheRecordRawSize(const FCacheRecord& Record)
 {
-	const uint64 ValueSize = Record.GetValuePayload().GetData().GetRawSize();
-	return int64(Algo::TransformAccumulate(Record.GetAttachmentPayloads(),
-		[](const FPayload& Payload) { return Payload.GetData().GetRawSize(); }, ValueSize));
+	const uint64 ValueSize = Record.GetValue().GetData().GetRawSize();
+	return int64(Algo::TransformAccumulate(Record.GetAttachments(),
+		[](const FValueWithId& Value) { return Value.GetData().GetRawSize(); }, ValueSize));
 }
 
 } // UE::DerivedData::Private
diff --git a/Engine/Source/Developer/DerivedDataCache/Private/DerivedDataPayload.cpp b/Engine/Source/Developer/DerivedDataCache/Private/DerivedDataPayload.cpp
deleted file mode 100644
diff --git a/Engine/Source/Developer/DerivedDataCache/Private/DerivedDataValue.cpp b/Engine/Source/Developer/DerivedDataCache/Private/DerivedDataValue.cpp
new file mode 100644
--- /dev/null
+++ b/Engine/Source/Developer/DerivedDataCache/Private/DerivedDataValue.cpp
@@ -0,0 +1,63 @@
+// Copyright Epic Games, Inc. All Rights Reserved.
+
+#include "DerivedDataValue.h"
+
+#include "Compression/OodleDataCompression.h"
+#include "Containers/StringConv.h"
+#include "Hash/xxhash.h"
+#include "Serialization/CompactBinary.h"
+#include "Templates/UnrealTemplate.h"
+
+namespace UE::DerivedData
+{
+
+constexpr ECompressedBufferCompressor GDefaultCompressor = ECompressedBufferCompressor::Mermaid;
+constexpr ECompressedBufferCompressionLevel GDefaultCompressionLevel = ECompressedBufferCompressionLevel::VeryFast;
+
+FValue FValue::Compress(const FCompositeBuffer& Buffer, const uint64 BlockSize)
+{
+	return FValue(FCompressedBuffer::Compress(Buffer, GDefaultCompressor, GDefaultCompressionLevel, BlockSize));
+}
+
+FValue FValue::Compress(const FSharedBuffer& Buffer, const uint64 BlockSize)
+{
+	return FValue(FCompressedBuffer::Compress(Buffer, GDefaultCompressor, GDefaultCompressionLevel, BlockSize));
+}
+
+FValueId::FValueId(const FMemoryView Id)
+{
+	checkf(Id.GetSize() == sizeof(ByteArray),
+		TEXT("FValueId cannot be constructed from a view of %" UINT64_FMT " bytes."), Id.GetSize());
+	FMemory::Memcpy(Bytes, Id.GetData(), sizeof(ByteArray));
+}
+
+FValueId::FValueId(const FCbObjectId& Id)
+	: FValueId(ImplicitConv<const ByteArray&>(Id))
+{
+}
+
+FValueId::operator FCbObjectId() const
+{
+	return FCbObjectId(ImplicitConv<const ByteArray&>(*this));
+}
+
+FValueId FValueId::FromHash(const FIoHash& Hash)
+{
+	checkf(!Hash.IsZero(), TEXT("FValueId requires a non-zero hash."));
+	return FValueId(MakeMemoryView(Hash.GetBytes()).Left(sizeof(ByteArray)));
+}
+
+FValueId FValueId::FromName(const FUtf8StringView Name)
+{
+	checkf(!Name.IsEmpty(), TEXT("FValueId requires a non-empty name."));
+	uint8 HashBytes[16];
+	FXxHash128::HashBuffer(Name.GetData(), Name.Len()).ToByteArray(HashBytes);
+	return FValueId(MakeMemoryView(HashBytes, sizeof(ByteArray)));
+}
+
+FValueId FValueId::FromName(const FWideStringView Name)
+{
+	return FValueId::FromName(FTCHARToUTF8(Name));
+}
+
+} // UE::DerivedData
diff --git a/Engine/Source/Developer/DerivedDataCache/Private/FileSystemDerivedDataBackend.cpp b/Engine/Source/Developer/DerivedDataCache/Private/FileSystemDerivedDataBackend.cpp
--- a/Engine/Source/Developer/DerivedDataCache/Private/FileSystemDerivedDataBackend.cpp
+++ b/Engine/Source/Developer/DerivedDataCache/Private/FileSystemDerivedDataBackend.cpp
@@ -14,7 +14,7 @@
 #include "DerivedDataCacheRecord.h"
 #include "DerivedDataCacheUsageStats.h"
 #include "DerivedDataChunk.h"
-#include "DerivedDataPayload.h"
+#include "DerivedDataValue.h"
 #include "Experimental/Async/LazyEvent.h"
 #include "Features/IModularFeatures.h"
 #include "HAL/Event.h"
@@ -780,15 +780,15 @@ private:
 		const FStringView RecordName,
 		const FCacheRecordPolicy& Policy,
 		EStatus& OutStatus);
-	FPayload GetCacheContent(
+	FValueWithId GetCacheContent(
 		const FCacheKey& Key,
-		const FPayload& Payload,
+		const FValueWithId& Value,
 		FStringView ContentName,
 		ECachePolicy Policy,
 		ECachePolicy SkipFlag) const;
 	void GetCacheContent(
 		const FCacheKey& Key,
-		const FPayload& Payload,
+		const FValueWithId& Value,
 		FStringView ContentName,
 		ECachePolicy Policy,
 		ECachePolicy SkipFlag,
@@ -1605,11 +1605,11 @@ void FFileSystemCacheStore::GetChunks(
 	TArray<FCacheChunkRequest, TInlineAllocator<16>> SortedRequests(Requests);
 	SortedRequests.StableSort(TChunkLess());
 
-	FPayload Payload;
-	FCacheKey PayloadKey;
-	TUniquePtr<FArchive> PayloadAr;
-	FCompressedBufferReader PayloadReader;
-	EStatus PayloadStatus = EStatus::Error;
+	FValueWithId Value;
+	FCacheKey ValueKey;
+	TUniquePtr<FArchive> ValueAr;
+	FCompressedBufferReader ValueReader;
+	EStatus ValueStatus = EStatus::Error;
 	FOptionalCacheRecord Record;
 	for (const FCacheChunkRequest& Request : SortedRequests)
 	{
@@ -1618,31 +1618,31 @@ void FFileSystemCacheStore::GetChunks(
 		TRACE_CPUPROFILER_EVENT_SCOPE(FileSystemDDC_Get);
 		TRACE_COUNTER_INCREMENT(FileSystemDDC_Get);
 		COOK_STAT(auto Timer = bExistsOnly ? UsageStats.TimeProbablyExists() : UsageStats.TimeGet());
-		if (!(Payload && PayloadKey == Request.Key && Payload.GetId() == Request.Id) || !bExistsOnly <= PayloadReader.HasSource())
+		if (!(Value && ValueKey == Request.Key && Value.GetId() == Request.Id) || !bExistsOnly <= ValueReader.HasSource())
 		{
-			PayloadStatus = EStatus::Error;
-			PayloadReader.ResetSource();
-			PayloadAr.Reset();
-			Payload.Reset();
+			ValueStatus = EStatus::Error;
+			ValueReader.ResetSource();
+			ValueAr.Reset();
+			Value.Reset();
 			if (!(Record && Record.Get().GetKey() == Request.Key))
 			{
 				FCacheRecordPolicyBuilder PolicyBuilder(ECachePolicy::None);
 				const ECachePolicy RecordSkipFlags = bExistsOnly ? ECachePolicy::SkipData : ECachePolicy::None;
-				PolicyBuilder.AddPayloadPolicy(Request.Id, Request.Policy | RecordSkipFlags);
+				PolicyBuilder.AddValuePolicy(Request.Id, Request.Policy | RecordSkipFlags);
 				Record.Reset();
 				Record = GetCacheRecordOnly(Request.Key, Request.Name, PolicyBuilder.Build());
 			}
 			if (Record)
 			{
-				Payload = Record.Get().GetPayload(Request.Id);
-				PayloadKey = Request.Key;
-				GetCacheContent(Request.Key, Payload, Request.Name, Request.Policy, SkipFlag, PayloadReader, PayloadAr);
+				Value = Record.Get().GetValue(Request.Id);
+				ValueKey = Request.Key;
+				GetCacheContent(Request.Key, Value, Request.Name, Request.Policy, SkipFlag, ValueReader, ValueAr);
 			}
 		}
-		if (Payload)
+		if (Value)
 		{
-			const uint64 RawOffset = FMath::Min(Payload.GetRawSize(), Request.RawOffset);
-			const uint64 RawSize = FMath::Min(Payload.GetRawSize() - RawOffset, Request.RawSize);
+			const uint64 RawOffset = FMath::Min(Value.GetRawSize(), Request.RawOffset);
+			const uint64 RawSize = FMath::Min(Value.GetRawSize() - RawOffset, Request.RawSize);
 			UE_LOG(LogDerivedDataCache, Verbose, TEXT("%s: Cache hit for %s from '%s'"),
 				*CachePath, *WriteToString<96>(Request.Key, '/', Request.Id), *Request.Name);
 			TRACE_COUNTER_INCREMENT(FileSystemDDC_GetHit);
@@ -1653,11 +1653,11 @@ void FFileSystemCacheStore::GetChunks(
 				FSharedBuffer Buffer;
 				if (!bExistsOnly)
 				{
-					Buffer = PayloadReader.Decompress(RawOffset, RawSize);
+					Buffer = ValueReader.Decompress(RawOffset, RawSize);
 				}
 				const EStatus ChunkStatus = bExistsOnly || Buffer.GetSize() == RawSize ? EStatus::Ok : EStatus::Error;
 				OnComplete({Request.Name, Request.Key, Request.Id, Request.RawOffset,
-					RawSize, Payload.GetRawHash(), MoveTemp(Buffer), Request.UserData, ChunkStatus});
+					RawSize, Value.GetRawHash(), MoveTemp(Buffer), Request.UserData, ChunkStatus});
 			}
 			continue;
 		}
@@ -1672,17 +1672,17 @@ void FFileSystemCacheStore::GetChunks(
 uint64 FFileSystemCacheStore::MeasureCompressedCacheRecord(const FCacheRecord& Record) const
 {
 	return Record.GetMeta().GetSize() +
-		Record.GetValuePayload().GetData().GetCompressedSize() +
-		Algo::TransformAccumulate(Record.GetAttachmentPayloads(),
-			[](const FPayload& Payload) { return Payload.GetData().GetCompressedSize(); }, uint64(0));
+		Record.GetValue().GetData().GetCompressedSize() +
+		Algo::TransformAccumulate(Record.GetAttachments(),
+			[](const FValue& Value) { return Value.GetData().GetCompressedSize(); }, uint64(0));
 }
 
 uint64 FFileSystemCacheStore::MeasureRawCacheRecord(const FCacheRecord& Record) const
 {
 	return Record.GetMeta().GetSize() +
-		Record.GetValuePayload().GetData().GetRawSize() +
-		Algo::TransformAccumulate(Record.GetAttachmentPayloads(),
-			[](const FPayload& Payload) { return Payload.GetData().GetRawSize(); }, uint64(0));
+		Record.GetValue().GetData().GetRawSize() +
+		Algo::TransformAccumulate(Record.GetAttachments(),
+			[](const FValue& Value) { return Value.GetData().GetRawSize(); }, uint64(0));
 }
 
 bool FFileSystemCacheStore::PutCacheRecord(
@@ -1700,9 +1700,9 @@ bool FFileSystemCacheStore::PutCacheRecord(
 	}
 
 	const FCacheKey& Key = Record.GetKey();
-	const ECachePolicy CombinedPayloadPolicy = Algo::TransformAccumulate(
-		Policy.GetPayloadPolicies(), &FCachePayloadPolicy::Policy, Policy.GetDefaultPayloadPolicy(), UE_PROJECTION(operator|));
-	const ECachePolicy CombinedPolicy = Policy.GetRecordPolicy() | CombinedPayloadPolicy;
+	const ECachePolicy CombinedValuePolicy = Algo::TransformAccumulate(
+		Policy.GetValuePolicies(), &FCacheValuePolicy::Policy, Policy.GetDefaultValuePolicy(), UE_PROJECTION(operator|));
+	const ECachePolicy CombinedPolicy = Policy.GetRecordPolicy() | CombinedValuePolicy;
 
 	// Skip the request if storing to the cache is disabled.
 	const ECachePolicy StoreFlag = SpeedClass == ESpeedClass::Local ? ECachePolicy::StoreLocal : ECachePolicy::StoreRemote;
@@ -1725,7 +1725,7 @@ bool FFileSystemCacheStore::PutCacheRecord(
 	FCbPackage ExistingPackage;
 	TStringBuilder<256> Path;
 	BuildCacheRecordPath(Key, Path);
-	if (EnumHasAllFlags(CombinedPayloadPolicy, ECachePolicy::SkipValue | ECachePolicy::SkipAttachments))
+	if (EnumHasAllFlags(CombinedValuePolicy, ECachePolicy::SkipValue | ECachePolicy::SkipAttachments))
 	{
 		bRecordExists = FileExists(Path);
 	}
@@ -1898,17 +1898,17 @@ FOptionalCacheRecord FFileSystemCacheStore::GetCacheRecord(
 		RecordBuilder.SetMeta(FCbObject(Record.Get().GetMeta()));
 	}
 
-	if (FPayload Payload = Record.Get().GetValuePayload())
+	if (FValueWithId Value = Record.Get().GetValue())
 	{
-		const ECachePolicy PayloadPolicy = Policy.GetPayloadPolicy(Payload.GetId());
-		if (FPayload Content = GetCacheContent(Key, Payload, RecordName, PayloadPolicy, ECachePolicy::SkipValue))
+		const ECachePolicy ValuePolicy = Policy.GetValuePolicy(Value.GetId());
+		if (FValueWithId Content = GetCacheContent(Key, Value, RecordName, ValuePolicy, ECachePolicy::SkipValue))
 		{
 			RecordBuilder.SetValue(MoveTemp(Content));
 		}
-		else if (EnumHasAnyFlags(PayloadPolicy, ECachePolicy::PartialOnError))
+		else if (EnumHasAnyFlags(ValuePolicy, ECachePolicy::PartialOnError))
 		{
 			OutStatus = EStatus::Error;
-			RecordBuilder.SetValue(MoveTemp(Payload));
+			RecordBuilder.SetValue(MoveTemp(Value));
 		}
 		else
 		{
@@ -1917,17 +1917,17 @@ FOptionalCacheRecord FFileSystemCacheStore::GetCacheRecord(
 		}
 	}
 
-	for (FPayload Payload : Record.Get().GetAttachmentPayloads())
+	for (FValueWithId Value : Record.Get().GetAttachments())
 	{
-		const ECachePolicy PayloadPolicy = Policy.GetPayloadPolicy(Payload.GetId());
-		if (FPayload Content = GetCacheContent(Key, Payload, RecordName, PayloadPolicy, ECachePolicy::SkipAttachments))
+		const ECachePolicy ValuePolicy = Policy.GetValuePolicy(Value.GetId());
+		if (FValueWithId Content = GetCacheContent(Key, Value, RecordName, ValuePolicy, ECachePolicy::SkipAttachments))
 		{
 			RecordBuilder.AddAttachment(MoveTemp(Content));
 		}
-		else if (EnumHasAnyFlags(PayloadPolicy, ECachePolicy::PartialOnError))
+		else if (EnumHasAnyFlags(ValuePolicy, ECachePolicy::PartialOnError))
 		{
 			OutStatus = EStatus::Error;
-			RecordBuilder.AddAttachment(MoveTemp(Payload));
+			RecordBuilder.AddAttachment(MoveTemp(Value));
 		}
 		else
 		{
@@ -1959,24 +1959,24 @@ bool FFileSystemCacheStore::PutCacheContent(
 	return true;
 }
 
-FPayload FFileSystemCacheStore::GetCacheContent(
+FValueWithId FFileSystemCacheStore::GetCacheContent(
 	const FCacheKey& Key,
-	const FPayload& Payload,
+	const FValueWithId& Value,
 	const FStringView ContentName,
 	const ECachePolicy Policy,
 	const ECachePolicy SkipFlag) const
 {
 	if (!EnumHasAnyFlags(Policy, ECachePolicy::Query))
 	{
-		return Payload.RemoveData();
+		return Value.RemoveData();
 	}
 
-	if (Payload.HasData())
+	if (Value.HasData())
 	{
-		return EnumHasAnyFlags(Policy, SkipFlag) ? Payload.RemoveData() : Payload;
+		return EnumHasAnyFlags(Policy, SkipFlag) ? Value.RemoveData() : Value;
 	}
 
-	const FIoHash& RawHash = Payload.GetRawHash();
+	const FIoHash& RawHash = Value.GetRawHash();
 
 	TStringBuilder<256> Path;
 	BuildCacheContentPath(RawHash, Path);
@@ -1988,7 +1988,7 @@ FPayload FFileSystemCacheStore::GetCacheContent(
 			{
 				AccessLogWriter->Append(RawHash, Path);
 			}
-			return Payload;
+			return Value;
 		}
 	}
 	else
@@ -2002,26 +2002,26 @@ FPayload FFileSystemCacheStore::GetCacheContent(
 				{
 					AccessLogWriter->Append(RawHash, Path);
 				}
-				return FPayload(Payload.GetId(), MoveTemp(CompressedBuffer));
+				return FValueWithId(Value.GetId(), MoveTemp(CompressedBuffer));
 			}
 			UE_LOG(LogDerivedDataCache, Display,
-				TEXT("%s: Cache miss with corrupted payload %s with hash %s for %s from '%.*s'"),
-				*CachePath, *WriteToString<16>(Payload.GetId()), *WriteToString<48>(RawHash),
+				TEXT("%s: Cache miss with corrupted value %s with hash %s for %s from '%.*s'"),
+				*CachePath, *WriteToString<16>(Value.GetId()), *WriteToString<48>(RawHash),
 				*WriteToString<96>(Key), ContentName.Len(), ContentName.GetData());
-			return FPayload::Null;
+			return FValueWithId::Null;
 		}
 	}
 
 	UE_LOG(LogDerivedDataCache, Verbose,
-		TEXT("%s: Cache miss with missing payload %s with hash %s for %s from '%.*s'"),
-		*CachePath, *WriteToString<16>(Payload.GetId()), *WriteToString<48>(RawHash), *WriteToString<96>(Key),
+		TEXT("%s: Cache miss with missing value %s with hash %s for %s from '%.*s'"),
+		*CachePath, *WriteToString<16>(Value.GetId()), *WriteToString<48>(RawHash), *WriteToString<96>(Key),
 		ContentName.Len(), ContentName.GetData());
-	return FPayload::Null;
+	return FValueWithId::Null;
 }
 
 void FFileSystemCacheStore::GetCacheContent(
 	const FCacheKey& Key,
-	const FPayload& Payload,
+	const FValueWithId& Value,
 	const FStringView ContentName,
 	const ECachePolicy Policy,
 	const ECachePolicy SkipFlag,
@@ -2033,17 +2033,17 @@ void FFileSystemCacheStore::GetCacheContent(
 		return;
 	}
 
-	if (Payload.HasData())
+	if (Value.HasData())
 	{
 		if (!EnumHasAnyFlags(Policy, SkipFlag))
 		{
-			Reader.SetSource(Payload.GetData());
+			Reader.SetSource(Value.GetData());
 		}
 		OutArchive.Reset();
 		return;
 	}
 
-	const FIoHash& RawHash = Payload.GetRawHash();
+	const FIoHash& RawHash = Value.GetRawHash();
 
 	TStringBuilder<256> Path;
 	BuildCacheContentPath(RawHash, Path);
@@ -2073,8 +2073,8 @@ void FFileSystemCacheStore::GetCacheContent(
 				return;
 			}
 			UE_LOG(LogDerivedDataCache, Display,
-				TEXT("%s: Cache miss with corrupted payload %s with hash %s for %s from '%.*s'"),
-				*CachePath, *WriteToString<16>(Payload.GetId()), *WriteToString<48>(RawHash),
+				TEXT("%s: Cache miss with corrupted value %s with hash %s for %s from '%.*s'"),
+				*CachePath, *WriteToString<16>(Value.GetId()), *WriteToString<48>(RawHash),
 				*WriteToString<96>(Key), ContentName.Len(), ContentName.GetData());
 			Reader.ResetSource();
 			OutArchive.Reset();
@@ -2083,8 +2083,8 @@ void FFileSystemCacheStore::GetCacheContent(
 	}
 
 	UE_LOG(LogDerivedDataCache, Verbose,
-		TEXT("%s: Cache miss with missing payload %s with hash %s for %s from '%.*s'"),
-		*CachePath, *WriteToString<16>(Payload.GetId()), *WriteToString<48>(RawHash), *WriteToString<96>(Key),
+		TEXT("%s: Cache miss with missing value %s with hash %s for %s from '%.*s'"),
+		*CachePath, *WriteToString<16>(Value.GetId()), *WriteToString<48>(RawHash), *WriteToString<96>(Key),
 		ContentName.Len(), ContentName.GetData());
 }
 
diff --git a/Engine/Source/Developer/DerivedDataCache/Private/HierarchicalDerivedDataBackend.h b/Engine/Source/Developer/DerivedDataCache/Private/HierarchicalDerivedDataBackend.h
--- a/Engine/Source/Developer/DerivedDataCache/Private/HierarchicalDerivedDataBackend.h
+++ b/Engine/Source/Developer/DerivedDataCache/Private/HierarchicalDerivedDataBackend.h
@@ -10,9 +10,9 @@
 #include "DerivedDataCacheRecord.h"
 #include "DerivedDataCacheUsageStats.h"
 #include "DerivedDataChunk.h"
-#include "DerivedDataPayload.h"
 #include "DerivedDataRequestOwner.h"
 #include "HAL/PlatformMath.h"
+#include "Memory/SharedBuffer.h"
 #include "Misc/ScopeRWLock.h"
 #include "Misc/StringBuilder.h"
 #include "ProfilingDebugging/CookStats.h"
@@ -586,10 +586,10 @@ public:
 					};
 					for (FCacheGetRequest& Request : RemainingRequests)
 					{
-						FCacheRecordPolicyBuilder Builder(ConvertPolicy(Request.Policy.GetDefaultPayloadPolicy()));
-						for (const FCachePayloadPolicy& PayloadPolicy : Request.Policy.GetPayloadPolicies())
+						FCacheRecordPolicyBuilder Builder(ConvertPolicy(Request.Policy.GetDefaultValuePolicy()));
+						for (const FCacheValuePolicy& ValuePolicy : Request.Policy.GetValuePolicies())
 						{
-							Builder.AddPayloadPolicy(PayloadPolicy.Id, ConvertPolicy(PayloadPolicy.Policy));
+							Builder.AddValuePolicy(ValuePolicy.Id, ConvertPolicy(ValuePolicy.Policy));
 						}
 						Request.Policy = Builder.Build();
 					}
diff --git a/Engine/Source/Developer/DerivedDataCache/Private/HttpDerivedDataBackend.cpp b/Engine/Source/Developer/DerivedDataCache/Private/HttpDerivedDataBackend.cpp
--- a/Engine/Source/Developer/DerivedDataCache/Private/HttpDerivedDataBackend.cpp
+++ b/Engine/Source/Developer/DerivedDataCache/Private/HttpDerivedDataBackend.cpp
@@ -16,11 +16,11 @@
 #include "Containers/StaticArray.h"
 #include "Containers/Ticker.h"
 #include "DerivedDataCacheRecord.h"
-#include "DerivedDataPayload.h"
 #include "Dom/JsonObject.h"
 #include "GenericPlatform/GenericPlatformFile.h"
 #include "HAL/PlatformFileManager.h"
 #include "IO/IoHash.h"
+#include "Memory/SharedBuffer.h"
 #include "Misc/FileHelper.h"
 #include "Misc/ScopeLock.h"
 #include "Misc/SecureHash.h"
diff --git a/Engine/Source/Developer/DerivedDataCache/Private/MemoryDerivedDataBackend.cpp b/Engine/Source/Developer/DerivedDataCache/Private/MemoryDerivedDataBackend.cpp
--- a/Engine/Source/Developer/DerivedDataCache/Private/MemoryDerivedDataBackend.cpp
+++ b/Engine/Source/Developer/DerivedDataCache/Private/MemoryDerivedDataBackend.cpp
@@ -4,7 +4,7 @@
 
 #include "Algo/Accumulate.h"
 #include "Algo/AllOf.h"
-#include "DerivedDataPayload.h"
+#include "DerivedDataValue.h"
 #include "Misc/ScopeExit.h"
 #include "Serialization/CompactBinary.h"
 #include "Templates/UniquePtr.h"
@@ -447,8 +447,8 @@ bool FMemoryDerivedDataBackend::ShouldSimulateMiss(const FCacheKey& Key)
 
 int64 FMemoryDerivedDataBackend::CalcRawCacheRecordSize(const FCacheRecord& Record) const
 {
-	const uint64 ValueSize = Record.GetValuePayload().GetRawSize();
-	return int64(Algo::TransformAccumulate(Record.GetAttachmentPayloads(), &FPayload::GetRawSize, ValueSize));
+	const uint64 ValueSize = Record.GetValue().GetRawSize();
+	return int64(Algo::TransformAccumulate(Record.GetAttachments(), &FValue::GetRawSize, ValueSize));
 }
 
 int64 FMemoryDerivedDataBackend::CalcSerializedCacheRecordSize(const FCacheRecord& Record) const
@@ -457,12 +457,12 @@ int64 FMemoryDerivedDataBackend::CalcSerializedCacheRecordSize(const FCacheRecor
 	uint64 TotalSize = 20;
 	TotalSize += Record.GetKey().Bucket.ToString().Len();
 	TotalSize += Record.GetMeta().GetSize();
-	const auto CalcCachePayloadSize = [](const FPayload& Payload)
+	const auto CalcCacheValueSize = [](const FValueWithId& Value)
 	{
-		return Payload ? Payload.GetData().GetCompressedSize() + 32 : 0;
+		return Value ? Value.GetData().GetCompressedSize() + 32 : 0;
 	};
-	TotalSize += CalcCachePayloadSize(Record.GetValuePayload());
-	TotalSize += Algo::TransformAccumulate(Record.GetAttachmentPayloads(), CalcCachePayloadSize, uint64(0));
+	TotalSize += CalcCacheValueSize(Record.GetValue());
+	TotalSize += Algo::TransformAccumulate(Record.GetAttachments(), CalcCacheValueSize, uint64(0));
 	return int64(TotalSize);
 }
 
@@ -491,10 +491,10 @@ void FMemoryDerivedDataBackend::Put(
 			continue;
 		}
 
-		const FPayload& Value = Record.GetValuePayload();
-		const TConstArrayView<FPayload> Attachments = Record.GetAttachmentPayloads();
+		const FValueWithId& Value = Record.GetValue();
+		const TConstArrayView<FValueWithId> Attachments = Record.GetAttachments();
 
-		if ((Value && !Value.HasData()) || !Algo::AllOf(Attachments, &FPayload::HasData))
+		if ((Value && !Value.HasData()) || !Algo::AllOf(Attachments, &FValue::HasData))
 		{
 			continue;
 		}
@@ -565,23 +565,23 @@ void FMemoryDerivedDataBackend::Get(
 		}
 		if (Record)
 		{
-			if (const FPayload& Payload = Record.Get().GetValuePayload();
-				!Payload.HasData() && !EnumHasAnyFlags(Policy.GetPayloadPolicy(Payload.GetId()), ECachePolicy::SkipValue))
+			if (const FValueWithId& Value = Record.Get().GetValue();
+				!Value.HasData() && !EnumHasAnyFlags(Policy.GetValuePolicy(Value.GetId()), ECachePolicy::SkipValue))
 			{
 				Status = EStatus::Error;
-				if (!EnumHasAllFlags(Policy.GetPayloadPolicy(Payload.GetId()), ECachePolicy::PartialOnError))
+				if (!EnumHasAllFlags(Policy.GetValuePolicy(Value.GetId()), ECachePolicy::PartialOnError))
 				{
 					Record.Reset();
 				}
 			}
 			if (Record)
 			{
-				for (const FPayload& Payload : Record.Get().GetAttachmentPayloads())
+				for (const FValueWithId& Value : Record.Get().GetAttachments())
 				{
-					if (!Payload.HasData() && !EnumHasAnyFlags(Policy.GetPayloadPolicy(Payload.GetId()), ECachePolicy::SkipAttachments))
+					if (!Value.HasData() && !EnumHasAnyFlags(Policy.GetValuePolicy(Value.GetId()), ECachePolicy::SkipAttachments))
 					{
 						Status = EStatus::Error;
-						if (!EnumHasAllFlags(Policy.GetPayloadPolicy(Payload.GetId()), ECachePolicy::PartialOnError))
+						if (!EnumHasAllFlags(Policy.GetValuePolicy(Value.GetId()), ECachePolicy::PartialOnError))
 						{
 							Record.Reset();
 							break;
@@ -613,8 +613,8 @@ void FMemoryDerivedDataBackend::GetChunks(
 	IRequestOwner& Owner,
 	FOnCacheChunkComplete&& OnComplete)
 {
-	FPayload Payload;
-	FCacheKey PayloadKey;
+	FValueWithId Value;
+	FCacheKey ValueKey;
 	FCompressedBufferReader Reader;
 	for (const FCacheChunkRequest& Request : Requests)
 	{
@@ -625,32 +625,32 @@ void FMemoryDerivedDataBackend::GetChunks(
 			UE_LOG(LogDerivedDataCache, Verbose, TEXT("%s: Simulated miss for get of %s from '%s'"),
 				*GetName(), *WriteToString<96>(Request.Key, '/', Request.Id), *Request.Name);
 		}
-		else if (PayloadKey == Request.Key && Payload.GetId() == Request.Id)
+		else if (ValueKey == Request.Key && Value.GetId() == Request.Id)
 		{
-			// Payload matches the request.
+			// Value matches the request.
 		}
 		else if (FWriteScopeLock ScopeLock(SynchronizationObject); const FCacheRecord* Record = CacheRecords.Find(Request.Key))
 		{
 			Reader.ResetSource();
-			Payload.Reset();
-			Payload = Record->GetAttachmentPayload(Request.Id);
-			PayloadKey = Request.Key;
-			Reader.SetSource(Payload.GetData());
+			Value.Reset();
+			Value = Record->GetAttachment(Request.Id);
+			ValueKey = Request.Key;
+			Reader.SetSource(Value.GetData());
 		}
-		if (Payload && Request.RawOffset <= Payload.GetRawSize())
+		if (Value && Request.RawOffset <= Value.GetRawSize())
 		{
-			const uint64 RawSize = FMath::Min(Payload.GetRawSize() - Request.RawOffset, Request.RawSize);
+			const uint64 RawSize = FMath::Min(Value.GetRawSize() - Request.RawOffset, Request.RawSize);
 			COOK_STAT(Timer.AddHit(RawSize));
 			if (OnComplete)
 			{
 				FSharedBuffer Buffer;
-				if (Payload.HasData() && !bExistsOnly)
+				if (Value.HasData() && !bExistsOnly)
 				{
 					Buffer = Reader.Decompress(Request.RawOffset, RawSize);
 				}
 				const EStatus Status = bExistsOnly || Buffer ? EStatus::Ok : EStatus::Error;
 				OnComplete({Request.Name, Request.Key, Request.Id, Request.RawOffset,
-					RawSize, Payload.GetRawHash(), MoveTemp(Buffer), Request.UserData, Status});
+					RawSize, Value.GetRawHash(), MoveTemp(Buffer), Request.UserData, Status});
 			}
 		}
 		else
diff --git a/Engine/Source/Developer/DerivedDataCache/Private/PakFileDerivedDataBackend.cpp b/Engine/Source/Developer/DerivedDataCache/Private/PakFileDerivedDataBackend.cpp
--- a/Engine/Source/Developer/DerivedDataCache/Private/PakFileDerivedDataBackend.cpp
+++ b/Engine/Source/Developer/DerivedDataCache/Private/PakFileDerivedDataBackend.cpp
@@ -8,7 +8,7 @@
 #include "DerivedDataCacheRecord.h"
 #include "DerivedDataCacheUsageStats.h"
 #include "DerivedDataChunk.h"
-#include "DerivedDataPayload.h"
+#include "DerivedDataValue.h"
 #include "HAL/FileManager.h"
 #include "HAL/PlatformFile.h"
 #include "HAL/PlatformFileManager.h"
@@ -558,31 +558,31 @@ void FPakFileDerivedDataBackend::GetChunks(
 		{
 			FCacheRecordPolicyBuilder PolicyBuilder(ECachePolicy::None);
 			const ECachePolicy RecordSkipFlags = bExistsOnly ? ECachePolicy::SkipData : ECachePolicy::None;
-			PolicyBuilder.AddPayloadPolicy(Request.Id, Request.Policy | RecordSkipFlags);
+			PolicyBuilder.AddValuePolicy(Request.Id, Request.Policy | RecordSkipFlags);
 			Record = GetCacheRecordOnly(Request.Key, Request.Name, PolicyBuilder.Build());
 		}
 		if (Record)
 		{
-			EStatus PayloadStatus = EStatus::Ok;
-			FPayload Payload = Record.Get().GetPayload(Request.Id);
-			GetCacheContent(Request.Key, Request.Name, Request.Policy, SkipFlag, Payload, PayloadStatus);
-			if (Payload)
+			EStatus ValueStatus = EStatus::Ok;
+			FValueWithId Value = Record.Get().GetValue(Request.Id);
+			GetCacheContent(Request.Key, Request.Name, Request.Policy, SkipFlag, Value, ValueStatus);
+			if (Value)
 			{
-				const uint64 RawOffset = FMath::Min(Payload.GetRawSize(), Request.RawOffset);
-				const uint64 RawSize = FMath::Min(Payload.GetRawSize() - RawOffset, Request.RawSize);
+				const uint64 RawOffset = FMath::Min(Value.GetRawSize(), Request.RawOffset);
+				const uint64 RawSize = FMath::Min(Value.GetRawSize() - RawOffset, Request.RawSize);
 				UE_LOG(LogDerivedDataCache, Verbose, TEXT("%s: Cache hit for %s from '%s'"),
 					*CachePath, *WriteToString<96>(Request.Key, '/', Request.Id), *Request.Name);
-				COOK_STAT(Timer.AddHit(Payload.HasData() ? RawSize : 0));
+				COOK_STAT(Timer.AddHit(Value.HasData() ? RawSize : 0));
 				if (OnComplete)
 				{
 					FSharedBuffer Buffer;
-					if (Payload.HasData() && !bExistsOnly)
+					if (Value.HasData() && !bExistsOnly)
 					{
-						FCompressedBufferReaderSourceScope Source(Reader, Payload.GetData());
+						FCompressedBufferReaderSourceScope Source(Reader, Value.GetData());
 						Buffer = Reader.Decompress(RawOffset, RawSize);
 					}
 					OnComplete({Request.Name, Request.Key, Request.Id, Request.RawOffset,
-						RawSize, Payload.GetRawHash(), MoveTemp(Buffer), Request.UserData, PayloadStatus});
+						RawSize, Value.GetRawHash(), MoveTemp(Buffer), Request.UserData, ValueStatus});
 				}
 				continue;
 			}
@@ -599,17 +599,17 @@ void FPakFileDerivedDataBackend::GetChunks(
 uint64 FPakFileDerivedDataBackend::MeasureCompressedCacheRecord(const FCacheRecord& Record) const
 {
 	return Record.GetMeta().GetSize() +
-		Record.GetValuePayload().GetData().GetCompressedSize() +
-		Algo::TransformAccumulate(Record.GetAttachmentPayloads(),
-			[](const FPayload& Payload) { return Payload.GetData().GetCompressedSize(); }, uint64(0));
+		Record.GetValue().GetData().GetCompressedSize() +
+		Algo::TransformAccumulate(Record.GetAttachments(),
+			[](const FValue& Value) { return Value.GetData().GetCompressedSize(); }, uint64(0));
 }
 
 uint64 FPakFileDerivedDataBackend::MeasureRawCacheRecord(const FCacheRecord& Record) const
 {
 	return Record.GetMeta().GetSize() +
-		Record.GetValuePayload().GetData().GetRawSize() +
-		Algo::TransformAccumulate(Record.GetAttachmentPayloads(),
-			[](const FPayload& Payload) { return Payload.GetData().GetRawSize(); }, uint64(0));
+		Record.GetValue().GetData().GetRawSize() +
+		Algo::TransformAccumulate(Record.GetAttachments(),
+			[](const FValue& Value) { return Value.GetData().GetRawSize(); }, uint64(0));
 }
 
 bool FPakFileDerivedDataBackend::PutCacheRecord(
@@ -644,9 +644,9 @@ bool FPakFileDerivedDataBackend::PutCacheRecord(
 	FCbPackage ExistingPackage;
 	TStringBuilder<256> Path;
 	FPathViews::Append(Path, TEXT("Buckets"), Key);
-	const ECachePolicy CombinedPayloadPolicy = Algo::TransformAccumulate(
-		Policy.GetPayloadPolicies(), &FCachePayloadPolicy::Policy, Policy.GetDefaultPayloadPolicy(), UE_PROJECTION(operator|));
-	if (EnumHasAllFlags(CombinedPayloadPolicy, ECachePolicy::SkipValue | ECachePolicy::SkipAttachments))
+	const ECachePolicy CombinedValuePolicy = Algo::TransformAccumulate(
+		Policy.GetValuePolicies(), &FCacheValuePolicy::Policy, Policy.GetDefaultValuePolicy(), UE_PROJECTION(operator|));
+	if (EnumHasAllFlags(CombinedValuePolicy, ECachePolicy::SkipValue | ECachePolicy::SkipAttachments))
 	{
 		bRecordExists = FileExists(Path);
 	}
@@ -804,26 +804,26 @@ FOptionalCacheRecord FPakFileDerivedDataBackend::GetCacheRecord(
 		RecordBuilder.SetMeta(FCbObject(Record.Get().GetMeta()));
 	}
 
-	if (FPayload Payload = Record.Get().GetValuePayload())
+	if (FValueWithId Value = Record.Get().GetValue())
 	{
-		const ECachePolicy PayloadPolicy = Policy.GetPayloadPolicy(Payload.GetId());
-		GetCacheContent(Key, RecordName, PayloadPolicy, ECachePolicy::SkipValue, Payload, OutStatus);
-		if (Payload.IsNull())
+		const ECachePolicy ValuePolicy = Policy.GetValuePolicy(Value.GetId());
+		GetCacheContent(Key, RecordName, ValuePolicy, ECachePolicy::SkipValue, Value, OutStatus);
+		if (Value.IsNull())
 		{
 			return FOptionalCacheRecord();
 		}
-		RecordBuilder.SetValue(MoveTemp(Payload));
+		RecordBuilder.SetValue(MoveTemp(Value));
 	}
 
-	for (FPayload Payload : Record.Get().GetAttachmentPayloads())
+	for (FValueWithId Value : Record.Get().GetAttachments())
 	{
-		const ECachePolicy PayloadPolicy = Policy.GetPayloadPolicy(Payload.GetId());
-		GetCacheContent(Key, RecordName, PayloadPolicy, ECachePolicy::SkipAttachments, Payload, OutStatus);
-		if (Payload.IsNull())
+		const ECachePolicy ValuePolicy = Policy.GetValuePolicy(Value.GetId());
+		GetCacheContent(Key, RecordName, ValuePolicy, ECachePolicy::SkipAttachments, Value, OutStatus);
+		if (Value.IsNull())
 		{
 			return FOptionalCacheRecord();
 		}
-		RecordBuilder.AddAttachment(MoveTemp(Payload));
+		RecordBuilder.AddAttachment(MoveTemp(Value));
 	}
 
 	return RecordBuilder.Build();
@@ -849,21 +849,21 @@ void FPakFileDerivedDataBackend::GetCacheContent(
 	const FStringView ContentName,
 	const ECachePolicy Policy,
 	const ECachePolicy SkipFlag,
-	FPayload& InOutPayload,
+	FValueWithId& InOutValue,
 	EStatus& InOutStatus)
 {
-	if (!EnumHasAnyFlags(Policy, ECachePolicy::Query) || (EnumHasAnyFlags(Policy, SkipFlag) && InOutPayload.HasData()))
+	if (!EnumHasAnyFlags(Policy, ECachePolicy::Query) || (EnumHasAnyFlags(Policy, SkipFlag) && InOutValue.HasData()))
 	{
-		InOutPayload = InOutPayload.RemoveData();
+		InOutValue = InOutValue.RemoveData();
 		return;
 	}
 
-	if (InOutPayload.HasData())
+	if (InOutValue.HasData())
 	{
 		return;
 	}
 
-	const FIoHash& RawHash = InOutPayload.GetRawHash();
+	const FIoHash& RawHash = InOutValue.GetRawHash();
 
 	TStringBuilder<256> Path;
 	FPathViews::Append(Path, TEXT("Content"), RawHash);
@@ -881,30 +881,30 @@ void FPakFileDerivedDataBackend::GetCacheContent(
 			if (FCompressedBuffer CompressedBuffer = FCompressedBuffer::FromCompressed(MoveTemp(CompressedData));
 				CompressedBuffer && CompressedBuffer.GetRawHash() == RawHash)
 			{
-				InOutPayload = FPayload(InOutPayload.GetId(), MoveTemp(CompressedBuffer));
+				InOutValue = FValueWithId(InOutValue.GetId(), MoveTemp(CompressedBuffer));
 				return;
 			}
 			UE_LOG(LogDerivedDataCache, Display,
-				TEXT("%s: Cache miss with corrupted payload %s with hash %s for %s from '%.*s'"),
-				*CachePath, *WriteToString<16>(InOutPayload.GetId()), *WriteToString<48>(RawHash),
+				TEXT("%s: Cache miss with corrupted value %s with hash %s for %s from '%.*s'"),
+				*CachePath, *WriteToString<16>(InOutValue.GetId()), *WriteToString<48>(RawHash),
 				*WriteToString<96>(Key), ContentName.Len(), ContentName.GetData());
 			InOutStatus = EStatus::Error;
 			if (!EnumHasAnyFlags(Policy, ECachePolicy::PartialOnError))
 			{
-				InOutPayload = FPayload::Null;
+				InOutValue = FValueWithId::Null;
 			}
 			return;
 		}
 	}
 
 	UE_LOG(LogDerivedDataCache, Verbose,
-		TEXT("%s: Cache miss with missing payload %s with hash %s for %s from '%.*s'"),
-		*CachePath, *WriteToString<16>(InOutPayload.GetId()), *WriteToString<48>(RawHash), *WriteToString<96>(Key),
+		TEXT("%s: Cache miss with missing value %s with hash %s for %s from '%.*s'"),
+		*CachePath, *WriteToString<16>(InOutValue.GetId()), *WriteToString<48>(RawHash), *WriteToString<96>(Key),
 		ContentName.Len(), ContentName.GetData());
 	InOutStatus = EStatus::Error;
 	if (!EnumHasAnyFlags(Policy, ECachePolicy::PartialOnError))
 	{
-		InOutPayload = FPayload::Null;
+		InOutValue = FValueWithId::Null;
 	}
 }
 
diff --git a/Engine/Source/Developer/DerivedDataCache/Private/PakFileDerivedDataBackend.h b/Engine/Source/Developer/DerivedDataCache/Private/PakFileDerivedDataBackend.h
--- a/Engine/Source/Developer/DerivedDataCache/Private/PakFileDerivedDataBackend.h
+++ b/Engine/Source/Developer/DerivedDataCache/Private/PakFileDerivedDataBackend.h
@@ -14,7 +14,7 @@ class FCompressedBuffer;
 class IFileHandle;
 
 namespace UE::DerivedData { class FOptionalCacheRecord; }
-namespace UE::DerivedData { class FPayload; }
+namespace UE::DerivedData { class FValueWithId; }
 
 namespace UE::DerivedData::CacheStore::PakFile
 {
@@ -140,7 +140,7 @@ private:
 		const FStringView ContentName,
 		const ECachePolicy Policy,
 		const ECachePolicy SkipFlag,
-		FPayload& InOutPayload,
+		FValueWithId& InOutValue,
 		EStatus& InOutStatus);
 
 	bool SaveFile(FStringView Path, FStringView DebugName, TFunctionRef<void (FArchive&)> WriteFunction);
diff --git a/Engine/Source/Developer/DerivedDataCache/Private/S3DerivedDataBackend.cpp b/Engine/Source/Developer/DerivedDataCache/Private/S3DerivedDataBackend.cpp
--- a/Engine/Source/Developer/DerivedDataCache/Private/S3DerivedDataBackend.cpp
+++ b/Engine/Source/Developer/DerivedDataCache/Private/S3DerivedDataBackend.cpp
@@ -17,6 +17,7 @@
 #include <openssl/sha.h>
 #include <openssl/hmac.h>
 #include "Algo/AllOf.h"
+#include "Memory/SharedBuffer.h"
 #include "Misc/Base64.h"
 #include "Misc/Paths.h"
 #include "Misc/SecureHash.h"
@@ -27,7 +28,6 @@
 #include "Serialization/MemoryReader.h"
 #include "DerivedDataBackendCorruptionWrapper.h"
 #include "DerivedDataCacheRecord.h"
-#include "DerivedDataPayload.h"
 #include "Dom/JsonObject.h"
 #include "Serialization/JsonReader.h"
 #include "Serialization/JsonSerializer.h"
diff --git a/Engine/Source/Developer/DerivedDataCache/Private/ZenDerivedDataBackend.cpp b/Engine/Source/Developer/DerivedDataCache/Private/ZenDerivedDataBackend.cpp
--- a/Engine/Source/Developer/DerivedDataCache/Private/ZenDerivedDataBackend.cpp
+++ b/Engine/Source/Developer/DerivedDataCache/Private/ZenDerivedDataBackend.cpp
@@ -7,7 +7,7 @@
 #include "Containers/StringFwd.h"
 #include "DerivedDataCacheRecord.h"
 #include "DerivedDataChunk.h"
-#include "DerivedDataPayload.h"
+#include "DerivedDataValue.h"
 #include "Misc/ConfigCacheIni.h"
 #include "Misc/ScopeLock.h"
 #include "ProfilingDebugging/CountersTrace.h"
@@ -352,10 +352,10 @@ void FZenDerivedDataBackend::AppendZenUri(const FCacheKey& CacheKey, FStringBuil
 	Out << TEXT("/z$/") << CacheKey.Bucket << TEXT('/') << CacheKey.Hash;
 }
 
-void FZenDerivedDataBackend::AppendZenUri(const FCacheKey& CacheKey, const FPayloadId& PayloadId, FStringBuilderBase& Out)
+void FZenDerivedDataBackend::AppendZenUri(const FCacheKey& CacheKey, const FValueId& Id, FStringBuilderBase& Out)
 {
 	AppendZenUri(CacheKey, Out);
-	Out << TEXT('/') << PayloadId;
+	Out << TEXT('/') << Id;
 }
 
 void FZenDerivedDataBackend::AppendPolicyQueryString(ECachePolicy Policy, FStringBuilderBase& Uri)
@@ -423,8 +423,8 @@ void FZenDerivedDataBackend::AppendPolicyQueryString(ECachePolicy Policy, FStrin
 uint64 FZenDerivedDataBackend::MeasureCacheRecord(const FCacheRecord& Record)
 {
 	return Record.GetMeta().GetSize() +
-		Record.GetValuePayload().GetRawSize() +
-		Algo::TransformAccumulate(Record.GetAttachmentPayloads(), &FPayload::GetRawSize, uint64(0));
+		Record.GetValue().GetRawSize() +
+		Algo::TransformAccumulate(Record.GetAttachments(), &FValue::GetRawSize, uint64(0));
 }
 
 void FZenDerivedDataBackend::RemoveCachedData(const TCHAR* CacheKey, bool bTransient)
@@ -668,17 +668,17 @@ void FZenDerivedDataBackend::Get(
 				BatchRequest.BeginObject("Policy"_ASV);
 				{
 					BatchRequest << "RecordPolicy"_ASV << static_cast<uint32>(Policy.GetRecordPolicy());
-					BatchRequest << "DefaultPayloadPolicy"_ASV << static_cast<uint32>(Policy.GetDefaultPayloadPolicy());
+					BatchRequest << "DefaultPayloadPolicy"_ASV << static_cast<uint32>(Policy.GetDefaultValuePolicy());
 					
-					TConstArrayView<FCachePayloadPolicy> PayloadPolicies = Policy.GetPayloadPolicies();
-					if (PayloadPolicies.Num())
+					TConstArrayView<FCacheValuePolicy> ValuePolicies = Policy.GetValuePolicies();
+					if (ValuePolicies.Num())
 					{
 						BatchRequest.BeginArray("PayloadPolicies"_ASV);
-						for (const FCachePayloadPolicy& PayloadPolicy : PayloadPolicies)
+						for (const FCacheValuePolicy& ValuePolicy : ValuePolicies)
 						{
 							BatchRequest.BeginObject();
-							BatchRequest.AddObjectId("Id"_ASV, PayloadPolicy.Id);
-							BatchRequest << "Policy"_ASV << static_cast<uint32>(PayloadPolicy.Policy);
+							BatchRequest.AddObjectId("Id"_ASV, ValuePolicy.Id);
+							BatchRequest << "Policy"_ASV << static_cast<uint32>(ValuePolicy.Policy);
 							BatchRequest.EndObject();
 						}
 						BatchRequest.EndArray();
@@ -873,7 +873,7 @@ void FZenDerivedDataBackend::GetChunks(
 				}
 				else
 				{
-					UE_LOG(LogDerivedDataCache, Display, TEXT("%s: Cache miss with missing payload '%s' for '%s' from '%s'"),
+					UE_LOG(LogDerivedDataCache, Display, TEXT("%s: Cache miss with missing value '%s' for '%s' from '%s'"),
 						*GetName(), *WriteToString<16>(Request.Id), *WriteToString<96>(Request.Key), *Request.Name);
 					
 					if (OnComplete)
@@ -891,7 +891,7 @@ void FZenDerivedDataBackend::GetChunks(
 			{
 				const FCacheChunkRequest& Request = SortedRequests[ChunkIndex];
 				
-				UE_LOG(LogDerivedDataCache, Display, TEXT("%s: Cache miss with missing payload '%s' for '%s' from '%s'"),
+				UE_LOG(LogDerivedDataCache, Display, TEXT("%s: Cache miss with missing value '%s' for '%s' from '%s'"),
 					*GetName(), *WriteToString<16>(Request.Id), *WriteToString<96>(Request.Key), *Request.Name);
 				
 				if (OnComplete)
diff --git a/Engine/Source/Developer/DerivedDataCache/Private/ZenDerivedDataBackend.h b/Engine/Source/Developer/DerivedDataCache/Private/ZenDerivedDataBackend.h
--- a/Engine/Source/Developer/DerivedDataCache/Private/ZenDerivedDataBackend.h
+++ b/Engine/Source/Developer/DerivedDataCache/Private/ZenDerivedDataBackend.h
@@ -37,8 +37,7 @@ namespace UE::Zen {
 namespace UE::DerivedData {
 	struct FCacheKey;
 	class FOptionalCacheRecord;
-	class FPayload;
-	struct FPayloadId;
+	struct FValueId;
 }
 
 namespace UE::DerivedData::Backends {
@@ -123,7 +122,7 @@ private:
 	bool IsServiceReady();
 	static FString MakeLegacyZenKey(const TCHAR* CacheKey);
 	static void AppendZenUri(const FCacheKey& CacheKey, FStringBuilderBase& Out);
-	static void AppendZenUri(const FCacheKey& CacheKey, const FPayloadId& PayloadId, FStringBuilderBase& Out);
+	static void AppendZenUri(const FCacheKey& CacheKey, const FValueId& Id, FStringBuilderBase& Out);
 	static void AppendPolicyQueryString(ECachePolicy Policy, FStringBuilderBase& Out);
 
 	static bool ShouldRetryOnError(int64 ResponseCode);
diff --git a/Engine/Source/Developer/DerivedDataCache/Public/DerivedDataBuild.h b/Engine/Source/Developer/DerivedDataCache/Public/DerivedDataBuild.h
--- a/Engine/Source/Developer/DerivedDataCache/Public/DerivedDataBuild.h
+++ b/Engine/Source/Developer/DerivedDataCache/Public/DerivedDataBuild.h
@@ -35,7 +35,7 @@ namespace UE::DerivedData
  * input references into the referenced data. Use FBuildSession::Build() to schedule a definition
  * to build, along with any of its transitive build dependencies.
  *
- * Implement IBuildFunction, with a unique name and version, to add payloads to the build context
+ * Implement a IBuildFunction, with a unique name and version, to add values to the build context
  * based on constants and inputs in the context. Use TBuildFunctionFactory to add the function to
  * the registry at IBuild::GetFunctionRegistry() to allow the build job to find it.
  */
@@ -79,8 +79,7 @@ public:
 	 * Create a build session.
 	 *
 	 * An input resolver is required for the session to perform builds with unresolved inputs, or
-	 * to resolve build payload keys. An input resolver is optional when the session is used only
-	 * to build actions with any inputs provided to the session directly.
+	 * to resolve build value keys.
 	 *
 	 * A default scheduler is used if one is not provided. Using the default is recommended.
 	 *
diff --git a/Engine/Source/Developer/DerivedDataCache/Public/DerivedDataBuildDefinition.h b/Engine/Source/Developer/DerivedDataCache/Public/DerivedDataBuildDefinition.h
--- a/Engine/Source/Developer/DerivedDataCache/Public/DerivedDataBuildDefinition.h
+++ b/Engine/Source/Developer/DerivedDataCache/Public/DerivedDataBuildDefinition.h
@@ -20,7 +20,7 @@ namespace UE::DerivedData { class FBuildDefinition; }
 namespace UE::DerivedData { class FBuildDefinitionBuilder; }
 namespace UE::DerivedData { class FOptionalBuildDefinition; }
 namespace UE::DerivedData { struct FBuildKey; }
-namespace UE::DerivedData { struct FBuildPayloadKey; }
+namespace UE::DerivedData { struct FBuildValueKey; }
 
 namespace UE::DerivedData::Private
 {
@@ -35,7 +35,7 @@ public:
 	virtual bool HasConstants() const = 0;
 	virtual bool HasInputs() const = 0;
 	virtual void IterateConstants(TFunctionRef<void (FStringView Key, FCbObject&& Value)> Visitor) const = 0;
-	virtual void IterateInputBuilds(TFunctionRef<void (FStringView Key, const FBuildPayloadKey& PayloadKey)> Visitor) const = 0;
+	virtual void IterateInputBuilds(TFunctionRef<void (FStringView Key, const FBuildValueKey& ValueKey)> Visitor) const = 0;
 	virtual void IterateInputBulkData(TFunctionRef<void (FStringView Key, const FGuid& BulkDataId)> Visitor) const = 0;
 	virtual void IterateInputFiles(TFunctionRef<void (FStringView Key, FStringView Path)> Visitor) const = 0;
 	virtual void IterateInputHashes(TFunctionRef<void (FStringView Key, const FIoHash& RawHash)> Visitor) const = 0;
@@ -51,7 +51,7 @@ class IBuildDefinitionBuilderInternal
 public:
 	virtual ~IBuildDefinitionBuilderInternal() = default;
 	virtual void AddConstant(FStringView Key, const FCbObject& Value) = 0;
-	virtual void AddInputBuild(FStringView Key, const FBuildPayloadKey& PayloadKey) = 0;
+	virtual void AddInputBuild(FStringView Key, const FBuildValueKey& ValueKey) = 0;
 	virtual void AddInputBulkData(FStringView Key, const FGuid& BulkDataId) = 0;
 	virtual void AddInputFile(FStringView Key, FStringView Path) = 0;
 	virtual void AddInputHash(FStringView Key, const FIoHash& RawHash) = 0;
@@ -108,8 +108,8 @@ public:
 		return Definition->IterateConstants(Visitor);
 	}
 
-	/** Visits every input build payload in order by key. */
-	inline void IterateInputBuilds(TFunctionRef<void (FStringView Key, const FBuildPayloadKey& PayloadKey)> Visitor) const
+	/** Visits every input build value in order by key. */
+	inline void IterateInputBuilds(TFunctionRef<void (FStringView Key, const FBuildValueKey& ValueKey)> Visitor) const
 	{
 		return Definition->IterateInputBuilds(Visitor);
 	}
@@ -176,10 +176,10 @@ public:
 		DefinitionBuilder->AddConstant(Key, Value);
 	}
 
-	/** Add a payload from another build with a key that is unique within this definition. */
-	inline void AddInputBuild(FStringView Key, const FBuildPayloadKey& PayloadKey)
+	/** Add a value from another build with a key that is unique within this definition. */
+	inline void AddInputBuild(FStringView Key, const FBuildValueKey& ValueKey)
 	{
-		DefinitionBuilder->AddInputBuild(Key, PayloadKey);
+		DefinitionBuilder->AddInputBuild(Key, ValueKey);
 	}
 
 	/**
diff --git a/Engine/Source/Developer/DerivedDataCache/Public/DerivedDataBuildFunction.h b/Engine/Source/Developer/DerivedDataCache/Public/DerivedDataBuildFunction.h
--- a/Engine/Source/Developer/DerivedDataCache/Public/DerivedDataBuildFunction.h
+++ b/Engine/Source/Developer/DerivedDataCache/Public/DerivedDataBuildFunction.h
@@ -12,8 +12,8 @@ class FCompressedBuffer;
 namespace UE::DerivedData { class FBuildConfigContext; }
 namespace UE::DerivedData { class FBuildContext; }
 namespace UE::DerivedData { class FCacheBucket; }
-namespace UE::DerivedData { class FPayload; }
-namespace UE::DerivedData { struct FPayloadId; }
+namespace UE::DerivedData { class FValue; }
+namespace UE::DerivedData { struct FValueId; }
 namespace UE::DerivedData { enum class EBuildPolicy : uint32; }
 namespace UE::DerivedData { enum class ECachePolicy : uint32; }
 
@@ -90,12 +90,12 @@ public:
 	/** Returns the input with the matching key, or a null buffer if not found. */
 	virtual FSharedBuffer FindInput(FStringView Key) const = 0;
 
-	/** Adds a payload to the build output. Must have a non-null buffer and a unique ID. */
-	virtual void AddPayload(const FPayload& Payload) = 0;
-	virtual void AddPayload(const FPayloadId& Id, const FCompressedBuffer& Buffer) = 0;
-	virtual void AddPayload(const FPayloadId& Id, const FCompositeBuffer& Buffer, uint64 BlockSize = 0) = 0;
-	virtual void AddPayload(const FPayloadId& Id, const FSharedBuffer& Buffer, uint64 BlockSize = 0) = 0;
-	virtual void AddPayload(const FPayloadId& Id, const FCbObject& Object) = 0;
+	/** Adds a value to the build output. Must have a non-null buffer and a unique ID. */
+	virtual void AddValue(const FValueId& Id, const FValue& Value) = 0;
+	virtual void AddValue(const FValueId& Id, const FCompressedBuffer& Buffer) = 0;
+	virtual void AddValue(const FValueId& Id, const FCompositeBuffer& Buffer, uint64 BlockSize = 0) = 0;
+	virtual void AddValue(const FValueId& Id, const FSharedBuffer& Buffer, uint64 BlockSize = 0) = 0;
+	virtual void AddValue(const FValueId& Id, const FCbObject& Object) = 0;
 
 	/** Adds the error to the build output. Must be deterministic and based only on inputs. */
 	virtual void AddError(FStringView Message) = 0;
diff --git a/Engine/Source/Developer/DerivedDataCache/Public/DerivedDataBuildKey.h b/Engine/Source/Developer/DerivedDataCache/Public/DerivedDataBuildKey.h
--- a/Engine/Source/Developer/DerivedDataCache/Public/DerivedDataBuildKey.h
+++ b/Engine/Source/Developer/DerivedDataCache/Public/DerivedDataBuildKey.h
@@ -4,7 +4,7 @@
 
 #include "CoreTypes.h"
 #include "Containers/StringFwd.h"
-#include "DerivedDataPayloadId.h"
+#include "DerivedDataValueId.h"
 #include "IO/IoHash.h"
 #include "Templates/TypeHash.h"
 
@@ -33,17 +33,17 @@ struct FBuildActionKey
 
 inline const FBuildActionKey FBuildActionKey::Empty;
 
-/** A key that uniquely identifies a payload within a build output. */
-struct FBuildPayloadKey
+/** A key that uniquely identifies a value within a build output. */
+struct FBuildValueKey
 {
 	FBuildKey BuildKey;
-	FPayloadId Id;
+	FValueId Id;
 
-	/** A payload key with an empty build key and a null payload identifier. */
-	static const FBuildPayloadKey Empty;
+	/** A value key with an empty build key and a null value identifier. */
+	static const FBuildValueKey Empty;
 };
 
-inline const FBuildPayloadKey FBuildPayloadKey::Empty;
+inline const FBuildValueKey FBuildValueKey::Empty;
 
 ///////////////////////////////////////////////////////////////////////////////////////////////////////////////////////
 
@@ -103,30 +103,30 @@ inline TStringBuilderBase<CharType>& operator<<(TStringBuilderBase<CharType>& Bu
 
 ///////////////////////////////////////////////////////////////////////////////////////////////////////////////////////
 
-inline bool operator==(const FBuildPayloadKey& A, const FBuildPayloadKey& B)
+inline bool operator==(const FBuildValueKey& A, const FBuildValueKey& B)
 {
 	return A.BuildKey == B.BuildKey && A.Id == B.Id;
 }
 
-inline bool operator!=(const FBuildPayloadKey& A, const FBuildPayloadKey& B)
+inline bool operator!=(const FBuildValueKey& A, const FBuildValueKey& B)
 {
 	return A.BuildKey != B.BuildKey || A.Id != B.Id;
 }
 
-inline bool operator<(const FBuildPayloadKey& A, const FBuildPayloadKey& B)
+inline bool operator<(const FBuildValueKey& A, const FBuildValueKey& B)
 {
 	const FBuildKey& KeyA = A.BuildKey;
 	const FBuildKey& KeyB = B.BuildKey;
 	return KeyA == KeyB ? A.Id < B.Id : KeyA < KeyB;
 }
 
-inline uint32 GetTypeHash(const FBuildPayloadKey& Key)
+inline uint32 GetTypeHash(const FBuildValueKey& Key)
 {
 	return HashCombine(GetTypeHash(Key.BuildKey), GetTypeHash(Key.Id));
 }
 
 template <typename CharType>
-inline TStringBuilderBase<CharType>& operator<<(TStringBuilderBase<CharType>& Builder, const FBuildPayloadKey& Key)
+inline TStringBuilderBase<CharType>& operator<<(TStringBuilderBase<CharType>& Builder, const FBuildValueKey& Key)
 {
 	return Builder << Key.BuildKey << CharType('/') << Key.Id;
 }
diff --git a/Engine/Source/Developer/DerivedDataCache/Public/DerivedDataBuildOutput.h b/Engine/Source/Developer/DerivedDataCache/Public/DerivedDataBuildOutput.h
--- a/Engine/Source/Developer/DerivedDataCache/Public/DerivedDataBuildOutput.h
+++ b/Engine/Source/Developer/DerivedDataCache/Public/DerivedDataBuildOutput.h
@@ -19,10 +19,11 @@ namespace UE::DerivedData { class FBuildOutputBuilder; }
 namespace UE::DerivedData { class FCacheRecord; }
 namespace UE::DerivedData { class FCacheRecordBuilder; }
 namespace UE::DerivedData { class FOptionalBuildOutput; }
-namespace UE::DerivedData { class FPayload; }
+namespace UE::DerivedData { class FValue; }
+namespace UE::DerivedData { class FValueWithId; }
 namespace UE::DerivedData { struct FBuildOutputLog; }
 namespace UE::DerivedData { struct FBuildOutputMessage; }
-namespace UE::DerivedData { struct FPayloadId; }
+namespace UE::DerivedData { struct FValueId; }
 
 namespace UE::DerivedData::Private
 {
@@ -34,8 +35,8 @@ public:
 	virtual FStringView GetName() const = 0;
 	virtual FStringView GetFunction() const = 0;
 	virtual const FCbObject& GetMeta() const = 0;
-	virtual const FPayload& GetPayload(const FPayloadId& Id) const = 0;
-	virtual TConstArrayView<FPayload> GetPayloads() const = 0;
+	virtual const FValueWithId& GetValue(const FValueId& Id) const = 0;
+	virtual TConstArrayView<FValueWithId> GetValues() const = 0;
 	virtual TConstArrayView<FBuildOutputMessage> GetMessages() const = 0;
 	virtual TConstArrayView<FBuildOutputLog> GetLogs() const = 0;
 	virtual bool HasLogs() const = 0;
@@ -53,7 +54,7 @@ class IBuildOutputBuilderInternal
 public:
 	virtual ~IBuildOutputBuilderInternal() = default;
 	virtual void SetMeta(FCbObject&& Meta) = 0;
-	virtual void AddPayload(const FPayload& Payload) = 0;
+	virtual void AddValue(const FValueId& Id, const FValue& Value) = 0;
 	virtual void AddMessage(const FBuildOutputMessage& Message) = 0;
 	virtual void AddLog(const FBuildOutputLog& Log) = 0;
 	virtual bool HasError() const = 0;
@@ -103,11 +104,11 @@ struct FBuildOutputLog
 };
 
 /**
- * A build output is an immutable container of payloads, messages, and logs produced by a build.
+ * A build output is an immutable container of values, messages, and logs produced by a build.
  *
- * The output will not contain any payloads if it has any errors.
+ * The output will not contain any values if it has any errors.
  *
- * The output can be requested without data, which means that the payloads will have null data.
+ * The output can be requested without data, which means that the values will have null data.
  */
 class FBuildOutput
 {
@@ -121,11 +122,11 @@ public:
 	/** Returns the optional metadata. */
 	inline const FCbObject& GetMeta() const { return Output->GetMeta(); }
 
-	/** Returns the payload matching the ID. Null if no match. Buffer is null if skipped. */
-	inline const FPayload& GetPayload(const FPayloadId& Id) const { return Output->GetPayload(Id); }
+	/** Returns the value matching the ID. Null if no match. Buffer is null if skipped. */
+	inline const FValueWithId& GetValue(const FValueId& Id) const { return Output->GetValue(Id); }
 
-	/** Returns the payloads in the output in order by ID. */
-	inline TConstArrayView<FPayload> GetPayloads() const { return Output->GetPayloads(); }
+	/** Returns the values in the output in order by ID. */
+	inline TConstArrayView<FValueWithId> GetValues() const { return Output->GetValues(); }
 
 	/** Returns the messages in the order that they were recorded. */
 	inline TConstArrayView<FBuildOutputMessage> GetMessages() const { return Output->GetMessages(); }
@@ -139,7 +140,7 @@ public:
 	/** Returns whether the output has any errors. */
 	inline bool HasError() const { return Output->HasError(); }
 
-	/** Saves the build output to a compact binary object with payloads as attachments. */
+	/** Saves the build output to a compact binary object with values as attachments. */
 	void Save(FCbWriter& Writer) const
 	{
 		Output->Save(Writer);
@@ -184,16 +185,16 @@ private:
 class FBuildOutputBuilder
 {
 public:
-	/** Set the metadata for the build output. Holds a reference and is cloned if not owned. */
+	/** Set the metadata for the output. Holds a reference and is cloned if not owned. */
 	inline void SetMeta(FCbObject&& Meta)
 	{
 		return OutputBuilder->SetMeta(MoveTemp(Meta));
 	}
 
-	/** Add a payload to the output. The ID must be unique in this output. */
-	inline void AddPayload(const FPayload& Payload)
+	/** Add a value to the output. The ID must be unique in this output. */
+	inline void AddValue(const FValueId& Id, const FValue& Value)
 	{
-		OutputBuilder->AddPayload(Payload);
+		OutputBuilder->AddValue(Id, Value);
 	}
 
 	/** Add a message to the output. */
diff --git a/Engine/Source/Developer/DerivedDataCache/Public/DerivedDataBuildSession.h b/Engine/Source/Developer/DerivedDataCache/Public/DerivedDataBuildSession.h
--- a/Engine/Source/Developer/DerivedDataCache/Public/DerivedDataBuildSession.h
+++ b/Engine/Source/Developer/DerivedDataCache/Public/DerivedDataBuildSession.h
@@ -13,14 +13,14 @@ namespace UE::DerivedData { class FBuildAction; }
 namespace UE::DerivedData { class FBuildDefinition; }
 namespace UE::DerivedData { class FBuildSession; }
 namespace UE::DerivedData { class FOptionalBuildInputs; }
-namespace UE::DerivedData { class FPayload; }
+namespace UE::DerivedData { class FValueWithId; }
 namespace UE::DerivedData { class IRequestOwner; }
-namespace UE::DerivedData { struct FBuildPayloadCompleteParams; }
+namespace UE::DerivedData { struct FBuildValueCompleteParams; }
 
 namespace UE::DerivedData
 {
 
-using FOnBuildPayloadComplete = TUniqueFunction<void (FBuildPayloadCompleteParams&& Params)>;
+using FOnBuildValueComplete = TUniqueFunction<void (FBuildValueCompleteParams&& Params)>;
 
 } // UE::DerivedData
 
@@ -44,11 +44,11 @@ public:
 		const FBuildPolicy& Policy,
 		IRequestOwner& Owner,
 		FOnBuildComplete&& OnComplete) = 0;
-	virtual void BuildPayload(
-		const FBuildPayloadKey& Payload,
+	virtual void BuildValue(
+		const FBuildValueKey& Value,
 		EBuildPolicy Policy,
 		IRequestOwner& Owner,
-		FOnBuildPayloadComplete&& OnComplete) = 0;
+		FOnBuildValueComplete&& OnComplete) = 0;
 };
 
 FBuildSession CreateBuildSession(IBuildSessionInternal* Session);
@@ -63,7 +63,7 @@ namespace UE::DerivedData
  *
  * The purpose of a session is to group together related builds that use the same input resolver,
  * such as grouping builds by target platform. A request to build one definition can lead to more
- * builds being scheduled if the definition references payloads from other builds as inputs.
+ * builds being scheduled if the definition references values from other builds as inputs.
  */
 class FBuildSession
 {
@@ -117,22 +117,22 @@ public:
 	}
 
 	/**
-	 * Asynchronous request to execute a build according to the policy and return one payload.
+	 * Asynchronous request to execute a build according to the policy and return one value.
 	 *
 	 * The callback will always be called, and may be called from an arbitrary thread.
 	 *
-	 * @param Payload      The key identifying the build definition and the payload to return.
+	 * @param Value        The key identifying the build definition and the value to return.
 	 * @param Policy       Flags to control the behavior of the request. See EBuildPolicy.
 	 * @param Owner        The owner to execute the build within.
 	 * @param OnComplete   A callback invoked when the build completes or is canceled.
 	 */
-	inline void BuildPayload(
-		const FBuildPayloadKey& Payload,
+	inline void BuildValue(
+		const FBuildValueKey& Value,
 		EBuildPolicy Policy,
 		IRequestOwner& Owner,
-		FOnBuildPayloadComplete&& OnComplete)
+		FOnBuildValueComplete&& OnComplete)
 	{
-		Session->BuildPayload(Payload, Policy, Owner, MoveTemp(OnComplete));
+		Session->BuildValue(Value, Policy, Owner, MoveTemp(OnComplete));
 	}
 
 private:
@@ -175,20 +175,20 @@ public:
 	inline void Reset() { *this = FOptionalBuildSession(); }
 };
 
-/** Parameters for the completion callback for build payload requests. */
-struct FBuildPayloadCompleteParams
+/** Parameters for the completion callback for build value requests. */
+struct FBuildValueCompleteParams
 {
-	/** Key for the build request that completed or was canceled. See Payload for ID. */
+	/** Key for the build request that completed or was canceled. See Value for ID. */
 	FBuildKey Key;
 
 	/**
-	 * Payload from the build payload request that completed or was canceled.
+	 * Value from the build value request that completed or was canceled.
 	 *
 	 * The ID is always populated.
 	 * The hash and size are populated when Status is Ok.
 	 * The data is populated when Status is Ok and the data was not skipped by the policy.
 	 */
-	FPayload&& Payload;
+	FValueWithId&& Value;
 
 	/** Status of the build request. */
 	EStatus Status = EStatus::Error;
diff --git a/Engine/Source/Developer/DerivedDataCache/Public/DerivedDataBuildTypes.h b/Engine/Source/Developer/DerivedDataCache/Public/DerivedDataBuildTypes.h
--- a/Engine/Source/Developer/DerivedDataCache/Public/DerivedDataBuildTypes.h
+++ b/Engine/Source/Developer/DerivedDataCache/Public/DerivedDataBuildTypes.h
@@ -3,8 +3,8 @@
 #pragma once
 
 #include "CoreTypes.h"
-#include "DerivedDataPayloadId.h"
 #include "DerivedDataRequestTypes.h"
+#include "DerivedDataValueId.h"
 #include "Misc/EnumClassFlags.h"
 #include "Templates/RefCounting.h"
 
@@ -25,13 +25,13 @@ using FOnBuildComplete = TUniqueFunction<void (FBuildCompleteParams&& Params)>;
  *
  * The build policy flags can be combined to support a variety of usage patterns. Examples:
  *
- * Build(Default): allow cache get; build if missing; return every payload.
- * Build(Build | CacheStore): never query the cache; always build; return every payload.
- * Build(Cache | CacheSkipData): allow cache get; never build; skip every payload.
+ * Build(Default): allow cache get; build if missing; return every value.
+ * Build(Build | CacheStore): never query the cache; always build; return every value.
+ * Build(Cache | CacheSkipData): allow cache get; never build; skip every value.
  *
- * BuildPayload(Default): allow cache get; build if missing; return one payload.
- * BuildPayload(Build | CacheStore): never get from the cache; always build; return one payload.
- * BuildPayload(Cache | CacheSkipData): allow cache get; never build; skip every payload.
+ * BuildValue(Default): allow cache get; build if missing; return one value.
+ * BuildValue(Build | CacheStore): never get from the cache; always build; return one value.
+ * BuildValue(Cache | CacheSkipData): allow cache get; never build; skip every value.
  */
 enum class EBuildPolicy : uint32
 {
@@ -59,19 +59,19 @@ enum class EBuildPolicy : uint32
 	/** Keep records in the cache for at least the duration of the session. */
 	CacheKeepAlive      = 1 << 5,
 
-	/** Skip fetching or returning data for the payloads. */
+	/** Skip fetching or returning data for the values. */
 	SkipData            = 1 << 6,
 
-	/** Allow cache query+store, allow local+remote build when missed or skipped, and fetch the payload(s). */
+	/** Allow cache query+store, allow local+remote build when missed or skipped, and fetch the value(s). */
 	Default             = Build | Cache,
 };
 
 ENUM_CLASS_FLAGS(EBuildPolicy);
 
-/** A payload ID and the build policy to use for that payload. */
-struct FBuildPayloadPolicy
+/** A value ID and the build policy to use for that value. */
+struct FBuildValuePolicy
 {
-	FPayloadId Id;
+	FValueId Id;
 	EBuildPolicy Policy = EBuildPolicy::Default;
 };
 
@@ -82,48 +82,48 @@ public:
 	virtual ~IBuildPolicyShared() = default;
 	virtual void AddRef() const = 0;
 	virtual void Release() const = 0;
-	virtual TConstArrayView<FBuildPayloadPolicy> GetPayloadPolicies() const = 0;
-	virtual void AddPayloadPolicy(const FBuildPayloadPolicy& Policy) = 0;
+	virtual TConstArrayView<FBuildValuePolicy> GetValuePolicies() const = 0;
+	virtual void AddValuePolicy(const FBuildValuePolicy& Policy) = 0;
 	virtual void Build() = 0;
 };
 
-/** Flags to control the behavior of build requests, with optional overrides by payload. */
+/** Flags to control the behavior of build requests, with optional overrides by value. */
 class FBuildPolicy
 {
 public:
 	/** Construct a build policy that uses the default policy. */
 	FBuildPolicy() = default;
 
-	/** Construct a build policy with a single policy for every payload. */
+	/** Construct a build policy with a single policy for every value. */
 	inline FBuildPolicy(EBuildPolicy Policy)
 		: CombinedPolicy(Policy)
-		, DefaultPayloadPolicy(Policy)
+		, DefaultValuePolicy(Policy)
 	{
 	}
 
-	/** Returns true if every payload uses the same build policy. */
+	/** Returns true if every value uses the same build policy. */
 	inline bool IsUniform() const { return !Shared; }
 
-	/** Returns the build policy combined from the payload policies. */
+	/** Returns the build policy combined from the value policies. */
 	inline EBuildPolicy GetCombinedPolicy() const { return CombinedPolicy; }
 
-	/** Returns the build policy to use for the payload. */
-	UE_API EBuildPolicy GetPayloadPolicy(const FPayloadId& Id) const;
+	/** Returns the build policy to use for the value. */
+	UE_API EBuildPolicy GetValuePolicy(const FValueId& Id) const;
 
-	/** Returns the build policy to use for payloads with no override. */
-	inline EBuildPolicy GetDefaultPayloadPolicy() const { return DefaultPayloadPolicy; }
+	/** Returns the build policy to use for values with no override. */
+	inline EBuildPolicy GetDefaultValuePolicy() const { return DefaultValuePolicy; }
 
-	/** Returns the array of build policy overrides for payloads, sorted by ID. */
-	inline TConstArrayView<FBuildPayloadPolicy> GetPayloadPolicies() const
+	/** Returns the array of build policy overrides for values, sorted by ID. */
+	inline TConstArrayView<FBuildValuePolicy> GetValuePolicies() const
 	{
-		return Shared ? Shared->GetPayloadPolicies() : TConstArrayView<FBuildPayloadPolicy>();
+		return Shared ? Shared->GetValuePolicies() : TConstArrayView<FBuildValuePolicy>();
 	}
 
 private:
 	friend class FBuildPolicyBuilder;
 
 	EBuildPolicy CombinedPolicy = EBuildPolicy::Default;
-	EBuildPolicy DefaultPayloadPolicy = EBuildPolicy::Default;
+	EBuildPolicy DefaultValuePolicy = EBuildPolicy::Default;
 	TRefCountPtr<const Private::IBuildPolicyShared> Shared;
 };
 
@@ -134,15 +134,15 @@ public:
 	/** Construct a policy builder that uses the default policy. */
 	FBuildPolicyBuilder() = default;
 
-	/** Construct a policy builder that uses the provided policy for payloads with no override. */
+	/** Construct a policy builder that uses the provided policy for values with no override. */
 	inline explicit FBuildPolicyBuilder(EBuildPolicy Policy)
 		: BasePolicy(Policy)
 	{
 	}
 
-	/** Adds a build policy override for a payload. */
-	UE_API void AddPayloadPolicy(const FBuildPayloadPolicy& Policy);
-	inline void AddPayloadPolicy(const FPayloadId& Id, EBuildPolicy Policy) { AddPayloadPolicy({Id, Policy}); }
+	/** Adds a build policy override for a value. */
+	UE_API void AddValuePolicy(const FBuildValuePolicy& Policy);
+	inline void AddValuePolicy(const FValueId& Id, EBuildPolicy Policy) { AddValuePolicy({Id, Policy}); }
 
 	/** Build a build policy, which makes this builder subsequently unusable. */
 	UE_API FBuildPolicy Build();
@@ -197,7 +197,7 @@ struct FBuildCompleteParams
 	 *
 	 * The name, function, and diagnostics are always populated.
 	 *
-	 * The payloads are populated when Status is Ok, but with null data if skipped by the policy.
+	 * The values are populated when Status is Ok, but with null data if skipped by the policy.
 	 */
 	FBuildOutput&& Output;
 
diff --git a/Engine/Source/Developer/DerivedDataCache/Public/DerivedDataCache.h b/Engine/Source/Developer/DerivedDataCache/Public/DerivedDataCache.h
--- a/Engine/Source/Developer/DerivedDataCache/Public/DerivedDataCache.h
+++ b/Engine/Source/Developer/DerivedDataCache/Public/DerivedDataCache.h
@@ -6,9 +6,9 @@
 #include "Containers/StringFwd.h"
 #include "DerivedDataCacheKey.h"
 #include "DerivedDataCacheRecord.h"
-#include "DerivedDataPayloadId.h"
 #include "DerivedDataRequestTypes.h"
 #include "DerivedDataSharedString.h"
+#include "DerivedDataValueId.h"
 #include "Math/NumericLimits.h"
 #include "Misc/EnumClassFlags.h"
 #include "Templates/Function.h"
@@ -82,7 +82,7 @@ enum class ECachePolicy : uint32
 	/**
 	 * Skip fetching the data for any requests.
 	 *
-	 * Put requests with skip flags may assume that record existence implies payload existence.
+	 * Put requests with skip flags may assume that record existence implies value existence.
 	 */
 	SkipData        = SkipMeta | SkipValue | SkipAttachments,
 
@@ -95,14 +95,15 @@ enum class ECachePolicy : uint32
 	KeepAlive       = 1 << 7,
 
 	/**
-	 * Partial output will be provided with the error status when a required payload is missing.
+	 * Partial output will be provided with the error status when a required value is missing.
 	 *
-	 * This is meant for cases when the missing payloads can be individually recovered or rebuilt
+	 * This is meant for cases when the missing values can be individually recovered, or rebuilt,
 	 * without rebuilding the whole record. The cache automatically adds this flag when there are
-	 * other cache stores that it may be able to recover missing payloads from.
+	 * other cache stores that it may be able to recover missing values from.
 	 *
-	 * Requests for records would return records where the missing payloads have a hash and size,
-	 * but no data. Requests for chunks or values would return the hash and size, but no data.
+	 * Missing values will be returned in the records or chunks, but with only the hash and size.
+	 *
+	 * Applying this flag for a put of a record allows a partial record to be stored.
 	 */
 	PartialOnError  = 1 << 8,
 
@@ -122,10 +123,10 @@ ENUM_CLASS_FLAGS(ECachePolicy);
 
 ///////////////////////////////////////////////////////////////////////////////////////////////////////////////////////
 
-/** A payload ID and the cache policy to use for that payload. */
-struct FCachePayloadPolicy
+/** A value ID and the cache policy to use for that value. */
+struct FCacheValuePolicy
 {
-	FPayloadId Id;
+	FValueId Id;
 	ECachePolicy Policy = ECachePolicy::Default;
 };
 
@@ -136,19 +137,19 @@ public:
 	virtual ~ICacheRecordPolicyShared() = default;
 	virtual void AddRef() const = 0;
 	virtual void Release() const = 0;
-	virtual TConstArrayView<FCachePayloadPolicy> GetPayloadPolicies() const = 0;
-	virtual void AddPayloadPolicy(const FCachePayloadPolicy& Policy) = 0;
+	virtual TConstArrayView<FCacheValuePolicy> GetValuePolicies() const = 0;
+	virtual void AddValuePolicy(const FCacheValuePolicy& Policy) = 0;
 	virtual void Build() = 0;
 };
 
 /**
- * Flags to control the behavior of cache requests, with optional overrides by payload.
+ * Flags to control the behavior of cache requests, with optional overrides by value.
  *
  * Examples:
- * - A base policy of Disable with payload policy overrides of Default will fetch those payloads,
- *   if they exist in the record, and skip data for any other payloads.
- * - A base policy of Default with payload policy overrides of (Query | SkipData) will skip those
- *   payloads, but still check if they exist, and will load any other payloads.
+ * - A base policy of Disable, with value policy overrides of Default, will fetch those values if
+ *   they exist in the record, and skip data for any other values.
+ * - A base policy of Default, with value policy overrides of (Query | SkipData), will skip those
+ *   values, but still check if they exist, and will load any other values.
  */
 class FCacheRecordPolicy
 {
@@ -156,36 +157,36 @@ public:
 	/** Construct a cache record policy that uses the default policy. */
 	FCacheRecordPolicy() = default;
 
-	/** Construct a cache record policy with a uniform policy for the record and every payload. */
+	/** Construct a cache record policy with a uniform policy for the record and every value. */
 	inline FCacheRecordPolicy(ECachePolicy Policy)
 		: RecordPolicy(Policy)
-		, DefaultPayloadPolicy(Policy)
+		, DefaultValuePolicy(Policy)
 	{
 	}
 
-	/** Returns true if the record and every payload use the same cache policy. */
-	inline bool IsUniform() const { return !Shared && RecordPolicy == DefaultPayloadPolicy; }
+	/** Returns true if the record and every value use the same cache policy. */
+	inline bool IsUniform() const { return !Shared && RecordPolicy == DefaultValuePolicy; }
 
 	/** Returns the cache policy to use for the record. */
 	inline ECachePolicy GetRecordPolicy() const { return RecordPolicy; }
 
-	/** Returns the cache policy to use for the payload. */
-	UE_API ECachePolicy GetPayloadPolicy(const FPayloadId& Id) const;
+	/** Returns the cache policy to use for the value. */
+	UE_API ECachePolicy GetValuePolicy(const FValueId& Id) const;
 
-	/** Returns the cache policy to use for payloads with no override. */
-	inline ECachePolicy GetDefaultPayloadPolicy() const { return DefaultPayloadPolicy; }
+	/** Returns the cache policy to use for values with no override. */
+	inline ECachePolicy GetDefaultValuePolicy() const { return DefaultValuePolicy; }
 
-	/** Returns the array of cache policy overrides for payloads, sorted by ID. */
-	inline TConstArrayView<FCachePayloadPolicy> GetPayloadPolicies() const
+	/** Returns the array of cache policy overrides for values, sorted by ID. */
+	inline TConstArrayView<FCacheValuePolicy> GetValuePolicies() const
 	{
-		return Shared ? Shared->GetPayloadPolicies() : TConstArrayView<FCachePayloadPolicy>();
+		return Shared ? Shared->GetValuePolicies() : TConstArrayView<FCacheValuePolicy>();
 	}
 
 private:
 	friend class FCacheRecordPolicyBuilder;
 
 	ECachePolicy RecordPolicy = ECachePolicy::Default;
-	ECachePolicy DefaultPayloadPolicy = ECachePolicy::Default;
+	ECachePolicy DefaultValuePolicy = ECachePolicy::Default;
 	TRefCountPtr<const Private::ICacheRecordPolicyShared> Shared;
 };
 
@@ -196,15 +197,15 @@ public:
 	/** Construct a policy builder that uses the default policy as its base policy. */
 	FCacheRecordPolicyBuilder() = default;
 
-	/** Construct a policy builder that uses the provided policy for the record and payloads with no override. */
+	/** Construct a policy builder that uses the provided policy for the record and values with no override. */
 	inline explicit FCacheRecordPolicyBuilder(ECachePolicy Policy)
 		: BasePolicy(Policy)
 	{
 	}
 
-	/** Adds a cache policy override for a payload. */
-	UE_API void AddPayloadPolicy(const FCachePayloadPolicy& Policy);
-	inline void AddPayloadPolicy(const FPayloadId& Id, ECachePolicy Policy) { AddPayloadPolicy({Id, Policy}); }
+	/** Adds a cache policy override for a value. */
+	UE_API void AddValuePolicy(const FCacheValuePolicy& Policy);
+	inline void AddValuePolicy(const FValueId& Id, ECachePolicy Policy) { AddValuePolicy({Id, Policy}); }
 
 	/** Build a cache record policy, which makes this builder subsequently unusable. */
 	UE_API FCacheRecordPolicy Build();
@@ -235,9 +236,9 @@ public:
 	 * Records may finish storing in any order, and from multiple threads concurrently.
 	 *
 	 * A cache store is free to interpret a record containing only a key as a request to delete that
-	 * record from the store. Records may contain payloads that do not have data, and these payloads
-	 * must reference an existing copy of the payload in the store, if available, and must otherwise
-	 * be stored as a partial record that can attempt recovery of missing payloads when fetched.
+	 * record from the store. Records may contain values that do not have data, and such values must
+	 * reference an existing value in the store, if available. The PartialOnError policy may be used
+	 * to request that a partial record be stored when a value is missing.
 	 *
 	 * @param Requests     Requests with the cache records to store. Records must have a key.
 	 * @param Owner        The owner to execute the request within.
@@ -255,11 +256,12 @@ public:
 	 * Records may become available in any order, and from multiple threads concurrently.
 	 *
 	 * Records may propagate into other cache stores, in accordance with the policy. A propagated
-	 * record may be a partial record, with some payloads missing data, depending on the policy.
+	 * record may be a partial record, containing values without data, depending on the policy. A
+	 * propagated put of a partial record will include the PartialOnError policy.
 	 *
-	 * When payloads are required by the policy, but not available, the status must be Error. The
-	 * cache store must produce a partial record with the available payloads when the policy flag
-	 * PartialOnError is set for the missing payloads.
+	 * When values are required by the policy but not available, the status must be Error. If the
+	 * PartialOnError policy applies to the missing values, the cache store must return a partial
+	 * record containing the available values.
 	 *
 	 * @param Requests     Requests with the keys identifying the cache records to query.
 	 * @param Owner        The owner to execute the request within.
@@ -271,14 +273,14 @@ public:
 		FOnCacheGetComplete&& OnComplete) = 0;
 
 	/**
-	 * Asynchronous request to get chunks, which are subsets of payloads, from cache records.
+	 * Asynchronous request to get chunks, which are subsets of values, from records or values.
 	 *
 	 * The callback will always be called for every chunk, and may be called from an arbitrary thread.
 	 * Chunks may become available in any order, and from multiple threads concurrently.
 	 *
-	 * There is no requirement that the cache store validate the existence of other payloads from
-	 * the requested records. Requests for whole payloads may propagate those payloads into other
-	 * cache stores, in accordance with the policy, through a put of a partial record.
+	 * Values may propagate into other cache stores, in accordance with the policy, if the entire
+	 * value is requested, through a put of a partial record or a value. The cache store does not
+	 * need to verify the existence of unrequested values within requested records.
 	 *
 	 * @param Requests     The keys, IDs, offsets, and sizes of the chunks to query.
 	 * @param Owner        The owner to execute the request within.
@@ -378,7 +380,7 @@ struct FCacheGetCompleteParams
 	 * The key is always populated. The remainder of the record is populated when Status is Ok.
 	 *
 	 * The value, attachments, and metadata may be skipped based on cache policy flags. When a value
-	 * or attachment has been skipped, it will have a payload but its data will be null.
+	 * or attachment has been skipped, it will have a hash and size but its data will be null.
 	 */
 	FCacheRecord&& Record;
 
@@ -389,25 +391,25 @@ struct FCacheGetCompleteParams
 	EStatus Status = EStatus::Error;
 };
 
-/** Parameters to request a chunk, which is a subset of a payload, from a cache record. */
+/** Parameters to request a chunk, which is a subset of a value, from a cache record. */
 struct FCacheChunkRequest
 {
 	/** A name to identify this request for logging and profiling. An object path is typically sufficient. */
 	FCacheRequestName Name;
 
-	/** The key identifying the cache record to fetch the payload from. */
+	/** The key identifying the cache record to fetch the value from. */
 	FCacheKey Key;
 
-	/** The ID of the payload to fetch from the cache record. */
-	FPayloadId Id;
+	/** The ID of the value to fetch from the cache record. */
+	FValueId Id;
 
-	/** The offset into the raw bytes of the payload at which to start fetching. */
+	/** The offset into the raw bytes of the value at which to start fetching. */
 	uint64 RawOffset = 0;
 
-	/** The maximum number of raw bytes of the payload to fetch, starting from the offset. */
+	/** The maximum number of raw bytes of the value to fetch, starting from the offset. */
 	uint64 RawSize = MAX_uint64;
 
-	/** The raw hash of the entire payload to fetch, if available, otherwise zero. */
+	/** The raw hash of the entire value to fetch, if available, otherwise zero. */
 	FIoHash RawHash;
 
 	/** Flags to control the behavior of the request. See ECachePolicy. */
@@ -427,18 +429,18 @@ struct FCacheChunkCompleteParams
 	FCacheKey Key;
 
 	/** ID from the chunk request that completed or was canceled. */
-	FPayloadId Id;
+	FValueId Id;
 
 	/** Offset from the chunk request that completed or was canceled. */
 	uint64 RawOffset = 0;
 
-	/** Size, in bytes, of the subset of the payload that was fetched, if any. */
+	/** Size, in bytes, of the subset of the value that was fetched, if any. */
 	uint64 RawSize = 0;
 
-	/** Hash of the entire payload, even if only a subset was fetched. */
+	/** Hash of the entire value, even if only a subset was fetched. */
 	const FIoHash& RawHash;
 
-	/** Data for the subset of the payload that was fetched when Status is Ok, otherwise null. */
+	/** Data for the subset of the value that was fetched when Status is Ok, otherwise null. */
 	FSharedBuffer&& RawData;
 
 	/** A copy of the value from the request. */
diff --git a/Engine/Source/Developer/DerivedDataCache/Public/DerivedDataCacheRecord.h b/Engine/Source/Developer/DerivedDataCache/Public/DerivedDataCacheRecord.h
--- a/Engine/Source/Developer/DerivedDataCache/Public/DerivedDataCacheRecord.h
+++ b/Engine/Source/Developer/DerivedDataCache/Public/DerivedDataCacheRecord.h
@@ -3,7 +3,7 @@
 #pragma once
 
 #include "DerivedDataCacheKey.h"
-#include "DerivedDataPayloadId.h"
+#include "DerivedDataValueId.h"
 #include "Memory/MemoryFwd.h"
 #include "Misc/ScopeExit.h"
 #include "Templates/Function.h"
@@ -17,7 +17,8 @@ class FCbPackage;
 
 namespace UE::DerivedData { class FCacheRecord; }
 namespace UE::DerivedData { class FOptionalCacheRecord; }
-namespace UE::DerivedData { class FPayload; }
+namespace UE::DerivedData { class FValue; }
+namespace UE::DerivedData { class FValueWithId; }
 namespace UE::DerivedData { class IRequestOwner; }
 namespace UE::DerivedData { using FOnCacheRecordComplete = TUniqueFunction<void (FCacheRecord&& Record)>; }
 
@@ -32,10 +33,10 @@ public:
 	virtual ~ICacheRecordInternal() = default;
 	virtual const FCacheKey& GetKey() const = 0;
 	virtual const FCbObject& GetMeta() const = 0;
-	virtual const FPayload& GetValue() const = 0;
-	virtual const FPayload& GetAttachment(const FPayloadId& Id) const = 0;
-	virtual TConstArrayView<FPayload> GetAttachments() const = 0;
-	virtual const FPayload& GetPayload(const FPayloadId& Id) const = 0;
+	virtual const FValueWithId& GetValue() const = 0;
+	virtual const FValueWithId& GetAttachment(const FValueId& Id) const = 0;
+	virtual TConstArrayView<FValueWithId> GetAttachments() const = 0;
+	virtual const FValueWithId& GetValue(const FValueId& Id) const = 0;
 	virtual void AddRef() const = 0;
 	virtual void Release() const = 0;
 };
@@ -47,12 +48,14 @@ class ICacheRecordBuilderInternal
 public:
 	virtual ~ICacheRecordBuilderInternal() = default;
 	virtual void SetMeta(FCbObject&& Meta) = 0;
-	virtual void SetValue(const FCompositeBuffer& Buffer, const FPayloadId& Id, uint64 BlockSize) = 0;
-	virtual void SetValue(const FSharedBuffer& Buffer, const FPayloadId& Id, uint64 BlockSize) = 0;
-	virtual void SetValue(const FPayload& Payload) = 0;
-	virtual void AddAttachment(const FCompositeBuffer& Buffer, const FPayloadId& Id, uint64 BlockSize) = 0;
-	virtual void AddAttachment(const FSharedBuffer& Buffer, const FPayloadId& Id, uint64 BlockSize) = 0;
-	virtual void AddAttachment(const FPayload& Payload) = 0;
+	virtual void SetValue(const FCompositeBuffer& Buffer, const FValueId& Id, uint64 BlockSize) = 0;
+	virtual void SetValue(const FSharedBuffer& Buffer, const FValueId& Id, uint64 BlockSize) = 0;
+	virtual void SetValue(const FValue& Value, const FValueId& Id) = 0;
+	virtual void SetValue(const FValueWithId& Value) = 0;
+	virtual void AddAttachment(const FCompositeBuffer& Buffer, const FValueId& Id, uint64 BlockSize) = 0;
+	virtual void AddAttachment(const FSharedBuffer& Buffer, const FValueId& Id, uint64 BlockSize) = 0;
+	virtual void AddAttachment(const FValue& Value, const FValueId& Id) = 0;
+	virtual void AddAttachment(const FValueWithId& Value) = 0;
 	virtual FCacheRecord Build() = 0;
 	virtual void BuildAsync(IRequestOwner& Owner, FOnCacheRecordComplete&& OnComplete) = 0;
 };
@@ -73,11 +76,11 @@ namespace UE::DerivedData
  * spent creating the value.
  *
  * The value and its attachments are compressed, and both the compressed and uncompressed formats
- * of the payloads are exposed by the cache record.
+ * of the values are exposed by the cache record.
  *
  * The value, attachments, and metadata are optional and can be skipped when requesting a record.
- * When the value or attachments have been skipped, the record will contain a payload with a null
- * data but will otherwise be populated.
+ * When the value or attachments have been skipped, the record will contain a value with the hash
+ * and size, but with null data.
  */
 class FCacheRecord
 {
@@ -89,23 +92,16 @@ public:
 	inline const FCbObject& GetMeta() const { return Record->GetMeta(); }
 
 	/** Returns the value. Data is null if skipped by the policy on request. */
-	inline const FPayload& GetValue() const { return Record->GetValue(); }
-
-	/** Returns the value payload. Null if no value. Data is null if skipped by the policy on request. */
-	inline const FPayload& GetValuePayload() const { return Record->GetValue(); }
+	inline const FValueWithId& GetValue() const { return Record->GetValue(); }
 
 	/** Returns the attachment matching the ID. Data is null if skipped by the policy on request. */
-	inline const FPayload& GetAttachment(const FPayloadId& Id) const { return Record->GetAttachment(Id); }
-
-	/** Returns the attachment payload matching the ID. Null if no match. Data is null if skipped by the policy on request. */
-	inline const FPayload& GetAttachmentPayload(const FPayloadId& Id) const { return Record->GetAttachment(Id); }
+	inline const FValueWithId& GetAttachment(const FValueId& Id) const { return Record->GetAttachment(Id); }
 
 	/** Returns a view of the attachments. Always available, but data may be skipped if skipped. */
-	inline TConstArrayView<FPayload> GetAttachments() const { return Record->GetAttachments(); }
-	inline TConstArrayView<FPayload> GetAttachmentPayloads() const { return Record->GetAttachments(); }
+	inline TConstArrayView<FValueWithId> GetAttachments() const { return Record->GetAttachments(); }
 
-	/** Returns the payload matching the ID, whether value or attachment. Null if no match. Data is null if skipped. */
-	inline const FPayload& GetPayload(const FPayloadId& Id) const { return Record->GetPayload(Id); }
+	/** Returns the value matching the ID. Null if no match. Data is null if skipped. */
+	inline const FValueWithId& GetValue(const FValueId& Id) const { return Record->GetValue(Id); }
 
 	/** Save the cache record to a compact binary package. */
 	UE_API FCbPackage Save() const;
@@ -133,9 +129,6 @@ private:
  * Create using a key that uniquely corresponds to the value and attachments for the cache record.
  * Metadata may vary between records of the same key.
  *
- * The value and attachments can be provided as buffers, which will be compressed, or as payloads
- * which were previously compressed and have an identifier assigned.
- *
  * @see FCacheRecord
  */
 class FCacheRecordBuilder
@@ -163,13 +156,12 @@ public:
 	 * @param Id          An ID for the value that is unique within this cache record.
 	 *                    When omitted, the hash of the buffer will be used as the ID.
 	 * @param BlockSize   The power-of-two block size to encode raw data in. 0 is default.
-	 * @return The ID that was provided or created.
 	 */
-	inline void SetValue(const FCompositeBuffer& Buffer, const FPayloadId& Id = FPayloadId(), const uint64 BlockSize = 0)
+	inline void SetValue(const FCompositeBuffer& Buffer, const FValueId& Id = FValueId(), const uint64 BlockSize = 0)
 	{
 		return RecordBuilder->SetValue(Buffer, Id, BlockSize);
 	}
-	inline void SetValue(const FSharedBuffer& Buffer, const FPayloadId& Id = FPayloadId(), const uint64 BlockSize = 0)
+	inline void SetValue(const FSharedBuffer& Buffer, const FValueId& Id = FValueId(), const uint64 BlockSize = 0)
 	{
 		return RecordBuilder->SetValue(Buffer, Id, BlockSize);
 	}
@@ -177,12 +169,15 @@ public:
 	/**
 	 * Set the value for the cache record.
 	 *
-	 * @param Payload   The payload, which must have data unless it is known to be in the cache.
-	 * @return The ID that was provided. Unique within the scope of the cache record.
+	 * @param Value   The value, which must have data unless it is known to be in the cache.
 	 */
-	inline void SetValue(const FPayload& Payload)
+	inline void SetValue(const FValue& Value, const FValueId& Id = FValueId())
 	{
-		return RecordBuilder->SetValue(Payload);
+		return RecordBuilder->SetValue(Value, Id);
+	}
+	inline void SetValue(const FValueWithId& Value)
+	{
+		return RecordBuilder->SetValue(Value);
 	}
 
 	/**
@@ -192,13 +187,12 @@ public:
 	 * @param Id          An ID for the attachment that is unique within this cache record.
 	 *                    When omitted, the hash of the buffer will be used as the ID.
 	 * @param BlockSize   The power-of-two block size to encode raw data in. 0 is default.
-	 * @return The ID that was provided or created.
 	 */
-	inline void AddAttachment(const FCompositeBuffer& Buffer, const FPayloadId& Id = FPayloadId(), const uint64 BlockSize = 0)
+	inline void AddAttachment(const FCompositeBuffer& Buffer, const FValueId& Id = FValueId(), const uint64 BlockSize = 0)
 	{
 		return RecordBuilder->AddAttachment(Buffer, Id, BlockSize);
 	}
-	inline void AddAttachment(const FSharedBuffer& Buffer, const FPayloadId& Id = FPayloadId(), const uint64 BlockSize = 0)
+	inline void AddAttachment(const FSharedBuffer& Buffer, const FValueId& Id = FValueId(), const uint64 BlockSize = 0)
 	{
 		return RecordBuilder->AddAttachment(Buffer, Id, BlockSize);
 	}
@@ -206,12 +200,15 @@ public:
 	/**
 	 * Add an attachment to the cache record.
 	 *
-	 * @param Payload   The payload, which must have data unless it is known to be in the cache.
-	 * @return The ID that was provided. Unique within the scope of the cache record.
+	 * @param Value   The value, which must have data unless it is known to be in the cache.
 	 */
-	inline void AddAttachment(const FPayload& Payload)
+	inline void AddAttachment(const FValue& Value, const FValueId& Id = FValueId())
+	{
+		return RecordBuilder->AddAttachment(Value, Id);
+	}
+	inline void AddAttachment(const FValueWithId& Value)
 	{
-		return RecordBuilder->AddAttachment(Payload);
+		return RecordBuilder->AddAttachment(Value);
 	}
 
 	/**
@@ -229,7 +226,7 @@ public:
 	/**
 	 * Build a cache record asynchronously, which makes this builder subsequently unusable.
 	 *
-	 * Prefer Build() when the value and attachments are added by payload, as compression is already
+	 * Prefer Build() when the value and attachments are added by value, as compression is already
 	 * complete and BuildAsync() will complete immediately in that case.
 	 */
 	inline void BuildAsync(IRequestOwner& Owner, FOnCacheRecordComplete&& OnComplete)
diff --git a/Engine/Source/Developer/DerivedDataCache/Public/DerivedDataChunk.h b/Engine/Source/Developer/DerivedDataCache/Public/DerivedDataChunk.h
--- a/Engine/Source/Developer/DerivedDataCache/Public/DerivedDataChunk.h
+++ b/Engine/Source/Developer/DerivedDataCache/Public/DerivedDataChunk.h
@@ -7,7 +7,7 @@
 namespace UE::DerivedData
 {
 
-/** Binary predicate that compares chunks by key, then payload ID, then raw offset. */
+/** Binary predicate that compares chunks by key, then value ID, then raw offset. */
 struct TChunkEqual
 {
 	template <typename ChunkTypeA, typename ChunkTypeB>
@@ -16,7 +16,7 @@ struct TChunkEqual
 		return A.Key == B.Key && A.Id == B.Id && A.RawOffset == B.RawOffset;
 	}
 };
-/** Binary predicate that compares chunks by key, then payload ID, then raw offset. */
+/** Binary predicate that compares chunks by key, then value ID, then raw offset. */
 struct TChunkLess
 {
 	template <typename ChunkTypeA, typename ChunkTypeB>
diff --git a/Engine/Source/Developer/DerivedDataCache/Public/DerivedDataPayload.h b/Engine/Source/Developer/DerivedDataCache/Public/DerivedDataPayload.h
deleted file mode 100644
diff --git a/Engine/Source/Developer/DerivedDataCache/Public/DerivedDataPayloadId.h b/Engine/Source/Developer/DerivedDataCache/Public/DerivedDataPayloadId.h
deleted file mode 100644
diff --git a/Engine/Source/Developer/DerivedDataCache/Public/DerivedDataValue.h b/Engine/Source/Developer/DerivedDataCache/Public/DerivedDataValue.h
new file mode 100644
--- /dev/null
+++ b/Engine/Source/Developer/DerivedDataCache/Public/DerivedDataValue.h
@@ -0,0 +1,173 @@
+// Copyright Epic Games, Inc. All Rights Reserved.
+
+#pragma once
+
+#include "CoreTypes.h"
+#include "Compression/CompressedBuffer.h"
+#include "DerivedDataValueId.h"
+#include "IO/IoHash.h"
+#include "Memory/MemoryFwd.h"
+
+#define UE_API DERIVEDDATACACHE_API
+
+namespace UE::DerivedData
+{
+
+/**
+ * A value is a compressed buffer identified by its raw hash and raw size.
+ *
+ * A value without data may be used as a reference to the value.
+ */
+class FValue
+{
+public:
+	/**
+	 * Compress the buffer to a value using default derived data compression parameters.
+	 *
+	 * @param RawData     The raw data to be compressed.
+	 * @param BlockSize   The power-of-two block size to encode raw data in. 0 is default.
+	 */
+	UE_API static FValue Compress(const FCompositeBuffer& RawData, uint64 BlockSize = 0);
+	UE_API static FValue Compress(const FSharedBuffer& RawData, uint64 BlockSize = 0);
+
+	/** Construct a null value, which is a hash and size of zero with no data. */
+	FValue() = default;
+
+	/** Construct a value reference from the hash and size of its data. */
+	inline FValue(const FIoHash& RawHash, uint64 RawSize);
+
+	/** Construct a value from a compressed buffer, which is cloned if not owned. */
+	inline explicit FValue(const FCompressedBuffer& Data);
+	inline explicit FValue(FCompressedBuffer&& Data);
+
+	/** Returns the hash of the raw buffer (uncompressed) for the value. */
+	inline const FIoHash& GetRawHash() const { return RawHash; }
+
+	/** Returns the size of the raw buffer (uncompressed) for the value. */
+	inline uint64 GetRawSize() const { return RawSize; }
+
+	/** Returns the compressed buffer for the value. May be null. */
+	inline const FCompressedBuffer& GetData() const { return Data; }
+
+	/** Whether the compressed buffer for the value is available. */
+	inline bool HasData() const { return !!Data; }
+
+	/** Create a copy of the value with the data removed. */
+	[[nodiscard]] inline FValue RemoveData() const { return FValue(RawHash, RawSize); }
+
+	/** Reset this to null. */
+	inline void Reset() { *this = FValue(); }
+
+	/** A null value. */
+	static const FValue Null;
+
+private:
+	FIoHash RawHash;
+	uint64 RawSize = 0;
+	FCompressedBuffer Data;
+};
+
+/**
+ * A value with an ID. Null if both ID and value are null.
+ *
+ * May have a null value with a non-null ID. A non-null value must have a non-null ID.
+ */
+class FValueWithId : public FValue
+{
+public:
+	/** Construct a null value. */
+	FValueWithId() = default;
+
+	/** Construct a value with a non-null ID and forward the other args to FValue(). */
+	template <typename... ArgTypes>
+	inline explicit FValueWithId(const FValueId& Id, ArgTypes&&... Args);
+
+	/** Returns the ID for the value. */
+	inline const FValueId& GetId() const { return Id; }
+
+	/** Create a copy of the value with the data removed. */
+	[[nodiscard]] inline FValueWithId RemoveData() const { return FValueWithId(Id, FValue::RemoveData()); }
+
+	/** Whether this is null. */
+	inline bool IsNull() const { return Id.IsNull(); }
+	/** Whether this is not null. */
+	inline bool IsValid() const { return !IsNull(); }
+	/** Whether this is not null. */
+	inline explicit operator bool() const { return IsValid(); }
+
+	/** Reset this to null. */
+	inline void Reset() { FValue::Reset(); Id.Reset(); }
+
+	/** A null value with a null ID. */
+	static const FValueWithId Null;
+
+private:
+	FValueId Id;
+};
+
+inline const FValue FValue::Null;
+inline const FValueWithId FValueWithId::Null;
+
+inline FValue::FValue(const FIoHash& InRawHash, const uint64 InRawSize)
+	: RawHash(InRawHash)
+	, RawSize(InRawSize)
+{
+}
+
+inline FValue::FValue(const FCompressedBuffer& InData)
+	: RawHash(InData.GetRawHash())
+	, RawSize(InData.GetRawSize())
+	, Data(InData.MakeOwned())
+{
+}
+
+inline FValue::FValue(FCompressedBuffer&& InData)
+	: RawHash(InData.GetRawHash())
+	, RawSize(InData.GetRawSize())
+	, Data(MoveTemp(InData).MakeOwned())
+{
+}
+
+template <typename... ArgTypes>
+inline FValueWithId::FValueWithId(const FValueId& InId, ArgTypes&&... InArgs)
+	: FValue(Forward<ArgTypes>(InArgs)...)
+	, Id(InId)
+{
+	checkf(Id.IsValid(), TEXT("A valid ID is required to construct a value."));
+}
+
+/** Compare values by the hash and size of their raw buffer. */
+inline bool operator==(const FValue& A, const FValue& B)
+{
+	return A.GetRawHash() == B.GetRawHash() && A.GetRawSize() == B.GetRawSize();
+}
+
+/** Compare values by their ID and the hash and size of their raw buffer. */
+inline bool operator==(const FValueWithId& A, const FValueWithId& B)
+{
+	const FValue& ValueA = A;
+	const FValue& ValueB = B;
+	return A.GetId() == B.GetId() && ValueA == ValueB;
+}
+
+/** Compare values by the hash and size of their raw buffer. */
+inline bool operator<(const FValue& A, const FValue& B)
+{
+	const FIoHash& HashA = A.GetRawHash();
+	const FIoHash& HashB = B.GetRawHash();
+	return !(HashA == HashB) ? HashA < HashB : A.GetRawSize() < B.GetRawSize();
+}
+
+/** Compare values by their ID and the hash and size of their raw buffer. */
+inline bool operator<(const FValueWithId& A, const FValueWithId& B)
+{
+	const FValue& ValueA = A;
+	const FValue& ValueB = B;
+	const FValueId& IdA = A.GetId();
+	const FValueId& IdB = B.GetId();
+	return !(IdA == IdB) ? IdA < IdB : ValueA < ValueB;
+}
+
+} // UE::DerivedData
+
+#undef UE_API
diff --git a/Engine/Source/Developer/DerivedDataCache/Public/DerivedDataValueId.h b/Engine/Source/Developer/DerivedDataCache/Public/DerivedDataValueId.h
new file mode 100644
--- /dev/null
+++ b/Engine/Source/Developer/DerivedDataCache/Public/DerivedDataValueId.h
@@ -0,0 +1,110 @@
+// Copyright Epic Games, Inc. All Rights Reserved.
+
+#pragma once
+
+#include "CoreTypes.h"
+#include "Containers/StringFwd.h"
+#include "Memory/MemoryView.h"
+#include "String/BytesToHex.h"
+
+#define UE_API DERIVEDDATACACHE_API
+
+class FCbObjectId;
+struct FIoHash;
+
+namespace UE::DerivedData
+{
+
+/** A 12-byte value that uniquely identifies a value in the context that it was created. */
+struct FValueId
+{
+public:
+	using ByteArray = uint8[12];
+
+	/** Construct a null ID. */
+	FValueId() = default;
+
+	/** Construct an ID from an array of 12 bytes. */
+	inline explicit FValueId(const ByteArray& Id);
+
+	/** Construct an ID from a view of 12 bytes. */
+	UE_API explicit FValueId(FMemoryView Id);
+
+	/** Construct an ID from a Compact Binary Object ID. */
+	UE_API FValueId(const FCbObjectId& Id);
+
+	/** Returns the ID as a Compact Binary Object ID. */
+	UE_API operator FCbObjectId() const;
+
+	/** Construct an ID from a non-zero hash. */
+	[[nodiscard]] UE_API static FValueId FromHash(const FIoHash& Hash);
+
+	/** Construct an ID from a non-empty name. */
+	[[nodiscard]] UE_API static FValueId FromName(FUtf8StringView Name);
+	[[nodiscard]] UE_API static FValueId FromName(FWideStringView Name);
+
+	/** Returns a reference to the raw byte array for the ID. */
+	inline const ByteArray& GetBytes() const { return Bytes; }
+	inline operator const ByteArray&() const { return Bytes; }
+
+	/** Returns a view of the raw byte array for the ID. */
+	inline FMemoryView GetView() const { return MakeMemoryView(Bytes); }
+
+	/** Whether this is null. */
+	inline bool IsNull() const;
+	/** Whether this is not null. */
+	inline bool IsValid() const { return !IsNull(); }
+
+	/** Reset this to null. */
+	inline void Reset() { *this = FValueId(); }
+
+	/** A null ID. */
+	static const FValueId Null;
+
+private:
+	alignas(uint32) ByteArray Bytes{};
+};
+
+inline const FValueId FValueId::Null;
+
+inline FValueId::FValueId(const ByteArray& Id)
+{
+	FMemory::Memcpy(Bytes, Id, sizeof(ByteArray));
+}
+
+inline bool operator==(const FValueId& A, const FValueId& B)
+{
+	return A.GetView().EqualBytes(B.GetView());
+}
+
+inline bool operator!=(const FValueId& A, const FValueId& B)
+{
+	return !A.GetView().EqualBytes(B.GetView());
+}
+
+inline bool operator<(const FValueId& A, const FValueId& B)
+{
+	return A.GetView().CompareBytes(B.GetView()) < 0;
+}
+
+inline uint32 GetTypeHash(const FValueId& Id)
+{
+	return *reinterpret_cast<const uint32*>(Id.GetView().GetData());
+}
+
+/** Convert the ID to a 24-character hex string. */
+template <typename CharType>
+inline TStringBuilderBase<CharType>& operator<<(TStringBuilderBase<CharType>& Builder, const FValueId& Id)
+{
+	UE::String::BytesToHexLower(Id.GetBytes(), Builder);
+	return Builder;
+}
+
+inline bool FValueId::IsNull() const
+{
+	return *this == FValueId();
+}
+
+} // UE::DerivedData
+
+#undef UE_API
diff --git a/Engine/Source/Developer/TextureBuild/Private/TextureBuildFunction.cpp b/Engine/Source/Developer/TextureBuild/Private/TextureBuildFunction.cpp
--- a/Engine/Source/Developer/TextureBuild/Private/TextureBuildFunction.cpp
+++ b/Engine/Source/Developer/TextureBuild/Private/TextureBuildFunction.cpp
@@ -3,7 +3,7 @@
 #include "TextureBuildFunction.h"
 
 #include "DerivedDataCache.h"
-#include "DerivedDataPayloadId.h"
+#include "DerivedDataValueId.h"
 #include "Engine/TextureDefines.h"
 #include "IImageWrapper.h"
 #include "IImageWrapperModule.h"
@@ -408,7 +408,7 @@ void FTextureBuildFunction::Build(UE::DerivedData::FBuildContext& Context) const
 		DescriptionWriter.AddInteger("NumMipsInTail"_ASV, NumMipsInTail);
 
 		DescriptionWriter.BeginArray("Mips"_ASV);
-		int64 MipPayloadOffset = 0;
+		int64 MipOffset = 0;
 		for (int32 MipIndex = 0; MipIndex < MipCount; ++MipIndex)
 		{
 			FCompressedImage2D& CompressedMip = CompressedMips[MipIndex];
@@ -427,7 +427,7 @@ void FTextureBuildFunction::Build(UE::DerivedData::FBuildContext& Context) const
 			DescriptionWriter.AddInteger("FileRegion"_ASV, static_cast<int32>(FileRegion));
 			if (bIsInlineMip)
 			{
-				DescriptionWriter.AddInteger("PayloadOffset"_ASV, MipPayloadOffset);
+				DescriptionWriter.AddInteger("PayloadOffset"_ASV, MipOffset);
 			}
 			DescriptionWriter.AddInteger("NumBytes"_ASV, CompressedMip.RawData.Num());
 
@@ -435,14 +435,14 @@ void FTextureBuildFunction::Build(UE::DerivedData::FBuildContext& Context) const
 
 			if (bIsInlineMip)
 			{
-				MipPayloadOffset += CompressedMip.RawData.Num();
+				MipOffset += CompressedMip.RawData.Num();
 			}
 		}
 		DescriptionWriter.EndArray();
 
 		DescriptionWriter.EndObject();
 		FCbObject DescriptionObject = DescriptionWriter.Save().AsObject();
-		Context.AddPayload(UE::DerivedData::FPayloadId::FromName("Description"_ASV), DescriptionObject);
+		Context.AddValue(UE::DerivedData::FValueId::FromName("Description"_ASV), DescriptionObject);
 	}
 
 	// Streaming mips
@@ -452,7 +452,7 @@ void FTextureBuildFunction::Build(UE::DerivedData::FBuildContext& Context) const
 		MipName << "Mip"_ASV << MipIndex;
 
 		FSharedBuffer MipData = MakeSharedBufferFromArray(MoveTemp(CompressedMips[MipIndex].RawData));
-		Context.AddPayload(UE::DerivedData::FPayloadId::FromName(MipName), MipData);
+		Context.AddValue(UE::DerivedData::FValueId::FromName(MipName), MipData);
 	}
 
 	// Mip tail
@@ -465,6 +465,6 @@ void FTextureBuildFunction::Build(UE::DerivedData::FBuildContext& Context) const
 	FCompositeBuffer MipTail(MipTailComponents);
 	if (MipTail.GetSize() > 0)
 	{
-		Context.AddPayload(UE::DerivedData::FPayloadId::FromName("MipTail"_ASV), MipTail);
+		Context.AddValue(UE::DerivedData::FValueId::FromName("MipTail"_ASV), MipTail);
 	}
 }
diff --git a/Engine/Source/Developer/Virtualization/Private/VirtualizationDDCBackend.cpp b/Engine/Source/Developer/Virtualization/Private/VirtualizationDDCBackend.cpp
--- a/Engine/Source/Developer/Virtualization/Private/VirtualizationDDCBackend.cpp
+++ b/Engine/Source/Developer/Virtualization/Private/VirtualizationDDCBackend.cpp
@@ -6,18 +6,18 @@
 #include "Virtualization/PayloadId.h"
 #include "DerivedDataCache.h"
 #include "DerivedDataCacheRecord.h"
-#include "DerivedDataPayload.h"
 #include "DerivedDataRequestOwner.h"
+#include "DerivedDataValue.h"
 
 namespace UE::Virtualization
 {
 
-/** Utility function to help convert from UE::Virtualization types to UE::DerivedData::FPayload */
-UE::DerivedData::FPayload ToDDCPayload(const FPayloadId& Id, const FCompressedBuffer& Payload)
+/** Utility function to help convert from UE::Virtualization types to UE::DerivedData::FValueWithId */
+UE::DerivedData::FValueWithId ToDDCPayload(const FPayloadId& Id, const FCompressedBuffer& Payload)
 {
-	UE::DerivedData::FPayloadId DDCPayloadId = UE::DerivedData::FPayloadId::FromHash(Id.GetIdentifier());
+	UE::DerivedData::FValueId ValueId = UE::DerivedData::FValueId::FromHash(Id.GetIdentifier());
 
-	return UE::DerivedData::FPayload(DDCPayloadId, Payload);
+	return UE::DerivedData::FValueWithId(ValueId, Payload);
 }
 
 FDDCBackend::FDDCBackend(FStringView ConfigName, FStringView InDebugName)
@@ -89,7 +89,7 @@ EPushResult FDDCBackend::PushData(const FPayloadId& Id, const FCompressedBuffer&
 	Key.Bucket = Bucket;
 	Key.Hash = Id.GetIdentifier();
 
-	UE::DerivedData::FPayload DDCPayload = ToDDCPayload(Id, Payload);
+	UE::DerivedData::FValueWithId DDCPayload = ToDDCPayload(Id, Payload);
 	check(DDCPayload.GetRawHash() == Id.GetIdentifier());
 
 	UE::DerivedData::FCacheRecordBuilder RecordBuilder(Key);
@@ -138,7 +138,7 @@ FCompressedBuffer FDDCBackend::PullData(const FPayloadId& Id)
 		ResultStatus = Params.Status;
 		if (ResultStatus == UE::DerivedData::EStatus::Ok)
 		{
-			ResultData = Params.Record.GetValuePayload().GetData();
+			ResultData = Params.Record.GetValue().GetData();
 		}
 	};
 
diff --git a/Engine/Source/Editor/UnrealEd/Private/DerivedDataBuildLocalExecutor.cpp b/Engine/Source/Editor/UnrealEd/Private/DerivedDataBuildLocalExecutor.cpp
--- a/Engine/Source/Editor/UnrealEd/Private/DerivedDataBuildLocalExecutor.cpp
+++ b/Engine/Source/Editor/UnrealEd/Private/DerivedDataBuildLocalExecutor.cpp
@@ -8,9 +8,9 @@
 #include "DerivedDataBuildOutput.h"
 #include "DerivedDataBuildTypes.h"
 #include "DerivedDataBuildWorker.h"
-#include "DerivedDataPayload.h"
 #include "DerivedDataRequest.h"
 #include "DerivedDataRequestOwner.h"
+#include "DerivedDataValue.h"
 #include "Features/IModularFeatures.h"
 #include "HAL/Event.h"
 #include "HAL/PlatformProcess.h"
@@ -238,30 +238,30 @@ public:
 		}
 
 		FBuildOutputBuilder OutputBuilder = BuildSystem.CreateOutput(Action.GetName(), Action.GetFunction());
-		for (const FPayload& Payload : RemoteBuildOutput.Get().GetPayloads())
+		for (const FValueWithId& Value : RemoteBuildOutput.Get().GetValues())
 		{
-			if (EnumHasAnyFlags(Policy.GetPayloadPolicy(Payload.GetId()), EBuildPolicy::SkipData))
+			if (EnumHasAnyFlags(Policy.GetValuePolicy(Value.GetId()), EBuildPolicy::SkipData))
 			{
-				OutputBuilder.AddPayload(Payload);
+				OutputBuilder.AddValue(Value.GetId(), Value);
 			}
 			else
 			{
-				FCompressedBuffer BufferForPayload;
+				FCompressedBuffer BufferForValue;
 
 				TStringBuilder<128> Path;
-				FPathViews::Append(Path, SandboxRoot, TEXT("Outputs"), FIoHash(Payload.GetRawHash()));
+				FPathViews::Append(Path, SandboxRoot, TEXT("Outputs"), FIoHash(Value.GetRawHash()));
 				if (TUniquePtr<FArchive> Ar{IFileManager::Get().CreateFileReader(*Path, FILEREAD_Silent)})
 				{
-					BufferForPayload = FCompressedBuffer::Load(*Ar);
+					BufferForValue = FCompressedBuffer::Load(*Ar);
 				}
 
-				if (BufferForPayload.IsNull())
+				if (BufferForValue.IsNull())
 				{
 					UE_LOG(LogDerivedDataBuildLocalExecutor, Warning, TEXT("Remote execution system error: payload blob missing!"));
 					return OnComplete({Action.GetKey(), {}, {}, EStatus::Error});
 				}
 
-				OutputBuilder.AddPayload(FPayload(Payload.GetId(), BufferForPayload));
+				OutputBuilder.AddValue(Value.GetId(), FValue(MoveTemp(BufferForValue)));
 			}
 		}
 
diff --git a/Engine/Source/Editor/UnrealEd/Private/DerivedDataBuildRemoteExecutor.cpp b/Engine/Source/Editor/UnrealEd/Private/DerivedDataBuildRemoteExecutor.cpp
--- a/Engine/Source/Editor/UnrealEd/Private/DerivedDataBuildRemoteExecutor.cpp
+++ b/Engine/Source/Editor/UnrealEd/Private/DerivedDataBuildRemoteExecutor.cpp
@@ -10,9 +10,9 @@
 #include "DerivedDataBuildOutput.h"
 #include "DerivedDataBuildTypes.h"
 #include "DerivedDataBuildWorker.h"
-#include "DerivedDataPayload.h"
 #include "DerivedDataRequest.h"
 #include "DerivedDataRequestOwner.h"
+#include "DerivedDataValue.h"
 #include "Features/IModularFeatures.h"
 #include "HAL/Event.h"
 #include "HAL/PlatformProcess.h"
@@ -1015,16 +1015,16 @@ void FRemoteBuildExecutionRequest::OnGetResultPackageComplete(const UE::Zen::FZe
 		OutputBuilder.AddLog(Log);
 	}
 
-	for (const FPayload& Payload : RemoteBuildOutput.Get().GetPayloads())
+	for (const FValueWithId& Value : RemoteBuildOutput.Get().GetValues())
 	{
-		FCompressedBuffer BufferForPayload;
+		FCompressedBuffer BufferForValue;
 
-		if (const FCbAttachment* Attachment = State.ResultPackage.FindAttachment(Payload.GetRawHash()))
+		if (const FCbAttachment* Attachment = State.ResultPackage.FindAttachment(Value.GetRawHash()))
 		{
-			BufferForPayload = Attachment->AsCompressedBinary();
+			BufferForValue = Attachment->AsCompressedBinary();
 		}
 
-		if (BufferForPayload.IsNull())
+		if (BufferForValue.IsNull())
 		{
 			UE_LOG(LogDerivedDataBuildRemoteExecutor, Warning, TEXT("Remote execution system error: payload blob missing!"));
 			State.Owner.End(this, [this]
@@ -1035,7 +1035,7 @@ void FRemoteBuildExecutionRequest::OnGetResultPackageComplete(const UE::Zen::FZe
 			return;
 		}
 
-		OutputBuilder.AddPayload(FPayload(Payload.GetId(), BufferForPayload));
+		OutputBuilder.AddValue(Value.GetId(), FValue(MoveTemp(BufferForValue)));
 	}
 
 	FBuildOutput BuildOutput = OutputBuilder.Build();
@@ -1045,8 +1045,8 @@ void FRemoteBuildExecutionRequest::OnGetResultPackageComplete(const UE::Zen::FZe
 	State.Owner.End(this, [this, &BuildOutput]() mutable
 		{
 			CompletionCallback({ State.BuildAction.GetKey(), MoveTemp(BuildOutput), {}, EStatus::Ok });
-			CompletionEvent->Trigger();
 		});
+	CompletionEvent->Trigger();
 }
 
 } // namespace UE::DerivedData
diff --git a/Engine/Source/Editor/UnrealEd/Private/EditorDomain/EditorBuildInputResolver.cpp b/Engine/Source/Editor/UnrealEd/Private/EditorDomain/EditorBuildInputResolver.cpp
--- a/Engine/Source/Editor/UnrealEd/Private/EditorDomain/EditorBuildInputResolver.cpp
+++ b/Engine/Source/Editor/UnrealEd/Private/EditorDomain/EditorBuildInputResolver.cpp
@@ -116,7 +116,7 @@ void FEditorBuildInputResolver::ResolveInputMeta(const FBuildDefinition& Definit
 	FOnBuildInputMetaResolved&& OnResolved)
 {
 	EStatus OtherStatus = EStatus::Ok;
-	Definition.IterateInputBuilds([&OtherStatus, &Definition](FStringView Key, const FBuildPayloadKey& PayloadKey)
+	Definition.IterateInputBuilds([&OtherStatus, &Definition](FStringView Key, const FBuildValueKey& ValueKey)
 		{
 			UE_LOG(LogCore, Error, TEXT("FEditorBuildInputResolver::ResolveInputMeta: resolving InputBuilds is not yet implemented. Context=%.*s, Key=%.*s"),
 				Definition.GetName().Len(), Definition.GetName().GetData(), Key.Len(), Key.GetData());
@@ -152,7 +152,7 @@ void FEditorBuildInputResolver::ResolveInputData(const FBuildDefinition& Definit
 	FOnBuildInputDataResolved&& OnResolved, FBuildInputFilter&& Filter)
 {
 	EStatus OtherStatus = EStatus::Ok;
-	Definition.IterateInputBuilds([&OtherStatus, &Definition](FStringView Key, const FBuildPayloadKey& PayloadKey)
+	Definition.IterateInputBuilds([&OtherStatus, &Definition](FStringView Key, const FBuildValueKey& PayloadKey)
 		{
 			UE_LOG(LogCore, Error, TEXT("FEditorBuildInputResolver::ResolveInputData resolving InputBuilds is not yet implemented. Context=%.*s, Key=%.*s"),
 				Definition.GetName().Len(), Definition.GetName().GetData(), Key.Len(), Key.GetData());
diff --git a/Engine/Source/Editor/UnrealEd/Private/EditorDomain/EditorDomainArchive.cpp b/Engine/Source/Editor/UnrealEd/Private/EditorDomain/EditorDomainArchive.cpp
--- a/Engine/Source/Editor/UnrealEd/Private/EditorDomain/EditorDomainArchive.cpp
+++ b/Engine/Source/Editor/UnrealEd/Private/EditorDomain/EditorDomainArchive.cpp
@@ -2,12 +2,13 @@
 
 #include "EditorDomain/EditorDomainArchive.h"
 
+#include "Algo/Sort.h"
 #include "AssetRegistry/AssetData.h"
 #include "Async/AsyncFileHandleNull.h"
 #include "DerivedDataCache.h"
 #include "DerivedDataCacheKey.h"
 #include "DerivedDataCacheRecord.h"
-#include "DerivedDataPayload.h"
+#include "DerivedDataValue.h"
 #include "EditorDomain/EditorDomainSave.h"
 #include "EditorDomain/EditorDomainUtils.h"
 #include "HAL/UnrealMemory.h"
@@ -281,32 +282,32 @@ void FEditorDomainPackageSegments::OnRecordRequestComplete(UE::DerivedData::FCac
 	{
 		FCacheRecord& Record = Params.Record;
 		uint64 FileSize = Record.GetMeta()["FileSize"].AsUInt64();
-		const FPayload& ValuePayload = Record.GetValuePayload();
-
-		// Sort the attachment payloads by PayloadId, in case they are not already sorted. EditorDomainPackage attachments
-		// are written with PayloadIds in the same order as the segment order
-		TConstArrayView<FPayload> AttachmentPayloads = Record.GetAttachmentPayloads();
-		TArray<const FPayload*, TInlineAllocator<4>> SortedPayloads;
-		SortedPayloads.Reserve(AttachmentPayloads.Num()+1);
-		if (ValuePayload.GetRawSize() > 0)
+		const FValueWithId& RecordValue = Record.GetValue();
+
+		// Sort the attachment values by ValueId, in case they are not already sorted. EditorDomainPackage attachments
+		// are written with ValueIds in the same order as the segment order
+		TConstArrayView<FValueWithId> Attachments = Record.GetAttachments();
+		TArray<const FValueWithId*, TInlineAllocator<4>> SortedValues;
+		SortedValues.Reserve(Attachments.Num()+1);
+		if (RecordValue.GetRawSize() > 0)
 		{
-			SortedPayloads.Add(&ValuePayload);
+			SortedValues.Add(&RecordValue);
 		}
-		for (const FPayload& Payload : AttachmentPayloads)
+		for (const FValueWithId& Attachment : Attachments)
 		{
-			if (Payload.GetRawSize() > 0)
+			if (Attachment.GetRawSize() > 0)
 			{
-				SortedPayloads.Add(&Payload);
+				SortedValues.Add(&Attachment);
 			}
 		}
-		SortedPayloads.Sort([](const FPayload& A, const FPayload& B) { return A.GetId() < B.GetId(); });
+		Algo::SortBy(SortedValues, &FValueWithId::GetId);
 
 		uint64 CompositeSize = 0;
-		Segments.Reserve(SortedPayloads.Num());
-		for (const FPayload* Payload : SortedPayloads)
+		Segments.Reserve(SortedValues.Num());
+		for (const FValueWithId* Value : SortedValues)
 		{
-			Segments.Emplace(Payload->GetId(), CompositeSize);
-			CompositeSize += Payload->GetRawSize();
+			Segments.Emplace(Value->GetId(), CompositeSize);
+			CompositeSize += Value->GetRawSize();
 		}
 
 		if (CompositeSize != FileSize)
@@ -407,7 +408,7 @@ void FEditorDomainPackageSegments::SendSegmentRequest(FSegment& Segment)
 
 	// Note that Segment.RequestOwner is Interface-only and so we can write it outside the lock
 	Segment.RequestOwner.Emplace(EPriority::Normal);
-	FCacheChunkRequest SegmentChunk{{PackagePath.GetDebugName()}, UE::EditorDomain::GetEditorDomainPackageKey(EditorDomainHash), Segment.PayloadId};
+	FCacheChunkRequest SegmentChunk{{PackagePath.GetDebugName()}, UE::EditorDomain::GetEditorDomainPackageKey(EditorDomainHash), Segment.ValueId};
 	SegmentChunk.Policy = ECachePolicy::Local;
 	Cache.GetChunks({SegmentChunk}, *Segment.RequestOwner,
 		[this, &Segment](FCacheChunkCompleteParams&& Params)
diff --git a/Engine/Source/Editor/UnrealEd/Private/EditorDomain/EditorDomainArchive.h b/Engine/Source/Editor/UnrealEd/Private/EditorDomain/EditorDomainArchive.h
--- a/Engine/Source/Editor/UnrealEd/Private/EditorDomain/EditorDomainArchive.h
+++ b/Engine/Source/Editor/UnrealEd/Private/EditorDomain/EditorDomainArchive.h
@@ -5,8 +5,8 @@
 #include "Async/AsyncFileHandle.h"
 #include "Containers/Array.h"
 #include "DerivedDataCache.h"
-#include "DerivedDataPayloadId.h"
 #include "DerivedDataRequestOwner.h"
+#include "DerivedDataValueId.h"
 #include "EditorDomain/EditorDomain.h"
 #include "HAL/CriticalSection.h"
 #include "HAL/Platform.h"
@@ -102,22 +102,22 @@ private:
 	 */
 	struct FSegment
 	{
-		FSegment(const UE::DerivedData::FPayloadId& InPayloadId, uint64 InStart)
-			: PayloadId(InPayloadId)
+		FSegment(const UE::DerivedData::FValueId& InValueId, uint64 InStart)
+			: ValueId(InValueId)
 			, Start(InStart)
 		{}
 
 		/** Attachment ID to request from cache for this segment. Read-only. */
-		UE::DerivedData::FPayloadId PayloadId;
+		UE::DerivedData::FValueId ValueId;
 		/** The request that will be (has been) sent for the bytes from cache. Interface-only. */
 		TOptional<UE::DerivedData::FRequestOwner> RequestOwner;
 		/** Offset from the start of the entire archive. Read-only. */
 		uint64 Start;
-		/** Payload bytes. Callback-only before bAsyncComplete, Interface-only after. */
+		/** Value bytes. Callback-only before bAsyncComplete, Interface-only after. */
 		FSharedBuffer Data;
 		/** Interface thread cache of bAsyncComplete. Interface-only. */
 		bool bComplete = false;
-		/** Set to true whent the callback for the cache request has finished writing Data. */
+		/** Set to true when the callback for the cache request has finished writing Data. */
 		std::atomic<bool> bAsyncComplete{ false };
 	};
 
@@ -162,7 +162,7 @@ private:
 	 * Read-only, *pointer requires EditorDomainLocks.Lock.
 	 */
 	TRefCountPtr<FEditorDomain::FPackageSource> PackageSource;
-	/** EditorDomainHash for requesting payloads, copied from PackageSource. Read-only. */
+	/** EditorDomainHash for requesting values, copied from PackageSource. Read-only. */
 	FIoHash EditorDomainHash;
 
 	/**
diff --git a/Engine/Source/Editor/UnrealEd/Private/EditorDomain/EditorDomainUtils.cpp b/Engine/Source/Editor/UnrealEd/Private/EditorDomain/EditorDomainUtils.cpp
--- a/Engine/Source/Editor/UnrealEd/Private/EditorDomain/EditorDomainUtils.cpp
+++ b/Engine/Source/Editor/UnrealEd/Private/EditorDomain/EditorDomainUtils.cpp
@@ -12,8 +12,8 @@
 #include "DerivedDataCache.h"
 #include "DerivedDataCacheKey.h"
 #include "DerivedDataCacheRecord.h"
-#include "DerivedDataPayload.h"
 #include "DerivedDataRequestOwner.h"
+#include "DerivedDataValue.h"
 #include "Editor.h"
 #include "Engine/StaticMesh.h"
 #include "HAL/FileManager.h"
@@ -962,7 +962,7 @@ public:
 	struct FAttachment
 	{
 		FSharedBuffer Buffer;
-		UE::DerivedData::FPayloadId PayloadId;
+		UE::DerivedData::FValueId ValueId;
 	};
 	/** The Buffer+Id for each section making up the EditorDomain's copy of the package */
 	TConstArrayView<FAttachment> GetAttachments() const
@@ -1056,29 +1056,29 @@ protected:
 			BulkDataSize += AdditionalRecord.Buffer.GetSize();
 		}
 
-		// We use a counter for PayloadIds rather than hashes of the Attachments. We do this because
-		// some attachments may be identical, and Attachments are not allowed to have identical PayloadIds.
+		// We use a counter for ValueIds rather than hashes of the Attachments. We do this because
+		// some attachments may be identical, and Attachments are not allowed to have identical ValueIds.
 		// We need to keep the duplicate copies of identical payloads because BulkDatas were written into
 		// the exports with offsets that expect all attachment segments to exist in the segmented archive.
-		auto IntToPayloadId = [](uint32 Value)
+		auto IntToValueId = [](uint32 Value)
 		{
-			alignas(decltype(Value)) UE::DerivedData::FPayloadId::ByteArray Bytes{};
+			alignas(decltype(Value)) UE::DerivedData::FValueId::ByteArray Bytes{};
 			static_assert(sizeof(Bytes) >= sizeof(Value), "We are storing an integer counter in the Bytes array");
-			// The PayloadIds are sorted as an array of bytes, so the bytes of the integer must be written big-endian
+			// The ValueIds are sorted as an array of bytes, so the bytes of the integer must be written big-endian
 			for (int ByteIndex = 0; ByteIndex < sizeof(Value); ++ByteIndex)
 			{
 				Bytes[sizeof(Bytes) - 1 - ByteIndex] = static_cast<uint8>(Value & 0xff);
 				Value >>= 8;
 			}
-			return UE::DerivedData::FPayloadId(Bytes);
+			return UE::DerivedData::FValueId(Bytes);
 		};
 
-		uint32 AttachmentIndex = 1; // 0 is not a valid value for IntToPayloadId
+		uint32 AttachmentIndex = 1; // 0 is not a valid value for IntToValueId
 		Attachments.Reserve(AttachmentBuffers.Num());
 		FileSize = 0;
 		for (const FSharedBuffer& Buffer : AttachmentBuffers)
 		{
-			Attachments.Add(FAttachment{ Buffer, IntToPayloadId(AttachmentIndex++) });
+			Attachments.Add(FAttachment{ Buffer, IntToValueId(AttachmentIndex++) });
 			FileSize += Buffer.GetSize();
 		}
 		WritePackageRecord = MoveTemp(*Record.Package);
@@ -1319,7 +1319,7 @@ bool TrySavePackage(UPackage* Package)
 	FCacheRecordBuilder RecordBuilder(GetEditorDomainPackageKey(PackageDigest.Hash));
 	for (const FEditorDomainPackageWriter::FAttachment& Attachment : PackageWriter->GetAttachments())
 	{
-		RecordBuilder.AddAttachment(Attachment.Buffer, Attachment.PayloadId);
+		RecordBuilder.AddAttachment(Attachment.Buffer, Attachment.ValueId);
 	}
 	RecordBuilder.SetMeta(MetaData.Save().AsObject());
 	FRequestOwner Owner(EPriority::Normal);
diff --git a/Engine/Source/Programs/DerivedDataBuildWorker/Private/DerivedDataBuildWorker.cpp b/Engine/Source/Programs/DerivedDataBuildWorker/Private/DerivedDataBuildWorker.cpp
--- a/Engine/Source/Programs/DerivedDataBuildWorker/Private/DerivedDataBuildWorker.cpp
+++ b/Engine/Source/Programs/DerivedDataBuildWorker/Private/DerivedDataBuildWorker.cpp
@@ -10,8 +10,8 @@
 #include "DerivedDataBuildInputs.h"
 #include "DerivedDataBuildOutput.h"
 #include "DerivedDataBuildSession.h"
-#include "DerivedDataPayload.h"
 #include "DerivedDataRequestOwner.h"
+#include "DerivedDataValue.h"
 #include "HAL/FileManager.h"
 #include "Memory/SharedBuffer.h"
 #include "Misc/CommandLine.h"
@@ -299,13 +299,13 @@ void FBuildWorkerProgram::BuildComplete(FBuildCompleteParams&& Params) const
 		}
 	}
 
-	for (const FPayload& Payload : Output.GetPayloads())
+	for (const FValueWithId& Value : Output.GetValues())
 	{
-		if (Payload.HasData())
+		if (Value.HasData())
 		{
-			if (TUniquePtr<FArchive> Ar = OpenOutput(Output.GetName(), Payload.GetRawHash()))
+			if (TUniquePtr<FArchive> Ar = OpenOutput(Output.GetName(), Value.GetRawHash()))
 			{
-				Payload.GetData().Save(*Ar);
+				Value.GetData().Save(*Ar);
 				if (Ar->Close())
 				{
 					continue;
@@ -313,7 +313,7 @@ void FBuildWorkerProgram::BuildComplete(FBuildCompleteParams&& Params) const
 			}
 			UE_LOG(LogDerivedDataBuildWorker, Error,
 				TEXT("Failed to store build output %s for build of '%.*s' by %.*s."),
-				*WriteToString<48>(Payload.GetRawHash()), Name.Len(), Name.GetData(), Function.Len(), Function.GetData());
+				*WriteToString<48>(Value.GetRawHash()), Name.Len(), Name.GetData(), Function.Len(), Function.GetData());
 		}
 	}
 
diff --git a/Engine/Source/Runtime/CoreUObject/Private/Serialization/DerivedData.cpp b/Engine/Source/Runtime/CoreUObject/Private/Serialization/DerivedData.cpp
--- a/Engine/Source/Runtime/CoreUObject/Private/Serialization/DerivedData.cpp
+++ b/Engine/Source/Runtime/CoreUObject/Private/Serialization/DerivedData.cpp
@@ -7,7 +7,7 @@
 #if WITH_EDITORONLY_DATA
 
 #include "DerivedDataBuildDefinition.h"
-#include "DerivedDataPayloadId.h"
+#include "DerivedDataValueId.h"
 #include "Serialization/EditorDerivedData.h"
 #include "Serialization/EditorDerivedDataIoStore.h"
 #include "Templates/SharedPointer.h"
@@ -107,8 +107,8 @@ FDerivedData::FDerivedData(FStringView CacheKey, FStringView CacheContext)
 {
 }
 
-FDerivedData::FDerivedData(const DerivedData::FBuildDefinition& BuildDefinition, const DerivedData::FPayloadId& PayloadId)
-	: EditorData(DerivedData::Private::MakeEditorDerivedData(BuildDefinition, PayloadId))
+FDerivedData::FDerivedData(const DerivedData::FBuildDefinition& BuildDefinition, const DerivedData::FValueId& ValueId)
+	: EditorData(DerivedData::Private::MakeEditorDerivedData(BuildDefinition, ValueId))
 	, ChunkId(DerivedData::IoStore::GEditorDerivedDataIoStore->AddData(EditorData.Get()))
 {
 }
diff --git a/Engine/Source/Runtime/CoreUObject/Private/Serialization/EditorDerivedData.cpp b/Engine/Source/Runtime/CoreUObject/Private/Serialization/EditorDerivedData.cpp
--- a/Engine/Source/Runtime/CoreUObject/Private/Serialization/EditorDerivedData.cpp
+++ b/Engine/Source/Runtime/CoreUObject/Private/Serialization/EditorDerivedData.cpp
@@ -8,7 +8,7 @@
 #include "DerivedDataBuildDefinition.h"
 #include "DerivedDataBuildSession.h"
 #include "DerivedDataCacheInterface.h"
-#include "DerivedDataPayloadId.h"
+#include "DerivedDataValueId.h"
 #include "Memory/CompositeBuffer.h"
 #include "Memory/SharedBuffer.h"
 #include "Misc/CoreMisc.h"
@@ -167,9 +167,9 @@ bool FEditorDerivedDataCache::TryGetSize(uint64& OutSize) const
 class FEditorDerivedDataBuild final : public FEditorDerivedData
 {
 public:
-	FEditorDerivedDataBuild(const FBuildDefinition& InBuildDefinition, const FPayloadId& InPayloadId)
+	FEditorDerivedDataBuild(const FBuildDefinition& InBuildDefinition, const FValueId& InValueId)
 		: BuildDefinition(InBuildDefinition)
-		, PayloadId(InPayloadId)
+		, ValueId(InValueId)
 	{
 		unimplemented();
 	}
@@ -181,7 +181,7 @@ public:
 
 private:
 	FBuildDefinition BuildDefinition;
-	FPayloadId PayloadId;
+	FValueId ValueId;
 };
 
 void FEditorDerivedDataBuild::Read(IoStore::FDerivedDataIoRequest Request) const
@@ -218,9 +218,9 @@ TUniquePtr<FEditorDerivedData> MakeEditorDerivedData(FStringView CacheKey, FStri
 
 TUniquePtr<FEditorDerivedData> MakeEditorDerivedData(
 	const FBuildDefinition& BuildDefinition,
-	const FPayloadId& PayloadId)
+	const FValueId& ValueId)
 {
-	return MakeUnique<FEditorDerivedDataBuild>(BuildDefinition, PayloadId);
+	return MakeUnique<FEditorDerivedDataBuild>(BuildDefinition, ValueId);
 }
 
 } // UE::DerivedData::Private
diff --git a/Engine/Source/Runtime/CoreUObject/Private/Serialization/EditorDerivedData.h b/Engine/Source/Runtime/CoreUObject/Private/Serialization/EditorDerivedData.h
--- a/Engine/Source/Runtime/CoreUObject/Private/Serialization/EditorDerivedData.h
+++ b/Engine/Source/Runtime/CoreUObject/Private/Serialization/EditorDerivedData.h
@@ -15,7 +15,7 @@ class FIoRequestImpl;
 
 namespace UE::DerivedData { class FBuildDefinition; }
 namespace UE::DerivedData { class FRequestOwner; }
-namespace UE::DerivedData { struct FPayloadId; }
+namespace UE::DerivedData { struct FValueId; }
 namespace UE::DerivedData::IoStore { class FDerivedDataIoRequestQueue; }
 
 namespace UE::DerivedData::IoStore
@@ -65,7 +65,7 @@ TUniquePtr<FEditorDerivedData> MakeEditorDerivedData(FStringView CacheKey, FStri
 
 TUniquePtr<FEditorDerivedData> MakeEditorDerivedData(
 	const DerivedData::FBuildDefinition& BuildDefinition,
-	const DerivedData::FPayloadId& PayloadId);
+	const DerivedData::FValueId& ValueId);
 
 } // UE::DerivedData::Private
 
diff --git a/Engine/Source/Runtime/CoreUObject/Public/Serialization/DerivedData.h b/Engine/Source/Runtime/CoreUObject/Public/Serialization/DerivedData.h
--- a/Engine/Source/Runtime/CoreUObject/Public/Serialization/DerivedData.h
+++ b/Engine/Source/Runtime/CoreUObject/Public/Serialization/DerivedData.h
@@ -17,7 +17,7 @@ class UObject;
 template <typename FuncType> class TFunctionRef;
 
 namespace UE::DerivedData { class FBuildDefinition; }
-namespace UE::DerivedData { struct FPayloadId; }
+namespace UE::DerivedData { struct FValueId; }
 namespace UE::DerivedData::Private { class FEditorDerivedData; }
 
 namespace UE
@@ -65,7 +65,7 @@ public:
 
 	UE_API FDerivedData(FStringView CacheKey, FStringView CacheContext);
 
-	UE_API FDerivedData(const DerivedData::FBuildDefinition& BuildDefinition, const DerivedData::FPayloadId& PayloadId);
+	UE_API FDerivedData(const DerivedData::FBuildDefinition& BuildDefinition, const DerivedData::FValueId& ValueId);
 
 	UE_API void SetFlags(EDerivedDataFlags Flags);
 
diff --git a/Engine/Source/Runtime/Engine/Classes/Engine/Texture.h b/Engine/Source/Runtime/Engine/Classes/Engine/Texture.h
--- a/Engine/Source/Runtime/Engine/Classes/Engine/Texture.h
+++ b/Engine/Source/Runtime/Engine/Classes/Engine/Texture.h
@@ -43,7 +43,7 @@ class UAssetUserData;
 struct FPropertyChangedEvent;
 
 #if WITH_EDITORONLY_DATA
-namespace UE::DerivedData { struct FPayloadId; }
+namespace UE::DerivedData { struct FValueId; }
 #endif
 
 UENUM()
@@ -771,7 +771,7 @@ public:
 
 #if WITH_EDITORONLY_DATA
 	FString GetDerivedDataMipKeyString(int32 MipIndex, const FTexture2DMipMap& Mip) const;
-	static UE::DerivedData::FPayloadId MakeMipId(int32 MipIndex);
+	static UE::DerivedData::FValueId MakeMipId(int32 MipIndex);
 #endif // WITH_EDITORONLY_DATA
 
 	/** 
diff --git a/Engine/Source/Runtime/Engine/Private/Streaming/Texture2DMipDataProvider_DDC.cpp b/Engine/Source/Runtime/Engine/Private/Streaming/Texture2DMipDataProvider_DDC.cpp
--- a/Engine/Source/Runtime/Engine/Private/Streaming/Texture2DMipDataProvider_DDC.cpp
+++ b/Engine/Source/Runtime/Engine/Private/Streaming/Texture2DMipDataProvider_DDC.cpp
@@ -16,7 +16,6 @@ Texture2DMipDataProvider_DDC.cpp : Implementation of FTextureMipDataProvider usi
 #include "DerivedDataCache.h"
 #include "DerivedDataCacheInterface.h"
 #include "DerivedDataCacheKey.h"
-#include "DerivedDataPayload.h"
 
 FTexture2DMipDataProvider_DDC::FTexture2DMipDataProvider_DDC(const UTexture* Texture)
 	: FTextureMipDataProvider(Texture, ETickState::Init, ETickThread::Async)
diff --git a/Engine/Source/Runtime/Engine/Private/Streaming/Texture2DStreamIn_DDC.cpp b/Engine/Source/Runtime/Engine/Private/Streaming/Texture2DStreamIn_DDC.cpp
--- a/Engine/Source/Runtime/Engine/Private/Streaming/Texture2DStreamIn_DDC.cpp
+++ b/Engine/Source/Runtime/Engine/Private/Streaming/Texture2DStreamIn_DDC.cpp
@@ -14,7 +14,6 @@ Texture2DStreamIn_DDC.cpp: Stream in helper for 2D textures loading DDC files.
 #include "DerivedDataCache.h"
 #include "DerivedDataCacheInterface.h"
 #include "DerivedDataCacheKey.h"
-#include "DerivedDataPayload.h"
 #include "DerivedDataRequestOwner.h"
 
 int32 GStreamingUseAsyncRequestsForDDC = 1;
diff --git a/Engine/Source/Runtime/Engine/Private/TextureDerivedData.cpp b/Engine/Source/Runtime/Engine/Private/TextureDerivedData.cpp
--- a/Engine/Source/Runtime/Engine/Private/TextureDerivedData.cpp
+++ b/Engine/Source/Runtime/Engine/Private/TextureDerivedData.cpp
@@ -37,9 +37,9 @@
 #include "DerivedDataCache.h"
 #include "DerivedDataCacheInterface.h"
 #include "DerivedDataCacheKey.h"
-#include "DerivedDataPayload.h"
 #include "DerivedDataRequestOwner.h"
 #include "DerivedDataRequestTypes.h"
+#include "DerivedDataValue.h"
 #include "Interfaces/ITargetPlatform.h"
 #include "Interfaces/ITargetPlatformManagerModule.h"
 #include "Interfaces/ITextureFormat.h"
@@ -2097,9 +2097,9 @@ FString FTexturePlatformData::GetDerivedDataMipKeyString(int32 MipIndex, const F
 	return FString::Printf(TEXT("%s_MIP%u_%dx%d"), *KeyString, MipIndex, Mip.SizeX, Mip.SizeY);
 }
 
-UE::DerivedData::FPayloadId FTexturePlatformData::MakeMipId(int32 MipIndex)
+UE::DerivedData::FValueId FTexturePlatformData::MakeMipId(int32 MipIndex)
 {
-	return UE::DerivedData::FPayloadId::FromName(WriteToString<16>(TEXT("Mip"_SV), MipIndex));
+	return UE::DerivedData::FValueId::FromName(WriteToString<16>(TEXT("Mip"_SV), MipIndex));
 }
 
 #endif // WITH_EDITORONLY_DATA
diff --git a/Engine/Source/Runtime/Engine/Private/TextureDerivedDataTask.cpp b/Engine/Source/Runtime/Engine/Private/TextureDerivedDataTask.cpp
--- a/Engine/Source/Runtime/Engine/Private/TextureDerivedDataTask.cpp
+++ b/Engine/Source/Runtime/Engine/Private/TextureDerivedDataTask.cpp
@@ -32,8 +32,8 @@
 #include "DerivedDataBuildSession.h"
 #include "DerivedDataCacheInterface.h"
 #include "DerivedDataCacheKey.h"
-#include "DerivedDataPayload.h"
 #include "DerivedDataRequestOwner.h"
+#include "DerivedDataValue.h"
 #include "Engine/TextureCube.h"
 #include "Engine/VolumeTexture.h"
 #include "GenericPlatform/GenericPlatformMath.h"
@@ -1118,8 +1118,8 @@ private:
 		using namespace UE::DerivedData;
 		DerivedData.DerivedDataKey.Emplace<FCacheKeyProxy>(Params.CacheKey);
 		bCacheHit = EnumHasAnyFlags(Params.BuildStatus, EBuildStatus::CacheQueryHit);
-		BuildOutputSize = Algo::TransformAccumulate(Params.Output.GetPayloads(),
-			[](const FPayload& Payload) { return Payload.GetData().GetRawSize(); }, uint64(0));
+		BuildOutputSize = Algo::TransformAccumulate(Params.Output.GetValues(),
+			[](const FValue& Value) { return Value.GetData().GetRawSize(); }, uint64(0));
 		if (Params.Status != EStatus::Canceled)
 		{
 			WriteDerivedData(MoveTemp(Params.Output));
@@ -1304,18 +1304,18 @@ public:
 
 private:
 
-	static bool DeserializeTextureFromPayloads(FTexturePlatformData& DerivedData, const UE::DerivedData::FBuildOutput& Output, int32 FirstMipToLoad, bool bInlineMips)
+	static bool DeserializeTextureFromValues(FTexturePlatformData& DerivedData, const UE::DerivedData::FBuildOutput& Output, int32 FirstMipToLoad, bool bInlineMips)
 	{
 		using namespace UE::DerivedData;
-		const FPayload& Payload = Output.GetPayload(FPayloadId::FromName("Description"_ASV));
-		if (!Payload)
+		const FValueWithId& Value = Output.GetValue(FValueId::FromName("Description"_ASV));
+		if (!Value)
 		{
 			UE_LOG(LogTexture, Error, TEXT("Missing texture description for build of '%s' by %s."),
 				*WriteToString<128>(Output.GetName()), *WriteToString<32>(Output.GetFunction()));
 			return false;
 		}
 
-		FCbObject TextureDescription(Payload.GetData().Decompress());
+		FCbObject TextureDescription(Value.GetData().Decompress());
 
 		FCbFieldViewIterator SizeIt = TextureDescription["Size"_ASV].AsArrayView().CreateViewIterator();
 		DerivedData.SizeX = SizeIt++->AsInt32();
@@ -1352,14 +1352,14 @@ private:
 		FSharedBuffer MipTailData;
 		if (NumMips > NumStreamingMips)
 		{
-			const FPayload& MipTailPayload = Output.GetPayload(FPayloadId::FromName("MipTail"_ASV));
-			if (!MipTailPayload)
+			const FValueWithId& MipTailValue = Output.GetValue(FValueId::FromName("MipTail"_ASV));
+			if (!MipTailValue)
 			{
 				UE_LOG(LogTexture, Error, TEXT("Missing texture mip tail for build of '%s' by %s."),
 					*WriteToString<128>(Output.GetName()), *WriteToString<32>(Output.GetFunction()));
 				return false;
 			}
-			MipTailData = MipTailPayload.GetData().Decompress();
+			MipTailData = MipTailValue.GetData().Decompress();
 		}
 
 		int32 MipIndex = 0;
@@ -1388,15 +1388,15 @@ private:
 			}
 			else if (bInlineMips && (MipIndex >= FirstMipToLoad))
 			{
-				const FPayload& StreamingMipPayload = Output.GetPayload(FTexturePlatformData::MakeMipId(MipIndex));
-				if (!StreamingMipPayload)
+				const FValueWithId& StreamingMipValue = Output.GetValue(FTexturePlatformData::MakeMipId(MipIndex));
+				if (!StreamingMipValue)
 				{
 					delete NewMip;
 					UE_LOG(LogTexture, Error, TEXT("Missing texture streaming mip '%d' for build of '%s' by %s."),
 						MipIndex, *WriteToString<128>(Output.GetName()), *WriteToString<32>(Output.GetFunction()));
 					return false;
 				}
-				FSharedBuffer StreamingMipData = StreamingMipPayload.GetData().Decompress();
+				FSharedBuffer StreamingMipData = StreamingMipValue.GetData().Decompress();
 				uint64 MipSize = StreamingMipData.GetSize();
 
 				NewMip->BulkData.Lock(LOCK_READ_WRITE);
@@ -1473,7 +1473,7 @@ private:
 			return;
 		}
 
-		DeserializeTextureFromPayloads(DerivedData, Output, FirstMipToLoad, bInlineMips);
+		DeserializeTextureFromValues(DerivedData, Output, FirstMipToLoad, bInlineMips);
 	}
 
 	static UE::DerivedData::EPriority ConvertPriority(EQueuedWorkPriority SourcePriority)
