DDC: Changed FCacheRecord to be a collection of values

The structure with a value and attachments is an artifact of an early design where the record *was* a compact binary package. That design changed months ago and it now simplifies the cache record significantly to make it a collection of values without specific "value" and "attachment" types.

#rnx

diff --git a/Engine/Source/Developer/DerivedDataCache/Private/DerivedDataBuild.cpp b/Engine/Source/Developer/DerivedDataCache/Private/DerivedDataBuild.cpp
--- a/Engine/Source/Developer/DerivedDataCache/Private/DerivedDataBuild.cpp
+++ b/Engine/Source/Developer/DerivedDataCache/Private/DerivedDataBuild.cpp
@@ -119,7 +119,7 @@ private:
 	TUniquePtr<IBuildFunctionRegistry> FunctionRegistry{CreateBuildFunctionRegistry()};
 	TUniquePtr<IBuildWorkerRegistry> WorkerRegistry{CreateBuildWorkerRegistry()};
 	TUniquePtr<IBuildScheduler> DefaultScheduler{CreateBuildScheduler()};
-	const FGuid Version{TEXT("02c9caed-123c-4fc3-b8c1-7816058abed5")};
+	const FGuid Version{TEXT("17fe280d-ccd8-4be8-a9d1-89c944a70969")};
 };
 
 ///////////////////////////////////////////////////////////////////////////////////////////////////////////////////////
diff --git a/Engine/Source/Developer/DerivedDataCache/Private/DerivedDataBuildDefinition.cpp b/Engine/Source/Developer/DerivedDataCache/Private/DerivedDataBuildDefinition.cpp
--- a/Engine/Source/Developer/DerivedDataCache/Private/DerivedDataBuildDefinition.cpp
+++ b/Engine/Source/Developer/DerivedDataCache/Private/DerivedDataBuildDefinition.cpp
@@ -195,7 +195,7 @@ FBuildDefinitionInternal::FBuildDefinitionInternal(FBuildDefinitionBuilderIntern
 				const FBuildValueKey& ValueKey = Pair.Value.Get<FBuildValueKey>();
 				Writer.BeginObject(FTCHARToUTF8(Pair.Key));
 				Writer.AddHash("Build"_ASV, ValueKey.BuildKey.Hash);
-				Writer.AddObjectId("Payload"_ASV, ValueKey.Id);
+				Writer.AddObjectId("Id"_ASV, ValueKey.Id);
 				Writer.EndObject();
 			}
 		}
@@ -267,7 +267,7 @@ FBuildDefinitionInternal::FBuildDefinitionInternal(FStringView InName, FCbObject
 			{
 				return Field.GetName().Len() > 0 && Field.IsObject()
 					&& Field["Build"_ASV].IsHash()
-					&& Field["Payload"_ASV].IsObjectId();
+					&& Field["Id"_ASV].IsObjectId();
 			})
 		&& Algo::AllOf(Definition.AsView()["Inputs"_ASV]["BulkData"_ASV],
 			[](FCbFieldView Field) { return Field.GetName().Len() > 0 && Field.IsUuid(); })
@@ -300,7 +300,7 @@ void FBuildDefinitionInternal::IterateInputBuilds(TFunctionRef<void (FStringView
 	for (FCbFieldView Field : Definition.AsView()["Inputs"_ASV]["Builds"_ASV])
 	{
 		const FBuildKey BuildKey{Field["Build"_ASV].AsHash()};
-		const FValueId Id = Field["Payload"_ASV].AsObjectId();
+		const FValueId Id = Field["Id"_ASV].AsObjectId();
 		Visitor(FUTF8ToTCHAR(Field.GetName()), FBuildValueKey{BuildKey, Id});
 	}
 }
diff --git a/Engine/Source/Developer/DerivedDataCache/Private/DerivedDataBuildJob.cpp b/Engine/Source/Developer/DerivedDataCache/Private/DerivedDataBuildJob.cpp
--- a/Engine/Source/Developer/DerivedDataCache/Private/DerivedDataBuildJob.cpp
+++ b/Engine/Source/Developer/DerivedDataCache/Private/DerivedDataBuildJob.cpp
@@ -478,7 +478,7 @@ static ECachePolicy MakeCacheQueryPolicy(EBuildPolicy BuildPolicy, const FBuildJ
 	}
 	if (EnumHasAnyFlags(BuildPolicy, EBuildPolicy::SkipData))
 	{
-		EnumAddFlags(CachePolicy, ECachePolicy::SkipAttachments);
+		EnumAddFlags(CachePolicy, ECachePolicy::SkipData);
 	}
 	CachePolicy &= Context.GetCachePolicyMask();
 	return CachePolicy;
@@ -486,11 +486,24 @@ static ECachePolicy MakeCacheQueryPolicy(EBuildPolicy BuildPolicy, const FBuildJ
 
 static FCacheRecordPolicy MakeCacheRecordQueryPolicy(const FBuildPolicy& BuildPolicy, const FBuildJobContext& Context)
 {
-	FCacheRecordPolicyBuilder Builder(MakeCacheQueryPolicy(BuildPolicy.GetDefaultValuePolicy(), Context));
+	const ECachePolicy DefaultCachePolicy = MakeCacheQueryPolicy(BuildPolicy.GetDefaultValuePolicy(), Context);
+	const ECachePolicy CombinedCachePolicy = MakeCacheQueryPolicy(BuildPolicy.GetCombinedPolicy(), Context);
+
+	FCacheRecordPolicyBuilder Builder(DefaultCachePolicy);
 	for (const FBuildValuePolicy& ValuePolicy : BuildPolicy.GetValuePolicies())
 	{
 		Builder.AddValuePolicy(ValuePolicy.Id, MakeCacheQueryPolicy(ValuePolicy.Policy, Context));
 	}
+
+	// Output can only be loaded if its value is present. Add a policy for the output value which
+	// is compatible with the combined policy if that differs from the default policy.
+	if (EnumHasAnyFlags(CombinedCachePolicy, ECachePolicy::Query) &&
+		(EnumHasAnyFlags(DefaultCachePolicy, ECachePolicy::SkipData) ||
+		(CombinedCachePolicy & ECachePolicy::Default) != (DefaultCachePolicy & ECachePolicy::Default)))
+	{
+		Builder.AddValuePolicy(GetBuildOutputValueId(), CombinedCachePolicy & ~ECachePolicy::SkipData);
+	}
+
 	return Builder.Build();
 }
 
diff --git a/Engine/Source/Developer/DerivedDataCache/Private/DerivedDataBuildOutput.cpp b/Engine/Source/Developer/DerivedDataCache/Private/DerivedDataBuildOutput.cpp
--- a/Engine/Source/Developer/DerivedDataCache/Private/DerivedDataBuildOutput.cpp
+++ b/Engine/Source/Developer/DerivedDataCache/Private/DerivedDataBuildOutput.cpp
@@ -17,6 +17,16 @@
 namespace UE::DerivedData::Private
 {
 
+static const ANSICHAR GBuildOutputValueName[] = "$Output";
+static const FValueId GBuildOutputValueId = FValueId::FromName(GBuildOutputValueName);
+
+FValueId GetBuildOutputValueId()
+{
+	return GBuildOutputValueId;
+}
+
+///////////////////////////////////////////////////////////////////////////////////////////////////////////////////////
+
 static FUtf8StringView LexToString(EBuildOutputMessageLevel Level)
 {
 	switch (Level)
@@ -194,11 +204,12 @@ FBuildOutputInternal::FBuildOutputInternal(FStringView InName, FStringView InFun
 	: Name(InName)
 	, Function(InFunction)
 	, Meta(InOutput.GetMeta())
-	, Output(InOutput.GetValue().GetData().Decompress())
-	, Values(InOutput.GetAttachments())
+	, Output(InOutput.GetValue(GBuildOutputValueId).GetData().Decompress())
+	, Values(InOutput.GetValues())
 {
 	checkf(!Name.IsEmpty(), TEXT("A build output requires a non-empty name."));
 	AssertValidBuildFunctionName(Function, Name);
+	Values.RemoveAll([](const FValueWithId& Value) { return Value.GetId() == GBuildOutputValueId; });
 	bOutIsValid = TryLoad();
 }
 
@@ -220,7 +231,7 @@ bool FBuildOutputInternal::TryLoad()
 
 	if (Values.IsEmpty())
 	{
-		for (FCbFieldView Value : OutputView["Payloads"_ASV])
+		for (FCbFieldView Value : OutputView["Values"_ASV])
 		{
 			const FValueId Id = Value["Id"_ASV].AsObjectId();
 			const FIoHash& RawHash = Value["RawHash"_ASV].AsAttachment();
@@ -282,7 +293,7 @@ void FBuildOutputInternal::Save(FCbWriter& Writer) const
 	Writer.BeginObject();
 	if (!Values.IsEmpty())
 	{
-		Writer.BeginArray("Payloads"_ASV);
+		Writer.BeginArray("Values"_ASV);
 		for (const FValueWithId& Value : Values)
 		{
 			Writer.BeginObject();
@@ -322,11 +333,11 @@ void FBuildOutputInternal::Save(FCacheRecordBuilder& RecordBuilder) const
 		{
 			Writer.AddField("Logs"_ASV, LogsField);
 		}
-		RecordBuilder.SetValue(Writer.Save().GetBuffer());
+		RecordBuilder.AddValue(GBuildOutputValueId, Writer.Save().GetBuffer());
 	}
 	for (const FValueWithId& Value : Values)
 	{
-		RecordBuilder.AddAttachment(Value);
+		RecordBuilder.AddValue(Value);
 	}
 }
 
@@ -335,6 +346,9 @@ void FBuildOutputInternal::Save(FCacheRecordBuilder& RecordBuilder) const
 void FBuildOutputBuilderInternal::AddValue(const FValueId& Id, const FValue& Value)
 {
 	checkf(Id.IsValid(), TEXT("Null value added in output for build of '%s' by %s."), *Name, *Function);
+	checkf(Id != GBuildOutputValueId,
+		TEXT("Value added with reserved ID %s from name '%hs' for build of '%s' by %s."),
+		*WriteToString<32>(Id), GBuildOutputValueName, *Name, *Function);
 	const int32 Index = Algo::LowerBoundBy(Values, Id, &FValueWithId::GetId);
 	checkf(!(Values.IsValidIndex(Index) && Values[Index].GetId() == Id),
 		TEXT("Duplicate ID %s used by value for build of '%s' by %s."), *WriteToString<32>(Id), *Name, *Function);
diff --git a/Engine/Source/Developer/DerivedDataCache/Private/DerivedDataBuildPrivate.h b/Engine/Source/Developer/DerivedDataCache/Private/DerivedDataBuildPrivate.h
--- a/Engine/Source/Developer/DerivedDataCache/Private/DerivedDataBuildPrivate.h
+++ b/Engine/Source/Developer/DerivedDataCache/Private/DerivedDataBuildPrivate.h
@@ -28,6 +28,7 @@ namespace UE::DerivedData { class IBuildWorkerRegistry; }
 namespace UE::DerivedData { class ICache; }
 namespace UE::DerivedData { class IRequestOwner; }
 namespace UE::DerivedData { struct FBuildKey; }
+namespace UE::DerivedData { struct FValueId; }
 
 namespace UE::DerivedData::Private
 {
@@ -57,6 +58,7 @@ FBuildInputsBuilder CreateBuildInputs(FStringView Name);
 
 // Implemented in DerivedDataBuildOutput.cpp
 FBuildOutputBuilder CreateBuildOutput(FStringView Name, FStringView Function);
+FValueId GetBuildOutputValueId();
 
 // Implemented in DerivedDataBuildSession.cpp
 FBuildSession CreateBuildSession(
diff --git a/Engine/Source/Developer/DerivedDataCache/Private/DerivedDataCacheRecord.cpp b/Engine/Source/Developer/DerivedDataCache/Private/DerivedDataCacheRecord.cpp
--- a/Engine/Source/Developer/DerivedDataCache/Private/DerivedDataCacheRecord.cpp
+++ b/Engine/Source/Developer/DerivedDataCache/Private/DerivedDataCacheRecord.cpp
@@ -3,6 +3,7 @@
 #include "DerivedDataCacheRecord.h"
 
 #include "Algo/Accumulate.h"
+#include "Algo/IsSorted.h"
 #include "DerivedDataCacheKey.h"
 #include "DerivedDataCachePrivate.h"
 #include "DerivedDataValue.h"
@@ -26,23 +27,14 @@ public:
 
 	void SetMeta(FCbObject&& Meta) final;
 
-	void SetValue(const FCompositeBuffer& Buffer, const FValueId& Id, uint64 BlockSize) final;
-	void SetValue(const FSharedBuffer& Buffer, const FValueId& Id, uint64 BlockSize) final;
-	void SetValue(const FValue& Value, const FValueId& Id) final;
-	void SetValue(const FValueWithId& Value) final;
-
-	void AddAttachment(const FCompositeBuffer& Buffer, const FValueId& Id, uint64 BlockSize) final;
-	void AddAttachment(const FSharedBuffer& Buffer, const FValueId& Id, uint64 BlockSize) final;
-	void AddAttachment(const FValue& Value, const FValueId& Id) final;
-	void AddAttachment(const FValueWithId& Value) final;
+	void AddValue(const FValueId& Id, const FValue& Value) final;
 
 	FCacheRecord Build() final;
 	void BuildAsync(IRequestOwner& Owner, FOnCacheRecordComplete&& OnComplete) final;
 
 	FCacheKey Key;
 	FCbObject Meta;
-	FValueWithId Value;
-	TArray<FValueWithId> Attachments;
+	TArray<FValueWithId> Values;
 };
 
 ///////////////////////////////////////////////////////////////////////////////////////////////////////////////////////
@@ -57,10 +49,8 @@ public:
 
 	const FCacheKey& GetKey() const final;
 	const FCbObject& GetMeta() const final;
-	const FValueWithId& GetValue() const final;
-	const FValueWithId& GetAttachment(const FValueId& Id) const final;
-	TConstArrayView<FValueWithId> GetAttachments() const final;
 	const FValueWithId& GetValue(const FValueId& Id) const final;
+	TConstArrayView<FValueWithId> GetValues() const final;
 
 	inline void AddRef() const final
 	{
@@ -78,25 +68,16 @@ public:
 private:
 	FCacheKey Key;
 	FCbObject Meta;
-	FValueWithId Value;
-	TArray<FValueWithId> Attachments;
+	TArray<FValueWithId> Values;
 	mutable std::atomic<uint32> ReferenceCount{0};
 };
 
 ///////////////////////////////////////////////////////////////////////////////////////////////////////////////////////
 
-static FValueId GetOrCreateValueId(const FValueId& Id, const FValue& Value)
-{
-	return Id.IsValid() ? Id : FValueId::FromHash(Value.GetRawHash());
-}
-
-///////////////////////////////////////////////////////////////////////////////////////////////////////////////////////
-
 FCacheRecordInternal::FCacheRecordInternal(FCacheRecordBuilderInternal&& RecordBuilder)
 	: Key(RecordBuilder.Key)
 	, Meta(MoveTemp(RecordBuilder.Meta))
-	, Value(MoveTemp(RecordBuilder.Value))
-	, Attachments(MoveTemp(RecordBuilder.Attachments))
+	, Values(MoveTemp(RecordBuilder.Values))
 {
 }
 
@@ -110,25 +91,15 @@ const FCbObject& FCacheRecordInternal::GetMeta() const
 	return Meta;
 }
 
-const FValueWithId& FCacheRecordInternal::GetValue() const
-{
-	return Value;
-}
-
-const FValueWithId& FCacheRecordInternal::GetAttachment(const FValueId& Id) const
-{
-	const int32 Index = Algo::BinarySearchBy(Attachments, Id, &FValueWithId::GetId);
-	return Attachments.IsValidIndex(Index) ? Attachments[Index] : FValueWithId::Null;
-}
-
-TConstArrayView<FValueWithId> FCacheRecordInternal::GetAttachments() const
+const FValueWithId& FCacheRecordInternal::GetValue(const FValueId& Id) const
 {
-	return Attachments;
+	const int32 Index = Algo::BinarySearchBy(Values, Id, &FValueWithId::GetId);
+	return Values.IsValidIndex(Index) ? Values[Index] : FValueWithId::Null;
 }
 
-const FValueWithId& FCacheRecordInternal::GetValue(const FValueId& Id) const
+TConstArrayView<FValueWithId> FCacheRecordInternal::GetValues() const
 {
-	return Value.GetId() == Id ? Value : GetAttachment(Id);
+	return Values;
 }
 
 ///////////////////////////////////////////////////////////////////////////////////////////////////////////////////////
@@ -144,65 +115,23 @@ void FCacheRecordBuilderInternal::SetMeta(FCbObject&& InMeta)
 	Meta.MakeOwned();
 }
 
-void FCacheRecordBuilderInternal::SetValue(const FCompositeBuffer& Buffer, const FValueId& Id, const uint64 BlockSize)
-{
-	return SetValue(FValue::Compress(Buffer, BlockSize), Id);
-}
-
-void FCacheRecordBuilderInternal::SetValue(const FSharedBuffer& Buffer, const FValueId& Id, const uint64 BlockSize)
-{
-	return SetValue(FValue::Compress(Buffer, BlockSize), Id);
-}
-
-void FCacheRecordBuilderInternal::SetValue(const FValue& InValue, const FValueId& Id)
+void FCacheRecordBuilderInternal::AddValue(const FValueId& Id, const FValue& Value)
 {
-	const FValueId ValueId = GetOrCreateValueId(Id, InValue);
-	checkf(Value.IsNull(),
-		TEXT("Failed to set value on %s with ID %s because it has an existing value with ID %s."),
-		*WriteToString<96>(Key), *WriteToString<32>(ValueId), *WriteToString<32>(Value.GetId()));
-	checkf(Algo::BinarySearchBy(Attachments, ValueId, &FValueWithId::GetId) == INDEX_NONE,
-		TEXT("Failed to set value on %s with ID %s because it has an existing attachment with that ID."),
+	checkf(Id, TEXT("Failed to add value on %s because the ID is null."), *WriteToString<96>(Key));
+	checkf(!(Value == FValue::Null), TEXT("Failed to add value on %s because the value is null."), *WriteToString<96>(Key));
+	const FValueId ValueId = Id ? Id : FValueId::FromHash(Value.GetRawHash());
+	const int32 Index = Algo::LowerBoundBy(Values, ValueId, &FValueWithId::GetId);
+	checkf(!(Values.IsValidIndex(Index) && Values[Index].GetId() == ValueId),
+		TEXT("Failed to add value on %s with ID %s because it has an existing value with that ID."),
 		*WriteToString<96>(Key), *WriteToString<32>(ValueId));
-	Value = FValueWithId(ValueId, InValue);
-}
-
-void FCacheRecordBuilderInternal::SetValue(const FValueWithId& InValue)
-{
-	checkf(InValue, TEXT("Failed to set value on %s because the value is null."), *WriteToString<96>(Key));
-	SetValue(InValue, InValue.GetId());
-}
-
-void FCacheRecordBuilderInternal::AddAttachment(const FCompositeBuffer& Buffer, const FValueId& Id, const uint64 BlockSize)
-{
-	return AddAttachment(FValue::Compress(Buffer, BlockSize), Id);
-}
-
-void FCacheRecordBuilderInternal::AddAttachment(const FSharedBuffer& Buffer, const FValueId& Id, const uint64 BlockSize)
-{
-	return AddAttachment(FValue::Compress(Buffer, BlockSize), Id);
-}
-
-void FCacheRecordBuilderInternal::AddAttachment(const FValue& InValue, const FValueId& Id)
-{
-	const FValueId ValueId = GetOrCreateValueId(Id, InValue);
-	checkf(Value.IsNull(),
-		TEXT("Failed to set value on %s with ID %s because it has an existing value with ID %s."),
-		*WriteToString<96>(Key), *WriteToString<32>(ValueId), *WriteToString<32>(Value.GetId()));
-	const int32 Index = Algo::LowerBoundBy(Attachments, Id, &FValueWithId::GetId);
-	checkf(!(Attachments.IsValidIndex(Index) && Attachments[Index].GetId() == Id) && Value.GetId() != Id,
-		TEXT("Failed to add attachment on %s with ID %s because it has an existing attachment or value with that ID."),
-		*WriteToString<96>(Key), *WriteToString<32>(Id));
-	Attachments.Insert(FValueWithId(ValueId, InValue), Index);
-}
-
-void FCacheRecordBuilderInternal::AddAttachment(const FValueWithId& InValue)
-{
-	checkf(InValue, TEXT("Failed to add attachment on %s because the value is null."), *WriteToString<96>(Key));
-	AddAttachment(InValue, InValue.GetId());
+	Values.Insert(FValueWithId(ValueId, Value), Index);
 }
 
 FCacheRecord FCacheRecordBuilderInternal::Build()
 {
+	checkf(Algo::IsSortedBy(Values, &FValueWithId::GetId),
+		TEXT("Values in the cache record %s are expected to be sorted when added."),
+		*WriteToString<96>(Key));
 	return CreateCacheRecord(new FCacheRecordInternal(MoveTemp(*this)));
 }
 
@@ -245,30 +174,21 @@ FCbPackage FCacheRecord::Save() const
 	{
 		Writer.AddObject("Meta"_ASV, Meta);
 	}
-	auto SaveValue = [&Package, &Writer](const FValueWithId& ValueWithId)
-	{
-		if (ValueWithId.HasData())
-		{
-			Package.AddAttachment(FCbAttachment(ValueWithId.GetData()));
-		}
-		Writer.BeginObject();
-		Writer.AddObjectId("Id"_ASV, ValueWithId.GetId());
-		Writer.AddBinaryAttachment("RawHash"_ASV, ValueWithId.GetRawHash());
-		Writer.AddInteger("RawSize"_ASV, ValueWithId.GetRawSize());
-		Writer.EndObject();
-	};
-	if (const FValueWithId& Value = GetValue())
-	{
-		Writer.SetName("Value"_ASV);
-		SaveValue(Value);
-	}
-	TConstArrayView<FValueWithId> Attachments = GetAttachments();
-	if (!Attachments.IsEmpty())
+	TConstArrayView<FValueWithId> Values = GetValues();
+	if (!Values.IsEmpty())
 	{
-		Writer.BeginArray("Attachments"_ASV);
-		for (const FValueWithId& Attachment : Attachments)
+		Writer.BeginArray("Values"_ASV);
+		for (const FValueWithId& Value : Values)
 		{
-			SaveValue(Attachment);
+			if (Value.HasData())
+			{
+				Package.AddAttachment(FCbAttachment(Value.GetData()));
+			}
+			Writer.BeginObject();
+			Writer.AddObjectId("Id"_ASV, Value.GetId());
+			Writer.AddBinaryAttachment("RawHash"_ASV, Value.GetRawHash());
+			Writer.AddInteger("RawSize"_ASV, Value.GetRawSize());
+			Writer.EndObject();
 		}
 		Writer.EndArray();
 	}
@@ -282,6 +202,12 @@ FOptionalCacheRecord FCacheRecord::Load(const FCbPackage& Attachments, const FCb
 {
 	const FCbObjectView ObjectView = Object;
 
+	// Check for the previous format of cache record. Remove this check in 5.1.
+	if (ObjectView["Value"_ASV] || ObjectView["Attachments"_ASV])
+	{
+		return FOptionalCacheRecord();
+	}
+
 	FCacheKey Key;
 	FCbObjectView KeyObject = ObjectView["Key"_ASV].AsObjectView();
 	auto TrySetBucketName = [](FUtf8StringView Name, FCacheKey& Key)
@@ -329,24 +255,14 @@ FOptionalCacheRecord FCacheRecord::Load(const FCbPackage& Attachments, const FCb
 		}
 	};
 
-	if (FCbObjectView ValueObject = ObjectView["Value"_ASV].AsObjectView())
+	for (FCbFieldView ValueField : ObjectView["Values"_ASV])
 	{
-		FValueWithId Value = LoadValue(ValueObject);
+		FValueWithId Value = LoadValue(ValueField.AsObjectView());
 		if (!Value)
 		{
 			return FOptionalCacheRecord();
 		}
-		Builder.SetValue(Value);
-	}
-
-	for (FCbFieldView AttachmentField : ObjectView["Attachments"_ASV])
-	{
-		FValueWithId Attachment = LoadValue(AttachmentField.AsObjectView());
-		if (!Attachment)
-		{
-			return FOptionalCacheRecord();
-		}
-		Builder.AddAttachment(Attachment);
+		Builder.AddValue(Value);
 	}
 
 	return Builder.Build();
@@ -362,6 +278,21 @@ FCacheRecordBuilder::FCacheRecordBuilder(const FCacheKey& Key)
 {
 }
 
+void FCacheRecordBuilder::AddValue(const FValueId& Id, const FCompositeBuffer& Buffer, const uint64 BlockSize)
+{
+	return RecordBuilder->AddValue(Id, FValue::Compress(Buffer, BlockSize));
+}
+
+void FCacheRecordBuilder::AddValue(const FValueId& Id, const FSharedBuffer& Buffer, const uint64 BlockSize)
+{
+	return RecordBuilder->AddValue(Id, FValue::Compress(Buffer, BlockSize));
+}
+
+void FCacheRecordBuilder::AddValue(const FValueWithId& Value)
+{
+	return RecordBuilder->AddValue(Value.GetId(), Value);
+}
+
 } // UE::DerivedData
 
 ///////////////////////////////////////////////////////////////////////////////////////////////////////////////////////
@@ -371,22 +302,22 @@ namespace UE::DerivedData::Private
 
 uint64 GetCacheRecordCompressedSize(const FCacheRecord& Record)
 {
-	const uint64 ValueSize = Record.GetValue().GetData().GetCompressedSize();
-	return int64(Algo::TransformAccumulate(Record.GetAttachments(),
-		[](const FValueWithId& Value) { return Value.GetData().GetCompressedSize(); }, ValueSize));
+	const uint64 MetaSize = Record.GetMeta().GetSize();
+	return int64(Algo::TransformAccumulate(Record.GetValues(),
+		[](const FValueWithId& Value) { return Value.GetData().GetCompressedSize(); }, MetaSize));
 }
 
 uint64 GetCacheRecordTotalRawSize(const FCacheRecord& Record)
 {
-	const uint64 ValueSize = Record.GetValue().GetRawSize();
-	return int64(Algo::TransformAccumulate(Record.GetAttachments(), &FValueWithId::GetRawSize, ValueSize));
+	const uint64 MetaSize = Record.GetMeta().GetSize();
+	return int64(Algo::TransformAccumulate(Record.GetValues(), &FValueWithId::GetRawSize, MetaSize));
 }
 
 uint64 GetCacheRecordRawSize(const FCacheRecord& Record)
 {
-	const uint64 ValueSize = Record.GetValue().GetData().GetRawSize();
-	return int64(Algo::TransformAccumulate(Record.GetAttachments(),
-		[](const FValueWithId& Value) { return Value.GetData().GetRawSize(); }, ValueSize));
+	const uint64 MetaSize = Record.GetMeta().GetSize();
+	return int64(Algo::TransformAccumulate(Record.GetValues(),
+		[](const FValueWithId& Value) { return Value.GetData().GetRawSize(); }, MetaSize));
 }
 
 } // UE::DerivedData::Private
diff --git a/Engine/Source/Developer/DerivedDataCache/Private/FileSystemDerivedDataBackend.cpp b/Engine/Source/Developer/DerivedDataCache/Private/FileSystemDerivedDataBackend.cpp
--- a/Engine/Source/Developer/DerivedDataCache/Private/FileSystemDerivedDataBackend.cpp
+++ b/Engine/Source/Developer/DerivedDataCache/Private/FileSystemDerivedDataBackend.cpp
@@ -11,6 +11,7 @@
 #include "DerivedDataBackendInterface.h"
 #include "DerivedDataCacheInterface.h"
 #include "DerivedDataCacheMaintainer.h"
+#include "DerivedDataCachePrivate.h"
 #include "DerivedDataCacheRecord.h"
 #include "DerivedDataCacheUsageStats.h"
 #include "DerivedDataChunk.h"
@@ -765,33 +766,28 @@ public:
 		FOnCacheChunkComplete&& OnComplete) final;
 
 private:
-	uint64 MeasureCompressedCacheRecord(const FCacheRecord& Record) const;
-	uint64 MeasureRawCacheRecord(const FCacheRecord& Record) const;
-
-	bool PutCacheRecord(const FCacheRecord& Record, FStringView RecordName, const FCacheRecordPolicy& Policy, uint64& OutWriteSize);
-	bool PutCacheContent(const FCompressedBuffer& Content, const FStringView ContentName, uint64& OutWriteSize) const;
+	bool PutCacheRecord(FStringView Name, const FCacheRecord& Record, const FCacheRecordPolicy& Policy, uint64& OutWriteSize);
+	bool PutCacheContent(FStringView Name, const FCompressedBuffer& Content, uint64& OutWriteSize) const;
 
 	FOptionalCacheRecord GetCacheRecordOnly(
+		FStringView Name,
 		const FCacheKey& Key,
-		const FStringView RecordName,
 		const FCacheRecordPolicy& Policy);
 	FOptionalCacheRecord GetCacheRecord(
+		FStringView Name,
 		const FCacheKey& Key,
-		const FStringView RecordName,
 		const FCacheRecordPolicy& Policy,
 		EStatus& OutStatus);
 	FValueWithId GetCacheContent(
+		FStringView Name,
 		const FCacheKey& Key,
 		const FValueWithId& Value,
-		FStringView ContentName,
-		ECachePolicy Policy,
-		ECachePolicy SkipFlag) const;
+		ECachePolicy Policy) const;
 	void GetCacheContent(
+		FStringView Name,
 		const FCacheKey& Key,
 		const FValueWithId& Value,
-		FStringView ContentName,
 		ECachePolicy Policy,
-		ECachePolicy SkipFlag,
 		FCompressedBufferReader& Reader,
 		TUniquePtr<FArchive>& OutArchive) const;
 
@@ -1542,12 +1538,12 @@ void FFileSystemCacheStore::Put(
 		TRACE_COUNTER_INCREMENT(FileSystemDDC_Put);
 		COOK_STAT(auto Timer = UsageStats.TimePut());
 		uint64 WriteSize = 0;
-		if (PutCacheRecord(Record, Request.Name, Request.Policy, WriteSize))
+		if (PutCacheRecord(Request.Name, Record, Request.Policy, WriteSize))
 		{
 			UE_LOG(LogDerivedDataCache, Verbose, TEXT("%s: Cache put complete for %s from '%s'"),
 				*CachePath, *WriteToString<96>(Record.GetKey()), *Request.Name);
 			TRACE_COUNTER_INCREMENT(FileSystemDDC_PutHit);
-			TRACE_COUNTER_ADD(FileSystemDDC_BytesWritten, MeasureCompressedCacheRecord(Record));
+			TRACE_COUNTER_ADD(FileSystemDDC_BytesWritten, WriteSize);
 			COOK_STAT(if (WriteSize) { Timer.AddHit(WriteSize); });
 			if (OnComplete)
 			{
@@ -1575,13 +1571,13 @@ void FFileSystemCacheStore::Get(
 		TRACE_COUNTER_INCREMENT(FileSystemDDC_Get);
 		COOK_STAT(auto Timer = UsageStats.TimeGet());
 		EStatus Status = EStatus::Ok;
-		if (FOptionalCacheRecord Record = GetCacheRecord(Request.Key, Request.Name, Request.Policy, Status))
+		if (FOptionalCacheRecord Record = GetCacheRecord(Request.Name, Request.Key, Request.Policy, Status))
 		{
 			UE_LOG(LogDerivedDataCache, Verbose, TEXT("%s: Cache hit for %s from '%s'"),
 				*CachePath, *WriteToString<96>(Request.Key), *Request.Name);
 			TRACE_COUNTER_INCREMENT(FileSystemDDC_GetHit);
-			TRACE_COUNTER_ADD(FileSystemDDC_BytesRead, MeasureCompressedCacheRecord(Record.Get()));
-			COOK_STAT(Timer.AddHit(MeasureRawCacheRecord(Record.Get())));
+			TRACE_COUNTER_ADD(FileSystemDDC_BytesRead, Private::GetCacheRecordCompressedSize(Record.Get()));
+			COOK_STAT(Timer.AddHit(Private::GetCacheRecordCompressedSize(Record.Get())));
 			if (OnComplete)
 			{
 				OnComplete({Request.Name, MoveTemp(Record).Get(), Request.UserData, Status});
@@ -1613,8 +1609,7 @@ void FFileSystemCacheStore::GetChunks(
 	FOptionalCacheRecord Record;
 	for (const FCacheChunkRequest& Request : SortedRequests)
 	{
-		constexpr ECachePolicy SkipFlag = ECachePolicy::SkipValue;
-		const bool bExistsOnly = EnumHasAnyFlags(Request.Policy, SkipFlag);
+		const bool bExistsOnly = EnumHasAnyFlags(Request.Policy, ECachePolicy::SkipData);
 		TRACE_CPUPROFILER_EVENT_SCOPE(FileSystemDDC_Get);
 		TRACE_COUNTER_INCREMENT(FileSystemDDC_Get);
 		COOK_STAT(auto Timer = bExistsOnly ? UsageStats.TimeProbablyExists() : UsageStats.TimeGet());
@@ -1627,16 +1622,15 @@ void FFileSystemCacheStore::GetChunks(
 			if (!(Record && Record.Get().GetKey() == Request.Key))
 			{
 				FCacheRecordPolicyBuilder PolicyBuilder(ECachePolicy::None);
-				const ECachePolicy RecordSkipFlags = bExistsOnly ? ECachePolicy::SkipData : ECachePolicy::None;
-				PolicyBuilder.AddValuePolicy(Request.Id, Request.Policy | RecordSkipFlags);
+				PolicyBuilder.AddValuePolicy(Request.Id, Request.Policy);
 				Record.Reset();
-				Record = GetCacheRecordOnly(Request.Key, Request.Name, PolicyBuilder.Build());
+				Record = GetCacheRecordOnly(Request.Name, Request.Key, PolicyBuilder.Build());
 			}
 			if (Record)
 			{
 				Value = Record.Get().GetValue(Request.Id);
 				ValueKey = Request.Key;
-				GetCacheContent(Request.Key, Value, Request.Name, Request.Policy, SkipFlag, ValueReader, ValueAr);
+				GetCacheContent(Request.Name, Request.Key, Value, Request.Policy, ValueReader, ValueAr);
 			}
 		}
 		if (Value)
@@ -1669,25 +1663,9 @@ void FFileSystemCacheStore::GetChunks(
 	}
 }
 
-uint64 FFileSystemCacheStore::MeasureCompressedCacheRecord(const FCacheRecord& Record) const
-{
-	return Record.GetMeta().GetSize() +
-		Record.GetValue().GetData().GetCompressedSize() +
-		Algo::TransformAccumulate(Record.GetAttachments(),
-			[](const FValue& Value) { return Value.GetData().GetCompressedSize(); }, uint64(0));
-}
-
-uint64 FFileSystemCacheStore::MeasureRawCacheRecord(const FCacheRecord& Record) const
-{
-	return Record.GetMeta().GetSize() +
-		Record.GetValue().GetData().GetRawSize() +
-		Algo::TransformAccumulate(Record.GetAttachments(),
-			[](const FValue& Value) { return Value.GetData().GetRawSize(); }, uint64(0));
-}
-
 bool FFileSystemCacheStore::PutCacheRecord(
+	const FStringView Name,
 	const FCacheRecord& Record,
-	const FStringView RecordName,
 	const FCacheRecordPolicy& Policy,
 	uint64& OutWriteSize)
 {
@@ -1695,7 +1673,7 @@ bool FFileSystemCacheStore::PutCacheRecord(
 	{
 		UE_LOG(LogDerivedDataCache, VeryVerbose,
 			TEXT("%s: Skipped put of %s from '%.*s' because this cache store is read-only"),
-			*CachePath, *WriteToString<96>(Record.GetKey()), RecordName.Len(), RecordName.GetData());
+			*CachePath, *WriteToString<96>(Record.GetKey()), Name.Len(), Name.GetData());
 		return false;
 	}
 
@@ -1709,14 +1687,14 @@ bool FFileSystemCacheStore::PutCacheRecord(
 	if (!EnumHasAnyFlags(CombinedPolicy, StoreFlag))
 	{
 		UE_LOG(LogDerivedDataCache, VeryVerbose, TEXT("%s: Skipped put of %s from '%.*s' due to cache policy"),
-			*CachePath, *WriteToString<96>(Key), RecordName.Len(), RecordName.GetData());
+			*CachePath, *WriteToString<96>(Key), Name.Len(), Name.GetData());
 		return false;
 	}
 
 	if (ShouldSimulateMiss(Key))
 	{
 		UE_LOG(LogDerivedDataCache, Verbose, TEXT("%s: Simulated miss for put of %s from '%.*s'"),
-			*CachePath, *WriteToString<96>(Key), RecordName.Len(), RecordName.GetData());
+			*CachePath, *WriteToString<96>(Key), Name.Len(), Name.GetData());
 		return false;
 	}
 
@@ -1725,13 +1703,13 @@ bool FFileSystemCacheStore::PutCacheRecord(
 	FCbPackage ExistingPackage;
 	TStringBuilder<256> Path;
 	BuildCacheRecordPath(Key, Path);
-	if (EnumHasAllFlags(CombinedValuePolicy, ECachePolicy::SkipValue | ECachePolicy::SkipAttachments))
+	if (EnumHasAnyFlags(CombinedValuePolicy, ECachePolicy::SkipData))
 	{
 		bRecordExists = FileExists(Path);
 	}
 	else
 	{
-		bRecordExists = LoadFileWithHash(Path, RecordName, [&ExistingPackage](FArchive& Ar) { ExistingPackage.TryLoad(Ar); });
+		bRecordExists = LoadFileWithHash(Path, Name, [&ExistingPackage](FArchive& Ar) { ExistingPackage.TryLoad(Ar); });
 	}
 
 	// Save the record to a package and remove attachments that will be stored externally.
@@ -1780,13 +1758,13 @@ bool FFileSystemCacheStore::PutCacheRecord(
 	for (FCompressedBuffer& Content : ExternalContent)
 	{
 		uint64 WriteSize = 0;
-		PutCacheContent(Content, RecordName, OutWriteSize);
+		PutCacheContent(Name, Content, OutWriteSize);
 		OutWriteSize += WriteSize;
 	}
 
 	// Save the record package to storage.
 	const auto WriteRecord = [&](FArchive& Ar) { Package.Save(Ar); OutWriteSize += uint64(Ar.TotalSize()); };
-	if (!bRecordExists && !SaveFileWithHash(Path, RecordName, WriteRecord))
+	if (!bRecordExists && !SaveFileWithHash(Path, Name, WriteRecord))
 	{
 		return false;
 	}
@@ -1800,15 +1778,15 @@ bool FFileSystemCacheStore::PutCacheRecord(
 }
 
 FOptionalCacheRecord FFileSystemCacheStore::GetCacheRecordOnly(
+	const FStringView Name,
 	const FCacheKey& Key,
-	const FStringView RecordName,
 	const FCacheRecordPolicy& Policy)
 {
 	if (!IsUsable())
 	{
 		UE_LOG(LogDerivedDataCache, VeryVerbose,
 			TEXT("%s: Skipped get of %s from '%.*s' because this cache store is not available"),
-			*CachePath, *WriteToString<96>(Key), RecordName.Len(), RecordName.GetData());
+			*CachePath, *WriteToString<96>(Key), Name.Len(), Name.GetData());
 		return FOptionalCacheRecord();
 	}
 
@@ -1818,14 +1796,14 @@ FOptionalCacheRecord FFileSystemCacheStore::GetCacheRecordOnly(
 	if (!EnumHasAnyFlags(Policy.GetRecordPolicy(), QueryPolicy))
 	{
 		UE_LOG(LogDerivedDataCache, VeryVerbose, TEXT("%s: Skipped get of %s from '%.*s' due to cache policy"),
-			*CachePath, *WriteToString<96>(Key), RecordName.Len(), RecordName.GetData());
+			*CachePath, *WriteToString<96>(Key), Name.Len(), Name.GetData());
 		return FOptionalCacheRecord();
 	}
 
 	if (ShouldSimulateMiss(Key))
 	{
 		UE_LOG(LogDerivedDataCache, Verbose, TEXT("%s: Simulated miss for get of %s from '%.*s'"),
-			*CachePath, *WriteToString<96>(Key), RecordName.Len(), RecordName.GetData());
+			*CachePath, *WriteToString<96>(Key), Name.Len(), Name.GetData());
 		return FOptionalCacheRecord();
 	}
 
@@ -1844,17 +1822,17 @@ FOptionalCacheRecord FFileSystemCacheStore::GetCacheRecordOnly(
 	FOptionalCacheRecord Record;
 	{
 		FCbPackage Package;
-		if (!LoadFileWithHash(Path, RecordName, [&Package](FArchive& Ar) { Package.TryLoad(Ar); }))
+		if (!LoadFileWithHash(Path, Name, [&Package](FArchive& Ar) { Package.TryLoad(Ar); }))
 		{
 			UE_LOG(LogDerivedDataCache, Verbose, TEXT("%s: Cache miss with missing record for %s from '%.*s'"),
-				*CachePath, *WriteToString<96>(Key), RecordName.Len(), RecordName.GetData());
+				*CachePath, *WriteToString<96>(Key), Name.Len(), Name.GetData());
 			return Record;
 		}
 
 		if (ValidateCompactBinary(Package, ECbValidateMode::Default | ECbValidateMode::Package) != ECbValidateError::None)
 		{
 			UE_LOG(LogDerivedDataCache, Display, TEXT("%s: Cache miss with invalid package for %s from '%.*s'"),
-				*CachePath, *WriteToString<96>(Key), RecordName.Len(), RecordName.GetData());
+				*CachePath, *WriteToString<96>(Key), Name.Len(), Name.GetData());
 			return FOptionalCacheRecord();
 		}
 
@@ -1862,7 +1840,7 @@ FOptionalCacheRecord FFileSystemCacheStore::GetCacheRecordOnly(
 		if (Record.IsNull())
 		{
 			UE_LOG(LogDerivedDataCache, Display, TEXT("%s: Cache miss with record load failure for %s from '%.*s'"),
-				*CachePath, *WriteToString<96>(Key), RecordName.Len(), RecordName.GetData());
+				*CachePath, *WriteToString<96>(Key), Name.Len(), Name.GetData());
 			return Record;
 		}
 	}
@@ -1877,12 +1855,12 @@ FOptionalCacheRecord FFileSystemCacheStore::GetCacheRecordOnly(
 }
 
 FOptionalCacheRecord FFileSystemCacheStore::GetCacheRecord(
+	const FStringView Name,
 	const FCacheKey& Key,
-	const FStringView RecordName,
 	const FCacheRecordPolicy& Policy,
 	EStatus& OutStatus)
 {
-	FOptionalCacheRecord Record = GetCacheRecordOnly(Key, RecordName, Policy);
+	FOptionalCacheRecord Record = GetCacheRecordOnly(Name, Key, Policy);
 	if (Record.IsNull())
 	{
 		OutStatus = EStatus::Error;
@@ -1898,36 +1876,17 @@ FOptionalCacheRecord FFileSystemCacheStore::GetCacheRecord(
 		RecordBuilder.SetMeta(FCbObject(Record.Get().GetMeta()));
 	}
 
-	if (FValueWithId Value = Record.Get().GetValue())
-	{
-		const ECachePolicy ValuePolicy = Policy.GetValuePolicy(Value.GetId());
-		if (FValueWithId Content = GetCacheContent(Key, Value, RecordName, ValuePolicy, ECachePolicy::SkipValue))
-		{
-			RecordBuilder.SetValue(MoveTemp(Content));
-		}
-		else if (EnumHasAnyFlags(ValuePolicy, ECachePolicy::PartialOnError))
-		{
-			OutStatus = EStatus::Error;
-			RecordBuilder.SetValue(MoveTemp(Value));
-		}
-		else
-		{
-			OutStatus = EStatus::Error;
-			return FOptionalCacheRecord();
-		}
-	}
-
-	for (FValueWithId Value : Record.Get().GetAttachments())
+	for (const FValueWithId& Value : Record.Get().GetValues())
 	{
 		const ECachePolicy ValuePolicy = Policy.GetValuePolicy(Value.GetId());
-		if (FValueWithId Content = GetCacheContent(Key, Value, RecordName, ValuePolicy, ECachePolicy::SkipAttachments))
+		if (FValueWithId Content = GetCacheContent(Name, Key, Value, ValuePolicy))
 		{
-			RecordBuilder.AddAttachment(MoveTemp(Content));
+			RecordBuilder.AddValue(MoveTemp(Content));
 		}
 		else if (EnumHasAnyFlags(ValuePolicy, ECachePolicy::PartialOnError))
 		{
 			OutStatus = EStatus::Error;
-			RecordBuilder.AddAttachment(MoveTemp(Value));
+			RecordBuilder.AddValue(Value);
 		}
 		else
 		{
@@ -1940,15 +1899,15 @@ FOptionalCacheRecord FFileSystemCacheStore::GetCacheRecord(
 }
 
 bool FFileSystemCacheStore::PutCacheContent(
+	const FStringView Name,
 	const FCompressedBuffer& Content,
-	const FStringView ContentName,
 	uint64& OutWriteSize) const
 {
 	const FIoHash& RawHash = Content.GetRawHash();
 	TStringBuilder<256> Path;
 	BuildCacheContentPath(RawHash, Path);
 	const auto WriteContent = [&](FArchive& Ar) { Content.Save(Ar); OutWriteSize += uint64(Ar.TotalSize()); };
-	if (!FileExists(Path) && !SaveFileWithHash(Path, ContentName, WriteContent))
+	if (!FileExists(Path) && !SaveFileWithHash(Path, Name, WriteContent))
 	{
 		return false;
 	}
@@ -1960,11 +1919,10 @@ bool FFileSystemCacheStore::PutCacheContent(
 }
 
 FValueWithId FFileSystemCacheStore::GetCacheContent(
+	const FStringView Name,
 	const FCacheKey& Key,
 	const FValueWithId& Value,
-	const FStringView ContentName,
-	const ECachePolicy Policy,
-	const ECachePolicy SkipFlag) const
+	const ECachePolicy Policy) const
 {
 	if (!EnumHasAnyFlags(Policy, ECachePolicy::Query))
 	{
@@ -1973,14 +1931,14 @@ FValueWithId FFileSystemCacheStore::GetCacheContent(
 
 	if (Value.HasData())
 	{
-		return EnumHasAnyFlags(Policy, SkipFlag) ? Value.RemoveData() : Value;
+		return EnumHasAnyFlags(Policy, ECachePolicy::SkipData) ? Value.RemoveData() : Value;
 	}
 
 	const FIoHash& RawHash = Value.GetRawHash();
 
 	TStringBuilder<256> Path;
 	BuildCacheContentPath(RawHash, Path);
-	if (EnumHasAnyFlags(Policy, SkipFlag))
+	if (EnumHasAnyFlags(Policy, ECachePolicy::SkipData))
 	{
 		if (FileExists(Path))
 		{
@@ -1994,7 +1952,7 @@ FValueWithId FFileSystemCacheStore::GetCacheContent(
 	else
 	{
 		FCompressedBuffer CompressedBuffer;
-		if (LoadFileWithHash(Path, ContentName, [&CompressedBuffer](FArchive& Ar) { CompressedBuffer = FCompressedBuffer::Load(Ar); }))
+		if (LoadFileWithHash(Path, Name, [&CompressedBuffer](FArchive& Ar) { CompressedBuffer = FCompressedBuffer::Load(Ar); }))
 		{
 			if (CompressedBuffer.GetRawHash() == RawHash)
 			{
@@ -2007,7 +1965,7 @@ FValueWithId FFileSystemCacheStore::GetCacheContent(
 			UE_LOG(LogDerivedDataCache, Display,
 				TEXT("%s: Cache miss with corrupted value %s with hash %s for %s from '%.*s'"),
 				*CachePath, *WriteToString<16>(Value.GetId()), *WriteToString<48>(RawHash),
-				*WriteToString<96>(Key), ContentName.Len(), ContentName.GetData());
+				*WriteToString<96>(Key), Name.Len(), Name.GetData());
 			return FValueWithId::Null;
 		}
 	}
@@ -2015,16 +1973,15 @@ FValueWithId FFileSystemCacheStore::GetCacheContent(
 	UE_LOG(LogDerivedDataCache, Verbose,
 		TEXT("%s: Cache miss with missing value %s with hash %s for %s from '%.*s'"),
 		*CachePath, *WriteToString<16>(Value.GetId()), *WriteToString<48>(RawHash), *WriteToString<96>(Key),
-		ContentName.Len(), ContentName.GetData());
+		Name.Len(), Name.GetData());
 	return FValueWithId::Null;
 }
 
 void FFileSystemCacheStore::GetCacheContent(
+	const FStringView Name,
 	const FCacheKey& Key,
 	const FValueWithId& Value,
-	const FStringView ContentName,
 	const ECachePolicy Policy,
-	const ECachePolicy SkipFlag,
 	FCompressedBufferReader& Reader,
 	TUniquePtr<FArchive>& OutArchive) const
 {
@@ -2035,7 +1992,7 @@ void FFileSystemCacheStore::GetCacheContent(
 
 	if (Value.HasData())
 	{
-		if (!EnumHasAnyFlags(Policy, SkipFlag))
+		if (!EnumHasAnyFlags(Policy, ECachePolicy::SkipData))
 		{
 			Reader.SetSource(Value.GetData());
 		}
@@ -2047,7 +2004,7 @@ void FFileSystemCacheStore::GetCacheContent(
 
 	TStringBuilder<256> Path;
 	BuildCacheContentPath(RawHash, Path);
-	if (EnumHasAllFlags(Policy, SkipFlag))
+	if (EnumHasAllFlags(Policy, ECachePolicy::SkipData))
 	{
 		if (FileExists(Path))
 		{
@@ -2060,7 +2017,7 @@ void FFileSystemCacheStore::GetCacheContent(
 	}
 	else
 	{
-		OutArchive = OpenFile(Path, ContentName);
+		OutArchive = OpenFile(Path, Name);
 		if (OutArchive)
 		{
 			Reader.SetSource(*OutArchive);
@@ -2075,7 +2032,7 @@ void FFileSystemCacheStore::GetCacheContent(
 			UE_LOG(LogDerivedDataCache, Display,
 				TEXT("%s: Cache miss with corrupted value %s with hash %s for %s from '%.*s'"),
 				*CachePath, *WriteToString<16>(Value.GetId()), *WriteToString<48>(RawHash),
-				*WriteToString<96>(Key), ContentName.Len(), ContentName.GetData());
+				*WriteToString<96>(Key), Name.Len(), Name.GetData());
 			Reader.ResetSource();
 			OutArchive.Reset();
 			return;
@@ -2085,7 +2042,7 @@ void FFileSystemCacheStore::GetCacheContent(
 	UE_LOG(LogDerivedDataCache, Verbose,
 		TEXT("%s: Cache miss with missing value %s with hash %s for %s from '%.*s'"),
 		*CachePath, *WriteToString<16>(Value.GetId()), *WriteToString<48>(RawHash), *WriteToString<96>(Key),
-		ContentName.Len(), ContentName.GetData());
+		Name.Len(), Name.GetData());
 }
 
 void FFileSystemCacheStore::BuildCacheRecordPath(const FCacheKey& CacheKey, FStringBuilderBase& Path) const
diff --git a/Engine/Source/Developer/DerivedDataCache/Private/MemoryDerivedDataBackend.cpp b/Engine/Source/Developer/DerivedDataCache/Private/MemoryDerivedDataBackend.cpp
--- a/Engine/Source/Developer/DerivedDataCache/Private/MemoryDerivedDataBackend.cpp
+++ b/Engine/Source/Developer/DerivedDataCache/Private/MemoryDerivedDataBackend.cpp
@@ -2,8 +2,8 @@
 
 #include "MemoryDerivedDataBackend.h"
 
-#include "Algo/Accumulate.h"
 #include "Algo/AllOf.h"
+#include "DerivedDataCachePrivate.h"
 #include "DerivedDataValue.h"
 #include "Misc/ScopeExit.h"
 #include "Serialization/CompactBinary.h"
@@ -445,27 +445,6 @@ bool FMemoryDerivedDataBackend::ShouldSimulateMiss(const FCacheKey& Key)
 	return false;
 }
 
-int64 FMemoryDerivedDataBackend::CalcRawCacheRecordSize(const FCacheRecord& Record) const
-{
-	const uint64 ValueSize = Record.GetValue().GetRawSize();
-	return int64(Algo::TransformAccumulate(Record.GetAttachments(), &FValue::GetRawSize, ValueSize));
-}
-
-int64 FMemoryDerivedDataBackend::CalcSerializedCacheRecordSize(const FCacheRecord& Record) const
-{
-	// Estimate the serialized size of the cache record.
-	uint64 TotalSize = 20;
-	TotalSize += Record.GetKey().Bucket.ToString().Len();
-	TotalSize += Record.GetMeta().GetSize();
-	const auto CalcCacheValueSize = [](const FValueWithId& Value)
-	{
-		return Value ? Value.GetData().GetCompressedSize() + 32 : 0;
-	};
-	TotalSize += CalcCacheValueSize(Record.GetValue());
-	TotalSize += Algo::TransformAccumulate(Record.GetAttachments(), CalcCacheValueSize, uint64(0));
-	return int64(TotalSize);
-}
-
 void FMemoryDerivedDataBackend::Put(
 	const TConstArrayView<FCachePutRequest> Requests,
 	IRequestOwner& Owner,
@@ -491,15 +470,14 @@ void FMemoryDerivedDataBackend::Put(
 			continue;
 		}
 
-		const FValueWithId& Value = Record.GetValue();
-		const TConstArrayView<FValueWithId> Attachments = Record.GetAttachments();
+		const TConstArrayView<FValueWithId> Values = Record.GetValues();
 
-		if ((Value && !Value.HasData()) || !Algo::AllOf(Attachments, &FValue::HasData))
+		if (!Algo::AllOf(Values, &FValue::HasData))
 		{
 			continue;
 		}
 
-		if (!Value && Attachments.IsEmpty())
+		if (Values.IsEmpty())
 		{
 			FWriteScopeLock ScopeLock(SynchronizationObject);
 			if (bDisabled)
@@ -509,7 +487,7 @@ void FMemoryDerivedDataBackend::Put(
 			if (const FCacheRecord* Existing = CacheRecords.Find(Key))
 			{
 				CacheRecords.Remove(Key);
-				CurrentCacheSize -= CalcSerializedCacheRecordSize(*Existing);
+				CurrentCacheSize -= Private::GetCacheRecordCompressedSize(*Existing);
 				bMaxSizeExceeded = false;
 			}
 			Status = EStatus::Ok;
@@ -517,7 +495,7 @@ void FMemoryDerivedDataBackend::Put(
 		else
 		{
 			COOK_STAT(auto Timer = UsageStats.TimePut());
-			const int64 RecordSize = CalcSerializedCacheRecordSize(Record);
+			const int64 RecordSize = Private::GetCacheRecordCompressedSize(Record);
 
 			FWriteScopeLock ScopeLock(SynchronizationObject);
 			Status = CacheRecords.Contains(Key) ? EStatus::Ok : EStatus::Error;
@@ -565,34 +543,22 @@ void FMemoryDerivedDataBackend::Get(
 		}
 		if (Record)
 		{
-			if (const FValueWithId& Value = Record.Get().GetValue();
-				!Value.HasData() && !EnumHasAnyFlags(Policy.GetValuePolicy(Value.GetId()), ECachePolicy::SkipValue))
-			{
-				Status = EStatus::Error;
-				if (!EnumHasAllFlags(Policy.GetValuePolicy(Value.GetId()), ECachePolicy::PartialOnError))
-				{
-					Record.Reset();
-				}
-			}
-			if (Record)
+			for (const FValueWithId& Value : Record.Get().GetValues())
 			{
-				for (const FValueWithId& Value : Record.Get().GetAttachments())
+				if (!Value.HasData() && !EnumHasAnyFlags(Policy.GetValuePolicy(Value.GetId()), ECachePolicy::SkipData))
 				{
-					if (!Value.HasData() && !EnumHasAnyFlags(Policy.GetValuePolicy(Value.GetId()), ECachePolicy::SkipAttachments))
+					Status = EStatus::Error;
+					if (!EnumHasAllFlags(Policy.GetValuePolicy(Value.GetId()), ECachePolicy::PartialOnError))
 					{
-						Status = EStatus::Error;
-						if (!EnumHasAllFlags(Policy.GetValuePolicy(Value.GetId()), ECachePolicy::PartialOnError))
-						{
-							Record.Reset();
-							break;
-						}
+						Record.Reset();
+						break;
 					}
 				}
 			}
 		}
 		if (Record)
 		{
-			COOK_STAT(Timer.AddHit(CalcRawCacheRecordSize(Record.Get())));
+			COOK_STAT(Timer.AddHit(Private::GetCacheRecordCompressedSize(Record.Get())));
 			if (OnComplete)
 			{
 				OnComplete({Request.Name, MoveTemp(Record).Get(), Request.UserData, Status});
@@ -618,7 +584,7 @@ void FMemoryDerivedDataBackend::GetChunks(
 	FCompressedBufferReader Reader;
 	for (const FCacheChunkRequest& Request : Requests)
 	{
-		const bool bExistsOnly = EnumHasAnyFlags(Request.Policy, ECachePolicy::SkipValue);
+		const bool bExistsOnly = EnumHasAnyFlags(Request.Policy, ECachePolicy::SkipData);
 		COOK_STAT(auto Timer = bExistsOnly ? UsageStats.TimeProbablyExists() : UsageStats.TimeGet());
 		if (ShouldSimulateMiss(Request.Key))
 		{
@@ -633,7 +599,7 @@ void FMemoryDerivedDataBackend::GetChunks(
 		{
 			Reader.ResetSource();
 			Value.Reset();
-			Value = Record->GetAttachment(Request.Id);
+			Value = Record->GetValue(Request.Id);
 			ValueKey = Request.Key;
 			Reader.SetSource(Value.GetData());
 		}
diff --git a/Engine/Source/Developer/DerivedDataCache/Private/MemoryDerivedDataBackend.h b/Engine/Source/Developer/DerivedDataCache/Private/MemoryDerivedDataBackend.h
--- a/Engine/Source/Developer/DerivedDataCache/Private/MemoryDerivedDataBackend.h
+++ b/Engine/Source/Developer/DerivedDataCache/Private/MemoryDerivedDataBackend.h
@@ -137,9 +137,6 @@ private:
 		return (Key.Len() + 1) * sizeof(TCHAR) + sizeof(FCacheValue::Age) + Data.Num();
 	}
 
-	int64 CalcRawCacheRecordSize(const FCacheRecord& Record) const;
-	int64 CalcSerializedCacheRecordSize(const FCacheRecord& Record) const;
-
 	/** Name of this cache (used for debugging) */
 	FString Name;
 
diff --git a/Engine/Source/Developer/DerivedDataCache/Private/PakFileDerivedDataBackend.cpp b/Engine/Source/Developer/DerivedDataCache/Private/PakFileDerivedDataBackend.cpp
--- a/Engine/Source/Developer/DerivedDataCache/Private/PakFileDerivedDataBackend.cpp
+++ b/Engine/Source/Developer/DerivedDataCache/Private/PakFileDerivedDataBackend.cpp
@@ -5,6 +5,7 @@
 #include "Algo/Accumulate.h"
 #include "Algo/StableSort.h"
 #include "Algo/Transform.h"
+#include "DerivedDataCachePrivate.h"
 #include "DerivedDataCacheRecord.h"
 #include "DerivedDataCacheUsageStats.h"
 #include "DerivedDataChunk.h"
@@ -487,11 +488,11 @@ void FPakFileDerivedDataBackend::Put(
 		const FCacheRecord& Record = Request.Record;
 		TRACE_CPUPROFILER_EVENT_SCOPE(PakFileDDC_Put);
 		COOK_STAT(auto Timer = UsageStats.TimePut());
-		if (PutCacheRecord(Record, Request.Name, Request.Policy))
+		if (PutCacheRecord(Request.Name, Record, Request.Policy))
 		{
 			UE_LOG(LogDerivedDataCache, Verbose, TEXT("%s: Cache put complete for %s from '%s'"),
 				*CachePath, *WriteToString<96>(Record.GetKey()), *Request.Name);
-			COOK_STAT(Timer.AddHit(MeasureRawCacheRecord(Record)));
+			COOK_STAT(Timer.AddHit(Private::GetCacheRecordCompressedSize(Record)));
 			if (OnComplete)
 			{
 				OnComplete({Request.Name, Record.GetKey(), Request.UserData, EStatus::Ok});
@@ -499,7 +500,7 @@ void FPakFileDerivedDataBackend::Put(
 		}
 		else
 		{
-			COOK_STAT(Timer.AddMiss(MeasureRawCacheRecord(Record)));
+			COOK_STAT(Timer.AddMiss());
 			if (OnComplete)
 			{
 				OnComplete({Request.Name, Record.GetKey(), Request.UserData, EStatus::Error});
@@ -518,11 +519,11 @@ void FPakFileDerivedDataBackend::Get(
 		TRACE_CPUPROFILER_EVENT_SCOPE(PakFileDDC_Get);
 		COOK_STAT(auto Timer = UsageStats.TimeGet());
 		EStatus Status = EStatus::Ok;
-		if (FOptionalCacheRecord Record = GetCacheRecord(Request.Key, Request.Name, Request.Policy, Status))
+		if (FOptionalCacheRecord Record = GetCacheRecord(Request.Name, Request.Key, Request.Policy, Status))
 		{
 			UE_LOG(LogDerivedDataCache, Verbose, TEXT("%s: Cache hit for %s from '%s'"),
 				*CachePath, *WriteToString<96>(Request.Key), *Request.Name);
-			COOK_STAT(Timer.AddHit(MeasureRawCacheRecord(Record.Get())));
+			COOK_STAT(Timer.AddHit(Private::GetCacheRecordCompressedSize(Record.Get())));
 			if (OnComplete)
 			{
 				OnComplete({Request.Name, MoveTemp(Record).Get(), Request.UserData, Status});
@@ -550,22 +551,20 @@ void FPakFileDerivedDataBackend::GetChunks(
 	FCompressedBufferReader Reader;
 	for (const FCacheChunkRequest& Request : SortedRequests)
 	{
-		constexpr ECachePolicy SkipFlag = ECachePolicy::SkipValue;
-		const bool bExistsOnly = EnumHasAnyFlags(Request.Policy, SkipFlag);
+		const bool bExistsOnly = EnumHasAnyFlags(Request.Policy, ECachePolicy::SkipData);
 		TRACE_CPUPROFILER_EVENT_SCOPE(PakFileDDC_Get);
 		COOK_STAT(auto Timer = bExistsOnly ? UsageStats.TimeProbablyExists() : UsageStats.TimeGet());
 		if (!Record || Record.Get().GetKey() != Request.Key)
 		{
 			FCacheRecordPolicyBuilder PolicyBuilder(ECachePolicy::None);
-			const ECachePolicy RecordSkipFlags = bExistsOnly ? ECachePolicy::SkipData : ECachePolicy::None;
-			PolicyBuilder.AddValuePolicy(Request.Id, Request.Policy | RecordSkipFlags);
-			Record = GetCacheRecordOnly(Request.Key, Request.Name, PolicyBuilder.Build());
+			PolicyBuilder.AddValuePolicy(Request.Id, Request.Policy);
+			Record = GetCacheRecordOnly(Request.Name, Request.Key, PolicyBuilder.Build());
 		}
 		if (Record)
 		{
 			EStatus ValueStatus = EStatus::Ok;
 			FValueWithId Value = Record.Get().GetValue(Request.Id);
-			GetCacheContent(Request.Key, Request.Name, Request.Policy, SkipFlag, Value, ValueStatus);
+			GetCacheContent(Request.Name, Request.Key, Request.Policy, Value, ValueStatus);
 			if (Value)
 			{
 				const uint64 RawOffset = FMath::Min(Value.GetRawSize(), Request.RawOffset);
@@ -596,25 +595,9 @@ void FPakFileDerivedDataBackend::GetChunks(
 	}
 }
 
-uint64 FPakFileDerivedDataBackend::MeasureCompressedCacheRecord(const FCacheRecord& Record) const
-{
-	return Record.GetMeta().GetSize() +
-		Record.GetValue().GetData().GetCompressedSize() +
-		Algo::TransformAccumulate(Record.GetAttachments(),
-			[](const FValue& Value) { return Value.GetData().GetCompressedSize(); }, uint64(0));
-}
-
-uint64 FPakFileDerivedDataBackend::MeasureRawCacheRecord(const FCacheRecord& Record) const
-{
-	return Record.GetMeta().GetSize() +
-		Record.GetValue().GetData().GetRawSize() +
-		Algo::TransformAccumulate(Record.GetAttachments(),
-			[](const FValue& Value) { return Value.GetData().GetRawSize(); }, uint64(0));
-}
-
 bool FPakFileDerivedDataBackend::PutCacheRecord(
+	const FStringView Name,
 	const FCacheRecord& Record,
-	const FStringView RecordName,
 	const FCacheRecordPolicy& Policy)
 {
 	if (!IsWritable())
@@ -628,14 +611,14 @@ bool FPakFileDerivedDataBackend::PutCacheRecord(
 	if (!EnumHasAnyFlags(Policy.GetRecordPolicy(), ECachePolicy::StoreLocal))
 	{
 		UE_LOG(LogDerivedDataCache, VeryVerbose, TEXT("%s: Skipped put of %s from '%.*s' due to cache policy"),
-			*CachePath, *WriteToString<96>(Key), RecordName.Len(), RecordName.GetData());
+			*CachePath, *WriteToString<96>(Key), Name.Len(), Name.GetData());
 		return false;
 	}
 
 	//if (ShouldSimulateMiss(Key))
 	//{
 	//	UE_LOG(LogDerivedDataCache, Verbose, TEXT("%s: Simulated miss for put of %s from '%.*s'"),
-	//		*CachePath, *WriteToString<96>(Key), RecordName.Len(), RecordName.GetData());
+	//		*CachePath, *WriteToString<96>(Key), Name.Len(), Name.GetData());
 	//	return false;
 	//}
 
@@ -646,11 +629,11 @@ bool FPakFileDerivedDataBackend::PutCacheRecord(
 	FPathViews::Append(Path, TEXT("Buckets"), Key);
 	const ECachePolicy CombinedValuePolicy = Algo::TransformAccumulate(
 		Policy.GetValuePolicies(), &FCacheValuePolicy::Policy, Policy.GetDefaultValuePolicy(), UE_PROJECTION(operator|));
-	if (EnumHasAllFlags(CombinedValuePolicy, ECachePolicy::SkipValue | ECachePolicy::SkipAttachments))
+	if (EnumHasAnyFlags(CombinedValuePolicy, ECachePolicy::SkipData))
 	{
 		bRecordExists = FileExists(Path);
 	}
-	else if (FSharedBuffer Buffer = LoadFile(Path, RecordName))
+	else if (FSharedBuffer Buffer = LoadFile(Path, Name))
 	{
 		FCbFieldIterator It = FCbFieldIterator::MakeRange(MoveTemp(Buffer));
 		bRecordExists = ExistingPackage.TryLoad(It);
@@ -701,11 +684,11 @@ bool FPakFileDerivedDataBackend::PutCacheRecord(
 	// Save the external content to storage.
 	for (FCompressedBuffer& Content : ExternalContent)
 	{
-		PutCacheContent(Content, RecordName);
+		PutCacheContent(Name, Content);
 	}
 
 	// Save the record package to storage.
-	if (!bRecordExists && !SaveFile(Path, RecordName, [&Package](FArchive& Ar) { Package.Save(Ar); }))
+	if (!bRecordExists && !SaveFile(Path, Name, [&Package](FArchive& Ar) { Package.Save(Ar); }))
 	{
 		return false;
 	}
@@ -714,15 +697,15 @@ bool FPakFileDerivedDataBackend::PutCacheRecord(
 }
 
 FOptionalCacheRecord FPakFileDerivedDataBackend::GetCacheRecordOnly(
+	const FStringView Name,
 	const FCacheKey& Key,
-	const FStringView RecordName,
 	const FCacheRecordPolicy& Policy)
 {
 	if (bClosed)
 	{
 		UE_LOG(LogDerivedDataCache, VeryVerbose,
 			TEXT("%s: Skipped get of %s from '%.*s' because this cache store is not available"),
-			*CachePath, *WriteToString<96>(Key), RecordName.Len(), RecordName.GetData());
+			*CachePath, *WriteToString<96>(Key), Name.Len(), Name.GetData());
 		return FOptionalCacheRecord();
 	}
 
@@ -730,25 +713,25 @@ FOptionalCacheRecord FPakFileDerivedDataBackend::GetCacheRecordOnly(
 	if (!EnumHasAnyFlags(Policy.GetRecordPolicy(), ECachePolicy::QueryLocal))
 	{
 		UE_LOG(LogDerivedDataCache, VeryVerbose, TEXT("%s: Skipped get of %s from '%.*s' due to cache policy"),
-			*CachePath, *WriteToString<96>(Key), RecordName.Len(), RecordName.GetData());
+			*CachePath, *WriteToString<96>(Key), Name.Len(), Name.GetData());
 		return FOptionalCacheRecord();
 	}
 
 	//if (ShouldSimulateMiss(Key))
 	//{
 	//	UE_LOG(LogDerivedDataCache, Verbose, TEXT("%s: Simulated miss for get of %s from '%.*s'"),
-	//		*CachePath, *WriteToString<96>(Key), RecordName.Len(), RecordName.GetData());
+	//		*CachePath, *WriteToString<96>(Key), Name.Len(), Name.GetData());
 	//	return FOptionalCacheRecord();
 	//}
 
 	// Request the record from storage.
 	TStringBuilder<256> Path;
 	FPathViews::Append(Path, TEXT("Buckets"), Key);
-	FSharedBuffer Buffer = LoadFile(Path, RecordName);
+	FSharedBuffer Buffer = LoadFile(Path, Name);
 	if (Buffer.IsNull())
 	{
 		UE_LOG(LogDerivedDataCache, Verbose, TEXT("%s: Cache miss with missing record for %s from '%.*s'"),
-			*CachePath, *WriteToString<96>(Key), RecordName.Len(), RecordName.GetData());
+			*CachePath, *WriteToString<96>(Key), Name.Len(), Name.GetData());
 		return FOptionalCacheRecord();
 	}
 
@@ -756,7 +739,7 @@ FOptionalCacheRecord FPakFileDerivedDataBackend::GetCacheRecordOnly(
 	if (ValidateCompactBinaryPackage(Buffer, ECbValidateMode::Default | ECbValidateMode::Package) != ECbValidateError::None)
 	{
 		UE_LOG(LogDerivedDataCache, Display, TEXT("%s: Cache miss with invalid package for %s from '%.*s'"),
-			*CachePath, *WriteToString<96>(Key), RecordName.Len(), RecordName.GetData());
+			*CachePath, *WriteToString<96>(Key), Name.Len(), Name.GetData());
 		return FOptionalCacheRecord();
 	}
 
@@ -767,14 +750,14 @@ FOptionalCacheRecord FPakFileDerivedDataBackend::GetCacheRecordOnly(
 		if (FCbFieldIterator It = FCbFieldIterator::MakeRange(Buffer); !Package.TryLoad(It))
 		{
 			UE_LOG(LogDerivedDataCache, Display, TEXT("%s: Cache miss with package load failure for %s from '%.*s'"),
-				*CachePath, *WriteToString<96>(Key), RecordName.Len(), RecordName.GetData());
+				*CachePath, *WriteToString<96>(Key), Name.Len(), Name.GetData());
 			return FOptionalCacheRecord();
 		}
 		Record = FCacheRecord::Load(Package);
 		if (Record.IsNull())
 		{
 			UE_LOG(LogDerivedDataCache, Display, TEXT("%s: Cache miss with record load failure for %s from '%.*s'"),
-				*CachePath, *WriteToString<96>(Key), RecordName.Len(), RecordName.GetData());
+				*CachePath, *WriteToString<96>(Key), Name.Len(), Name.GetData());
 			return FOptionalCacheRecord();
 		}
 	}
@@ -783,12 +766,12 @@ FOptionalCacheRecord FPakFileDerivedDataBackend::GetCacheRecordOnly(
 }
 
 FOptionalCacheRecord FPakFileDerivedDataBackend::GetCacheRecord(
+	const FStringView Name,
 	const FCacheKey& Key,
-	const FStringView RecordName,
 	const FCacheRecordPolicy& Policy,
 	EStatus& OutStatus)
 {
-	FOptionalCacheRecord Record = GetCacheRecordOnly(Key, RecordName, Policy);
+	FOptionalCacheRecord Record = GetCacheRecordOnly(Name, Key, Policy);
 	if (Record.IsNull())
 	{
 		OutStatus = EStatus::Error;
@@ -804,39 +787,28 @@ FOptionalCacheRecord FPakFileDerivedDataBackend::GetCacheRecord(
 		RecordBuilder.SetMeta(FCbObject(Record.Get().GetMeta()));
 	}
 
-	if (FValueWithId Value = Record.Get().GetValue())
-	{
-		const ECachePolicy ValuePolicy = Policy.GetValuePolicy(Value.GetId());
-		GetCacheContent(Key, RecordName, ValuePolicy, ECachePolicy::SkipValue, Value, OutStatus);
-		if (Value.IsNull())
-		{
-			return FOptionalCacheRecord();
-		}
-		RecordBuilder.SetValue(MoveTemp(Value));
-	}
-
-	for (FValueWithId Value : Record.Get().GetAttachments())
+	for (FValueWithId Value : Record.Get().GetValues())
 	{
 		const ECachePolicy ValuePolicy = Policy.GetValuePolicy(Value.GetId());
-		GetCacheContent(Key, RecordName, ValuePolicy, ECachePolicy::SkipAttachments, Value, OutStatus);
+		GetCacheContent(Name, Key, ValuePolicy, Value, OutStatus);
 		if (Value.IsNull())
 		{
 			return FOptionalCacheRecord();
 		}
-		RecordBuilder.AddAttachment(MoveTemp(Value));
+		RecordBuilder.AddValue(MoveTemp(Value));
 	}
 
 	return RecordBuilder.Build();
 }
 
-bool FPakFileDerivedDataBackend::PutCacheContent(const FCompressedBuffer& Content, const FStringView ContentName)
+bool FPakFileDerivedDataBackend::PutCacheContent(const FStringView Name, const FCompressedBuffer& Content)
 {
 	const FIoHash& RawHash = Content.GetRawHash();
 	TStringBuilder<256> Path;
 	FPathViews::Append(Path, TEXT("Content"), RawHash);
 	if (!FileExists(Path))
 	{
-		if (!SaveFile(Path, ContentName, [&Content](FArchive& Ar) { Content.Save(Ar); }))
+		if (!SaveFile(Path, Name, [&Content](FArchive& Ar) { Content.Save(Ar); }))
 		{
 			return false;
 		}
@@ -845,14 +817,14 @@ bool FPakFileDerivedDataBackend::PutCacheContent(const FCompressedBuffer& Conten
 }
 
 void FPakFileDerivedDataBackend::GetCacheContent(
+	const FStringView Name,
 	const FCacheKey& Key,
-	const FStringView ContentName,
 	const ECachePolicy Policy,
-	const ECachePolicy SkipFlag,
 	FValueWithId& InOutValue,
 	EStatus& InOutStatus)
 {
-	if (!EnumHasAnyFlags(Policy, ECachePolicy::Query) || (EnumHasAnyFlags(Policy, SkipFlag) && InOutValue.HasData()))
+	if (!EnumHasAnyFlags(Policy, ECachePolicy::Query) ||
+		(EnumHasAnyFlags(Policy, ECachePolicy::SkipData) && InOutValue.HasData()))
 	{
 		InOutValue = InOutValue.RemoveData();
 		return;
@@ -867,7 +839,7 @@ void FPakFileDerivedDataBackend::GetCacheContent(
 
 	TStringBuilder<256> Path;
 	FPathViews::Append(Path, TEXT("Content"), RawHash);
-	if (EnumHasAllFlags(Policy, SkipFlag))
+	if (EnumHasAllFlags(Policy, ECachePolicy::SkipData))
 	{
 		if (FileExists(Path))
 		{
@@ -876,7 +848,7 @@ void FPakFileDerivedDataBackend::GetCacheContent(
 	}
 	else
 	{
-		if (FSharedBuffer CompressedData = LoadFile(Path, ContentName))
+		if (FSharedBuffer CompressedData = LoadFile(Path, Name))
 		{
 			if (FCompressedBuffer CompressedBuffer = FCompressedBuffer::FromCompressed(MoveTemp(CompressedData));
 				CompressedBuffer && CompressedBuffer.GetRawHash() == RawHash)
@@ -887,7 +859,7 @@ void FPakFileDerivedDataBackend::GetCacheContent(
 			UE_LOG(LogDerivedDataCache, Display,
 				TEXT("%s: Cache miss with corrupted value %s with hash %s for %s from '%.*s'"),
 				*CachePath, *WriteToString<16>(InOutValue.GetId()), *WriteToString<48>(RawHash),
-				*WriteToString<96>(Key), ContentName.Len(), ContentName.GetData());
+				*WriteToString<96>(Key), Name.Len(), Name.GetData());
 			InOutStatus = EStatus::Error;
 			if (!EnumHasAnyFlags(Policy, ECachePolicy::PartialOnError))
 			{
@@ -900,7 +872,7 @@ void FPakFileDerivedDataBackend::GetCacheContent(
 	UE_LOG(LogDerivedDataCache, Verbose,
 		TEXT("%s: Cache miss with missing value %s with hash %s for %s from '%.*s'"),
 		*CachePath, *WriteToString<16>(InOutValue.GetId()), *WriteToString<48>(RawHash), *WriteToString<96>(Key),
-		ContentName.Len(), ContentName.GetData());
+		Name.Len(), Name.GetData());
 	InOutStatus = EStatus::Error;
 	if (!EnumHasAnyFlags(Policy, ECachePolicy::PartialOnError))
 	{
diff --git a/Engine/Source/Developer/DerivedDataCache/Private/PakFileDerivedDataBackend.h b/Engine/Source/Developer/DerivedDataCache/Private/PakFileDerivedDataBackend.h
--- a/Engine/Source/Developer/DerivedDataCache/Private/PakFileDerivedDataBackend.h
+++ b/Engine/Source/Developer/DerivedDataCache/Private/PakFileDerivedDataBackend.h
@@ -120,26 +120,22 @@ public:
 		FOnCacheChunkComplete&& OnComplete) override;
 
 private:
-	uint64 MeasureCompressedCacheRecord(const FCacheRecord& Record) const;
-	uint64 MeasureRawCacheRecord(const FCacheRecord& Record) const;
-
-	bool PutCacheRecord(const FCacheRecord& Record, FStringView RecordName, const FCacheRecordPolicy& Policy);
-	bool PutCacheContent(const FCompressedBuffer& Content, const FStringView ContentName);
+	bool PutCacheRecord(FStringView Name, const FCacheRecord& Record, const FCacheRecordPolicy& Policy);
+	bool PutCacheContent(FStringView Name, const FCompressedBuffer& Content);
 
 	FOptionalCacheRecord GetCacheRecordOnly(
+		FStringView Name,
 		const FCacheKey& Key,
-		const FStringView RecordName,
 		const FCacheRecordPolicy& Policy);
 	FOptionalCacheRecord GetCacheRecord(
+		FStringView Name,
 		const FCacheKey& Key,
-		const FStringView RecordName,
 		const FCacheRecordPolicy& Policy,
 		EStatus& OutStatus);
 	void GetCacheContent(
+		FStringView Name,
 		const FCacheKey& Key,
-		const FStringView ContentName,
-		const ECachePolicy Policy,
-		const ECachePolicy SkipFlag,
+		ECachePolicy Policy,
 		FValueWithId& InOutValue,
 		EStatus& InOutStatus);
 
diff --git a/Engine/Source/Developer/DerivedDataCache/Private/ZenDerivedDataBackend.cpp b/Engine/Source/Developer/DerivedDataCache/Private/ZenDerivedDataBackend.cpp
--- a/Engine/Source/Developer/DerivedDataCache/Private/ZenDerivedDataBackend.cpp
+++ b/Engine/Source/Developer/DerivedDataCache/Private/ZenDerivedDataBackend.cpp
@@ -3,11 +3,10 @@
 
 #if WITH_ZEN_DDC_BACKEND
 
-#include "Algo/Accumulate.h"
 #include "Containers/StringFwd.h"
+#include "DerivedDataCachePrivate.h"
 #include "DerivedDataCacheRecord.h"
 #include "DerivedDataChunk.h"
-#include "DerivedDataValue.h"
 #include "Misc/ConfigCacheIni.h"
 #include "Misc/ScopeLock.h"
 #include "ProfilingDebugging/CountersTrace.h"
@@ -407,26 +406,19 @@ void FZenDerivedDataBackend::AppendPolicyQueryString(ECachePolicy Policy, FStrin
 		if (EnumHasAnyFlags(Policy, ECachePolicy::StoreRemote)) { AppendValue(TEXT("remote")); }
 		if (!EnumHasAnyFlags(Policy, ECachePolicy::Store)) { AppendValue(TEXT("none")); }
 	}
-	if (EnumHasAnyFlags(Policy, ECachePolicy::SkipData))
+	if (EnumHasAnyFlags(Policy, ECachePolicy::SkipMeta | ECachePolicy::SkipData))
 	{
 		AppendKey(TEXT("skip="));
-		if (EnumHasAllFlags(Policy, ECachePolicy::SkipData)) { AppendValue(TEXT("data")); }
+		if (EnumHasAllFlags(Policy, ECachePolicy::SkipMeta | ECachePolicy::SkipData)) { AppendValue(TEXT("data")); }
 		else
 		{
 			if (EnumHasAnyFlags(Policy, ECachePolicy::SkipMeta)) { AppendValue(TEXT("meta")); }
-			if (EnumHasAnyFlags(Policy, ECachePolicy::SkipValue)) { AppendValue(TEXT("value")); }
-			if (EnumHasAnyFlags(Policy, ECachePolicy::SkipAttachments)) { AppendValue(TEXT("attachments")); }
+			if (EnumHasAnyFlags(Policy, ECachePolicy::SkipData)) { AppendValue(TEXT("value")); }
+			if (EnumHasAnyFlags(Policy, ECachePolicy::SkipData)) { AppendValue(TEXT("attachments")); }
 		}
 	}
 }
 
-uint64 FZenDerivedDataBackend::MeasureCacheRecord(const FCacheRecord& Record)
-{
-	return Record.GetMeta().GetSize() +
-		Record.GetValue().GetRawSize() +
-		Algo::TransformAccumulate(Record.GetAttachments(), &FValue::GetRawSize, uint64(0));
-}
-
 void FZenDerivedDataBackend::RemoveCachedData(const TCHAR* CacheKey, bool bTransient)
 {
 	TRACE_CPUPROFILER_EVENT_SCOPE(ZenDDC_Remove);
@@ -612,7 +604,7 @@ void FZenDerivedDataBackend::Put(
 		{
 			UE_LOG(LogDerivedDataCache, Verbose, TEXT("%s: Cache put complete for %s from '%s'"),
 				*GetName(), *WriteToString<96>(Record.GetKey()), *Request.Name);
-			COOK_STAT(Timer.AddHit(MeasureCacheRecord(Record)));
+			COOK_STAT(Timer.AddHit(Private::GetCacheRecordCompressedSize(Record)));
 			if (OnComplete)
 			{
 				OnComplete({ Request.Name, Record.GetKey(), Request.UserData, EStatus::Ok });
@@ -620,7 +612,7 @@ void FZenDerivedDataBackend::Put(
 		}
 		else
 		{
-			COOK_STAT(Timer.AddMiss(MeasureCacheRecord(Record)));
+			COOK_STAT(Timer.AddMiss());
 			if (OnComplete)
 			{
 				OnComplete({ Request.Name, Record.GetKey(), Request.UserData, EStatus::Error });
@@ -729,7 +721,7 @@ void FZenDerivedDataBackend::Get(
 						*GetName(), *WriteToString<96>(Key), *Request.Name);
 					
 					TRACE_COUNTER_ADD(ZenDDC_GetHit, int64(1));
-					int64 ReceivedSize = MeasureCacheRecord(Record.Get());
+					int64 ReceivedSize = Private::GetCacheRecordCompressedSize(Record.Get());
 					TRACE_COUNTER_ADD(ZenDDC_BytesReceived, ReceivedSize);
 					COOK_STAT(Timer.AddHit(ReceivedSize));
 
diff --git a/Engine/Source/Developer/DerivedDataCache/Private/ZenDerivedDataBackend.h b/Engine/Source/Developer/DerivedDataCache/Private/ZenDerivedDataBackend.h
--- a/Engine/Source/Developer/DerivedDataCache/Private/ZenDerivedDataBackend.h
+++ b/Engine/Source/Developer/DerivedDataCache/Private/ZenDerivedDataBackend.h
@@ -126,7 +126,6 @@ private:
 	static void AppendPolicyQueryString(ECachePolicy Policy, FStringBuilderBase& Out);
 
 	static bool ShouldRetryOnError(int64 ResponseCode);
-	static uint64 MeasureCacheRecord(const FCacheRecord& Record);
 
 	/* Debug helpers */
 	bool ShouldSimulateMiss(const TCHAR* InKey);
diff --git a/Engine/Source/Developer/DerivedDataCache/Public/DerivedDataCache.h b/Engine/Source/Developer/DerivedDataCache/Public/DerivedDataCache.h
--- a/Engine/Source/Developer/DerivedDataCache/Public/DerivedDataCache.h
+++ b/Engine/Source/Developer/DerivedDataCache/Public/DerivedDataCache.h
@@ -75,24 +75,8 @@ enum class ECachePolicy : uint32
 
 	/** Skip fetching the metadata for record requests. */
 	SkipMeta        = 1 << 4,
-	/** Skip fetching the value for record, chunk, or value requests. */
-	SkipValue       = 1 << 5,
-	/** Skip fetching the attachments for record requests. */
-	SkipAttachments = 1 << 6,
-	/**
-	 * Skip fetching the data for any requests.
-	 *
-	 * Put requests with skip flags may assume that record existence implies value existence.
-	 */
-	SkipData        = SkipMeta | SkipValue | SkipAttachments,
-
-	/**
-	 * Keep records in the cache for at least the duration of the session.
-	 *
-	 * This is a hint that the record may be accessed again in this session. This is mainly meant
-	 * to be used when subsequent accesses will not tolerate a cache miss.
-	 */
-	KeepAlive       = 1 << 7,
+	/** Skip fetching the data for values. */
+	SkipData        = 1 << 5,
 
 	/**
 	 * Partial output will be provided with the error status when a required value is missing.
@@ -105,7 +89,15 @@ enum class ECachePolicy : uint32
 	 *
 	 * Applying this flag for a put of a record allows a partial record to be stored.
 	 */
-	PartialOnError  = 1 << 8,
+	PartialOnError  = 1 << 6,
+
+	/**
+	 * Keep records in the cache for at least the duration of the session.
+	 *
+	 * This is a hint that the record may be accessed again in this session. This is mainly meant
+	 * to be used when subsequent accesses will not tolerate a cache miss.
+	 */
+	KeepAlive       = 1 << 7,
 
 	/** Allow cache requests to query and store records and values in local caches. */
 	Local           = QueryLocal | StoreLocal,
@@ -379,8 +371,8 @@ struct FCacheGetCompleteParams
 	 *
 	 * The key is always populated. The remainder of the record is populated when Status is Ok.
 	 *
-	 * The value, attachments, and metadata may be skipped based on cache policy flags. When a value
-	 * or attachment has been skipped, it will have a hash and size but its data will be null.
+	 * The metadata or the data for values may be skipped based on cache policy flags. Values for
+	 * which data has been skipped will have a hash and size but null data.
 	 */
 	FCacheRecord&& Record;
 
diff --git a/Engine/Source/Developer/DerivedDataCache/Public/DerivedDataCacheRecord.h b/Engine/Source/Developer/DerivedDataCache/Public/DerivedDataCacheRecord.h
--- a/Engine/Source/Developer/DerivedDataCache/Public/DerivedDataCacheRecord.h
+++ b/Engine/Source/Developer/DerivedDataCache/Public/DerivedDataCacheRecord.h
@@ -33,10 +33,8 @@ public:
 	virtual ~ICacheRecordInternal() = default;
 	virtual const FCacheKey& GetKey() const = 0;
 	virtual const FCbObject& GetMeta() const = 0;
-	virtual const FValueWithId& GetValue() const = 0;
-	virtual const FValueWithId& GetAttachment(const FValueId& Id) const = 0;
-	virtual TConstArrayView<FValueWithId> GetAttachments() const = 0;
 	virtual const FValueWithId& GetValue(const FValueId& Id) const = 0;
+	virtual TConstArrayView<FValueWithId> GetValues() const = 0;
 	virtual void AddRef() const = 0;
 	virtual void Release() const = 0;
 };
@@ -48,14 +46,7 @@ class ICacheRecordBuilderInternal
 public:
 	virtual ~ICacheRecordBuilderInternal() = default;
 	virtual void SetMeta(FCbObject&& Meta) = 0;
-	virtual void SetValue(const FCompositeBuffer& Buffer, const FValueId& Id, uint64 BlockSize) = 0;
-	virtual void SetValue(const FSharedBuffer& Buffer, const FValueId& Id, uint64 BlockSize) = 0;
-	virtual void SetValue(const FValue& Value, const FValueId& Id) = 0;
-	virtual void SetValue(const FValueWithId& Value) = 0;
-	virtual void AddAttachment(const FCompositeBuffer& Buffer, const FValueId& Id, uint64 BlockSize) = 0;
-	virtual void AddAttachment(const FSharedBuffer& Buffer, const FValueId& Id, uint64 BlockSize) = 0;
-	virtual void AddAttachment(const FValue& Value, const FValueId& Id) = 0;
-	virtual void AddAttachment(const FValueWithId& Value) = 0;
+	virtual void AddValue(const FValueId& Id, const FValue& Value) = 0;
 	virtual FCacheRecord Build() = 0;
 	virtual void BuildAsync(IRequestOwner& Owner, FOnCacheRecordComplete&& OnComplete) = 0;
 };
@@ -68,19 +59,11 @@ namespace UE::DerivedData
 {
 
 /**
- * A cache record is a key, a value, attachments, and metadata.
+ * A cache record is a key, an array of values, and metadata.
  *
- * The key is expected to correspond uniquely to its value and attachments. The key should not be
- * used with any other value and attachments. The metadata does not have the same requirement and
- * may be used to persist details such as the machine that created the value or how much time was
- * spent creating the value.
- *
- * The value and its attachments are compressed, and both the compressed and uncompressed formats
- * of the values are exposed by the cache record.
- *
- * The value, attachments, and metadata are optional and can be skipped when requesting a record.
- * When the value or attachments have been skipped, the record will contain a value with the hash
- * and size, but with null data.
+ * The key must uniquely correspond to its values. The key must never be reused for other values.
+ * The metadata does not have this requirement and may be used to persist details that vary, such
+ * as the time to generate the values or the machine that generated them.
  */
 class FCacheRecord
 {
@@ -91,18 +74,12 @@ public:
 	/** Returns the metadata. Null when requested with ECachePolicy::SkipMeta. */
 	inline const FCbObject& GetMeta() const { return Record->GetMeta(); }
 
-	/** Returns the value. Data is null if skipped by the policy on request. */
-	inline const FValueWithId& GetValue() const { return Record->GetValue(); }
-
-	/** Returns the attachment matching the ID. Data is null if skipped by the policy on request. */
-	inline const FValueWithId& GetAttachment(const FValueId& Id) const { return Record->GetAttachment(Id); }
-
-	/** Returns a view of the attachments. Always available, but data may be skipped if skipped. */
-	inline TConstArrayView<FValueWithId> GetAttachments() const { return Record->GetAttachments(); }
-
 	/** Returns the value matching the ID. Null if no match. Data is null if skipped. */
 	inline const FValueWithId& GetValue(const FValueId& Id) const { return Record->GetValue(Id); }
 
+	/** Returns a view of the values ordered by ID. Data is null if skipped. */
+	inline TConstArrayView<FValueWithId> GetValues() const { return Record->GetValues(); }
+
 	/** Save the cache record to a compact binary package. */
 	UE_API FCbPackage Save() const;
 
@@ -126,8 +103,8 @@ private:
 /**
  * A cache record builder is used to construct a cache record.
  *
- * Create using a key that uniquely corresponds to the value and attachments for the cache record.
- * Metadata may vary between records of the same key.
+ * Create using a key that uniquely corresponds to the values for the cache record.
+ * The metadata may vary between records with the same key.
  *
  * @see FCacheRecord
  */
@@ -150,72 +127,39 @@ public:
 	}
 
 	/**
-	 * Set the value for the cache record.
+	 * Add a value to the cache record.
 	 *
-	 * @param Buffer      The value, which is compressed by the builder, and cloned if not owned.
 	 * @param Id          An ID for the value that is unique within this cache record.
 	 *                    When omitted, the hash of the buffer will be used as the ID.
+	 * @param Buffer      The value, which is compressed by the builder, and cloned if not owned.
 	 * @param BlockSize   The power-of-two block size to encode raw data in. 0 is default.
 	 */
-	inline void SetValue(const FCompositeBuffer& Buffer, const FValueId& Id = FValueId(), const uint64 BlockSize = 0)
-	{
-		return RecordBuilder->SetValue(Buffer, Id, BlockSize);
-	}
-	inline void SetValue(const FSharedBuffer& Buffer, const FValueId& Id = FValueId(), const uint64 BlockSize = 0)
-	{
-		return RecordBuilder->SetValue(Buffer, Id, BlockSize);
-	}
-
-	/**
-	 * Set the value for the cache record.
-	 *
-	 * @param Value   The value, which must have data unless it is known to be in the cache.
-	 */
-	inline void SetValue(const FValue& Value, const FValueId& Id = FValueId())
-	{
-		return RecordBuilder->SetValue(Value, Id);
-	}
-	inline void SetValue(const FValueWithId& Value)
-	{
-		return RecordBuilder->SetValue(Value);
-	}
+	UE_API void AddValue(const FValueId& Id, const FCompositeBuffer& Buffer, uint64 BlockSize = 0);
+	UE_API void AddValue(const FValueId& Id, const FSharedBuffer& Buffer, uint64 BlockSize = 0);
 
 	/**
-	 * Add an attachment to the cache record.
+	 * Add a value to the cache record.
 	 *
-	 * @param Buffer      The attachment, which is compressed by the builder, and cloned if not owned.
-	 * @param Id          An ID for the attachment that is unique within this cache record.
-	 *                    When omitted, the hash of the buffer will be used as the ID.
-	 * @param BlockSize   The power-of-two block size to encode raw data in. 0 is default.
+	 * @note The ID for the value must be unique within this cache record.
 	 */
-	inline void AddAttachment(const FCompositeBuffer& Buffer, const FValueId& Id = FValueId(), const uint64 BlockSize = 0)
-	{
-		return RecordBuilder->AddAttachment(Buffer, Id, BlockSize);
-	}
-	inline void AddAttachment(const FSharedBuffer& Buffer, const FValueId& Id = FValueId(), const uint64 BlockSize = 0)
-	{
-		return RecordBuilder->AddAttachment(Buffer, Id, BlockSize);
-	}
+	UE_API void AddValue(const FValueWithId& Value);
 
 	/**
-	 * Add an attachment to the cache record.
+	 * Add a value to the cache record.
 	 *
-	 * @param Value   The value, which must have data unless it is known to be in the cache.
+	 * @param Id   An ID for the value that is unique within this cache record.
+	 *             When omitted, the hash of the value will be used as the ID.
 	 */
-	inline void AddAttachment(const FValue& Value, const FValueId& Id = FValueId())
-	{
-		return RecordBuilder->AddAttachment(Value, Id);
-	}
-	inline void AddAttachment(const FValueWithId& Value)
+	inline void AddValue(const FValueId& Id, const FValue& Value)
 	{
-		return RecordBuilder->AddAttachment(Value);
+		return RecordBuilder->AddValue(Id, Value);
 	}
 
 	/**
 	 * Build a cache record, which makes this builder subsequently unusable.
 	 *
-	 * Prefer BuildAsync() when the value or attachments are added from a buffer, as this must block
-	 * on compression of those buffers before it can construct a cache record.
+	 * Prefer BuildAsync() when the values are added from raw buffers, which requires compressing
+	 * those buffers before constructing the cache record.
 	 */
 	inline FCacheRecord Build()
 	{
@@ -226,7 +170,7 @@ public:
 	/**
 	 * Build a cache record asynchronously, which makes this builder subsequently unusable.
 	 *
-	 * Prefer Build() when the value and attachments are added by value, as compression is already
+	 * Prefer Build() when the values are added from compressed buffers as compression is already
 	 * complete and BuildAsync() will complete immediately in that case.
 	 */
 	inline void BuildAsync(IRequestOwner& Owner, FOnCacheRecordComplete&& OnComplete)
diff --git a/Engine/Source/Developer/TextureBuild/Private/TextureBuildFunction.cpp b/Engine/Source/Developer/TextureBuild/Private/TextureBuildFunction.cpp
--- a/Engine/Source/Developer/TextureBuild/Private/TextureBuildFunction.cpp
+++ b/Engine/Source/Developer/TextureBuild/Private/TextureBuildFunction.cpp
@@ -427,7 +427,7 @@ void FTextureBuildFunction::Build(UE::DerivedData::FBuildContext& Context) const
 			DescriptionWriter.AddInteger("FileRegion"_ASV, static_cast<int32>(FileRegion));
 			if (bIsInlineMip)
 			{
-				DescriptionWriter.AddInteger("PayloadOffset"_ASV, MipOffset);
+				DescriptionWriter.AddInteger("MipOffset"_ASV, MipOffset);
 			}
 			DescriptionWriter.AddInteger("NumBytes"_ASV, CompressedMip.RawData.Num());
 
diff --git a/Engine/Source/Developer/Virtualization/Private/VirtualizationDDCBackend.cpp b/Engine/Source/Developer/Virtualization/Private/VirtualizationDDCBackend.cpp
--- a/Engine/Source/Developer/Virtualization/Private/VirtualizationDDCBackend.cpp
+++ b/Engine/Source/Developer/Virtualization/Private/VirtualizationDDCBackend.cpp
@@ -12,12 +12,10 @@
 namespace UE::Virtualization
 {
 
-/** Utility function to help convert from UE::Virtualization types to UE::DerivedData::FValueWithId */
-UE::DerivedData::FValueWithId ToDDCPayload(const FPayloadId& Id, const FCompressedBuffer& Payload)
+/** Utility function to help convert from UE::Virtualization::FPayloadId to UE::DerivedData::FValueId */
+static UE::DerivedData::FValueId ToDerivedDataValueId(const FPayloadId& Id)
 {
-	UE::DerivedData::FValueId ValueId = UE::DerivedData::FValueId::FromHash(Id.GetIdentifier());
-
-	return UE::DerivedData::FValueWithId(ValueId, Payload);
+	return UE::DerivedData::FValueId::FromHash(Id.GetIdentifier());
 }
 
 FDDCBackend::FDDCBackend(FStringView ConfigName, FStringView InDebugName)
@@ -89,11 +87,11 @@ EPushResult FDDCBackend::PushData(const FPayloadId& Id, const FCompressedBuffer&
 	Key.Bucket = Bucket;
 	Key.Hash = Id.GetIdentifier();
 
-	UE::DerivedData::FValueWithId DDCPayload = ToDDCPayload(Id, Payload);
-	check(DDCPayload.GetRawHash() == Id.GetIdentifier());
+	UE::DerivedData::FValue DerivedDataValue(Payload);
+	check(DerivedDataValue.GetRawHash() == Id.GetIdentifier());
 
 	UE::DerivedData::FCacheRecordBuilder RecordBuilder(Key);
-	RecordBuilder.SetValue(DDCPayload);
+	RecordBuilder.AddValue(ToDerivedDataValueId(Id), DerivedDataValue);
 
 	UE::DerivedData::FRequestOwner Owner(UE::DerivedData::EPriority::Blocking);
 
@@ -133,12 +131,12 @@ FCompressedBuffer FDDCBackend::PullData(const FPayloadId& Id)
 	FCompressedBuffer ResultData;
 	UE::DerivedData::EStatus ResultStatus;
 
-	auto Callback = [&ResultData, &ResultStatus](UE::DerivedData::FCacheGetCompleteParams&& Params)
+	auto Callback = [&Id, &ResultData, &ResultStatus](UE::DerivedData::FCacheGetCompleteParams&& Params)
 	{
 		ResultStatus = Params.Status;
 		if (ResultStatus == UE::DerivedData::EStatus::Ok)
 		{
-			ResultData = Params.Record.GetValue().GetData();
+			ResultData = Params.Record.GetValue(ToDerivedDataValueId(Id)).GetData();
 		}
 	};
 
diff --git a/Engine/Source/Editor/UnrealEd/Private/EditorDomain/EditorDomainArchive.cpp b/Engine/Source/Editor/UnrealEd/Private/EditorDomain/EditorDomainArchive.cpp
--- a/Engine/Source/Editor/UnrealEd/Private/EditorDomain/EditorDomainArchive.cpp
+++ b/Engine/Source/Editor/UnrealEd/Private/EditorDomain/EditorDomainArchive.cpp
@@ -280,34 +280,19 @@ void FEditorDomainPackageSegments::OnRecordRequestComplete(UE::DerivedData::FCac
 	int32 InitialRequestIndex = -1;
 	if (Params.Status == EStatus::Ok)
 	{
-		FCacheRecord& Record = Params.Record;
-		uint64 FileSize = Record.GetMeta()["FileSize"].AsUInt64();
-		const FValueWithId& RecordValue = Record.GetValue();
-
-		// Sort the attachment values by ValueId, in case they are not already sorted. EditorDomainPackage attachments
-		// are written with ValueIds in the same order as the segment order
-		TConstArrayView<FValueWithId> Attachments = Record.GetAttachments();
-		TArray<const FValueWithId*, TInlineAllocator<4>> SortedValues;
-		SortedValues.Reserve(Attachments.Num()+1);
-		if (RecordValue.GetRawSize() > 0)
-		{
-			SortedValues.Add(&RecordValue);
-		}
-		for (const FValueWithId& Attachment : Attachments)
-		{
-			if (Attachment.GetRawSize() > 0)
-			{
-				SortedValues.Add(&Attachment);
-			}
-		}
-		Algo::SortBy(SortedValues, &FValueWithId::GetId);
+		const FCacheRecord& Record = Params.Record;
+		const uint64 FileSize = Record.GetMeta()["FileSize"].AsUInt64();
 
+		TConstArrayView<FValueWithId> Values = Record.GetValues();
+		Segments.Reserve(Values.Num());
 		uint64 CompositeSize = 0;
-		Segments.Reserve(SortedValues.Num());
-		for (const FValueWithId* Value : SortedValues)
+		for (const FValueWithId& Value : Values)
 		{
-			Segments.Emplace(Value->GetId(), CompositeSize);
-			CompositeSize += Value->GetRawSize();
+			if (const uint64 ValueSize = Value.GetRawSize())
+			{
+				Segments.Emplace(Value.GetId(), CompositeSize);
+				CompositeSize += ValueSize;
+			}
 		}
 
 		if (CompositeSize != FileSize)
diff --git a/Engine/Source/Editor/UnrealEd/Private/EditorDomain/EditorDomainUtils.cpp b/Engine/Source/Editor/UnrealEd/Private/EditorDomain/EditorDomainUtils.cpp
--- a/Engine/Source/Editor/UnrealEd/Private/EditorDomain/EditorDomainUtils.cpp
+++ b/Engine/Source/Editor/UnrealEd/Private/EditorDomain/EditorDomainUtils.cpp
@@ -79,7 +79,7 @@ bool bGGlobalAddedCustomVersionsInitialized = false;
 int64 GMaxBulkDataSize = -1;
 
 // Change to a new guid when EditorDomain needs to be invalidated
-const TCHAR* EditorDomainVersion = TEXT("D0A667ED7E374A60962A4AA61B63E6B5");
+const TCHAR* EditorDomainVersion = TEXT("3237D60426C74B8B93E8F6BA49488E84");
 
 // Identifier of the CacheBuckets for EditorDomain tables
 const TCHAR* EditorDomainPackageBucketName = TEXT("EditorDomainPackage");
@@ -903,18 +903,24 @@ UE::DerivedData::FCacheKey GetEditorDomainPackageKey(const FIoHash& EditorDomain
 	return UE::DerivedData::FCacheKey{EditorDomainPackageCacheBucket, EditorDomainHash };
 }
 
-UE::DerivedData::FCacheKey GetBulkDataListKey(const FIoHash& EditorDomainHash)
+static UE::DerivedData::FCacheKey GetBulkDataListKey(const FIoHash& EditorDomainHash)
 {
 	static UE::DerivedData::FCacheBucket EditorDomainBulkDataListBucket(EditorDomainBulkDataListBucketName);
 	return UE::DerivedData::FCacheKey{ EditorDomainBulkDataListBucket, EditorDomainHash };
 }
 
-UE::DerivedData::FCacheKey GetBulkDataPayloadIdKey(const FIoHash& PackageAndGuidHash)
+static UE::DerivedData::FCacheKey GetBulkDataPayloadIdKey(const FIoHash& PackageAndGuidHash)
 {
 	static UE::DerivedData::FCacheBucket EditorDomainBulkDataPayloadIdBucket(EditorDomainBulkDataPayloadIdBucketName);
 	return UE::DerivedData::FCacheKey{ EditorDomainBulkDataPayloadIdBucket, PackageAndGuidHash };
 }
 
+static UE::DerivedData::FValueId GetBulkDataValueId()
+{
+	static UE::DerivedData::FValueId Id = UE::DerivedData::FValueId::FromName(ANSITEXTVIEW("Value"));
+	return Id;
+}
+
 void RequestEditorDomainPackage(const FPackagePath& PackagePath,
 	const FIoHash& EditorDomainHash, UE::DerivedData::ECachePolicy SkipFlags, UE::DerivedData::IRequestOwner& Owner,
 	UE::DerivedData::FOnCacheGetComplete&& Callback)
@@ -1319,7 +1325,7 @@ bool TrySavePackage(UPackage* Package)
 	FCacheRecordBuilder RecordBuilder(GetEditorDomainPackageKey(PackageDigest.Hash));
 	for (const FEditorDomainPackageWriter::FAttachment& Attachment : PackageWriter->GetAttachments())
 	{
-		RecordBuilder.AddAttachment(Attachment.Buffer, Attachment.ValueId);
+		RecordBuilder.AddValue(Attachment.ValueId, Attachment.Buffer);
 	}
 	RecordBuilder.SetMeta(MetaData.Save().AsObject());
 	FRequestOwner Owner(EPriority::Normal);
@@ -1369,7 +1375,7 @@ void GetBulkDataList(FName PackageName, UE::DerivedData::IRequestOwner& Owner, T
 		[InnerCallback = MoveTemp(Callback)](FCacheGetCompleteParams&& Params)
 		{
 			bool bOk = Params.Status == EStatus::Ok;
-			InnerCallback(bOk ? Params.Record.GetValue().GetData().Decompress() : FSharedBuffer());
+			InnerCallback(bOk ? Params.Record.GetValue(GetBulkDataValueId()).GetData().Decompress() : FSharedBuffer());
 		});
 }
 
@@ -1396,7 +1402,7 @@ void PutBulkDataList(FName PackageName, FSharedBuffer Buffer)
 	ICache& Cache = GetCache();
 	FRequestOwner Owner(EPriority::Normal);
 	FCacheRecordBuilder RecordBuilder(GetBulkDataListKey(PackageDigest.Hash));
-	RecordBuilder.SetValue(Buffer);
+	RecordBuilder.AddValue(GetBulkDataValueId(), Buffer);
 	Cache.Put({{{WriteToString<128>(PackageName)}, RecordBuilder.Build()}}, Owner);
 	Owner.KeepAlive();
 }
@@ -1440,7 +1446,7 @@ void GetBulkDataPayloadId(FName PackageName, const FGuid& BulkDataId, UE::Derive
 		[InnerCallback = MoveTemp(Callback)](FCacheGetCompleteParams&& Params)
 		{
 			bool bOk = Params.Status == EStatus::Ok;
-			InnerCallback(bOk ? Params.Record.GetValue().GetData().Decompress() : FSharedBuffer());
+			InnerCallback(bOk ? Params.Record.GetValue(GetBulkDataValueId()).GetData().Decompress() : FSharedBuffer());
 		});
 }
 
@@ -1468,7 +1474,7 @@ void PutBulkDataPayloadId(FName PackageName, const FGuid& BulkDataId, FSharedBuf
 	ICache& Cache = GetCache();
 	FRequestOwner Owner(EPriority::Normal);
 	FCacheRecordBuilder RecordBuilder(GetBulkDataPayloadIdKey(PackageAndGuidHash));
-	RecordBuilder.SetValue(Buffer);
+	RecordBuilder.AddValue(GetBulkDataValueId(), Buffer);
 	Cache.Put({{{WriteToString<128>(PackageName)}, RecordBuilder.Build()}}, Owner);
 	Owner.KeepAlive();
 }
diff --git a/Engine/Source/Runtime/Engine/Private/TextureDerivedDataTask.cpp b/Engine/Source/Runtime/Engine/Private/TextureDerivedDataTask.cpp
--- a/Engine/Source/Runtime/Engine/Private/TextureDerivedDataTask.cpp
+++ b/Engine/Source/Runtime/Engine/Private/TextureDerivedDataTask.cpp
@@ -1378,7 +1378,7 @@ private:
 			if (MipIndex >= NumStreamingMips)
 			{
 				uint64 MipSize = MipObjectView["NumBytes"_ASV].AsUInt64();
-				FMemoryView MipView = MipTailData.GetView().Mid(MipObjectView["PayloadOffset"_ASV].AsUInt64(), MipSize);
+				FMemoryView MipView = MipTailData.GetView().Mid(MipObjectView["MipOffset"_ASV].AsUInt64(), MipSize);
 
 				NewMip->BulkData.Lock(LOCK_READ_WRITE);
 				void* MipAllocData = NewMip->BulkData.Realloc(int64(MipSize));
