DDC: Fixes to the cache record and cache value refactor

#rnx

diff --git a/Engine/Source/Developer/DerivedDataCache/Private/DerivedDataCacheRecord.cpp b/Engine/Source/Developer/DerivedDataCache/Private/DerivedDataCacheRecord.cpp
--- a/Engine/Source/Developer/DerivedDataCache/Private/DerivedDataCacheRecord.cpp
+++ b/Engine/Source/Developer/DerivedDataCache/Private/DerivedDataCacheRecord.cpp
@@ -117,14 +117,13 @@ void FCacheRecordBuilderInternal::SetMeta(FCbObject&& InMeta)
 
 void FCacheRecordBuilderInternal::AddValue(const FValueId& Id, const FValue& Value)
 {
-	checkf(Id, TEXT("Failed to add value on %s because the ID is null."), *WriteToString<96>(Key));
+	checkf(Id.IsValid(), TEXT("Failed to add value on %s because the ID is null."), *WriteToString<96>(Key));
 	checkf(!(Value == FValue::Null), TEXT("Failed to add value on %s because the value is null."), *WriteToString<96>(Key));
-	const FValueId ValueId = Id ? Id : FValueId::FromHash(Value.GetRawHash());
-	const int32 Index = Algo::LowerBoundBy(Values, ValueId, &FValueWithId::GetId);
-	checkf(!(Values.IsValidIndex(Index) && Values[Index].GetId() == ValueId),
+	const int32 Index = Algo::LowerBoundBy(Values, Id, &FValueWithId::GetId);
+	checkf(!(Values.IsValidIndex(Index) && Values[Index].GetId() == Id),
 		TEXT("Failed to add value on %s with ID %s because it has an existing value with that ID."),
-		*WriteToString<96>(Key), *WriteToString<32>(ValueId));
-	Values.Insert(FValueWithId(ValueId, Value), Index);
+		*WriteToString<96>(Key), *WriteToString<32>(Id));
+	Values.Insert(FValueWithId(Id, Value), Index);
 }
 
 FCacheRecord FCacheRecordBuilderInternal::Build()
diff --git a/Engine/Source/Developer/DerivedDataCache/Private/DerivedDataValue.cpp b/Engine/Source/Developer/DerivedDataCache/Private/DerivedDataValue.cpp
--- a/Engine/Source/Developer/DerivedDataCache/Private/DerivedDataValue.cpp
+++ b/Engine/Source/Developer/DerivedDataCache/Private/DerivedDataValue.cpp
@@ -38,7 +38,7 @@ FValueId::FValueId(const FCbObjectId& Id)
 
 FValueId::operator FCbObjectId() const
 {
-	return FCbObjectId(ImplicitConv<const ByteArray&>(*this));
+	return FCbObjectId(GetBytes());
 }
 
 FValueId FValueId::FromHash(const FIoHash& Hash)
diff --git a/Engine/Source/Developer/DerivedDataCache/Public/DerivedDataCacheRecord.h b/Engine/Source/Developer/DerivedDataCache/Public/DerivedDataCacheRecord.h
--- a/Engine/Source/Developer/DerivedDataCache/Public/DerivedDataCacheRecord.h
+++ b/Engine/Source/Developer/DerivedDataCache/Public/DerivedDataCacheRecord.h
@@ -130,7 +130,6 @@ public:
 	 * Add a value to the cache record.
 	 *
 	 * @param Id          An ID for the value that is unique within this cache record.
-	 *                    When omitted, the hash of the buffer will be used as the ID.
 	 * @param Buffer      The value, which is compressed by the builder, and cloned if not owned.
 	 * @param BlockSize   The power-of-two block size to encode raw data in. 0 is default.
 	 */
@@ -148,7 +147,6 @@ public:
 	 * Add a value to the cache record.
 	 *
 	 * @param Id   An ID for the value that is unique within this cache record.
-	 *             When omitted, the hash of the value will be used as the ID.
 	 */
 	inline void AddValue(const FValueId& Id, const FValue& Value)
 	{
diff --git a/Engine/Source/Developer/DerivedDataCache/Public/DerivedDataValueId.h b/Engine/Source/Developer/DerivedDataCache/Public/DerivedDataValueId.h
--- a/Engine/Source/Developer/DerivedDataCache/Public/DerivedDataValueId.h
+++ b/Engine/Source/Developer/DerivedDataCache/Public/DerivedDataValueId.h
@@ -45,7 +45,6 @@ public:
 
 	/** Returns a reference to the raw byte array for the ID. */
 	inline const ByteArray& GetBytes() const { return Bytes; }
-	inline operator const ByteArray&() const { return Bytes; }
 
 	/** Returns a view of the raw byte array for the ID. */
 	inline FMemoryView GetView() const { return MakeMemoryView(Bytes); }
