DDC: Added a value API to ICacheStore

The new API is only implemented for the file system cache store for now. Support for the other cache stores will follow.

#rnx

diff --git a/Engine/Source/Developer/DerivedDataCache/Private/DerivedDataBackendAsyncPutWrapper.cpp b/Engine/Source/Developer/DerivedDataCache/Private/DerivedDataBackendAsyncPutWrapper.cpp
--- a/Engine/Source/Developer/DerivedDataCache/Private/DerivedDataBackendAsyncPutWrapper.cpp
+++ b/Engine/Source/Developer/DerivedDataCache/Private/DerivedDataBackendAsyncPutWrapper.cpp
@@ -519,6 +519,70 @@ void FDerivedDataBackendAsyncPutWrapper::Get(
 	}
 }
 
+void FDerivedDataBackendAsyncPutWrapper::PutValue(
+	const TConstArrayView<FCachePutValueRequest> Requests,
+	IRequestOwner& Owner,
+	FOnCachePutValueComplete&& OnComplete)
+{
+	if (Owner.GetPriority() == EPriority::Blocking || !GDDCIOThreadPool)
+	{
+		InnerBackend->PutValue(Requests, Owner, MoveTemp(OnComplete));
+	}
+	else
+	{
+		FDerivedDataAsyncWrapperRequest* Request = new FDerivedDataAsyncWrapperRequest(Owner,
+			[this, Requests = TArray<FCachePutValueRequest>(Requests), OnComplete = MoveTemp(OnComplete)](bool bCancel) mutable
+			{
+				if (!bCancel)
+				{
+					FRequestOwner BlockingOwner(EPriority::Blocking);
+					InnerBackend->PutValue(Requests, BlockingOwner, MoveTemp(OnComplete));
+					BlockingOwner.Wait();
+				}
+				else if (OnComplete)
+				{
+					for (const FCachePutValueRequest& Request : Requests)
+					{
+						OnComplete({Request.Name, Request.Key, Request.UserData, EStatus::Canceled});
+					}
+				}
+			});
+		Request->Start(Owner.GetPriority());
+	}
+}
+
+void FDerivedDataBackendAsyncPutWrapper::GetValue(
+	const TConstArrayView<FCacheGetValueRequest> Requests,
+	IRequestOwner& Owner,
+	FOnCacheGetValueComplete&& OnComplete)
+{
+	if (Owner.GetPriority() == EPriority::Blocking || !GDDCIOThreadPool)
+	{
+		InnerBackend->GetValue(Requests, Owner, MoveTemp(OnComplete));
+	}
+	else
+	{
+		FDerivedDataAsyncWrapperRequest* Request = new FDerivedDataAsyncWrapperRequest(Owner,
+			[this, Requests = TArray<FCacheGetValueRequest>(Requests), OnComplete = MoveTemp(OnComplete)](bool bCancel) mutable
+			{
+				if (!bCancel)
+				{
+					FRequestOwner BlockingOwner(EPriority::Blocking);
+					InnerBackend->GetValue(Requests, BlockingOwner, MoveTemp(OnComplete));
+					BlockingOwner.Wait();
+				}
+				else if (OnComplete)
+				{
+					for (const FCacheGetValueRequest& Request : Requests)
+					{
+						OnComplete({Request.Name, Request.Key, {}, Request.UserData, EStatus::Canceled});
+					}
+				}
+			});
+		Request->Start(Owner.GetPriority());
+	}
+}
+
 void FDerivedDataBackendAsyncPutWrapper::GetChunks(
 	TConstArrayView<FCacheChunkRequest> Requests,
 	IRequestOwner& Owner,
diff --git a/Engine/Source/Developer/DerivedDataCache/Private/DerivedDataBackendAsyncPutWrapper.h b/Engine/Source/Developer/DerivedDataCache/Private/DerivedDataBackendAsyncPutWrapper.h
--- a/Engine/Source/Developer/DerivedDataCache/Private/DerivedDataBackendAsyncPutWrapper.h
+++ b/Engine/Source/Developer/DerivedDataCache/Private/DerivedDataBackendAsyncPutWrapper.h
@@ -145,6 +145,16 @@ public:
 		IRequestOwner& Owner,
 		FOnCacheGetComplete&& OnComplete) override;
 
+	virtual void PutValue(
+		TConstArrayView<FCachePutValueRequest> Requests,
+		IRequestOwner& Owner,
+		FOnCachePutValueComplete&& OnComplete) override;
+
+	virtual void GetValue(
+		TConstArrayView<FCacheGetValueRequest> Requests,
+		IRequestOwner& Owner,
+		FOnCacheGetValueComplete&& OnComplete) override;
+
 	virtual void GetChunks(
 		TConstArrayView<FCacheChunkRequest> Requests,
 		IRequestOwner& Owner,
diff --git a/Engine/Source/Developer/DerivedDataCache/Private/DerivedDataCache.cpp b/Engine/Source/Developer/DerivedDataCache/Private/DerivedDataCache.cpp
--- a/Engine/Source/Developer/DerivedDataCache/Private/DerivedDataCache.cpp
+++ b/Engine/Source/Developer/DerivedDataCache/Private/DerivedDataCache.cpp
@@ -237,6 +237,34 @@ FCacheRecordPolicy FCacheRecordPolicyBuilder::Build()
 	return Policy;
 }
 
+void ICacheStore::PutValue(
+	const TConstArrayView<FCachePutValueRequest> Requests,
+	IRequestOwner& Owner,
+	FOnCachePutValueComplete&& OnComplete)
+{
+	if (OnComplete)
+	{
+		for (const FCachePutValueRequest& Request : Requests)
+		{
+			OnComplete({Request.Name, Request.Key, Request.UserData, EStatus::Error});
+		}
+	}
+}
+
+void ICacheStore::GetValue(
+	const TConstArrayView<FCacheGetValueRequest> Requests,
+	IRequestOwner& Owner,
+	FOnCacheGetValueComplete&& OnComplete)
+{
+	if (OnComplete)
+	{
+		for (const FCacheGetValueRequest& Request : Requests)
+		{
+			OnComplete({Request.Name, Request.Key, {}, Request.UserData, EStatus::Error});
+		}
+	}
+}
+
 } // UE::DerivedData
 
 namespace UE::DerivedData::Private
@@ -801,6 +829,22 @@ public:
 		return FDerivedDataBackend::Get().GetRoot().Get(Requests, Owner, MoveTemp(OnComplete));
 	}
 
+	void PutValue(
+		TConstArrayView<FCachePutValueRequest> Requests,
+		IRequestOwner& Owner,
+		FOnCachePutValueComplete&& OnComplete) final
+	{
+		return FDerivedDataBackend::Get().GetRoot().PutValue(Requests, Owner, MoveTemp(OnComplete));
+	}
+
+	void GetValue(
+		TConstArrayView<FCacheGetValueRequest> Requests,
+		IRequestOwner& Owner,
+		FOnCacheGetValueComplete&& OnComplete) final
+	{
+		return FDerivedDataBackend::Get().GetRoot().GetValue(Requests, Owner, MoveTemp(OnComplete));
+	}
+
 	void GetChunks(
 		TConstArrayView<FCacheChunkRequest> Requests,
 		IRequestOwner& Owner,
diff --git a/Engine/Source/Developer/DerivedDataCache/Private/DerivedDataLimitKeyLengthWrapper.h b/Engine/Source/Developer/DerivedDataCache/Private/DerivedDataLimitKeyLengthWrapper.h
--- a/Engine/Source/Developer/DerivedDataCache/Private/DerivedDataLimitKeyLengthWrapper.h
+++ b/Engine/Source/Developer/DerivedDataCache/Private/DerivedDataLimitKeyLengthWrapper.h
@@ -284,6 +284,52 @@ public:
 	#endif
 	}
 
+	virtual void PutValue(
+		TConstArrayView<FCachePutValueRequest> Requests,
+		IRequestOwner& Owner,
+		FOnCachePutValueComplete&& OnComplete) override
+	{
+	#if ENABLE_COOK_STATS
+		return InnerBackend->PutValue(Requests, Owner,
+			[this, OnComplete = MoveTemp(OnComplete)](FCachePutValueCompleteParams&& Params)
+			{
+				if (Params.Status == EStatus::Ok)
+				{
+					UsageStats.TimePut().AddHit(0);
+				}
+				if (OnComplete)
+				{
+					OnComplete(MoveTemp(Params));
+				}
+			});
+	#else
+		return InnerBackend->PutValue(Requests, Owner, MoveTemp(OnComplete));
+	#endif
+	}
+
+	virtual void GetValue(
+		TConstArrayView<FCacheGetValueRequest> Requests,
+		IRequestOwner& Owner,
+		FOnCacheGetValueComplete&& OnComplete) override
+	{
+	#if ENABLE_COOK_STATS
+		return InnerBackend->GetValue(Requests, Owner,
+			[this, OnComplete = MoveTemp(OnComplete)](FCacheGetValueCompleteParams&& Params)
+			{
+				if (Params.Status == EStatus::Ok)
+				{
+					UsageStats.TimeGet().AddHit(0);
+				}
+				if (OnComplete)
+				{
+					OnComplete(MoveTemp(Params));
+				}
+			});
+	#else
+		return InnerBackend->GetValue(Requests, Owner, MoveTemp(OnComplete));
+	#endif
+	}
+
 	virtual void GetChunks(
 		TConstArrayView<FCacheChunkRequest> Requests,
 		IRequestOwner& Owner,
diff --git a/Engine/Source/Developer/DerivedDataCache/Private/FileSystemDerivedDataBackend.cpp b/Engine/Source/Developer/DerivedDataCache/Private/FileSystemDerivedDataBackend.cpp
--- a/Engine/Source/Developer/DerivedDataCache/Private/FileSystemDerivedDataBackend.cpp
+++ b/Engine/Source/Developer/DerivedDataCache/Private/FileSystemDerivedDataBackend.cpp
@@ -23,6 +23,7 @@
 #include "HAL/Thread.h"
 #include "Hash/xxhash.h"
 #include "HashingArchiveProxy.h"
+#include "Math/NumericLimits.h"
 #include "Misc/CommandLine.h"
 #include "Misc/ConfigCacheIni.h"
 #include "Misc/CoreMisc.h"
@@ -72,7 +73,7 @@ static const TCHAR GContentDirectoryName[] = TEXT("Content");
 
 ///////////////////////////////////////////////////////////////////////////////////////////////////////////////////////
 
-void BuildPathForLegacyCache(const TCHAR* const CacheKey, FStringBuilderBase& Path)
+static void BuildPathForLegacyCache(const TCHAR* const CacheKey, FStringBuilderBase& Path)
 {
 	TStringBuilder<256> Key;
 	Key << CacheKey;
@@ -87,7 +88,7 @@ void BuildPathForLegacyCache(const TCHAR* const CacheKey, FStringBuilderBase& Pa
 	Path.Appendf(TEXT("%1d/%1d/%1d/%s.udd"), (Hash / 100) % 10, (Hash / 10) % 10, Hash % 10, *Key);
 }
 
-void BuildPathForCacheRecord(const FCacheKey& CacheKey, FStringBuilderBase& Path)
+static void BuildPathForCachePackage(const FCacheKey& CacheKey, FStringBuilderBase& Path)
 {
 	const FIoHash::ByteArray& Bytes = CacheKey.Hash.GetBytes();
 	Path.Appendf(TEXT("%s/%hs/%02x/%02x/"), GBucketsDirectoryName, CacheKey.Bucket.ToCString(), Bytes[0], Bytes[1]);
@@ -95,7 +96,7 @@ void BuildPathForCacheRecord(const FCacheKey& CacheKey, FStringBuilderBase& Path
 	Path << TEXT(".udd"_SV);
 }
 
-void BuildPathForCacheContent(const FIoHash& RawHash, FStringBuilderBase& Path)
+static void BuildPathForCacheContent(const FIoHash& RawHash, FStringBuilderBase& Path)
 {
 	const FIoHash::ByteArray& Bytes = RawHash.GetBytes();
 	Path.Appendf(TEXT("%s/%02x/%02x/"), GContentDirectoryName, Bytes[0], Bytes[1]);
@@ -760,50 +761,69 @@ public:
 		TConstArrayView<FCacheGetRequest> Requests,
 		IRequestOwner& Owner,
 		FOnCacheGetComplete&& OnComplete) final;
+	void PutValue(
+		TConstArrayView<FCachePutValueRequest> Requests,
+		IRequestOwner& Owner,
+		FOnCachePutValueComplete&& OnComplete) final;
+	void GetValue(
+		TConstArrayView<FCacheGetValueRequest> Requests,
+		IRequestOwner& Owner,
+		FOnCacheGetValueComplete&& OnComplete) final;
 	void GetChunks(
 		TConstArrayView<FCacheChunkRequest> Requests,
 		IRequestOwner& Owner,
 		FOnCacheChunkComplete&& OnComplete) final;
 
 private:
-	bool PutCacheRecord(FStringView Name, const FCacheRecord& Record, const FCacheRecordPolicy& Policy, uint64& OutWriteSize);
-	bool PutCacheContent(FStringView Name, const FCompressedBuffer& Content, uint64& OutWriteSize) const;
+	[[nodiscard]] bool PutCacheRecord(FStringView Name, const FCacheRecord& Record, const FCacheRecordPolicy& Policy, uint64& OutWriteSize);
 
-	FOptionalCacheRecord GetCacheRecordOnly(
+	[[nodiscard]] FOptionalCacheRecord GetCacheRecordOnly(
 		FStringView Name,
 		const FCacheKey& Key,
 		const FCacheRecordPolicy& Policy);
-	FOptionalCacheRecord GetCacheRecord(
+	[[nodiscard]] FOptionalCacheRecord GetCacheRecord(
 		FStringView Name,
 		const FCacheKey& Key,
 		const FCacheRecordPolicy& Policy,
 		EStatus& OutStatus);
-	FValueWithId GetCacheContent(
+
+	[[nodiscard]] bool PutCacheValue(FStringView Name, const FCacheKey& Key, const FValue& Value, ECachePolicy Policy, uint64& OutWriteSize);
+
+	[[nodiscard]] bool GetCacheValueOnly(FStringView Name, const FCacheKey& Key, ECachePolicy Policy, FValue& OutValue);
+	[[nodiscard]] bool GetCacheValue(FStringView Name, const FCacheKey& Key, ECachePolicy Policy, FValue& OutValue);
+
+	[[nodiscard]] bool PutCacheContent(FStringView Name, const FCompressedBuffer& Content, uint64& OutWriteSize) const;
+
+	[[nodiscard]] bool GetCacheContentExists(const FCacheKey& Key, const FIoHash& RawHash) const;
+	[[nodiscard]] bool GetCacheContent(
 		FStringView Name,
 		const FCacheKey& Key,
-		const FValueWithId& Value,
-		ECachePolicy Policy) const;
+		const FValueId& Id,
+		const FValue& Value,
+		ECachePolicy Policy,
+		FValue& OutValue) const;
 	void GetCacheContent(
 		FStringView Name,
 		const FCacheKey& Key,
-		const FValueWithId& Value,
+		const FValueId& Id,
+		const FValue& Value,
 		ECachePolicy Policy,
 		FCompressedBufferReader& Reader,
 		TUniquePtr<FArchive>& OutArchive) const;
 
-	void BuildCacheRecordPath(const FCacheKey& CacheKey, FStringBuilderBase& Path) const;
+	void BuildCachePackagePath(const FCacheKey& CacheKey, FStringBuilderBase& Path) const;
 	void BuildCacheContentPath(const FIoHash& RawHash, FStringBuilderBase& Path) const;
 
-	bool SaveFileWithHash(FStringBuilderBase& Path, FStringView DebugName, TFunctionRef<void (FArchive&)> WriteFunction) const;
-	bool LoadFileWithHash(FStringBuilderBase& Path, FStringView DebugName, TFunctionRef<void (FArchive&)> ReadFunction) const;
-	bool SaveFile(FStringBuilderBase& Path, FStringView DebugName, TFunctionRef<void (FArchive&)> WriteFunction) const;
-	bool LoadFile(FStringBuilderBase& Path, FStringView DebugName, TFunctionRef<void (FArchive&)> ReadFunction) const;
-	TUniquePtr<FArchive> OpenFile(FStringBuilderBase& Path, FStringView DebugName) const;
+	[[nodiscard]] bool SaveFileWithHash(FStringBuilderBase& Path, FStringView DebugName, TFunctionRef<void (FArchive&)> WriteFunction) const;
+	[[nodiscard]] bool LoadFileWithHash(FStringBuilderBase& Path, FStringView DebugName, TFunctionRef<void (FArchive&)> ReadFunction) const;
+	[[nodiscard]] bool SaveFile(FStringBuilderBase& Path, FStringView DebugName, TFunctionRef<void (FArchive&)> WriteFunction) const;
+	[[nodiscard]] bool LoadFile(FStringBuilderBase& Path, FStringView DebugName, TFunctionRef<void (FArchive&)> ReadFunction) const;
+	[[nodiscard]] TUniquePtr<FArchive> OpenFile(FStringBuilderBase& Path, FStringView DebugName) const;
 
-	bool FileExists(FStringBuilderBase& Path) const;
+	[[nodiscard]] bool FileExists(FStringBuilderBase& Path) const;
 
-	bool ShouldSimulateMiss(const TCHAR* InKey);
-	bool ShouldSimulateMiss(const FCacheKey& Key);
+	[[nodiscard]] bool ShouldSimulateMiss(const TCHAR* InKey);
+	[[nodiscard]] bool ShouldSimulateMiss(const FCacheKey& Key);
 
 private:
 	void BuildCacheLegacyPath(const TCHAR* const CacheKey, FStringBuilderBase& Path) const
@@ -1593,6 +1613,72 @@ void FFileSystemCacheStore::Get(
 	}
 }
 
+void FFileSystemCacheStore::PutValue(
+	const TConstArrayView<FCachePutValueRequest> Requests,
+	IRequestOwner& Owner,
+	FOnCachePutValueComplete&& OnComplete)
+{
+	for (const FCachePutValueRequest& Request : Requests)
+	{
+		TRACE_CPUPROFILER_EVENT_SCOPE(FileSystemDDC_Put);
+		TRACE_COUNTER_INCREMENT(FileSystemDDC_Put);
+		COOK_STAT(auto Timer = UsageStats.TimePut());
+		uint64 WriteSize = 0;
+		if (PutCacheValue(Request.Name, Request.Key, Request.Value, Request.Policy, WriteSize))
+		{
+			UE_LOG(LogDerivedDataCache, Verbose, TEXT("%s: Cache put complete for %s from '%s'"),
+				*CachePath, *WriteToString<96>(Request.Key), *Request.Name);
+			TRACE_COUNTER_INCREMENT(FileSystemDDC_PutHit);
+			TRACE_COUNTER_ADD(FileSystemDDC_BytesWritten, WriteSize);
+			COOK_STAT(if (WriteSize) { Timer.AddHit(WriteSize); });
+			if (OnComplete)
+			{
+				OnComplete({Request.Name, Request.Key, Request.UserData, EStatus::Ok});
+			}
+		}
+		else
+		{
+			if (OnComplete)
+			{
+				OnComplete({Request.Name, Request.Key, Request.UserData, EStatus::Error});
+			}
+		}
+	}
+}
+
+void FFileSystemCacheStore::GetValue(
+	const TConstArrayView<FCacheGetValueRequest> Requests,
+	IRequestOwner& Owner,
+	FOnCacheGetValueComplete&& OnComplete)
+{
+	for (const FCacheGetValueRequest& Request : Requests)
+	{
+		TRACE_CPUPROFILER_EVENT_SCOPE(FileSystemDDC_Get);
+		TRACE_COUNTER_INCREMENT(FileSystemDDC_Get);
+		COOK_STAT(auto Timer = UsageStats.TimeGet());
+		FValue Value;
+		if (GetCacheValue(Request.Name, Request.Key, Request.Policy, Value))
+		{
+			UE_LOG(LogDerivedDataCache, Verbose, TEXT("%s: Cache hit for %s from '%s'"),
+				*CachePath, *WriteToString<96>(Request.Key), *Request.Name);
+			TRACE_COUNTER_INCREMENT(FileSystemDDC_GetHit);
+			TRACE_COUNTER_ADD(FileSystemDDC_BytesRead, Value.GetData().GetCompressedSize());
+			COOK_STAT(Timer.AddHit(Value.GetData().GetCompressedSize()));
+			if (OnComplete)
+			{
+				OnComplete({Request.Name, Request.Key, Value, Request.UserData, EStatus::Ok});
+			}
+		}
+		else
+		{
+			if (OnComplete)
+			{
+				OnComplete({Request.Name, Request.Key, {}, Request.UserData, EStatus::Error});
+			}
+		}
+	}
+}
+
 void FFileSystemCacheStore::GetChunks(
 	const TConstArrayView<FCacheChunkRequest> Requests,
 	IRequestOwner& Owner,
@@ -1601,7 +1687,9 @@ void FFileSystemCacheStore::GetChunks(
 	TArray<FCacheChunkRequest, TInlineAllocator<16>> SortedRequests(Requests);
 	SortedRequests.StableSort(TChunkLess());
 
-	FValueWithId Value;
+	bool bHasValue = false;
+	FValue Value;
+	FValueId ValueId;
 	FCacheKey ValueKey;
 	TUniquePtr<FArchive> ValueAr;
 	FCompressedBufferReader ValueReader;
@@ -1613,27 +1701,42 @@ void FFileSystemCacheStore::GetChunks(
 		TRACE_CPUPROFILER_EVENT_SCOPE(FileSystemDDC_Get);
 		TRACE_COUNTER_INCREMENT(FileSystemDDC_Get);
 		COOK_STAT(auto Timer = bExistsOnly ? UsageStats.TimeProbablyExists() : UsageStats.TimeGet());
-		if (!(Value && ValueKey == Request.Key && Value.GetId() == Request.Id) || !bExistsOnly <= ValueReader.HasSource())
+		if (!(bHasValue && ValueKey == Request.Key && ValueId == Request.Id) || ValueReader.HasSource() < !bExistsOnly)
 		{
 			ValueStatus = EStatus::Error;
 			ValueReader.ResetSource();
 			ValueAr.Reset();
+			ValueKey = {};
+			ValueId.Reset();
 			Value.Reset();
-			if (!(Record && Record.Get().GetKey() == Request.Key))
+			bHasValue = false;
+			if (Request.Id.IsValid())
 			{
-				FCacheRecordPolicyBuilder PolicyBuilder(ECachePolicy::None);
-				PolicyBuilder.AddValuePolicy(Request.Id, Request.Policy);
-				Record.Reset();
-				Record = GetCacheRecordOnly(Request.Name, Request.Key, PolicyBuilder.Build());
+				if (!(Record && Record.Get().GetKey() == Request.Key))
+				{
+					FCacheRecordPolicyBuilder PolicyBuilder(ECachePolicy::None);
+					PolicyBuilder.AddValuePolicy(Request.Id, Request.Policy);
+					Record.Reset();
+					Record = GetCacheRecordOnly(Request.Name, Request.Key, PolicyBuilder.Build());
+				}
+				if (Record)
+				{
+					const FValueWithId& ValueWithId = Record.Get().GetValue(Request.Id);
+					bHasValue = ValueWithId.IsValid();
+					Value = ValueWithId;
+					ValueId = Request.Id;
+					ValueKey = Request.Key;
+					GetCacheContent(Request.Name, Request.Key, ValueId, Value, Request.Policy, ValueReader, ValueAr);
+				}
 			}
-			if (Record)
+			else
 			{
-				Value = Record.Get().GetValue(Request.Id);
 				ValueKey = Request.Key;
-				GetCacheContent(Request.Name, Request.Key, Value, Request.Policy, ValueReader, ValueAr);
+				bHasValue = GetCacheValueOnly(Request.Name, Request.Key, Request.Policy, Value);
+				GetCacheContent(Request.Name, Request.Key, Request.Id, Value, Request.Policy, ValueReader, ValueAr);
 			}
 		}
-		if (Value)
+		if (bHasValue)
 		{
 			const uint64 RawOffset = FMath::Min(Value.GetRawSize(), Request.RawOffset);
 			const uint64 RawSize = FMath::Min(Value.GetRawSize() - RawOffset, Request.RawSize);
@@ -1702,7 +1805,7 @@ bool FFileSystemCacheStore::PutCacheRecord(
 	bool bRecordExists = false;
 	FCbPackage ExistingPackage;
 	TStringBuilder<256> Path;
-	BuildCacheRecordPath(Key, Path);
+	BuildCachePackagePath(Key, Path);
 	if (EnumHasAnyFlags(CombinedValuePolicy, ECachePolicy::SkipData))
 	{
 		bRecordExists = FileExists(Path);
@@ -1758,7 +1861,10 @@ bool FFileSystemCacheStore::PutCacheRecord(
 	for (FCompressedBuffer& Content : ExternalContent)
 	{
 		uint64 WriteSize = 0;
-		PutCacheContent(Name, Content, OutWriteSize);
+		if (!PutCacheContent(Name, Content, OutWriteSize))
+		{
+			return false;
+		}
 		OutWriteSize += WriteSize;
 	}
 
@@ -1791,9 +1897,8 @@ FOptionalCacheRecord FFileSystemCacheStore::GetCacheRecordOnly(
 	}
 
 	// Skip the request if querying the cache is disabled.
-	const ECachePolicy QueryPolicy = SpeedClass == ESpeedClass::Local
-		? ECachePolicy::QueryLocal : ECachePolicy::QueryRemote;
-	if (!EnumHasAnyFlags(Policy.GetRecordPolicy(), QueryPolicy))
+	const ECachePolicy QueryFlag = SpeedClass == ESpeedClass::Local ? ECachePolicy::QueryLocal : ECachePolicy::QueryRemote;
+	if (!EnumHasAnyFlags(Policy.GetRecordPolicy(), QueryFlag))
 	{
 		UE_LOG(LogDerivedDataCache, VeryVerbose, TEXT("%s: Skipped get of %s from '%.*s' due to cache policy"),
 			*CachePath, *WriteToString<96>(Key), Name.Len(), Name.GetData());
@@ -1808,12 +1913,12 @@ FOptionalCacheRecord FFileSystemCacheStore::GetCacheRecordOnly(
 	}
 
 	TStringBuilder<256> Path;
-	BuildCacheRecordPath(Key, Path);
+	BuildCachePackagePath(Key, Path);
 
-	bool bDeleteRecord = true;
+	bool bDeletePackage = true;
 	ON_SCOPE_EXIT
 	{
-		if (bDeleteRecord && !bReadOnly)
+		if (bDeletePackage && !bReadOnly)
 		{
 			IFileManager::Get().Delete(*Path, /*bRequireExists*/ false, /*bEvenReadOnly*/ false, /*bQuiet*/ true);
 		}
@@ -1824,7 +1929,7 @@ FOptionalCacheRecord FFileSystemCacheStore::GetCacheRecordOnly(
 		FCbPackage Package;
 		if (!LoadFileWithHash(Path, Name, [&Package](FArchive& Ar) { Package.TryLoad(Ar); }))
 		{
-			UE_LOG(LogDerivedDataCache, Verbose, TEXT("%s: Cache miss with missing record for %s from '%.*s'"),
+			UE_LOG(LogDerivedDataCache, Verbose, TEXT("%s: Cache miss with missing package for %s from '%.*s'"),
 				*CachePath, *WriteToString<96>(Key), Name.Len(), Name.GetData());
 			return Record;
 		}
@@ -1833,7 +1938,7 @@ FOptionalCacheRecord FFileSystemCacheStore::GetCacheRecordOnly(
 		{
 			UE_LOG(LogDerivedDataCache, Display, TEXT("%s: Cache miss with invalid package for %s from '%.*s'"),
 				*CachePath, *WriteToString<96>(Key), Name.Len(), Name.GetData());
-			return FOptionalCacheRecord();
+			return Record;
 		}
 
 		Record = FCacheRecord::Load(Package);
@@ -1850,7 +1955,7 @@ FOptionalCacheRecord FFileSystemCacheStore::GetCacheRecordOnly(
 		AccessLogWriter->Append(Key, Path);
 	}
 
-	bDeleteRecord = false;
+	bDeletePackage = false;
 	return Record;
 }
 
@@ -1878,10 +1983,12 @@ FOptionalCacheRecord FFileSystemCacheStore::GetCacheRecord(
 
 	for (const FValueWithId& Value : Record.Get().GetValues())
 	{
-		const ECachePolicy ValuePolicy = Policy.GetValuePolicy(Value.GetId());
-		if (FValueWithId Content = GetCacheContent(Name, Key, Value, ValuePolicy))
+		const FValueId& Id = Value.GetId();
+		const ECachePolicy ValuePolicy = Policy.GetValuePolicy(Id);
+		FValue Content;
+		if (GetCacheContent(Name, Key, Id, Value, ValuePolicy, Content))
 		{
-			RecordBuilder.AddValue(MoveTemp(Content));
+			RecordBuilder.AddValue(Id, MoveTemp(Content));
 		}
 		else if (EnumHasAnyFlags(ValuePolicy, ECachePolicy::PartialOnError))
 		{
@@ -1898,6 +2005,210 @@ FOptionalCacheRecord FFileSystemCacheStore::GetCacheRecord(
 	return RecordBuilder.Build();
 }
 
+bool FFileSystemCacheStore::PutCacheValue(
+	const FStringView Name,
+	const FCacheKey& Key,
+	const FValue& Value,
+	const ECachePolicy Policy,
+	uint64& OutWriteSize)
+{
+	if (!IsWritable())
+	{
+		UE_LOG(LogDerivedDataCache, VeryVerbose,
+			TEXT("%s: Skipped put of %s from '%.*s' because this cache store is read-only"),
+			*CachePath, *WriteToString<96>(Key), Name.Len(), Name.GetData());
+		return false;
+	}
+
+	// Skip the request if storing to the cache is disabled.
+	const ECachePolicy StoreFlag = SpeedClass == ESpeedClass::Local ? ECachePolicy::StoreLocal : ECachePolicy::StoreRemote;
+	if (!EnumHasAnyFlags(Policy, StoreFlag))
+	{
+		UE_LOG(LogDerivedDataCache, VeryVerbose, TEXT("%s: Skipped put of %s from '%.*s' due to cache policy"),
+			*CachePath, *WriteToString<96>(Key), Name.Len(), Name.GetData());
+		return false;
+	}
+
+	if (ShouldSimulateMiss(Key))
+	{
+		UE_LOG(LogDerivedDataCache, Verbose, TEXT("%s: Simulated miss for put of %s from '%.*s'"),
+			*CachePath, *WriteToString<96>(Key), Name.Len(), Name.GetData());
+		return false;
+	}
+
+	// Check if there is an existing value package.
+	bool bValueExists = false;
+	FCbPackage ExistingPackage;
+	TStringBuilder<256> Path;
+	BuildCachePackagePath(Key, Path);
+	if (EnumHasAnyFlags(Policy, ECachePolicy::SkipData))
+	{
+		bValueExists = FileExists(Path);
+	}
+	else
+	{
+		bValueExists = LoadFileWithHash(Path, Name, [&ExistingPackage](FArchive& Ar) { ExistingPackage.TryLoad(Ar); })
+			&& (ExistingPackage.FindAttachment(Value.GetRawHash()) || GetCacheContentExists(Key, Value.GetRawHash()));
+	}
+
+	// Save the value to a package and save the data to external content depending on its size.
+	if (!bValueExists)
+	{
+		FCbWriter Writer;
+		Writer.BeginObject();
+		Writer.AddBinaryAttachment("RawHash", Value.GetRawHash());
+		Writer.AddInteger("RawSize", Value.GetRawSize());
+		Writer.EndObject();
+
+		FCbPackage Package(Writer.Save().AsObject());
+		if (!Value.HasData())
+		{
+			// Verify that the content exists in storage.
+			if (!GetCacheContentExists(Key, Value.GetRawHash()))
+			{
+				UE_LOG(LogDerivedDataCache, Verbose, TEXT("%s: Failed due to missing data for put of %s from '%.*s'"),
+					*CachePath, *WriteToString<96>(Key), Name.Len(), Name.GetData());
+				return false;
+			}
+		}
+		else if (Value.GetData().GetCompressedSize() <= MaxValueSizeKB)
+		{
+			// Store the content in the package.
+			Package.AddAttachment(FCbAttachment(Value.GetData()));
+		}
+		else
+		{
+			// Save the external content to storage.
+			uint64 WriteSize = 0;
+			if (!PutCacheContent(Name, Value.GetData(), WriteSize))
+			{
+				return false;
+			}
+			OutWriteSize += WriteSize;
+		}
+
+		// Save the value package to storage.
+		const auto WritePackage = [&](FArchive& Ar) { Package.Save(Ar); OutWriteSize += uint64(Ar.TotalSize()); };
+		if (!SaveFileWithHash(Path, Name, WritePackage))
+		{
+			return false;
+		}
+	}
+
+	if (AccessLogWriter)
+	{
+		AccessLogWriter->Append(Key, Path);
+	}
+
+	return true;
+}
+
+bool FFileSystemCacheStore::GetCacheValueOnly(
+	const FStringView Name,
+	const FCacheKey& Key,
+	const ECachePolicy Policy,
+	FValue& OutValue)
+{
+	if (!IsUsable())
+	{
+		UE_LOG(LogDerivedDataCache, VeryVerbose,
+			TEXT("%s: Skipped get of %s from '%.*s' because this cache store is not available"),
+			*CachePath, *WriteToString<96>(Key), Name.Len(), Name.GetData());
+		return false;
+	}
+
+	// Skip the request if querying the cache is disabled.
+	const ECachePolicy QueryFlag = SpeedClass == ESpeedClass::Local ? ECachePolicy::QueryLocal : ECachePolicy::QueryRemote;
+	if (!EnumHasAnyFlags(Policy, QueryFlag))
+	{
+		UE_LOG(LogDerivedDataCache, VeryVerbose, TEXT("%s: Skipped get of %s from '%.*s' due to cache policy"),
+			*CachePath, *WriteToString<96>(Key), Name.Len(), Name.GetData());
+		return false;
+	}
+
+	if (ShouldSimulateMiss(Key))
+	{
+		UE_LOG(LogDerivedDataCache, Verbose, TEXT("%s: Simulated miss for get of %s from '%.*s'"),
+			*CachePath, *WriteToString<96>(Key), Name.Len(), Name.GetData());
+		return false;
+	}
+
+	TStringBuilder<256> Path;
+	BuildCachePackagePath(Key, Path);
+
+	bool bDeletePackage = true;
+	ON_SCOPE_EXIT
+	{
+		if (bDeletePackage && !bReadOnly)
+		{
+			IFileManager::Get().Delete(*Path, /*bRequireExists*/ false, /*bEvenReadOnly*/ false, /*bQuiet*/ true);
+		}
+	};
+
+	FCbPackage Package;
+	if (!LoadFileWithHash(Path, Name, [&Package](FArchive& Ar) { Package.TryLoad(Ar); }))
+	{
+		UE_LOG(LogDerivedDataCache, Verbose, TEXT("%s: Cache miss with missing package for %s from '%.*s'"),
+			*CachePath, *WriteToString<96>(Key), Name.Len(), Name.GetData());
+		return false;
+	}
+
+	if (ValidateCompactBinary(Package, ECbValidateMode::Default | ECbValidateMode::Package) != ECbValidateError::None)
+	{
+		UE_LOG(LogDerivedDataCache, Display, TEXT("%s: Cache miss with invalid package for %s from '%.*s'"),
+			*CachePath, *WriteToString<96>(Key), Name.Len(), Name.GetData());
+		return false;
+	}
+
+	const FCbObjectView Object = Package.GetObject();
+	const FIoHash RawHash = Object["RawHash"].AsHash();
+	const uint64 RawSize = Object["RawSize"].AsUInt64(MAX_uint64);
+	if (RawHash.IsZero() || RawSize == MAX_uint64)
+	{
+		UE_LOG(LogDerivedDataCache, Display, TEXT("%s: Cache miss with invalid value for %s from '%.*s'"),
+			*CachePath, *WriteToString<96>(Key), Name.Len(), Name.GetData());
+		return false;
+	}
+
+	if (EnumHasAnyFlags(Policy, ECachePolicy::SkipData))
+	{
+		OutValue = FValue(RawHash, RawSize);
+	}
+	else if (const FCbAttachment* const Attachment = Package.FindAttachment(RawHash))
+	{
+		const FCompressedBuffer& Data = Attachment->AsCompressedBinary();
+		if (Data.GetRawHash() != RawHash || Data.GetRawSize() != RawSize)
+		{
+			UE_LOG(LogDerivedDataCache, Display,
+				TEXT("%s: Cache miss with invalid value attachment for %s from '%.*s'"),
+				*CachePath, *WriteToString<96>(Key), Name.Len(), Name.GetData());
+			return false;
+		}
+		OutValue = FValue(Data);
+	}
+	else
+	{
+		OutValue = FValue(RawHash, RawSize);
+	}
+
+	if (AccessLogWriter)
+	{
+		AccessLogWriter->Append(Key, Path);
+	}
+
+	bDeletePackage = false;
+	return true;
+}
+
+bool FFileSystemCacheStore::GetCacheValue(
+	const FStringView Name,
+	const FCacheKey& Key,
+	const ECachePolicy Policy,
+	FValue& OutValue)
+{
+	return GetCacheValueOnly(Name, Key, Policy, OutValue) && GetCacheContent(Name, Key, {}, OutValue, Policy, OutValue);
+}
+
 bool FFileSystemCacheStore::PutCacheContent(
 	const FStringView Name,
 	const FCompressedBuffer& Content,
@@ -1918,20 +2229,31 @@ bool FFileSystemCacheStore::PutCacheContent(
 	return true;
 }
 
-FValueWithId FFileSystemCacheStore::GetCacheContent(
+bool FFileSystemCacheStore::GetCacheContentExists(const FCacheKey& Key, const FIoHash& RawHash) const
+{
+	TStringBuilder<256> Path;
+	BuildCacheContentPath(RawHash, Path);
+	return FileExists(Path);
+}
+
+bool FFileSystemCacheStore::GetCacheContent(
 	const FStringView Name,
 	const FCacheKey& Key,
-	const FValueWithId& Value,
-	const ECachePolicy Policy) const
+	const FValueId& Id,
+	const FValue& Value,
+	const ECachePolicy Policy,
+	FValue& OutValue) const
 {
 	if (!EnumHasAnyFlags(Policy, ECachePolicy::Query))
 	{
-		return Value.RemoveData();
+		OutValue = Value.RemoveData();
+		return true;
 	}
 
 	if (Value.HasData())
 	{
-		return EnumHasAnyFlags(Policy, ECachePolicy::SkipData) ? Value.RemoveData() : Value;
+		OutValue = EnumHasAnyFlags(Policy, ECachePolicy::SkipData) ? Value.RemoveData() : Value;
+		return true;
 	}
 
 	const FIoHash& RawHash = Value.GetRawHash();
@@ -1946,7 +2268,8 @@ FValueWithId FFileSystemCacheStore::GetCacheContent(
 			{
 				AccessLogWriter->Append(RawHash, Path);
 			}
-			return Value;
+			OutValue = Value;
+			return true;
 		}
 	}
 	else
@@ -1960,27 +2283,29 @@ FValueWithId FFileSystemCacheStore::GetCacheContent(
 				{
 					AccessLogWriter->Append(RawHash, Path);
 				}
-				return FValueWithId(Value.GetId(), MoveTemp(CompressedBuffer));
+				OutValue = FValue(MoveTemp(CompressedBuffer));
+				return true;
 			}
 			UE_LOG(LogDerivedDataCache, Display,
 				TEXT("%s: Cache miss with corrupted value %s with hash %s for %s from '%.*s'"),
-				*CachePath, *WriteToString<16>(Value.GetId()), *WriteToString<48>(RawHash),
+				*CachePath, *WriteToString<16>(Id), *WriteToString<48>(RawHash),
 				*WriteToString<96>(Key), Name.Len(), Name.GetData());
-			return FValueWithId::Null;
+			return false;
 		}
 	}
 
 	UE_LOG(LogDerivedDataCache, Verbose,
 		TEXT("%s: Cache miss with missing value %s with hash %s for %s from '%.*s'"),
-		*CachePath, *WriteToString<16>(Value.GetId()), *WriteToString<48>(RawHash), *WriteToString<96>(Key),
+		*CachePath, *WriteToString<16>(Id), *WriteToString<48>(RawHash), *WriteToString<96>(Key),
 		Name.Len(), Name.GetData());
-	return FValueWithId::Null;
+	return false;
 }
 
 void FFileSystemCacheStore::GetCacheContent(
 	const FStringView Name,
 	const FCacheKey& Key,
-	const FValueWithId& Value,
+	const FValueId& Id,
+	const FValue& Value,
 	const ECachePolicy Policy,
 	FCompressedBufferReader& Reader,
 	TUniquePtr<FArchive>& OutArchive) const
@@ -2031,7 +2356,7 @@ void FFileSystemCacheStore::GetCacheContent(
 			}
 			UE_LOG(LogDerivedDataCache, Display,
 				TEXT("%s: Cache miss with corrupted value %s with hash %s for %s from '%.*s'"),
-				*CachePath, *WriteToString<16>(Value.GetId()), *WriteToString<48>(RawHash),
+				*CachePath, *WriteToString<16>(Id), *WriteToString<48>(RawHash),
 				*WriteToString<96>(Key), Name.Len(), Name.GetData());
 			Reader.ResetSource();
 			OutArchive.Reset();
@@ -2041,14 +2366,14 @@ void FFileSystemCacheStore::GetCacheContent(
 
 	UE_LOG(LogDerivedDataCache, Verbose,
 		TEXT("%s: Cache miss with missing value %s with hash %s for %s from '%.*s'"),
-		*CachePath, *WriteToString<16>(Value.GetId()), *WriteToString<48>(RawHash), *WriteToString<96>(Key),
+		*CachePath, *WriteToString<16>(Id), *WriteToString<48>(RawHash), *WriteToString<96>(Key),
 		Name.Len(), Name.GetData());
 }
 
-void FFileSystemCacheStore::BuildCacheRecordPath(const FCacheKey& CacheKey, FStringBuilderBase& Path) const
+void FFileSystemCacheStore::BuildCachePackagePath(const FCacheKey& CacheKey, FStringBuilderBase& Path) const
 {
 	Path << CachePath << TEXT('/');
-	BuildPathForCacheRecord(CacheKey, Path);
+	BuildPathForCachePackage(CacheKey, Path);
 }
 
 void FFileSystemCacheStore::BuildCacheContentPath(const FIoHash& RawHash, FStringBuilderBase& Path) const
diff --git a/Engine/Source/Developer/DerivedDataCache/Private/HierarchicalDerivedDataBackend.h b/Engine/Source/Developer/DerivedDataCache/Private/HierarchicalDerivedDataBackend.h
--- a/Engine/Source/Developer/DerivedDataCache/Private/HierarchicalDerivedDataBackend.h
+++ b/Engine/Source/Developer/DerivedDataCache/Private/HierarchicalDerivedDataBackend.h
@@ -606,6 +606,135 @@ public:
 		}
 	}
 
+	virtual void PutValue(
+		TConstArrayView<FCachePutValueRequest> Requests,
+		IRequestOwner& Owner,
+		FOnCachePutValueComplete&& OnComplete) override
+	{
+		FRequestOwner AsyncOwner(FPlatformMath::Min(Owner.GetPriority(), EPriority::Highest));
+		FRequestBarrier AsyncBarrier(AsyncOwner);
+		AsyncOwner.KeepAlive();
+
+		TSet<FCacheKey> RequestsOk;
+
+		{
+			FReadScopeLock LockScope(Lock);
+
+			for (int32 PutCacheIndex = 0; PutCacheIndex < InnerBackends.Num(); PutCacheIndex++)
+			{
+				if (InnerBackends[PutCacheIndex]->IsWritable())
+				{
+					// Every record must put synchronously before switching to async calls.
+					if (RequestsOk.Num() < Requests.Num())
+					{
+						FRequestOwner BlockingOwner(EPriority::Blocking);
+						InnerBackends[PutCacheIndex]->PutValue(Requests, BlockingOwner,
+							[&OnComplete, &RequestsOk](FCachePutValueCompleteParams&& Params)
+							{
+								if (Params.Status == EStatus::Ok)
+								{
+									bool bIsAlreadyInSet = false;
+									RequestsOk.FindOrAdd(Params.Key, &bIsAlreadyInSet);
+									if (OnComplete && !bIsAlreadyInSet)
+									{
+										OnComplete(MoveTemp(Params));
+									}
+								}
+							});
+						BlockingOwner.Wait();
+					}
+					else
+					{
+						AsyncPutInnerBackends[PutCacheIndex]->PutValue(Requests, AsyncOwner);
+					}
+				}
+			}
+		}
+
+		if (OnComplete && RequestsOk.Num() < Requests.Num())
+		{
+			for (const FCachePutValueRequest& Request : Requests)
+			{
+				if (!RequestsOk.Contains(Request.Key))
+				{
+					OnComplete({Request.Name, Request.Key, Request.UserData, EStatus::Error});
+				}
+			}
+		}
+	}
+
+	virtual void GetValue(
+		TConstArrayView<FCacheGetValueRequest> Requests,
+		IRequestOwner& Owner,
+		FOnCacheGetValueComplete&& OnComplete) override
+	{
+		TArray<FCacheGetValueRequest, TInlineAllocator<16>> RemainingRequests(Requests);
+
+		{
+			FReadScopeLock LockScope(Lock);
+			TSet<FCacheKey> KeysOk;
+
+			bool bHadMiss = false;
+
+			for (int32 GetCacheIndex = 0; GetCacheIndex < InnerBackends.Num() && !RemainingRequests.IsEmpty(); ++GetCacheIndex)
+			{
+				// Block on this because backends in this hierarchy are not expected to be asynchronous.
+				FRequestOwner BlockingOwner(EPriority::Blocking);
+				InnerBackends[GetCacheIndex]->GetValue(RemainingRequests, BlockingOwner,
+					[this, GetCacheIndex, &Owner, &OnComplete, &KeysOk](FCacheGetValueCompleteParams&& Params)
+					{
+						if (Params.Status == EStatus::Ok)
+						{
+							FRequestOwner AsyncOwner(FPlatformMath::Min(Owner.GetPriority(), EPriority::Highest));
+							FRequestBarrier AsyncBarrier(AsyncOwner);
+							AsyncOwner.KeepAlive();
+							for (int32 FillCacheIndex = 0; FillCacheIndex < InnerBackends.Num(); ++FillCacheIndex)
+							{
+								if (GetCacheIndex != FillCacheIndex)
+								{
+									AsyncPutInnerBackends[FillCacheIndex]->PutValue({{Params.Name, Params.Key, Params.Value, ECachePolicy::Default}}, AsyncOwner);
+								}
+							}
+
+							KeysOk.Add(Params.Key);
+							if (OnComplete)
+							{
+								OnComplete(MoveTemp(Params));
+							}
+						}
+					});
+				BlockingOwner.Wait();
+
+				RemainingRequests.RemoveAll([&KeysOk](const FCacheGetValueRequest& Request) { return KeysOk.Contains(Request.Key); });
+
+				if (!bHadMiss && !RemainingRequests.IsEmpty() && InnerBackends[GetCacheIndex]->IsWritable())
+				{
+					bHadMiss = true;
+					const auto ConvertPolicy = [](ECachePolicy P) -> ECachePolicy
+					{
+						if (EnumHasAnyFlags(P, ECachePolicy::Store))
+						{
+							EnumRemoveFlags(P, ECachePolicy::SkipData);
+						}
+						return P;
+					};
+					for (FCacheGetValueRequest& Request : RemainingRequests)
+					{
+						Request.Policy = ConvertPolicy(Request.Policy);
+					}
+				}
+			}
+		}
+
+		if (OnComplete)
+		{
+			for (const FCacheGetValueRequest& Request : RemainingRequests)
+			{
+				OnComplete({Request.Name, Request.Key, {}, Request.UserData, EStatus::Error});
+			}
+		}
+	}
+
 	virtual void GetChunks(
 		TConstArrayView<FCacheChunkRequest> Requests,
 		IRequestOwner& Owner,
diff --git a/Engine/Source/Developer/DerivedDataCache/Public/DerivedDataCache.h b/Engine/Source/Developer/DerivedDataCache/Public/DerivedDataCache.h
--- a/Engine/Source/Developer/DerivedDataCache/Public/DerivedDataCache.h
+++ b/Engine/Source/Developer/DerivedDataCache/Public/DerivedDataCache.h
@@ -8,8 +8,10 @@
 #include "DerivedDataCacheRecord.h"
 #include "DerivedDataRequestTypes.h"
 #include "DerivedDataSharedString.h"
+#include "DerivedDataValue.h"
 #include "DerivedDataValueId.h"
 #include "Math/NumericLimits.h"
+#include "Memory/SharedBuffer.h"
 #include "Misc/EnumClassFlags.h"
 #include "Templates/Function.h"
 #include "Templates/RefCounting.h"
@@ -22,8 +24,12 @@ namespace UE::DerivedData { struct FCacheChunkCompleteParams; }
 namespace UE::DerivedData { struct FCacheChunkRequest; }
 namespace UE::DerivedData { struct FCacheGetCompleteParams; }
 namespace UE::DerivedData { struct FCacheGetRequest; }
+namespace UE::DerivedData { struct FCacheGetValueCompleteParams; }
+namespace UE::DerivedData { struct FCacheGetValueRequest; }
 namespace UE::DerivedData { struct FCachePutCompleteParams; }
 namespace UE::DerivedData { struct FCachePutRequest; }
+namespace UE::DerivedData { struct FCachePutValueCompleteParams; }
+namespace UE::DerivedData { struct FCachePutValueRequest; }
 namespace UE::DerivedData::Private { class ICacheRecordPolicyShared; }
 
 namespace UE::DerivedData
@@ -33,6 +39,8 @@ using FCacheRequestName = FSharedString;
 
 using FOnCachePutComplete = TUniqueFunction<void (FCachePutCompleteParams&& Params)>;
 using FOnCacheGetComplete = TUniqueFunction<void (FCacheGetCompleteParams&& Params)>;
+using FOnCachePutValueComplete = TUniqueFunction<void (FCachePutValueCompleteParams&& Params)>;
+using FOnCacheGetValueComplete = TUniqueFunction<void (FCacheGetValueCompleteParams&& Params)>;
 using FOnCacheChunkComplete = TUniqueFunction<void (FCacheChunkCompleteParams&& Params)>;
 
 ///////////////////////////////////////////////////////////////////////////////////////////////////////////////////////
@@ -135,7 +143,7 @@ public:
 };
 
 /**
- * Flags to control the behavior of cache requests, with optional overrides by value.
+ * Flags to control the behavior of cache record requests, with optional overrides by value.
  *
  * Examples:
  * - A base policy of Disable, with value policy overrides of Default, will fetch those values if
@@ -210,7 +218,7 @@ private:
 ///////////////////////////////////////////////////////////////////////////////////////////////////////////////////////
 
 /**
- * Interface to a store of cache records.
+ * Interface to a store of cache records and cache values.
  *
  * Functions on this interface may be called from any thread. When a callback is provided, it may
  * be invoked from a different thread than the request was made on, and may be invoked before the
@@ -255,7 +263,7 @@ public:
 	 * PartialOnError policy applies to the missing values, the cache store must return a partial
 	 * record containing the available values.
 	 *
-	 * @param Requests     Requests with the keys identifying the cache records to query.
+	 * @param Requests     Requests with the keys identifying the cache records to fetch.
 	 * @param Owner        The owner to execute the request within.
 	 * @param OnComplete   A callback invoked for every key as it completes or is canceled.
 	 */
@@ -264,6 +272,16 @@ public:
 		IRequestOwner& Owner,
 		FOnCacheGetComplete&& OnComplete) = 0;
 
+	UE_API virtual void PutValue(
+		TConstArrayView<FCachePutValueRequest> Requests,
+		IRequestOwner& Owner,
+		FOnCachePutValueComplete&& OnComplete = FOnCachePutValueComplete());
+
+	UE_API virtual void GetValue(
+		TConstArrayView<FCacheGetValueRequest> Requests,
+		IRequestOwner& Owner,
+		FOnCacheGetValueComplete&& OnComplete);
+
 	/**
 	 * Asynchronous request to get chunks, which are subsets of values, from records or values.
 	 *
@@ -274,7 +292,7 @@ public:
 	 * value is requested, through a put of a partial record or a value. The cache store does not
 	 * need to verify the existence of unrequested values within requested records.
 	 *
-	 * @param Requests     The keys, IDs, offsets, and sizes of the chunks to query.
+	 * @param Requests     The keys, IDs, offsets, and sizes of the chunks to fetch.
 	 * @param Owner        The owner to execute the request within.
 	 * @param OnComplete   A callback invoked for every chunk as it completes or is canceled.
 	 */
@@ -289,7 +307,7 @@ public:
 /**
  * Interface to the cache.
  *
- * @see ICacheStore for cache record storage.
+ * @see ICacheStore for cache record and cache value storage.
  */
 class ICache : public ICacheStore
 {
@@ -318,7 +336,7 @@ struct FCachePutRequest
 	/** A name to identify this request for logging and profiling. An object path is typically sufficient. */
 	FCacheRequestName Name;
 
-	/** The cache record to put. */
+	/** A record to store. */
 	FCacheRecord Record;
 
 	/** Flags to control the behavior of the request. See FCacheRecordPolicy. */
@@ -334,23 +352,23 @@ struct FCachePutCompleteParams
 	/** A copy of the name from the request. */
 	FCacheRequestName Name;
 
-	/** Key for the part of the put request that completed or was canceled. */
+	/** A copy of the key from the request. */
 	FCacheKey Key;
 
 	/** A copy of the value from the request. */
 	uint64 UserData = 0;
 
-	/** Status of the cache request. */
+	/** The status of the request. */
 	EStatus Status = EStatus::Error;
 };
 
-/** Parameters to request a cache record. */
+/** Parameters to request to get a cache record. */
 struct FCacheGetRequest
 {
 	/** A name to identify this request for logging and profiling. An object path is typically sufficient. */
 	FCacheRequestName Name;
 
-	/** The key identifying the cache record to fetch. */
+	/** A key identifying the record to fetch. */
 	FCacheKey Key;
 
 	/** Flags to control the behavior of the request. See FCacheRecordPolicy. */
@@ -367,32 +385,107 @@ struct FCacheGetCompleteParams
 	FCacheRequestName Name;
 
 	/**
-	 * Record for the part of the get request that completed or was canceled.
+	 * Record for the request that completed or was canceled.
 	 *
 	 * The key is always populated. The remainder of the record is populated when Status is Ok.
 	 *
 	 * The metadata or the data for values may be skipped based on cache policy flags. Values for
 	 * which data has been skipped will have a hash and size but null data.
 	 */
-	FCacheRecord&& Record;
+	FCacheRecord Record;
+
+	/** A copy of the value from the request. */
+	uint64 UserData = 0;
+
+	/** The status of the request. */
+	EStatus Status = EStatus::Error;
+};
+
+/** Parameters to request to put a cache value. */
+struct FCachePutValueRequest
+{
+	/** A name to identify this request for logging and profiling. An object path is typically sufficient. */
+	FCacheRequestName Name;
+
+	/** A key that will uniquely identify the value in the cache. */
+	FCacheKey Key;
+
+	/** A value to store. */
+	FValue Value;
+
+	/** Flags to control the behavior of the request. See ECachePolicy. */
+	ECachePolicy Policy = ECachePolicy::Default;
+
+	/** A value that will be returned in the completion callback. */
+	uint64 UserData = 0;
+};
+
+/** Parameters for the completion callback for cache value put requests. */
+struct FCachePutValueCompleteParams
+{
+	/** A copy of the name from the request. */
+	FCacheRequestName Name;
+
+	/** A copy of the key from the request. */
+	FCacheKey Key;
+
+	/** A copy of the value from the request. */
+	uint64 UserData = 0;
+
+	/** The status of the request. */
+	EStatus Status = EStatus::Error;
+};
+
+/** Parameters to request to get a cache value. */
+struct FCacheGetValueRequest
+{
+	/** A name to identify this request for logging and profiling. An object path is typically sufficient. */
+	FCacheRequestName Name;
+
+	/** A key identifying the value to fetch. */
+	FCacheKey Key;
+
+	/** Flags to control the behavior of the request. See ECachePolicy. */
+	ECachePolicy Policy = ECachePolicy::Default;
+
+	/** A value that will be returned in the completion callback. */
+	uint64 UserData = 0;
+};
+
+/** Parameters for the completion callback for cache value get requests. */
+struct FCacheGetValueCompleteParams
+{
+	/** A copy of the name from the request. */
+	FCacheRequestName Name;
+
+	/** A copy of the key from the request. */
+	FCacheKey Key;
+
+	/**
+	 * Value for the request that completed or was canceled.
+	 *
+	 * The data may be skipped based on cache policy flags. A value for which data has been skipped
+	 * will have a hash and size but null data.
+	 */
+	FValue Value;
 
 	/** A copy of the value from the request. */
 	uint64 UserData = 0;
 
-	/** Status of the cache request. */
+	/** The status of the request. */
 	EStatus Status = EStatus::Error;
 };
 
-/** Parameters to request a chunk, which is a subset of a value, from a cache record. */
+/** Parameters to request a chunk, which is a subset of a value, from a cache record or cache value. */
 struct FCacheChunkRequest
 {
 	/** A name to identify this request for logging and profiling. An object path is typically sufficient. */
 	FCacheRequestName Name;
 
-	/** The key identifying the cache record to fetch the value from. */
+	/** A key identifying the record or value to fetch the chunk from. */
 	FCacheKey Key;
 
-	/** The ID of the value to fetch from the cache record. */
+	/** An ID identifying the value to fetch, if fetching from a record, otherwise null. */
 	FValueId Id;
 
 	/** The offset into the raw bytes of the value at which to start fetching. */
@@ -417,28 +510,28 @@ struct FCacheChunkCompleteParams
 	/** A copy of the name from the request. */
 	FCacheRequestName Name;
 
-	/** Key from the chunk request that completed or was canceled. */
+	/** A copy of the key from the request. */
 	FCacheKey Key;
 
-	/** ID from the chunk request that completed or was canceled. */
+	/** A copy of the ID from the request. */
 	FValueId Id;
 
-	/** Offset from the chunk request that completed or was canceled. */
+	/** A copy of the offset from the request. */
 	uint64 RawOffset = 0;
 
-	/** Size, in bytes, of the subset of the value that was fetched, if any. */
+	/** The size, in bytes, of the subset of the value that was fetched, if any. */
 	uint64 RawSize = 0;
 
-	/** Hash of the entire value, even if only a subset was fetched. */
+	/** The hash of the entire value, even if only a subset was fetched. */
 	const FIoHash& RawHash;
 
 	/** Data for the subset of the value that was fetched when Status is Ok, otherwise null. */
-	FSharedBuffer&& RawData;
+	FSharedBuffer RawData;
 
 	/** A copy of the value from the request. */
 	uint64 UserData = 0;
 
-	/** Status of the cache request. */
+	/** The status of the request. */
 	EStatus Status = EStatus::Error;
 };
 
