Add support for structured cache operations through the Jupiter DDC backend.

diff --git a/Engine/Source/Developer/DerivedDataCache/Private/DerivedDataBackends.cpp b/Engine/Source/Developer/DerivedDataCache/Private/DerivedDataBackends.cpp
--- a/Engine/Source/Developer/DerivedDataCache/Private/DerivedDataBackends.cpp
+++ b/Engine/Source/Developer/DerivedDataCache/Private/DerivedDataBackends.cpp
@@ -747,6 +747,7 @@ public:
 		const TCHAR* IniSection,
 		FString& Host,
 		FString& Namespace,
+		FString& StructuredNamespace,
 		FString& OAuthProvider,
 		FString& OAuthClientId,
 		FString& OAuthSecret,
@@ -759,7 +760,7 @@ public:
 			const TCHAR* ServerSection = TEXT("HordeStorageServers");
 			if (GConfig->GetString(ServerSection, *ServerId, ServerEntry, IniFilename))
 			{
-				ParseHttpCacheParams(NodeName, *ServerEntry, IniFilename, IniSection, Host, Namespace, OAuthProvider, OAuthClientId, OAuthSecret, bReadOnly);
+				ParseHttpCacheParams(NodeName, *ServerEntry, IniFilename, IniSection, Host, Namespace, StructuredNamespace, OAuthProvider, OAuthClientId, OAuthSecret, bReadOnly);
 			}
 			else
 			{
@@ -790,6 +791,7 @@ public:
 		}
 
 		FParse::Value(Entry, TEXT("Namespace="), Namespace);
+		FParse::Value(Entry, TEXT("StructuredNamespace="), StructuredNamespace);
 		FParse::Value(Entry, TEXT("OAuthProvider="), OAuthProvider);
 		FParse::Value(Entry, TEXT("OAuthClientId="), OAuthClientId);
 		FParse::Value(Entry, TEXT("OAuthSecret="), OAuthSecret);
@@ -808,12 +810,13 @@ public:
 #if WITH_HTTP_DDC_BACKEND
 		FString Host;
 		FString Namespace;
+		FString StructuredNamespace;
 		FString OAuthProvider;
 		FString OAuthClientId;
 		FString OAuthSecret;
 		bool bReadOnly = false;
 
-		ParseHttpCacheParams(NodeName, Entry, IniFilename, IniSection, Host, Namespace, OAuthProvider, OAuthClientId, OAuthSecret, bReadOnly);
+		ParseHttpCacheParams(NodeName, Entry, IniFilename, IniSection, Host, Namespace, StructuredNamespace, OAuthProvider, OAuthClientId, OAuthSecret, bReadOnly);
 
 		if (Host.IsEmpty())
 		{
@@ -833,6 +836,11 @@ public:
 			UE_LOG(LogDerivedDataCache, Warning, TEXT("Node %s does not specify 'Namespace', falling back to '%s'"), NodeName, *Namespace);
 		}
 
+		if (StructuredNamespace.IsEmpty())
+		{
+			StructuredNamespace = Namespace;
+		}
+
 		if (OAuthProvider.IsEmpty())
 		{
 			UE_LOG(LogDerivedDataCache, Error, TEXT("Node %s does not specify 'OAuthProvider'"), NodeName);
@@ -877,7 +885,7 @@ public:
 			}
 		}
 
-		FHttpDerivedDataBackend* Backend = new FHttpDerivedDataBackend(*Host, *Namespace, *OAuthProvider, *OAuthClientId, *OAuthSecret, bReadOnly);
+		FHttpDerivedDataBackend* Backend = new FHttpDerivedDataBackend(*Host, *Namespace, *StructuredNamespace, *OAuthProvider, *OAuthClientId, *OAuthSecret, bReadOnly);
 
 		if (ForceSpeedClass != FDerivedDataBackendInterface::ESpeedClass::Unknown)
 		{
diff --git a/Engine/Source/Developer/DerivedDataCache/Private/HttpDerivedDataBackend.cpp b/Engine/Source/Developer/DerivedDataCache/Private/HttpDerivedDataBackend.cpp
--- a/Engine/Source/Developer/DerivedDataCache/Private/HttpDerivedDataBackend.cpp
+++ b/Engine/Source/Developer/DerivedDataCache/Private/HttpDerivedDataBackend.cpp
@@ -11,11 +11,14 @@
 #if PLATFORM_WINDOWS || PLATFORM_HOLOLENS
 #include "Windows/HideWindowsPlatformTypes.h"
 #endif
+#include "Algo/Accumulate.h"
 #include "Algo/Transform.h"
 #include "Algo/Find.h"
 #include "Containers/StaticArray.h"
 #include "Containers/Ticker.h"
+#include "DerivedDataCachePrivate.h"
 #include "DerivedDataCacheRecord.h"
+#include "DerivedDataValue.h"
 #include "Dom/JsonObject.h"
 #include "GenericPlatform/GenericPlatformFile.h"
 #include "HAL/PlatformFileManager.h"
@@ -29,6 +32,9 @@
 #include "ProfilingDebugging/CpuProfilerTrace.h"
 #include "ProfilingDebugging/CountersTrace.h"
 #include "Serialization/BufferArchive.h"
+#include "Serialization/CompactBinary.h"
+#include "Serialization/CompactBinaryPackage.h"
+#include "Serialization/CompactBinaryValidation.h"
 #include "Serialization/JsonReader.h"
 #include "Serialization/JsonWriter.h"
 #include "Serialization/JsonSerializer.h"
@@ -114,6 +120,8 @@ public:
 	{
 		Get,
 		Put,
+		PutCompactBinary,
+		PutCompressedBlob,
 		Post,
 		PostJson,
 		Delete,
@@ -231,17 +239,28 @@ public:
 	template<RequestVerb V>
 	Result PerformBlockingUpload(const TCHAR* Uri, TArrayView<const uint8> Buffer)
 	{
-		static_assert(V == Put || V == Post || V == PostJson, "Upload should use either Put or Post verbs.");
+		static_assert(V == Put || V == PutCompactBinary || V == PutCompressedBlob || V == Post || V == PostJson, "Upload should use either Put or Post verbs.");
 		
 		uint32 ContentLength = 0u;
 
-		if (V == Put)
+		if constexpr (V == Put || V == PutCompactBinary || V == PutCompressedBlob)
 		{
 			curl_easy_setopt(Curl, CURLOPT_UPLOAD, 1L);
 			curl_easy_setopt(Curl, CURLOPT_INFILESIZE, Buffer.Num());
 			curl_easy_setopt(Curl, CURLOPT_READDATA, this);
 			curl_easy_setopt(Curl, CURLOPT_READFUNCTION, StaticReadFn);
-			Headers.Add(FString(TEXT("Content-Type: application/octet-stream")));
+			if constexpr (V == PutCompactBinary)
+			{
+				Headers.Add(FString(TEXT("Content-Type: application/x-ue-cb")));
+			}
+			else if constexpr (V == PutCompressedBlob)
+			{
+				Headers.Add(FString(TEXT("Content-Type: application/x-ue-comp")));
+			}
+			else
+			{
+				Headers.Add(FString(TEXT("Content-Type: application/octet-stream")));
+			}
 			ContentLength = Buffer.Num();
 			ReadDataView = Buffer;
 		}
@@ -500,6 +519,8 @@ private:
 				AdditionalInfo = FString::Printf(TEXT("Received: %d bytes."), BytesReceived);
 				break;
 			case Put:
+			case PutCompactBinary:
+			case PutCompressedBlob:
 				bSuccess = IsSuccessResponse(ResponseCode);
 				VerbStr = TEXT("updating");
 				AdditionalInfo = FString::Printf(TEXT("Sent: %d bytes."), BytesSent);
@@ -1901,6 +1922,11 @@ bool ShouldAbortForShutdown()
 	return !GIsBuildMachine && FDerivedDataBackend::Get().IsShuttingDown();
 }
 
+TConstArrayView<uint8> MakeArrayView(FSharedBuffer Buffer)
+{
+	return TConstArrayView<uint8>(reinterpret_cast<const uint8*>(Buffer.GetData()), Buffer.GetSize());
+}
+
 //----------------------------------------------------------------------------------------------------------
 // FHttpAccessToken
 //----------------------------------------------------------------------------------------------------------
@@ -1933,12 +1959,14 @@ uint32 FHttpAccessToken::GetSerial() const
 FHttpDerivedDataBackend::FHttpDerivedDataBackend(
 	const TCHAR* InServiceUrl, 
 	const TCHAR* InNamespace, 
+	const TCHAR* InStructuredNamespace, 
 	const TCHAR* InOAuthProvider,
 	const TCHAR* InOAuthClientId,
 	const TCHAR* InOAuthSecret,
 	const bool bInReadOnly)
 	: Domain(InServiceUrl)
 	, Namespace(InNamespace)
+	, StructuredNamespace(InStructuredNamespace)
 	, DefaultBucket(TEXT("default"))
 	, OAuthProvider(InOAuthProvider)
 	, OAuthClientId(InOAuthClientId)
@@ -2163,7 +2191,6 @@ bool FHttpDerivedDataBackend::ShouldSimulateMiss(const TCHAR* InKey)
 	if (DebugOptions.ShouldSimulateMiss(InKey))
 	{
 		FScopeLock Lock(&MissedKeysCS);
-		UE_LOG(LogDerivedDataCache, Verbose, TEXT("Simulating miss in %s for %s"), *GetName(), InKey);
 		DebugMissedKeys.AddByHash(Hash, Key);
 		return true;
 	}
@@ -2171,6 +2198,474 @@ bool FHttpDerivedDataBackend::ShouldSimulateMiss(const TCHAR* InKey)
 	return false;
 }
 
+bool FHttpDerivedDataBackend::ShouldSimulateMiss(const FCacheKey& Key)
+{
+	if (DebugOptions.RandomMissRate == 0 && DebugOptions.SimulateMissTypes.IsEmpty())
+	{
+		return false;
+	}
+
+	const uint32 Hash = GetTypeHash(Key);
+
+	if (FScopeLock Lock(&MissedKeysCS); DebugMissedCacheKeys.ContainsByHash(Hash, Key))
+	{
+		return true;
+	}
+
+	if (DebugOptions.ShouldSimulateMiss(Key))
+	{
+		FScopeLock Lock(&MissedKeysCS);
+		DebugMissedCacheKeys.AddByHash(Hash, Key);
+		return true;
+	}
+
+	return false;
+}
+
+uint64 FHttpDerivedDataBackend::PutRef(const FCacheRecord& Record, const FCbPackage& Package, FStringView Bucket, bool bFinalize, TArray<FIoHash>& OutNeededBlobHashes, bool& bOutPutCompletedSuccessfully)
+{
+	bOutPutCompletedSuccessfully = false;
+
+	uint64 BytesSent = 0;
+	const FCacheKey& Key = Record.GetKey();
+	TStringBuilder<256> RefsUri;
+	RefsUri << "api/v1/refs/" << StructuredNamespace << "/" << Bucket << "/" << Key.Hash;
+	if (bFinalize)
+	{
+		RefsUri << "/finalize/" << Package.GetObjectHash();
+	}
+
+	int64 ResponseCode = 0;
+	for (uint32 Attempts = 0; (Attempts < UE_HTTPDDC_MAX_ATTEMPTS) && !ShouldAbortForShutdown() && (Attempts == 0 || ShouldRetryOnError(ResponseCode)); ++Attempts)
+	{
+		FScopedRequestPtr Request(PutRequestPools[IsInGameThread()].Get());
+
+		if (bFinalize)
+		{
+			Request->PerformBlockingUpload<FHttpRequest::Post>(*RefsUri, TArrayView<const uint8>());
+		}
+		else
+		{
+			Request->SetHeader(TEXT("X-Jupiter-IoHash"), *WriteToString<48>(Package.GetObjectHash()));
+			Request->PerformBlockingUpload<FHttpRequest::PutCompactBinary>(*RefsUri, MakeArrayView(Package.GetObject().GetBuffer().ToShared()));
+		}
+		ResponseCode = Request->GetResponseCode();
+
+		if (FHttpRequest::IsSuccessResponse(ResponseCode))
+		{
+			BytesSent += Request->GetBytesSent();
+			if (TSharedPtr<FJsonObject> ResponseObject = Request->GetResponseAsJsonObject())
+			{
+				TArray<FString> NeedsArrayStrings;
+				ResponseObject->TryGetStringArrayField(TEXT("needs"), NeedsArrayStrings);
+
+				OutNeededBlobHashes.Empty(NeedsArrayStrings.Num());
+				for (const FString& NeededString : NeedsArrayStrings)
+				{
+					FIoHash BlobHash;
+					LexFromString(BlobHash, *NeededString);
+					if (!BlobHash.IsZero())
+					{
+						OutNeededBlobHashes.Add(BlobHash);
+					}
+				}
+			}
+			else
+			{
+				OutNeededBlobHashes.Empty();
+			}
+
+			if (OutNeededBlobHashes.IsEmpty())
+			{
+				bOutPutCompletedSuccessfully = true;
+			}
+			break;
+		}
+		else
+		{
+			OutNeededBlobHashes.Empty();
+		}
+	}
+
+	return BytesSent;
+}
+
+bool FHttpDerivedDataBackend::PutCacheRecord(
+	const FCacheRecord& Record,
+	const FStringView Context,
+	const FCacheRecordPolicy& Policy,
+	uint64& OutWriteSize)
+{
+	OutWriteSize = 0;
+
+	if (!IsWritable())
+	{
+		UE_LOG(LogDerivedDataCache, VeryVerbose,
+			TEXT("%s: Skipped put of %s from '%.*s' because this cache store is read-only"),
+			*GetName(), *WriteToString<96>(Record.GetKey()), Context.Len(), Context.GetData());
+		return false;
+	}
+
+	const FCacheKey& Key = Record.GetKey();
+	const ECachePolicy CombinedValuePolicy = Algo::TransformAccumulate(
+		Policy.GetValuePolicies(), &FCacheValuePolicy::Policy, Policy.GetDefaultValuePolicy(), UE_PROJECTION(operator|));
+	const ECachePolicy CombinedPolicy = Policy.GetRecordPolicy() | CombinedValuePolicy;
+
+	// Skip the request if storing to the cache is disabled.
+	// Http backends won't generally be "local" but including handling for this possibility for consistency
+	const ECachePolicy StoreFlag = SpeedClass == ESpeedClass::Local ? ECachePolicy::StoreLocal : ECachePolicy::StoreRemote;
+	if (!EnumHasAnyFlags(CombinedPolicy, StoreFlag))
+	{
+		UE_LOG(LogDerivedDataCache, VeryVerbose, TEXT("%s: Skipped put of %s from '%.*s' due to cache policy"),
+			*GetName(), *WriteToString<96>(Key), Context.Len(), Context.GetData());
+		return false;
+	}
+
+	if (ShouldSimulateMiss(Key))
+	{
+		UE_LOG(LogDerivedDataCache, Verbose, TEXT("%s: Simulated miss for put of %s from '%.*s'"),
+			*GetName(), *WriteToString<96>(Key), Context.Len(), Context.GetData());
+		return false;
+	}
+
+	// TODO: Jupiter currently always overwrites.  It doesn't have a "write if not present" feature (for records or attachments),
+	//		 but would require one to implement all policy correctly.
+
+	TStringConversion<TStringConvert<ANSICHAR, TCHAR>> BucketStringConversion(Key.Bucket.ToString());
+	FString Bucket(FStringView(BucketStringConversion.Get(), BucketStringConversion.Length()));
+	Bucket.ToLowerInline();
+
+	bool bPutCompletedSuccessfully = false;
+
+	FCbPackage Package = Record.Save();
+	TArray<FIoHash> NeededBlobHashes;
+
+	// Initial record upload
+	size_t PutRefBytesSent = PutRef(Record, Package, Bucket, false, NeededBlobHashes, bPutCompletedSuccessfully);
+	OutWriteSize += PutRefBytesSent;
+
+	// TODO: blob uploading and finalization should be replaced with a single batch compressed blob upload endpoint in the future.
+
+	// Needed blob upload (if any missing)
+	for (const FIoHash& NeededBlobHash : NeededBlobHashes)
+	{
+		TStringBuilder<256> CompressedBlobsUri;
+		CompressedBlobsUri << "api/v1/compressed-blobs/" << StructuredNamespace << "/" << NeededBlobHash;
+
+		if (const FCbAttachment* Attachment = Package.FindAttachment(NeededBlobHash))
+		{
+			FSharedBuffer TempBuffer;
+			TConstArrayView<uint8> BlobArrayView;
+			if (Attachment->IsCompressedBinary())
+			{
+				TempBuffer = Attachment->AsCompressedBinary().GetCompressed().ToShared();
+				BlobArrayView = MakeArrayView(TempBuffer);
+			}
+			else if (Attachment->IsBinary())
+			{
+				TempBuffer = FCompressedBuffer::Compress(Attachment->AsCompositeBinary()).GetCompressed().ToShared();
+				BlobArrayView = MakeArrayView(TempBuffer);
+			}
+			else
+			{
+				TempBuffer = FCompressedBuffer::Compress(Attachment->AsObject().GetBuffer()).GetCompressed().ToShared();
+				BlobArrayView = MakeArrayView(TempBuffer);
+			}
+
+			int64 ResponseCode = 0;
+			for (uint32 Attempts = 0; (Attempts < UE_HTTPDDC_MAX_ATTEMPTS) && !ShouldAbortForShutdown() && (Attempts == 0 || ShouldRetryOnError(ResponseCode)); ++Attempts)
+			{
+				FScopedRequestPtr Request(PutRequestPools[IsInGameThread()].Get());
+				Request->PerformBlockingUpload<FHttpRequest::PutCompressedBlob>(*CompressedBlobsUri, BlobArrayView);
+
+				ResponseCode = Request->GetResponseCode();
+
+				if (FHttpRequest::IsSuccessResponse(ResponseCode))
+				{
+					OutWriteSize += Request->GetBytesSent();
+					break;
+				}
+			}
+		}
+	}
+
+	// Finalization (if any blobs needed)
+	if (!NeededBlobHashes.IsEmpty())
+	{
+		size_t FinalizeRefBytesSent = PutRef(Record, Package, Bucket, true, NeededBlobHashes, bPutCompletedSuccessfully);
+		OutWriteSize += FinalizeRefBytesSent;
+	}
+
+	return bPutCompletedSuccessfully;
+}
+
+FOptionalCacheRecord FHttpDerivedDataBackend::GetCacheRecordOnly(
+	const FCacheKey& Key,
+	const FStringView Context,
+	const FCacheRecordPolicy& Policy)
+{
+	if (!IsUsable())
+	{
+		UE_LOG(LogDerivedDataCache, VeryVerbose,
+			TEXT("%s: Skipped get of %s from '%.*s' because this cache store is not available"),
+			*GetName(), *WriteToString<96>(Key), Context.Len(), Context.GetData());
+		return FOptionalCacheRecord();
+	}
+
+	// Skip the request if querying the cache is disabled.
+	const ECachePolicy QueryPolicy = SpeedClass == ESpeedClass::Local
+		? ECachePolicy::QueryLocal : ECachePolicy::QueryRemote;
+	if (!EnumHasAnyFlags(Policy.GetRecordPolicy(), QueryPolicy))
+	{
+		UE_LOG(LogDerivedDataCache, VeryVerbose, TEXT("%s: Skipped get of %s from '%.*s' due to cache policy"),
+			*GetName(), *WriteToString<96>(Key), Context.Len(), Context.GetData());
+		return FOptionalCacheRecord();
+	}
+
+	if (ShouldSimulateMiss(Key))
+	{
+		UE_LOG(LogDerivedDataCache, Verbose, TEXT("%s: Simulated miss for get of %s from '%.*s'"),
+			*GetName(), *WriteToString<96>(Key), Context.Len(), Context.GetData());
+		return FOptionalCacheRecord();
+	}
+
+	FOptionalCacheRecord Record;
+	{
+		TStringConversion<TStringConvert<ANSICHAR, TCHAR>> BucketStringConversion(Key.Bucket.ToString());
+		FString Bucket(FStringView(BucketStringConversion.Get(), BucketStringConversion.Length()));
+		Bucket.ToLowerInline();
+
+		TStringBuilder<256> RefsUri;
+		RefsUri << "api/v1/refs/" << StructuredNamespace << "/" << Bucket << "/" << Key.Hash;
+
+		FSharedBuffer ResponseBuffer;
+		int64 ResponseCode = 0;
+		for (uint32 Attempts = 0; (Attempts < UE_HTTPDDC_MAX_ATTEMPTS) && !ShouldAbortForShutdown() && (Attempts == 0 || ShouldRetryOnError(ResponseCode)); ++Attempts)
+		{
+			FScopedRequestPtr Request(PutRequestPools[IsInGameThread()].Get());
+
+			TArray<uint8> ByteArray;
+			Request->SetHeader(TEXT("Accept"), TEXT("application/x-ue-cb"));
+			Request->PerformBlockingDownload(*RefsUri, &ByteArray);
+			ResponseCode = Request->GetResponseCode();
+
+			if (FHttpRequest::IsSuccessResponse(ResponseCode))
+			{
+				ResponseBuffer = MakeSharedBufferFromArray(MoveTemp(ByteArray));
+				break;
+			}
+		}
+
+		if (ValidateCompactBinary(ResponseBuffer, ECbValidateMode::Default) != ECbValidateError::None)
+		{
+			UE_LOG(LogDerivedDataCache, Display, TEXT("%s: Cache miss with invalid package for %s from '%.*s'"),
+				*GetName(), *WriteToString<96>(Key), Context.Len(), Context.GetData());
+			return FOptionalCacheRecord();
+		}
+
+		Record = FCacheRecord::Load(FCbPackage(FCbObject(ResponseBuffer)));
+		if (Record.IsNull())
+		{
+			UE_LOG(LogDerivedDataCache, Display, TEXT("%s: Cache miss with record load failure for %s from '%.*s'"),
+				*GetName(), *WriteToString<96>(Key), Context.Len(), Context.GetData());
+			return Record;
+		}
+	}
+
+	return Record;
+}
+
+FOptionalCacheRecord FHttpDerivedDataBackend::GetCacheRecord(
+	const FCacheKey& Key,
+	const FStringView Context,
+	const FCacheRecordPolicy& Policy,
+	EStatus& OutStatus)
+{
+	FOptionalCacheRecord Record = GetCacheRecordOnly(Key, Context, Policy);
+	if (Record.IsNull())
+	{
+		OutStatus = EStatus::Error;
+		return Record;
+	}
+
+	OutStatus = EStatus::Ok;
+
+	FCacheRecordBuilder RecordBuilder(Key);
+
+	if (!EnumHasAnyFlags(Policy.GetRecordPolicy(), ECachePolicy::SkipMeta))
+	{
+		RecordBuilder.SetMeta(FCbObject(Record.Get().GetMeta()));
+	}
+
+	// TODO: There is not currently a batched GET endpoint for Jupiter.  Once there is, all payload data should be fetched in one call.
+	//		 In the meantime, we try to keep the code structured in a way that is friendly to future batching of GETs.
+
+	TArray<FValueWithId> RequiredGets;
+	TArray<FValueWithId> RequiredHeads;
+
+	auto IsValueDataReady = [](FValueWithId& Value, const ECachePolicy Policy, const ECachePolicy SkipFlag)
+	{
+		if (!EnumHasAnyFlags(Policy, ECachePolicy::Query))
+		{
+			Value = Value.RemoveData();
+			return true;
+		}
+
+		if (Value.HasData())
+		{
+			if (EnumHasAnyFlags(Policy, SkipFlag))
+			{
+				Value = Value.RemoveData();
+			}
+			return true;
+		}
+		return false;
+	};
+
+	for (FValueWithId Value : Record.Get().GetValues())
+	{
+		const ECachePolicy ValuePolicy = Policy.GetValuePolicy(Value.GetId());
+		const ECachePolicy SkipFlag = ECachePolicy::SkipData;
+		if (IsValueDataReady(Value, ValuePolicy, SkipFlag))
+		{
+			RecordBuilder.AddValue(MoveTemp(Value));
+		}
+		else
+		{
+			if (EnumHasAnyFlags(ValuePolicy, SkipFlag))
+			{
+				RequiredHeads.Emplace(Value);
+			}
+			else
+			{
+				RequiredGets.Emplace(Value);
+			}
+		}
+	}
+
+	if (!CachedDataProbablyExistsBatch(Key, RequiredHeads, Context))
+	{
+		OutStatus = EStatus::Error;
+		return FOptionalCacheRecord();
+	}
+
+	TArray<FCompressedBuffer> FetchedBuffers;
+	if (!TryGetCachedDataBatch(Key, RequiredGets, Context, FetchedBuffers))
+	{
+		OutStatus = EStatus::Error;
+		return FOptionalCacheRecord();
+	}
+
+	for (int32 Index = 0; Index < RequiredHeads.Num(); ++Index)
+	{
+		RecordBuilder.AddValue(RequiredHeads[Index].RemoveData());
+	}
+
+	for (int32 Index = 0; Index < RequiredGets.Num(); ++Index)
+	{
+		RecordBuilder.AddValue(FValueWithId(RequiredGets[Index].GetId(), FetchedBuffers[Index]));
+	}
+
+	return RecordBuilder.Build();
+}
+
+bool FHttpDerivedDataBackend::TryGetCachedDataBatch(
+	const FCacheKey& Key,
+	TArrayView<FValueWithId> Values,
+	const FStringView Context,
+	TArray<FCompressedBuffer>& OutBuffers)
+{
+	for (const FValueWithId& Value : Values)
+	{
+		const FIoHash& RawHash = Value.GetRawHash();
+		TStringBuilder<256> CompressedBlobsUri;
+		CompressedBlobsUri << "api/v1/compressed-blobs/" << StructuredNamespace << "/" << RawHash;
+
+		bool bHit = false;
+		FCompressedBuffer CompressedBuffer;
+		int64 ResponseCode = 0;
+		for (uint32 Attempts = 0; (Attempts < UE_HTTPDDC_MAX_ATTEMPTS) && !ShouldAbortForShutdown() && (Attempts == 0 || ShouldRetryOnError(ResponseCode)); ++Attempts)
+		{
+			FScopedRequestPtr Request(PutRequestPools[IsInGameThread()].Get());
+
+			TArray<uint8> ByteArray;
+			const FHttpRequest::Result Result = Request->PerformBlockingDownload(*CompressedBlobsUri, &ByteArray);
+			ResponseCode = Request->GetResponseCode();
+
+			if (FHttpRequest::IsSuccessResponse(ResponseCode))
+			{
+				CompressedBuffer = FCompressedBuffer::FromCompressed(MakeSharedBufferFromArray(MoveTemp(ByteArray)));
+				bHit = true;
+				break;
+			}
+		}
+
+		if (!bHit)
+		{
+			UE_LOG(LogDerivedDataCache, Verbose,
+				TEXT("%s: Cache miss with missing value %s with hash %s for %s from '%.*s'"),
+				*GetName(), *WriteToString<16>(Value.GetId()), *WriteToString<48>(RawHash), *WriteToString<96>(Key),
+				Context.Len(), Context.GetData());
+			return false;
+		}
+		else if (CompressedBuffer.GetRawHash() != RawHash)
+		{
+			UE_LOG(LogDerivedDataCache, Display,
+				TEXT("%s: Cache miss with corrupted value %s with hash %s for %s from '%.*s'"),
+				*GetName(), *WriteToString<16>(Value.GetId()), *WriteToString<48>(RawHash),
+				*WriteToString<96>(Key), Context.Len(), Context.GetData());
+			return false;
+		}
+		else
+		{
+			OutBuffers.Add(CompressedBuffer);
+		}
+	}
+	return true;
+}
+
+bool FHttpDerivedDataBackend::CachedDataProbablyExistsBatch(
+	const FCacheKey& Key,
+	TConstArrayView<FValueWithId> Values,
+	const FStringView Context)
+{
+	for (const FValueWithId& Value : Values)
+	{
+		const FIoHash& RawHash = Value.GetRawHash();
+		TStringBuilder<256> CompressedBlobsUri;
+		CompressedBlobsUri << "api/v1/compressed-blobs/" << StructuredNamespace << "/" << RawHash;
+
+		bool bIsHit = false;
+		int64 ResponseCode = 0;
+		for (uint32 Attempts = 0; (Attempts < UE_HTTPDDC_MAX_ATTEMPTS) && !ShouldAbortForShutdown() && (Attempts == 0 || ShouldRetryOnError(ResponseCode)); ++Attempts)
+		{
+			FScopedRequestPtr Request(PutRequestPools[IsInGameThread()].Get());
+
+			//TODO: This should be a HEAD operation, but those are returning 404 from this endpoint.  Issue has been reported and will change to HEAD when resolved.
+			//const FHttpRequest::Result Result = Request->PerformBlockingQuery<FHttpRequest::Head>(*CompressedBlobsUri);
+			TArray<uint8> ByteArray;
+			const FHttpRequest::Result Result = Request->PerformBlockingDownload(*CompressedBlobsUri, &ByteArray);
+			ResponseCode = Request->GetResponseCode();
+
+			if (FHttpRequest::IsSuccessResponse(ResponseCode) || ResponseCode == 400)
+			{
+				bIsHit = (Result == FHttpRequest::Success && FHttpRequest::IsSuccessResponse(ResponseCode));
+				break;
+			}
+		}
+
+		if (!bIsHit)
+		{
+			UE_LOG(LogDerivedDataCache, Verbose,
+				TEXT("%s: Cache miss with missing value %s with hash %s for %s from '%.*s'"),
+				*GetName(), *WriteToString<16>(Value.GetId()), *WriteToString<48>(RawHash), *WriteToString<96>(Key),
+				Context.Len(), Context.GetData());
+			return false;
+		}
+	}
+	return true;
+}
+
 bool FHttpDerivedDataBackend::CachedDataProbablyExists(const TCHAR* CacheKey)
 {
 	TRACE_CPUPROFILER_EVENT_SCOPE(HttpDDC_Exist);
@@ -2531,11 +3026,28 @@ void FHttpDerivedDataBackend::Put(
 	IRequestOwner& Owner,
 	FOnCachePutComplete&& OnComplete)
 {
-	if (OnComplete)
+	for (const FCachePutRequest& Request : Requests)
 	{
-		for (const FCachePutRequest& Request : Requests)
+		const FCacheRecord& Record = Request.Record;
+		COOK_STAT(auto Timer = UsageStats.TimePut());
+		uint64 BytesSent = 0;
+		if (PutCacheRecord(Record, Request.Name, Request.Policy, BytesSent))
+		{
+			UE_LOG(LogDerivedDataCache, Verbose, TEXT("%s: Cache put complete for %s from '%s'"),
+				*GetName(), *WriteToString<96>(Record.GetKey()), *Request.Name);
+			TRACE_COUNTER_ADD(HttpDDC_BytesSent, int64(BytesSent));
+			COOK_STAT(Timer.AddHit(BytesSent));
+			if (OnComplete)
+			{
+				OnComplete({Request.Name, Record.GetKey(), Request.UserData, EStatus::Ok});
+			}
+		}
+		else
 		{
-			OnComplete({Request.Name, Request.Record.GetKey(), Request.UserData, EStatus::Error});
+			if (OnComplete)
+			{
+				OnComplete({Request.Name, Record.GetKey(), Request.UserData, EStatus::Error});
+			}
 		}
 	}
 }
@@ -2545,11 +3057,27 @@ void FHttpDerivedDataBackend::Get(
 	IRequestOwner& Owner,
 	FOnCacheGetComplete&& OnComplete)
 {
-	if (OnComplete)
+	for (const FCacheGetRequest& Request : Requests)
 	{
-		for (const FCacheGetRequest& Request : Requests)
+		COOK_STAT(auto Timer = UsageStats.TimeGet());
+		EStatus Status = EStatus::Ok;
+		if (FOptionalCacheRecord Record = GetCacheRecord(Request.Key, Request.Name, Request.Policy, Status))
+		{
+			UE_LOG(LogDerivedDataCache, Verbose, TEXT("%s: Cache hit for %s from '%s'"),
+				*GetName(), *WriteToString<96>(Request.Key), *Request.Name);
+			TRACE_COUNTER_ADD(HttpDDC_BytesReceived, Private::GetCacheRecordCompressedSize(Record.Get()));
+			COOK_STAT(Timer.AddHit(Private::GetCacheRecordCompressedSize(Record.Get())));
+			if (OnComplete)
+			{
+				OnComplete({Request.Name, MoveTemp(Record).Get(), Request.UserData, Status});
+			}
+		}
+		else
 		{
-			OnComplete({Request.Name, FCacheRecordBuilder(Request.Key).Build(), Request.UserData, EStatus::Error});
+			if (OnComplete)
+			{
+				OnComplete({Request.Name, FCacheRecordBuilder(Request.Key).Build(), Request.UserData, Status});
+			}
 		}
 	}
 }
diff --git a/Engine/Source/Developer/DerivedDataCache/Private/HttpDerivedDataBackend.h b/Engine/Source/Developer/DerivedDataCache/Private/HttpDerivedDataBackend.h
--- a/Engine/Source/Developer/DerivedDataCache/Private/HttpDerivedDataBackend.h
+++ b/Engine/Source/Developer/DerivedDataCache/Private/HttpDerivedDataBackend.h
@@ -3,16 +3,22 @@
 #pragma once
 
 #include "CoreMinimal.h"
+#include "Compression/CompressedBuffer.h"
+#include "Containers/StringView.h"
 #include "DerivedDataBackendInterface.h"
+#include "DerivedDataCacheKey.h"
+#include "DerivedDataCacheRecord.h"
 #include "DerivedDataCacheUsageStats.h"
 
-// Macro for whether to enable the S3 backend. libcurl is not currently available on Mac.
+// Macro for whether to enable the Jupiter backend. libcurl is not currently available on Mac.
 #if !defined(WITH_HTTP_DDC_BACKEND)
 	#define WITH_HTTP_DDC_BACKEND 0
 #endif
 
 #if WITH_HTTP_DDC_BACKEND
 
+class FCbPackage;
+
 namespace UE::DerivedData::Backends
 {
 
@@ -30,15 +36,17 @@ public:
 	/**
 	 * Creates the backend, checks health status and attempts to acquire an access token.
 	 *
-	 * @param ServiceUrl	Base url to the service including schema.
-	 * @param Namespace		Namespace to use.
-	 * @param OAuthProvider	Url to OAuth provider, for example "https://myprovider.com/oauth2/v1/token".
-	 * @param OAuthClientId	OAuth client identifier.
-	 * @param OAuthData		OAuth form data to send to login service. Can either be the raw form data or a Windows network file address (starting with "\\").
+	 * @param ServiceUrl			Base url to the service including schema.
+	 * @param Namespace				Namespace to use.
+	 * @param StructuredNamespace	Namespace to use for structured cache operations.
+	 * @param OAuthProvider			Url to OAuth provider, for example "https://myprovider.com/oauth2/v1/token".
+	 * @param OAuthClientId			OAuth client identifier.
+	 * @param OAuthData				OAuth form data to send to login service. Can either be the raw form data or a Windows network file address (starting with "\\").
 	 */
 	FHttpDerivedDataBackend(
 		const TCHAR* ServiceUrl, 
 		const TCHAR* Namespace, 
+		const TCHAR* StructuredNamespace, 
 		const TCHAR* OAuthProvider, 
 		const TCHAR* OAuthClientId, 
 		const TCHAR* OAuthData,
@@ -93,9 +101,17 @@ public:
 		return AnyInstance;
 	}
 
+	const FString& GetDomain() const { return Domain; }
+	const FString& GetNamespace() const { return Namespace; }
+	const FString& GetStructuredNamespace() const { return StructuredNamespace; }
+	const FString& GetOAuthProvider() const { return OAuthProvider; }
+	const FString& GetOAuthClientId() const { return OAuthClientId; }
+	const FString& GetOAuthSecret() const { return OAuthSecret; }
+
 private:
 	FString Domain;
 	FString Namespace;
+	FString StructuredNamespace;
 	FString DefaultBucket;
 	FString OAuthProvider;
 	FString OAuthClientId;
@@ -105,6 +121,7 @@ private:
 	FBackendDebugOptions DebugOptions;
 	FCriticalSection MissedKeysCS;
 	TSet<FName> DebugMissedKeys;
+	TSet<FCacheKey> DebugMissedCacheKeys;
 	TUniquePtr<struct FRequestPool> GetRequestPools[2];
 	TUniquePtr<struct FRequestPool> PutRequestPools[2];
 	TUniquePtr<struct FHttpAccessToken> Access;
@@ -118,6 +135,29 @@ private:
 	bool AcquireAccessToken();
 	bool ShouldRetryOnError(int64 ResponseCode);
 	bool ShouldSimulateMiss(const TCHAR* InKey);
+	bool ShouldSimulateMiss(const FCacheKey& Key);
+
+	bool PutCacheRecord(const FCacheRecord& Record, FStringView Context, const FCacheRecordPolicy& Policy, uint64& OutWriteSize);
+	uint64 PutRef(const FCacheRecord& Record, const FCbPackage& Package, FStringView Bucket, bool bFinalize, TArray<FIoHash>& OutNeededBlobHashes, bool& bOutPutCompletedSuccessfully);
+
+	FOptionalCacheRecord GetCacheRecordOnly(
+		const FCacheKey& Key,
+		const FStringView Context,
+		const FCacheRecordPolicy& Policy);
+	FOptionalCacheRecord GetCacheRecord(
+		const FCacheKey& Key,
+		const FStringView Context,
+		const FCacheRecordPolicy& Policy,
+		EStatus& OutStatus);
+	bool TryGetCachedDataBatch(
+		const FCacheKey& Key,
+		TArrayView<FValueWithId> Values,
+		const FStringView Context,
+		TArray<FCompressedBuffer>& OutBuffers);
+	bool CachedDataProbablyExistsBatch(
+		const FCacheKey& Key,
+		TConstArrayView<FValueWithId> Values,
+		const FStringView Context);
 };
 
 } // UE::DerivedData::Backends
diff --git a/Engine/Source/Developer/DerivedDataCache/Private/Tests/HttpDerivedDataBackendTest.cpp b/Engine/Source/Developer/DerivedDataCache/Private/Tests/HttpDerivedDataBackendTest.cpp
--- a/Engine/Source/Developer/DerivedDataCache/Private/Tests/HttpDerivedDataBackendTest.cpp
+++ b/Engine/Source/Developer/DerivedDataCache/Private/Tests/HttpDerivedDataBackendTest.cpp
@@ -1,11 +1,20 @@
 // Copyright Epic Games, Inc. All Rights Reserved.
 
 #include "HttpDerivedDataBackend.h"
+#include "ZenDerivedDataBackend.h"
 
 #include "Async/ParallelFor.h"
 #include "DerivedDataBackendInterface.h"
+#include "DerivedDataCacheKey.h"
+#include "DerivedDataCacheRecord.h"
+#include "DerivedDataRequestOwner.h"
+#include "DerivedDataValue.h"
+#include "DerivedDataValueId.h"
+#include "Memory/CompositeBuffer.h"
 #include "Misc/AutomationTest.h"
+#include "Misc/Paths.h"
 #include "Misc/SecureHash.h"
+#include "Serialization/CompactBinaryWriter.h"
 
 // Test is targeted at HttpDerivedDataBackend but with some backend test interface it could be generalized
 // to function against all backends.
@@ -51,6 +60,7 @@ public:
 	}
 
 protected:
+
 	void ConcurrentTestWithStats(TFunctionRef<void()> TestFunction, int32 ThreadCount, double Duration)
 	{
 		std::atomic<uint64> Requests{ 0 };
@@ -144,6 +154,121 @@ protected:
 		return CachedBackend;
 	}
 
+	bool GetRecords(TConstArrayView<UE::DerivedData::FCacheRecord> Records, UE::DerivedData::FCacheRecordPolicy Policy, TArray<UE::DerivedData::FCacheRecord>& OutRecords)
+	{
+		using namespace UE::DerivedData;
+		UE::DerivedData::ICacheStore* TestBackend = GetTestBackend();
+
+		TArray<FCacheGetRequest> Requests;
+		Requests.Reserve(Records.Num());
+
+		for (int32 RecordIndex = 0; RecordIndex < Records.Num(); ++RecordIndex)
+		{
+			const FCacheRecord& Record = Records[RecordIndex];
+			Requests.Add({ {TEXT("FHttpDerivedDataTestBase")}, Record.GetKey(), Policy, static_cast<uint64>(RecordIndex) });
+		}
+
+		struct FGetOutput
+		{
+			FCacheRecord Record;
+			EStatus Status = EStatus::Error;
+		};
+
+		TArray<TOptional<FGetOutput>> GetOutputs;
+		GetOutputs.SetNum(Records.Num());
+		FRequestOwner RequestOwner(EPriority::Blocking);
+		TestBackend->Get(Requests, RequestOwner, [&GetOutputs](FCacheGetCompleteParams&& Params)
+			{
+				FCacheRecordBuilder RecordBuilder(Params.Record.GetKey());
+
+				if (Params.Record.GetMeta())
+				{
+					RecordBuilder.SetMeta(FCbObject::Clone(Params.Record.GetMeta()));
+				}
+
+				for (const FValueWithId& Value : Params.Record.GetValues())
+				{
+					if (Value)
+					{
+						RecordBuilder.AddValue(Value);
+					}
+				}
+				GetOutputs[Params.UserData].Emplace(FGetOutput{ RecordBuilder.Build(), Params.Status });
+			});
+		RequestOwner.Wait();
+
+		for (int32 RecordIndex = 0; RecordIndex < Records.Num(); ++RecordIndex)
+		{
+			FGetOutput& ReceivedOutput = GetOutputs[RecordIndex].GetValue();
+			if (ReceivedOutput.Status != EStatus::Ok)
+			{
+				return false;
+			}
+			OutRecords.Add(MoveTemp(ReceivedOutput.Record));
+		}
+
+		return true;
+	}
+	
+	void ValidateRecords(const TCHAR* Name, TConstArrayView<UE::DerivedData::FCacheRecord> RecordsToTest, TConstArrayView<UE::DerivedData::FCacheRecord> ReferenceRecords, UE::DerivedData::FCacheRecordPolicy Policy)
+	{
+		using namespace UE::DerivedData;
+
+		if (!TestEqual(FString::Printf(TEXT("%s::Record quantity"), Name), RecordsToTest.Num(), ReferenceRecords.Num()))
+		{
+			return;
+		}
+
+		for (int32 RecordIndex = 0; RecordIndex < RecordsToTest.Num(); ++RecordIndex)
+		{
+			const FCacheRecord& ExpectedRecord = ReferenceRecords[RecordIndex];
+			const FCacheRecord& RecordToTest = RecordsToTest[RecordIndex];
+
+			if (EnumHasAnyFlags(Policy.GetRecordPolicy(), ECachePolicy::SkipMeta))
+			{
+				TestTrue(FString::Printf(TEXT("%s::Get meta null"), Name), !RecordToTest.GetMeta());
+			}
+			else
+			{
+				TestTrue(FString::Printf(TEXT("%s::Get meta equality"), Name), ExpectedRecord.GetMeta().Equals(RecordToTest.GetMeta()));
+			}
+
+			TestEqual(FString::Printf(TEXT("%s::Get value quantity"), Name), ExpectedRecord.GetValues().Num(), RecordToTest.GetValues().Num());
+
+			const TConstArrayView<FValueWithId> ExpectedValues = ExpectedRecord.GetValues();
+			const TConstArrayView<FValueWithId> ReceivedValues = RecordToTest.GetValues();
+			for (int32 ValueIndex = 0; ValueIndex < ExpectedValues.Num(); ++ValueIndex)
+			{
+				if (EnumHasAnyFlags(Policy.GetRecordPolicy(), ECachePolicy::SkipData))
+				{
+					TestTrue(FString::Printf(TEXT("%s::Get value[%d] !HasData"), Name, ValueIndex), !ReceivedValues[ValueIndex].HasData());
+				}
+				else
+				{
+					TestTrue(FString::Printf(TEXT("%s::Get value[%d] HasData"), Name, ValueIndex), ReceivedValues[ValueIndex].HasData());
+					TestTrue(FString::Printf(TEXT("%s::Get value[%d] equality"), Name, ValueIndex), ExpectedValues[ValueIndex] == ReceivedValues[ValueIndex]);
+					TestTrue(FString::Printf(TEXT("%s::Get value[%d] data equality"), Name, ValueIndex), FIoHash::HashBuffer(ReceivedValues[ValueIndex].GetData().GetCompressed()) == FIoHash::HashBuffer(ExpectedValues[ValueIndex].GetData().GetCompressed()));
+				}
+			}
+		}
+	}
+
+	TArray<UE::DerivedData::FCacheRecord> GetAndValidateRecords(const TCHAR* Name, TConstArrayView<UE::DerivedData::FCacheRecord> Records, UE::DerivedData::FCacheRecordPolicy Policy)
+	{
+		using namespace UE::DerivedData;
+		TArray<FCacheRecord> ReceivedRecords;
+		bool bGetSuccessful = GetRecords(Records, Policy, ReceivedRecords);
+		TestTrue(FString::Printf(TEXT("%s::Get status"), Name), bGetSuccessful);
+
+		if (!bGetSuccessful)
+		{
+			return TArray<FCacheRecord>();
+		}
+
+		ValidateRecords(Name, ReceivedRecords, Records, Policy);
+		return ReceivedRecords;
+	}
+
 private:
 	UE::DerivedData::Backends::FHttpDerivedDataBackend* FetchTestBackend_Internal() const
 	{
@@ -174,6 +299,64 @@ TArray<FString> CreateTestCacheKeys(UE::DerivedData::Backends::FHttpDerivedDataB
 	return Keys;
 }
 
+TArray<UE::DerivedData::FCacheRecord> CreateTestCacheRecords(UE::DerivedData::ICacheStore* InTestBackend, uint32 InNumKeys, uint32 InNumValues, FCbObject MetaContents = FCbObject(), const char* BucketName = nullptr)
+{
+	using namespace UE::DerivedData;
+	FCacheBucket TestCacheBucket(BucketName ? BucketName : "AutoTestDummy");
+
+	TArray<FCacheRecord> CacheRecords;
+	TArray<FCachePutRequest> PutRequests;
+	PutRequests.Reserve(InNumKeys);
+
+	for (uint32 KeyIndex = 0; KeyIndex < InNumKeys; ++KeyIndex)
+	{
+		FIoHashBuilder HashBuilder;
+
+		TArray<FSharedBuffer> Values;
+		for (uint32 ValueIndex = 0; ValueIndex < InNumValues; ++ValueIndex)
+		{
+			TArray<uint8> ValueContents;
+			// Add N zeroed bytes where N corresponds to the value index.
+			ValueContents.AddUninitialized(ValueIndex+1);
+			for (uint32 ContentIndex = 0; ContentIndex < (ValueIndex+1); ++ContentIndex)
+			{
+				ValueContents[ContentIndex] = (uint8)KeyIndex;
+			}
+			Values.Emplace(MakeSharedBufferFromArray(MoveTemp(ValueContents)));
+			HashBuilder.Update(Values.Last().GetView());
+		}
+
+		FCacheKey Key;
+		Key.Bucket = TestCacheBucket;
+		Key.Hash = HashBuilder.Finalize();
+
+		FCacheRecordBuilder RecordBuilder(Key);
+
+		for (const FSharedBuffer& ValueBuffer : Values)
+		{
+			FIoHash ValueHash(FIoHash::HashBuffer(ValueBuffer));
+			RecordBuilder.AddValue(FValueId::FromHash(ValueHash), ValueBuffer);
+		}
+
+		if (MetaContents)
+		{
+			RecordBuilder.SetMeta(MoveTemp(MetaContents));
+		}
+
+		PutRequests.Add({ {TEXT("AutoTest")}, RecordBuilder.Build(), ECachePolicy::Default, KeyIndex });
+	}
+
+	FRequestOwner Owner(EPriority::Blocking);
+	InTestBackend->Put(PutRequests, Owner, [&CacheRecords, &PutRequests] (FCachePutCompleteParams&& Params)
+	{
+		check(Params.Status == EStatus::Ok);
+		CacheRecords.Add(PutRequests[Params.UserData].Record);
+	});
+	Owner.Wait();
+
+	return CacheRecords;
+}
+
 IMPLEMENT_HTTPDERIVEDDATA_AUTOMATION_TEST(FConcurrentCachedDataProbablyExistsBatch, TEXT(".FConcurrentCachedDataProbablyExistsBatch"), EAutomationTestFlags::EditorContext | EAutomationTestFlags::ProductFilter)
 bool FConcurrentCachedDataProbablyExistsBatch::RunTest(const FString& Parameters)
 {
@@ -246,5 +429,94 @@ bool FConcurrentExistsAndGetForSameKeyBatch::RunTest(const FString& Parameters)
 	return true;
 }
 
+// Tests basic functionality for structured cache operations
+IMPLEMENT_HTTPDERIVEDDATA_AUTOMATION_TEST(CacheStore, TEXT(".CacheStore"), EAutomationTestFlags::EditorContext | EAutomationTestFlags::ProductFilter)
+bool CacheStore::RunTest(const FString& Parameters)
+{
+	using namespace UE::DerivedData;
+	UE::DerivedData::Backends::FHttpDerivedDataBackend* TestBackend = GetTestBackend();
+
+#if WITH_ZEN_DDC_BACKEND
+	using namespace UE::Zen;
+	FServiceSettings ZenTestServiceSettings;
+	FServiceAutoLaunchSettings& ZenTestAutoLaunchSettings = ZenTestServiceSettings.SettingsVariant.Get<FServiceAutoLaunchSettings>();
+	ZenTestAutoLaunchSettings.DataPath = FPaths::ConvertRelativePathToFull(FPaths::Combine(FPaths::EngineSavedDir(), "ZenUnitTest"));
+	ZenTestAutoLaunchSettings.ExtraArgs = FString::Printf(TEXT("--http asio --upstream-jupiter-url \"%s\" --upstream-jupiter-oauth-url \"%s\" --upstream-jupiter-oauth-clientid \"%s\" --upstream-jupiter-oauth-clientsecret \"%s\" --upstream-jupiter-namespace-ddc \"%s\" --upstream-jupiter-namespace \"%s\""),
+		*TestBackend->GetDomain(),
+		*TestBackend->GetOAuthProvider(),
+		*TestBackend->GetOAuthClientId(),
+		*TestBackend->GetOAuthSecret(),
+		*TestBackend->GetNamespace(),
+		*TestBackend->GetStructuredNamespace()
+	);
+	ZenTestAutoLaunchSettings.DesiredPort = 13337; // Avoid the normal default port
+	ZenTestAutoLaunchSettings.bShowConsole = true;
+	ZenTestAutoLaunchSettings.bLimitProcessLifetime = true;
+
+	FScopeZenService ScopeZenService(MoveTemp(ZenTestServiceSettings));
+	TUniquePtr<Backends::FZenDerivedDataBackend> ZenIntermediaryBackend = MakeUnique<UE::DerivedData::Backends::FZenDerivedDataBackend>(ScopeZenService.GetInstance().GetURL(), *TestBackend->GetNamespace());
+	auto WaitForZenPushToUpstream = [](Backends::FZenDerivedDataBackend* ZenBackend, TConstArrayView<FCacheRecord> Records)
+	{
+		// TODO: Expecting a legitimate means to wait for zen to finish pushing records to its upstream in the future
+		FPlatformProcess::Sleep(1.0f);
+	};
+#endif // WITH_ZEN_DDC_BACKEND
+
+	const uint32 RecordsInBatch = 3;
+
+	{
+		TArray<FCacheRecord> PutRecords = CreateTestCacheRecords(TestBackend, RecordsInBatch, 1);
+		TArray<FCacheRecord> RecievedRecords = GetAndValidateRecords(TEXT("SimpleValue"), PutRecords, ECachePolicy::Default);
+		TArray<FCacheRecord> RecievedRecordsSkipMeta = GetAndValidateRecords(TEXT("SimpleValueSkipMeta"), PutRecords, ECachePolicy::Default | ECachePolicy::SkipMeta);
+		TArray<FCacheRecord> RecievedRecordsSkipData = GetAndValidateRecords(TEXT("SimpleValueSkipData"), PutRecords, ECachePolicy::Default | ECachePolicy::SkipData);
+
+#if WITH_ZEN_DDC_BACKEND
+		TArray<FCacheRecord> PutRecordsZen = CreateTestCacheRecords(ZenIntermediaryBackend.Get(), RecordsInBatch, 1, FCbObject(), "AutoTestDummyZen");
+		WaitForZenPushToUpstream(ZenIntermediaryBackend.Get(), PutRecordsZen);
+		ValidateRecords(TEXT("SimpleValueZenAndDirect"), GetAndValidateRecords(TEXT("SimpleValueZen"), PutRecordsZen, ECachePolicy::Default), RecievedRecords, ECachePolicy::Default);
+		ValidateRecords(TEXT("SimpleValueSkipMetaZenAndDirect"), GetAndValidateRecords(TEXT("SimpleValueSkipMetaZen"), PutRecordsZen, ECachePolicy::Default | ECachePolicy::SkipMeta), RecievedRecordsSkipMeta, ECachePolicy::Default | ECachePolicy::SkipMeta);
+		ValidateRecords(TEXT("SimpleValueSkipDataZenAndDirect"), GetAndValidateRecords(TEXT("SimpleValueSkipDataZen"), PutRecordsZen, ECachePolicy::Default | ECachePolicy::SkipData), RecievedRecordsSkipData, ECachePolicy::Default | ECachePolicy::SkipData);
+#endif // WITH_ZEN_DDC_BACKEND
+	}
+
+	{
+		TCbWriter<64> MetaWriter;
+		MetaWriter.BeginObject();
+		MetaWriter.AddInteger("MetaKey"_ASV, 42);
+		MetaWriter.EndObject();
+		FCbObject MetaObject = MetaWriter.Save().AsObject();
+
+		TArray<FCacheRecord> PutRecords = CreateTestCacheRecords(TestBackend, RecordsInBatch, 1, MetaObject);
+		TArray<FCacheRecord> RecievedRecords = GetAndValidateRecords(TEXT("SimpleValueWithMeta"), PutRecords, ECachePolicy::Default);
+		TArray<FCacheRecord> RecievedRecordsSkipMeta = GetAndValidateRecords(TEXT("SimpleValueWithMetaSkipMeta"), PutRecords, ECachePolicy::Default | ECachePolicy::SkipMeta);
+		TArray<FCacheRecord> RecievedRecordsSkipData = GetAndValidateRecords(TEXT("SimpleValueWithMetaSkipData"), PutRecords, ECachePolicy::Default | ECachePolicy::SkipData);
+
+#if WITH_ZEN_DDC_BACKEND
+		TArray<FCacheRecord> PutRecordsZen = CreateTestCacheRecords(ZenIntermediaryBackend.Get(), RecordsInBatch, 1, MetaObject, "AutoTestDummyZen");
+		WaitForZenPushToUpstream(ZenIntermediaryBackend.Get(), PutRecordsZen);
+		ValidateRecords(TEXT("SimpleValueWithMetaZenAndDirect"), GetAndValidateRecords(TEXT("SimpleValueWithMetaZen"), PutRecordsZen, ECachePolicy::Default), RecievedRecords, ECachePolicy::Default);
+		ValidateRecords(TEXT("SimpleValueWithMetaSkipMetaZenAndDirect"), GetAndValidateRecords(TEXT("SimpleValueWithMetaSkipMetaZen"), PutRecordsZen, ECachePolicy::Default | ECachePolicy::SkipMeta), RecievedRecordsSkipMeta, ECachePolicy::Default | ECachePolicy::SkipMeta);
+		ValidateRecords(TEXT("SimpleValueWithMetaSkipDataZenAndDirect"), GetAndValidateRecords(TEXT("SimpleValueWithMetaSkipDataZen"), PutRecordsZen, ECachePolicy::Default | ECachePolicy::SkipData), RecievedRecordsSkipData, ECachePolicy::Default | ECachePolicy::SkipData);
+#endif // WITH_ZEN_DDC_BACKEND
+	}
+
+	{
+		TArray<FCacheRecord> PutRecords = CreateTestCacheRecords(TestBackend, RecordsInBatch, 5);
+		TArray<FCacheRecord> RecievedRecords = GetAndValidateRecords(TEXT("MultiValue"), PutRecords, ECachePolicy::Default);
+		TArray<FCacheRecord> RecievedRecordsSkipMeta = GetAndValidateRecords(TEXT("MultiValueSkipMeta"), PutRecords, ECachePolicy::Default | ECachePolicy::SkipMeta);
+		TArray<FCacheRecord> RecievedRecordsSkipData = GetAndValidateRecords(TEXT("MultiValueSkipData"), PutRecords, ECachePolicy::Default | ECachePolicy::SkipData);
+
+#if WITH_ZEN_DDC_BACKEND
+		TArray<FCacheRecord> PutRecordsZen = CreateTestCacheRecords(ZenIntermediaryBackend.Get(), RecordsInBatch, 5, FCbObject(), "AutoTestDummyZen");
+		WaitForZenPushToUpstream(ZenIntermediaryBackend.Get(), PutRecordsZen);
+		ValidateRecords(TEXT("MultiValueZenAndDirect"), GetAndValidateRecords(TEXT("MultiValueZen"), PutRecordsZen, ECachePolicy::Default), RecievedRecords, ECachePolicy::Default);
+		ValidateRecords(TEXT("MultiValueSkipMetaZenAndDirect"), GetAndValidateRecords(TEXT("MultiValueSkipMetaZen"), PutRecordsZen, ECachePolicy::Default | ECachePolicy::SkipMeta), RecievedRecordsSkipMeta, ECachePolicy::Default | ECachePolicy::SkipMeta);
+		ValidateRecords(TEXT("MultiValueSkipDataZenAndDirect"), GetAndValidateRecords(TEXT("MultiValueSkipDataZen"), PutRecordsZen, ECachePolicy::Default | ECachePolicy::SkipData), RecievedRecordsSkipData, ECachePolicy::Default | ECachePolicy::SkipData);
+#endif // WITH_ZEN_DDC_BACKEND
+	}
+
+	return true;
+}
+
 }
 #endif // #if WITH_DEV_AUTOMATION_TESTS && WITH_HTTP_DDC_BACKEND
