DDC: Do not pass null completion callbacks from ICache into ICacheStore

#rnx

diff --git a/Engine/Source/Developer/DerivedDataCache/Private/DerivedDataCache.cpp b/Engine/Source/Developer/DerivedDataCache/Private/DerivedDataCache.cpp
--- a/Engine/Source/Developer/DerivedDataCache/Private/DerivedDataCache.cpp
+++ b/Engine/Source/Developer/DerivedDataCache/Private/DerivedDataCache.cpp
@@ -818,7 +818,7 @@ public:
 		IRequestOwner& Owner,
 		FOnCachePutComplete&& OnComplete) final
 	{
-		return FDerivedDataBackend::Get().GetRoot().Put(Requests, Owner, MoveTemp(OnComplete));
+		return FDerivedDataBackend::Get().GetRoot().Put(Requests, Owner, OnComplete ? MoveTemp(OnComplete) : [](auto&&){});
 	}
 
 	void Get(
@@ -826,7 +826,7 @@ public:
 		IRequestOwner& Owner,
 		FOnCacheGetComplete&& OnComplete) final
 	{
-		return FDerivedDataBackend::Get().GetRoot().Get(Requests, Owner, MoveTemp(OnComplete));
+		return FDerivedDataBackend::Get().GetRoot().Get(Requests, Owner, OnComplete ? MoveTemp(OnComplete) : [](auto&&){});
 	}
 
 	void PutValue(
@@ -834,7 +834,7 @@ public:
 		IRequestOwner& Owner,
 		FOnCachePutValueComplete&& OnComplete) final
 	{
-		return FDerivedDataBackend::Get().GetRoot().PutValue(Requests, Owner, MoveTemp(OnComplete));
+		return FDerivedDataBackend::Get().GetRoot().PutValue(Requests, Owner, OnComplete ? MoveTemp(OnComplete) : [](auto&&){});
 	}
 
 	void GetValue(
@@ -842,7 +842,7 @@ public:
 		IRequestOwner& Owner,
 		FOnCacheGetValueComplete&& OnComplete) final
 	{
-		return FDerivedDataBackend::Get().GetRoot().GetValue(Requests, Owner, MoveTemp(OnComplete));
+		return FDerivedDataBackend::Get().GetRoot().GetValue(Requests, Owner, OnComplete ? MoveTemp(OnComplete) : [](auto&&){});
 	}
 
 	void GetChunks(
@@ -850,7 +850,7 @@ public:
 		IRequestOwner& Owner,
 		FOnCacheChunkComplete&& OnComplete) final
 	{
-		return FDerivedDataBackend::Get().GetRoot().GetChunks(Requests, Owner, MoveTemp(OnComplete));
+		return FDerivedDataBackend::Get().GetRoot().GetChunks(Requests, Owner, OnComplete ? MoveTemp(OnComplete) : [](auto&&){});
 	}
 
 	// ICache Interface
