Switch to using HEAD when doing existance checks against the structured cache in Jupiter.  Was not possible before due to a bug in Jupiter, but that hs now been fixed.

diff --git a/Engine/Source/Developer/DerivedDataCache/Private/HttpDerivedDataBackend.cpp b/Engine/Source/Developer/DerivedDataCache/Private/HttpDerivedDataBackend.cpp
--- a/Engine/Source/Developer/DerivedDataCache/Private/HttpDerivedDataBackend.cpp
+++ b/Engine/Source/Developer/DerivedDataCache/Private/HttpDerivedDataBackend.cpp
@@ -2641,10 +2641,7 @@ bool FHttpDerivedDataBackend::CachedDataProbablyExistsBatch(
 		{
 			FScopedRequestPtr Request(PutRequestPools[IsInGameThread()].Get());
 
-			//TODO: This should be a HEAD operation, but those are returning 404 from this endpoint.  Issue has been reported and will change to HEAD when resolved.
-			//const FHttpRequest::Result Result = Request->PerformBlockingQuery<FHttpRequest::Head>(*CompressedBlobsUri);
-			TArray<uint8> ByteArray;
-			const FHttpRequest::Result Result = Request->PerformBlockingDownload(*CompressedBlobsUri, &ByteArray);
+			const FHttpRequest::Result Result = Request->PerformBlockingQuery<FHttpRequest::Head>(*CompressedBlobsUri);
 			ResponseCode = Request->GetResponseCode();
 
 			if (FHttpRequest::IsSuccessResponse(ResponseCode) || ResponseCode == 400)
