DDC: Moved cache backends out of headers

This hides implementation details and reduces the number of files that need to change when refactoring backends.

Backends moved in this change: AsyncPut, Http, Memory, S3, Zen.

#rnx

diff --git a/Engine/Source/Developer/DerivedDataCache/DerivedDataCache.Build.cs b/Engine/Source/Developer/DerivedDataCache/DerivedDataCache.Build.cs
--- a/Engine/Source/Developer/DerivedDataCache/DerivedDataCache.Build.cs
+++ b/Engine/Source/Developer/DerivedDataCache/DerivedDataCache.Build.cs
@@ -17,10 +17,7 @@ public class DerivedDataCache : ModuleRules
 		AddEngineThirdPartyPrivateStaticDependencies(Target, "OpenSSL");
 
 		// Platform-specific opt-in
-		if (Target.Platform == UnrealTargetPlatform.Win64)
-		{
-			PrivateDefinitions.Add("WITH_HTTP_DDC_BACKEND=1");
-			PrivateDefinitions.Add("WITH_S3_DDC_BACKEND=1");
-		}
+		PrivateDefinitions.Add($"WITH_HTTP_DDC_BACKEND={(Target.Platform == UnrealTargetPlatform.Win64 ? 1 : 0)}");
+		PrivateDefinitions.Add($"WITH_S3_DDC_BACKEND={(Target.Platform == UnrealTargetPlatform.Win64 ? 1 : 0)}");
 	}
 }
diff --git a/Engine/Source/Developer/DerivedDataCache/Private/DerivedDataBackendAsyncPutWrapper.cpp b/Engine/Source/Developer/DerivedDataCache/Private/DerivedDataBackendAsyncPutWrapper.cpp
--- a/Engine/Source/Developer/DerivedDataCache/Private/DerivedDataBackendAsyncPutWrapper.cpp
+++ b/Engine/Source/Developer/DerivedDataCache/Private/DerivedDataBackendAsyncPutWrapper.cpp
@@ -1,17 +1,186 @@
 // Copyright Epic Games, Inc. All Rights Reserved.
 
-#include "DerivedDataBackendAsyncPutWrapper.h"
-
+#include "Async/AsyncWork.h"
+#include "DerivedDataBackendInterface.h"
+#include "DerivedDataCacheUsageStats.h"
 #include "DerivedDataRequest.h"
 #include "DerivedDataRequestOwner.h"
 #include "DerivedDataValueId.h"
 #include "Experimental/Async/LazyEvent.h"
-#include "MemoryDerivedDataBackend.h"
+#include "FileBackedDerivedDataBackend.h"
 #include "Memory/SharedBuffer.h"
+#include "Misc/ScopeLock.h"
+#include "ProfilingDebugging/CookStats.h"
+#include "Stats/Stats.h"
 #include "Tasks/Task.h"
 
-namespace UE::DerivedData::Backends
+namespace UE::DerivedData::CacheStore::Memory
+{
+
+FFileBackedDerivedDataBackend* CreateMemoryDerivedDataBackend(const TCHAR* Name, int64 MaxCacheSize, bool bCanBeDisabled);
+
+} // UE::DerivedData::CacheStore::Memory
+
+namespace UE::DerivedData::CacheStore::AsyncPut
+{
+
+/** 
+ * Thread safe set helper
+**/
+struct FThreadSet
+{
+	FCriticalSection	SynchronizationObject;
+	TSet<FString>		FilesInFlight;
+
+	void Add(const FString& Key)
+	{
+		FScopeLock ScopeLock(&SynchronizationObject);
+		check(Key.Len());
+		FilesInFlight.Add(Key);
+	}
+	void Remove(const FString& Key)
+	{
+		FScopeLock ScopeLock(&SynchronizationObject);
+		FilesInFlight.Remove(Key);
+	}
+	bool Exists(const FString& Key)
+	{
+		FScopeLock ScopeLock(&SynchronizationObject);
+		return FilesInFlight.Contains(Key);
+	}
+	bool AddIfNotExists(const FString& Key)
+	{
+		FScopeLock ScopeLock(&SynchronizationObject);
+		check(Key.Len());
+		if (!FilesInFlight.Contains(Key))
+		{
+			FilesInFlight.Add(Key);
+			return true;
+		}
+		return false;
+	}
+};
+
+/** 
+ * A backend wrapper that coordinates async puts. This means that a Get will hit an in-memory cache while the async put is still in flight.
+**/
+class FDerivedDataBackendAsyncPutWrapper : public FDerivedDataBackendInterface
 {
+public:
+
+	/**
+	 * Constructor
+	 *
+	 * @param	InInnerBackend		Backend to use for storage, my responsibilities are about async puts
+	 * @param	bCacheInFlightPuts	if true, cache in-flight puts in a memory cache so that they hit immediately
+	 */
+	FDerivedDataBackendAsyncPutWrapper(FDerivedDataBackendInterface* InInnerBackend, bool bCacheInFlightPuts);
+
+	/** Return a name for this interface */
+	virtual FString GetName() const override
+	{
+		return FString::Printf(TEXT("AsyncPutWrapper (%s)"), *InnerBackend->GetName());
+	}
+
+	/** return true if this cache is writable **/
+	virtual bool IsWritable() const override;
+
+	/** Returns a class of speed for this interface **/
+	virtual ESpeedClass GetSpeedClass() const override;
+
+	/** Return true if hits on this cache should propagate to lower cache level. */
+	virtual bool BackfillLowerCacheLevels() const override;
+
+	/**
+	 * Synchronous test for the existence of a cache item
+	 *
+	 * @param	CacheKey	Alphanumeric+underscore key of this cache item
+	 * @return				true if the data probably will be found, this can't be guaranteed because of concurrency in the backends, corruption, etc
+	 */
+	virtual bool CachedDataProbablyExists(const TCHAR* CacheKey) override;
+
+	/**
+	 * Synchronous test for the existence of multiple cache items
+	 *
+	 * @param	CacheKeys	Alphanumeric+underscore key of the cache items
+	 * @return				A bit array with bits indicating whether the data for the corresponding key will probably be found
+	 */
+	virtual TBitArray<> CachedDataProbablyExistsBatch(TConstArrayView<FString> CacheKeys) override;
+
+	/**
+	 * Attempt to make sure the cached data will be available as optimally as possible.
+	 *
+	 * @param	CacheKeys	Alphanumeric+underscore keys of the cache items
+	 * @return				true if the data will probably be found in a fast backend on a future request.
+	 */
+	virtual bool TryToPrefetch(TConstArrayView<FString> CacheKeys) override;
+
+	/**
+	 * Allows the DDC backend to determine if it wants to cache the provided data. Reasons for returning false could be a slow connection,
+	 * a file size limit, etc.
+	 */
+	virtual bool WouldCache(const TCHAR* CacheKey, TArrayView<const uint8> InData) override;
+
+	/**
+	 * Synchronous retrieve of a cache item
+	 *
+	 * @param	CacheKey	Alphanumeric+underscore key of this cache item
+	 * @param	OutData		Buffer to receive the results, if any were found
+	 * @return				true if any data was found, and in this case OutData is non-empty
+	 */
+	virtual bool GetCachedData(const TCHAR* CacheKey, TArray<uint8>& OutData) override;
+
+	/**
+	 * Asynchronous, fire-and-forget placement of a cache item
+	 *
+	 * @param	CacheKey	Alphanumeric+underscore key of this cache item
+	 * @param	InData		Buffer containing the data to cache, can be destroyed after the call returns, immediately
+	 * @param	bPutEvenIfExists	If true, then do not attempt skip the put even if CachedDataProbablyExists returns true
+	 */
+	virtual EPutStatus PutCachedData(const TCHAR* CacheKey, TArrayView<const uint8> InData, bool bPutEvenIfExists) override;
+
+	virtual void RemoveCachedData(const TCHAR* CacheKey, bool bTransient) override;
+
+	virtual TSharedRef<FDerivedDataCacheStatsNode> GatherUsageStats() const override;
+
+	virtual bool ApplyDebugOptions(FBackendDebugOptions& InOptions) override;
+
+	virtual void Put(
+		TConstArrayView<FCachePutRequest> Requests,
+		IRequestOwner& Owner,
+		FOnCachePutComplete&& OnComplete) override;
+
+	virtual void Get(
+		TConstArrayView<FCacheGetRequest> Requests,
+		IRequestOwner& Owner,
+		FOnCacheGetComplete&& OnComplete) override;
+
+	virtual void PutValue(
+		TConstArrayView<FCachePutValueRequest> Requests,
+		IRequestOwner& Owner,
+		FOnCachePutValueComplete&& OnComplete) override;
+
+	virtual void GetValue(
+		TConstArrayView<FCacheGetValueRequest> Requests,
+		IRequestOwner& Owner,
+		FOnCacheGetValueComplete&& OnComplete) override;
+
+	virtual void GetChunks(
+		TConstArrayView<FCacheChunkRequest> Requests,
+		IRequestOwner& Owner,
+		FOnCacheChunkComplete&& OnComplete) override;
+
+private:
+	FDerivedDataCacheUsageStats UsageStats;
+	FDerivedDataCacheUsageStats PutSyncUsageStats;
+
+	/** Backend to use for storage, my responsibilities are about async puts **/
+	FDerivedDataBackendInterface*					InnerBackend;
+	/** Memory based cache to deal with gets that happen while an async put is still in flight **/
+	TUniquePtr<FDerivedDataBackendInterface>		InflightCache;
+	/** We remember outstanding puts so that we don't do them redundantly **/
+	FThreadSet										FilesInFlight;
+};
 
 /** 
  * Async task to handle the fire and forget async put
@@ -170,7 +339,7 @@ public:
 
 FDerivedDataBackendAsyncPutWrapper::FDerivedDataBackendAsyncPutWrapper(FDerivedDataBackendInterface* InInnerBackend, bool bCacheInFlightPuts)
 	: InnerBackend(InInnerBackend)
-	, InflightCache(bCacheInFlightPuts ? (new FMemoryDerivedDataBackend(TEXT("InflightMemoryCache"))) : NULL)
+	, InflightCache(bCacheInFlightPuts ? Memory::CreateMemoryDerivedDataBackend(TEXT("InflightMemoryCache"), /*MaxCacheSize*/ -1, /*bCanBeDisabled*/ false) : nullptr)
 {
 	check(InnerBackend);
 }
@@ -615,4 +784,9 @@ void FDerivedDataBackendAsyncPutWrapper::GetChunks(
 	}
 }
 
-} // UE::DerivedData::Backends
+FDerivedDataBackendInterface* CreateAsyncPutDerivedDataBackend(FDerivedDataBackendInterface* InnerBackend, bool bCacheInFlightPuts)
+{
+	return new FDerivedDataBackendAsyncPutWrapper(InnerBackend, bCacheInFlightPuts);
+}
+
+} // UE::DerivedData::CacheStore::AsyncPut
diff --git a/Engine/Source/Developer/DerivedDataCache/Private/DerivedDataBackendAsyncPutWrapper.h b/Engine/Source/Developer/DerivedDataCache/Private/DerivedDataBackendAsyncPutWrapper.h
deleted file mode 100644
diff --git a/Engine/Source/Developer/DerivedDataCache/Private/DerivedDataBackends.cpp b/Engine/Source/Developer/DerivedDataCache/Private/DerivedDataBackends.cpp
--- a/Engine/Source/Developer/DerivedDataCache/Private/DerivedDataBackends.cpp
+++ b/Engine/Source/Developer/DerivedDataCache/Private/DerivedDataBackends.cpp
@@ -1,35 +1,31 @@
 // Copyright Epic Games, Inc. All Rights Reserved.
 
-#include "CoreMinimal.h"
-#include "HAL/PlatformFileManager.h"
-#include "HAL/FileManager.h"
-#include "Misc/CommandLine.h"
-#include "Misc/Paths.h"
-#include "HAL/ThreadSafeCounter.h"
-#include "Misc/Guid.h"
-#include "HAL/IConsoleManager.h"
-#include "Misc/App.h"
+#include "CoreTypes.h"
+#include "Containers/StringView.h"
 #include "DerivedDataBackendInterface.h"
-#include "DerivedDataCacheUsageStats.h"
-#include "MemoryDerivedDataBackend.h"
-#include "HttpDerivedDataBackend.h"
-#include "DerivedDataBackendAsyncPutWrapper.h"
-#include "PakFileDerivedDataBackend.h"
-#include "S3DerivedDataBackend.h"
-#include "ZenDerivedDataBackend.h"
-#include "HierarchicalDerivedDataBackend.h"
-#include "DerivedDataLimitKeyLengthWrapper.h"
 #include "DerivedDataBackendThrottleWrapper.h"
 #include "DerivedDataBackendVerifyWrapper.h"
-#include "Misc/EngineBuildSettings.h"
-#include "Modules/ModuleManager.h"
+#include "DerivedDataCacheUsageStats.h"
+#include "DerivedDataLimitKeyLengthWrapper.h"
+#include "FileBackedDerivedDataBackend.h"
+#include "HAL/FileManager.h"
+#include "HAL/IConsoleManager.h"
+#include "HAL/PlatformFileManager.h"
+#include "HAL/ThreadSafeCounter.h"
+#include "HierarchicalDerivedDataBackend.h"
+#include "Internationalization/FastDecimalFormat.h"
+#include "Math/BasicMathExpressionEvaluator.h"
+#include "Math/UnitConversion.h"
+#include "Misc/App.h"
+#include "Misc/CommandLine.h"
 #include "Misc/ConfigCacheIni.h"
-#include "Containers/StringFwd.h"
+#include "Misc/EngineBuildSettings.h"
+#include "Misc/Guid.h"
+#include "Misc/Paths.h"
 #include "Misc/StringBuilder.h"
+#include "Modules/ModuleManager.h"
+#include "PakFileDerivedDataBackend.h"
 #include "ProfilingDebugging/CookStats.h"
-#include "Math/UnitConversion.h"
-#include "Internationalization/FastDecimalFormat.h"
-#include "Math/BasicMathExpressionEvaluator.h"
 #include "Serialization/CompactBinaryPackage.h"
 #include <atomic>
 
@@ -38,12 +34,44 @@ DEFINE_LOG_CATEGORY(LogDerivedDataCache);
 #define MAX_BACKEND_KEY_LENGTH (120)
 #define LOCTEXT_NAMESPACE "DerivedDataBackendGraph"
 
+namespace UE::DerivedData::CacheStore::AsyncPut
+{
+FDerivedDataBackendInterface* CreateAsyncPutDerivedDataBackend(FDerivedDataBackendInterface* InnerBackend, bool bCacheInFlightPuts);
+} // UE::DerivedData::CacheStore::AsyncPut
+
 namespace UE::DerivedData::CacheStore::FileSystem
 {
+FDerivedDataBackendInterface* CreateFileSystemDerivedDataBackend(const TCHAR* CacheDirectory, const TCHAR* Params, const TCHAR* AccessLogFileName);
+} // UE::DerivedData::CacheStore::FileSystem
 
-FDerivedDataBackendInterface* CreateFileSystemDerivedDataBackend(const TCHAR* CacheDirectory, const TCHAR* InParams, const TCHAR* InAccessLogFileName = nullptr);
+namespace UE::DerivedData::CacheStore::Http
+{
+FDerivedDataBackendInterface* CreateHttpDerivedDataBackend(
+	const TCHAR* NodeName,
+	const TCHAR* ServiceUrl,
+	const TCHAR* Namespace,
+	const TCHAR* StructuredNamespace,
+	const TCHAR* OAuthProvider,
+	const TCHAR* OAuthClientId,
+	const TCHAR* OAuthData, 
+	const FDerivedDataBackendInterface::ESpeedClass* ForceSpeedClass,
+	bool bReadOnly);
+} // UE::DerivedData::CacheStore::Http
+
+namespace UE::DerivedData::CacheStore::Memory
+{
+FFileBackedDerivedDataBackend* CreateMemoryDerivedDataBackend(const TCHAR* Name, int64 MaxCacheSize, bool bCanBeDisabled);
+} // UE::DerivedData::CacheStore::Memory
 
-} // UE::DerivedData::CacheStore::FileSystem
+namespace UE::DerivedData::CacheStore::S3
+{
+FDerivedDataBackendInterface* CreateS3DerivedDataBackend(const TCHAR* RootManifestPath, const TCHAR* BaseUrl, const TCHAR* Region, const TCHAR* CanaryObjectKey, const TCHAR* CachePath);
+} // UE::DerivedData::CacheStore::S3
+
+namespace UE::DerivedData::CacheStore::ZenCache
+{
+FDerivedDataBackendInterface* CreateZenDerivedDataBackend(const TCHAR* NodeName, const TCHAR* ServiceUrl, const TCHAR* Namespace);
+} // UE::DerivedData::CacheStore::ZenCache
 
 namespace UE::DerivedData::Backends
 {
@@ -124,7 +152,7 @@ public:
 		// Make sure AsyncPutWrapper and KeyLengthWrapper are created
 		if( !AsyncPutWrapper )
 		{
-			AsyncPutWrapper = new FDerivedDataBackendAsyncPutWrapper( RootCache, true );
+			AsyncPutWrapper = CacheStore::AsyncPut::CreateAsyncPutDerivedDataBackend( RootCache, /*bCacheInFlightPuts*/ true );
 			check(AsyncPutWrapper);
 			CreatedBackends.Add( AsyncPutWrapper );
 			RootCache = AsyncPutWrapper;
@@ -410,7 +438,7 @@ public:
 		FDerivedDataBackendInterface* AsyncNode = NULL;
 		if( InnerNode )
 		{
-			AsyncNode = new FDerivedDataBackendAsyncPutWrapper( InnerNode, true );
+			AsyncNode = CacheStore::AsyncPut::CreateAsyncPutDerivedDataBackend( InnerNode, /*bCacheInFlightPuts*/ true );
 		}
 		else
 		{
@@ -674,7 +702,6 @@ public:
 	 */
 	FDerivedDataBackendInterface* ParseS3Cache(const TCHAR* NodeName, const TCHAR* Entry)
 	{
-#if WITH_S3_DDC_BACKEND
 		FString ManifestPath;
 		if (!FParse::Value(Entry, TEXT("Manifest="), ManifestPath))
 		{
@@ -733,11 +760,13 @@ public:
 			}
 		}
 
-		return new CacheStore::S3::FS3DerivedDataBackend(*ManifestPath, *BaseUrl, *Region, *CanaryObjectKey, *CachePath);
-#else
+		if (FDerivedDataBackendInterface* Backend = CacheStore::S3::CreateS3DerivedDataBackend(*ManifestPath, *BaseUrl, *Region, *CanaryObjectKey, *CachePath))
+		{
+			return Backend;
+		}
+
 		UE_LOG(LogDerivedDataCache, Log, TEXT("S3 backend is not supported on the current platform."));
 		return nullptr;
-#endif
 	}
 
 	void ParseHttpCacheParams(
@@ -807,7 +836,6 @@ public:
 		const FString& IniFilename,
 		const TCHAR* IniSection)
 	{
-#if WITH_HTTP_DDC_BACKEND
 		FString Host;
 		FString Namespace;
 		FString StructuredNamespace;
@@ -885,25 +913,14 @@ public:
 			}
 		}
 
-		FHttpDerivedDataBackend* Backend = new FHttpDerivedDataBackend(*Host, *Namespace, *StructuredNamespace, *OAuthProvider, *OAuthClientId, *OAuthSecret, bReadOnly);
-
 		if (ForceSpeedClass != FDerivedDataBackendInterface::ESpeedClass::Unknown)
 		{
 			UE_LOG(LogDerivedDataCache, Log, TEXT("Node %s found speed class override ForceSpeedClass=%s"), NodeName, *ForceSpeedClassValue);
-			Backend->SetSpeedClass(ForceSpeedClass);
 		}
 
-		if (!Backend->IsUsable())
-		{
-			UE_LOG(LogDerivedDataCache, Warning, TEXT("Node %s could not contact the service (%s), will not use it"), NodeName, *Host);
-			delete Backend;
-			return nullptr;
-		}
-		return Backend;
-#else
-		UE_LOG(LogDerivedDataCache, Warning, TEXT("HTTP backend is not yet supported in the current build configuration."));
-		return nullptr;
-#endif
+		return CacheStore::Http::CreateHttpDerivedDataBackend(
+			NodeName, *Host, *Namespace, *StructuredNamespace, *OAuthProvider, *OAuthClientId, *OAuthSecret,
+			ForceSpeedClass == FDerivedDataBackendInterface::ESpeedClass::Unknown ? nullptr : &ForceSpeedClass, bReadOnly);
 	}
 
 	/**
@@ -911,7 +928,6 @@ public:
 	 */
 	FDerivedDataBackendInterface* ParseZenCache(const TCHAR* NodeName, const TCHAR* Entry)
 	{
-#if WITH_ZEN_DDC_BACKEND
 		FString ServiceUrl;
 		FParse::Value(Entry, TEXT("Host="), ServiceUrl);
 
@@ -922,18 +938,13 @@ public:
 			UE_LOG(LogDerivedDataCache, Warning, TEXT("Node %s does not specify 'Namespace', falling back to '%s'"), NodeName, *Namespace);
 		}
 
-		FZenDerivedDataBackend* backend = new FZenDerivedDataBackend(*ServiceUrl, *Namespace);
-		if (!backend->IsUsable())
+		if (FDerivedDataBackendInterface* Backend = CacheStore::ZenCache::CreateZenDerivedDataBackend(NodeName, *ServiceUrl, *Namespace))
 		{
-			UE_LOG(LogDerivedDataCache, Warning, TEXT("%s could not contact the service (%s), will not use it."), NodeName, *backend->GetName());
-			delete backend;
-			return nullptr;
+			return Backend;
 		}
-		return backend;
-#else
+		
 		UE_LOG(LogDerivedDataCache, Warning, TEXT("Zen backend is not yet supported in the current build configuration."));
 		return nullptr;
-#endif
 	}
 
 	/**
@@ -946,7 +957,7 @@ public:
 	 */
 	FFileBackedDerivedDataBackend* ParseBootCache( const TCHAR* NodeName, const TCHAR* Entry, FString& OutFilename )
 	{
-		FMemoryDerivedDataBackend* Cache = NULL;
+		FFileBackedDerivedDataBackend* Cache = nullptr;
 
 		// Only allow boot cache with the editor. We don't want other tools and utilities (eg. SCW) writing to the same file.
 #if WITH_EDITOR
@@ -966,7 +977,7 @@ public:
 			MaxCacheSize = FMath::Min(MaxCacheSize, MaxSupportedCacheSize);
 
 			UE_LOG( LogDerivedDataCache, Display, TEXT("Max Cache Size: %d MB"), MaxCacheSize);
-			Cache = new FMemoryDerivedDataBackend(TEXT("Boot"), MaxCacheSize * 1024 * 1024, true /* bCanBeDisabled */);
+			Cache = CacheStore::Memory::CreateMemoryDerivedDataBackend(TEXT("Boot"), MaxCacheSize * 1024 * 1024, /*bCanBeDisabled*/ true);
 
 			if( Cache && Filename.Len() )
 			{
@@ -1002,13 +1013,13 @@ public:
 	 * @param Entry Node definition.
 	 * @return Memory data cache backend interface instance or NULL if unsuccessfull
 	 */
-	FMemoryDerivedDataBackend* ParseMemoryCache( const TCHAR* NodeName, const TCHAR* Entry )
+	FFileBackedDerivedDataBackend* ParseMemoryCache( const TCHAR* NodeName, const TCHAR* Entry )
 	{
-		FMemoryDerivedDataBackend* Cache = NULL;
+		FFileBackedDerivedDataBackend* Cache = nullptr;
 		FString Filename;
 
 		FParse::Value( Entry, TEXT("Filename="), Filename );
-		Cache = new FMemoryDerivedDataBackend(NodeName);
+		Cache = CacheStore::Memory::CreateMemoryDerivedDataBackend(NodeName, /*MaxCacheSize*/ -1, /*bCanBeDisabled*/ false);
 		if( Cache && Filename.Len() )
 		{
 			if( Cache->LoadCache( *Filename ) )
diff --git a/Engine/Source/Developer/DerivedDataCache/Private/DerivedDataCacheStorePrivate.h b/Engine/Source/Developer/DerivedDataCache/Private/DerivedDataCacheStorePrivate.h
new file mode 100644
--- /dev/null
+++ b/Engine/Source/Developer/DerivedDataCache/Private/DerivedDataCacheStorePrivate.h
@@ -0,0 +1,10 @@
+// Copyright Epic Games, Inc. All Rights Reserved.
+
+#pragma once
+
+#include "DerivedDataCache.h"
+
+namespace UE::DerivedData::Private
+{
+
+} // UE::DerivedData::Private
diff --git a/Engine/Source/Developer/DerivedDataCache/Private/HierarchicalDerivedDataBackend.h b/Engine/Source/Developer/DerivedDataCache/Private/HierarchicalDerivedDataBackend.h
--- a/Engine/Source/Developer/DerivedDataCache/Private/HierarchicalDerivedDataBackend.h
+++ b/Engine/Source/Developer/DerivedDataCache/Private/HierarchicalDerivedDataBackend.h
@@ -5,7 +5,6 @@
 #include "CoreTypes.h"
 #include "Algo/BinarySearch.h"
 #include "Containers/ArrayView.h"
-#include "DerivedDataBackendAsyncPutWrapper.h"
 #include "DerivedDataBackendInterface.h"
 #include "DerivedDataCacheRecord.h"
 #include "DerivedDataCacheUsageStats.h"
@@ -20,6 +19,13 @@
 
 extern bool GVerifyDDC;
 
+namespace UE::DerivedData::CacheStore::AsyncPut
+{
+
+FDerivedDataBackendInterface* CreateAsyncPutDerivedDataBackend(FDerivedDataBackendInterface* InnerBackend, bool bCacheInFlightPuts);
+
+} // UE::DerivedData::CacheStore::AsyncPut
+
 namespace UE::DerivedData::Backends
 {
 
@@ -93,7 +99,7 @@ private:
 			// async puts to allow us to fill all levels without holding up the engine
 			// we need to cache inflight puts to avoid having inconsistent miss and redownload on lower cache levels while puts are still async
 			const bool bCacheInFlightPuts = true;
-			AsyncPutInnerBackends.Emplace(new FDerivedDataBackendAsyncPutWrapper(InnerBackends[CacheIndex], bCacheInFlightPuts));
+			AsyncPutInnerBackends.Emplace(CacheStore::AsyncPut::CreateAsyncPutDerivedDataBackend(InnerBackends[CacheIndex], bCacheInFlightPuts));
 		}
 	}
 
diff --git a/Engine/Source/Developer/DerivedDataCache/Private/HttpDerivedDataBackend.cpp b/Engine/Source/Developer/DerivedDataCache/Private/HttpDerivedDataBackend.cpp
--- a/Engine/Source/Developer/DerivedDataCache/Private/HttpDerivedDataBackend.cpp
+++ b/Engine/Source/Developer/DerivedDataCache/Private/HttpDerivedDataBackend.cpp
@@ -1,5 +1,6 @@
 // Copyright Epic Games, Inc. All Rights Reserved.
-#include "HttpDerivedDataBackend.h"
+
+#include "DerivedDataBackendInterface.h"
 
 #if WITH_HTTP_DDC_BACKEND
 
@@ -12,14 +13,19 @@
 #include "Windows/HideWindowsPlatformTypes.h"
 #endif
 #include "Algo/Accumulate.h"
-#include "Algo/Transform.h"
 #include "Algo/Find.h"
+#include "Algo/Transform.h"
+#include "Compression/CompressedBuffer.h"
 #include "Containers/StaticArray.h"
+#include "Containers/StringView.h"
 #include "Containers/Ticker.h"
+#include "DerivedDataCacheKey.h"
 #include "DerivedDataCachePrivate.h"
 #include "DerivedDataCacheRecord.h"
+#include "DerivedDataCacheUsageStats.h"
 #include "DerivedDataValue.h"
 #include "Dom/JsonObject.h"
+#include "Experimental/Containers/FAAArrayQueue.h"
 #include "GenericPlatform/GenericPlatformFile.h"
 #include "HAL/PlatformFileManager.h"
 #include "IO/IoHash.h"
@@ -27,18 +33,17 @@
 #include "Misc/FileHelper.h"
 #include "Misc/ScopeLock.h"
 #include "Misc/SecureHash.h"
-#include "Experimental/Containers/FAAArrayQueue.h"
 #include "Misc/StringBuilder.h"
-#include "ProfilingDebugging/CpuProfilerTrace.h"
+#include "Policies/CondensedJsonPrintPolicy.h"
 #include "ProfilingDebugging/CountersTrace.h"
+#include "ProfilingDebugging/CpuProfilerTrace.h"
 #include "Serialization/BufferArchive.h"
 #include "Serialization/CompactBinary.h"
 #include "Serialization/CompactBinaryPackage.h"
 #include "Serialization/CompactBinaryValidation.h"
 #include "Serialization/JsonReader.h"
-#include "Serialization/JsonWriter.h"
 #include "Serialization/JsonSerializer.h"
-#include "Policies/CondensedJsonPrintPolicy.h"
+#include "Serialization/JsonWriter.h"
 
 #if WITH_SSL
 #include "Ssl.h"
@@ -73,7 +78,7 @@
 #define UE_HTTPDDC_BATCH_HEAD_WEIGHT 1
 #define UE_HTTPDDC_BATCH_WEIGHT_HINT 12
 
-namespace UE::DerivedData::Backends
+namespace UE::DerivedData::CacheStore::Http
 {
 
 TRACE_DECLARE_INT_COUNTER(HttpDDC_Exist, TEXT("HttpDDC Exist"));
@@ -87,6 +92,9 @@ TRACE_DECLARE_INT_COUNTER(HttpDDC_BytesSent, TEXT("HttpDDC Bytes Sent"));
 
 static CURLcode sslctx_function(CURL * curl, void * sslctx, void * parm);
 
+typedef TSharedPtr<class IHttpRequest> FHttpRequestPtr;
+typedef TSharedPtr<class IHttpResponse, ESPMode::ThreadSafe> FHttpResponsePtr;
+
 /**
  * Encapsulation for access token shared by all requests.
  */
@@ -1956,6 +1964,140 @@ uint32 FHttpAccessToken::GetSerial() const
 // FHttpDerivedDataBackend
 //----------------------------------------------------------------------------------------------------------
 
+/**
+ * Backend for a HTTP based caching service (Jupiter).
+ **/
+class FHttpDerivedDataBackend : public FDerivedDataBackendInterface
+{
+public:
+	
+	/**
+	 * Creates the backend, checks health status and attempts to acquire an access token.
+	 *
+	 * @param ServiceUrl			Base url to the service including schema.
+	 * @param Namespace				Namespace to use.
+	 * @param StructuredNamespace	Namespace to use for structured cache operations.
+	 * @param OAuthProvider			Url to OAuth provider, for example "https://myprovider.com/oauth2/v1/token".
+	 * @param OAuthClientId			OAuth client identifier.
+	 * @param OAuthData				OAuth form data to send to login service. Can either be the raw form data or a Windows network file address (starting with "\\").
+	 */
+	FHttpDerivedDataBackend(
+		const TCHAR* ServiceUrl, 
+		const TCHAR* Namespace, 
+		const TCHAR* StructuredNamespace, 
+		const TCHAR* OAuthProvider, 
+		const TCHAR* OAuthClientId, 
+		const TCHAR* OAuthData,
+		bool bReadOnly);
+
+	~FHttpDerivedDataBackend();
+
+	/**
+	 * Checks is backend is usable (reachable and accessible).
+	 * @return true if usable
+	 */
+	bool IsUsable() const { return bIsUsable; }
+
+	/** return true if this cache is writable **/
+	virtual bool IsWritable() const override
+	{
+		return !bReadOnly && bIsUsable;
+	}
+
+	virtual bool CachedDataProbablyExists(const TCHAR* CacheKey) override;
+	virtual TBitArray<> CachedDataProbablyExistsBatch(TConstArrayView<FString> CacheKeys) override;
+	virtual bool GetCachedData(const TCHAR* CacheKey, TArray<uint8>& OutData) override;
+	virtual EPutStatus PutCachedData(const TCHAR* CacheKey, TArrayView<const uint8> InData, bool bPutEvenIfExists) override;
+	virtual void RemoveCachedData(const TCHAR* CacheKey, bool bTransient) override;
+	virtual TSharedRef<FDerivedDataCacheStatsNode> GatherUsageStats() const override;
+	
+	virtual FString GetName() const override;
+	virtual bool TryToPrefetch(TConstArrayView<FString> CacheKeys) override;
+	virtual bool WouldCache(const TCHAR* CacheKey, TArrayView<const uint8> InData) override;
+	virtual ESpeedClass GetSpeedClass() const override;
+	virtual bool ApplyDebugOptions(FBackendDebugOptions& InOptions) override;
+
+	void SetSpeedClass(ESpeedClass InSpeedClass) { SpeedClass = InSpeedClass; }
+
+	virtual void Put(
+		TConstArrayView<FCachePutRequest> Requests,
+		IRequestOwner& Owner,
+		FOnCachePutComplete&& OnComplete) override;
+
+	virtual void Get(
+		TConstArrayView<FCacheGetRequest> Requests,
+		IRequestOwner& Owner,
+		FOnCacheGetComplete&& OnComplete) override;
+
+	virtual void GetChunks(
+		TConstArrayView<FCacheChunkRequest> Requests,
+		IRequestOwner& Owner,
+		FOnCacheChunkComplete&& OnComplete) override;
+
+	static FHttpDerivedDataBackend* GetAny()
+	{
+		return AnyInstance;
+	}
+
+	const FString& GetDomain() const { return Domain; }
+	const FString& GetNamespace() const { return Namespace; }
+	const FString& GetStructuredNamespace() const { return StructuredNamespace; }
+	const FString& GetOAuthProvider() const { return OAuthProvider; }
+	const FString& GetOAuthClientId() const { return OAuthClientId; }
+	const FString& GetOAuthSecret() const { return OAuthSecret; }
+
+private:
+	FString Domain;
+	FString Namespace;
+	FString StructuredNamespace;
+	FString DefaultBucket;
+	FString OAuthProvider;
+	FString OAuthClientId;
+	FString OAuthSecret;
+	FCriticalSection AccessCs;
+	FDerivedDataCacheUsageStats UsageStats;
+	FBackendDebugOptions DebugOptions;
+	FCriticalSection MissedKeysCS;
+	TSet<FName> DebugMissedKeys;
+	TSet<FCacheKey> DebugMissedCacheKeys;
+	TUniquePtr<struct FRequestPool> GetRequestPools[2];
+	TUniquePtr<struct FRequestPool> PutRequestPools[2];
+	TUniquePtr<struct FHttpAccessToken> Access;
+	bool bIsUsable;
+	bool bReadOnly;
+	uint32 FailedLoginAttempts;
+	ESpeedClass SpeedClass;
+	static inline FHttpDerivedDataBackend* AnyInstance = nullptr;
+
+	bool IsServiceReady();
+	bool AcquireAccessToken();
+	bool ShouldRetryOnError(int64 ResponseCode);
+	bool ShouldSimulateMiss(const TCHAR* InKey);
+	bool ShouldSimulateMiss(const FCacheKey& Key);
+
+	bool PutCacheRecord(const FCacheRecord& Record, FStringView Context, const FCacheRecordPolicy& Policy, uint64& OutWriteSize);
+	uint64 PutRef(const FCacheRecord& Record, const FCbPackage& Package, FStringView Bucket, bool bFinalize, TArray<FIoHash>& OutNeededBlobHashes, bool& bOutPutCompletedSuccessfully);
+
+	FOptionalCacheRecord GetCacheRecordOnly(
+		const FCacheKey& Key,
+		const FStringView Context,
+		const FCacheRecordPolicy& Policy);
+	FOptionalCacheRecord GetCacheRecord(
+		const FCacheKey& Key,
+		const FStringView Context,
+		const FCacheRecordPolicy& Policy,
+		EStatus& OutStatus);
+	bool TryGetCachedDataBatch(
+		const FCacheKey& Key,
+		TArrayView<FValueWithId> Values,
+		const FStringView Context,
+		TArray<FCompressedBuffer>& OutBuffers);
+	bool CachedDataProbablyExistsBatch(
+		const FCacheKey& Key,
+		TConstArrayView<FValueWithId> Values,
+		const FStringView Context);
+};
+
 FHttpDerivedDataBackend::FHttpDerivedDataBackend(
 	const TCHAR* InServiceUrl, 
 	const TCHAR* InNamespace, 
@@ -3093,6 +3235,52 @@ void FHttpDerivedDataBackend::GetChunks(
 	}
 }
 
-} // UE::DerivedData::Backends
+} // UE::DerivedData::CacheStore::Http
+
+#endif // WITH_HTTP_DDC_BACKEND
+
+namespace UE::DerivedData::CacheStore::Http
+{
+
+FDerivedDataBackendInterface* CreateHttpDerivedDataBackend(
+	const TCHAR* NodeName,
+	const TCHAR* ServiceUrl,
+	const TCHAR* Namespace,
+	const TCHAR* StructuredNamespace,
+	const TCHAR* OAuthProvider,
+	const TCHAR* OAuthClientId,
+	const TCHAR* OAuthData,
+	const FDerivedDataBackendInterface::ESpeedClass* ForceSpeedClass,
+	bool bReadOnly)
+{
+#if WITH_HTTP_DDC_BACKEND
+	FHttpDerivedDataBackend* Backend = new FHttpDerivedDataBackend(ServiceUrl, Namespace, StructuredNamespace, OAuthProvider, OAuthClientId, OAuthData, bReadOnly);
+	if (Backend->IsUsable())
+	{
+		return Backend;
+	}
+	UE_LOG(LogDerivedDataCache, Warning, TEXT("Node %s could not contact the service (%s), will not use it"), NodeName, *ServiceUrl);
+	delete Backend;
+	return nullptr;
+#else
+	UE_LOG(LogDerivedDataCache, Warning, TEXT("HTTP backend is not yet supported in the current build configuration."));
+	return nullptr;
+#endif
+}
+
+FDerivedDataBackendInterface* GetAnyHttpDerivedDataBackend(
+	FString& OutDomain,
+	FString& OutOAuthProvider,
+	FString& OutOAuthClientId,
+	FString& OutOAuthSecret,
+	FString& OutNamespace,
+	FString& OutStructuredNamespace)
+{
+#if WITH_HTTP_DDC_BACKEND
+	return FHttpDerivedDataBackend::GetAny();
+#else
+	return nullptr;
+#endif
+}
 
-#endif //WITH_HTTP_DDC_BACKEND
\ No newline at end of file
+} // UE::DerivedData::CacheStore::Http
diff --git a/Engine/Source/Developer/DerivedDataCache/Private/HttpDerivedDataBackend.h b/Engine/Source/Developer/DerivedDataCache/Private/HttpDerivedDataBackend.h
deleted file mode 100644
diff --git a/Engine/Source/Developer/DerivedDataCache/Private/MemoryDerivedDataBackend.cpp b/Engine/Source/Developer/DerivedDataCache/Private/MemoryDerivedDataBackend.cpp
--- a/Engine/Source/Developer/DerivedDataCache/Private/MemoryDerivedDataBackend.cpp
+++ b/Engine/Source/Developer/DerivedDataCache/Private/MemoryDerivedDataBackend.cpp
@@ -1,17 +1,210 @@
 // Copyright Epic Games, Inc. All Rights Reserved.
 
-#include "MemoryDerivedDataBackend.h"
-
 #include "Algo/AllOf.h"
 #include "DerivedDataCachePrivate.h"
+#include "DerivedDataCacheRecord.h"
+#include "DerivedDataCacheUsageStats.h"
 #include "DerivedDataValue.h"
+#include "FileBackedDerivedDataBackend.h"
+#include "HAL/FileManager.h"
 #include "Misc/ScopeExit.h"
+#include "Misc/ScopeLock.h"
+#include "Misc/ScopeRWLock.h"
+#include "ProfilingDebugging/CookStats.h"
 #include "Serialization/CompactBinary.h"
 #include "Templates/UniquePtr.h"
-#include "Misc/ScopeRWLock.h"
 
-namespace UE::DerivedData::Backends
+namespace UE::DerivedData::CacheStore::Memory
+{
+
+/**
+ * A simple thread safe, memory based backend. This is used for Async puts and the boot cache.
+ */
+class FMemoryDerivedDataBackend : public FFileBackedDerivedDataBackend
 {
+public:
+	explicit FMemoryDerivedDataBackend(const TCHAR* InName, int64 InMaxCacheSize = -1, bool bCanBeDisabled = false);
+	~FMemoryDerivedDataBackend();
+
+	/** Return a name for this interface */
+	virtual FString GetName() const override { return Name; }
+
+	/** return true if this cache is writable **/
+	virtual bool IsWritable() const override;
+
+	/** Returns a class of speed for this interface **/
+	virtual ESpeedClass GetSpeedClass()  const override;
+
+	/**
+	 * Synchronous test for the existence of a cache item
+	 *
+	 * @param	CacheKey	Alphanumeric+underscore key of this cache item
+	 * @return				true if the data probably will be found, this can't be guaranteed because of concurrency in the backends, corruption, etc
+	 */
+	virtual bool CachedDataProbablyExists(const TCHAR* CacheKey) override;
+
+	/**
+	 * Synchronous retrieve of a cache item
+	 *
+	 * @param	CacheKey	Alphanumeric+underscore key of this cache item
+	 * @param	OutData		Buffer to receive the results, if any were found
+	 * @return				true if any data was found, and in this case OutData is non-empty
+	 */
+	virtual bool GetCachedData(const TCHAR* CacheKey, TArray<uint8>& OutData) override;
+
+	/**
+	 * Asynchronous, fire-and-forget placement of a cache item
+	 *
+	 * @param	CacheKey	Alphanumeric+underscore key of this cache item
+	 * @param	InData		Buffer containing the data to cache, can be destroyed after the call returns, immediately
+	 * @param	bPutEvenIfExists	If true, then do not attempt skip the put even if CachedDataProbablyExists returns true
+	 */
+	virtual EPutStatus PutCachedData(const TCHAR* CacheKey, TArrayView<const uint8> InData, bool bPutEvenIfExists) override;
+
+	virtual void RemoveCachedData(const TCHAR* CacheKey, bool bTransient) override;
+
+	/**
+	 * Save the cache to disk
+	 * @param	Filename	Filename to save
+	 * @return	true if file was saved successfully
+	 */
+	bool SaveCache(const TCHAR* Filename);
+
+	/**
+	 * Load the cache from disk
+	 * @param	Filename	Filename to load
+	 * @return	true if file was loaded successfully
+	 */
+	bool LoadCache(const TCHAR* Filename);
+
+	/**
+	 * Disable cache and ignore all subsequent requests
+	 */
+	void Disable() override;
+
+	virtual TSharedRef<FDerivedDataCacheStatsNode> GatherUsageStats() const override;
+
+	virtual bool TryToPrefetch(TConstArrayView<FString> CacheKeys) override;
+
+	/**
+	 *  Determines if we would cache the provided data
+	 */
+	virtual bool WouldCache(const TCHAR* CacheKey, TArrayView<const uint8> InData) override;
+
+	/**
+	 * Apply debug options
+	 */
+	bool ApplyDebugOptions(FBackendDebugOptions& InOptions) override;
+
+	virtual void Put(
+		TConstArrayView<FCachePutRequest> Requests,
+		IRequestOwner& Owner,
+		FOnCachePutComplete&& OnComplete) override;
+
+	virtual void Get(
+		TConstArrayView<FCacheGetRequest> Requests,
+		IRequestOwner& Owner,
+		FOnCacheGetComplete&& OnComplete) override;
+
+	virtual void GetChunks(
+		TConstArrayView<FCacheChunkRequest> Requests,
+		IRequestOwner& Owner,
+		FOnCacheChunkComplete&& OnComplete) override;
+
+private:
+	/** Name of the cache file loaded (if any). */
+	FString CacheFilename;
+	FDerivedDataCacheUsageStats UsageStats;
+
+	struct FCacheValue
+	{
+		int32 Age;
+		int32 Size;
+		FRWLock DataLock;
+		TArray<uint8> Data;
+		FCacheValue(int32 InSize, int32 InAge = 0)
+			: Age(InAge)
+			, Size(InSize)
+		{
+		}
+	};
+
+	FORCEINLINE int32 CalcSerializedCacheValueSize(const FString& Key, const FCacheValue& Val)
+	{
+		return (Key.Len() + 1) * sizeof(TCHAR) + sizeof(FCacheValue::Age) + Val.Size;
+	}
+
+	FORCEINLINE int32 CalcSerializedCacheValueSize(const FString& Key, const TArrayView<const uint8>& Data)
+	{
+		return (Key.Len() + 1) * sizeof(TCHAR) + sizeof(FCacheValue::Age) + Data.Num();
+	}
+
+	/** Name of this cache (used for debugging) */
+	FString Name;
+
+	/** Set of files that are being written to disk asynchronously. */
+	TMap<FString, FCacheValue*> CacheItems;
+	/** Set of records in this cache. */
+	TSet<FCacheRecord, FCacheRecordKeyFuncs> CacheRecords;
+	/** Maximum size the cached items can grow up to ( in bytes ) */
+	int64 MaxCacheSize;
+	/** When set to true, this cache is disabled...ignore all requests. */
+	std::atomic<bool> bDisabled;
+	/** Object used for synchronization via a scoped lock						*/
+	mutable FRWLock SynchronizationObject;
+	/** Current estimated cache size in bytes */
+	int64 CurrentCacheSize;
+	/** Indicates that the cache max size has been exceeded. This is used to avoid
+		warning spam after the size has reached the limit. */
+	bool bMaxSizeExceeded;
+
+	/** When a memory cache can be disabled, it won't return true for CachedDataProbablyExists calls.
+	  * This is to avoid having the Boot DDC tells it has some resources that will suddenly disappear after the boot.
+	  * Get() requests will still get fulfilled and other cache level will be properly back-filled 
+	  * offering the speed benefit of the boot cache while maintaining coherency at all cache levels.
+	  * 
+	  * The problem is that most asset types (audio/staticmesh/texture) will always verify if their different LODS/Chunks can be found in the cache using CachedDataProbablyExists.
+	  * If any of the LOD/MIP can't be found, a build of the asset is triggered, otherwise they skip asset compilation altogether.
+	  * However, we should not skip the compilation based on the CachedDataProbablyExists result of the boot cache because it is a lie and will disappear at some point.
+	  * When the boot cache disappears and the streamer tries to fetch a LOD that it has been told was cached, it will fail and will then have no choice but to rebuild the asset synchronously.
+	  * This obviously causes heavy game-thread stutters.
+
+	  * However, if the bootcache returns false during CachedDataProbablyExists. The async compilation will be triggered and data will be put in the both the boot.ddc and the local cache.
+	  * This way, no more heavy game-thread stutters during streaming...
+
+	  * This can be reproed when you clear the local cache but do not clear the boot.ddc file, but even if it's a corner case, I stumbled upon it enough times that I though it was worth to fix so the caches are coherent.
+	  */
+	bool bCanBeDisabled = false;
+	bool bShuttingDown  = false;
+
+	enum 
+	{
+		/** Magic number to use in header */
+		MemCache_Magic = 0x0cac0ddc,
+		/** Magic number to use in header (new, > 2GB size compatible) */
+		MemCache_Magic64 = 0x0cac1ddc,
+		/** Oldest cache items to keep */
+		MaxAge = 3,
+		/** Size of data that is stored in the cachefile apart from the cache entries (64 bit size). */
+		SerializationSpecificDataSize = sizeof(uint32)	// Magic
+									  + sizeof(int64)	// Size
+									  + sizeof(uint32), // CRC
+	};
+
+protected:
+
+	/* Debug helpers */
+	bool ShouldSimulateMiss(const TCHAR* InKey);
+	bool ShouldSimulateMiss(const FCacheKey& InKey);
+
+	/** Debug Options */
+	FBackendDebugOptions DebugOptions;
+
+	/** Keys we ignored due to miss rate settings */
+	FCriticalSection MissedKeysCS;
+	TSet<FName> DebugMissedKeys;
+	TSet<FCacheKey> DebugMissedCacheKeys;
+};
 
 FMemoryDerivedDataBackend::FMemoryDerivedDataBackend(const TCHAR* InName, int64 InMaxCacheSize, bool bInCanBeDisabled)
 	: Name(InName)
@@ -630,4 +823,9 @@ void FMemoryDerivedDataBackend::GetChunks(
 	}
 }
 
-} // UE::DerivedData::Backends
+FFileBackedDerivedDataBackend* CreateMemoryDerivedDataBackend(const TCHAR* Name, int64 MaxCacheSize, bool bCanBeDisabled)
+{
+	return new FMemoryDerivedDataBackend(Name, MaxCacheSize, bCanBeDisabled);
+}
+
+} // UE::DerivedData::CacheStore::Memory
diff --git a/Engine/Source/Developer/DerivedDataCache/Private/MemoryDerivedDataBackend.h b/Engine/Source/Developer/DerivedDataCache/Private/MemoryDerivedDataBackend.h
deleted file mode 100644
diff --git a/Engine/Source/Developer/DerivedDataCache/Private/S3DerivedDataBackend.cpp b/Engine/Source/Developer/DerivedDataCache/Private/S3DerivedDataBackend.cpp
--- a/Engine/Source/Developer/DerivedDataCache/Private/S3DerivedDataBackend.cpp
+++ b/Engine/Source/Developer/DerivedDataCache/Private/S3DerivedDataBackend.cpp
@@ -1,6 +1,6 @@
 // Copyright Epic Games, Inc. All Rights Reserved.
 
-#include "S3DerivedDataBackend.h"
+#include "DerivedDataBackendInterface.h"
 
 #if WITH_S3_DDC_BACKEND
 
@@ -8,39 +8,40 @@
 	#include "Windows/WindowsHWrapper.h"
 	#include "Windows/AllowWindowsPlatformTypes.h"
 #endif
-#include "curl/curl.h"
 #if PLATFORM_WINDOWS || PLATFORM_HOLOLENS
 	#include "Windows/HideWindowsPlatformTypes.h"
 #endif
-#include "Ssl.h"
-#include <openssl/ssl.h>
-#include <openssl/sha.h>
-#include <openssl/hmac.h>
+
 #include "Algo/AllOf.h"
-#include "Memory/SharedBuffer.h"
-#include "Misc/Base64.h"
-#include "Misc/Paths.h"
-#include "Misc/SecureHash.h"
-#include "Misc/FileHelper.h"
-#include "Misc/FeedbackContext.h"
-#include "Misc/OutputDeviceRedirector.h"
 #include "Async/ParallelFor.h"
-#include "Serialization/MemoryReader.h"
 #include "DerivedDataBackendCorruptionWrapper.h"
 #include "DerivedDataCacheRecord.h"
+#include "DerivedDataCacheUsageStats.h"
+#include "DesktopPlatformModule.h"
 #include "Dom/JsonObject.h"
-#include "Serialization/JsonReader.h"
-#include "Serialization/JsonSerializer.h"
-#include "ProfilingDebugging/CpuProfilerTrace.h"
-#include "ProfilingDebugging/CountersTrace.h"
 #include "HAL/PlatformFile.h"
 #include "HAL/PlatformFileManager.h"
 #include "HAL/Runnable.h"
-#include "DesktopPlatformModule.h"
+#include "Memory/SharedBuffer.h"
+#include "Misc/Base64.h"
 #include "Misc/ConfigCacheIni.h"
+#include "Misc/FeedbackContext.h"
+#include "Misc/FileHelper.h"
+#include "Misc/OutputDeviceRedirector.h"
+#include "Misc/Paths.h"
+#include "Misc/SecureHash.h"
+#include "ProfilingDebugging/CountersTrace.h"
+#include "ProfilingDebugging/CpuProfilerTrace.h"
+#include "Serialization/JsonReader.h"
+#include "Serialization/JsonSerializer.h"
+#include "Serialization/MemoryReader.h"
+
+#include "curl/curl.h"
 
 #if WITH_SSL
 #include "Ssl.h"
+#include <openssl/hmac.h>
+#include <openssl/sha.h>
 #include <openssl/ssl.h>
 #endif
 
@@ -187,6 +188,101 @@ struct IRequestCallback
 	virtual bool Update(int NumBytes, int TotalBytes) = 0;
 };
 
+/**
+ * Backend for a read-only AWS S3 based caching service.
+ **/
+class FS3DerivedDataBackend : public FDerivedDataBackendInterface
+{
+public:
+	/**
+	 * Creates the backend, checks health status and attempts to acquire an access token.
+	 *
+	 * @param  InRootManifestPath   Local path to the JSON manifest in the workspace containing a list of files to download
+	 * @param  InBaseUrl            Base URL for the bucket, with trailing slash (eg. https://foo.s3.us-east-1.amazonaws.com/)
+	 * @param  InRegion	            Name of the AWS region (eg. us-east-1)
+	 * @param  InCanaryObjectKey    Key for a canary object used to test whether this backend is usable
+	 * @param  InCachePath          Path to cache the DDC files
+	 */
+	FS3DerivedDataBackend(const TCHAR* InRootManifestPath, const TCHAR* InBaseUrl, const TCHAR* InRegion, const TCHAR* InCanaryObjectKey, const TCHAR* InCachePath);
+	~FS3DerivedDataBackend();
+
+	/**
+	 * Checks is backend is usable (reachable and accessible).
+	 * @return true if usable
+	 */
+	bool IsUsable() const;
+
+	/* S3 Cache cannot be written to*/
+	bool IsWritable() const override { return false; }
+
+	/* S3 Cache does not try to write back to lower caches (e.g. Shared DDC) */
+	bool BackfillLowerCacheLevels() const override { return false; }
+
+	bool CachedDataProbablyExists(const TCHAR* CacheKey) override;
+	bool GetCachedData(const TCHAR* CacheKey, TArray<uint8>& OutData) override;
+	EPutStatus PutCachedData(const TCHAR* CacheKey, TArrayView<const uint8> InData, bool bPutEvenIfExists) override;
+	void RemoveCachedData(const TCHAR* CacheKey, bool bTransient) override;
+	virtual TSharedRef<FDerivedDataCacheStatsNode> GatherUsageStats() const override;
+
+	FString GetName() const override;
+	ESpeedClass GetSpeedClass() const override;
+	bool TryToPrefetch(TConstArrayView<FString> CacheKeys) override;
+	bool WouldCache(const TCHAR* CacheKey, TArrayView<const uint8> InData) override;
+
+	bool ApplyDebugOptions(FBackendDebugOptions& InOptions) override;
+
+	virtual void Put(
+		TConstArrayView<FCachePutRequest> Requests,
+		IRequestOwner& Owner,
+		FOnCachePutComplete&& OnComplete) override;
+
+	virtual void Get(
+		TConstArrayView<FCacheGetRequest> Requests,
+		IRequestOwner& Owner,
+		FOnCacheGetComplete&& OnComplete) override;
+
+	virtual void GetChunks(
+		TConstArrayView<FCacheChunkRequest> Requests,
+		IRequestOwner& Owner,
+		FOnCacheChunkComplete&& OnComplete) override;
+
+private:
+	struct FBundle;
+	struct FBundleEntry;
+	struct FBundleDownload;
+
+	struct FRootManifest;
+
+	class FHttpRequest;
+	class FRequestPool;
+
+	FString RootManifestPath;
+	FString BaseUrl;
+	FString Region;
+	FString CanaryObjectKey;
+	FString CacheDir;
+	TArray<FBundle> Bundles;
+	TUniquePtr<FRequestPool> RequestPool;
+	FDerivedDataCacheUsageStats UsageStats;
+	bool bEnabled;
+
+	bool DownloadManifest(const FRootManifest& RootManifest, FFeedbackContext* Context);
+	void RemoveUnusedBundles();
+	void ReadBundle(FBundle& Bundle);
+	bool FindBundleEntry(const TCHAR* CacheKey, const FBundle*& OutBundle, const FBundleEntry*& OutBundleEntry) const;
+
+	/* Debug helpers */
+	bool DidSimulateMiss(const TCHAR* InKey);
+	bool ShouldSimulateMiss(const TCHAR* InKey);
+
+	/** Debug Options */
+	FBackendDebugOptions DebugOptions;
+
+	/** Keys we ignored due to miss rate settings */
+	FCriticalSection MissedKeysCS;
+	TSet<FName> DebugMissedKeys;
+};
+
 /**
  * Minimal HTTP request type wrapping CURL without the need for managers. This request
  * is written to allow reuse of request objects, in order to allow connections to be reused.
@@ -1260,6 +1356,23 @@ void FS3DerivedDataBackend::GetChunks(
 	}
 }
 
+FDerivedDataBackendInterface* CreateS3DerivedDataBackend(const TCHAR* RootManifestPath, const TCHAR* BaseUrl, const TCHAR* Region, const TCHAR* CanaryObjectKey, const TCHAR* CachePath)
+{
+	return new FS3DerivedDataBackend(RootManifestPath, BaseUrl, Region, CanaryObjectKey, CachePath);
+}
+
 } // UE::DerivedData::CacheStore::S3
 
-#endif
+#else
+
+namespace UE::DerivedData::CacheStore::S3
+{
+
+FDerivedDataBackendInterface* CreateS3DerivedDataBackend(const TCHAR* RootManifestPath, const TCHAR* BaseUrl, const TCHAR* Region, const TCHAR* CanaryObjectKey, const TCHAR* CachePath)
+{
+	return nullptr;
+}
+
+} // UE::DerivedData::CacheStore::S3
+
+#endif // WITH_S3_DDC_BACKEND
diff --git a/Engine/Source/Developer/DerivedDataCache/Private/S3DerivedDataBackend.h b/Engine/Source/Developer/DerivedDataCache/Private/S3DerivedDataBackend.h
deleted file mode 100644
diff --git a/Engine/Source/Developer/DerivedDataCache/Private/Tests/HttpDerivedDataBackendTest.cpp b/Engine/Source/Developer/DerivedDataCache/Private/Tests/HttpDerivedDataBackendTest.cpp
--- a/Engine/Source/Developer/DerivedDataCache/Private/Tests/HttpDerivedDataBackendTest.cpp
+++ b/Engine/Source/Developer/DerivedDataCache/Private/Tests/HttpDerivedDataBackendTest.cpp
@@ -1,8 +1,5 @@
 // Copyright Epic Games, Inc. All Rights Reserved.
 
-#include "HttpDerivedDataBackend.h"
-#include "ZenDerivedDataBackend.h"
-
 #include "Async/ParallelFor.h"
 #include "DerivedDataBackendInterface.h"
 #include "DerivedDataCacheKey.h"
@@ -15,6 +12,7 @@
 #include "Misc/Paths.h"
 #include "Misc/SecureHash.h"
 #include "Serialization/CompactBinaryWriter.h"
+#include "ZenServerInterface.h"
 
 // Test is targeted at HttpDerivedDataBackend but with some backend test interface it could be generalized
 // to function against all backends.
@@ -36,6 +34,22 @@ DEFINE_LOG_CATEGORY_STATIC(LogHttpDerivedDataBackendTests, Log, All);
 		} \
 	}
 
+namespace UE::DerivedData::CacheStore::Http
+{
+FDerivedDataBackendInterface* GetAnyHttpDerivedDataBackend(
+	FString& OutDomain,
+	FString& OutOAuthProvider,
+	FString& OutOAuthClientId,
+	FString& OutOAuthSecret,
+	FString& OutNamespace,
+	FString& OutStructuredNamespace);
+} // UE::DerivedData::CacheStore::Http
+
+namespace UE::DerivedData::CacheStore::ZenCache
+{
+FDerivedDataBackendInterface* CreateZenDerivedDataBackend(const TCHAR* NodeName, const TCHAR* ServiceUrl, const TCHAR* Namespace);
+} // UE::DerivedData::CacheStore::ZenCache
+
 namespace HttpDerivedDataBackendTest
 {
 
@@ -49,12 +63,9 @@ public:
 
 	bool CheckPrequisites() const
 	{
-		if (UE::DerivedData::Backends::FHttpDerivedDataBackend* Backend = GetTestBackend())
+		if (FDerivedDataBackendInterface* Backend = GetTestBackend())
 		{
-			if (Backend->IsUsable())
-			{
-				return true;
-			}
+			return true;
 		}
 		return false;
 	}
@@ -148,9 +159,10 @@ protected:
 		FPlatformProcess::ReturnSynchEventToPool(LastEvent);
 	}
 
-	UE::DerivedData::Backends::FHttpDerivedDataBackend* GetTestBackend() const
+	FDerivedDataBackendInterface* GetTestBackend() const
 	{
-		static UE::DerivedData::Backends::FHttpDerivedDataBackend* CachedBackend = FetchTestBackend_Internal();
+		static FDerivedDataBackendInterface* CachedBackend = UE::DerivedData::CacheStore::Http::GetAnyHttpDerivedDataBackend(
+			TestDomain, TestOAuthProvider, TestOAuthClientId, TestOAuthSecret, TestNamespace, TestStructuredNamespace);
 		return CachedBackend;
 	}
 
@@ -269,15 +281,17 @@ protected:
 		return ReceivedRecords;
 	}
 
-private:
-	UE::DerivedData::Backends::FHttpDerivedDataBackend* FetchTestBackend_Internal() const
-	{
-		return UE::DerivedData::Backends::FHttpDerivedDataBackend::GetAny();
-	}
+protected:
+	static inline FString TestDomain;
+	static inline FString TestOAuthProvider;
+	static inline FString TestOAuthClientId;
+	static inline FString TestOAuthSecret;
+	static inline FString TestNamespace;
+	static inline FString TestStructuredNamespace;
 };
 
 // Helper function to create a number of dummy cache keys for testing
-TArray<FString> CreateTestCacheKeys(UE::DerivedData::Backends::FHttpDerivedDataBackend* InTestBackend, uint32 InNumKeys)
+TArray<FString> CreateTestCacheKeys(FDerivedDataBackendInterface* InTestBackend, uint32 InNumKeys)
 {
 	TArray<FString> Keys;
 	TArray<uint8> KeyContents;
@@ -360,7 +374,7 @@ TArray<UE::DerivedData::FCacheRecord> CreateTestCacheRecords(UE::DerivedData::IC
 IMPLEMENT_HTTPDERIVEDDATA_AUTOMATION_TEST(FConcurrentCachedDataProbablyExistsBatch, TEXT(".FConcurrentCachedDataProbablyExistsBatch"), EAutomationTestFlags::EditorContext | EAutomationTestFlags::ProductFilter)
 bool FConcurrentCachedDataProbablyExistsBatch::RunTest(const FString& Parameters)
 {
-	UE::DerivedData::Backends::FHttpDerivedDataBackend* TestBackend = GetTestBackend();
+	FDerivedDataBackendInterface* TestBackend = GetTestBackend();
 
 	const int32 ThreadCount = 64;
 	const double Duration = 10;
@@ -393,7 +407,7 @@ bool FConcurrentCachedDataProbablyExistsBatch::RunTest(const FString& Parameters
 IMPLEMENT_HTTPDERIVEDDATA_AUTOMATION_TEST(FConcurrentExistsAndGetForSameKeyBatch, TEXT(".FConcurrentExistsAndGetForSameKeyBatch"), EAutomationTestFlags::EditorContext | EAutomationTestFlags::ProductFilter)
 bool FConcurrentExistsAndGetForSameKeyBatch::RunTest(const FString& Parameters)
 {
-	UE::DerivedData::Backends::FHttpDerivedDataBackend* TestBackend = GetTestBackend();
+	FDerivedDataBackendInterface* TestBackend = GetTestBackend();
 
 	const int32 ParallelTasks = 32;
 	const uint32 Iterations = 20;
@@ -434,33 +448,33 @@ IMPLEMENT_HTTPDERIVEDDATA_AUTOMATION_TEST(CacheStore, TEXT(".CacheStore"), EAuto
 bool CacheStore::RunTest(const FString& Parameters)
 {
 	using namespace UE::DerivedData;
-	UE::DerivedData::Backends::FHttpDerivedDataBackend* TestBackend = GetTestBackend();
+	FDerivedDataBackendInterface* TestBackend = GetTestBackend();
 
-#if WITH_ZEN_DDC_BACKEND
+#if UE_WITH_ZEN
 	using namespace UE::Zen;
 	FServiceSettings ZenTestServiceSettings;
 	FServiceAutoLaunchSettings& ZenTestAutoLaunchSettings = ZenTestServiceSettings.SettingsVariant.Get<FServiceAutoLaunchSettings>();
 	ZenTestAutoLaunchSettings.DataPath = FPaths::ConvertRelativePathToFull(FPaths::Combine(FPaths::EngineSavedDir(), "ZenUnitTest"));
 	ZenTestAutoLaunchSettings.ExtraArgs = FString::Printf(TEXT("--http asio --upstream-jupiter-url \"%s\" --upstream-jupiter-oauth-url \"%s\" --upstream-jupiter-oauth-clientid \"%s\" --upstream-jupiter-oauth-clientsecret \"%s\" --upstream-jupiter-namespace-ddc \"%s\" --upstream-jupiter-namespace \"%s\""),
-		*TestBackend->GetDomain(),
-		*TestBackend->GetOAuthProvider(),
-		*TestBackend->GetOAuthClientId(),
-		*TestBackend->GetOAuthSecret(),
-		*TestBackend->GetNamespace(),
-		*TestBackend->GetStructuredNamespace()
+		*TestDomain,
+		*TestOAuthProvider,
+		*TestOAuthClientId,
+		*TestOAuthSecret,
+		*TestNamespace,
+		*TestStructuredNamespace
 	);
 	ZenTestAutoLaunchSettings.DesiredPort = 13337; // Avoid the normal default port
 	ZenTestAutoLaunchSettings.bShowConsole = true;
 	ZenTestAutoLaunchSettings.bLimitProcessLifetime = true;
 
 	FScopeZenService ScopeZenService(MoveTemp(ZenTestServiceSettings));
-	TUniquePtr<Backends::FZenDerivedDataBackend> ZenIntermediaryBackend = MakeUnique<UE::DerivedData::Backends::FZenDerivedDataBackend>(ScopeZenService.GetInstance().GetURL(), *TestBackend->GetNamespace());
-	auto WaitForZenPushToUpstream = [](Backends::FZenDerivedDataBackend* ZenBackend, TConstArrayView<FCacheRecord> Records)
+	TUniquePtr<FDerivedDataBackendInterface> ZenIntermediaryBackend(UE::DerivedData::CacheStore::ZenCache::CreateZenDerivedDataBackend(TEXT("Test"), ScopeZenService.GetInstance().GetURL(), *TestNamespace));
+	auto WaitForZenPushToUpstream = [](FDerivedDataBackendInterface* ZenBackend, TConstArrayView<FCacheRecord> Records)
 	{
 		// TODO: Expecting a legitimate means to wait for zen to finish pushing records to its upstream in the future
 		FPlatformProcess::Sleep(1.0f);
 	};
-#endif // WITH_ZEN_DDC_BACKEND
+#endif // UE_WITH_ZEN
 
 	const uint32 RecordsInBatch = 3;
 
@@ -470,13 +484,16 @@ bool CacheStore::RunTest(const FString& Parameters)
 		TArray<FCacheRecord> RecievedRecordsSkipMeta = GetAndValidateRecords(TEXT("SimpleValueSkipMeta"), PutRecords, ECachePolicy::Default | ECachePolicy::SkipMeta);
 		TArray<FCacheRecord> RecievedRecordsSkipData = GetAndValidateRecords(TEXT("SimpleValueSkipData"), PutRecords, ECachePolicy::Default | ECachePolicy::SkipData);
 
-#if WITH_ZEN_DDC_BACKEND
-		TArray<FCacheRecord> PutRecordsZen = CreateTestCacheRecords(ZenIntermediaryBackend.Get(), RecordsInBatch, 1, FCbObject(), "AutoTestDummyZen");
-		WaitForZenPushToUpstream(ZenIntermediaryBackend.Get(), PutRecordsZen);
-		ValidateRecords(TEXT("SimpleValueZenAndDirect"), GetAndValidateRecords(TEXT("SimpleValueZen"), PutRecordsZen, ECachePolicy::Default), RecievedRecords, ECachePolicy::Default);
-		ValidateRecords(TEXT("SimpleValueSkipMetaZenAndDirect"), GetAndValidateRecords(TEXT("SimpleValueSkipMetaZen"), PutRecordsZen, ECachePolicy::Default | ECachePolicy::SkipMeta), RecievedRecordsSkipMeta, ECachePolicy::Default | ECachePolicy::SkipMeta);
-		ValidateRecords(TEXT("SimpleValueSkipDataZenAndDirect"), GetAndValidateRecords(TEXT("SimpleValueSkipDataZen"), PutRecordsZen, ECachePolicy::Default | ECachePolicy::SkipData), RecievedRecordsSkipData, ECachePolicy::Default | ECachePolicy::SkipData);
-#endif // WITH_ZEN_DDC_BACKEND
+#if UE_WITH_ZEN
+		if (ZenIntermediaryBackend)
+		{
+			TArray<FCacheRecord> PutRecordsZen = CreateTestCacheRecords(ZenIntermediaryBackend.Get(), RecordsInBatch, 1, FCbObject(), "AutoTestDummyZen");
+			WaitForZenPushToUpstream(ZenIntermediaryBackend.Get(), PutRecordsZen);
+			ValidateRecords(TEXT("SimpleValueZenAndDirect"), GetAndValidateRecords(TEXT("SimpleValueZen"), PutRecordsZen, ECachePolicy::Default), RecievedRecords, ECachePolicy::Default);
+			ValidateRecords(TEXT("SimpleValueSkipMetaZenAndDirect"), GetAndValidateRecords(TEXT("SimpleValueSkipMetaZen"), PutRecordsZen, ECachePolicy::Default | ECachePolicy::SkipMeta), RecievedRecordsSkipMeta, ECachePolicy::Default | ECachePolicy::SkipMeta);
+			ValidateRecords(TEXT("SimpleValueSkipDataZenAndDirect"), GetAndValidateRecords(TEXT("SimpleValueSkipDataZen"), PutRecordsZen, ECachePolicy::Default | ECachePolicy::SkipData), RecievedRecordsSkipData, ECachePolicy::Default | ECachePolicy::SkipData);
+		}
+#endif // UE_WITH_ZEN
 	}
 
 	{
@@ -491,13 +508,16 @@ bool CacheStore::RunTest(const FString& Parameters)
 		TArray<FCacheRecord> RecievedRecordsSkipMeta = GetAndValidateRecords(TEXT("SimpleValueWithMetaSkipMeta"), PutRecords, ECachePolicy::Default | ECachePolicy::SkipMeta);
 		TArray<FCacheRecord> RecievedRecordsSkipData = GetAndValidateRecords(TEXT("SimpleValueWithMetaSkipData"), PutRecords, ECachePolicy::Default | ECachePolicy::SkipData);
 
-#if WITH_ZEN_DDC_BACKEND
-		TArray<FCacheRecord> PutRecordsZen = CreateTestCacheRecords(ZenIntermediaryBackend.Get(), RecordsInBatch, 1, MetaObject, "AutoTestDummyZen");
-		WaitForZenPushToUpstream(ZenIntermediaryBackend.Get(), PutRecordsZen);
-		ValidateRecords(TEXT("SimpleValueWithMetaZenAndDirect"), GetAndValidateRecords(TEXT("SimpleValueWithMetaZen"), PutRecordsZen, ECachePolicy::Default), RecievedRecords, ECachePolicy::Default);
-		ValidateRecords(TEXT("SimpleValueWithMetaSkipMetaZenAndDirect"), GetAndValidateRecords(TEXT("SimpleValueWithMetaSkipMetaZen"), PutRecordsZen, ECachePolicy::Default | ECachePolicy::SkipMeta), RecievedRecordsSkipMeta, ECachePolicy::Default | ECachePolicy::SkipMeta);
-		ValidateRecords(TEXT("SimpleValueWithMetaSkipDataZenAndDirect"), GetAndValidateRecords(TEXT("SimpleValueWithMetaSkipDataZen"), PutRecordsZen, ECachePolicy::Default | ECachePolicy::SkipData), RecievedRecordsSkipData, ECachePolicy::Default | ECachePolicy::SkipData);
-#endif // WITH_ZEN_DDC_BACKEND
+#if UE_WITH_ZEN
+		if (ZenIntermediaryBackend)
+		{
+			TArray<FCacheRecord> PutRecordsZen = CreateTestCacheRecords(ZenIntermediaryBackend.Get(), RecordsInBatch, 1, MetaObject, "AutoTestDummyZen");
+			WaitForZenPushToUpstream(ZenIntermediaryBackend.Get(), PutRecordsZen);
+			ValidateRecords(TEXT("SimpleValueWithMetaZenAndDirect"), GetAndValidateRecords(TEXT("SimpleValueWithMetaZen"), PutRecordsZen, ECachePolicy::Default), RecievedRecords, ECachePolicy::Default);
+			ValidateRecords(TEXT("SimpleValueWithMetaSkipMetaZenAndDirect"), GetAndValidateRecords(TEXT("SimpleValueWithMetaSkipMetaZen"), PutRecordsZen, ECachePolicy::Default | ECachePolicy::SkipMeta), RecievedRecordsSkipMeta, ECachePolicy::Default | ECachePolicy::SkipMeta);
+			ValidateRecords(TEXT("SimpleValueWithMetaSkipDataZenAndDirect"), GetAndValidateRecords(TEXT("SimpleValueWithMetaSkipDataZen"), PutRecordsZen, ECachePolicy::Default | ECachePolicy::SkipData), RecievedRecordsSkipData, ECachePolicy::Default | ECachePolicy::SkipData);
+		}
+#endif // UE_WITH_ZEN
 	}
 
 	{
@@ -506,13 +526,16 @@ bool CacheStore::RunTest(const FString& Parameters)
 		TArray<FCacheRecord> RecievedRecordsSkipMeta = GetAndValidateRecords(TEXT("MultiValueSkipMeta"), PutRecords, ECachePolicy::Default | ECachePolicy::SkipMeta);
 		TArray<FCacheRecord> RecievedRecordsSkipData = GetAndValidateRecords(TEXT("MultiValueSkipData"), PutRecords, ECachePolicy::Default | ECachePolicy::SkipData);
 
-#if WITH_ZEN_DDC_BACKEND
-		TArray<FCacheRecord> PutRecordsZen = CreateTestCacheRecords(ZenIntermediaryBackend.Get(), RecordsInBatch, 5, FCbObject(), "AutoTestDummyZen");
-		WaitForZenPushToUpstream(ZenIntermediaryBackend.Get(), PutRecordsZen);
-		ValidateRecords(TEXT("MultiValueZenAndDirect"), GetAndValidateRecords(TEXT("MultiValueZen"), PutRecordsZen, ECachePolicy::Default), RecievedRecords, ECachePolicy::Default);
-		ValidateRecords(TEXT("MultiValueSkipMetaZenAndDirect"), GetAndValidateRecords(TEXT("MultiValueSkipMetaZen"), PutRecordsZen, ECachePolicy::Default | ECachePolicy::SkipMeta), RecievedRecordsSkipMeta, ECachePolicy::Default | ECachePolicy::SkipMeta);
-		ValidateRecords(TEXT("MultiValueSkipDataZenAndDirect"), GetAndValidateRecords(TEXT("MultiValueSkipDataZen"), PutRecordsZen, ECachePolicy::Default | ECachePolicy::SkipData), RecievedRecordsSkipData, ECachePolicy::Default | ECachePolicy::SkipData);
-#endif // WITH_ZEN_DDC_BACKEND
+#if UE_WITH_ZEN
+		if (ZenIntermediaryBackend)
+		{
+			TArray<FCacheRecord> PutRecordsZen = CreateTestCacheRecords(ZenIntermediaryBackend.Get(), RecordsInBatch, 5, FCbObject(), "AutoTestDummyZen");
+			WaitForZenPushToUpstream(ZenIntermediaryBackend.Get(), PutRecordsZen);
+			ValidateRecords(TEXT("MultiValueZenAndDirect"), GetAndValidateRecords(TEXT("MultiValueZen"), PutRecordsZen, ECachePolicy::Default), RecievedRecords, ECachePolicy::Default);
+			ValidateRecords(TEXT("MultiValueSkipMetaZenAndDirect"), GetAndValidateRecords(TEXT("MultiValueSkipMetaZen"), PutRecordsZen, ECachePolicy::Default | ECachePolicy::SkipMeta), RecievedRecordsSkipMeta, ECachePolicy::Default | ECachePolicy::SkipMeta);
+			ValidateRecords(TEXT("MultiValueSkipDataZenAndDirect"), GetAndValidateRecords(TEXT("MultiValueSkipDataZen"), PutRecordsZen, ECachePolicy::Default | ECachePolicy::SkipData), RecievedRecordsSkipData, ECachePolicy::Default | ECachePolicy::SkipData);
+		}
+#endif // UE_WITH_ZEN
 	}
 
 	return true;
diff --git a/Engine/Source/Developer/DerivedDataCache/Private/ZenDerivedDataBackend.cpp b/Engine/Source/Developer/DerivedDataCache/Private/ZenDerivedDataBackend.cpp
--- a/Engine/Source/Developer/DerivedDataCache/Private/ZenDerivedDataBackend.cpp
+++ b/Engine/Source/Developer/DerivedDataCache/Private/ZenDerivedDataBackend.cpp
@@ -1,22 +1,28 @@
 // Copyright Epic Games, Inc. All Rights Reserved.
-#include "ZenDerivedDataBackend.h"
 
-#if WITH_ZEN_DDC_BACKEND
+#include "DerivedDataBackendInterface.h"
+#include "ZenServerInterface.h"
 
+#if UE_WITH_ZEN
+
+#include "Containers/Set.h"
+#include "Containers/StringFwd.h"
 #include "Containers/StringFwd.h"
 #include "DerivedDataCachePrivate.h"
 #include "DerivedDataCacheRecord.h"
+#include "DerivedDataCacheUsageStats.h"
 #include "DerivedDataChunk.h"
+#include "HAL/CriticalSection.h"
 #include "Misc/ConfigCacheIni.h"
 #include "Misc/ScopeLock.h"
 #include "ProfilingDebugging/CountersTrace.h"
 #include "ProfilingDebugging/CpuProfilerTrace.h"
+#include "Serialization/BufferArchive.h"
 #include "Serialization/CompactBinary.h"
 #include "Serialization/CompactBinaryPackage.h"
+#include "Serialization/CompactBinarySerialization.h"
 #include "Serialization/CompactBinaryValidation.h"
 #include "Serialization/CompactBinaryWriter.h"
-#include "Serialization/CompactBinarySerialization.h"
-#include "Serialization/BufferArchive.h"
 #include "ZenBackendUtils.h"
 #include "ZenSerialization.h"
 #include "ZenServerHttp.h"
@@ -33,7 +39,8 @@ TRACE_DECLARE_INT_COUNTER(ZenDDC_BytesSent,		TEXT("ZenDDC Bytes Sent"));
 TRACE_DECLARE_INT_COUNTER(ZenDDC_CacheRecordRequestCount, TEXT("ZenDDC CacheRecord Request Count"));
 TRACE_DECLARE_INT_COUNTER(ZenDDC_ChunkRequestCount, TEXT("ZenDDC Chunk Request Count"));
 
-namespace UE::DerivedData::Backends {
+namespace UE::DerivedData::CacheStore::ZenCache
+{
 
 template<typename T>
 void ForEachBatch(const int32 BatchSize, const int32 TotalCount, T&& Fn)
@@ -59,6 +66,116 @@ void ForEachBatch(const int32 BatchSize, const int32 TotalCount, T&& Fn)
 // FZenDerivedDataBackend
 //----------------------------------------------------------------------------------------------------------
 
+/**
+ * Backend for a HTTP based caching service (Zen)
+ **/
+class FZenDerivedDataBackend : public FDerivedDataBackendInterface
+{
+public:
+	
+	/**
+	 * Creates the backend, checks health status and attempts to acquire an access token.
+	 *
+	 * @param ServiceUrl	Base url to the service including schema.
+	 * @param Namespace		Namespace to use.
+	 */
+	FZenDerivedDataBackend(const TCHAR* ServiceUrl, const TCHAR* Namespace);
+
+	~FZenDerivedDataBackend();
+
+	/**
+	 * Checks is backend is usable (reachable and accessible).
+	 * @return true if usable
+	 */
+	bool IsUsable() const { return bIsUsable; }
+
+	virtual bool IsWritable() const override;
+	virtual ESpeedClass GetSpeedClass() const override;
+	virtual TSharedRef<FDerivedDataCacheStatsNode> GatherUsageStats() const override;
+
+	/**
+	 * Synchronous attempt to make sure the cached data will be available as optimally as possible.
+	 *
+	 * @param	CacheKeys	Alphanumeric+underscore keys of the cache items
+	 * @return				true if the data will probably be found in a fast backend on a future request.
+	 */
+	virtual bool TryToPrefetch(TConstArrayView<FString> CacheKeys) override;
+
+	virtual bool CachedDataProbablyExists(const TCHAR* CacheKey) override;
+	virtual bool GetCachedData(const TCHAR* CacheKey, TArray<uint8>& OutData) override;
+	virtual EPutStatus PutCachedData(const TCHAR* CacheKey, TArrayView<const uint8> InData, bool bPutEvenIfExists) override;
+
+	virtual void RemoveCachedData(const TCHAR* CacheKey, bool bTransient) override;
+	
+	virtual FString GetName() const override;
+	virtual bool WouldCache(const TCHAR* CacheKey, TArrayView<const uint8> InData) override;
+	virtual bool ApplyDebugOptions(FBackendDebugOptions& InOptions) override;
+
+	// ICacheStore
+
+	virtual void Put(
+		TConstArrayView<FCachePutRequest> Requests,
+		IRequestOwner& Owner,
+		FOnCachePutComplete&& OnComplete = FOnCachePutComplete()) override;
+
+	virtual void Get(
+		TConstArrayView<FCacheGetRequest> Requests,
+		IRequestOwner& Owner,
+		FOnCacheGetComplete&& OnComplete) override;
+
+	virtual void GetChunks(
+		TConstArrayView<FCacheChunkRequest> Requests,
+		IRequestOwner& Owner,
+		FOnCacheChunkComplete&& OnComplete) override;
+
+private:
+	enum class EGetResult
+	{
+		Success,
+		NotFound,
+		Corrupted
+	};
+	EGetResult GetZenData(FStringView Uri, TArray64<uint8>* OutData, Zen::EContentType ContentType) const;
+
+	// TODO: need ability to specify content type
+	FDerivedDataBackendInterface::EPutStatus PutZenData(const TCHAR* Uri, const FCompositeBuffer& InData, Zen::EContentType ContentType);
+	EGetResult GetZenData(const FCacheKey& Key, ECachePolicy CachePolicy, FCbPackage& OutPackage) const;
+
+	bool PutCacheRecord(const FCacheRecord& Record, const FCacheRecordPolicy& Policy);
+
+	bool IsServiceReady();
+	static FString MakeLegacyZenKey(const TCHAR* CacheKey);
+	static void AppendZenUri(const FCacheKey& CacheKey, FStringBuilderBase& Out);
+	static void AppendZenUri(const FCacheKey& CacheKey, const FValueId& Id, FStringBuilderBase& Out);
+	static void AppendPolicyQueryString(ECachePolicy Policy, FStringBuilderBase& Out);
+
+	static bool ShouldRetryOnError(int64 ResponseCode);
+
+	/* Debug helpers */
+	bool ShouldSimulateMiss(const TCHAR* InKey);
+	bool ShouldSimulateMiss(const FCacheKey& InKey);
+
+private:
+	FString Namespace;
+	UE::Zen::FScopeZenService ZenService;
+	mutable FDerivedDataCacheUsageStats UsageStats;
+	TUniquePtr<UE::Zen::FZenHttpRequestPool> RequestPool;
+	bool bIsUsable = false;
+	bool bIsRemote = false;
+	uint32 FailedLoginAttempts = 0;
+	uint32 MaxAttempts = 4;
+	int32 CacheRecordBatchSize = 8;
+	int32 CacheChunksBatchSize = 8;
+
+	/** Debug Options */
+	FBackendDebugOptions DebugOptions;
+
+	/** Keys we ignored due to miss rate settings */
+	FCriticalSection MissedKeysCS;
+	TSet<FName> DebugMissedKeys;
+	TSet<FCacheKey> DebugMissedCacheKeys;
+};
+
 FZenDerivedDataBackend::FZenDerivedDataBackend(
 	const TCHAR* InServiceUrl,
 	const TCHAR* InNamespace)
@@ -919,6 +1036,30 @@ bool FZenDerivedDataBackend::PutCacheRecord(const FCacheRecord& Record, const FC
 	return true;
 }
 
-} // namespace UE::DerivedData::Backends
+FDerivedDataBackendInterface* CreateZenDerivedDataBackend(const TCHAR* NodeName, const TCHAR* ServiceUrl, const TCHAR* Namespace)
+{
+	FZenDerivedDataBackend* Backend = new FZenDerivedDataBackend(ServiceUrl, Namespace);
+	if (Backend->IsUsable())
+	{
+		return Backend;
+	}
+	UE_LOG(LogDerivedDataCache, Warning, TEXT("%s could not contact the service (%s), will not use it."), NodeName, *Backend->GetName());
+	delete Backend;
+	return nullptr;
+}
+
+} // namespace UE::DerivedData::CacheStore::ZenCache
+
+#else
+
+namespace UE::DerivedData::CacheStore::ZenCache
+{
+
+FDerivedDataBackendInterface* CreateZenDerivedDataBackend(const TCHAR* NodeName, const TCHAR* ServiceUrl, const TCHAR* Namespace)
+{
+	return nullptr;
+}
+
+} // UE::DerivedData::CacheStore::ZenCache
 
-#endif //WITH_HTTP_DDC_BACKEND
+#endif // UE_WITH_ZEN
diff --git a/Engine/Source/Developer/DerivedDataCache/Private/ZenDerivedDataBackend.h b/Engine/Source/Developer/DerivedDataCache/Private/ZenDerivedDataBackend.h
deleted file mode 100644
