DDC: Replaced BuildValue(ValueKey) with Build(Key)

The policy allows filtering the result to an arbitrary set of values, which makes the new function both more powerful and better aligned with the rest of the API.

#rnx

diff --git a/Engine/Source/Developer/DerivedDataCache/Private/DerivedDataBuildSession.cpp b/Engine/Source/Developer/DerivedDataCache/Private/DerivedDataBuildSession.cpp
--- a/Engine/Source/Developer/DerivedDataCache/Private/DerivedDataBuildSession.cpp
+++ b/Engine/Source/Developer/DerivedDataCache/Private/DerivedDataBuildSession.cpp
@@ -32,6 +32,12 @@ public:
 
 	FStringView GetName() const final { return Name; }
 
+	void Build(
+		const FBuildKey& Key,
+		const FBuildPolicy& Policy,
+		IRequestOwner& Owner,
+		FOnBuildComplete&& OnComplete) final;
+
 	void Build(
 		const FBuildDefinition& Definition,
 		const FOptionalBuildInputs& Inputs,
@@ -46,12 +52,6 @@ public:
 		IRequestOwner& Owner,
 		FOnBuildComplete&& OnComplete) final;
 
-	void BuildValue(
-		const FBuildValueKey& Value,
-		EBuildPolicy Policy,
-		IRequestOwner& Owner,
-		FOnBuildValueComplete&& OnComplete) final;
-
 	FString Name;
 	ICache& Cache;
 	IBuild& BuildSystem;
@@ -61,6 +61,16 @@ public:
 
 ///////////////////////////////////////////////////////////////////////////////////////////////////////////////////////
 
+void FBuildSessionInternal::Build(
+	const FBuildKey& Key,
+	const FBuildPolicy& Policy,
+	IRequestOwner& Owner,
+	FOnBuildComplete&& OnComplete)
+{
+	CreateBuildJob({Cache, BuildSystem, Scheduler, InputResolver, Owner}, Key, Policy,
+		OnComplete ? MoveTemp(OnComplete) : [](FBuildCompleteParams&&){});
+}
+
 void FBuildSessionInternal::Build(
 	const FBuildDefinition& Definition,
 	const FOptionalBuildInputs& Inputs,
@@ -83,35 +93,6 @@ void FBuildSessionInternal::Build(
 		OnComplete ? MoveTemp(OnComplete) : [](FBuildCompleteParams&&){});
 }
 
-void FBuildSessionInternal::BuildValue(
-	const FBuildValueKey& ValueKey,
-	EBuildPolicy Policy,
-	IRequestOwner& Owner,
-	FOnBuildValueComplete&& OnComplete)
-{
-	// This requests the entire output to get one value. It will be optimized later to request only one value.
-	FOnBuildComplete OnJobComplete;
-	if (OnComplete)
-	{
-		OnJobComplete = [ValueKey, OnComplete = MoveTemp(OnComplete)](FBuildCompleteParams&& Params)
-		{
-			FValueWithId Value;
-			EStatus Status = Params.Status;
-			if (Status == EStatus::Ok)
-			{
-				Value = Params.Output.GetValue(ValueKey.Id);
-				Status = Value ? EStatus::Ok : EStatus::Error;
-			}
-			if (!Value)
-			{
-				Value = FValueWithId(ValueKey.Id);
-			}
-			OnComplete({ValueKey.BuildKey, MoveTemp(Value), Status});
-		};
-	}
-	CreateBuildJob({Cache, BuildSystem, Scheduler, InputResolver, Owner}, ValueKey.BuildKey, Policy, MoveTemp(OnJobComplete));
-}
-
 ///////////////////////////////////////////////////////////////////////////////////////////////////////////////////////
 
 FBuildSession CreateBuildSession(IBuildSessionInternal* Session)
diff --git a/Engine/Source/Developer/DerivedDataCache/Public/DerivedDataBuildSession.h b/Engine/Source/Developer/DerivedDataCache/Public/DerivedDataBuildSession.h
--- a/Engine/Source/Developer/DerivedDataCache/Public/DerivedDataBuildSession.h
+++ b/Engine/Source/Developer/DerivedDataCache/Public/DerivedDataBuildSession.h
@@ -4,7 +4,6 @@
 
 #include "CoreTypes.h"
 #include "Containers/StringView.h"
-#include "DerivedDataBuildKey.h"
 #include "DerivedDataBuildTypes.h"
 #include "Templates/Function.h"
 #include "Templates/UniquePtr.h"
@@ -13,16 +12,8 @@ namespace UE::DerivedData { class FBuildAction; }
 namespace UE::DerivedData { class FBuildDefinition; }
 namespace UE::DerivedData { class FBuildSession; }
 namespace UE::DerivedData { class FOptionalBuildInputs; }
-namespace UE::DerivedData { class FValueWithId; }
 namespace UE::DerivedData { class IRequestOwner; }
-namespace UE::DerivedData { struct FBuildValueCompleteParams; }
-
-namespace UE::DerivedData
-{
-
-using FOnBuildValueComplete = TUniqueFunction<void (FBuildValueCompleteParams&& Params)>;
-
-} // UE::DerivedData
+namespace UE::DerivedData { struct FBuildKey; }
 
 namespace UE::DerivedData::Private
 {
@@ -32,6 +23,11 @@ class IBuildSessionInternal
 public:
 	virtual ~IBuildSessionInternal() = default;
 	virtual FStringView GetName() const = 0;
+	virtual void Build(
+		const FBuildKey& Key,
+		const FBuildPolicy& Policy,
+		IRequestOwner& Owner,
+		FOnBuildComplete&& OnComplete) = 0;
 	virtual void Build(
 		const FBuildDefinition& Definition,
 		const FOptionalBuildInputs& Inputs,
@@ -44,11 +40,6 @@ public:
 		const FBuildPolicy& Policy,
 		IRequestOwner& Owner,
 		FOnBuildComplete&& OnComplete) = 0;
-	virtual void BuildValue(
-		const FBuildValueKey& Value,
-		EBuildPolicy Policy,
-		IRequestOwner& Owner,
-		FOnBuildValueComplete&& OnComplete) = 0;
 };
 
 FBuildSession CreateBuildSession(IBuildSessionInternal* Session);
@@ -79,20 +70,18 @@ public:
 	 *
 	 * The callback will always be called, and may be called from an arbitrary thread.
 	 *
-	 * @param Definition   The build function to execute and references to its inputs.
-	 * @param Inputs       The build inputs referenced by the definition, if it has any. Optional.
+	 * @param Key          The key of the build definition to resolve and build.
 	 * @param Policy       Flags to control the behavior of the request. See FBuildPolicy.
 	 * @param Owner        The owner to execute the build within.
 	 * @param OnComplete   A callback invoked when the build completes or is canceled.
 	 */
 	inline void Build(
-		const FBuildDefinition& Definition,
-		const FOptionalBuildInputs& Inputs,
+		const FBuildKey& Key,
 		const FBuildPolicy& Policy,
 		IRequestOwner& Owner,
 		FOnBuildComplete&& OnComplete)
 	{
-		Session->Build(Definition, Inputs, Policy, Owner, MoveTemp(OnComplete));
+		Session->Build(Key, Policy, Owner, MoveTemp(OnComplete));
 	}
 
 	/**
@@ -100,39 +89,41 @@ public:
 	 *
 	 * The callback will always be called, and may be called from an arbitrary thread.
 	 *
-	 * @param Action       The build function to execute and references to its inputs.
-	 * @param Inputs       The build inputs referenced by the action, if it has any. Optional.
+	 * @param Definition   The build function to execute and references to its inputs.
+	 * @param Inputs       The build inputs referenced by the definition, if it has any. Optional.
 	 * @param Policy       Flags to control the behavior of the request. See FBuildPolicy.
 	 * @param Owner        The owner to execute the build within.
 	 * @param OnComplete   A callback invoked when the build completes or is canceled.
 	 */
 	inline void Build(
-		const FBuildAction& Action,
+		const FBuildDefinition& Definition,
 		const FOptionalBuildInputs& Inputs,
 		const FBuildPolicy& Policy,
 		IRequestOwner& Owner,
 		FOnBuildComplete&& OnComplete)
 	{
-		Session->Build(Action, Inputs, Policy, Owner, MoveTemp(OnComplete));
+		Session->Build(Definition, Inputs, Policy, Owner, MoveTemp(OnComplete));
 	}
 
 	/**
-	 * Asynchronous request to execute a build according to the policy and return one value.
+	 * Asynchronous request to execute a build according to the policy.
 	 *
 	 * The callback will always be called, and may be called from an arbitrary thread.
 	 *
-	 * @param Value        The key identifying the build definition and the value to return.
-	 * @param Policy       Flags to control the behavior of the request. See EBuildPolicy.
+	 * @param Action       The build function to execute and references to its inputs.
+	 * @param Inputs       The build inputs referenced by the action, if it has any. Optional.
+	 * @param Policy       Flags to control the behavior of the request. See FBuildPolicy.
 	 * @param Owner        The owner to execute the build within.
 	 * @param OnComplete   A callback invoked when the build completes or is canceled.
 	 */
-	inline void BuildValue(
-		const FBuildValueKey& Value,
-		EBuildPolicy Policy,
+	inline void Build(
+		const FBuildAction& Action,
+		const FOptionalBuildInputs& Inputs,
+		const FBuildPolicy& Policy,
 		IRequestOwner& Owner,
-		FOnBuildValueComplete&& OnComplete)
+		FOnBuildComplete&& OnComplete)
 	{
-		Session->BuildValue(Value, Policy, Owner, MoveTemp(OnComplete));
+		Session->Build(Action, Inputs, Policy, Owner, MoveTemp(OnComplete));
 	}
 
 private:
@@ -175,23 +166,4 @@ public:
 	inline void Reset() { *this = FOptionalBuildSession(); }
 };
 
-/** Parameters for the completion callback for build value requests. */
-struct FBuildValueCompleteParams
-{
-	/** Key for the build request that completed or was canceled. See Value for ID. */
-	FBuildKey Key;
-
-	/**
-	 * Value from the build value request that completed or was canceled.
-	 *
-	 * The ID is always populated.
-	 * The hash and size are populated when Status is Ok.
-	 * The data is populated when Status is Ok and the data was not skipped by the policy.
-	 */
-	FValueWithId&& Value;
-
-	/** Status of the build request. */
-	EStatus Status = EStatus::Error;
-};
-
 } // UE::DerivedData
