DDC: Moved the pak file backends out of their header

This hides implementation details and reduces the number of files that need to change when refactoring backends.

#rnx

diff --git a/Engine/Source/Developer/DerivedDataCache/Private/DerivedDataBackends.cpp b/Engine/Source/Developer/DerivedDataCache/Private/DerivedDataBackends.cpp
--- a/Engine/Source/Developer/DerivedDataCache/Private/DerivedDataBackends.cpp
+++ b/Engine/Source/Developer/DerivedDataCache/Private/DerivedDataBackends.cpp
@@ -63,6 +63,11 @@ namespace UE::DerivedData::CacheStore::Memory
 FFileBackedDerivedDataBackend* CreateMemoryDerivedDataBackend(const TCHAR* Name, int64 MaxCacheSize, bool bCanBeDisabled);
 } // UE::DerivedData::CacheStore::Memory
 
+namespace UE::DerivedData::CacheStore::PakFile
+{
+IPakFileDerivedDataBackend* CreatePakFileDerivedDataBackend(const TCHAR* Filename, bool bWriting, bool bCompressed);
+} // UE::DerivedData::CacheStore::PakFile
+
 namespace UE::DerivedData::CacheStore::S3
 {
 FDerivedDataBackendInterface* CreateS3DerivedDataBackend(const TCHAR* RootManifestPath, const TCHAR* BaseUrl, const TCHAR* Region, const TCHAR* CanaryObjectKey, const TCHAR* CachePath);
@@ -341,7 +346,7 @@ public:
 				FGuid Temp = FGuid::NewGuid();
 				ReadPakFilename = PakFilename;
 				WritePakFilename = PakFilename + TEXT(".") + Temp.ToString();
-				WritePakCache = bCompressed? new FCompressedPakFileDerivedDataBackend( *WritePakFilename, true ) : new FPakFileDerivedDataBackend( *WritePakFilename, true );
+				WritePakCache = CreatePakFileDerivedDataBackend(*WritePakFilename, /*bWriting*/ true, bCompressed);
 				PakNode = WritePakCache;
 			}
 			else
@@ -349,7 +354,7 @@ public:
 				bool bReadPak = FPlatformFileManager::Get().GetPlatformFile().FileExists( *PakFilename );
 				if( bReadPak )
 				{
-					FPakFileDerivedDataBackend* ReadPak = bCompressed? new FCompressedPakFileDerivedDataBackend( *PakFilename, false ) : new FPakFileDerivedDataBackend( *PakFilename, false );
+					IPakFileDerivedDataBackend* ReadPak = CreatePakFileDerivedDataBackend(*PakFilename, /*bWriting*/ false, bCompressed);
 					ReadPakFilename = PakFilename;
 					PakNode = ReadPak;
 					ReadPakCache.Add(ReadPak);
@@ -539,7 +544,7 @@ public:
 		FDerivedDataBackendInterface* Hierarchy = NULL;
 		if( InnerNodes.Num() > 1 )
 		{
-			FHierarchicalDerivedDataBackend* HierarchyBackend = new FHierarchicalDerivedDataBackend( InnerNodes );
+			CacheStore::Hierarchical::FHierarchicalDerivedDataBackend* HierarchyBackend = new CacheStore::Hierarchical::FHierarchicalDerivedDataBackend( InnerNodes );
 			Hierarchy = HierarchyBackend;
 			if (HierarchicalWrapper == NULL)
 			{
@@ -1090,8 +1095,9 @@ public:
 
 				for(const FString& MergePakName : MergePakList)
 				{
-					CacheStore::PakFile::FPakFileDerivedDataBackend ReadPak(*FPaths::Combine(*FPaths::GetPath(WritePakFilename), *MergePakName), false);
-					WritePakCache->MergeCache(&ReadPak);
+					TUniquePtr<CacheStore::PakFile::IPakFileDerivedDataBackend> ReadPak(
+						CacheStore::PakFile::CreatePakFileDerivedDataBackend(*FPaths::Combine(*FPaths::GetPath(WritePakFilename), *MergePakName), /*bWriting*/ false, /*bCompressed*/ false));
+					WritePakCache->MergeCache(ReadPak.Get());
 				}
 			}
 			for (int32 ReadPakIndex = 0; ReadPakIndex < ReadPakCache.Num(); ReadPakIndex++)
@@ -1115,7 +1121,7 @@ public:
 							UE_LOG(LogDerivedDataCache, Error, TEXT("Could not delete the pak file %s to overwrite it with a new one."), *ReadPakFilename);
 						}
 					}
-					if (!CacheStore::PakFile::FPakFileDerivedDataBackend::SortAndCopy(WritePakFilename, ReadPakFilename))
+					if (!CacheStore::PakFile::IPakFileDerivedDataBackend::SortAndCopy(WritePakFilename, ReadPakFilename))
 					{
 						UE_LOG(LogDerivedDataCache, Error, TEXT("Couldn't sort pak file (%s)"), *WritePakFilename);
 					}
@@ -1179,10 +1185,10 @@ public:
 	{
 		// Assumptions: there's at least one read-only pak backend in the hierarchy
 		// and its parent is a hierarchical backend.
-		CacheStore::PakFile::FPakFileDerivedDataBackend* ReadPak = NULL;
+		CacheStore::PakFile::IPakFileDerivedDataBackend* ReadPak = nullptr;
 		if (HierarchicalWrapper && FPlatformFileManager::Get().GetPlatformFile().FileExists(PakFilename))
 		{
-			ReadPak = new CacheStore::PakFile::FPakFileDerivedDataBackend(PakFilename, false);
+			ReadPak = CacheStore::PakFile::CreatePakFileDerivedDataBackend(PakFilename, /*bWriting*/ false, /*bCompressed*/ false);
 
 			HierarchicalWrapper->AddInnerBackend(ReadPak);
 			CreatedBackends.Add(ReadPak);
@@ -1200,7 +1206,7 @@ public:
 	{
 		for (int PakIndex = 0; PakIndex < ReadPakCache.Num(); ++PakIndex)
 		{
-			CacheStore::PakFile::FPakFileDerivedDataBackend* ReadPak = ReadPakCache[PakIndex];
+			CacheStore::PakFile::IPakFileDerivedDataBackend* ReadPak = ReadPakCache[PakIndex];
 			if (ReadPak->GetFilename() == PakFilename)
 			{
 				check(HierarchicalWrapper);
@@ -1277,12 +1283,12 @@ private:
 
 	/** Instances of backend interfaces which exist in only one copy */
 	FFileBackedDerivedDataBackend*	BootCache;
-	CacheStore::PakFile::FPakFileDerivedDataBackend* WritePakCache;
+	CacheStore::PakFile::IPakFileDerivedDataBackend* WritePakCache;
 	FDerivedDataBackendInterface*	AsyncPutWrapper;
 	FDerivedDataBackendInterface*	KeyLengthWrapper;
-	FHierarchicalDerivedDataBackend* HierarchicalWrapper;
+	CacheStore::Hierarchical::FHierarchicalDerivedDataBackend* HierarchicalWrapper;
 	/** Support for multiple read only pak files. */
-	TArray<CacheStore::PakFile::FPakFileDerivedDataBackend*> ReadPakCache;
+	TArray<CacheStore::PakFile::IPakFileDerivedDataBackend*> ReadPakCache;
 
 	/** List of directories used by the DDC */
 	TArray<FString> Directories;
diff --git a/Engine/Source/Developer/DerivedDataCache/Private/HierarchicalDerivedDataBackend.h b/Engine/Source/Developer/DerivedDataCache/Private/HierarchicalDerivedDataBackend.h
--- a/Engine/Source/Developer/DerivedDataCache/Private/HierarchicalDerivedDataBackend.h
+++ b/Engine/Source/Developer/DerivedDataCache/Private/HierarchicalDerivedDataBackend.h
@@ -21,12 +21,10 @@ extern bool GVerifyDDC;
 
 namespace UE::DerivedData::CacheStore::AsyncPut
 {
-
 FDerivedDataBackendInterface* CreateAsyncPutDerivedDataBackend(FDerivedDataBackendInterface* InnerBackend, bool bCacheInFlightPuts);
-
 } // UE::DerivedData::CacheStore::AsyncPut
 
-namespace UE::DerivedData::Backends
+namespace UE::DerivedData::CacheStore::Hierarchical
 {
 
 /** 
@@ -99,7 +97,7 @@ private:
 			// async puts to allow us to fill all levels without holding up the engine
 			// we need to cache inflight puts to avoid having inconsistent miss and redownload on lower cache levels while puts are still async
 			const bool bCacheInFlightPuts = true;
-			AsyncPutInnerBackends.Emplace(CacheStore::AsyncPut::CreateAsyncPutDerivedDataBackend(InnerBackends[CacheIndex], bCacheInFlightPuts));
+			AsyncPutInnerBackends.Emplace(AsyncPut::CreateAsyncPutDerivedDataBackend(InnerBackends[CacheIndex], bCacheInFlightPuts));
 		}
 	}
 
@@ -810,4 +808,4 @@ private:
 	bool bHasWritableRemoteBackends;
 };
 
-} // UE::DerivedData::Backends
+} // UE::DerivedData::CacheStore::Hierarchical
diff --git a/Engine/Source/Developer/DerivedDataCache/Private/PakFileDerivedDataBackend.cpp b/Engine/Source/Developer/DerivedDataCache/Private/PakFileDerivedDataBackend.cpp
--- a/Engine/Source/Developer/DerivedDataCache/Private/PakFileDerivedDataBackend.cpp
+++ b/Engine/Source/Developer/DerivedDataCache/Private/PakFileDerivedDataBackend.cpp
@@ -10,6 +10,7 @@
 #include "DerivedDataCacheUsageStats.h"
 #include "DerivedDataChunk.h"
 #include "DerivedDataValue.h"
+#include "HAL/CriticalSection.h"
 #include "HAL/FileManager.h"
 #include "HAL/PlatformFile.h"
 #include "HAL/PlatformFileManager.h"
@@ -19,16 +20,183 @@
 #include "Misc/Paths.h"
 #include "Misc/PathViews.h"
 #include "Misc/ScopeRWLock.h"
+#include "ProfilingDebugging/CookStats.h"
 #include "Serialization/CompactBinary.h"
 #include "Serialization/CompactBinaryPackage.h"
 #include "Serialization/CompactBinaryValidation.h"
 #include "Serialization/MemoryReader.h"
 #include "Serialization/MemoryWriter.h"
 #include "Templates/Greater.h"
+#include "Templates/UniquePtr.h"
 
 namespace UE::DerivedData::CacheStore::PakFile
 {
 
+/**
+ * A simple thread safe, pak file based backend.
+ */
+class FPakFileDerivedDataBackend : public IPakFileDerivedDataBackend
+{
+public:
+	FPakFileDerivedDataBackend(const TCHAR* InFilename, bool bInWriting);
+	~FPakFileDerivedDataBackend();
+
+	virtual void Close() override;
+
+	/** Return a name for this interface */
+	virtual FString GetName() const override;
+
+	/** return true if this cache is writable **/
+	virtual bool IsWritable() const override;
+
+	/** Returns a class of speed for this interface **/
+	virtual ESpeedClass GetSpeedClass() const override;
+
+	virtual bool BackfillLowerCacheLevels() const override;
+
+	/**
+	 * Synchronous test for the existence of a cache item
+	 *
+	 * @param	CacheKey	Alphanumeric+underscore key of this cache item
+	 * @return				true if the data probably will be found, this can't be guaranteed because of concurrency in the backends, corruption, etc
+	 */
+	virtual bool CachedDataProbablyExists(const TCHAR* CacheKey) override;
+
+	/**
+	 * Synchronous retrieve of a cache item
+	 *
+	 * @param	CacheKey	Alphanumeric+underscore key of this cache item
+	 * @param	OutData		Buffer to receive the results, if any were found
+	 * @return				true if any data was found, and in this case OutData is non-empty
+	 */
+	virtual bool GetCachedData(const TCHAR* CacheKey, TArray<uint8>& OutData) override;
+
+	/**
+	 * Asynchronous, fire-and-forget placement of a cache item
+	 *
+	 * @param	CacheKey	Alphanumeric+underscore key of this cache item
+	 * @param	InData		Buffer containing the data to cache, can be destroyed after the call returns, immediately
+	 * @param	bPutEvenIfExists	If true, then do not attempt skip the put even if CachedDataProbablyExists returns true
+	 */
+	virtual EPutStatus PutCachedData(const TCHAR* CacheKey, TArrayView<const uint8> InData, bool bPutEvenIfExists) override;
+	virtual void RemoveCachedData(const TCHAR* CacheKey, bool bTransient) override;
+
+	/**
+	 * Save the cache to disk
+	 * @return	true if file was saved successfully
+	 */
+	virtual bool SaveCache() override;
+
+	/**
+	 * Load the cache to disk	 * @param	Filename	Filename to load
+	 * @return	true if file was loaded successfully
+	 */
+	virtual bool LoadCache(const TCHAR* InFilename) override;
+
+	/**
+	 * Merges another cache file into this one.
+	 * @return true on success
+	 */
+	virtual void MergeCache(IPakFileDerivedDataBackend* OtherPak) override;
+	
+	virtual const FString& GetFilename() const override
+	{
+		return CachePath;
+	}
+
+	virtual TSharedRef<FDerivedDataCacheStatsNode> GatherUsageStats() const override;
+
+	virtual bool TryToPrefetch(TConstArrayView<FString> CacheKeys) override
+	{
+		return CachedDataProbablyExistsBatch(CacheKeys).CountSetBits() == CacheKeys.Num();
+	}
+
+	virtual bool WouldCache(const TCHAR* CacheKey, TArrayView<const uint8> InData) override { return true; }
+
+	virtual bool ApplyDebugOptions(FBackendDebugOptions& InOptions) override { return false; }
+
+	virtual void Put(
+		TConstArrayView<FCachePutRequest> Requests,
+		IRequestOwner& Owner,
+		FOnCachePutComplete&& OnComplete) override;
+
+	virtual void Get(
+		TConstArrayView<FCacheGetRequest> Requests,
+		IRequestOwner& Owner,
+		FOnCacheGetComplete&& OnComplete) override;
+
+	virtual void GetChunks(
+		TConstArrayView<FCacheChunkRequest> Requests,
+		IRequestOwner& Owner,
+		FOnCacheChunkComplete&& OnComplete) override;
+
+private:
+	bool PutCacheRecord(FStringView Name, const FCacheRecord& Record, const FCacheRecordPolicy& Policy);
+	bool PutCacheContent(FStringView Name, const FCompressedBuffer& Content);
+
+	FOptionalCacheRecord GetCacheRecordOnly(
+		FStringView Name,
+		const FCacheKey& Key,
+		const FCacheRecordPolicy& Policy);
+	FOptionalCacheRecord GetCacheRecord(
+		FStringView Name,
+		const FCacheKey& Key,
+		const FCacheRecordPolicy& Policy,
+		EStatus& OutStatus);
+	void GetCacheContent(
+		FStringView Name,
+		const FCacheKey& Key,
+		ECachePolicy Policy,
+		FValueWithId& InOutValue,
+		EStatus& InOutStatus);
+
+	bool SaveFile(FStringView Path, FStringView DebugName, TFunctionRef<void (FArchive&)> WriteFunction);
+	FSharedBuffer LoadFile(FStringView Path, FStringView DebugName);
+	bool FileExists(FStringView Path);
+
+private:
+	FDerivedDataCacheUsageStats UsageStats;
+
+	struct FCacheValue
+	{
+		int64 Offset;
+		int64 Size;
+		uint32 Crc;
+		FCacheValue(int64 InOffset, int64 InSize, uint32 InCrc)
+			: Offset(InOffset)
+			, Size(InSize)
+			, Crc(InCrc)
+		{
+		}
+	};
+
+	/** When set to true, we are a pak writer (we don't do reads). */
+	bool bWriting;
+	/** When set to true, we are a pak writer and we saved, so we shouldn't be used anymore. Also, a read cache that failed to open. */
+	bool bClosed;
+	/** Object used for synchronization via scoped read or write locks. */
+	FRWLock SynchronizationObject;
+	/** Set of files that are being written to disk asynchronously. */
+	TMap<FString, FCacheValue> CacheItems;
+	/** File handle of pak. */
+	TUniquePtr<IFileHandle> FileHandle;
+	/** File name of pak. */
+	FString CachePath;
+
+	/** Maximum total size of compressed data stored within a record package with multiple attachments. */
+	uint64 MaxRecordSizeKB = 256;
+	/** Maximum total size of compressed data stored within a value package, or a record package with one attachment. */
+	uint64 MaxValueSizeKB = 1024;
+
+	enum 
+	{
+		/** Magic number to use in header */
+		PakCache_Magic = 0x0c7c0ddc,
+	};
+
+	friend class IPakFileDerivedDataBackend;
+};
+
 FPakFileDerivedDataBackend::FPakFileDerivedDataBackend(const TCHAR* const InCachePath, const bool bInWriting)
 	: bWriting(bInWriting)
 	, bClosed(false)
@@ -405,8 +573,10 @@ bool FPakFileDerivedDataBackend::LoadCache(const TCHAR* InFilename)
 	return true;
 }
 
-void FPakFileDerivedDataBackend::MergeCache(FPakFileDerivedDataBackend* OtherPak)
+void FPakFileDerivedDataBackend::MergeCache(IPakFileDerivedDataBackend* OtherPakInterface)
 {
+	FPakFileDerivedDataBackend* OtherPak = static_cast<FPakFileDerivedDataBackend*>(OtherPakInterface);
+
 	// Get all the existing keys
 	TArray<FString> KeyNames;
 	OtherPak->CacheItems.GenerateKeyArray(KeyNames);
@@ -434,7 +604,7 @@ void FPakFileDerivedDataBackend::MergeCache(FPakFileDerivedDataBackend* OtherPak
 	}
 }
 
-bool FPakFileDerivedDataBackend::SortAndCopy(const FString &InputFilename, const FString &OutputFilename)
+bool IPakFileDerivedDataBackend::SortAndCopy(const FString &InputFilename, const FString &OutputFilename)
 {
 	// Open the input and output files
 	FPakFileDerivedDataBackend InputPak(*InputFilename, false);
@@ -1036,6 +1206,25 @@ bool FPakFileDerivedDataBackend::FileExists(const FStringView Path)
 	return CacheItems.ContainsByHash(PathHash, Path);
 }
 
+class FCompressedPakFileDerivedDataBackend : public FPakFileDerivedDataBackend
+{
+public:
+	FCompressedPakFileDerivedDataBackend(const TCHAR* InFilename, bool bInWriting);
+
+	virtual EPutStatus PutCachedData(const TCHAR* CacheKey, TArrayView<const uint8> InData, bool bPutEvenIfExists) override;
+	virtual bool GetCachedData(const TCHAR* CacheKey, TArray<uint8>& OutData) override;
+
+	/** Returns a class of speed for this interface **/
+	virtual ESpeedClass GetSpeedClass() const override
+	{
+		return ESpeedClass::Fast;
+	}
+
+private:
+	static const EName CompressionFormat = NAME_Zlib;
+	static const ECompressionFlags CompressionFlags = COMPRESS_BiasMemory;
+};
+
 FCompressedPakFileDerivedDataBackend::FCompressedPakFileDerivedDataBackend(const TCHAR* InFilename, bool bInWriting)
 	: FPakFileDerivedDataBackend(InFilename, bInWriting)
 {
@@ -1072,4 +1261,9 @@ bool FCompressedPakFileDerivedDataBackend::GetCachedData(const TCHAR* CacheKey,
 	return true;
 }
 
+IPakFileDerivedDataBackend* CreatePakFileDerivedDataBackend(const TCHAR* Filename, bool bWriting, bool bCompressed)
+{
+	return bCompressed ? new FCompressedPakFileDerivedDataBackend(Filename, bWriting) : new FPakFileDerivedDataBackend(Filename, bWriting);
+}
+
 } // UE::DerivedData::CacheStore::PakFile
diff --git a/Engine/Source/Developer/DerivedDataCache/Private/PakFileDerivedDataBackend.h b/Engine/Source/Developer/DerivedDataCache/Private/PakFileDerivedDataBackend.h
--- a/Engine/Source/Developer/DerivedDataCache/Private/PakFileDerivedDataBackend.h
+++ b/Engine/Source/Developer/DerivedDataCache/Private/PakFileDerivedDataBackend.h
@@ -5,202 +5,20 @@
 #include "CoreTypes.h"
 #include "Containers/StringFwd.h"
 #include "DerivedDataBackendInterface.h"
-#include "DerivedDataCacheUsageStats.h"
-#include "HAL/CriticalSection.h"
-#include "ProfilingDebugging/CookStats.h"
-#include "Templates/UniquePtr.h"
-
-class FCompressedBuffer;
-class IFileHandle;
-
-namespace UE::DerivedData { class FOptionalCacheRecord; }
-namespace UE::DerivedData { class FValueWithId; }
 
 namespace UE::DerivedData::CacheStore::PakFile
 {
 
-/** 
- * A simple thread safe, pak file based backend. 
-**/
-class FPakFileDerivedDataBackend : public FDerivedDataBackendInterface
+class IPakFileDerivedDataBackend : public FDerivedDataBackendInterface
 {
 public:
-	FPakFileDerivedDataBackend(const TCHAR* InFilename, bool bInWriting);
-	~FPakFileDerivedDataBackend();
-
-	void Close();
-
-	/** Return a name for this interface */
-	virtual FString GetName() const override;
-
-	/** return true if this cache is writable **/
-	virtual bool IsWritable() const override;
-
-	/** Returns a class of speed for this interface **/
-	virtual ESpeedClass GetSpeedClass() const override;
-
-	virtual bool BackfillLowerCacheLevels() const override;
-
-	/**
-	 * Synchronous test for the existence of a cache item
-	 *
-	 * @param	CacheKey	Alphanumeric+underscore key of this cache item
-	 * @return				true if the data probably will be found, this can't be guaranteed because of concurrency in the backends, corruption, etc
-	 */
-	virtual bool CachedDataProbablyExists(const TCHAR* CacheKey) override;
-
-	/**
-	 * Synchronous retrieve of a cache item
-	 *
-	 * @param	CacheKey	Alphanumeric+underscore key of this cache item
-	 * @param	OutData		Buffer to receive the results, if any were found
-	 * @return				true if any data was found, and in this case OutData is non-empty
-	 */
-	virtual bool GetCachedData(const TCHAR* CacheKey, TArray<uint8>& OutData) override;
-
-	/**
-	 * Asynchronous, fire-and-forget placement of a cache item
-	 *
-	 * @param	CacheKey	Alphanumeric+underscore key of this cache item
-	 * @param	InData		Buffer containing the data to cache, can be destroyed after the call returns, immediately
-	 * @param	bPutEvenIfExists	If true, then do not attempt skip the put even if CachedDataProbablyExists returns true
-	 */
-	virtual EPutStatus PutCachedData(const TCHAR* CacheKey, TArrayView<const uint8> InData, bool bPutEvenIfExists) override;
-	virtual void RemoveCachedData(const TCHAR* CacheKey, bool bTransient) override;
-
-	/**
-	 * Save the cache to disk
-	 * @return	true if file was saved sucessfully
-	 */
-	bool SaveCache();
-
-	/**
-	 * Load the cache to disk	 * @param	Filename	Filename to load
-	 * @return	true if file was loaded sucessfully
-	 */
-	bool LoadCache(const TCHAR* InFilename);
-
-	/**
-	 * Merges another cache file into this one.
-	 * @return true on success
-	 */
-	void MergeCache(FPakFileDerivedDataBackend* OtherPak);
-	
-	const FString& GetFilename() const
-	{
-		return CachePath;
-	}
-
-	static bool SortAndCopy(const FString &InputFilename, const FString &OutputFilename);
-
-	virtual TSharedRef<FDerivedDataCacheStatsNode> GatherUsageStats() const override;
-
-	virtual bool TryToPrefetch(TConstArrayView<FString> CacheKeys) override
-	{
-		return CachedDataProbablyExistsBatch(CacheKeys).CountSetBits() == CacheKeys.Num();
-	}
-
-	virtual bool WouldCache(const TCHAR* CacheKey, TArrayView<const uint8> InData) override { return true; }
-
-	virtual bool ApplyDebugOptions(FBackendDebugOptions& InOptions) override { return false; }
-
-	virtual void Put(
-		TConstArrayView<FCachePutRequest> Requests,
-		IRequestOwner& Owner,
-		FOnCachePutComplete&& OnComplete) override;
-
-	virtual void Get(
-		TConstArrayView<FCacheGetRequest> Requests,
-		IRequestOwner& Owner,
-		FOnCacheGetComplete&& OnComplete) override;
-
-	virtual void GetChunks(
-		TConstArrayView<FCacheChunkRequest> Requests,
-		IRequestOwner& Owner,
-		FOnCacheChunkComplete&& OnComplete) override;
-
-private:
-	bool PutCacheRecord(FStringView Name, const FCacheRecord& Record, const FCacheRecordPolicy& Policy);
-	bool PutCacheContent(FStringView Name, const FCompressedBuffer& Content);
-
-	FOptionalCacheRecord GetCacheRecordOnly(
-		FStringView Name,
-		const FCacheKey& Key,
-		const FCacheRecordPolicy& Policy);
-	FOptionalCacheRecord GetCacheRecord(
-		FStringView Name,
-		const FCacheKey& Key,
-		const FCacheRecordPolicy& Policy,
-		EStatus& OutStatus);
-	void GetCacheContent(
-		FStringView Name,
-		const FCacheKey& Key,
-		ECachePolicy Policy,
-		FValueWithId& InOutValue,
-		EStatus& InOutStatus);
-
-	bool SaveFile(FStringView Path, FStringView DebugName, TFunctionRef<void (FArchive&)> WriteFunction);
-	FSharedBuffer LoadFile(FStringView Path, FStringView DebugName);
-	bool FileExists(FStringView Path);
-
-private:
-	FDerivedDataCacheUsageStats UsageStats;
-
-	struct FCacheValue
-	{
-		int64 Offset;
-		int64 Size;
-		uint32 Crc;
-		FCacheValue(int64 InOffset, int64 InSize, uint32 InCrc)
-			: Offset(InOffset)
-			, Size(InSize)
-			, Crc(InCrc)
-		{
-		}
-	};
-
-	/** When set to true, we are a pak writer (we don't do reads). */
-	bool bWriting;
-	/** When set to true, we are a pak writer and we saved, so we shouldn't be used anymore. Also, a read cache that failed to open. */
-	bool bClosed;
-	/** Object used for synchronization via scoped read or write locks. */
-	FRWLock SynchronizationObject;
-	/** Set of files that are being written to disk asynchronously. */
-	TMap<FString, FCacheValue> CacheItems;
-	/** File handle of pak. */
-	TUniquePtr<IFileHandle> FileHandle;
-	/** File name of pak. */
-	FString CachePath;
-
-	/** Maximum total size of compressed data stored within a record package with multiple attachments. */
-	uint64 MaxRecordSizeKB = 256;
-	/** Maximum total size of compressed data stored within a value package, or a record package with one attachment. */
-	uint64 MaxValueSizeKB = 1024;
-
-	enum 
-	{
-		/** Magic number to use in header */
-		PakCache_Magic = 0x0c7c0ddc,
-	};
-};
-
-class FCompressedPakFileDerivedDataBackend : public FPakFileDerivedDataBackend
-{
-public:
-	FCompressedPakFileDerivedDataBackend(const TCHAR* InFilename, bool bInWriting);
-
-	virtual EPutStatus PutCachedData(const TCHAR* CacheKey, TArrayView<const uint8> InData, bool bPutEvenIfExists) override;
-	virtual bool GetCachedData(const TCHAR* CacheKey, TArray<uint8>& OutData) override;
-
-	/** Returns a class of speed for this interface **/
-	virtual ESpeedClass GetSpeedClass() const override
-	{
-		return ESpeedClass::Fast;
-	}
+	virtual void Close() = 0;
+	virtual bool SaveCache() = 0;
+	virtual bool LoadCache(const TCHAR* InFilename) = 0;
+	virtual void MergeCache(IPakFileDerivedDataBackend* OtherPak) = 0;
+	virtual const FString& GetFilename() const = 0;
 
-private:
-	static const EName CompressionFormat = NAME_Zlib;
-	static const ECompressionFlags CompressionFlags = COMPRESS_BiasMemory;
+	static bool SortAndCopy(const FString& InputFilename, const FString& OutputFilename);
 };
 
 } // UE::DerivedData::CacheStore::PakFile
