DDC: Rename CompleteParams to Response and use FSharedString with no alias

#rnx

diff --git a/Engine/Source/Developer/DerivedDataCache/Private/DerivedDataBackendThrottleWrapper.h b/Engine/Source/Developer/DerivedDataCache/Private/DerivedDataBackendThrottleWrapper.h
--- a/Engine/Source/Developer/DerivedDataCache/Private/DerivedDataBackendThrottleWrapper.h
+++ b/Engine/Source/Developer/DerivedDataCache/Private/DerivedDataBackendThrottleWrapper.h
@@ -183,13 +183,13 @@ public:
 		});
 
 		InnerBackend->Put(Requests, Owner,
-			[this, RecordSizes = MoveTemp(RecordSizes), State = EnterThrottlingScope(), OnComplete = MoveTemp(OnComplete)](FCachePutCompleteParams&& Params)
+			[this, RecordSizes = MoveTemp(RecordSizes), State = EnterThrottlingScope(), OnComplete = MoveTemp(OnComplete)](FCachePutResponse&& Response)
 			{
-				const FRecordSize* Size = Algo::FindBy(RecordSizes, Params.Key, &FRecordSize::Key);
+				const FRecordSize* Size = Algo::FindBy(RecordSizes, Response.Key, &FRecordSize::Key);
 				CloseThrottlingScope(State, FThrottlingState(this, Size ? Size->Size : 0));
 				if (OnComplete)
 				{
-					OnComplete(MoveTemp(Params));
+					OnComplete(MoveTemp(Response));
 				}
 			});
 	}
@@ -200,12 +200,12 @@ public:
 		FOnCacheGetComplete&& OnComplete) override
 	{
 		InnerBackend->Get(Requests, Owner,
-			[this, State = EnterThrottlingScope(), OnComplete = MoveTemp(OnComplete)](FCacheGetCompleteParams&& Params)
+			[this, State = EnterThrottlingScope(), OnComplete = MoveTemp(OnComplete)](FCacheGetResponse&& Response)
 			{
-				CloseThrottlingScope(State, FThrottlingState(this, Private::GetCacheRecordCompressedSize(Params.Record)));
+				CloseThrottlingScope(State, FThrottlingState(this, Private::GetCacheRecordCompressedSize(Response.Record)));
 				if (OnComplete)
 				{
-					OnComplete(MoveTemp(Params));
+					OnComplete(MoveTemp(Response));
 				}
 			});
 	}
@@ -216,12 +216,12 @@ public:
 		FOnCacheChunkComplete&& OnComplete) override
 	{
 		InnerBackend->GetChunks(Requests, Owner,
-			[this, State = EnterThrottlingScope(), OnComplete = MoveTemp(OnComplete)](FCacheChunkCompleteParams&& Params)
+			[this, State = EnterThrottlingScope(), OnComplete = MoveTemp(OnComplete)](FCacheChunkResponse&& Response)
 			{
-				CloseThrottlingScope(State, FThrottlingState(this, Params.RawData.GetSize()));
+				CloseThrottlingScope(State, FThrottlingState(this, Response.RawData.GetSize()));
 				if (OnComplete)
 				{
-					OnComplete(MoveTemp(Params));
+					OnComplete(MoveTemp(Response));
 				}
 			});
 	}
diff --git a/Engine/Source/Developer/DerivedDataCache/Private/DerivedDataBuildJob.cpp b/Engine/Source/Developer/DerivedDataCache/Private/DerivedDataBuildJob.cpp
--- a/Engine/Source/Developer/DerivedDataCache/Private/DerivedDataBuildJob.cpp
+++ b/Engine/Source/Developer/DerivedDataCache/Private/DerivedDataBuildJob.cpp
@@ -153,8 +153,8 @@ private:
 	void BeginExecuteRemote();
 	void BeginExecuteLocal();
 
-	void EndCacheQuery(FCacheGetCompleteParams&& Params);
-	void EndCacheStore(FCachePutCompleteParams&& Params);
+	void EndCacheQuery(FCacheGetResponse&& Response);
+	void EndCacheStore(FCachePutResponse&& Response);
 	void EndResolveKey(FBuildKeyResolvedParams&& Params);
 	void EndResolveInputMeta(FBuildInputMetaResolvedParams&& Params);
 	void EndResolveInputData(FBuildInputDataResolvedParams&& Params);
@@ -522,18 +522,18 @@ void FBuildJob::BeginCacheQuery()
 	TRACE_CPUPROFILER_EVENT_SCOPE(FBuildJob::CacheQuery);
 	EnumAddFlags(BuildStatus, EBuildStatus::CacheQuery);
 	Cache.Get({{Name, Context->GetCacheKey(), MakeCacheRecordQueryPolicy(BuildPolicy, *Context)}}, Owner,
-		[this](FCacheGetCompleteParams&& Params) { EndCacheQuery(MoveTemp(Params)); });
+		[this](FCacheGetResponse&& Response) { EndCacheQuery(MoveTemp(Response)); });
 }
 
-void FBuildJob::EndCacheQuery(FCacheGetCompleteParams&& Params)
+void FBuildJob::EndCacheQuery(FCacheGetResponse&& Response)
 {
-	if (Params.Status == EStatus::Canceled)
+	if (Response.Status == EStatus::Canceled)
 	{
 		return CompleteWithError("Build was canceled."_ASV);
 	}
-	if (Params.Status == EStatus::Ok)
+	if (Response.Status == EStatus::Ok)
 	{
-		if (FOptionalBuildOutput CacheOutput = FBuildOutput::Load(Name, FunctionName, Params.Record))
+		if (FOptionalBuildOutput CacheOutput = FBuildOutput::Load(Name, FunctionName, Response.Record))
 		{
 			EnumAddFlags(BuildStatus, EBuildStatus::CacheQueryHit);
 			return SetOutputNoCheck(MoveTemp(CacheOutput).Get());
@@ -588,12 +588,12 @@ void FBuildJob::BeginCacheStore()
 	FCacheRecordBuilder RecordBuilder(Context->GetCacheKey());
 	Output.Get().Save(RecordBuilder);
 	Cache.Put({{Name, RecordBuilder.Build(), MakeCacheStorePolicy(BuildPolicy.GetCombinedPolicy(), *Context)}}, Owner,
-		[this](FCachePutCompleteParams&& Params) { EndCacheStore(MoveTemp(Params)); });
+		[this](FCachePutResponse&& Response) { EndCacheStore(MoveTemp(Response)); });
 }
 
-void FBuildJob::EndCacheStore(FCachePutCompleteParams&& Params)
+void FBuildJob::EndCacheStore(FCachePutResponse&& Response)
 {
-	if (Params.Status == EStatus::Ok)
+	if (Response.Status == EStatus::Ok)
 	{
 		EnumAddFlags(BuildStatus, EBuildStatus::CacheStoreHit);
 	}
diff --git a/Engine/Source/Developer/DerivedDataCache/Private/DerivedDataLimitKeyLengthWrapper.h b/Engine/Source/Developer/DerivedDataCache/Private/DerivedDataLimitKeyLengthWrapper.h
--- a/Engine/Source/Developer/DerivedDataCache/Private/DerivedDataLimitKeyLengthWrapper.h
+++ b/Engine/Source/Developer/DerivedDataCache/Private/DerivedDataLimitKeyLengthWrapper.h
@@ -245,15 +245,15 @@ public:
 	{
 	#if ENABLE_COOK_STATS
 		return InnerBackend->Put(Requests, Owner,
-			[this, OnComplete = MoveTemp(OnComplete)](FCachePutCompleteParams&& Params)
+			[this, OnComplete = MoveTemp(OnComplete)](FCachePutResponse&& Response)
 			{
-				if (Params.Status == EStatus::Ok)
+				if (Response.Status == EStatus::Ok)
 				{
 					UsageStats.TimePut().AddHit(0);
 				}
 				if (OnComplete)
 				{
-					OnComplete(MoveTemp(Params));
+					OnComplete(MoveTemp(Response));
 				}
 			});
 	#else
@@ -268,15 +268,15 @@ public:
 	{
 	#if ENABLE_COOK_STATS
 		return InnerBackend->Get(Requests, Owner,
-			[this, OnComplete = MoveTemp(OnComplete)](FCacheGetCompleteParams&& Params)
+			[this, OnComplete = MoveTemp(OnComplete)](FCacheGetResponse&& Response)
 			{
-				if (Params.Status == EStatus::Ok)
+				if (Response.Status == EStatus::Ok)
 				{
 					UsageStats.TimeGet().AddHit(0);
 				}
 				if (OnComplete)
 				{
-					OnComplete(MoveTemp(Params));
+					OnComplete(MoveTemp(Response));
 				}
 			});
 	#else
@@ -291,15 +291,15 @@ public:
 	{
 	#if ENABLE_COOK_STATS
 		return InnerBackend->PutValue(Requests, Owner,
-			[this, OnComplete = MoveTemp(OnComplete)](FCachePutValueCompleteParams&& Params)
+			[this, OnComplete = MoveTemp(OnComplete)](FCachePutValueResponse&& Response)
 			{
-				if (Params.Status == EStatus::Ok)
+				if (Response.Status == EStatus::Ok)
 				{
 					UsageStats.TimePut().AddHit(0);
 				}
 				if (OnComplete)
 				{
-					OnComplete(MoveTemp(Params));
+					OnComplete(MoveTemp(Response));
 				}
 			});
 	#else
@@ -314,15 +314,15 @@ public:
 	{
 	#if ENABLE_COOK_STATS
 		return InnerBackend->GetValue(Requests, Owner,
-			[this, OnComplete = MoveTemp(OnComplete)](FCacheGetValueCompleteParams&& Params)
+			[this, OnComplete = MoveTemp(OnComplete)](FCacheGetValueResponse&& Response)
 			{
-				if (Params.Status == EStatus::Ok)
+				if (Response.Status == EStatus::Ok)
 				{
 					UsageStats.TimeGet().AddHit(0);
 				}
 				if (OnComplete)
 				{
-					OnComplete(MoveTemp(Params));
+					OnComplete(MoveTemp(Response));
 				}
 			});
 	#else
@@ -337,15 +337,15 @@ public:
 	{
 	#if ENABLE_COOK_STATS
 		return InnerBackend->GetChunks(Requests, Owner,
-			[this, OnComplete = MoveTemp(OnComplete)](FCacheChunkCompleteParams&& Params)
+			[this, OnComplete = MoveTemp(OnComplete)](FCacheChunkResponse&& Response)
 			{
-				if (Params.Status == EStatus::Ok)
+				if (Response.Status == EStatus::Ok)
 				{
 					UsageStats.TimeGet().AddHit(0);
 				}
 				if (OnComplete)
 				{
-					OnComplete(MoveTemp(Params));
+					OnComplete(MoveTemp(Response));
 				}
 			});
 	#else
diff --git a/Engine/Source/Developer/DerivedDataCache/Private/HierarchicalDerivedDataBackend.h b/Engine/Source/Developer/DerivedDataCache/Private/HierarchicalDerivedDataBackend.h
--- a/Engine/Source/Developer/DerivedDataCache/Private/HierarchicalDerivedDataBackend.h
+++ b/Engine/Source/Developer/DerivedDataCache/Private/HierarchicalDerivedDataBackend.h
@@ -499,15 +499,15 @@ public:
 					{
 						FRequestOwner BlockingOwner(EPriority::Blocking);
 						InnerBackends[PutCacheIndex]->Put(Requests, BlockingOwner,
-							[&OnComplete, &RequestsOk](FCachePutCompleteParams&& Params)
+							[&OnComplete, &RequestsOk](FCachePutResponse&& Response)
 							{
-								if (Params.Status == EStatus::Ok)
+								if (Response.Status == EStatus::Ok)
 								{
 									bool bIsAlreadyInSet = false;
-									RequestsOk.FindOrAdd(Params.Key, &bIsAlreadyInSet);
+									RequestsOk.FindOrAdd(Response.Key, &bIsAlreadyInSet);
 									if (OnComplete && !bIsAlreadyInSet)
 									{
-										OnComplete(MoveTemp(Params));
+										OnComplete(MoveTemp(Response));
 									}
 								}
 							});
@@ -551,9 +551,9 @@ public:
 				// Block on this because backends in this hierarchy are not expected to be asynchronous.
 				FRequestOwner BlockingOwner(EPriority::Blocking);
 				InnerBackends[GetCacheIndex]->Get(RemainingRequests, BlockingOwner,
-					[this, GetCacheIndex, &Owner, &OnComplete, &KeysOk](FCacheGetCompleteParams&& Params)
+					[this, GetCacheIndex, &Owner, &OnComplete, &KeysOk](FCacheGetResponse&& Response)
 					{
-						if (Params.Status == EStatus::Ok)
+						if (Response.Status == EStatus::Ok)
 						{
 							FRequestOwner AsyncOwner(FPlatformMath::Min(Owner.GetPriority(), EPriority::Highest));
 							FRequestBarrier AsyncBarrier(AsyncOwner);
@@ -562,14 +562,14 @@ public:
 							{
 								if (GetCacheIndex != FillCacheIndex)
 								{
-									AsyncPutInnerBackends[FillCacheIndex]->Put({{Params.Name, Params.Record, ECachePolicy::Default}}, AsyncOwner);
+									AsyncPutInnerBackends[FillCacheIndex]->Put({{Response.Name, Response.Record, ECachePolicy::Default}}, AsyncOwner);
 								}
 							}
 
-							KeysOk.Add(Params.Record.GetKey());
+							KeysOk.Add(Response.Record.GetKey());
 							if (OnComplete)
 							{
-								OnComplete(MoveTemp(Params));
+								OnComplete(MoveTemp(Response));
 							}
 						}
 					});
@@ -633,15 +633,15 @@ public:
 					{
 						FRequestOwner BlockingOwner(EPriority::Blocking);
 						InnerBackends[PutCacheIndex]->PutValue(Requests, BlockingOwner,
-							[&OnComplete, &RequestsOk](FCachePutValueCompleteParams&& Params)
+							[&OnComplete, &RequestsOk](FCachePutValueResponse&& Response)
 							{
-								if (Params.Status == EStatus::Ok)
+								if (Response.Status == EStatus::Ok)
 								{
 									bool bIsAlreadyInSet = false;
-									RequestsOk.FindOrAdd(Params.Key, &bIsAlreadyInSet);
+									RequestsOk.FindOrAdd(Response.Key, &bIsAlreadyInSet);
 									if (OnComplete && !bIsAlreadyInSet)
 									{
-										OnComplete(MoveTemp(Params));
+										OnComplete(MoveTemp(Response));
 									}
 								}
 							});
@@ -685,9 +685,9 @@ public:
 				// Block on this because backends in this hierarchy are not expected to be asynchronous.
 				FRequestOwner BlockingOwner(EPriority::Blocking);
 				InnerBackends[GetCacheIndex]->GetValue(RemainingRequests, BlockingOwner,
-					[this, GetCacheIndex, &Owner, &OnComplete, &KeysOk](FCacheGetValueCompleteParams&& Params)
+					[this, GetCacheIndex, &Owner, &OnComplete, &KeysOk](FCacheGetValueResponse&& Response)
 					{
-						if (Params.Status == EStatus::Ok)
+						if (Response.Status == EStatus::Ok)
 						{
 							FRequestOwner AsyncOwner(FPlatformMath::Min(Owner.GetPriority(), EPriority::Highest));
 							FRequestBarrier AsyncBarrier(AsyncOwner);
@@ -696,14 +696,14 @@ public:
 							{
 								if (GetCacheIndex != FillCacheIndex)
 								{
-									AsyncPutInnerBackends[FillCacheIndex]->PutValue({{Params.Name, Params.Key, Params.Value, ECachePolicy::Default}}, AsyncOwner);
+									AsyncPutInnerBackends[FillCacheIndex]->PutValue({{Response.Name, Response.Key, Response.Value, ECachePolicy::Default}}, AsyncOwner);
 								}
 							}
 
-							KeysOk.Add(Params.Key);
+							KeysOk.Add(Response.Key);
 							if (OnComplete)
 							{
-								OnComplete(MoveTemp(Params));
+								OnComplete(MoveTemp(Response));
 							}
 						}
 					});
@@ -759,19 +759,19 @@ public:
 				TArray<FCacheChunkRequest, TInlineAllocator<16>> ErrorChunks;
 				FRequestOwner BlockingOwner(EPriority::Blocking);
 				InnerBackend->GetChunks(RemainingRequests, BlockingOwner,
-					[&OnComplete, &RemainingRequests, &ErrorChunks](FCacheChunkCompleteParams&& Params)
+					[&OnComplete, &RemainingRequests, &ErrorChunks](FCacheChunkResponse&& Response)
 					{
-						if (Params.Status == EStatus::Error)
+						if (Response.Status == EStatus::Error)
 						{
-							const int32 ChunkIndex = Algo::BinarySearch(RemainingRequests, Params, TChunkLess());
+							const int32 ChunkIndex = Algo::BinarySearch(RemainingRequests, Response, TChunkLess());
 							checkf(ChunkIndex != INDEX_NONE, TEXT("Failed to find remaining chunk %s ")
 								TEXT(" with raw offset %") UINT64_FMT TEXT("."),
-								*WriteToString<96>(Params.Key, '/', Params.Id), Params.RawOffset);
+								*WriteToString<96>(Response.Key, '/', Response.Id), Response.RawOffset);
 							ErrorChunks.Add(RemainingRequests[ChunkIndex]);
 						}
 						else if (OnComplete)
 						{
-							OnComplete(MoveTemp(Params));
+							OnComplete(MoveTemp(Response));
 						}
 					});
 				BlockingOwner.Wait();
diff --git a/Engine/Source/Developer/DerivedDataCache/Private/Tests/HttpDerivedDataBackendTest.cpp b/Engine/Source/Developer/DerivedDataCache/Private/Tests/HttpDerivedDataBackendTest.cpp
--- a/Engine/Source/Developer/DerivedDataCache/Private/Tests/HttpDerivedDataBackendTest.cpp
+++ b/Engine/Source/Developer/DerivedDataCache/Private/Tests/HttpDerivedDataBackendTest.cpp
@@ -189,23 +189,23 @@ protected:
 		TArray<TOptional<FGetOutput>> GetOutputs;
 		GetOutputs.SetNum(Records.Num());
 		FRequestOwner RequestOwner(EPriority::Blocking);
-		TestBackend->Get(Requests, RequestOwner, [&GetOutputs](FCacheGetCompleteParams&& Params)
+		TestBackend->Get(Requests, RequestOwner, [&GetOutputs](FCacheGetResponse&& Response)
 			{
-				FCacheRecordBuilder RecordBuilder(Params.Record.GetKey());
+				FCacheRecordBuilder RecordBuilder(Response.Record.GetKey());
 
-				if (Params.Record.GetMeta())
+				if (Response.Record.GetMeta())
 				{
-					RecordBuilder.SetMeta(FCbObject::Clone(Params.Record.GetMeta()));
+					RecordBuilder.SetMeta(FCbObject::Clone(Response.Record.GetMeta()));
 				}
 
-				for (const FValueWithId& Value : Params.Record.GetValues())
+				for (const FValueWithId& Value : Response.Record.GetValues())
 				{
 					if (Value)
 					{
 						RecordBuilder.AddValue(Value);
 					}
 				}
-				GetOutputs[Params.UserData].Emplace(FGetOutput{ RecordBuilder.Build(), Params.Status });
+				GetOutputs[Response.UserData].Emplace(FGetOutput{ RecordBuilder.Build(), Response.Status });
 			});
 		RequestOwner.Wait();
 
@@ -361,10 +361,10 @@ TArray<UE::DerivedData::FCacheRecord> CreateTestCacheRecords(UE::DerivedData::IC
 	}
 
 	FRequestOwner Owner(EPriority::Blocking);
-	InTestBackend->Put(PutRequests, Owner, [&CacheRecords, &PutRequests] (FCachePutCompleteParams&& Params)
+	InTestBackend->Put(PutRequests, Owner, [&CacheRecords, &PutRequests] (FCachePutResponse&& Response)
 	{
-		check(Params.Status == EStatus::Ok);
-		CacheRecords.Add(PutRequests[Params.UserData].Record);
+		check(Response.Status == EStatus::Ok);
+		CacheRecords.Add(PutRequests[Response.UserData].Record);
 	});
 	Owner.Wait();
 
diff --git a/Engine/Source/Developer/DerivedDataCache/Public/DerivedDataCache.h b/Engine/Source/Developer/DerivedDataCache/Public/DerivedDataCache.h
--- a/Engine/Source/Developer/DerivedDataCache/Public/DerivedDataCache.h
+++ b/Engine/Source/Developer/DerivedDataCache/Public/DerivedDataCache.h
@@ -20,28 +20,28 @@
 
 namespace UE::DerivedData { class ICacheStoreMaintainer; }
 namespace UE::DerivedData { class IRequestOwner; }
-namespace UE::DerivedData { struct FCacheChunkCompleteParams; }
 namespace UE::DerivedData { struct FCacheChunkRequest; }
-namespace UE::DerivedData { struct FCacheGetCompleteParams; }
+namespace UE::DerivedData { struct FCacheChunkResponse; }
 namespace UE::DerivedData { struct FCacheGetRequest; }
-namespace UE::DerivedData { struct FCacheGetValueCompleteParams; }
+namespace UE::DerivedData { struct FCacheGetResponse; }
 namespace UE::DerivedData { struct FCacheGetValueRequest; }
-namespace UE::DerivedData { struct FCachePutCompleteParams; }
+namespace UE::DerivedData { struct FCacheGetValueResponse; }
 namespace UE::DerivedData { struct FCachePutRequest; }
-namespace UE::DerivedData { struct FCachePutValueCompleteParams; }
+namespace UE::DerivedData { struct FCachePutResponse; }
 namespace UE::DerivedData { struct FCachePutValueRequest; }
+namespace UE::DerivedData { struct FCachePutValueResponse; }
 namespace UE::DerivedData::Private { class ICacheRecordPolicyShared; }
 
 namespace UE::DerivedData
 {
 
-using FCacheRequestName = FSharedString;
+using FCacheRequestName UE_DEPRECATED(5.0, "Use FSharedString.") = FSharedString;
 
-using FOnCachePutComplete = TUniqueFunction<void (FCachePutCompleteParams&& Params)>;
-using FOnCacheGetComplete = TUniqueFunction<void (FCacheGetCompleteParams&& Params)>;
-using FOnCachePutValueComplete = TUniqueFunction<void (FCachePutValueCompleteParams&& Params)>;
-using FOnCacheGetValueComplete = TUniqueFunction<void (FCacheGetValueCompleteParams&& Params)>;
-using FOnCacheChunkComplete = TUniqueFunction<void (FCacheChunkCompleteParams&& Params)>;
+using FOnCachePutComplete = TUniqueFunction<void (FCachePutResponse&& Response)>;
+using FOnCacheGetComplete = TUniqueFunction<void (FCacheGetResponse&& Response)>;
+using FOnCachePutValueComplete = TUniqueFunction<void (FCachePutValueResponse&& Response)>;
+using FOnCacheGetValueComplete = TUniqueFunction<void (FCacheGetValueResponse&& Response)>;
+using FOnCacheChunkComplete = TUniqueFunction<void (FCacheChunkResponse&& Response)>;
 
 ///////////////////////////////////////////////////////////////////////////////////////////////////////////////////////
 
@@ -334,7 +334,7 @@ public:
 struct FCachePutRequest
 {
 	/** A name to identify this request for logging and profiling. An object path is typically sufficient. */
-	FCacheRequestName Name;
+	FSharedString Name;
 
 	/** A record to store. */
 	FCacheRecord Record;
@@ -347,10 +347,10 @@ struct FCachePutRequest
 };
 
 /** Parameters for the completion callback for cache put requests. */
-struct FCachePutCompleteParams
+struct FCachePutResponse
 {
 	/** A copy of the name from the request. */
-	FCacheRequestName Name;
+	FSharedString Name;
 
 	/** A copy of the key from the request. */
 	FCacheKey Key;
@@ -362,11 +362,13 @@ struct FCachePutCompleteParams
 	EStatus Status = EStatus::Error;
 };
 
+using FCachePutCompleteParams UE_DEPRECATED(5.0, "Use FCachePutResponse.") = FCachePutResponse;
+
 /** Parameters to request to get a cache record. */
 struct FCacheGetRequest
 {
 	/** A name to identify this request for logging and profiling. An object path is typically sufficient. */
-	FCacheRequestName Name;
+	FSharedString Name;
 
 	/** A key identifying the record to fetch. */
 	FCacheKey Key;
@@ -379,10 +381,10 @@ struct FCacheGetRequest
 };
 
 /** Parameters for the completion callback for cache get requests. */
-struct FCacheGetCompleteParams
+struct FCacheGetResponse
 {
 	/** A copy of the name from the request. */
-	FCacheRequestName Name;
+	FSharedString Name;
 
 	/**
 	 * Record for the request that completed or was canceled.
@@ -401,11 +403,13 @@ struct FCacheGetCompleteParams
 	EStatus Status = EStatus::Error;
 };
 
+using FCacheGetCompleteParams UE_DEPRECATED(5.0, "Use FCacheGetResponse.") = FCacheGetResponse;
+
 /** Parameters to request to put a cache value. */
 struct FCachePutValueRequest
 {
 	/** A name to identify this request for logging and profiling. An object path is typically sufficient. */
-	FCacheRequestName Name;
+	FSharedString Name;
 
 	/** A key that will uniquely identify the value in the cache. */
 	FCacheKey Key;
@@ -421,10 +425,10 @@ struct FCachePutValueRequest
 };
 
 /** Parameters for the completion callback for cache value put requests. */
-struct FCachePutValueCompleteParams
+struct FCachePutValueResponse
 {
 	/** A copy of the name from the request. */
-	FCacheRequestName Name;
+	FSharedString Name;
 
 	/** A copy of the key from the request. */
 	FCacheKey Key;
@@ -436,11 +440,13 @@ struct FCachePutValueCompleteParams
 	EStatus Status = EStatus::Error;
 };
 
+using FCachePutValueCompleteParams UE_DEPRECATED(5.0, "Use FCachePutValueResponse.") = FCachePutValueResponse;
+
 /** Parameters to request to get a cache value. */
 struct FCacheGetValueRequest
 {
 	/** A name to identify this request for logging and profiling. An object path is typically sufficient. */
-	FCacheRequestName Name;
+	FSharedString Name;
 
 	/** A key identifying the value to fetch. */
 	FCacheKey Key;
@@ -453,10 +459,10 @@ struct FCacheGetValueRequest
 };
 
 /** Parameters for the completion callback for cache value get requests. */
-struct FCacheGetValueCompleteParams
+struct FCacheGetValueResponse
 {
 	/** A copy of the name from the request. */
-	FCacheRequestName Name;
+	FSharedString Name;
 
 	/** A copy of the key from the request. */
 	FCacheKey Key;
@@ -476,11 +482,13 @@ struct FCacheGetValueCompleteParams
 	EStatus Status = EStatus::Error;
 };
 
+using FCacheGetValueCompleteParams UE_DEPRECATED(5.0, "Use FCacheGetValueResponse.") = FCacheGetValueResponse;
+
 /** Parameters to request a chunk, which is a subset of a value, from a cache record or cache value. */
 struct FCacheChunkRequest
 {
 	/** A name to identify this request for logging and profiling. An object path is typically sufficient. */
-	FCacheRequestName Name;
+	FSharedString Name;
 
 	/** A key identifying the record or value to fetch the chunk from. */
 	FCacheKey Key;
@@ -505,10 +513,10 @@ struct FCacheChunkRequest
 };
 
 /** Parameters for the completion callback for cache chunk requests. */
-struct FCacheChunkCompleteParams
+struct FCacheChunkResponse
 {
 	/** A copy of the name from the request. */
-	FCacheRequestName Name;
+	FSharedString Name;
 
 	/** A copy of the key from the request. */
 	FCacheKey Key;
@@ -535,6 +543,8 @@ struct FCacheChunkCompleteParams
 	EStatus Status = EStatus::Error;
 };
 
+using FCacheChunkCompleteParams UE_DEPRECATED(5.0, "Use FCacheChunkResponse.") = FCacheChunkResponse;
+
 ///////////////////////////////////////////////////////////////////////////////////////////////////////////////////////
 
 /** Returns a reference to the cache. Asserts if not available. */
diff --git a/Engine/Source/Developer/Virtualization/Private/VirtualizationDDCBackend.cpp b/Engine/Source/Developer/Virtualization/Private/VirtualizationDDCBackend.cpp
--- a/Engine/Source/Developer/Virtualization/Private/VirtualizationDDCBackend.cpp
+++ b/Engine/Source/Developer/Virtualization/Private/VirtualizationDDCBackend.cpp
@@ -95,10 +95,10 @@ EPushResult FDDCBackend::PushData(const FPayloadId& Id, const FCompressedBuffer&
 
 	UE::DerivedData::FRequestOwner Owner(UE::DerivedData::EPriority::Blocking);
 
-	UE::DerivedData::FCachePutCompleteParams Result;
-	auto Callback = [&Result](UE::DerivedData::FCachePutCompleteParams&& Params)
+	UE::DerivedData::FCachePutResponse Result;
+	auto Callback = [&Result](UE::DerivedData::FCachePutResponse&& Response)
 	{
-		Result = Params;
+		Result = Response;
 	};
 
 	// TODO: Improve the name when we start passing more context to this function
@@ -131,12 +131,12 @@ FCompressedBuffer FDDCBackend::PullData(const FPayloadId& Id)
 	FCompressedBuffer ResultData;
 	UE::DerivedData::EStatus ResultStatus;
 
-	auto Callback = [&Id, &ResultData, &ResultStatus](UE::DerivedData::FCacheGetCompleteParams&& Params)
+	auto Callback = [&Id, &ResultData, &ResultStatus](UE::DerivedData::FCacheGetResponse&& Response)
 	{
-		ResultStatus = Params.Status;
+		ResultStatus = Response.Status;
 		if (ResultStatus == UE::DerivedData::EStatus::Ok)
 		{
-			ResultData = Params.Record.GetValue(ToDerivedDataValueId(Id)).GetData();
+			ResultData = Response.Record.GetValue(ToDerivedDataValueId(Id)).GetData();
 		}
 	};
 
@@ -161,9 +161,9 @@ bool FDDCBackend::DoesPayloadExist(const FPayloadId& Id)
 	UE::DerivedData::FRequestOwner Owner(UE::DerivedData::EPriority::Blocking);
 	
 	UE::DerivedData::EStatus ResultStatus;
-	auto Callback = [&ResultStatus](UE::DerivedData::FCacheGetCompleteParams&& Params)
+	auto Callback = [&ResultStatus](UE::DerivedData::FCacheGetResponse&& Response)
 	{
-		ResultStatus = Params.Status;
+		ResultStatus = Response.Status;
 	};
 
 	// TODO: Improve the name when we start passing more context to this function
diff --git a/Engine/Source/Editor/UnrealEd/Private/EditorDomain/EditorDomain.cpp b/Engine/Source/Editor/UnrealEd/Private/EditorDomain/EditorDomain.cpp
--- a/Engine/Source/Editor/UnrealEd/Private/EditorDomain/EditorDomain.cpp
+++ b/Engine/Source/Editor/UnrealEd/Private/EditorDomain/EditorDomain.cpp
@@ -292,13 +292,13 @@ int64 FEditorDomain::FileSize(const FPackagePath& PackagePath, EPackageSegment P
 
 		auto MetaDataGetComplete =
 			[&FileSize, &PackageSource, &PackagePath, PackageSegment, Locks=this->Locks, OutUpdatedPath]
-			(UE::DerivedData::FCacheGetCompleteParams&& Params)
+			(UE::DerivedData::FCacheGetResponse&& Response)
 		{
 			FScopeLock ScopeLock(&Locks->Lock);
 			if ((PackageSource->Source == FEditorDomain::EPackageSource::Undecided || PackageSource->Source == FEditorDomain::EPackageSource::Editor) &&
-				Params.Status == UE::DerivedData::EStatus::Ok)
+				Response.Status == UE::DerivedData::EStatus::Ok)
 			{
-				const FCbObject& MetaData = Params.Record.GetMeta();
+				const FCbObject& MetaData = Response.Record.GetMeta();
 				FileSize = MetaData["FileSize"].AsInt64();
 				PackageSource->Source = EPackageSource::Editor;
 			}
@@ -368,10 +368,10 @@ FOpenPackageResult FEditorDomain::OpenReadPackage(const FPackagePath& PackagePat
 	ECachePolicy SkipFlags = ECachePolicy::SkipData & ~ECachePolicy::SkipMeta;
 	RequestEditorDomainPackage(PackagePath, PackageEditorHash,
 		SkipFlags, Result->GetRequestOwner(),
-		[Result](FCacheGetCompleteParams&& Params)
+		[Result](FCacheGetResponse&& Response)
 		{
 			// Note that ~FEditorDomainReadArchive waits for this callback to be called, so Result cannot dangle
-			Result->OnRecordRequestComplete(MoveTemp(Params));
+			Result->OnRecordRequestComplete(MoveTemp(Response));
 		});
 
 	// Precache the exports segment
@@ -426,10 +426,10 @@ FOpenAsyncPackageResult FEditorDomain::OpenAsyncReadPackage(const FPackagePath&
 	// Fetch meta-data only in the initial request
 	ECachePolicy SkipFlags = ECachePolicy::SkipData & ~ECachePolicy::SkipMeta;
 	RequestEditorDomainPackage(PackagePath, EditorDomainHash, SkipFlags, Result->GetRequestOwner(),
-		[Result](FCacheGetCompleteParams&& Params)
+		[Result](FCacheGetResponse&& Response)
 		{
 			// Note that ~FEditorDomainAsyncReadFileHandle waits for this callback to be called, so Result cannot dangle
-			Result->OnRecordRequestComplete(MoveTemp(Params));
+			Result->OnRecordRequestComplete(MoveTemp(Response));
 		});
 
 	const EPackageFormat Format = bHasEditorSource ? EPackageFormat::Binary : Result->GetPackageFormat();
@@ -632,12 +632,12 @@ void FEditorDomain::BatchDownload(TArrayView<FName> PackageNames)
 	if (!CacheRequests.IsEmpty())
 	{
 		FRequestBarrier Barrier(*BatchDownloadOwner);
-		GetCache().Get(CacheRequests, *BatchDownloadOwner, [this](FCacheGetCompleteParams&& Params)
+		GetCache().Get(CacheRequests, *BatchDownloadOwner, [this](FCacheGetResponse&& Response)
 			{
 				FScopeLock ScopeLock(&Locks->Lock);
 				TRefCountPtr<FPackageSource> PackageSource;
 				FPackageDigest ErrorDigest;
-				FName PackageName = FName(*Params.Name);
+				FName PackageName = FName(*Response.Name);
 				if (TryFindOrAddPackageSource(PackageName, PackageSource, &ErrorDigest))
 				{
 					PackageSource->bHasQueriedCatalog = true;
diff --git a/Engine/Source/Editor/UnrealEd/Private/EditorDomain/EditorDomainArchive.cpp b/Engine/Source/Editor/UnrealEd/Private/EditorDomain/EditorDomainArchive.cpp
--- a/Engine/Source/Editor/UnrealEd/Private/EditorDomain/EditorDomainArchive.cpp
+++ b/Engine/Source/Editor/UnrealEd/Private/EditorDomain/EditorDomainArchive.cpp
@@ -259,7 +259,7 @@ FEditorDomainPackageSegments::FSegment* FEditorDomainPackageSegments::EnsureSegm
 	return StartSegment;
 }
 
-void FEditorDomainPackageSegments::OnRecordRequestComplete(UE::DerivedData::FCacheGetCompleteParams&& Params,
+void FEditorDomainPackageSegments::OnRecordRequestComplete(UE::DerivedData::FCacheGetResponse&& Response,
 	TUniqueFunction<void(bool bValid)>&& CreateSegmentData,
 	TUniqueFunction<bool(FEditorDomain& EditorDomain)>&& TryCreateFallbackData)
 {
@@ -278,9 +278,9 @@ void FEditorDomainPackageSegments::OnRecordRequestComplete(UE::DerivedData::FCac
 
 	ESource NewAsyncSource = ESource::Uninitialized;
 	int32 InitialRequestIndex = -1;
-	if (Params.Status == EStatus::Ok)
+	if (Response.Status == EStatus::Ok)
 	{
-		const FCacheRecord& Record = Params.Record;
+		const FCacheRecord& Record = Response.Record;
 		const uint64 FileSize = Record.GetMeta()["FileSize"].AsUInt64();
 
 		TConstArrayView<FValueWithId> Values = Record.GetValues();
@@ -396,7 +396,7 @@ void FEditorDomainPackageSegments::SendSegmentRequest(FSegment& Segment)
 	FCacheChunkRequest SegmentChunk{{PackagePath.GetDebugName()}, UE::EditorDomain::GetEditorDomainPackageKey(EditorDomainHash), Segment.ValueId};
 	SegmentChunk.Policy = ECachePolicy::Local;
 	Cache.GetChunks({SegmentChunk}, *Segment.RequestOwner,
-		[this, &Segment](FCacheChunkCompleteParams&& Params)
+		[this, &Segment](FCacheChunkResponse&& Response)
 		{
 			if (AsyncSource == ESource::Closed)
 			{
@@ -408,13 +408,13 @@ void FEditorDomainPackageSegments::SendSegmentRequest(FSegment& Segment)
 				return;
 			}
 
-			if (Params.Status != EStatus::Ok)
+			if (Response.Status != EStatus::Ok)
 			{
 				UE_LOG(LogEditorDomain, Error, TEXT("Package %s is missing cache data in segment %d."), *PackagePath.GetDebugName(), (int32)(&Segment - &Segments[0]));
 			}
 			else
 			{
-				Segment.Data = MoveTemp(Params.RawData);
+				Segment.Data = MoveTemp(Response.RawData);
 				uint64 SegmentEnd = &Segment < &Segments.Last() ? (&Segment + 1)->Start : Size;
 				uint64 SegmentSize = SegmentEnd - Segment.Start;
 				if (Segment.Data.GetSize() != SegmentSize)
@@ -468,9 +468,9 @@ void FEditorDomainReadArchive::WaitForReady() const
 	}
 }
 
-void FEditorDomainReadArchive::OnRecordRequestComplete(UE::DerivedData::FCacheGetCompleteParams&& Params)
+void FEditorDomainReadArchive::OnRecordRequestComplete(UE::DerivedData::FCacheGetResponse&& Response)
 {
-	Segments.OnRecordRequestComplete(MoveTemp(Params),
+	Segments.OnRecordRequestComplete(MoveTemp(Response),
 		[this](bool bValid) { CreateSegmentData(bValid); },
 		[this](FEditorDomain& EditorDomain) { return TryCreateFallbackData(EditorDomain); }
 	);
@@ -731,14 +731,14 @@ protected:
 
 FEditorDomainAsyncReadFileHandle::FEditorDomainAsyncReadFileHandle(const TRefCountPtr<FEditorDomain::FLocks>& InLocks,
 	const FPackagePath& InPackagePath, const TRefCountPtr<FEditorDomain::FPackageSource>& InPackageSource,
-	UE::DerivedData::EPriority Priority)
-	: Segments(InLocks, InPackagePath, InPackageSource, Priority)
+	UE::DerivedData::EPriority InPriority)
+	: Segments(InLocks, InPackagePath, InPackageSource, InPriority)
 {
 }
 
-void FEditorDomainAsyncReadFileHandle::OnRecordRequestComplete(UE::DerivedData::FCacheGetCompleteParams&& Params)
+void FEditorDomainAsyncReadFileHandle::OnRecordRequestComplete(UE::DerivedData::FCacheGetResponse&& Response)
 {
-	Segments.OnRecordRequestComplete(MoveTemp(Params),
+	Segments.OnRecordRequestComplete(MoveTemp(Response),
 		[this](bool bValid) { CreateSegmentData(bValid); },
 		[this](FEditorDomain& EditorDomain) { return TryCreateFallbackData(EditorDomain); }
 	);
diff --git a/Engine/Source/Editor/UnrealEd/Private/EditorDomain/EditorDomainArchive.h b/Engine/Source/Editor/UnrealEd/Private/EditorDomain/EditorDomainArchive.h
--- a/Engine/Source/Editor/UnrealEd/Private/EditorDomain/EditorDomainArchive.h
+++ b/Engine/Source/Editor/UnrealEd/Private/EditorDomain/EditorDomainArchive.h
@@ -19,7 +19,7 @@
 #include "UObject/PackageResourceManager.h"
 
 class FAssetPackageData;
-namespace UE { namespace DerivedData { struct FCacheGetCompleteParams; } }
+namespace UE { namespace DerivedData { struct FCacheGetResponse; } }
 
 /**
  * Helper class for archive classes that server either an EditorDomain version of a package
@@ -47,7 +47,7 @@ public:
 	/** Get the request owner for requests that will feed this archive. */
 	UE::DerivedData::IRequestOwner& GetRequestOwner() { return RequestOwner; }
 	/** Callback from the request of the Record; set whether we're reading from EditorDomain bytes or WorkspaceDomain archive. */
-	void OnRecordRequestComplete(UE::DerivedData::FCacheGetCompleteParams&& Params,
+	void OnRecordRequestComplete(UE::DerivedData::FCacheGetResponse&& Response,
 		TUniqueFunction<void(bool bValid)>&& CreateSegmentData,
 		TUniqueFunction<bool(FEditorDomain& EditorDomain)>&& TryCreateFallbackData);
 
@@ -204,7 +204,7 @@ public:
 	/** Get the request owner for requests that will feed this archive. */
 	UE::DerivedData::IRequestOwner& GetRequestOwner() { return Segments.GetRequestOwner(); }
 	/** Callback from the request of the Record; set whether we're reading from EditorDomain bytes or WorkspaceDomain archive. */
-	void OnRecordRequestComplete(UE::DerivedData::FCacheGetCompleteParams&& Params);
+	void OnRecordRequestComplete(UE::DerivedData::FCacheGetResponse&& Response);
 	/** Get the PackageFormat, which depends on the domain the data is read from. */
 	EPackageFormat GetPackageFormat() const;
 	FEditorDomain::EPackageSource GetPackageSource() const;
@@ -244,7 +244,7 @@ public:
 	/** Get the request owner for requests that will feed this archive. */
 	UE::DerivedData::IRequestOwner& GetRequestOwner() { return Segments.GetRequestOwner(); }
 	/** Callback from the request of the Record; set whether we're reading from EditorDomain bytes or WorkspaceDomain archive. */
-	void OnRecordRequestComplete(UE::DerivedData::FCacheGetCompleteParams&& Params);
+	void OnRecordRequestComplete(UE::DerivedData::FCacheGetResponse&& Response);
 	/** Get the PackageFormat, which depends on the domain the data is read from. */
 	EPackageFormat GetPackageFormat() const;
 	FEditorDomain::EPackageSource GetPackageSource() const;
diff --git a/Engine/Source/Editor/UnrealEd/Private/EditorDomain/EditorDomainUtils.cpp b/Engine/Source/Editor/UnrealEd/Private/EditorDomain/EditorDomainUtils.cpp
--- a/Engine/Source/Editor/UnrealEd/Private/EditorDomain/EditorDomainUtils.cpp
+++ b/Engine/Source/Editor/UnrealEd/Private/EditorDomain/EditorDomainUtils.cpp
@@ -1372,10 +1372,10 @@ void GetBulkDataList(FName PackageName, UE::DerivedData::IRequestOwner& Owner, T
 	using namespace UE::DerivedData;
 	ICache& Cache = GetCache();
 	Cache.Get({{{WriteToString<128>(PackageName)}, GetBulkDataListKey(PackageDigest.Hash)}}, Owner,
-		[InnerCallback = MoveTemp(Callback)](FCacheGetCompleteParams&& Params)
+		[InnerCallback = MoveTemp(Callback)](FCacheGetResponse&& Response)
 		{
-			bool bOk = Params.Status == EStatus::Ok;
-			InnerCallback(bOk ? Params.Record.GetValue(GetBulkDataValueId()).GetData().Decompress() : FSharedBuffer());
+			bool bOk = Response.Status == EStatus::Ok;
+			InnerCallback(bOk ? Response.Record.GetValue(GetBulkDataValueId()).GetData().Decompress() : FSharedBuffer());
 		});
 }
 
@@ -1443,10 +1443,10 @@ void GetBulkDataPayloadId(FName PackageName, const FGuid& BulkDataId, UE::Derive
 	ICache& Cache = GetCache();
 	Cache.Get({{{WriteToString<192>(PackageName, TEXT("/"), BulkDataId)}, GetBulkDataPayloadIdKey(PackageAndGuidHash)}},
 		Owner,
-		[InnerCallback = MoveTemp(Callback)](FCacheGetCompleteParams&& Params)
+		[InnerCallback = MoveTemp(Callback)](FCacheGetResponse&& Response)
 		{
-			bool bOk = Params.Status == EStatus::Ok;
-			InnerCallback(bOk ? Params.Record.GetValue(GetBulkDataValueId()).GetData().Decompress() : FSharedBuffer());
+			bool bOk = Response.Status == EStatus::Ok;
+			InnerCallback(bOk ? Response.Record.GetValue(GetBulkDataValueId()).GetData().Decompress() : FSharedBuffer());
 		});
 }
 
diff --git a/Engine/Source/Runtime/Engine/Private/Streaming/Texture2DMipDataProvider_DDC.cpp b/Engine/Source/Runtime/Engine/Private/Streaming/Texture2DMipDataProvider_DDC.cpp
--- a/Engine/Source/Runtime/Engine/Private/Streaming/Texture2DMipDataProvider_DDC.cpp
+++ b/Engine/Source/Runtime/Engine/Private/Streaming/Texture2DMipDataProvider_DDC.cpp
@@ -79,13 +79,13 @@ void FTexture2DMipDataProvider_DDC::Init(const FTextureUpdateContext& Context, c
 
 				if (MipKeys.Num())
 				{
-					GetCache().GetChunks(MipKeys, DDCRequestOwner, [this](FCacheChunkCompleteParams&& Params)
+					GetCache().GetChunks(MipKeys, DDCRequestOwner, [this](FCacheChunkResponse&& Response)
 					{
-						if (Params.Status == EStatus::Ok)
+						if (Response.Status == EStatus::Ok)
 						{
-							const int32 MipIndex = int32(Params.UserData);
+							const int32 MipIndex = int32(Response.UserData);
 							check(!DDCBuffers[MipIndex]);
-							DDCBuffers[MipIndex] = MoveTemp(Params.RawData);
+							DDCBuffers[MipIndex] = MoveTemp(Response.RawData);
 						}
 					});
 				}
diff --git a/Engine/Source/Runtime/Engine/Private/Streaming/Texture2DStreamIn_DDC.cpp b/Engine/Source/Runtime/Engine/Private/Streaming/Texture2DStreamIn_DDC.cpp
--- a/Engine/Source/Runtime/Engine/Private/Streaming/Texture2DStreamIn_DDC.cpp
+++ b/Engine/Source/Runtime/Engine/Private/Streaming/Texture2DStreamIn_DDC.cpp
@@ -182,13 +182,13 @@ void FTexture2DStreamIn_DDC::DoCreateAsyncDDCRequests(const FContext& Context)
 
 			if (MipKeys.Num())
 			{
-				GetCache().GetChunks(MipKeys, DDCRequestOwner, [this](FCacheChunkCompleteParams&& Params)
+				GetCache().GetChunks(MipKeys, DDCRequestOwner, [this](FCacheChunkResponse&& Response)
 				{
-					if (Params.Status == EStatus::Ok)
+					if (Response.Status == EStatus::Ok)
 					{
-						const int32 MipIndex = int32(Params.UserData);
+						const int32 MipIndex = int32(Response.UserData);
 						check(!DDCBuffers[MipIndex]);
-						DDCBuffers[MipIndex] = MoveTemp(Params.RawData);
+						DDCBuffers[MipIndex] = MoveTemp(Response.RawData);
 					}
 				});
 			}
diff --git a/Engine/Source/Runtime/Engine/Private/TextureDerivedData.cpp b/Engine/Source/Runtime/Engine/Private/TextureDerivedData.cpp
--- a/Engine/Source/Runtime/Engine/Private/TextureDerivedData.cpp
+++ b/Engine/Source/Runtime/Engine/Private/TextureDerivedData.cpp
@@ -1304,11 +1304,11 @@ static bool BeginLoadDerivedMips(FTexturePlatformData& PlatformData, int32 First
 			check(Callback);
 			bool bMiss = false;
 			FRequestOwner RequestOwner(EPriority::Blocking);
-			GetCache().GetChunks(MipKeys, RequestOwner, [Callback = MoveTemp(Callback), &bMiss](FCacheChunkCompleteParams&& Params)
+			GetCache().GetChunks(MipKeys, RequestOwner, [Callback = MoveTemp(Callback), &bMiss](FCacheChunkResponse&& Response)
 			{
-				if (Params.Status == EStatus::Ok)
+				if (Response.Status == EStatus::Ok)
 				{
-					Callback(int32(Params.UserData), MoveTemp(Params.RawData));
+					Callback(int32(Response.UserData), MoveTemp(Response.RawData));
 				}
 				else
 				{
