Fixed a crash in logging code when CreateHttpDerivedDataBackend couldn't connect to a backend (extra dereference on a TCHAR* format arg)

diff --git a/Engine/Source/Developer/DerivedDataCache/Private/HttpDerivedDataBackend.cpp b/Engine/Source/Developer/DerivedDataCache/Private/HttpDerivedDataBackend.cpp
--- a/Engine/Source/Developer/DerivedDataCache/Private/HttpDerivedDataBackend.cpp
+++ b/Engine/Source/Developer/DerivedDataCache/Private/HttpDerivedDataBackend.cpp
@@ -3259,7 +3259,7 @@ FDerivedDataBackendInterface* CreateHttpDerivedDataBackend(
 	{
 		return Backend;
 	}
-	UE_LOG(LogDerivedDataCache, Warning, TEXT("Node %s could not contact the service (%s), will not use it"), NodeName, *ServiceUrl);
+	UE_LOG(LogDerivedDataCache, Warning, TEXT("Node %s could not contact the service (%s), will not use it"), NodeName, ServiceUrl);
 	delete Backend;
 	return nullptr;
 #else
