DDC: Made the backend namespaces consistent in advance of introducing ILegacyCacheStore

#rnx

diff --git a/Engine/Source/Developer/DerivedDataCache/Private/DerivedDataBackendInterface.h b/Engine/Source/Developer/DerivedDataCache/Private/DerivedDataBackendInterface.h
--- a/Engine/Source/Developer/DerivedDataCache/Private/DerivedDataBackendInterface.h
+++ b/Engine/Source/Developer/DerivedDataCache/Private/DerivedDataBackendInterface.h
@@ -25,13 +25,14 @@ DECLARE_FLOAT_ACCUMULATOR_STAT_EXTERN(TEXT("Sync Put Time"),STAT_DDC_PutTime,STA
 DECLARE_FLOAT_ACCUMULATOR_STAT_EXTERN(TEXT("Sync Build Time"),STAT_DDC_SyncBuildTime,STATGROUP_DDC, );
 DECLARE_FLOAT_ACCUMULATOR_STAT_EXTERN(TEXT("Exists Time"),STAT_DDC_ExistTime,STATGROUP_DDC, );
 
-
+namespace UE::DerivedData
+{
 
 /**
  * Interface for cache server backends.
  * The entire API should be callable from any thread (except the singleton can be assumed to be called at least once before concurrent access).
  */
-class FDerivedDataBackendInterface : public UE::DerivedData::ICacheStore
+class FDerivedDataBackendInterface : public ICacheStore
 {
 public:
 
@@ -284,3 +285,5 @@ inline void LexFromString(FDerivedDataBackendInterface::ESpeedClass& OutValue, c
 		OutValue = FDerivedDataBackendInterface::ESpeedClass::Local;
 	}
 }
+
+} // UE::DerivedData
diff --git a/Engine/Source/Developer/DerivedDataCache/Private/DerivedDataBackendThrottleWrapper.h b/Engine/Source/Developer/DerivedDataCache/Private/DerivedDataBackendThrottleWrapper.h
--- a/Engine/Source/Developer/DerivedDataCache/Private/DerivedDataBackendThrottleWrapper.h
+++ b/Engine/Source/Developer/DerivedDataCache/Private/DerivedDataBackendThrottleWrapper.h
@@ -12,7 +12,7 @@
 #include "Memory/SharedBuffer.h"
 #include <atomic>
 
-namespace UE::DerivedData::Backends
+namespace UE::DerivedData::CacheStore::Throttle
 {
 
 /** 
@@ -302,4 +302,4 @@ private:
 	ESpeedClass SpeedClass;
 };
 
-} // UE::DerivedData::Backends
+} // UE::DerivedData::CacheStore::Throttle
diff --git a/Engine/Source/Developer/DerivedDataCache/Private/DerivedDataBackendVerifyWrapper.h b/Engine/Source/Developer/DerivedDataCache/Private/DerivedDataBackendVerifyWrapper.h
--- a/Engine/Source/Developer/DerivedDataCache/Private/DerivedDataBackendVerifyWrapper.h
+++ b/Engine/Source/Developer/DerivedDataCache/Private/DerivedDataBackendVerifyWrapper.h
@@ -12,7 +12,7 @@
 
 class Error;
 
-namespace UE::DerivedData::Backends
+namespace UE::DerivedData::CacheStore::Verify
 {
 
 /**
@@ -195,4 +195,4 @@ private:
 	FDerivedDataBackendInterface*					InnerBackend;
 };
 
-} // UE::DerivedData::Backends
+} // UE::DerivedData::CacheStore::Verify
diff --git a/Engine/Source/Developer/DerivedDataCache/Private/DerivedDataBackends.cpp b/Engine/Source/Developer/DerivedDataCache/Private/DerivedDataBackends.cpp
--- a/Engine/Source/Developer/DerivedDataCache/Private/DerivedDataBackends.cpp
+++ b/Engine/Source/Developer/DerivedDataCache/Private/DerivedDataBackends.cpp
@@ -78,7 +78,7 @@ namespace UE::DerivedData::CacheStore::ZenCache
 FDerivedDataBackendInterface* CreateZenDerivedDataBackend(const TCHAR* NodeName, const TCHAR* ServiceUrl, const TCHAR* Namespace);
 } // UE::DerivedData::CacheStore::ZenCache
 
-namespace UE::DerivedData::Backends
+namespace UE::DerivedData::CacheStore::Graph
 {
 
 /**
@@ -157,14 +157,14 @@ public:
 		// Make sure AsyncPutWrapper and KeyLengthWrapper are created
 		if( !AsyncPutWrapper )
 		{
-			AsyncPutWrapper = CacheStore::AsyncPut::CreateAsyncPutDerivedDataBackend( RootCache, /*bCacheInFlightPuts*/ true );
+			AsyncPutWrapper = AsyncPut::CreateAsyncPutDerivedDataBackend( RootCache, /*bCacheInFlightPuts*/ true );
 			check(AsyncPutWrapper);
 			CreatedBackends.Add( AsyncPutWrapper );
 			RootCache = AsyncPutWrapper;
 		}
 		if ( !KeyLengthWrapper )
 		{
-			KeyLengthWrapper = new FDerivedDataLimitKeyLengthWrapper( RootCache, MAX_BACKEND_KEY_LENGTH );
+			KeyLengthWrapper = new LimitKeyLength::FDerivedDataLimitKeyLengthWrapper( RootCache, MAX_BACKEND_KEY_LENGTH );
 			check(KeyLengthWrapper);
 			CreatedBackends.Add( KeyLengthWrapper );
 			RootCache = KeyLengthWrapper;
@@ -289,7 +289,7 @@ public:
 
 					if (LatencyMS != 0 || MaxBytesPerSecond != 0)
 					{
-						ParsedNode = new FDerivedDataBackendThrottleWrapper(ParsedNode, LatencyMS, MaxBytesPerSecond, *Entry);
+						ParsedNode = new Throttle::FDerivedDataBackendThrottleWrapper(ParsedNode, LatencyMS, MaxBytesPerSecond, *Entry);
 					}
 				}
 			}
@@ -327,8 +327,6 @@ public:
 	 */
 	FDerivedDataBackendInterface* ParsePak( const TCHAR* NodeName, const TCHAR* Entry, const bool bWriting )
 	{
-		using namespace CacheStore::PakFile;
-
 		FDerivedDataBackendInterface* PakNode = NULL;
 		FString PakFilename;
 		FParse::Value( Entry, TEXT("Filename="), PakFilename );
@@ -346,7 +344,7 @@ public:
 				FGuid Temp = FGuid::NewGuid();
 				ReadPakFilename = PakFilename;
 				WritePakFilename = PakFilename + TEXT(".") + Temp.ToString();
-				WritePakCache = CreatePakFileDerivedDataBackend(*WritePakFilename, /*bWriting*/ true, bCompressed);
+				WritePakCache = PakFile::CreatePakFileDerivedDataBackend(*WritePakFilename, /*bWriting*/ true, bCompressed);
 				PakNode = WritePakCache;
 			}
 			else
@@ -354,7 +352,7 @@ public:
 				bool bReadPak = FPlatformFileManager::Get().GetPlatformFile().FileExists( *PakFilename );
 				if( bReadPak )
 				{
-					IPakFileDerivedDataBackend* ReadPak = CreatePakFileDerivedDataBackend(*PakFilename, /*bWriting*/ false, bCompressed);
+					PakFile::IPakFileDerivedDataBackend* ReadPak = PakFile::CreatePakFileDerivedDataBackend(*PakFilename, /*bWriting*/ false, bCompressed);
 					ReadPakFilename = PakFilename;
 					PakNode = ReadPak;
 					ReadPakCache.Add(ReadPak);
@@ -402,7 +400,7 @@ public:
 			IFileManager::Get().DeleteDirectory(*(FPaths::ProjectSavedDir() / TEXT("VerifyDDC/")), false, true);
 
 			const bool bFix = GetParsedBool( Entry, TEXT("Fix=") );
-			VerifyNode = new FDerivedDataBackendVerifyWrapper( InnerNode, bFix );
+			VerifyNode = new Verify::FDerivedDataBackendVerifyWrapper( InnerNode, bFix );
 		}
 		else
 		{
@@ -443,7 +441,7 @@ public:
 		FDerivedDataBackendInterface* AsyncNode = NULL;
 		if( InnerNode )
 		{
-			AsyncNode = CacheStore::AsyncPut::CreateAsyncPutDerivedDataBackend( InnerNode, /*bCacheInFlightPuts*/ true );
+			AsyncNode = AsyncPut::CreateAsyncPutDerivedDataBackend( InnerNode, /*bCacheInFlightPuts*/ true );
 		}
 		else
 		{
@@ -488,7 +486,7 @@ public:
 			FParse::Value( Entry, TEXT("Length="), KeyLength );
 			KeyLength = FMath::Clamp( KeyLength, 0, MAX_BACKEND_KEY_LENGTH );
 
-			KeyLengthNode = new FDerivedDataLimitKeyLengthWrapper( InnerNode, KeyLength );
+			KeyLengthNode = new LimitKeyLength::FDerivedDataLimitKeyLengthWrapper( InnerNode, KeyLength );
 		}
 		else
 		{
@@ -544,7 +542,7 @@ public:
 		FDerivedDataBackendInterface* Hierarchy = NULL;
 		if( InnerNodes.Num() > 1 )
 		{
-			CacheStore::Hierarchical::FHierarchicalDerivedDataBackend* HierarchyBackend = new CacheStore::Hierarchical::FHierarchicalDerivedDataBackend( InnerNodes );
+			Hierarchical::FHierarchicalDerivedDataBackend* HierarchyBackend = new Hierarchical::FHierarchicalDerivedDataBackend( InnerNodes );
 			Hierarchy = HierarchyBackend;
 			if (HierarchicalWrapper == NULL)
 			{
@@ -671,7 +669,7 @@ public:
 
 				if (!bShared || IFileManager::Get().DirectoryExists(*Path))
 				{
-					InnerFileSystem = CacheStore::FileSystem::CreateFileSystemDerivedDataBackend(*Path, Entry, *WriteAccessLog);
+					InnerFileSystem = FileSystem::CreateFileSystemDerivedDataBackend(*Path, Entry, *WriteAccessLog);
 				}
 
 				if (InnerFileSystem)
@@ -765,7 +763,7 @@ public:
 			}
 		}
 
-		if (FDerivedDataBackendInterface* Backend = CacheStore::S3::CreateS3DerivedDataBackend(*ManifestPath, *BaseUrl, *Region, *CanaryObjectKey, *CachePath))
+		if (FDerivedDataBackendInterface* Backend = S3::CreateS3DerivedDataBackend(*ManifestPath, *BaseUrl, *Region, *CanaryObjectKey, *CachePath))
 		{
 			return Backend;
 		}
@@ -923,7 +921,7 @@ public:
 			UE_LOG(LogDerivedDataCache, Log, TEXT("Node %s found speed class override ForceSpeedClass=%s"), NodeName, *ForceSpeedClassValue);
 		}
 
-		return CacheStore::Http::CreateHttpDerivedDataBackend(
+		return Http::CreateHttpDerivedDataBackend(
 			NodeName, *Host, *Namespace, *StructuredNamespace, *OAuthProvider, *OAuthClientId, *OAuthSecret,
 			ForceSpeedClass == FDerivedDataBackendInterface::ESpeedClass::Unknown ? nullptr : &ForceSpeedClass, bReadOnly);
 	}
@@ -943,7 +941,7 @@ public:
 			UE_LOG(LogDerivedDataCache, Warning, TEXT("Node %s does not specify 'Namespace', falling back to '%s'"), NodeName, *Namespace);
 		}
 
-		if (FDerivedDataBackendInterface* Backend = CacheStore::ZenCache::CreateZenDerivedDataBackend(NodeName, *ServiceUrl, *Namespace))
+		if (FDerivedDataBackendInterface* Backend = ZenCache::CreateZenDerivedDataBackend(NodeName, *ServiceUrl, *Namespace))
 		{
 			return Backend;
 		}
@@ -982,7 +980,7 @@ public:
 			MaxCacheSize = FMath::Min(MaxCacheSize, MaxSupportedCacheSize);
 
 			UE_LOG( LogDerivedDataCache, Display, TEXT("Max Cache Size: %d MB"), MaxCacheSize);
-			Cache = CacheStore::Memory::CreateMemoryDerivedDataBackend(TEXT("Boot"), MaxCacheSize * 1024 * 1024, /*bCanBeDisabled*/ true);
+			Cache = Memory::CreateMemoryDerivedDataBackend(TEXT("Boot"), MaxCacheSize * 1024 * 1024, /*bCanBeDisabled*/ true);
 
 			if( Cache && Filename.Len() )
 			{
@@ -1024,7 +1022,7 @@ public:
 		FString Filename;
 
 		FParse::Value( Entry, TEXT("Filename="), Filename );
-		Cache = CacheStore::Memory::CreateMemoryDerivedDataBackend(NodeName, /*MaxCacheSize*/ -1, /*bCanBeDisabled*/ false);
+		Cache = Memory::CreateMemoryDerivedDataBackend(NodeName, /*MaxCacheSize*/ -1, /*bCanBeDisabled*/ false);
 		if( Cache && Filename.Len() )
 		{
 			if( Cache->LoadCache( *Filename ) )
@@ -1095,8 +1093,8 @@ public:
 
 				for(const FString& MergePakName : MergePakList)
 				{
-					TUniquePtr<CacheStore::PakFile::IPakFileDerivedDataBackend> ReadPak(
-						CacheStore::PakFile::CreatePakFileDerivedDataBackend(*FPaths::Combine(*FPaths::GetPath(WritePakFilename), *MergePakName), /*bWriting*/ false, /*bCompressed*/ false));
+					TUniquePtr<PakFile::IPakFileDerivedDataBackend> ReadPak(
+						PakFile::CreatePakFileDerivedDataBackend(*FPaths::Combine(*FPaths::GetPath(WritePakFilename), *MergePakName), /*bWriting*/ false, /*bCompressed*/ false));
 					WritePakCache->MergeCache(ReadPak.Get());
 				}
 			}
@@ -1121,7 +1119,7 @@ public:
 							UE_LOG(LogDerivedDataCache, Error, TEXT("Could not delete the pak file %s to overwrite it with a new one."), *ReadPakFilename);
 						}
 					}
-					if (!CacheStore::PakFile::IPakFileDerivedDataBackend::SortAndCopy(WritePakFilename, ReadPakFilename))
+					if (!PakFile::IPakFileDerivedDataBackend::SortAndCopy(WritePakFilename, ReadPakFilename))
 					{
 						UE_LOG(LogDerivedDataCache, Error, TEXT("Couldn't sort pak file (%s)"), *WritePakFilename);
 					}
@@ -1185,10 +1183,10 @@ public:
 	{
 		// Assumptions: there's at least one read-only pak backend in the hierarchy
 		// and its parent is a hierarchical backend.
-		CacheStore::PakFile::IPakFileDerivedDataBackend* ReadPak = nullptr;
+		PakFile::IPakFileDerivedDataBackend* ReadPak = nullptr;
 		if (HierarchicalWrapper && FPlatformFileManager::Get().GetPlatformFile().FileExists(PakFilename))
 		{
-			ReadPak = CacheStore::PakFile::CreatePakFileDerivedDataBackend(PakFilename, /*bWriting*/ false, /*bCompressed*/ false);
+			ReadPak = PakFile::CreatePakFileDerivedDataBackend(PakFilename, /*bWriting*/ false, /*bCompressed*/ false);
 
 			HierarchicalWrapper->AddInnerBackend(ReadPak);
 			CreatedBackends.Add(ReadPak);
@@ -1206,7 +1204,7 @@ public:
 	{
 		for (int PakIndex = 0; PakIndex < ReadPakCache.Num(); ++PakIndex)
 		{
-			CacheStore::PakFile::IPakFileDerivedDataBackend* ReadPak = ReadPakCache[PakIndex];
+			PakFile::IPakFileDerivedDataBackend* ReadPak = ReadPakCache[PakIndex];
 			if (ReadPak->GetFilename() == PakFilename)
 			{
 				check(HierarchicalWrapper);
@@ -1283,12 +1281,12 @@ private:
 
 	/** Instances of backend interfaces which exist in only one copy */
 	FFileBackedDerivedDataBackend*	BootCache;
-	CacheStore::PakFile::IPakFileDerivedDataBackend* WritePakCache;
+	PakFile::IPakFileDerivedDataBackend* WritePakCache;
 	FDerivedDataBackendInterface*	AsyncPutWrapper;
 	FDerivedDataBackendInterface*	KeyLengthWrapper;
-	CacheStore::Hierarchical::FHierarchicalDerivedDataBackend* HierarchicalWrapper;
+	Hierarchical::FHierarchicalDerivedDataBackend* HierarchicalWrapper;
 	/** Support for multiple read only pak files. */
-	TArray<CacheStore::PakFile::IPakFileDerivedDataBackend*> ReadPakCache;
+	TArray<PakFile::IPakFileDerivedDataBackend*> ReadPakCache;
 
 	/** List of directories used by the DDC */
 	TArray<FString> Directories;
@@ -1305,11 +1303,14 @@ private:
 	FAutoConsoleCommand UnountPakCommand;
 };
 
-} // UE::DerivedData::Backends
+} // UE::DerivedData::CacheStore::Graph
+
+namespace UE::DerivedData
+{
 
 FDerivedDataBackend& FDerivedDataBackend::Get()
 {
-	return UE::DerivedData::Backends::FDerivedDataBackendGraph::Get();
+	return CacheStore::Graph::FDerivedDataBackendGraph::Get();
 }
 
 /**
@@ -1392,4 +1393,6 @@ bool FDerivedDataBackendInterface::FBackendDebugOptions::ShouldSimulateMiss(cons
 	return RandomMissRate > 0 && FMath::RandHelper(100) < RandomMissRate;
 }
 
+} // UE::DerivedData
+
 #undef LOCTEXT_NAMESPACE
diff --git a/Engine/Source/Developer/DerivedDataCache/Private/DerivedDataCache.cpp b/Engine/Source/Developer/DerivedDataCache/Private/DerivedDataCache.cpp
--- a/Engine/Source/Developer/DerivedDataCache/Private/DerivedDataCache.cpp
+++ b/Engine/Source/Developer/DerivedDataCache/Private/DerivedDataCache.cpp
@@ -42,7 +42,7 @@ DEFINE_STAT(STAT_DDC_ExistTime);
 
 #if ENABLE_COOK_STATS
 #include "DerivedDataCacheUsageStats.h"
-namespace DerivedDataCacheCookStats
+namespace UE::DerivedData::CookStats
 {
 	// Use to prevent potential divide by zero issues
 	inline double SafeDivide(const int64 Numerator, const int64 Denominator)
diff --git a/Engine/Source/Developer/DerivedDataCache/Private/DerivedDataLimitKeyLengthWrapper.h b/Engine/Source/Developer/DerivedDataCache/Private/DerivedDataLimitKeyLengthWrapper.h
--- a/Engine/Source/Developer/DerivedDataCache/Private/DerivedDataLimitKeyLengthWrapper.h
+++ b/Engine/Source/Developer/DerivedDataCache/Private/DerivedDataLimitKeyLengthWrapper.h
@@ -8,7 +8,7 @@
 #include "DerivedDataCacheUsageStats.h"
 #include "Misc/SecureHash.h"
 
-namespace UE::DerivedData::Backends
+namespace UE::DerivedData::CacheStore::LimitKeyLength
 {
 
 /** 
@@ -401,4 +401,4 @@ private:
 	int32 MaxKeyLength;
 };
 
-} // UE::DerivedData::Backends
+} // UE::DerivedData::CacheStore::LimitKeyLength
diff --git a/Engine/Source/Developer/DerivedDataCache/Private/FileBackedDerivedDataBackend.h b/Engine/Source/Developer/DerivedDataCache/Private/FileBackedDerivedDataBackend.h
--- a/Engine/Source/Developer/DerivedDataCache/Private/FileBackedDerivedDataBackend.h
+++ b/Engine/Source/Developer/DerivedDataCache/Private/FileBackedDerivedDataBackend.h
@@ -4,6 +4,9 @@
 
 #include "DerivedDataBackendInterface.h"
 
+namespace UE::DerivedData
+{
+
 /**
  * A file backed backend. This is used for local caching of data where more flexibility is needed than loose files.
  */
@@ -29,3 +32,5 @@ public:
 	 */
 	virtual void Disable() = 0;
 };
+
+} // UE::DerivedData
diff --git a/Engine/Source/Developer/DerivedDataCache/Private/Tests/HttpDerivedDataBackendTest.cpp b/Engine/Source/Developer/DerivedDataCache/Private/Tests/HttpDerivedDataBackendTest.cpp
--- a/Engine/Source/Developer/DerivedDataCache/Private/Tests/HttpDerivedDataBackendTest.cpp
+++ b/Engine/Source/Developer/DerivedDataCache/Private/Tests/HttpDerivedDataBackendTest.cpp
@@ -50,7 +50,7 @@ namespace UE::DerivedData::CacheStore::ZenCache
 FDerivedDataBackendInterface* CreateZenDerivedDataBackend(const TCHAR* NodeName, const TCHAR* ServiceUrl, const TCHAR* Namespace);
 } // UE::DerivedData::CacheStore::ZenCache
 
-namespace HttpDerivedDataBackendTest
+namespace UE::DerivedData::CacheStore::Http
 {
 
 class FHttpDerivedDataTestBase : public FAutomationTestBase
@@ -541,5 +541,6 @@ bool CacheStore::RunTest(const FString& Parameters)
 	return true;
 }
 
-}
+} // UE::DerivedData::CacheStore::Http
+
 #endif // #if WITH_DEV_AUTOMATION_TESTS && WITH_HTTP_DDC_BACKEND
