DDC: Disabled propagation of gets to slow backends if they have not had a cache miss yet

#rnx

diff --git a/Engine/Source/Developer/DerivedDataCache/Private/HierarchicalDerivedDataBackend.h b/Engine/Source/Developer/DerivedDataCache/Private/HierarchicalDerivedDataBackend.h
--- a/Engine/Source/Developer/DerivedDataCache/Private/HierarchicalDerivedDataBackend.h
+++ b/Engine/Source/Developer/DerivedDataCache/Private/HierarchicalDerivedDataBackend.h
@@ -560,9 +560,11 @@ public:
 							AsyncOwner.KeepAlive();
 							for (int32 FillCacheIndex = 0; FillCacheIndex < InnerBackends.Num(); ++FillCacheIndex)
 							{
-								if (GetCacheIndex != FillCacheIndex)
+								FDerivedDataBackendInterface* FillBackend = AsyncPutInnerBackends[FillCacheIndex].Get();
+								if ((FillCacheIndex < GetCacheIndex) ||
+									(FillCacheIndex > GetCacheIndex && FillBackend->GetSpeedClass() >= ESpeedClass::Fast))
 								{
-									AsyncPutInnerBackends[FillCacheIndex]->Put({{Response.Name, Response.Record, ECachePolicy::Default}}, AsyncOwner);
+									FillBackend->Put({{Response.Name, Response.Record, ECachePolicy::Default}}, AsyncOwner);
 								}
 							}
 
