DDC: Factored out common code in AsyncPut and added stats

#rnx

diff --git a/Engine/Source/Developer/DerivedDataCache/Private/DerivedDataBackendAsyncPutWrapper.cpp b/Engine/Source/Developer/DerivedDataCache/Private/DerivedDataBackendAsyncPutWrapper.cpp
--- a/Engine/Source/Developer/DerivedDataCache/Private/DerivedDataBackendAsyncPutWrapper.cpp
+++ b/Engine/Source/Developer/DerivedDataCache/Private/DerivedDataBackendAsyncPutWrapper.cpp
@@ -16,9 +16,7 @@
 
 namespace UE::DerivedData::CacheStore::Memory
 {
-
 FFileBackedDerivedDataBackend* CreateMemoryDerivedDataBackend(const TCHAR* Name, int64 MaxCacheSize, bool bCanBeDisabled);
-
 } // UE::DerivedData::CacheStore::Memory
 
 namespace UE::DerivedData::CacheStore::AsyncPut
@@ -148,29 +146,54 @@ public:
 	virtual void Put(
 		TConstArrayView<FCachePutRequest> Requests,
 		IRequestOwner& Owner,
-		FOnCachePutComplete&& OnComplete) override;
+		FOnCachePutComplete&& OnComplete) override
+	{
+		Execute(COOK_STAT(&FDerivedDataCacheUsageStats::TimePut,) Requests, Owner, MoveTemp(OnComplete), &ICacheStore::Put);
+	}
 
 	virtual void Get(
 		TConstArrayView<FCacheGetRequest> Requests,
 		IRequestOwner& Owner,
-		FOnCacheGetComplete&& OnComplete) override;
+		FOnCacheGetComplete&& OnComplete) override
+	{
+		Execute(COOK_STAT(&FDerivedDataCacheUsageStats::TimeGet,) Requests, Owner, MoveTemp(OnComplete), &ICacheStore::Get);
+	}
 
 	virtual void PutValue(
 		TConstArrayView<FCachePutValueRequest> Requests,
 		IRequestOwner& Owner,
-		FOnCachePutValueComplete&& OnComplete) override;
+		FOnCachePutValueComplete&& OnComplete) override
+	{
+		Execute(COOK_STAT(&FDerivedDataCacheUsageStats::TimePut,) Requests, Owner, MoveTemp(OnComplete), &ICacheStore::PutValue);
+	}
 
 	virtual void GetValue(
 		TConstArrayView<FCacheGetValueRequest> Requests,
 		IRequestOwner& Owner,
-		FOnCacheGetValueComplete&& OnComplete) override;
+		FOnCacheGetValueComplete&& OnComplete) override
+	{
+		Execute(COOK_STAT(&FDerivedDataCacheUsageStats::TimeGet,) Requests, Owner, MoveTemp(OnComplete), &ICacheStore::GetValue);
+	}
 
 	virtual void GetChunks(
 		TConstArrayView<FCacheChunkRequest> Requests,
 		IRequestOwner& Owner,
-		FOnCacheChunkComplete&& OnComplete) override;
+		FOnCacheChunkComplete&& OnComplete) override
+	{
+		Execute(COOK_STAT(&FDerivedDataCacheUsageStats::TimeGet,) Requests, Owner, MoveTemp(OnComplete), &ICacheStore::GetChunks);
+	}
 
 private:
+	COOK_STAT(using CookStatsFunction = FCookStats::FScopedStatsCounter (FDerivedDataCacheUsageStats::*)());
+
+	template <typename RequestType, typename OnCompleteType, typename OnExecuteType>
+	void Execute(
+		COOK_STAT(CookStatsFunction OnAddStats,)
+		TConstArrayView<RequestType> Requests,
+		IRequestOwner& Owner,
+		OnCompleteType&& OnComplete,
+		OnExecuteType&& OnExecute);
+
 	FDerivedDataCacheUsageStats UsageStats;
 	FDerivedDataCacheUsageStats PutSyncUsageStats;
 
@@ -624,164 +647,81 @@ private:
 	FLazyEvent DoneEvent{EEventMode::ManualReset};
 };
 
-void FDerivedDataBackendAsyncPutWrapper::Put(
-	const TConstArrayView<FCachePutRequest> Requests,
-	IRequestOwner& Owner,
-	FOnCachePutComplete&& OnComplete)
+static FCachePutResponse MakeCanceledResponse(const FCachePutRequest& Request)
 {
-	if (Owner.GetPriority() == EPriority::Blocking || !GDDCIOThreadPool)
-	{
-		InnerBackend->Put(Requests, Owner, MoveTemp(OnComplete));
-	}
-	else
-	{
-		FDerivedDataAsyncWrapperRequest* Request = new FDerivedDataAsyncWrapperRequest(Owner,
-			[this, Requests = TArray<FCachePutRequest>(Requests), OnComplete = MoveTemp(OnComplete)](bool bCancel) mutable
-			{
-				if (!bCancel)
-				{
-					FRequestOwner BlockingOwner(EPriority::Blocking);
-					InnerBackend->Put(Requests, BlockingOwner, MoveTemp(OnComplete));
-					BlockingOwner.Wait();
-				}
-				else if (OnComplete)
-				{
-					for (const FCachePutRequest& Request : Requests)
-					{
-						OnComplete({Request.Name, Request.Record.GetKey(), Request.UserData, EStatus::Canceled});
-					}
-				}
-			});
-		Request->Start(Owner.GetPriority());
-	}
+	return {Request.Name, Request.Record.GetKey(), Request.UserData, EStatus::Canceled};
 }
 
-void FDerivedDataBackendAsyncPutWrapper::Get(
-	const TConstArrayView<FCacheGetRequest> Requests,
-	IRequestOwner& Owner,
-	FOnCacheGetComplete&& OnComplete)
+static FCacheGetResponse MakeCanceledResponse(const FCacheGetRequest& Request)
 {
-	if (Owner.GetPriority() == EPriority::Blocking || !GDDCIOThreadPool)
-	{
-		InnerBackend->Get(Requests, Owner, MoveTemp(OnComplete));
-	}
-	else
-	{
-		FDerivedDataAsyncWrapperRequest* Request = new FDerivedDataAsyncWrapperRequest(Owner,
-			[this, Requests = TArray<FCacheGetRequest>(Requests), OnComplete = MoveTemp(OnComplete)](bool bCancel) mutable
-			{
-				if (!bCancel)
-				{
-					FRequestOwner BlockingOwner(EPriority::Blocking);
-					InnerBackend->Get(Requests, BlockingOwner, MoveTemp(OnComplete));
-					BlockingOwner.Wait();
-				}
-				else if (OnComplete)
-				{
-					for (const FCacheGetRequest& Request : Requests)
-					{
-						OnComplete({Request.Name, FCacheRecordBuilder(Request.Key).Build(), Request.UserData, EStatus::Canceled});
-					}
-				}
-			});
-		Request->Start(Owner.GetPriority());
-	}
+	return {Request.Name, FCacheRecordBuilder(Request.Key).Build(), Request.UserData, EStatus::Canceled};
 }
 
-void FDerivedDataBackendAsyncPutWrapper::PutValue(
-	const TConstArrayView<FCachePutValueRequest> Requests,
-	IRequestOwner& Owner,
-	FOnCachePutValueComplete&& OnComplete)
+static FCachePutValueResponse MakeCanceledResponse(const FCachePutValueRequest& Request)
 {
-	if (Owner.GetPriority() == EPriority::Blocking || !GDDCIOThreadPool)
-	{
-		InnerBackend->PutValue(Requests, Owner, MoveTemp(OnComplete));
-	}
-	else
-	{
-		FDerivedDataAsyncWrapperRequest* Request = new FDerivedDataAsyncWrapperRequest(Owner,
-			[this, Requests = TArray<FCachePutValueRequest>(Requests), OnComplete = MoveTemp(OnComplete)](bool bCancel) mutable
-			{
-				if (!bCancel)
-				{
-					FRequestOwner BlockingOwner(EPriority::Blocking);
-					InnerBackend->PutValue(Requests, BlockingOwner, MoveTemp(OnComplete));
-					BlockingOwner.Wait();
-				}
-				else if (OnComplete)
-				{
-					for (const FCachePutValueRequest& Request : Requests)
-					{
-						OnComplete({Request.Name, Request.Key, Request.UserData, EStatus::Canceled});
-					}
-				}
-			});
-		Request->Start(Owner.GetPriority());
-	}
+	return {Request.Name, Request.Key, Request.UserData, EStatus::Canceled};
 }
 
-void FDerivedDataBackendAsyncPutWrapper::GetValue(
-	const TConstArrayView<FCacheGetValueRequest> Requests,
-	IRequestOwner& Owner,
-	FOnCacheGetValueComplete&& OnComplete)
+static FCacheGetValueResponse MakeCanceledResponse(const FCacheGetValueRequest& Request)
 {
-	if (Owner.GetPriority() == EPriority::Blocking || !GDDCIOThreadPool)
-	{
-		InnerBackend->GetValue(Requests, Owner, MoveTemp(OnComplete));
-	}
-	else
-	{
-		FDerivedDataAsyncWrapperRequest* Request = new FDerivedDataAsyncWrapperRequest(Owner,
-			[this, Requests = TArray<FCacheGetValueRequest>(Requests), OnComplete = MoveTemp(OnComplete)](bool bCancel) mutable
-			{
-				if (!bCancel)
-				{
-					FRequestOwner BlockingOwner(EPriority::Blocking);
-					InnerBackend->GetValue(Requests, BlockingOwner, MoveTemp(OnComplete));
-					BlockingOwner.Wait();
-				}
-				else if (OnComplete)
-				{
-					for (const FCacheGetValueRequest& Request : Requests)
-					{
-						OnComplete({Request.Name, Request.Key, {}, Request.UserData, EStatus::Canceled});
-					}
-				}
-			});
-		Request->Start(Owner.GetPriority());
-	}
+	return {Request.Name, Request.Key, {}, Request.UserData, EStatus::Canceled};
 }
 
-void FDerivedDataBackendAsyncPutWrapper::GetChunks(
-	TConstArrayView<FCacheChunkRequest> Requests,
+static FCacheChunkResponse MakeCanceledResponse(const FCacheChunkRequest& Request)
+{
+	return {Request.Name, Request.Key, Request.Id, Request.RawOffset, 0, {}, {}, Request.UserData, EStatus::Canceled};
+}
+
+template <typename RequestType, typename OnCompleteType, typename OnExecuteType>
+void FDerivedDataBackendAsyncPutWrapper::Execute(
+	COOK_STAT(CookStatsFunction OnAddStats,)
+	const TConstArrayView<RequestType> Requests,
 	IRequestOwner& Owner,
-	FOnCacheChunkComplete&& OnComplete)
+	OnCompleteType&& OnComplete,
+	OnExecuteType&& OnExecute)
 {
+	auto ExecuteWithStats = [this, COOK_STAT(OnAddStats,) OnExecute](TConstArrayView<RequestType> Requests, IRequestOwner& Owner, OnCompleteType&& OnComplete) mutable
+	{
+	#if ENABLE_COOK_STATS
+		Invoke(OnExecute, *InnerBackend, Requests, Owner, [this, OnAddStats, OnComplete = MoveTemp(OnComplete)](auto&& Response)
+		{
+			if (Response.Status == EStatus::Ok)
+			{
+				(UsageStats.*OnAddStats)().AddHit(0);
+			}
+			if (OnComplete)
+			{
+				OnComplete(MoveTemp(Response));
+			}
+		});
+	#else
+		Invoke(OnExecute, *InnerBackend, Requests, Owner, MoveTemp(OnComplete));
+	#endif
+	};
+
 	if (Owner.GetPriority() == EPriority::Blocking || !GDDCIOThreadPool)
 	{
-		InnerBackend->GetChunks(Requests, Owner, MoveTemp(OnComplete));
+		return ExecuteWithStats(Requests, Owner, MoveTemp(OnComplete));
 	}
-	else
-	{
-		FDerivedDataAsyncWrapperRequest* Request = new FDerivedDataAsyncWrapperRequest(Owner,
-			[this, Requests = TArray<FCacheChunkRequest>(Requests), OnComplete = MoveTemp(OnComplete)](bool bCancel) mutable
+
+	FDerivedDataAsyncWrapperRequest* Request = new FDerivedDataAsyncWrapperRequest(Owner,
+		[this, Requests = TArray<RequestType>(Requests), OnComplete = MoveTemp(OnComplete), ExecuteWithStats = MoveTemp(ExecuteWithStats)](bool bCancel) mutable
+		{
+			if (!bCancel)
+			{
+				FRequestOwner BlockingOwner(EPriority::Blocking);
+				ExecuteWithStats(Requests, BlockingOwner, MoveTemp(OnComplete));
+				BlockingOwner.Wait();
+			}
+			else if (OnComplete)
 			{
-				if (!bCancel)
+				for (const RequestType& Request : Requests)
 				{
-					FRequestOwner BlockingOwner(EPriority::Blocking);
-					InnerBackend->GetChunks(Requests, BlockingOwner, MoveTemp(OnComplete));
-					BlockingOwner.Wait();
+					OnComplete(MakeCanceledResponse(Request));
 				}
-				else if (OnComplete)
-				{
-					for (const FCacheChunkRequest& Request : Requests)
-					{
-						OnComplete({Request.Name, Request.Key, Request.Id, Request.RawOffset, 0, {}, {}, Request.UserData, EStatus::Canceled});
-					}
-				}
-			});
-		Request->Start(Owner.GetPriority());
-	}
+			}
+		});
+	Request->Start(Owner.GetPriority());
 }
 
 FDerivedDataBackendInterface* CreateAsyncPutDerivedDataBackend(FDerivedDataBackendInterface* InnerBackend, bool bCacheInFlightPuts)
diff --git a/Engine/Source/Developer/DerivedDataCache/Public/DerivedDataCache.h b/Engine/Source/Developer/DerivedDataCache/Public/DerivedDataCache.h
--- a/Engine/Source/Developer/DerivedDataCache/Public/DerivedDataCache.h
+++ b/Engine/Source/Developer/DerivedDataCache/Public/DerivedDataCache.h
@@ -531,7 +531,7 @@ struct FCacheChunkResponse
 	uint64 RawSize = 0;
 
 	/** The hash of the entire value, even if only a subset was fetched. */
-	const FIoHash& RawHash;
+	FIoHash RawHash;
 
 	/** Data for the subset of the value that was fetched when Status is Ok, otherwise null. */
 	FSharedBuffer RawData;
diff --git a/Engine/Source/Runtime/Core/Public/ProfilingDebugging/CookStats.h b/Engine/Source/Runtime/Core/Public/ProfilingDebugging/CookStats.h
--- a/Engine/Source/Runtime/Core/Public/ProfilingDebugging/CookStats.h
+++ b/Engine/Source/Runtime/Core/Public/ProfilingDebugging/CookStats.h
@@ -421,8 +421,8 @@ struct FDerivedDataUsageStats
 };
 
 /** any expression inside this block will be compiled out if ENABLE_COOK_STATS is not true. Shorthand for the multi-line #if ENABLE_COOK_STATS guard. */
-#define COOK_STAT(expr) expr
+#define COOK_STAT(...) __VA_ARGS__
 
 #else
-#define COOK_STAT(expr)
+#define COOK_STAT(...)
 #endif
