DDC: Added ILegacyCacheStore and folded the key length limit into the new FLegacyCacheKey

This API adds a little bit of overhead but, in exchange:
- Offers a path to migrate to the new Value API.
- Offers a path to provide compression while falling back to the previous format.
- Offers a path to backends batching requests without consuming a worker thread.
- Avoids copying the value multiple times during async puts.

#rnx

diff --git a/Engine/Source/Developer/DerivedDataCache/Private/DerivedDataBackendAsyncPutWrapper.cpp b/Engine/Source/Developer/DerivedDataCache/Private/DerivedDataBackendAsyncPutWrapper.cpp
--- a/Engine/Source/Developer/DerivedDataCache/Private/DerivedDataBackendAsyncPutWrapper.cpp
+++ b/Engine/Source/Developer/DerivedDataCache/Private/DerivedDataBackendAsyncPutWrapper.cpp
@@ -183,6 +183,30 @@ public:
 		Execute(COOK_STAT(&FDerivedDataCacheUsageStats::TimeGet,) Requests, Owner, MoveTemp(OnComplete), &ICacheStore::GetChunks);
 	}
 
+	virtual void LegacyPut(
+		TConstArrayView<FLegacyCachePutRequest> Requests,
+		IRequestOwner& Owner,
+		FOnLegacyCachePutComplete&& OnComplete) override
+	{
+		Execute(COOK_STAT(&FDerivedDataCacheUsageStats::TimePut,) Requests, Owner, MoveTemp(OnComplete), &ILegacyCacheStore::LegacyPut);
+	}
+
+	virtual void LegacyGet(
+		TConstArrayView<FLegacyCacheGetRequest> Requests,
+		IRequestOwner& Owner,
+		FOnLegacyCacheGetComplete&& OnComplete) override
+	{
+		Execute(COOK_STAT(&FDerivedDataCacheUsageStats::TimeGet,) Requests, Owner, MoveTemp(OnComplete), &ILegacyCacheStore::LegacyGet);
+	}
+
+	virtual void LegacyDelete(
+		TConstArrayView<FLegacyCacheDeleteRequest> Requests,
+		IRequestOwner& Owner,
+		FOnLegacyCacheDeleteComplete&& OnComplete) override
+	{
+		Execute(COOK_STAT(&FDerivedDataCacheUsageStats::TimePut,) Requests, Owner, MoveTemp(OnComplete), &ILegacyCacheStore::LegacyDelete);
+	}
+
 private:
 	COOK_STAT(using CookStatsFunction = FCookStats::FScopedStatsCounter (FDerivedDataCacheUsageStats::*)());
 
@@ -672,6 +696,21 @@ static FCacheChunkResponse MakeCanceledResponse(const FCacheChunkRequest& Reques
 	return {Request.Name, Request.Key, Request.Id, Request.RawOffset, 0, {}, {}, Request.UserData, EStatus::Canceled};
 }
 
+static FLegacyCachePutResponse MakeCanceledResponse(const FLegacyCachePutRequest& Request)
+{
+	return {Request.Name, Request.Key, Request.UserData, EStatus::Canceled};
+}
+
+static FLegacyCacheGetResponse MakeCanceledResponse(const FLegacyCacheGetRequest& Request)
+{
+	return {Request.Name, Request.Key, {}, Request.UserData, EStatus::Canceled};
+}
+
+static FLegacyCacheDeleteResponse MakeCanceledResponse(const FLegacyCacheDeleteRequest& Request)
+{
+	return {Request.Name, Request.Key, Request.UserData, EStatus::Canceled};
+}
+
 template <typename RequestType, typename OnCompleteType, typename OnExecuteType>
 void FDerivedDataBackendAsyncPutWrapper::Execute(
 	COOK_STAT(CookStatsFunction OnAddStats,)
diff --git a/Engine/Source/Developer/DerivedDataCache/Private/DerivedDataBackendInterface.h b/Engine/Source/Developer/DerivedDataCache/Private/DerivedDataBackendInterface.h
--- a/Engine/Source/Developer/DerivedDataCache/Private/DerivedDataBackendInterface.h
+++ b/Engine/Source/Developer/DerivedDataCache/Private/DerivedDataBackendInterface.h
@@ -6,6 +6,7 @@
 #include "CoreTypes.h"
 #include "Containers/BitArray.h"
 #include "DerivedDataCache.h"
+#include "DerivedDataLegacyCacheStore.h"
 #include "Stats/Stats.h"
 
 class FDerivedDataCacheUsageStats;
@@ -32,7 +33,7 @@ namespace UE::DerivedData
  * Interface for cache server backends.
  * The entire API should be callable from any thread (except the singleton can be assumed to be called at least once before concurrent access).
  */
-class FDerivedDataBackendInterface : public ICacheStore
+class FDerivedDataBackendInterface : public ILegacyCacheStore
 {
 public:
 
@@ -190,6 +191,21 @@ public:
 	 * false which will result in a warning if an attempt is made to apply these options.
 	 */
 	virtual bool ApplyDebugOptions(FBackendDebugOptions& InOptions) = 0;
+
+	virtual void LegacyPut(
+		TConstArrayView<FLegacyCachePutRequest> Requests,
+		IRequestOwner& Owner,
+		FOnLegacyCachePutComplete&& OnComplete) override;
+
+	virtual void LegacyGet(
+		TConstArrayView<FLegacyCacheGetRequest> Requests,
+		IRequestOwner& Owner,
+		FOnLegacyCacheGetComplete&& OnComplete) override;
+
+	virtual void LegacyDelete(
+		TConstArrayView<FLegacyCacheDeleteRequest> Requests,
+		IRequestOwner& Owner,
+		FOnLegacyCacheDeleteComplete&& OnComplete) override;
 };
 
 class FDerivedDataBackend
@@ -206,7 +222,9 @@ public:
 	 * Singleton to retrieve the root cache
 	 * @return Reference to the global cache root
 	 */
-	virtual FDerivedDataBackendInterface& GetRoot() = 0;
+	virtual ILegacyCacheStore& GetRoot() = 0;
+
+	virtual int32 GetMaxKeyLength() const = 0;
 
 	//--------------------
 	// System Interface, copied from FDerivedDataCacheInterface
diff --git a/Engine/Source/Developer/DerivedDataCache/Private/DerivedDataBackends.cpp b/Engine/Source/Developer/DerivedDataCache/Private/DerivedDataBackends.cpp
--- a/Engine/Source/Developer/DerivedDataCache/Private/DerivedDataBackends.cpp
+++ b/Engine/Source/Developer/DerivedDataCache/Private/DerivedDataBackends.cpp
@@ -6,7 +6,6 @@
 #include "DerivedDataBackendThrottleWrapper.h"
 #include "DerivedDataBackendVerifyWrapper.h"
 #include "DerivedDataCacheUsageStats.h"
-#include "DerivedDataLimitKeyLengthWrapper.h"
 #include "FileBackedDerivedDataBackend.h"
 #include "HAL/FileManager.h"
 #include "HAL/IConsoleManager.h"
@@ -96,7 +95,6 @@ public:
 		, BootCache(NULL)
 		, WritePakCache(NULL)
 		, AsyncPutWrapper(NULL)
-		, KeyLengthWrapper(NULL)
 		, HierarchicalWrapper(NULL)
 		, bUsingSharedDDC(false)
 		, bIsShuttingDown(false)
@@ -154,7 +152,7 @@ public:
 			}
 		}
 
-		// Make sure AsyncPutWrapper and KeyLengthWrapper are created
+		// Make sure AsyncPutWrapper is created
 		if( !AsyncPutWrapper )
 		{
 			AsyncPutWrapper = AsyncPut::CreateAsyncPutDerivedDataBackend( RootCache, /*bCacheInFlightPuts*/ true );
@@ -162,12 +160,10 @@ public:
 			CreatedBackends.Add( AsyncPutWrapper );
 			RootCache = AsyncPutWrapper;
 		}
-		if ( !KeyLengthWrapper )
+
+		if (MaxKeyLength == 0)
 		{
-			KeyLengthWrapper = new LimitKeyLength::FDerivedDataLimitKeyLengthWrapper( RootCache, MAX_BACKEND_KEY_LENGTH );
-			check(KeyLengthWrapper);
-			CreatedBackends.Add( KeyLengthWrapper );
-			RootCache = KeyLengthWrapper;
+			MaxKeyLength = MAX_BACKEND_KEY_LENGTH;
 		}
 	}
 
@@ -231,14 +227,13 @@ public:
 				}
 				else if( NodeType == TEXT("KeyLength") )
 				{
-					if( KeyLengthWrapper == NULL )
+					if( MaxKeyLength == 0 )
 					{
-						KeyLengthWrapper = ParseKeyLength( NodeName, *Entry, IniFilename, IniSection, InParsedNodes );
-						ParsedNode = KeyLengthWrapper;
+						ParsedNode = ParseKeyLength( NodeName, *Entry, IniFilename, IniSection, InParsedNodes );
 					}
 					else
 					{
-						UE_LOG( LogDerivedDataCache, Warning, TEXT("FDerivedDataBackendGraph:  Unable to create %s KeyLengthWrapper because only one KeyLength node is supported."), NodeName );
+						UE_LOG( LogDerivedDataCache, Warning, TEXT("FDerivedDataBackendGraph:  Unable to create %s KeyLength because only one KeyLength node is supported."), NodeName );
 					}
 				}
 				else if( NodeType == TEXT("AsyncPut") )
@@ -463,7 +458,7 @@ public:
 	 */
 	FDerivedDataBackendInterface* ParseKeyLength( const TCHAR* NodeName, const TCHAR* Entry, const FString& IniFilename, const TCHAR* IniSection, TMap<FString, FDerivedDataBackendInterface*>& InParsedNodes )
 	{
-		FDerivedDataBackendInterface* InnerNode = NULL;
+		FDerivedDataBackendInterface* InnerNode = nullptr;
 		FString InnerName;
 		if( FParse::Value( Entry, TEXT("Inner="), InnerName ) )
 		{
@@ -479,21 +474,18 @@ public:
 			}
 		}
 
-		FDerivedDataBackendInterface* KeyLengthNode = NULL;
 		if( InnerNode )
 		{
 			int32 KeyLength = MAX_BACKEND_KEY_LENGTH;
 			FParse::Value( Entry, TEXT("Length="), KeyLength );
-			KeyLength = FMath::Clamp( KeyLength, 0, MAX_BACKEND_KEY_LENGTH );
-
-			KeyLengthNode = new LimitKeyLength::FDerivedDataLimitKeyLengthWrapper( InnerNode, KeyLength );
+			MaxKeyLength = FMath::Clamp( KeyLength, 0, MAX_BACKEND_KEY_LENGTH );
 		}
 		else
 		{
 			UE_LOG( LogDerivedDataCache, Warning, TEXT("Unable to find inner node %s for KeyLength node %s. KeyLength node will not be created."), *InnerName, NodeName );
 		}
 
-		return KeyLengthNode;
+		return InnerNode;
 	}
 
 	/**
@@ -1044,12 +1036,17 @@ public:
 		DestroyCreatedBackends();
 	}
 
-	FDerivedDataBackendInterface& GetRoot() override
+	ILegacyCacheStore& GetRoot() override
 	{
 		check(RootCache);
 		return *RootCache;
 	}
 
+	virtual int32 GetMaxKeyLength() const override
+	{
+		return MaxKeyLength;
+	}
+
 	virtual void NotifyBootComplete() override
 	{
 		check(RootCache);
@@ -1283,7 +1280,6 @@ private:
 	FFileBackedDerivedDataBackend*	BootCache;
 	PakFile::IPakFileDerivedDataBackend* WritePakCache;
 	FDerivedDataBackendInterface*	AsyncPutWrapper;
-	FDerivedDataBackendInterface*	KeyLengthWrapper;
 	Hierarchical::FHierarchicalDerivedDataBackend* HierarchicalWrapper;
 	/** Support for multiple read only pak files. */
 	TArray<PakFile::IPakFileDerivedDataBackend*> ReadPakCache;
@@ -1291,6 +1287,8 @@ private:
 	/** List of directories used by the DDC */
 	TArray<FString> Directories;
 
+	int32 MaxKeyLength = 0;
+
 	/** Whether a shared cache is in use */
 	bool bUsingSharedDDC;
 
@@ -1308,6 +1306,96 @@ private:
 namespace UE::DerivedData
 {
 
+void FDerivedDataBackendInterface::LegacyPut(
+	const TConstArrayView<FLegacyCachePutRequest> Requests,
+	IRequestOwner& Owner,
+	FOnLegacyCachePutComplete&& OnComplete)
+{
+	for (const FLegacyCachePutRequest& Request : Requests)
+	{
+		FCompositeBuffer CompositeValue = Request.Value;
+		Request.Key.WriteValueTrailer(CompositeValue);
+
+		checkf(CompositeValue.GetSize() < MAX_int32,
+			TEXT("Value is 2 GiB or greater, which is not supported for put of '%s' from '%s'"),
+			*Request.Key.GetFullKey(), *Request.Name);
+	
+		UE_CLOG(Request.Key.HasShortKey(), LogDerivedDataCache, VeryVerbose,
+			TEXT("ShortenKey %s -> %s"), *Request.Key.GetFullKey(), *Request.Key.GetShortKey());
+
+		FSharedBuffer Value = MoveTemp(CompositeValue).ToShared();
+		const TArrayView<const uint8> Data(MakeArrayView(static_cast<const uint8*>(Value.GetData()), int32(Value.GetSize())));
+		const EPutStatus Status = PutCachedData(*Request.Key.GetShortKey(), Data, /*bPutEvenIfExists*/ false);
+		OnComplete({Request.Name, Request.Key, Request.UserData, Status == EPutStatus::Cached ? EStatus::Ok : EStatus::Error});
+	}
+}
+
+void FDerivedDataBackendInterface::LegacyGet(
+	const TConstArrayView<FLegacyCacheGetRequest> Requests,
+	IRequestOwner& Owner,
+	FOnLegacyCacheGetComplete&& OnComplete)
+{
+	TArray<FString, TInlineAllocator<8>> ExistsKeys;
+	TArray<const FLegacyCacheGetRequest*, TInlineAllocator<8>> ExistsRequests;
+
+	TArray<FString, TInlineAllocator<8>> PrefetchKeys;
+	TArray<const FLegacyCacheGetRequest*, TInlineAllocator<8>> PrefetchRequests;
+
+	for (const FLegacyCacheGetRequest& Request : Requests)
+	{
+		if (!EnumHasAnyFlags(Request.Policy, ECachePolicy::Query))
+		{
+			OnComplete({Request.Name, Request.Key, {}, Request.UserData, EStatus::Error});
+		}
+		else if (EnumHasAnyFlags(Request.Policy, ECachePolicy::SkipData))
+		{
+			const bool bExists = !EnumHasAnyFlags(Request.Policy, ECachePolicy::Store);
+			(bExists ? ExistsKeys : PrefetchKeys).Emplace(Request.Key.GetShortKey());
+			(bExists ? ExistsRequests : PrefetchRequests).Add(&Request);
+		}
+		else
+		{
+			TArray<uint8> Data;
+			const bool bGetOk = GetCachedData(*Request.Key.GetShortKey(), Data);
+			FCompositeBuffer Value(MakeSharedBufferFromArray(MoveTemp(Data)));
+			const bool bKeyOk = bGetOk && Request.Key.ReadValueTrailer(Value);
+			OnComplete({Request.Name, Request.Key, MoveTemp(Value).ToShared(), Request.UserData, bKeyOk ? EStatus::Ok : EStatus::Error});
+		}
+	}
+
+	if (!PrefetchKeys.IsEmpty())
+	{
+		const bool bOk = TryToPrefetch(PrefetchKeys);
+		for (const FLegacyCacheGetRequest* const Request : PrefetchRequests)
+		{
+			OnComplete({Request->Name, Request->Key, {}, Request->UserData, bOk ? EStatus::Ok : EStatus::Error});
+		}
+	}
+
+	if (!ExistsKeys.IsEmpty())
+	{
+		int32 Index = 0;
+		const TBitArray<> Exists = CachedDataProbablyExistsBatch(ExistsKeys);
+		for (const FLegacyCacheGetRequest* const Request : ExistsRequests)
+		{
+			OnComplete({Request->Name, Request->Key, {}, Request->UserData, Exists[Index] ? EStatus::Ok : EStatus::Error});
+			++Index;
+		}
+	}
+}
+
+void FDerivedDataBackendInterface::LegacyDelete(
+	const TConstArrayView<FLegacyCacheDeleteRequest> Requests,
+	IRequestOwner& Owner,
+	FOnLegacyCacheDeleteComplete&& OnComplete)
+{
+	for (const FLegacyCacheDeleteRequest& Request : Requests)
+	{
+		RemoveCachedData(*Request.Key.GetShortKey(), Request.bTransient);
+		OnComplete({Request.Name, Request.Key, Request.UserData, EStatus::Ok});
+	}
+}
+
 FDerivedDataBackend& FDerivedDataBackend::Get()
 {
 	return CacheStore::Graph::FDerivedDataBackendGraph::Get();
diff --git a/Engine/Source/Developer/DerivedDataCache/Private/DerivedDataCache.cpp b/Engine/Source/Developer/DerivedDataCache/Private/DerivedDataCache.cpp
--- a/Engine/Source/Developer/DerivedDataCache/Private/DerivedDataCache.cpp
+++ b/Engine/Source/Developer/DerivedDataCache/Private/DerivedDataCache.cpp
@@ -16,6 +16,7 @@
 #include "DerivedDataCachePrivate.h"
 #include "DerivedDataCacheUsageStats.h"
 #include "DerivedDataPluginInterface.h"
+#include "DerivedDataRequestOwner.h"
 #include "Features/IModularFeatures.h"
 #include "HAL/ThreadSafeCounter.h"
 #include "Misc/CommandLine.h"
@@ -301,12 +302,13 @@ class FDerivedDataCache final
 		 * @param	InDataDeriver	plugin to produce cache key and in the event of a miss, return the data.
 		 * @param	InCacheKey		Complete cache key for this data.
 		**/
-		FBuildAsyncWorker(FDerivedDataPluginInterface* InDataDeriver, const TCHAR* InCacheKey, bool bInSynchronousForStats)
+		FBuildAsyncWorker(FDerivedDataPluginInterface* InDataDeriver, const TCHAR* InCacheKey, FStringView InDebugContext, bool bInSynchronousForStats)
 		: bSuccess(false)
 		, bSynchronousForStats(bInSynchronousForStats)
 		, bDataWasBuilt(false)
 		, DataDeriver(InDataDeriver)
 		, CacheKey(InCacheKey)
+		, DebugContext(InDebugContext)
 		{
 		}
 
@@ -342,7 +344,20 @@ class FDerivedDataCache final
 				STAT(double ThisTime = 0);
 				{
 					SCOPE_SECONDS_COUNTER(ThisTime);
-					bGetResult = FDerivedDataBackend::Get().GetRoot().GetCachedData(*CacheKey, Data);
+					FLegacyCacheGetRequest LegacyRequest;
+					LegacyRequest.Name = DebugContext;
+					LegacyRequest.Key = FLegacyCacheKey(CacheKey, FDerivedDataBackend::Get().GetMaxKeyLength());
+					FRequestOwner BlockingOwner(EPriority::Blocking);
+					FDerivedDataBackend::Get().GetRoot().LegacyGet({LegacyRequest}, BlockingOwner,
+						[this, &bGetResult](FLegacyCacheGetResponse&& Response)
+						{
+							bGetResult = Response.Status == EStatus::Ok && Response.Value.GetSize() < MAX_int32;
+							if (bGetResult)
+							{
+								Data = MakeArrayView(static_cast<const uint8*>(Response.Value.GetData()), int32(Response.Value.GetSize()));
+							}
+						});
+					BlockingOwner.Wait();
 				}
 				INC_FLOAT_STAT_BY(STAT_DDC_SyncGetTime, bSynchronousForStats ? (float)ThisTime : 0.0f);
 			}
@@ -413,7 +428,13 @@ class FDerivedDataCache final
 					STAT(double ThisTime = 0);
 					{
 						SCOPE_SECONDS_COUNTER(ThisTime);
-						FDerivedDataBackend::Get().GetRoot().PutCachedData(*CacheKey, Data, true);
+						FLegacyCachePutRequest LegacyRequest;
+						LegacyRequest.Name = DebugContext;
+						LegacyRequest.Key = FLegacyCacheKey(CacheKey, FDerivedDataBackend::Get().GetMaxKeyLength());
+						LegacyRequest.Value = FCompositeBuffer(MakeSharedBufferFromArray(MoveTemp(Data)));
+						FRequestOwner BlockingOwner(EPriority::Blocking);
+						FDerivedDataBackend::Get().GetRoot().LegacyPut({LegacyRequest}, BlockingOwner, [](auto&&){});
+						BlockingOwner.Wait();
 					}
 					INC_FLOAT_STAT_BY(STAT_DDC_PutTime, bSynchronousForStats ? (float)ThisTime : 0.0f);
 				}
@@ -449,6 +470,8 @@ class FDerivedDataCache final
 		FDerivedDataPluginInterface*	DataDeriver;
 		/** Cache key associated with this build **/
 		FString							CacheKey;
+		/** Context from the caller */
+		FSharedString					DebugContext;
 		/** Data to return to caller, later **/
 		TArray<uint8>					Data;
 	};
@@ -487,7 +510,7 @@ public:
 		check(DataDeriver);
 		FString CacheKey = FDerivedDataCache::BuildCacheKey(DataDeriver);
 		UE_LOG(LogDerivedDataCache, VeryVerbose, TEXT("GetSynchronous %s from '%s'"), *CacheKey, *DataDeriver->GetDebugContextString());
-		FAsyncTask<FBuildAsyncWorker> PendingTask(DataDeriver, *CacheKey, true);
+		FAsyncTask<FBuildAsyncWorker> PendingTask(DataDeriver, *CacheKey, DataDeriver->GetDebugContextString(), true);
 		AddToAsyncCompletionCounter(1);
 		PendingTask.StartSynchronousTask();
 		OutData = PendingTask.GetTask().Data;
@@ -506,7 +529,7 @@ public:
 		FString CacheKey = FDerivedDataCache::BuildCacheKey(DataDeriver);
 		UE_LOG(LogDerivedDataCache, VeryVerbose, TEXT("GetAsynchronous %s from '%s', Handle %d"), *CacheKey, *DataDeriver->GetDebugContextString(), Handle);
 		const bool bSync = !DataDeriver->IsBuildThreadsafe();
-		FAsyncTask<FBuildAsyncWorker>* AsyncTask = new FAsyncTask<FBuildAsyncWorker>(DataDeriver, *CacheKey, bSync);
+		FAsyncTask<FBuildAsyncWorker>* AsyncTask = new FAsyncTask<FBuildAsyncWorker>(DataDeriver, *CacheKey, DataDeriver->GetDebugContextString(), bSync);
 		check(!PendingTasks.Contains(Handle));
 		PendingTasks.Add(Handle,AsyncTask);
 		AddToAsyncCompletionCounter(1);
@@ -581,26 +604,26 @@ public:
 		return true;
 	}
 
-	virtual bool GetSynchronous(const TCHAR* CacheKey, TArray<uint8>& OutData, FStringView DataContext) override
+	virtual bool GetSynchronous(const TCHAR* CacheKey, TArray<uint8>& OutData, FStringView DebugContext) override
 	{
 		DDC_SCOPE_CYCLE_COUNTER(DDC_GetSynchronous_Data);
-		UE_LOG(LogDerivedDataCache, VeryVerbose, TEXT("GetSynchronous %s from '%.*s'"), CacheKey, DataContext.Len(), DataContext.GetData());
+		UE_LOG(LogDerivedDataCache, VeryVerbose, TEXT("GetSynchronous %s from '%.*s'"), CacheKey, DebugContext.Len(), DebugContext.GetData());
 		ValidateCacheKey(CacheKey);
-		FAsyncTask<FBuildAsyncWorker> PendingTask((FDerivedDataPluginInterface*)NULL, CacheKey, true);
+		FAsyncTask<FBuildAsyncWorker> PendingTask((FDerivedDataPluginInterface*)NULL, CacheKey, DebugContext, true);
 		AddToAsyncCompletionCounter(1);
 		PendingTask.StartSynchronousTask();
 		OutData = PendingTask.GetTask().Data;
 		return PendingTask.GetTask().bSuccess;
 	}
 
-	virtual uint32 GetAsynchronous(const TCHAR* CacheKey, FStringView DataContext) override
+	virtual uint32 GetAsynchronous(const TCHAR* CacheKey, FStringView DebugContext) override
 	{
 		DDC_SCOPE_CYCLE_COUNTER(DDC_GetAsynchronous_Handle);
 		FScopeLock ScopeLock(&SynchronizationObject);
 		const uint32 Handle = NextHandle();
-		UE_LOG(LogDerivedDataCache, VeryVerbose, TEXT("GetAsynchronous %s from '%.*s', Handle %d"), CacheKey, DataContext.Len(), DataContext.GetData(), Handle);
+		UE_LOG(LogDerivedDataCache, VeryVerbose, TEXT("GetAsynchronous %s from '%.*s', Handle %d"), CacheKey, DebugContext.Len(), DebugContext.GetData(), Handle);
 		ValidateCacheKey(CacheKey);
-		FAsyncTask<FBuildAsyncWorker>* AsyncTask = new FAsyncTask<FBuildAsyncWorker>((FDerivedDataPluginInterface*)NULL, CacheKey, false);
+		FAsyncTask<FBuildAsyncWorker>* AsyncTask = new FAsyncTask<FBuildAsyncWorker>((FDerivedDataPluginInterface*)NULL, CacheKey, DebugContext, false);
 		check(!PendingTasks.Contains(Handle));
 		PendingTasks.Add(Handle, AsyncTask);
 		AddToAsyncCompletionCounter(1);
@@ -609,15 +632,21 @@ public:
 		return Handle;
 	}
 
-	virtual void Put(const TCHAR* CacheKey, TArrayView<const uint8> Data, FStringView DataContext, bool bPutEvenIfExists = false) override
+	virtual void Put(const TCHAR* CacheKey, TArrayView<const uint8> Data, FStringView DebugContext, bool bPutEvenIfExists = false) override
 	{
 		DDC_SCOPE_CYCLE_COUNTER(DDC_Put);
-		UE_LOG(LogDerivedDataCache, VeryVerbose, TEXT("Put %s from '%.*s'"), CacheKey, DataContext.Len(), DataContext.GetData());
+		UE_LOG(LogDerivedDataCache, VeryVerbose, TEXT("Put %s from '%.*s'"), CacheKey, DebugContext.Len(), DebugContext.GetData());
 		ValidateCacheKey(CacheKey);
 		STAT(double ThisTime = 0);
 		{
 			SCOPE_SECONDS_COUNTER(ThisTime);
-			FDerivedDataBackend::Get().GetRoot().PutCachedData(CacheKey, Data, bPutEvenIfExists);
+			FLegacyCachePutRequest LegacyRequest;
+			LegacyRequest.Name = DebugContext;
+			LegacyRequest.Key = FLegacyCacheKey(CacheKey, FDerivedDataBackend::Get().GetMaxKeyLength());
+			LegacyRequest.Value = FCompositeBuffer(FSharedBuffer::MakeView(MakeMemoryView(Data)));
+			FRequestOwner BlockingOwner(EPriority::Blocking);
+			FDerivedDataBackend::Get().GetRoot().LegacyPut({LegacyRequest}, BlockingOwner, [](auto&&){});
+			BlockingOwner.Wait();
 		}
 		INC_FLOAT_STAT_BY(STAT_DDC_PutTime,(float)ThisTime);
 		INC_DWORD_STAT(STAT_DDC_NumPuts);
@@ -626,7 +655,13 @@ public:
 	virtual void MarkTransient(const TCHAR* CacheKey) override
 	{
 		ValidateCacheKey(CacheKey);
-		FDerivedDataBackend::Get().GetRoot().RemoveCachedData(CacheKey, /*bTransient=*/ true);
+		FLegacyCacheDeleteRequest LegacyRequest;
+		LegacyRequest.Key = FLegacyCacheKey(CacheKey, FDerivedDataBackend::Get().GetMaxKeyLength());
+		LegacyRequest.Name = LegacyRequest.Key.GetFullKey();
+		LegacyRequest.bTransient = true;
+		FRequestOwner BlockingOwner(EPriority::Blocking);
+		FDerivedDataBackend::Get().GetRoot().LegacyDelete({LegacyRequest}, BlockingOwner, [](auto&&){});
+		BlockingOwner.Wait();
 	}
 
 	virtual bool CachedDataProbablyExists(const TCHAR* CacheKey) override
@@ -638,7 +673,14 @@ public:
 		STAT(double ThisTime = 0);
 		{
 			SCOPE_SECONDS_COUNTER(ThisTime);
-			bResult = FDerivedDataBackend::Get().GetRoot().CachedDataProbablyExists(CacheKey);
+			FLegacyCacheGetRequest LegacyRequest;
+			LegacyRequest.Key = FLegacyCacheKey(CacheKey, FDerivedDataBackend::Get().GetMaxKeyLength());
+			LegacyRequest.Name = LegacyRequest.Key.GetFullKey();
+			LegacyRequest.Policy = ECachePolicy::Query | ECachePolicy::SkipData;
+			FRequestOwner BlockingOwner(EPriority::Blocking);
+			FDerivedDataBackend::Get().GetRoot().LegacyGet({LegacyRequest}, BlockingOwner,
+				[&bResult](FLegacyCacheGetResponse&& Response) { bResult = Response.Status == EStatus::Ok; });
+			BlockingOwner.Wait();
 		}
 		INC_FLOAT_STAT_BY(STAT_DDC_ExistTime, (float)ThisTime);
 		return bResult;
@@ -646,23 +688,35 @@ public:
 
 	virtual TBitArray<> CachedDataProbablyExistsBatch(TConstArrayView<FString> CacheKeys) override
 	{
-		TBitArray<> Result;
-		if (CacheKeys.Num() > 1)
+		TBitArray<> Result(false, CacheKeys.Num());
+		if (!CacheKeys.IsEmpty())
 		{
 			DDC_SCOPE_CYCLE_COUNTER(DDC_CachedDataProbablyExistsBatch);
 			INC_DWORD_STAT(STAT_DDC_NumExist);
 			STAT(double ThisTime = 0);
 			{
 				SCOPE_SECONDS_COUNTER(ThisTime);
-				Result = FDerivedDataBackend::Get().GetRoot().CachedDataProbablyExistsBatch(CacheKeys);
-				check(Result.Num() == CacheKeys.Num());
+				TArray<FLegacyCacheGetRequest, TInlineAllocator<8>> LegacyRequests;
+				int32 Index = 0;
+				for (const FString& CacheKey : CacheKeys)
+				{
+					FLegacyCacheGetRequest& LegacyRequest = LegacyRequests.AddDefaulted_GetRef();
+					LegacyRequest.Key = FLegacyCacheKey(CacheKey, FDerivedDataBackend::Get().GetMaxKeyLength());
+					LegacyRequest.Name = LegacyRequest.Key.GetFullKey();
+					LegacyRequest.Policy = ECachePolicy::Query | ECachePolicy::SkipData;
+					LegacyRequest.UserData = uint64(Index);
+					++Index;
+				}
+				FRequestOwner BlockingOwner(EPriority::Blocking);
+				FDerivedDataBackend::Get().GetRoot().LegacyGet(LegacyRequests, BlockingOwner,
+					[&Result](FLegacyCacheGetResponse&& Response)
+					{
+						Result[int32(Response.UserData)] = Response.Status == EStatus::Ok;
+					});
+				BlockingOwner.Wait();
 			}
 			INC_FLOAT_STAT_BY(STAT_DDC_ExistTime, (float)ThisTime);
 		}
-		else if (CacheKeys.Num() == 1)
-		{
-			Result.Add(CachedDataProbablyExists(*CacheKeys[0]));
-		}
 		return Result;
 	}
 
@@ -677,8 +731,28 @@ public:
 		{
 			DDC_SCOPE_CYCLE_COUNTER(DDC_TryToPrefetch);
 			UE_LOG(LogDerivedDataCache, VeryVerbose, TEXT("TryToPrefetch %d keys including %s from '%.*s'"),
-				CacheKeys.Num(), *CacheKeys[0], DebugContext.Len(), DebugContext.GetData());
-			return FDerivedDataBackend::Get().GetRoot().TryToPrefetch(CacheKeys);
+			CacheKeys.Num(), *CacheKeys[0], DebugContext.Len(), DebugContext.GetData());
+			TArray<FLegacyCacheGetRequest, TInlineAllocator<8>> LegacyRequests;
+			int32 Index = 0;
+			const FSharedString Name = DebugContext;
+			for (const FString& CacheKey : CacheKeys)
+			{
+				FLegacyCacheGetRequest& LegacyRequest = LegacyRequests.AddDefaulted_GetRef();
+				LegacyRequest.Name = Name;
+				LegacyRequest.Key = FLegacyCacheKey(CacheKey, FDerivedDataBackend::Get().GetMaxKeyLength());
+				LegacyRequest.Policy = ECachePolicy::Default | ECachePolicy::SkipData;
+				LegacyRequest.UserData = uint64(Index);
+				++Index;
+			}
+			bool bOk = true;
+			FRequestOwner BlockingOwner(EPriority::Blocking);
+			FDerivedDataBackend::Get().GetRoot().LegacyGet(LegacyRequests, BlockingOwner,
+				[&bOk](FLegacyCacheGetResponse&& Response)
+				{
+					bOk &= Response.Status == EStatus::Ok;
+				});
+			BlockingOwner.Wait();
+			return bOk;
 		}
 		return true;
 	}
diff --git a/Engine/Source/Developer/DerivedDataCache/Private/DerivedDataCacheStorePrivate.h b/Engine/Source/Developer/DerivedDataCache/Private/DerivedDataCacheStorePrivate.h
deleted file mode 100644
diff --git a/Engine/Source/Developer/DerivedDataCache/Private/DerivedDataLegacyCacheStore.cpp b/Engine/Source/Developer/DerivedDataCache/Private/DerivedDataLegacyCacheStore.cpp
new file mode 100644
--- /dev/null
+++ b/Engine/Source/Developer/DerivedDataCache/Private/DerivedDataLegacyCacheStore.cpp
@@ -0,0 +1,86 @@
+// Copyright Epic Games, Inc. All Rights Reserved.
+
+#include "DerivedDataLegacyCacheStore.h"
+
+#include "Containers/StringConv.h"
+#include "DerivedDataBackendInterface.h"
+#include "Misc/Crc.h"
+#include "Misc/SecureHash.h"
+#include "Misc/StringBuilder.h"
+#include "String/BytesToHex.h"
+#include "String/Find.h"
+
+namespace UE::DerivedData
+{
+
+FLegacyCacheKey::FLegacyCacheKey(const FStringView InFullKey, const int32 MaxKeyLength)
+	: FullKey(InFullKey)
+{
+	if (InFullKey.Len() > MaxKeyLength)
+	{
+		const auto FullKeyUCS2 = StringCast<UCS2CHAR>(InFullKey.GetData(), InFullKey.Len());
+		const int32 FullKeyLength = FullKeyUCS2.Length();
+		const uint32 CRCofPayload(FCrc::MemCrc32(FullKeyUCS2.Get(), FullKeyLength * sizeof(UCS2CHAR)));
+
+		uint8 Hash[FSHA1::DigestSize];
+		FSHA1 HashState;
+		HashState.Update((const uint8*)&FullKeyLength, sizeof(int32));
+		HashState.Update((const uint8*)&CRCofPayload, sizeof(uint32));
+		HashState.Update((const uint8*)FullKeyUCS2.Get(), FullKeyLength * sizeof(UCS2CHAR));
+		HashState.Final();
+		HashState.GetHash(Hash);
+
+		TStringBuilder<128> ShortKeyBuilder;
+		ShortKeyBuilder << InFullKey.Left(MaxKeyLength - FSHA1::DigestSize * 2 - 2) << TEXT("__");
+		String::BytesToHex(Hash, ShortKeyBuilder);
+
+		check(ShortKeyBuilder.Len() == MaxKeyLength && ShortKeyBuilder.Len() > 0);
+		ShortKey = ShortKeyBuilder;
+	}
+
+	TStringBuilder<64> Bucket;
+	Bucket << TEXT("Legacy");
+	if (const int32 BucketEnd = String::FindFirstChar(InFullKey, TEXT('_')); BucketEnd != INDEX_NONE)
+	{
+		Bucket << InFullKey.Left(BucketEnd);
+	}
+	Key.Bucket = FCacheBucket(Bucket);
+	Key.Hash = FIoHash::HashBuffer(MakeMemoryView(FTCHARToUTF8(InFullKey)));
+}
+
+bool FLegacyCacheKey::ReadValueTrailer(FCompositeBuffer& Value) const
+{
+	if (!ShortKey.IsEmpty())
+	{
+		TUtf8StringBuilder<256> FullKeyUtf8;
+		FullKeyUtf8 << FullKey;
+		const FMemoryView CompareKey(FullKeyUtf8.ToString(), (FullKeyUtf8.Len() + 1) * sizeof(UTF8CHAR));
+		if (Value.GetSize() < CompareKey.GetSize())
+		{
+			UE_LOG(LogDerivedDataCache, Display, TEXT("FLegacyCacheKey: Hash collision or short value for key %s."), *FullKey);
+			return false;
+		}
+		FUniqueBuffer CopyBuffer;
+		const uint64 KeyOffset = Value.GetSize() - CompareKey.GetSize();
+		const FMemoryView ValueKey = Value.ViewOrCopyRange(KeyOffset, CompareKey.GetSize(), CopyBuffer);
+		if (!CompareKey.EqualBytes(ValueKey))
+		{
+			UE_LOG(LogDerivedDataCache, Display, TEXT("FLegacyCacheKey: Hash collision for key %s."), *FullKey);
+			return false;
+		}
+		Value = Value.Mid(0, KeyOffset);
+	}
+	return true;
+}
+
+void FLegacyCacheKey::WriteValueTrailer(FCompositeBuffer& Value) const
+{
+	if (!ShortKey.IsEmpty())
+	{
+		TUtf8StringBuilder<256> FullKeyUtf8;
+		FullKeyUtf8 << FullKey;
+		Value = FCompositeBuffer(Value, FSharedBuffer::Clone(FullKeyUtf8.ToString(), (FullKeyUtf8.Len() + 1) * sizeof(UTF8CHAR)));
+	}
+}
+
+} // UE::DerivedData
diff --git a/Engine/Source/Developer/DerivedDataCache/Private/DerivedDataLegacyCacheStore.h b/Engine/Source/Developer/DerivedDataCache/Private/DerivedDataLegacyCacheStore.h
new file mode 100644
--- /dev/null
+++ b/Engine/Source/Developer/DerivedDataCache/Private/DerivedDataLegacyCacheStore.h
@@ -0,0 +1,112 @@
+// Copyright Epic Games, Inc. All Rights Reserved.
+
+#pragma once
+
+#include "DerivedDataCache.h"
+
+namespace UE::DerivedData { class FLegacyCacheKey; }
+namespace UE::DerivedData { struct FLegacyCacheDeleteRequest; }
+namespace UE::DerivedData { struct FLegacyCacheDeleteResponse; }
+namespace UE::DerivedData { struct FLegacyCacheGetRequest; }
+namespace UE::DerivedData { struct FLegacyCacheGetResponse; }
+namespace UE::DerivedData { struct FLegacyCachePutRequest; }
+namespace UE::DerivedData { struct FLegacyCachePutResponse; }
+
+namespace UE::DerivedData
+{
+
+using FOnLegacyCachePutComplete = TUniqueFunction<void (FLegacyCachePutResponse&& Response)>;
+using FOnLegacyCacheGetComplete = TUniqueFunction<void (FLegacyCacheGetResponse&& Response)>;
+using FOnLegacyCacheDeleteComplete = TUniqueFunction<void (FLegacyCacheDeleteResponse&& Response)>;
+
+class ILegacyCacheStore : public ICacheStore
+{
+public:
+	virtual void LegacyPut(
+		TConstArrayView<FLegacyCachePutRequest> Requests,
+		IRequestOwner& Owner,
+		FOnLegacyCachePutComplete&& OnComplete) = 0;
+
+	virtual void LegacyGet(
+		TConstArrayView<FLegacyCacheGetRequest> Requests,
+		IRequestOwner& Owner,
+		FOnLegacyCacheGetComplete&& OnComplete) = 0;
+
+	virtual void LegacyDelete(
+		TConstArrayView<FLegacyCacheDeleteRequest> Requests,
+		IRequestOwner& Owner,
+		FOnLegacyCacheDeleteComplete&& OnComplete) = 0;
+};
+
+class FLegacyCacheKey
+{
+public:
+	FLegacyCacheKey() = default;
+	FLegacyCacheKey(FStringView FullKey, int32 MaxKeyLength);
+
+	const FCacheKey& GetKey() const { return Key; }
+	const FSharedString& GetFullKey() const { return FullKey; }
+	const FSharedString& GetShortKey() const { return ShortKey.IsEmpty() ? FullKey : ShortKey; }
+	bool HasShortKey() const { return !ShortKey.IsEmpty(); }
+
+	bool ReadValueTrailer(FCompositeBuffer& Value) const;
+	void WriteValueTrailer(FCompositeBuffer& Value) const;
+
+private:
+	FCacheKey Key;
+	FSharedString FullKey;
+	FSharedString ShortKey;
+};
+
+struct FLegacyCachePutRequest
+{
+	FSharedString Name;
+	FLegacyCacheKey Key;
+	FCompositeBuffer Value;
+	ECachePolicy Policy = ECachePolicy::Default;
+	uint64 UserData = 0;
+};
+
+struct FLegacyCachePutResponse
+{
+	FSharedString Name;
+	FLegacyCacheKey Key;
+	uint64 UserData = 0;
+	EStatus Status = EStatus::Error;
+};
+
+struct FLegacyCacheGetRequest
+{
+	FSharedString Name;
+	FLegacyCacheKey Key;
+	ECachePolicy Policy = ECachePolicy::Default;
+	uint64 UserData = 0;
+};
+
+struct FLegacyCacheGetResponse
+{
+	FSharedString Name;
+	FLegacyCacheKey Key;
+	FSharedBuffer Value;
+	uint64 UserData = 0;
+	EStatus Status = EStatus::Error;
+};
+
+struct FLegacyCacheDeleteRequest
+{
+	FSharedString Name;
+	FLegacyCacheKey Key;
+	ECachePolicy Policy = ECachePolicy::Default;
+	bool bTransient = false;
+	uint64 UserData = 0;
+};
+
+struct FLegacyCacheDeleteResponse
+{
+	FSharedString Name;
+	FLegacyCacheKey Key;
+	uint64 UserData = 0;
+	EStatus Status = EStatus::Error;
+};
+
+} // UE::DerivedData
diff --git a/Engine/Source/Developer/DerivedDataCache/Private/DerivedDataLimitKeyLengthWrapper.h b/Engine/Source/Developer/DerivedDataCache/Private/DerivedDataLimitKeyLengthWrapper.h
deleted file mode 100644
