DDC: Added functions to transform a cache policy and build policy

#rnx

diff --git a/Engine/Source/Developer/DerivedDataCache/Private/DerivedDataBuild.cpp b/Engine/Source/Developer/DerivedDataCache/Private/DerivedDataBuild.cpp
--- a/Engine/Source/Developer/DerivedDataCache/Private/DerivedDataBuild.cpp
+++ b/Engine/Source/Developer/DerivedDataCache/Private/DerivedDataBuild.cpp
@@ -187,6 +187,20 @@ EBuildPolicy FBuildPolicy::GetValuePolicy(const FValueId& Id) const
 	return DefaultValuePolicy;
 }
 
+FBuildPolicy FBuildPolicy::Transform(TFunctionRef<EBuildPolicy (EBuildPolicy)> Op) const
+{
+	if (IsUniform())
+	{
+		return Op(CombinedPolicy);
+	}
+	FBuildPolicyBuilder Builder(Op(DefaultValuePolicy));
+	for (const FBuildValuePolicy& Value : GetValuePolicies())
+	{
+		Builder.AddValuePolicy({Value.Id, Op(Value.Policy)});
+	}
+	return Builder.Build();
+}
+
 void FBuildPolicyBuilder::AddValuePolicy(const FBuildValuePolicy& Policy)
 {
 	if (!Shared)
diff --git a/Engine/Source/Developer/DerivedDataCache/Private/DerivedDataBuildJob.cpp b/Engine/Source/Developer/DerivedDataCache/Private/DerivedDataBuildJob.cpp
--- a/Engine/Source/Developer/DerivedDataCache/Private/DerivedDataBuildJob.cpp
+++ b/Engine/Source/Developer/DerivedDataCache/Private/DerivedDataBuildJob.cpp
@@ -818,12 +818,7 @@ void FBuildJob::BeginExecuteRemote()
 	}
 	else
 	{
-		FBuildPolicyBuilder PolicyBuilder(BuildPolicy.GetDefaultValuePolicy() & Mask);
-		for (const FBuildValuePolicy& Value : BuildPolicy.GetValuePolicies())
-		{
-			PolicyBuilder.AddValuePolicy({Value.Id, Value.Policy & Mask});
-		}
-		RemoteBuildPolicy = PolicyBuilder.Build();
+		RemoteBuildPolicy = BuildPolicy.Transform([Mask](EBuildPolicy Policy) { return Policy & Mask; });
 	}
 
 	WorkerExecutor->Build(Action.Get(), Inputs, RemoteBuildPolicy, *Worker, BuildSystem, Owner,
diff --git a/Engine/Source/Developer/DerivedDataCache/Private/DerivedDataCache.cpp b/Engine/Source/Developer/DerivedDataCache/Private/DerivedDataCache.cpp
--- a/Engine/Source/Developer/DerivedDataCache/Private/DerivedDataCache.cpp
+++ b/Engine/Source/Developer/DerivedDataCache/Private/DerivedDataCache.cpp
@@ -215,6 +215,20 @@ ECachePolicy FCacheRecordPolicy::GetValuePolicy(const FValueId& Id) const
 	return DefaultValuePolicy;
 }
 
+FCacheRecordPolicy FCacheRecordPolicy::Transform(TFunctionRef<ECachePolicy (ECachePolicy)> Op) const
+{
+	if (IsUniform())
+	{
+		return Op(RecordPolicy);
+	}
+	FCacheRecordPolicyBuilder Builder(Op(DefaultValuePolicy));
+	for (const FCacheValuePolicy& Value : GetValuePolicies())
+	{
+		Builder.AddValuePolicy({Value.Id, Op(Value.Policy)});
+	}
+	return Builder.Build();
+}
+
 void FCacheRecordPolicyBuilder::AddValuePolicy(const FCacheValuePolicy& Policy)
 {
 	if (!Shared)
diff --git a/Engine/Source/Developer/DerivedDataCache/Public/DerivedDataBuildTypes.h b/Engine/Source/Developer/DerivedDataCache/Public/DerivedDataBuildTypes.h
--- a/Engine/Source/Developer/DerivedDataCache/Public/DerivedDataBuildTypes.h
+++ b/Engine/Source/Developer/DerivedDataCache/Public/DerivedDataBuildTypes.h
@@ -119,6 +119,9 @@ public:
 		return Shared ? Shared->GetValuePolicies() : TConstArrayView<FBuildValuePolicy>();
 	}
 
+	/** Returns a copy of this policy transformed by an operation. */
+	UE_API FBuildPolicy Transform(TFunctionRef<EBuildPolicy (EBuildPolicy)> Op) const;
+
 private:
 	friend class FBuildPolicyBuilder;
 
diff --git a/Engine/Source/Developer/DerivedDataCache/Public/DerivedDataCache.h b/Engine/Source/Developer/DerivedDataCache/Public/DerivedDataCache.h
--- a/Engine/Source/Developer/DerivedDataCache/Public/DerivedDataCache.h
+++ b/Engine/Source/Developer/DerivedDataCache/Public/DerivedDataCache.h
@@ -182,6 +182,9 @@ public:
 		return Shared ? Shared->GetValuePolicies() : TConstArrayView<FCacheValuePolicy>();
 	}
 
+	/** Returns a copy of this policy transformed by an operation. */
+	UE_API FCacheRecordPolicy Transform(TFunctionRef<ECachePolicy (ECachePolicy)> Op) const;
+
 private:
 	friend class FCacheRecordPolicyBuilder;
 
