Ensure Jupiter DDC backend implements GetChunks in a functionally correct way.
Fix erroneous logging in Jupiter DDC backend when a cache record was missed.
Fix breakage in Jupiter DDC backend automated test caused by the header re-org.
Add GetChunks tests to Jupiter DDC backend automated test and increase the size and  byte variety for values in the test content.

diff --git a/Engine/Source/Developer/DerivedDataCache/Private/HttpDerivedDataBackend.cpp b/Engine/Source/Developer/DerivedDataCache/Private/HttpDerivedDataBackend.cpp
--- a/Engine/Source/Developer/DerivedDataCache/Private/HttpDerivedDataBackend.cpp
+++ b/Engine/Source/Developer/DerivedDataCache/Private/HttpDerivedDataBackend.cpp
@@ -23,6 +23,7 @@
 #include "DerivedDataCachePrivate.h"
 #include "DerivedDataCacheRecord.h"
 #include "DerivedDataCacheUsageStats.h"
+#include "DerivedDataChunk.h"
 #include "DerivedDataValue.h"
 #include "Dom/JsonObject.h"
 #include "Experimental/Containers/FAAArrayQueue.h"
@@ -2075,27 +2076,27 @@ private:
 	bool ShouldSimulateMiss(const TCHAR* InKey);
 	bool ShouldSimulateMiss(const FCacheKey& Key);
 
-	bool PutCacheRecord(const FCacheRecord& Record, FStringView Context, const FCacheRecordPolicy& Policy, uint64& OutWriteSize);
+	bool PutCacheRecord(FStringView Name, const FCacheRecord& Record, const FCacheRecordPolicy& Policy, uint64& OutWriteSize);
 	uint64 PutRef(const FCacheRecord& Record, const FCbPackage& Package, FStringView Bucket, bool bFinalize, TArray<FIoHash>& OutNeededBlobHashes, bool& bOutPutCompletedSuccessfully);
 
 	FOptionalCacheRecord GetCacheRecordOnly(
+		FStringView Name,
 		const FCacheKey& Key,
-		const FStringView Context,
 		const FCacheRecordPolicy& Policy);
 	FOptionalCacheRecord GetCacheRecord(
+		FStringView Name,
 		const FCacheKey& Key,
-		const FStringView Context,
 		const FCacheRecordPolicy& Policy,
 		EStatus& OutStatus);
 	bool TryGetCachedDataBatch(
+		FStringView Name,
 		const FCacheKey& Key,
-		TArrayView<FValueWithId> Values,
-		const FStringView Context,
+		TConstArrayView<FValueWithId> Values,
 		TArray<FCompressedBuffer>& OutBuffers);
 	bool CachedDataProbablyExistsBatch(
+		FStringView Name,
 		const FCacheKey& Key,
-		TConstArrayView<FValueWithId> Values,
-		const FStringView Context);
+		TConstArrayView<FValueWithId> Values);
 };
 
 FHttpDerivedDataBackend::FHttpDerivedDataBackend(
@@ -2433,8 +2434,8 @@ uint64 FHttpDerivedDataBackend::PutRef(const FCacheRecord& Record, const FCbPack
 }
 
 bool FHttpDerivedDataBackend::PutCacheRecord(
+	FStringView Name,
 	const FCacheRecord& Record,
-	const FStringView Context,
 	const FCacheRecordPolicy& Policy,
 	uint64& OutWriteSize)
 {
@@ -2444,7 +2445,7 @@ bool FHttpDerivedDataBackend::PutCacheRecord(
 	{
 		UE_LOG(LogDerivedDataCache, VeryVerbose,
 			TEXT("%s: Skipped put of %s from '%.*s' because this cache store is read-only"),
-			*GetName(), *WriteToString<96>(Record.GetKey()), Context.Len(), Context.GetData());
+			*GetName(), *WriteToString<96>(Record.GetKey()), Name.Len(), Name.GetData());
 		return false;
 	}
 
@@ -2459,14 +2460,14 @@ bool FHttpDerivedDataBackend::PutCacheRecord(
 	if (!EnumHasAnyFlags(CombinedPolicy, StoreFlag))
 	{
 		UE_LOG(LogDerivedDataCache, VeryVerbose, TEXT("%s: Skipped put of %s from '%.*s' due to cache policy"),
-			*GetName(), *WriteToString<96>(Key), Context.Len(), Context.GetData());
+			*GetName(), *WriteToString<96>(Key), Name.Len(), Name.GetData());
 		return false;
 	}
 
 	if (ShouldSimulateMiss(Key))
 	{
 		UE_LOG(LogDerivedDataCache, Verbose, TEXT("%s: Simulated miss for put of %s from '%.*s'"),
-			*GetName(), *WriteToString<96>(Key), Context.Len(), Context.GetData());
+			*GetName(), *WriteToString<96>(Key), Name.Len(), Name.GetData());
 		return false;
 	}
 
@@ -2542,15 +2543,15 @@ bool FHttpDerivedDataBackend::PutCacheRecord(
 }
 
 FOptionalCacheRecord FHttpDerivedDataBackend::GetCacheRecordOnly(
+	const FStringView Name,
 	const FCacheKey& Key,
-	const FStringView Context,
 	const FCacheRecordPolicy& Policy)
 {
 	if (!IsUsable())
 	{
 		UE_LOG(LogDerivedDataCache, VeryVerbose,
 			TEXT("%s: Skipped get of %s from '%.*s' because this cache store is not available"),
-			*GetName(), *WriteToString<96>(Key), Context.Len(), Context.GetData());
+			*GetName(), *WriteToString<96>(Key), Name.Len(), Name.GetData());
 		return FOptionalCacheRecord();
 	}
 
@@ -2560,14 +2561,14 @@ FOptionalCacheRecord FHttpDerivedDataBackend::GetCacheRecordOnly(
 	if (!EnumHasAnyFlags(Policy.GetRecordPolicy(), QueryPolicy))
 	{
 		UE_LOG(LogDerivedDataCache, VeryVerbose, TEXT("%s: Skipped get of %s from '%.*s' due to cache policy"),
-			*GetName(), *WriteToString<96>(Key), Context.Len(), Context.GetData());
+			*GetName(), *WriteToString<96>(Key), Name.Len(), Name.GetData());
 		return FOptionalCacheRecord();
 	}
 
 	if (ShouldSimulateMiss(Key))
 	{
 		UE_LOG(LogDerivedDataCache, Verbose, TEXT("%s: Simulated miss for get of %s from '%.*s'"),
-			*GetName(), *WriteToString<96>(Key), Context.Len(), Context.GetData());
+			*GetName(), *WriteToString<96>(Key), Name.Len(), Name.GetData());
 		return FOptionalCacheRecord();
 	}
 
@@ -2580,6 +2581,7 @@ FOptionalCacheRecord FHttpDerivedDataBackend::GetCacheRecordOnly(
 		TStringBuilder<256> RefsUri;
 		RefsUri << "api/v1/refs/" << StructuredNamespace << "/" << Bucket << "/" << Key.Hash;
 
+		bool bIsSuccessfulResponse = false;
 		FSharedBuffer ResponseBuffer;
 		int64 ResponseCode = 0;
 		for (uint32 Attempts = 0; (Attempts < UE_HTTPDDC_MAX_ATTEMPTS) && !ShouldAbortForShutdown() && (Attempts == 0 || ShouldRetryOnError(ResponseCode)); ++Attempts)
@@ -2594,22 +2596,30 @@ FOptionalCacheRecord FHttpDerivedDataBackend::GetCacheRecordOnly(
 			if (FHttpRequest::IsSuccessResponse(ResponseCode))
 			{
 				ResponseBuffer = MakeSharedBufferFromArray(MoveTemp(ByteArray));
+				bIsSuccessfulResponse = true;
 				break;
 			}
 		}
 
+		if (!bIsSuccessfulResponse)
+		{
+			UE_LOG(LogDerivedDataCache, Verbose, TEXT("%s: Cache miss with missing package for %s from '%.*s'"),
+				*GetName(), *WriteToString<96>(Key), Name.Len(), Name.GetData());
+			return Record;
+		}
+
 		if (ValidateCompactBinary(ResponseBuffer, ECbValidateMode::Default) != ECbValidateError::None)
 		{
 			UE_LOG(LogDerivedDataCache, Display, TEXT("%s: Cache miss with invalid package for %s from '%.*s'"),
-				*GetName(), *WriteToString<96>(Key), Context.Len(), Context.GetData());
-			return FOptionalCacheRecord();
+				*GetName(), *WriteToString<96>(Key), Name.Len(), Name.GetData());
+			return Record;
 		}
 
 		Record = FCacheRecord::Load(FCbPackage(FCbObject(ResponseBuffer)));
 		if (Record.IsNull())
 		{
 			UE_LOG(LogDerivedDataCache, Display, TEXT("%s: Cache miss with record load failure for %s from '%.*s'"),
-				*GetName(), *WriteToString<96>(Key), Context.Len(), Context.GetData());
+				*GetName(), *WriteToString<96>(Key), Name.Len(), Name.GetData());
 			return Record;
 		}
 	}
@@ -2617,13 +2627,32 @@ FOptionalCacheRecord FHttpDerivedDataBackend::GetCacheRecordOnly(
 	return Record;
 }
 
+static bool IsValueDataReady(FValue& Value, const ECachePolicy Policy)
+{
+	if (!EnumHasAnyFlags(Policy, ECachePolicy::Query))
+	{
+		Value = Value.RemoveData();
+		return true;
+	}
+
+	if (Value.HasData())
+	{
+		if (EnumHasAnyFlags(Policy, ECachePolicy::SkipData))
+		{
+			Value = Value.RemoveData();
+		}
+		return true;
+	}
+	return false;
+};
+
 FOptionalCacheRecord FHttpDerivedDataBackend::GetCacheRecord(
+	const FStringView Name,
 	const FCacheKey& Key,
-	const FStringView Context,
 	const FCacheRecordPolicy& Policy,
 	EStatus& OutStatus)
 {
-	FOptionalCacheRecord Record = GetCacheRecordOnly(Key, Context, Policy);
+	FOptionalCacheRecord Record = GetCacheRecordOnly(Name, Key, Policy);
 	if (Record.IsNull())
 	{
 		OutStatus = EStatus::Error;
@@ -2645,36 +2674,16 @@ FOptionalCacheRecord FHttpDerivedDataBackend::GetCacheRecord(
 	TArray<FValueWithId> RequiredGets;
 	TArray<FValueWithId> RequiredHeads;
 
-	auto IsValueDataReady = [](FValueWithId& Value, const ECachePolicy Policy, const ECachePolicy SkipFlag)
-	{
-		if (!EnumHasAnyFlags(Policy, ECachePolicy::Query))
-		{
-			Value = Value.RemoveData();
-			return true;
-		}
-
-		if (Value.HasData())
-		{
-			if (EnumHasAnyFlags(Policy, SkipFlag))
-			{
-				Value = Value.RemoveData();
-			}
-			return true;
-		}
-		return false;
-	};
-
 	for (FValueWithId Value : Record.Get().GetValues())
 	{
 		const ECachePolicy ValuePolicy = Policy.GetValuePolicy(Value.GetId());
-		const ECachePolicy SkipFlag = ECachePolicy::SkipData;
-		if (IsValueDataReady(Value, ValuePolicy, SkipFlag))
+		if (IsValueDataReady(Value, ValuePolicy))
 		{
 			RecordBuilder.AddValue(MoveTemp(Value));
 		}
 		else
 		{
-			if (EnumHasAnyFlags(ValuePolicy, SkipFlag))
+			if (EnumHasAnyFlags(ValuePolicy, ECachePolicy::SkipData))
 			{
 				RequiredHeads.Emplace(Value);
 			}
@@ -2685,14 +2694,14 @@ FOptionalCacheRecord FHttpDerivedDataBackend::GetCacheRecord(
 		}
 	}
 
-	if (!CachedDataProbablyExistsBatch(Key, RequiredHeads, Context))
+	if (!CachedDataProbablyExistsBatch(Name, Key, RequiredHeads))
 	{
 		OutStatus = EStatus::Error;
 		return FOptionalCacheRecord();
 	}
 
 	TArray<FCompressedBuffer> FetchedBuffers;
-	if (!TryGetCachedDataBatch(Key, RequiredGets, Context, FetchedBuffers))
+	if (!TryGetCachedDataBatch(Name, Key, RequiredGets, FetchedBuffers))
 	{
 		OutStatus = EStatus::Error;
 		return FOptionalCacheRecord();
@@ -2712,9 +2721,9 @@ FOptionalCacheRecord FHttpDerivedDataBackend::GetCacheRecord(
 }
 
 bool FHttpDerivedDataBackend::TryGetCachedDataBatch(
+	const FStringView Name,
 	const FCacheKey& Key,
-	TArrayView<FValueWithId> Values,
-	const FStringView Context,
+	TConstArrayView<FValueWithId> Values,
 	TArray<FCompressedBuffer>& OutBuffers)
 {
 	for (const FValueWithId& Value : Values)
@@ -2747,7 +2756,7 @@ bool FHttpDerivedDataBackend::TryGetCachedDataBatch(
 			UE_LOG(LogDerivedDataCache, Verbose,
 				TEXT("%s: Cache miss with missing value %s with hash %s for %s from '%.*s'"),
 				*GetName(), *WriteToString<16>(Value.GetId()), *WriteToString<48>(RawHash), *WriteToString<96>(Key),
-				Context.Len(), Context.GetData());
+				Name.Len(), Name.GetData());
 			return false;
 		}
 		else if (CompressedBuffer.GetRawHash() != RawHash)
@@ -2755,7 +2764,7 @@ bool FHttpDerivedDataBackend::TryGetCachedDataBatch(
 			UE_LOG(LogDerivedDataCache, Display,
 				TEXT("%s: Cache miss with corrupted value %s with hash %s for %s from '%.*s'"),
 				*GetName(), *WriteToString<16>(Value.GetId()), *WriteToString<48>(RawHash),
-				*WriteToString<96>(Key), Context.Len(), Context.GetData());
+				*WriteToString<96>(Key), Name.Len(), Name.GetData());
 			return false;
 		}
 		else
@@ -2767,9 +2776,9 @@ bool FHttpDerivedDataBackend::TryGetCachedDataBatch(
 }
 
 bool FHttpDerivedDataBackend::CachedDataProbablyExistsBatch(
+	const FStringView Name,
 	const FCacheKey& Key,
-	TConstArrayView<FValueWithId> Values,
-	const FStringView Context)
+	TConstArrayView<FValueWithId> Values)
 {
 	for (const FValueWithId& Value : Values)
 	{
@@ -2798,7 +2807,7 @@ bool FHttpDerivedDataBackend::CachedDataProbablyExistsBatch(
 			UE_LOG(LogDerivedDataCache, Verbose,
 				TEXT("%s: Cache miss with missing value %s with hash %s for %s from '%.*s'"),
 				*GetName(), *WriteToString<16>(Value.GetId()), *WriteToString<48>(RawHash), *WriteToString<96>(Key),
-				Context.Len(), Context.GetData());
+				Name.Len(), Name.GetData());
 			return false;
 		}
 	}
@@ -3170,7 +3179,7 @@ void FHttpDerivedDataBackend::Put(
 		const FCacheRecord& Record = Request.Record;
 		COOK_STAT(auto Timer = UsageStats.TimePut());
 		uint64 BytesSent = 0;
-		if (PutCacheRecord(Record, Request.Name, Request.Policy, BytesSent))
+		if (PutCacheRecord(Request.Name, Record, Request.Policy, BytesSent))
 		{
 			UE_LOG(LogDerivedDataCache, Verbose, TEXT("%s: Cache put complete for %s from '%s'"),
 				*GetName(), *WriteToString<96>(Record.GetKey()), *Request.Name);
@@ -3200,7 +3209,7 @@ void FHttpDerivedDataBackend::Get(
 	{
 		COOK_STAT(auto Timer = UsageStats.TimeGet());
 		EStatus Status = EStatus::Ok;
-		if (FOptionalCacheRecord Record = GetCacheRecord(Request.Key, Request.Name, Request.Policy, Status))
+		if (FOptionalCacheRecord Record = GetCacheRecord(Request.Name, Request.Key, Request.Policy, Status))
 		{
 			UE_LOG(LogDerivedDataCache, Verbose, TEXT("%s: Cache hit for %s from '%s'"),
 				*GetName(), *WriteToString<96>(Request.Key), *Request.Name);
@@ -3226,9 +3235,102 @@ void FHttpDerivedDataBackend::GetChunks(
 	IRequestOwner& Owner,
 	FOnCacheChunkComplete&& OnComplete)
 {
-	if (OnComplete)
+	// TODO: This is inefficient because Jupiter doesn't allow us to get only part of a compressed blob, so we have to
+	//		 get the whole thing and then decompress only the portion we need.  Furthermore, because there is no propagation
+	//		 between cache stores during chunk requests, the fetched result won't end up in the local store.
+	//		 These efficiency issues will be addressed by changes to the Hierarchy that translate chunk requests that
+	//		 are missing in local/fast stores and have to be retrieved from slow stores into record requests instead.  That
+	//		 will make this code path unused/uncommon as Jupiter will most always be a slow store with a local/fast store in front of it.
+	//		 Regardless, to adhere to the functional contract, this implementation must exist.
+	TArray<FCacheChunkRequest, TInlineAllocator<16>> SortedRequests(Requests);
+	SortedRequests.StableSort(TChunkLess());
+
+	bool bHasValue = false;
+	FValue Value;
+	FValueId ValueId;
+	FCacheKey ValueKey;
+	FCompressedBuffer ValueBuffer;
+	FCompressedBufferReader ValueReader;
+	EStatus ValueStatus = EStatus::Error;
+	FOptionalCacheRecord Record;
+	for (const FCacheChunkRequest& Request : SortedRequests)
 	{
-		for (const FCacheChunkRequest& Request : Requests)
+		const bool bExistsOnly = EnumHasAnyFlags(Request.Policy, ECachePolicy::SkipData);
+		COOK_STAT(auto Timer = bExistsOnly ? UsageStats.TimeProbablyExists() : UsageStats.TimeGet());
+		if (!(bHasValue && ValueKey == Request.Key && ValueId == Request.Id) || ValueReader.HasSource() < !bExistsOnly)
+		{
+			ValueStatus = EStatus::Error;
+			ValueReader.ResetSource();
+			ValueKey = {};
+			ValueId.Reset();
+			Value.Reset();
+			bHasValue = false;
+			if (Request.Id.IsValid())
+			{
+				if (!(Record && Record.Get().GetKey() == Request.Key))
+				{
+					FCacheRecordPolicyBuilder PolicyBuilder(ECachePolicy::None);
+					PolicyBuilder.AddValuePolicy(Request.Id, Request.Policy);
+					Record.Reset();
+					Record = GetCacheRecordOnly(Request.Name, Request.Key, PolicyBuilder.Build());
+				}
+				if (Record)
+				{
+					const FValueWithId& ValueWithId = Record.Get().GetValue(Request.Id);
+					bHasValue = ValueWithId.IsValid();
+					Value = ValueWithId;
+					ValueId = Request.Id;
+					ValueKey = Request.Key;
+
+					if (IsValueDataReady(Value, Request.Policy))
+					{
+						ValueReader.SetSource(Value.GetData());
+					}
+					else
+					{
+						TArray<FCompressedBuffer> ValueBuffers;
+						if (TryGetCachedDataBatch(Request.Name, Request.Key, ::MakeArrayView({ ValueWithId }), ValueBuffers))
+						{
+							ValueBuffer = ValueBuffers[0];
+							ValueReader.SetSource(ValueBuffer);
+						}
+						else
+						{
+							ValueBuffer.Reset();
+							ValueReader.ResetSource();
+						}
+					}
+				}
+			}
+			else
+			{
+				UE_LOG(LogDerivedDataCache, Verbose, TEXT("%s: Cache miss with unimplemented value support for %s from '%s'"),
+					*GetName(), *WriteToString<96>(Request.Key, '/', Request.Id), *Request.Name);
+				checkf(false, TEXT("Value API not implemented for HTTP backend"));
+			}
+		}
+		if (bHasValue)
+		{
+			const uint64 RawOffset = FMath::Min(Value.GetRawSize(), Request.RawOffset);
+			const uint64 RawSize = FMath::Min(Value.GetRawSize() - RawOffset, Request.RawSize);
+			UE_LOG(LogDerivedDataCache, Verbose, TEXT("%s: Cache hit for %s from '%s'"),
+				*GetName(), *WriteToString<96>(Request.Key, '/', Request.Id), *Request.Name);
+			COOK_STAT(Timer.AddHit(!bExistsOnly ? RawSize : 0));
+			if (OnComplete)
+			{
+				FSharedBuffer Buffer;
+				if (!bExistsOnly)
+				{
+					Buffer = ValueReader.Decompress(RawOffset, RawSize);
+				}
+				const EStatus ChunkStatus = bExistsOnly || Buffer.GetSize() == RawSize ? EStatus::Ok : EStatus::Error;
+				OnComplete({Request.Name, Request.Key, Request.Id, Request.RawOffset,
+					RawSize, Value.GetRawHash(), MoveTemp(Buffer), Request.UserData, ChunkStatus});
+			}
+			continue;
+		}
+
+		if (OnComplete)
 		{
 			OnComplete({Request.Name, Request.Key, Request.Id, Request.RawOffset, 0, {}, {}, Request.UserData, EStatus::Error});
 		}
@@ -3277,7 +3379,18 @@ FDerivedDataBackendInterface* GetAnyHttpDerivedDataBackend(
 	FString& OutStructuredNamespace)
 {
 #if WITH_HTTP_DDC_BACKEND
-	return FHttpDerivedDataBackend::GetAny();
+	if (FHttpDerivedDataBackend* HttpBackend = FHttpDerivedDataBackend::GetAny())
+	{
+		OutDomain = HttpBackend->GetDomain();
+		OutOAuthProvider = HttpBackend->GetOAuthProvider();
+		OutOAuthClientId = HttpBackend->GetOAuthClientId();
+		OutOAuthSecret = HttpBackend->GetOAuthSecret();
+		OutNamespace = HttpBackend->GetNamespace();
+		OutStructuredNamespace = HttpBackend->GetStructuredNamespace();
+
+		return HttpBackend;
+	}
+	return nullptr;
 #else
 	return nullptr;
 #endif
diff --git a/Engine/Source/Developer/DerivedDataCache/Private/Tests/HttpDerivedDataBackendTest.cpp b/Engine/Source/Developer/DerivedDataCache/Private/Tests/HttpDerivedDataBackendTest.cpp
--- a/Engine/Source/Developer/DerivedDataCache/Private/Tests/HttpDerivedDataBackendTest.cpp
+++ b/Engine/Source/Developer/DerivedDataCache/Private/Tests/HttpDerivedDataBackendTest.cpp
@@ -221,6 +221,55 @@ protected:
 
 		return true;
 	}
+
+	bool GetRecordChunks(TConstArrayView<UE::DerivedData::FCacheRecord> Records, UE::DerivedData::FCacheRecordPolicy Policy, uint64 Offset, uint64 Size, TArray<FSharedBuffer>& OutChunks)
+	{
+		using namespace UE::DerivedData;
+		UE::DerivedData::ICacheStore* TestBackend = GetTestBackend();
+
+		TArray<FCacheChunkRequest> Requests;
+
+		int32 OverallIndex = 0;
+		for (int32 RecordIndex = 0; RecordIndex < Records.Num(); ++RecordIndex)
+		{
+			const FCacheRecord& Record = Records[RecordIndex];
+			TConstArrayView<FValueWithId> Values = Record.GetValues();
+			for (int32 ValueIndex = 0; ValueIndex < Values.Num(); ++ValueIndex)
+			{
+				const FValueWithId& Value = Values[ValueIndex];
+				Requests.Add({ {TEXT("FHttpDerivedDataTestBase")}, Record.GetKey(),  Value.GetId(), Offset, Size, Value.GetRawHash(), Policy.GetValuePolicy(Value.GetId()), static_cast<uint64>(OverallIndex) });
+				++OverallIndex;
+			}
+		}
+
+		struct FGetChunksOutput
+		{
+			FSharedBuffer Chunk;
+			EStatus Status = EStatus::Error;
+		};
+
+		TArray<TOptional<FGetChunksOutput>> GetOutputs;
+		GetOutputs.SetNum(Requests.Num());
+		FRequestOwner RequestOwner(EPriority::Blocking);
+		TestBackend->GetChunks(Requests, RequestOwner, [&GetOutputs](FCacheChunkResponse&& Response)
+			{
+				GetOutputs[Response.UserData].Emplace(FGetChunksOutput { Response.RawData, Response.Status });
+			});
+		RequestOwner.Wait();
+
+		for (int32 RequestIndex = 0; RequestIndex < Requests.Num(); ++RequestIndex)
+		{
+			FGetChunksOutput& ReceivedOutput = GetOutputs[RequestIndex].GetValue();
+			if (ReceivedOutput.Status != EStatus::Ok)
+			{
+				return false;
+			}
+			OutChunks.Add(MoveTemp(ReceivedOutput.Chunk));
+		}
+
+		return true;
+
+	}
 	
 	void ValidateRecords(const TCHAR* Name, TConstArrayView<UE::DerivedData::FCacheRecord> RecordsToTest, TConstArrayView<UE::DerivedData::FCacheRecord> ReferenceRecords, UE::DerivedData::FCacheRecordPolicy Policy)
 	{
@@ -265,6 +314,51 @@ protected:
 		}
 	}
 
+	void ValidateRecordChunks(const TCHAR* Name, TConstArrayView<FSharedBuffer> RecordChunksToTest, TConstArrayView<UE::DerivedData::FCacheRecord> ReferenceRecords, UE::DerivedData::FCacheRecordPolicy Policy, uint64 Offset, uint64 Size)
+	{
+		using namespace UE::DerivedData;
+
+		int32 TotalChunks = 0;
+		for (int32 RecordIndex = 0; RecordIndex < ReferenceRecords.Num(); ++RecordIndex)
+		{
+			const FCacheRecord& Record = ReferenceRecords[RecordIndex];
+			TConstArrayView<FValueWithId> Values = Record.GetValues();
+			for (int32 ValueIndex = 0; ValueIndex < Values.Num(); ++ValueIndex)
+			{
+				++TotalChunks;
+			}
+		}
+
+		if (!TestEqual(FString::Printf(TEXT("%s::Chunk quantity"), Name), RecordChunksToTest.Num(), TotalChunks))
+		{
+			return;
+		}
+
+		int32 ChunkIndex = 0;
+		for (int32 RecordIndex = 0; RecordIndex < ReferenceRecords.Num(); ++RecordIndex)
+		{
+			const FCacheRecord& ExpectedRecord = ReferenceRecords[RecordIndex];
+
+			const TConstArrayView<FValueWithId> ExpectedValues = ExpectedRecord.GetValues();
+			for (int32 ValueIndex = 0; ValueIndex < ExpectedValues.Num(); ++ValueIndex)
+			{
+				const FSharedBuffer& ChunkToTest = RecordChunksToTest[ChunkIndex];
+
+				if (EnumHasAnyFlags(Policy.GetRecordPolicy(), ECachePolicy::SkipData))
+				{
+					TestTrue(FString::Printf(TEXT("%s::Get chunk[%d] IsNull"), Name, ChunkIndex), ChunkToTest.IsNull());
+				}
+				else
+				{
+					FSharedBuffer ReferenceBuffer = ExpectedValues[ValueIndex].GetData().Decompress();
+					FMemoryView ReferenceView = ReferenceBuffer.GetView().Mid(Offset, Size);
+					TestTrue(FString::Printf(TEXT("%s::Get chunk[%d] data equality"), Name, ChunkIndex), ReferenceView.EqualBytes(ChunkToTest.GetView()));
+				}
+				++ChunkIndex;
+			}
+		}
+	}
+
 	TArray<UE::DerivedData::FCacheRecord> GetAndValidateRecords(const TCHAR* Name, TConstArrayView<UE::DerivedData::FCacheRecord> Records, UE::DerivedData::FCacheRecordPolicy Policy)
 	{
 		using namespace UE::DerivedData;
@@ -281,6 +375,28 @@ protected:
 		return ReceivedRecords;
 	}
 
+	TArray<FSharedBuffer> GetAndValidateRecordChunks(const TCHAR* Name, TConstArrayView<UE::DerivedData::FCacheRecord> Records, UE::DerivedData::FCacheRecordPolicy Policy, uint64 Offset, uint64 Size)
+	{
+		using namespace UE::DerivedData;
+		TArray<FSharedBuffer> ReceivedChunks;
+		bool bGetSuccessful = GetRecordChunks(Records, Policy, Offset, Size, ReceivedChunks);
+		TestTrue(FString::Printf(TEXT("%s::GetChunks status"), Name), bGetSuccessful);
+
+		if (!bGetSuccessful)
+		{
+			return TArray<FSharedBuffer>();
+		}
+
+		ValidateRecordChunks(Name, ReceivedChunks, Records, Policy, Offset, Size);
+		return ReceivedChunks;
+	}
+
+	TArray<UE::DerivedData::FCacheRecord> GetAndValidateRecordsAndChunks(const TCHAR* Name, TConstArrayView<UE::DerivedData::FCacheRecord> Records, UE::DerivedData::FCacheRecordPolicy Policy)
+	{
+		GetAndValidateRecordChunks(Name, Records, Policy, 5, 5);
+		return GetAndValidateRecords(Name, Records, Policy);
+	}
+
 protected:
 	static inline FString TestDomain;
 	static inline FString TestOAuthProvider;
@@ -330,11 +446,12 @@ TArray<UE::DerivedData::FCacheRecord> CreateTestCacheRecords(UE::DerivedData::IC
 		for (uint32 ValueIndex = 0; ValueIndex < InNumValues; ++ValueIndex)
 		{
 			TArray<uint8> ValueContents;
-			// Add N zeroed bytes where N corresponds to the value index.
-			ValueContents.AddUninitialized(ValueIndex+1);
-			for (uint32 ContentIndex = 0; ContentIndex < (ValueIndex+1); ++ContentIndex)
+			// Add N zeroed bytes where N corresponds to the value index times 10.
+			const int32 NumBytes = (ValueIndex+1)*10;
+			ValueContents.AddUninitialized(NumBytes);
+			for (int32 ContentIndex = 0; ContentIndex < NumBytes; ++ContentIndex)
 			{
-				ValueContents[ContentIndex] = (uint8)KeyIndex;
+				ValueContents[ContentIndex] = (uint8)(KeyIndex + ContentIndex);
 			}
 			Values.Emplace(MakeSharedBufferFromArray(MoveTemp(ValueContents)));
 			HashBuilder.Update(Values.Last().GetView());
@@ -480,9 +597,9 @@ bool CacheStore::RunTest(const FString& Parameters)
 
 	{
 		TArray<FCacheRecord> PutRecords = CreateTestCacheRecords(TestBackend, RecordsInBatch, 1);
-		TArray<FCacheRecord> RecievedRecords = GetAndValidateRecords(TEXT("SimpleValue"), PutRecords, ECachePolicy::Default);
-		TArray<FCacheRecord> RecievedRecordsSkipMeta = GetAndValidateRecords(TEXT("SimpleValueSkipMeta"), PutRecords, ECachePolicy::Default | ECachePolicy::SkipMeta);
-		TArray<FCacheRecord> RecievedRecordsSkipData = GetAndValidateRecords(TEXT("SimpleValueSkipData"), PutRecords, ECachePolicy::Default | ECachePolicy::SkipData);
+		TArray<FCacheRecord> RecievedRecords = GetAndValidateRecordsAndChunks(TEXT("SimpleValue"), PutRecords, ECachePolicy::Default);
+		TArray<FCacheRecord> RecievedRecordsSkipMeta = GetAndValidateRecordsAndChunks(TEXT("SimpleValueSkipMeta"), PutRecords, ECachePolicy::Default | ECachePolicy::SkipMeta);
+		TArray<FCacheRecord> RecievedRecordsSkipData = GetAndValidateRecordsAndChunks(TEXT("SimpleValueSkipData"), PutRecords, ECachePolicy::Default | ECachePolicy::SkipData);
 
 #if UE_WITH_ZEN
 		if (ZenIntermediaryBackend)
@@ -504,9 +621,9 @@ bool CacheStore::RunTest(const FString& Parameters)
 		FCbObject MetaObject = MetaWriter.Save().AsObject();
 
 		TArray<FCacheRecord> PutRecords = CreateTestCacheRecords(TestBackend, RecordsInBatch, 1, MetaObject);
-		TArray<FCacheRecord> RecievedRecords = GetAndValidateRecords(TEXT("SimpleValueWithMeta"), PutRecords, ECachePolicy::Default);
-		TArray<FCacheRecord> RecievedRecordsSkipMeta = GetAndValidateRecords(TEXT("SimpleValueWithMetaSkipMeta"), PutRecords, ECachePolicy::Default | ECachePolicy::SkipMeta);
-		TArray<FCacheRecord> RecievedRecordsSkipData = GetAndValidateRecords(TEXT("SimpleValueWithMetaSkipData"), PutRecords, ECachePolicy::Default | ECachePolicy::SkipData);
+		TArray<FCacheRecord> RecievedRecords = GetAndValidateRecordsAndChunks(TEXT("SimpleValueWithMeta"), PutRecords, ECachePolicy::Default);
+		TArray<FCacheRecord> RecievedRecordsSkipMeta = GetAndValidateRecordsAndChunks(TEXT("SimpleValueWithMetaSkipMeta"), PutRecords, ECachePolicy::Default | ECachePolicy::SkipMeta);
+		TArray<FCacheRecord> RecievedRecordsSkipData = GetAndValidateRecordsAndChunks(TEXT("SimpleValueWithMetaSkipData"), PutRecords, ECachePolicy::Default | ECachePolicy::SkipData);
 
 #if UE_WITH_ZEN
 		if (ZenIntermediaryBackend)
@@ -522,9 +639,9 @@ bool CacheStore::RunTest(const FString& Parameters)
 
 	{
 		TArray<FCacheRecord> PutRecords = CreateTestCacheRecords(TestBackend, RecordsInBatch, 5);
-		TArray<FCacheRecord> RecievedRecords = GetAndValidateRecords(TEXT("MultiValue"), PutRecords, ECachePolicy::Default);
-		TArray<FCacheRecord> RecievedRecordsSkipMeta = GetAndValidateRecords(TEXT("MultiValueSkipMeta"), PutRecords, ECachePolicy::Default | ECachePolicy::SkipMeta);
-		TArray<FCacheRecord> RecievedRecordsSkipData = GetAndValidateRecords(TEXT("MultiValueSkipData"), PutRecords, ECachePolicy::Default | ECachePolicy::SkipData);
+		TArray<FCacheRecord> RecievedRecords = GetAndValidateRecordsAndChunks(TEXT("MultiValue"), PutRecords, ECachePolicy::Default);
+		TArray<FCacheRecord> RecievedRecordsSkipMeta = GetAndValidateRecordsAndChunks(TEXT("MultiValueSkipMeta"), PutRecords, ECachePolicy::Default | ECachePolicy::SkipMeta);
+		TArray<FCacheRecord> RecievedRecordsSkipData = GetAndValidateRecordsAndChunks(TEXT("MultiValueSkipData"), PutRecords, ECachePolicy::Default | ECachePolicy::SkipData);
 
 #if UE_WITH_ZEN
 		if (ZenIntermediaryBackend)
