DDC: Fixed an assert when doing a put from a cache plugin

#rnx

diff --git a/Engine/Source/Developer/DerivedDataCache/Private/DerivedDataCache.cpp b/Engine/Source/Developer/DerivedDataCache/Private/DerivedDataCache.cpp
--- a/Engine/Source/Developer/DerivedDataCache/Private/DerivedDataCache.cpp
+++ b/Engine/Source/Developer/DerivedDataCache/Private/DerivedDataCache.cpp
@@ -445,7 +445,7 @@ class FDerivedDataCache final
 						FLegacyCachePutRequest LegacyRequest;
 						LegacyRequest.Name = DebugContext;
 						LegacyRequest.Key = FLegacyCacheKey(CacheKey, FDerivedDataBackend::Get().GetMaxKeyLength());
-						LegacyRequest.Value = FCompositeBuffer(MakeSharedBufferFromArray(MoveTemp(Data)));
+						LegacyRequest.Value = FCompositeBuffer(FSharedBuffer::Clone(MakeMemoryView(Data)));
 						FRequestOwner BlockingOwner(EPriority::Blocking);
 						FDerivedDataBackend::Get().GetRoot().LegacyPut({LegacyRequest}, BlockingOwner, [](auto&&){});
 						BlockingOwner.Wait();
@@ -657,7 +657,7 @@ public:
 			FLegacyCachePutRequest LegacyRequest;
 			LegacyRequest.Name = DebugContext;
 			LegacyRequest.Key = FLegacyCacheKey(CacheKey, FDerivedDataBackend::Get().GetMaxKeyLength());
-			LegacyRequest.Value = FCompositeBuffer(FSharedBuffer::MakeView(MakeMemoryView(Data)));
+			LegacyRequest.Value = FCompositeBuffer(FSharedBuffer::Clone(MakeMemoryView(Data)));
 			FRequestOwner BlockingOwner(EPriority::Blocking);
 			FDerivedDataBackend::Get().GetRoot().LegacyPut({LegacyRequest}, BlockingOwner, [](auto&&){});
 			BlockingOwner.Wait();
