DDC: Fixed cache record propagation to respect the cache policy

The hierarchical backend now considers the Store and Query flags, whether backends are writable, and whether backends allow lower cache levels to be filled.

Features of cache propagation that are missing at this point:
- Propagation within GetChunks.
- Consideration of the local/remote classification of the backend relative to the cache policy.
- Removal from the response of data that was requested for propagation purposes but was meant to be skipped.
- Use of the PartialOnError flag and custom policies to allow higher level backends to return parts of records while attempting to recover only the missing pieces from lower level caches.
- Disabling propagation of Get and Put after reaching backend that contains the record or value, and disables filling of lower cache levels.

#rnx

diff --git a/Engine/Source/Developer/DerivedDataCache/Private/HierarchicalDerivedDataBackend.h b/Engine/Source/Developer/DerivedDataCache/Private/HierarchicalDerivedDataBackend.h
--- a/Engine/Source/Developer/DerivedDataCache/Private/HierarchicalDerivedDataBackend.h
+++ b/Engine/Source/Developer/DerivedDataCache/Private/HierarchicalDerivedDataBackend.h
@@ -41,12 +41,6 @@ public:
 	FHierarchicalDerivedDataBackend(const TArray<FDerivedDataBackendInterface*>& InInnerBackends)
 		: InnerBackends(InInnerBackends)
 		, bIsWritable(false)
-		, bHasLocalBackends(false)
-		, bHasRemoteBackends(false)
-		, bHasMultipleLocalBackends(false)
-		, bHasMultipleRemoteBackends(false)
-		, bHasWritableLocalBackends(false)
-		, bHasWritableRemoteBackends(false)
 	{
 		check(InnerBackends.Num() > 1); // if it is just one, then you don't need this wrapper
 		UpdateAsyncInnerBackends();
@@ -68,28 +62,10 @@ private:
 	void UpdateAsyncInnerBackends()
 	{
 		bIsWritable = false;
-		bHasLocalBackends = false;
-		bHasRemoteBackends = false;
-		bHasMultipleLocalBackends = false;
-		bHasMultipleRemoteBackends = false;
-		bHasWritableLocalBackends = false;
-		bHasWritableRemoteBackends = false;
 		for (int32 CacheIndex = 0; CacheIndex < InnerBackends.Num(); CacheIndex++)
 		{
 			const bool bIsWritableBackend = InnerBackends[CacheIndex]->IsWritable();
 			bIsWritable |= bIsWritableBackend;
-			if (InnerBackends[CacheIndex]->GetSpeedClass() == ESpeedClass::Local)
-			{
-				bHasWritableLocalBackends |= bIsWritableBackend;
-				bHasMultipleLocalBackends = bHasLocalBackends;
-				bHasLocalBackends = true;
-			}
-			else
-			{
-				bHasWritableRemoteBackends |= bIsWritableBackend;
-				bHasMultipleRemoteBackends = bHasRemoteBackends;
-				bHasRemoteBackends = true;
-			}
 		}
 
 		for (int32 CacheIndex = 0; CacheIndex < InnerBackends.Num(); CacheIndex++)
@@ -540,21 +516,75 @@ public:
 	{
 		TArray<FCacheGetRequest, TInlineAllocator<16>> RemainingRequests(Requests);
 
+		struct FUserData
 		{
-			FReadScopeLock LockScope(Lock);
-			TSet<FCacheKey> KeysOk;
+			int32 OriginalIndex = 0;
+			int32 RemainingIndex = 0;
 
-			bool bHadMiss = false;
+			FUserData(int32 InOriginalIndex, int32 InRemainingIndex)
+				: OriginalIndex(InOriginalIndex)
+				, RemainingIndex(InRemainingIndex)
+			{
+			}
 
-			for (int32 GetCacheIndex = 0; GetCacheIndex < InnerBackends.Num() && !RemainingRequests.IsEmpty(); ++GetCacheIndex)
+			explicit FUserData(uint64 UserData)
+				: OriginalIndex(int32(UserData >> 32))
+				, RemainingIndex(int32(UserData))
 			{
-				// Block on this because backends in this hierarchy are not expected to be asynchronous.
-				FRequestOwner BlockingOwner(EPriority::Blocking);
-				InnerBackends[GetCacheIndex]->Get(RemainingRequests, BlockingOwner,
-					[this, GetCacheIndex, &Owner, &OnComplete, &KeysOk](FCacheGetResponse&& Response)
+			}
+
+			uint64 ToUserData() const
+			{
+				return (uint64(OriginalIndex) << 32) | uint64(RemainingIndex);
+			}
+		};
+
+		for (int32 Index = 0, Count = RemainingRequests.Num(); Index < Count; ++Index)
+		{
+			RemainingRequests[Index].UserData = FUserData(Index, Index).ToUserData();
+		}
+
+		FReadScopeLock LockScope(Lock);
+
+		for (int32 GetCacheIndex = 0; GetCacheIndex < InnerBackends.Num() && !RemainingRequests.IsEmpty(); ++GetCacheIndex)
+		{
+			// Block on this because backends in this hierarchy are not expected to be asynchronous.
+			FRequestOwner BlockingOwner(EPriority::Blocking);
+			InnerBackends[GetCacheIndex]->Get(RemainingRequests, BlockingOwner,
+				[this, GetCacheIndex, &RemainingRequests, &Requests, &Owner, &OnComplete](FCacheGetResponse&& Response)
+				{
+					const FUserData UserData(Response.UserData);
+					FCacheRecordPolicy& Policy = RemainingRequests[UserData.RemainingIndex].Policy;
+					if (Response.Status == EStatus::Error && GetCacheIndex + 1 < InnerBackends.Num())
 					{
-						if (Response.Status == EStatus::Ok)
+						if (InnerBackends[GetCacheIndex]->IsWritable())
 						{
+							const auto ConvertPolicy = [](ECachePolicy P) -> ECachePolicy
+							{
+								if (EnumHasAnyFlags(P, ECachePolicy::Store))
+								{
+									EnumRemoveFlags(P, ECachePolicy::SkipData);
+								}
+								return P;
+							};
+							Policy = Policy.Transform(ConvertPolicy);
+						}
+					}
+					else
+					{
+						if (Response.Status == EStatus::Ok &&
+							EnumHasAnyFlags(Policy.GetRecordPolicy(), ECachePolicy::Store) &&
+							InnerBackends[GetCacheIndex]->BackfillLowerCacheLevels())
+						{
+							const FCacheRecordPolicy PutPolicy = Policy.Transform([](ECachePolicy P) -> ECachePolicy
+							{
+								if (!EnumHasAllFlags(P, ECachePolicy::Query))
+								{
+									EnumAddFlags(P, ECachePolicy::PartialOnError);
+								}
+								return P;
+							});
+
 							FRequestOwner AsyncOwner(FPlatformMath::Min(Owner.GetPriority(), EPriority::Highest));
 							FRequestBarrier AsyncBarrier(AsyncOwner);
 							AsyncOwner.KeepAlive();
@@ -564,50 +594,26 @@ public:
 								if ((FillCacheIndex < GetCacheIndex) ||
 									(FillCacheIndex > GetCacheIndex && FillBackend->GetSpeedClass() >= ESpeedClass::Fast))
 								{
-									FillBackend->Put({{Response.Name, Response.Record, ECachePolicy::Default}}, AsyncOwner);
+									if (FillBackend->IsWritable())
+									{
+										FillBackend->Put({{Response.Name, Response.Record, PutPolicy}}, AsyncOwner);
+									}
 								}
 							}
-
-							KeysOk.Add(Response.Record.GetKey());
-							if (OnComplete)
-							{
-								OnComplete(MoveTemp(Response));
-							}
 						}
-					});
-				BlockingOwner.Wait();
 
-				RemainingRequests.RemoveAll([&KeysOk](const FCacheGetRequest& Request) { return KeysOk.Contains(Request.Key); });
-
-				if (!bHadMiss && !RemainingRequests.IsEmpty() && InnerBackends[GetCacheIndex]->IsWritable())
-				{
-					bHadMiss = true;
-					const auto ConvertPolicy = [](ECachePolicy P) -> ECachePolicy
-					{
-						if (EnumHasAnyFlags(P, ECachePolicy::Store))
-						{
-							EnumRemoveFlags(P, ECachePolicy::SkipData);
-						}
-						return P;
-					};
-					for (FCacheGetRequest& Request : RemainingRequests)
-					{
-						FCacheRecordPolicyBuilder Builder(ConvertPolicy(Request.Policy.GetDefaultValuePolicy()));
-						for (const FCacheValuePolicy& ValuePolicy : Request.Policy.GetValuePolicies())
-						{
-							Builder.AddValuePolicy(ValuePolicy.Id, ConvertPolicy(ValuePolicy.Policy));
-						}
-						Request.Policy = Builder.Build();
+						Response.UserData = Requests[UserData.OriginalIndex].UserData;
+						OnComplete(MoveTemp(Response));
+						RemainingRequests[UserData.RemainingIndex].UserData = MAX_uint64;
 					}
-				}
-			}
-		}
+				});
+			BlockingOwner.Wait();
 
-		if (OnComplete)
-		{
-			for (const FCacheGetRequest& Request : RemainingRequests)
+			RemainingRequests.RemoveAll([](const FCacheGetRequest& Request) { return Request.UserData == MAX_uint64; });
+			for (int32 Index = 0, Count = RemainingRequests.Num(); Index < Count; ++Index)
 			{
-				OnComplete({Request.Name, FCacheRecordBuilder(Request.Key).Build(), Request.UserData, EStatus::Error});
+				const FUserData UserData(RemainingRequests[Index].UserData);
+				RemainingRequests[Index].UserData = FUserData(UserData.OriginalIndex, Index).ToUserData();
 			}
 		}
 	}
@@ -802,12 +808,6 @@ private:
 	TArray<TUniquePtr<FDerivedDataBackendInterface> > AsyncPutInnerBackends;
 	/** As an optimization, we check our writable status at contruction **/
 	bool bIsWritable;
-	bool bHasLocalBackends;
-	bool bHasRemoteBackends;
-	bool bHasMultipleLocalBackends;
-	bool bHasMultipleRemoteBackends;
-	bool bHasWritableLocalBackends;
-	bool bHasWritableRemoteBackends;
 };
 
 } // UE::DerivedData::CacheStore::Hierarchical
