DDC/ZenDerivedDataBackend: Implement PutValue,GetValue API.
#rnx

diff --git a/Engine/Source/Developer/DerivedDataCache/Private/Tests/CacheStoreTest.cpp b/Engine/Source/Developer/DerivedDataCache/Private/Tests/CacheStoreTest.cpp
new file mode 100644
--- /dev/null
+++ b/Engine/Source/Developer/DerivedDataCache/Private/Tests/CacheStoreTest.cpp
@@ -0,0 +1,248 @@
+// Copyright Epic Games, Inc. All Rights Reserved.
+
+#include "Misc/AutomationTest.h"
+
+#if WITH_DEV_AUTOMATION_TESTS
+
+#include "Containers/Array.h"
+#include "Containers/Map.h"
+#include "DerivedDataCache.h"
+#include "DerivedDataRequestOwner.h"
+#include "Hash/Blake3.h"
+#include "Misc/StringBuilder.h"
+
+IMPLEMENT_SIMPLE_AUTOMATION_TEST(FCacheStoreTest, "System.DerivedDataCache.CacheStore",
+	EAutomationTestFlags::EditorContext | EAutomationTestFlags::EngineFilter);
+
+// TODO: remove this ifdef and enable the code it surrounds when we submit zenserver.exe
+// with the change to HandleRpcGetCachePayloads to support the ValueAPI
+#define ZENSERVER_SUPPORTS_VALUEAPI 0
+
+bool FCacheStoreTest::RunTest(const FString& Parameters)
+{
+	using namespace UE::DerivedData;
+	ICache& Cache = GetCache();
+
+	FRequestOwner Owner(EPriority::Blocking);
+	FCacheBucket DDCTestBucket(TEXT("DDCTest"));
+
+	// NumRequests/2 must be larger than the batch size in ZenDerivedData,
+	// and NumRequests >= 8 so that we get two results for each possible combination of flags
+	const int32 NumRequests = 24;
+	TArray<FCachePutRequest> PutRequests;
+	TArray<FCachePutValueRequest> PutValueRequests;
+	TArray<FCacheGetRequest> GetRequests;
+	TArray<FCacheGetValueRequest> GetValueRequests;
+	TArray<FCacheChunkRequest> ChunkRequests;
+	struct FTestData
+	{
+		FValue Value;
+		uint8 ByteValue;
+		bool bGetRequestsData = true;
+		bool bUseValueAPI = true;
+		bool bReceivedPut = false;
+		bool bReceivedGet = false;
+		bool bReceivedPutValue = false;
+		bool bReceivedGetValue = false;
+		bool bReceivedChunk = false;
+	};
+
+	TMap<uint64, FTestData> TestDatas;
+	FValueId ValueId = FValueId::FromName(TEXT("ValueName"));
+	for (int32 n = 0; n < NumRequests; ++n)
+	{
+		FBlake3 KeyWriter;
+		KeyWriter.Update(*TestName, TestName.Len() * sizeof(TestName[0]));
+		KeyWriter.Update(&n, sizeof(n));
+		FIoHash KeyHash = KeyWriter.Finalize();
+		FCacheKey Key{ DDCTestBucket, KeyHash };
+		FSharedString Name(WriteToString<16>(TEXT("Request"), n));
+		uint64 UserData = (uint64)n;
+		FTestData& TestData = TestDatas.FindOrAdd(UserData);
+
+		ECachePolicy PutPolicy = ECachePolicy::Default;
+		ECachePolicy GetPolicy = ECachePolicy::Default;
+		TestData.bGetRequestsData = (n & 0x1) == 0;
+		TestData.bUseValueAPI = (n & 0x2) != 0;
+		if (!TestData.bGetRequestsData)
+		{
+			GetPolicy |= ECachePolicy::SkipData;
+		}
+		TestData.ByteValue = (uint8)n;
+		TestData.Value = FValue::Compress(FSharedBuffer::MakeView(&TestData.ByteValue, 1));
+
+		if (!TestData.bUseValueAPI)
+		{
+			FCacheRecordBuilder Builder(Key);
+			Builder.AddValue(ValueId, TestData.Value);
+
+			PutRequests.Add(FCachePutRequest{ Name, Builder.Build(), FCacheRecordPolicy(PutPolicy), UserData });
+			GetRequests.Add(FCacheGetRequest{ Name, Key, FCacheRecordPolicy(GetPolicy), UserData });
+		}
+		else
+		{
+			PutValueRequests.Add(FCachePutValueRequest{ Name, Key, TestData.Value, PutPolicy, UserData });
+			GetValueRequests.Add(FCacheGetValueRequest{ Name, Key, GetPolicy, UserData });
+		}
+		ChunkRequests.Add(FCacheChunkRequest{ Name, Key, TestData.bUseValueAPI ? FValueId() : ValueId,
+			0, TestData.Value.GetRawSize(), FIoHash(), GetPolicy, UserData });
+	}
+
+	{
+		FRequestBarrier Barrier(Owner);
+		Cache.Put(PutRequests, Owner, [&TestDatas, this](FCachePutResponse&& Response)
+			{
+				FTestData* TestData = TestDatas.Find(Response.UserData);
+				if (TestTrue(TEXT("Valid UserData in Put Callback"), TestData != nullptr))
+				{
+					TestData->bReceivedPut = true;
+				}
+			});
+		Cache.PutValue(PutValueRequests, Owner, [&TestDatas, this](FCachePutValueResponse&& Response)
+			{
+				FTestData* TestData = TestDatas.Find(Response.UserData);
+				if (TestTrue(TEXT("Valid UserData in PutValue Callback"), TestData != nullptr))
+				{
+					TestData->bReceivedPutValue = true;
+				}
+			});
+	}
+	Owner.Wait();
+	for (TPair<uint64, FTestData>& Pair : TestDatas)
+	{
+		int32 n = (int32)Pair.Key;
+		if (!Pair.Value.bUseValueAPI)
+		{
+			TestTrue(*WriteToString<16>(TEXT("Put %d received"), n), Pair.Value.bReceivedPut);
+		}
+		else
+		{
+			TestTrue(*WriteToString<16>(TEXT("PutValue %d received"), n), Pair.Value.bReceivedPutValue);
+		}
+
+	}
+
+	{
+		FRequestBarrier Barrier(Owner);
+		Cache.Get(GetRequests, Owner, [&TestDatas, &ValueId, this](FCacheGetResponse&& Response)
+			{
+				FTestData* TestData = TestDatas.Find(Response.UserData);
+				int32 n = (int32)Response.UserData;
+				if (TestTrue(TEXT("Valid UserData in Get Callback"), TestData != nullptr))
+				{
+					TestData->bReceivedGet = true;
+					if (TestTrue(*WriteToString<32>(TEXT("Get %d succeeded"), n), Response.Status == EStatus::Ok))
+					{
+						FCacheRecord& Record = Response.Record;
+						TConstArrayView<FValueWithId> Values = Record.GetValues();
+						if (TestEqual(*WriteToString<32>(TEXT("Get %d ValuesLen"), n), Values.Num(), 1))
+						{
+							const FValueWithId& ActualValue = Values[0];
+							TestEqual(*WriteToString<32>(TEXT("Get %d ValueId"), n), ActualValue.GetId(), ValueId);
+							FValue& ExpectedValue = TestData->Value;
+							TestEqual(*WriteToString<32>(TEXT("Get %d Hash"), n),
+								ActualValue.GetRawHash(), ExpectedValue.GetRawHash());
+							TestEqual(*WriteToString<32>(TEXT("Get %d Size"), n),
+								ActualValue.GetRawSize(), ExpectedValue.GetRawSize());
+
+							if (TestData->bGetRequestsData)
+							{
+								const FCompressedBuffer& Compressed = ActualValue.GetData();
+								FSharedBuffer Buffer = Compressed.Decompress();
+								TestEqual(*WriteToString<32>(TEXT("Get %d Data Size"), n),
+									Buffer.GetSize(), ActualValue.GetRawSize());
+								if (Buffer.GetSize())
+								{
+									TestEqual(*WriteToString<32>(TEXT("Get %d Data Equals"), n),
+										((const uint8*)Buffer.GetData())[0], TestData->ByteValue);
+								}
+							}
+						}
+					}
+				}
+			});
+
+		Cache.GetValue(GetValueRequests, Owner, [&TestDatas, this](FCacheGetValueResponse&& Response)
+			{
+				FTestData* TestData = TestDatas.Find(Response.UserData);
+				int32 n = (int32)Response.UserData;
+				if (TestTrue(TEXT("Valid UserData in GetValue Callback"), TestData != nullptr))
+				{
+					TestData->bReceivedGetValue = true;
+					if (TestTrue(*WriteToString<32>(TEXT("GetValue %d succeeded"), n), Response.Status == EStatus::Ok))
+					{
+						FValue& ActualValue = Response.Value;
+						FValue& ExpectedValue = TestData->Value;
+						TestEqual(*WriteToString<32>(TEXT("GetValue %d Hash"), n),
+							ActualValue.GetRawHash(), ExpectedValue.GetRawHash());
+						TestEqual(*WriteToString<32>(TEXT("GetValue %d Size"), n),
+							ActualValue.GetRawSize(), ExpectedValue.GetRawSize());
+
+						if (TestData->bGetRequestsData)
+						{
+							const FCompressedBuffer& Compressed = ActualValue.GetData();
+							FSharedBuffer Buffer = Compressed.Decompress();
+							TestEqual(*WriteToString<32>(TEXT("GetValue %d Data Size"), n),
+								Buffer.GetSize(), ActualValue.GetRawSize());
+							if (Buffer.GetSize())
+							{
+								TestEqual(*WriteToString<32>(TEXT("GetValue %d Data Equals"), n),
+									((const uint8*)Buffer.GetData())[0], TestData->ByteValue);
+							}
+						}
+					}
+				}
+			});
+
+#if ZENSERVER_SUPPORTS_VALUEAPI
+		Cache.GetChunks(ChunkRequests, Owner, [&TestDatas, this](FCacheChunkResponse&& Response)
+			{
+				FTestData* TestData = TestDatas.Find(Response.UserData);
+				int32 n = (int32)Response.UserData;
+				if (TestTrue(TEXT("Valid UserData in GetChunks Callback"), TestData != nullptr))
+				{
+					TestData->bReceivedChunk = true;
+					if (TestTrue(*WriteToString<32>(TEXT("GetChunks %d succeeded"), n), Response.Status == EStatus::Ok))
+					{
+						FValue& ExpectedValue = TestData->Value;
+						TestEqual(*WriteToString<32>(TEXT("GetChunks %d Hash"), n),
+							Response.RawHash, ExpectedValue.GetRawHash());
+						TestEqual(*WriteToString<32>(TEXT("GetChunks %d Size"), n),
+							Response.RawSize, ExpectedValue.GetRawSize());
+						if (TestData->bGetRequestsData)
+						{
+							FSharedBuffer Buffer = Response.RawData;
+							TestEqual(*WriteToString<32>(TEXT("GetChunks %d Data Size"), n),
+								Buffer.GetSize(), Response.RawSize);
+							if (Buffer.GetSize())
+							{
+								TestEqual(*WriteToString<32>(TEXT("GetChunks %d Data Equals"), n),
+									((const uint8*)Buffer.GetData())[0], TestData->ByteValue);
+							}
+						}
+					}
+				}
+			});
+#endif
+	}
+	Owner.Wait();
+	for (TPair<uint64, FTestData>& Pair : TestDatas)
+	{
+		int32 n = (int32)Pair.Key;
+		if (!Pair.Value.bUseValueAPI)
+		{
+			TestTrue(*WriteToString<16>(TEXT("Get %d received"), n), Pair.Value.bReceivedGet);
+		}
+		else
+		{
+			TestTrue(*WriteToString<16>(TEXT("GetValue %d received"), n), TestDatas[n].bReceivedGetValue);
+		}
+#if ZENSERVER_SUPPORTS_VALUEAPI
+		TestTrue(*WriteToString<16>(TEXT("GetChunk %d received"), n), TestDatas[n].bReceivedChunk);
+#endif
+	}
+
+	return true;
+}
+
+#endif
\ No newline at end of file
diff --git a/Engine/Source/Developer/DerivedDataCache/Private/ZenDerivedDataBackend.cpp b/Engine/Source/Developer/DerivedDataCache/Private/ZenDerivedDataBackend.cpp
--- a/Engine/Source/Developer/DerivedDataCache/Private/ZenDerivedDataBackend.cpp
+++ b/Engine/Source/Developer/DerivedDataCache/Private/ZenDerivedDataBackend.cpp
@@ -123,6 +123,16 @@ public:
 		IRequestOwner& Owner,
 		FOnCacheGetComplete&& OnComplete) override;
 
+	virtual void PutValue(
+		TConstArrayView<FCachePutValueRequest> Requests,
+		IRequestOwner& Owner,
+		FOnCachePutValueComplete&& OnComplete = FOnCachePutValueComplete()) override;
+
+	virtual void GetValue(
+		TConstArrayView<FCacheGetValueRequest> Requests,
+		IRequestOwner& Owner,
+		FOnCacheGetValueComplete&& OnComplete = FOnCacheGetValueComplete()) override;
+
 	virtual void GetChunks(
 		TConstArrayView<FCacheChunkRequest> Requests,
 		IRequestOwner& Owner,
@@ -142,6 +152,7 @@ private:
 	EGetResult GetZenData(const FCacheKey& Key, ECachePolicy CachePolicy, FCbPackage& OutPackage) const;
 
 	bool PutCacheRecord(const FCacheRecord& Record, const FCacheRecordPolicy& Policy);
+	bool PutCacheValue(const FCacheKey& Key, const FValue& Value, const ECachePolicy& ValuePolicy);
 
 	bool IsServiceReady();
 	static FString MakeLegacyZenKey(const TCHAR* CacheKey);
@@ -714,8 +725,10 @@ void FZenDerivedDataBackend::Put(
 				*GetName(), *WriteToString<96>(Record.GetKey()), *Request.Name);
 			bResult = false;
 		}
-
-		bResult = PutCacheRecord(Record, Request.Policy);
+		else
+		{
+			bResult = PutCacheRecord(Record, Request.Policy);
+		}
 
 		if (bResult)
 		{
@@ -885,6 +898,190 @@ void FZenDerivedDataBackend::Get(
 	TRACE_COUNTER_SUBTRACT(ZenDDC_CacheRecordRequestCount, int64(Requests.Num()));
 }
 
+void FZenDerivedDataBackend::PutValue(
+	TConstArrayView<FCachePutValueRequest> Requests,
+	IRequestOwner& Owner,
+	FOnCachePutValueComplete&& OnComplete)
+{
+	TRACE_CPUPROFILER_EVENT_SCOPE(ZenDDC::PutValue);
+
+	for (const FCachePutValueRequest& Request : Requests)
+	{
+		COOK_STAT(auto Timer = UsageStats.TimePut());
+		bool bResult;
+		if (ShouldSimulateMiss(Request.Key))
+		{
+			UE_LOG(LogDerivedDataCache, Verbose, TEXT("%s: Simulated miss for PutValue of %s from '%s'"),
+				*GetName(), *WriteToString<96>(Request.Key), *Request.Name);
+			bResult = false;
+		}
+		else
+		{
+			bResult = PutCacheValue(Request.Key, Request.Value, Request.Policy);
+		}
+
+		if (bResult)
+		{
+			UE_LOG(LogDerivedDataCache, Verbose, TEXT("%s: Cache PutValue complete for %s from '%s'"),
+				*GetName(), *WriteToString<96>(Request.Key), *Request.Name);
+			COOK_STAT(Timer.AddHit(Request.Value.GetData().GetCompressedSize()));
+			if (OnComplete)
+			{
+				OnComplete({ Request.Name, Request.Key, Request.UserData, EStatus::Ok });
+			}
+		}
+		else
+		{
+			COOK_STAT(Timer.AddMiss());
+			if (OnComplete)
+			{
+				OnComplete({ Request.Name, Request.Key, Request.UserData, EStatus::Error });
+			}
+		}
+	}
+}
+
+void FZenDerivedDataBackend::GetValue(
+	TConstArrayView<FCacheGetValueRequest> Requests,
+	IRequestOwner& Owner,
+	FOnCacheGetValueComplete&& OnComplete)
+{
+	TRACE_CPUPROFILER_EVENT_SCOPE(ZenDDC::GetValue);
+	TRACE_COUNTER_ADD(ZenDDC_CacheRecordRequestCount, int64(Requests.Num()));
+
+	int32 TotalCompleted = 0;
+
+	ForEachBatch(CacheRecordBatchSize, Requests.Num(),
+		[this, &Requests, &Owner, &OnComplete, &TotalCompleted](int32 BatchFirst, int32 BatchLast)
+		{
+			TRACE_COUNTER_ADD(ZenDDC_Get, int64(BatchLast) - int64(BatchFirst) + 1);
+			COOK_STAT(auto Timer = UsageStats.TimeGet());
+
+			FCbWriter BatchRequest;
+			BatchRequest.BeginObject();
+			{
+				BatchRequest << "Method"_ASV << "GetCacheRecords";
+				BatchRequest.BeginObject("Params"_ASV);
+				{
+					BatchRequest.BeginArray("CacheKeys"_ASV);
+					for (int32 KeyIndex = BatchFirst; KeyIndex <= BatchLast; KeyIndex++)
+					{
+						const FCacheKey& Key = Requests[KeyIndex].Key;
+
+						BatchRequest.BeginObject();
+						BatchRequest << "Bucket"_ASV << Key.Bucket.ToString();
+						BatchRequest << "Hash"_ASV << Key.Hash;
+						BatchRequest.EndObject();
+					}
+					BatchRequest.EndArray();
+
+					// TODO: The policy needs to be sent with each key.
+					const ECachePolicy& Policy = Requests[BatchFirst].Policy;
+					BatchRequest.BeginObject("Policy"_ASV);
+					{
+						BatchRequest << "RecordPolicy"_ASV << static_cast<uint32>(Policy);
+					}
+					BatchRequest.EndObject();
+				}
+				BatchRequest.EndObject();
+			}
+			BatchRequest.EndObject();
+
+			FCbPackage BatchResponse;
+			Zen::FZenHttpRequest::Result HttpResult = Zen::FZenHttpRequest::Result::Failed;
+
+			{
+				Zen::FZenScopedRequestPtr Request(RequestPool.Get());
+				HttpResult = Request->PerformRpc(TEXT("/z$/$rpc"_SV), BatchRequest.Save().AsObject(), BatchResponse);
+			}
+
+			if (HttpResult == Zen::FZenHttpRequest::Result::Success)
+			{
+				const FCbObject& ResponseObj = BatchResponse.GetObject();
+
+				int32 KeyIndex = BatchFirst;
+				for (FCbField ValueReferenceField : ResponseObj["Result"_ASV])
+				{
+					const FCacheGetValueRequest& Request = Requests[KeyIndex++];
+					const FCacheKey& Key = Request.Key;
+					FIoHash RawHash = ValueReferenceField["RawHash"_ASV].AsBinaryAttachment();
+					TOptional<FValue> ResultValue;
+
+					if (!RawHash.IsZero())
+					{
+						if (ShouldSimulateMiss(Key))
+						{
+							UE_LOG(LogDerivedDataCache, Verbose, TEXT("%s: Simulated miss for put of '%s' from '%s'"),
+								*GetName(), *WriteToString<96>(Key), *Request.Name);
+						}
+						else if (EnumHasAnyFlags(Request.Policy, UE::DerivedData::ECachePolicy::SkipData))
+						{
+							ResultValue.Emplace(RawHash, ValueReferenceField["RawSize"].AsUInt64());
+						}
+						else
+						{
+							const FCbAttachment* Attachment = BatchResponse.FindAttachment(RawHash);
+							if (Attachment)
+							{
+								ResultValue.Emplace(Attachment->AsCompressedBinary());
+							}
+						}
+					}
+
+					if (ResultValue)
+					{
+						UE_LOG(LogDerivedDataCache, Verbose, TEXT("%s: Cache hit for '%s' from '%s'"),
+							*GetName(), *WriteToString<96>(Key), *Request.Name);
+
+						TRACE_COUNTER_ADD(ZenDDC_GetHit, int64(1));
+						int64 ReceivedSize = ResultValue->GetData().GetCompressedSize();
+						TRACE_COUNTER_ADD(ZenDDC_BytesReceived, ReceivedSize);
+						COOK_STAT(Timer.AddHit(ReceivedSize));
+
+						if (OnComplete)
+						{
+							OnComplete({ Request.Name, Key, MoveTemp(*ResultValue), Request.UserData, EStatus::Ok });
+						}
+					}
+					else
+					{
+						UE_LOG(LogDerivedDataCache, Verbose, TEXT("%s: Cache miss for '%s' from '%s'"),
+							*GetName(), *WriteToString<96>(Key), *Request.Name);
+
+						if (OnComplete)
+						{
+							OnComplete({ Request.Name, Key, {}, Request.UserData, EStatus::Error });
+						}
+					}
+
+					TotalCompleted++;
+				}
+			}
+			else
+			{
+				for (int32 KeyIndex = BatchFirst; KeyIndex <= BatchLast; KeyIndex++)
+				{
+					const FCacheGetValueRequest& Request = Requests[KeyIndex];
+					const FCacheKey& Key = Request.Key;
+
+					UE_LOG(LogDerivedDataCache, Verbose, TEXT("%s: Cache miss for '%s' from '%s'"),
+						*GetName(), *WriteToString<96>(Key), *Request.Name);
+
+					if (OnComplete)
+					{
+						OnComplete({ Request.Name, Key, {}, Request.UserData, EStatus::Error });
+					}
+
+					TotalCompleted++;
+				}
+			}
+		});
+
+	UE_CLOG(TotalCompleted != Requests.Num(), LogDerivedDataCache, Warning, TEXT("Only '%d/%d' GetValue request(s) completed"), TotalCompleted, Requests.Num());
+	TRACE_COUNTER_SUBTRACT(ZenDDC_CacheRecordRequestCount, int64(Requests.Num()));
+}
+
+
 void FZenDerivedDataBackend::GetChunks(
 	TConstArrayView<FCacheChunkRequest> Requests,
 	IRequestOwner& Owner,
@@ -1036,6 +1233,35 @@ bool FZenDerivedDataBackend::PutCacheRecord(const FCacheRecord& Record, const FC
 	return true;
 }
 
+bool FZenDerivedDataBackend::PutCacheValue(const FCacheKey& Key, const FValue& Value, const ECachePolicy& ValuePolicy)
+{
+	FCbWriter Writer;
+	Writer.BeginObject();
+	Writer.AddBinaryAttachment("RawHash", Value.GetRawHash());
+	Writer.AddInteger("RawSize", Value.GetRawSize());
+	Writer.EndObject();
+
+	FCbPackage Package(Writer.Save().AsObject());
+	if (Value.HasData())
+	{
+		Package.AddAttachment(FCbAttachment(Value.GetData()));
+	}
+
+	FBufferArchive Ar;
+	Package.Save(Ar);
+	FCompositeBuffer Buffer = FCompositeBuffer(FSharedBuffer::MakeView(Ar.GetData(), Ar.Num()));
+	TStringBuilder<256> Uri;
+	AppendZenUri(Key, Uri);
+	AppendPolicyQueryString(ValuePolicy, Uri);
+	if (PutZenData(Uri.ToString(), Buffer, Zen::EContentType::CbPackage)
+		!= FDerivedDataBackendInterface::EPutStatus::Cached)
+	{
+		return false;
+	}
+
+	return true;
+}
+
 FDerivedDataBackendInterface* CreateZenDerivedDataBackend(const TCHAR* NodeName, const TCHAR* ServiceUrl, const TCHAR* Namespace)
 {
 	FZenDerivedDataBackend* Backend = new FZenDerivedDataBackend(ServiceUrl, Namespace);
