DDC: Renamed ECachePolicy::PartialOnError to PartialRecord

The name more accurately reflects the purpose of the flag, which is to allow a partial record to be fetched or stored in the absence of data for otherwise-required values.

#rnx

diff --git a/Engine/Source/Developer/DerivedDataCache/Private/DerivedDataBuildJobContext.cpp b/Engine/Source/Developer/DerivedDataCache/Private/DerivedDataBuildJobContext.cpp
--- a/Engine/Source/Developer/DerivedDataCache/Private/DerivedDataBuildJobContext.cpp
+++ b/Engine/Source/Developer/DerivedDataCache/Private/DerivedDataBuildJobContext.cpp
@@ -179,8 +179,8 @@ void FBuildJobContext::SetCachePolicyMask(ECachePolicy Policy)
 		TEXT("KeepAlive flag may not be masked out on the cache policy for build of '%s' by %s. ")
 		TEXT("Flags for cache record lifetime may be set indirectly through EBuildPolicy."),
 		*WriteToString<128>(Job.GetName()), *WriteToString<32>(Job.GetFunction()));
-	checkf(EnumHasAllFlags(Policy, ECachePolicy::PartialOnError),
-		TEXT("PartialOnError flag may not be masked out on the cache policy for build of '%s' by %s."),
+	checkf(EnumHasAllFlags(Policy, ECachePolicy::PartialRecord),
+		TEXT("PartialRecord flag may not be masked out on the cache policy for build of '%s' by %s."),
 		*WriteToString<128>(Job.GetName()), *WriteToString<32>(Job.GetFunction()));
 	CachePolicyMask = Policy;
 }
diff --git a/Engine/Source/Developer/DerivedDataCache/Private/FileSystemDerivedDataBackend.cpp b/Engine/Source/Developer/DerivedDataCache/Private/FileSystemDerivedDataBackend.cpp
--- a/Engine/Source/Developer/DerivedDataCache/Private/FileSystemDerivedDataBackend.cpp
+++ b/Engine/Source/Developer/DerivedDataCache/Private/FileSystemDerivedDataBackend.cpp
@@ -1990,7 +1990,7 @@ FOptionalCacheRecord FFileSystemCacheStore::GetCacheRecord(
 		{
 			RecordBuilder.AddValue(Id, MoveTemp(Content));
 		}
-		else if (EnumHasAnyFlags(ValuePolicy, ECachePolicy::PartialOnError))
+		else if (EnumHasAnyFlags(ValuePolicy, ECachePolicy::PartialRecord))
 		{
 			OutStatus = EStatus::Error;
 			RecordBuilder.AddValue(Value);
diff --git a/Engine/Source/Developer/DerivedDataCache/Private/HierarchicalDerivedDataBackend.h b/Engine/Source/Developer/DerivedDataCache/Private/HierarchicalDerivedDataBackend.h
--- a/Engine/Source/Developer/DerivedDataCache/Private/HierarchicalDerivedDataBackend.h
+++ b/Engine/Source/Developer/DerivedDataCache/Private/HierarchicalDerivedDataBackend.h
@@ -580,7 +580,7 @@ public:
 							{
 								if (!EnumHasAllFlags(P, ECachePolicy::Query))
 								{
-									EnumAddFlags(P, ECachePolicy::PartialOnError);
+									EnumAddFlags(P, ECachePolicy::PartialRecord);
 								}
 								return P;
 							});
diff --git a/Engine/Source/Developer/DerivedDataCache/Private/MemoryDerivedDataBackend.cpp b/Engine/Source/Developer/DerivedDataCache/Private/MemoryDerivedDataBackend.cpp
--- a/Engine/Source/Developer/DerivedDataCache/Private/MemoryDerivedDataBackend.cpp
+++ b/Engine/Source/Developer/DerivedDataCache/Private/MemoryDerivedDataBackend.cpp
@@ -741,7 +741,7 @@ void FMemoryDerivedDataBackend::Get(
 				if (!Value.HasData() && !EnumHasAnyFlags(Policy.GetValuePolicy(Value.GetId()), ECachePolicy::SkipData))
 				{
 					Status = EStatus::Error;
-					if (!EnumHasAllFlags(Policy.GetValuePolicy(Value.GetId()), ECachePolicy::PartialOnError))
+					if (!EnumHasAllFlags(Policy.GetValuePolicy(Value.GetId()), ECachePolicy::PartialRecord))
 					{
 						Record.Reset();
 						break;
diff --git a/Engine/Source/Developer/DerivedDataCache/Private/PakFileDerivedDataBackend.cpp b/Engine/Source/Developer/DerivedDataCache/Private/PakFileDerivedDataBackend.cpp
--- a/Engine/Source/Developer/DerivedDataCache/Private/PakFileDerivedDataBackend.cpp
+++ b/Engine/Source/Developer/DerivedDataCache/Private/PakFileDerivedDataBackend.cpp
@@ -1031,7 +1031,7 @@ void FPakFileDerivedDataBackend::GetCacheContent(
 				*CachePath, *WriteToString<16>(InOutValue.GetId()), *WriteToString<48>(RawHash),
 				*WriteToString<96>(Key), Name.Len(), Name.GetData());
 			InOutStatus = EStatus::Error;
-			if (!EnumHasAnyFlags(Policy, ECachePolicy::PartialOnError))
+			if (!EnumHasAnyFlags(Policy, ECachePolicy::PartialRecord))
 			{
 				InOutValue = FValueWithId::Null;
 			}
@@ -1044,7 +1044,7 @@ void FPakFileDerivedDataBackend::GetCacheContent(
 		*CachePath, *WriteToString<16>(InOutValue.GetId()), *WriteToString<48>(RawHash), *WriteToString<96>(Key),
 		Name.Len(), Name.GetData());
 	InOutStatus = EStatus::Error;
-	if (!EnumHasAnyFlags(Policy, ECachePolicy::PartialOnError))
+	if (!EnumHasAnyFlags(Policy, ECachePolicy::PartialRecord))
 	{
 		InOutValue = FValueWithId::Null;
 	}
diff --git a/Engine/Source/Developer/DerivedDataCache/Public/DerivedDataCache.h b/Engine/Source/Developer/DerivedDataCache/Public/DerivedDataCache.h
--- a/Engine/Source/Developer/DerivedDataCache/Public/DerivedDataCache.h
+++ b/Engine/Source/Developer/DerivedDataCache/Public/DerivedDataCache.h
@@ -97,7 +97,7 @@ enum class ECachePolicy : uint32
 	 *
 	 * Applying this flag for a put of a record allows a partial record to be stored.
 	 */
-	PartialOnError  = 1 << 6,
+	PartialRecord   = 1 << 6,
 
 	/**
 	 * Keep records in the cache for at least the duration of the session.
@@ -240,7 +240,7 @@ public:
 	 *
 	 * A cache store is free to interpret a record containing only a key as a request to delete that
 	 * record from the store. Records may contain values that do not have data, and such values must
-	 * reference an existing value in the store, if available. The PartialOnError policy may be used
+	 * reference an existing value in the store, if available. The PartialRecord policy may be used
 	 * to request that a partial record be stored when a value is missing.
 	 *
 	 * @param Requests     Requests with the cache records to store. Records must have a key.
@@ -260,10 +260,10 @@ public:
 	 *
 	 * Records may propagate into other cache stores, in accordance with the policy. A propagated
 	 * record may be a partial record, containing values without data, depending on the policy. A
-	 * propagated put of a partial record will include the PartialOnError policy.
+	 * propagated put of a partial record will include the PartialRecord policy.
 	 *
 	 * When values are required by the policy but not available, the status must be Error. If the
-	 * PartialOnError policy applies to the missing values, the cache store must return a partial
+	 * PartialRecord policy applies to the missing values, the cache store must return a partial
 	 * record containing the available values.
 	 *
 	 * @param Requests     Requests with the keys identifying the cache records to fetch.
