DDC: Moved GDDCIOThreadPool out of Core

The DDC thread pool is a temporary workaround for issues with high-latency cache accesses happening on regular worker thread, which are reduced in number in 5.0 compared to 4.27. Moving it to be private within DDC will keep it easy to remove once the high-latency caches are accessed through an async API.

#rnx

diff --git a/Engine/Source/Developer/DerivedDataCache/Private/DerivedDataBackendAsyncPutWrapper.cpp b/Engine/Source/Developer/DerivedDataCache/Private/DerivedDataBackendAsyncPutWrapper.cpp
--- a/Engine/Source/Developer/DerivedDataCache/Private/DerivedDataBackendAsyncPutWrapper.cpp
+++ b/Engine/Source/Developer/DerivedDataCache/Private/DerivedDataBackendAsyncPutWrapper.cpp
@@ -2,6 +2,7 @@
 
 #include "Async/AsyncWork.h"
 #include "DerivedDataBackendInterface.h"
+#include "DerivedDataCachePrivate.h"
 #include "DerivedDataCacheUsageStats.h"
 #include "DerivedDataRequest.h"
 #include "DerivedDataRequestOwner.h"
@@ -521,7 +522,7 @@ FDerivedDataBackendInterface::EPutStatus FDerivedDataBackendAsyncPutWrapper::Put
 	UE_LOG(LogDerivedDataCache, Verbose, TEXT("%s queueing %s for put"), *GetName(), CacheKey);
 
 	FDerivedDataBackend::Get().AddToAsyncCompletionCounter(1);
-	(new FAutoDeleteAsyncTask<FCachePutAsyncWorker>(CacheKey, InData, InnerBackend, bPutEvenIfExists, InflightCache.Get(), &FilesInFlight, UsageStats))->StartBackgroundTask(GDDCIOThreadPool, EQueuedWorkPriority::Low);
+	(new FAutoDeleteAsyncTask<FCachePutAsyncWorker>(CacheKey, InData, InnerBackend, bPutEvenIfExists, InflightCache.Get(), &FilesInFlight, UsageStats))->StartBackgroundTask(Private::GCacheThreadPool, EQueuedWorkPriority::Low);
 
 	return EPutStatus::Executing;
 }
@@ -576,7 +577,7 @@ public:
 		Owner.Begin(this);
 
 		DoneEvent.Reset();
-		GDDCIOThreadPool->AddQueuedWork(this, GetPriority(Priority));
+		Private::GCacheThreadPool->AddQueuedWork(this, GetPriority(Priority));
 	}
 
 	inline void Execute(bool bCancel)
@@ -595,9 +596,9 @@ public:
 
 	inline void SetPriority(EPriority Priority) final
 	{
-		if (GDDCIOThreadPool->RetractQueuedWork(this))
+		if (Private::GCacheThreadPool->RetractQueuedWork(this))
 		{
-			GDDCIOThreadPool->AddQueuedWork(this, GetPriority(Priority));
+			Private::GCacheThreadPool->AddQueuedWork(this, GetPriority(Priority));
 		}
 	}
 
@@ -605,7 +606,7 @@ public:
 	{
 		if (!DoneEvent.Wait(0))
 		{
-			if (GDDCIOThreadPool->RetractQueuedWork(this))
+			if (Private::GCacheThreadPool->RetractQueuedWork(this))
 			{
 				Abandon();
 			}
@@ -621,7 +622,7 @@ public:
 	{
 		if (!DoneEvent.Wait(0))
 		{
-			if (GDDCIOThreadPool->RetractQueuedWork(this))
+			if (Private::GCacheThreadPool->RetractQueuedWork(this))
 			{
 				DoThreadedWork();
 			}
@@ -738,7 +739,7 @@ void FDerivedDataBackendAsyncPutWrapper::Execute(
 	#endif
 	};
 
-	if (Owner.GetPriority() == EPriority::Blocking || !GDDCIOThreadPool)
+	if (Owner.GetPriority() == EPriority::Blocking || !Private::GCacheThreadPool)
 	{
 		return ExecuteWithStats(Requests, Owner, MoveTemp(OnComplete));
 	}
diff --git a/Engine/Source/Developer/DerivedDataCache/Private/DerivedDataCache.cpp b/Engine/Source/Developer/DerivedDataCache/Private/DerivedDataCache.cpp
--- a/Engine/Source/Developer/DerivedDataCache/Private/DerivedDataCache.cpp
+++ b/Engine/Source/Developer/DerivedDataCache/Private/DerivedDataCache.cpp
@@ -285,6 +285,8 @@ void ICacheStore::GetValue(
 namespace UE::DerivedData::Private
 {
 
+FQueuedThreadPool* GCacheThreadPool;
+
 /**
  * Implementation of the derived data cache
  * This API is fully threadsafe
@@ -496,6 +498,13 @@ public:
 	FDerivedDataCache()
 		: CurrentHandle(19248) // we will skip some potential handles to catch errors
 	{
+		if (FPlatformProcess::SupportsMultithreading())
+		{
+			GCacheThreadPool = FQueuedThreadPool::Allocate();
+			const int32 ThreadCount = FPlatformMisc::NumberOfIOWorkerThreadsToSpawn();
+			verify(GCacheThreadPool->Create(ThreadCount, 96 * 1024, TPri_AboveNormal, TEXT("DDC IO ThreadPool")));
+		}
+
 		FDerivedDataBackend::Get(); // we need to make sure this starts before we allow us to start
 
 		CacheStoreMaintainers = IModularFeatures::Get().GetModularFeatureImplementations<ICacheStoreMaintainer>(FeatureName);
@@ -642,7 +651,7 @@ public:
 		PendingTasks.Add(Handle, AsyncTask);
 		AddToAsyncCompletionCounter(1);
 		// This request is I/O only, doesn't do any processing, send it to the I/O only thread-pool to avoid wasting worker threads on long I/O waits.
-		AsyncTask->StartBackgroundTask(GDDCIOThreadPool);
+		AsyncTask->StartBackgroundTask(GCacheThreadPool);
 		return Handle;
 	}
 
diff --git a/Engine/Source/Developer/DerivedDataCache/Private/DerivedDataCachePrivate.h b/Engine/Source/Developer/DerivedDataCache/Private/DerivedDataCachePrivate.h
--- a/Engine/Source/Developer/DerivedDataCache/Private/DerivedDataCachePrivate.h
+++ b/Engine/Source/Developer/DerivedDataCache/Private/DerivedDataCachePrivate.h
@@ -7,6 +7,7 @@
 #include "Misc/AsciiSet.h"
 
 class FDerivedDataCacheInterface;
+class FQueuedThreadPool;
 
 namespace UE::DerivedData { class FCacheRecord; }
 namespace UE::DerivedData { class ICache; }
@@ -15,6 +16,9 @@ namespace UE::DerivedData { struct FCacheKey; }
 namespace UE::DerivedData::Private
 {
 
+// Created by CreateCache
+extern FQueuedThreadPool* GCacheThreadPool;
+
 // Implemented in DerivedDataCache.cpp
 ICache* CreateCache(FDerivedDataCacheInterface** OutLegacyCache);
 
diff --git a/Engine/Source/Runtime/Core/Private/HAL/ThreadingBase.cpp b/Engine/Source/Runtime/Core/Private/HAL/ThreadingBase.cpp
--- a/Engine/Source/Runtime/Core/Private/HAL/ThreadingBase.cpp
+++ b/Engine/Source/Runtime/Core/Private/HAL/ThreadingBase.cpp
@@ -43,8 +43,6 @@ FQueuedThreadPool* GThreadPool = nullptr;
 
 FQueuedThreadPool* GIOThreadPool = nullptr;
 
-FQueuedThreadPool* GDDCIOThreadPool = nullptr;
-
 FQueuedThreadPool* GBackgroundPriorityThreadPool = nullptr;
 
 #if WITH_EDITOR
diff --git a/Engine/Source/Runtime/Core/Public/Misc/QueuedThreadPool.h b/Engine/Source/Runtime/Core/Public/Misc/QueuedThreadPool.h
--- a/Engine/Source/Runtime/Core/Public/Misc/QueuedThreadPool.h
+++ b/Engine/Source/Runtime/Core/Public/Misc/QueuedThreadPool.h
@@ -164,8 +164,6 @@ extern CORE_API FQueuedThreadPool* GThreadPool;
 
 extern CORE_API FQueuedThreadPool* GIOThreadPool;
 
-extern CORE_API FQueuedThreadPool* GDDCIOThreadPool;
-
 extern CORE_API FQueuedThreadPool* GBackgroundPriorityThreadPool;
 
 #if WITH_EDITOR
diff --git a/Engine/Source/Runtime/Engine/Private/AudioDerivedData.cpp b/Engine/Source/Runtime/Engine/Private/AudioDerivedData.cpp
--- a/Engine/Source/Runtime/Engine/Private/AudioDerivedData.cpp
+++ b/Engine/Source/Runtime/Engine/Private/AudioDerivedData.cpp
@@ -897,11 +897,11 @@ bool FStreamedAudioPlatformData::AreDerivedChunksAvailable(FStringView InContext
 
 	// If this is called from the game thread, try to prefetch chunks locally on background thread
 	// to avoid doing high latency remote calls every time we reload this data.
-	if (bAllCachedDataProbablyExists && GDDCIOThreadPool && IsInGameThread())
+	if (bAllCachedDataProbablyExists && GIOThreadPool && IsInGameThread())
 	{
 		FString Context{ InContext };
 		AsyncPool(
-			*GDDCIOThreadPool,
+			*GIOThreadPool,
 			[ChunkKeys, Context]()
 			{
 				GetDerivedDataCacheRef().TryToPrefetch(ChunkKeys, Context);
diff --git a/Engine/Source/Runtime/Launch/Private/LaunchEngineLoop.cpp b/Engine/Source/Runtime/Launch/Private/LaunchEngineLoop.cpp
--- a/Engine/Source/Runtime/Launch/Private/LaunchEngineLoop.cpp
+++ b/Engine/Source/Runtime/Launch/Private/LaunchEngineLoop.cpp
@@ -2469,18 +2469,6 @@ PRAGMA_ENABLE_DEPRECATION_WARNINGS
 			}
 			verify(GIOThreadPool->Create(NumThreadsInThreadPool, 96 * 1024, TPri_AboveNormal, TEXT("IOThreadPool")));
 		}
- 
-		if (!FPlatformProperties::RequiresCookedData())
-		{
-			SCOPED_BOOT_TIMING("GDDCIOThreadPool->Create");
-			GDDCIOThreadPool = FQueuedThreadPool::Allocate();
-			int32 NumThreadsInThreadPool = FPlatformMisc::NumberOfIOWorkerThreadsToSpawn();
-			if (FPlatformProperties::IsServerOnly())
-			{
-				NumThreadsInThreadPool = 2;
-			}
-			verify(GDDCIOThreadPool->Create(NumThreadsInThreadPool, 96 * 1024, TPri_AboveNormal, TEXT("DDC IO ThreadPool")));
-		}
 	}
 
 	FEmbeddedCommunication::ForceTick(1);
