DDC: Renamed cache request and response types from Chunk to GetChunk

#rnx

diff --git a/Engine/Source/Developer/DerivedDataCache/Private/DerivedDataBackendAsyncPutWrapper.cpp b/Engine/Source/Developer/DerivedDataCache/Private/DerivedDataBackendAsyncPutWrapper.cpp
--- a/Engine/Source/Developer/DerivedDataCache/Private/DerivedDataBackendAsyncPutWrapper.cpp
+++ b/Engine/Source/Developer/DerivedDataCache/Private/DerivedDataBackendAsyncPutWrapper.cpp
@@ -177,9 +177,9 @@ public:
 	}
 
 	virtual void GetChunks(
-		TConstArrayView<FCacheChunkRequest> Requests,
+		TConstArrayView<FCacheGetChunkRequest> Requests,
 		IRequestOwner& Owner,
-		FOnCacheChunkComplete&& OnComplete) override
+		FOnCacheGetChunkComplete&& OnComplete) override
 	{
 		Execute(COOK_STAT(&FDerivedDataCacheUsageStats::TimeGet,) Requests, Owner, MoveTemp(OnComplete), &ICacheStore::GetChunks);
 	}
@@ -692,7 +692,7 @@ static FCacheGetValueResponse MakeCanceledResponse(const FCacheGetValueRequest&
 	return {Request.Name, Request.Key, {}, Request.UserData, EStatus::Canceled};
 }
 
-static FCacheChunkResponse MakeCanceledResponse(const FCacheChunkRequest& Request)
+static FCacheGetChunkResponse MakeCanceledResponse(const FCacheGetChunkRequest& Request)
 {
 	return {Request.Name, Request.Key, Request.Id, Request.RawOffset, 0, {}, {}, Request.UserData, EStatus::Canceled};
 }
diff --git a/Engine/Source/Developer/DerivedDataCache/Private/DerivedDataBackendThrottleWrapper.h b/Engine/Source/Developer/DerivedDataCache/Private/DerivedDataBackendThrottleWrapper.h
--- a/Engine/Source/Developer/DerivedDataCache/Private/DerivedDataBackendThrottleWrapper.h
+++ b/Engine/Source/Developer/DerivedDataCache/Private/DerivedDataBackendThrottleWrapper.h
@@ -211,12 +211,12 @@ public:
 	}
 
 	virtual void GetChunks(
-		TConstArrayView<FCacheChunkRequest> Requests,
+		TConstArrayView<FCacheGetChunkRequest> Requests,
 		IRequestOwner& Owner,
-		FOnCacheChunkComplete&& OnComplete) override
+		FOnCacheGetChunkComplete&& OnComplete) override
 	{
 		InnerBackend->GetChunks(Requests, Owner,
-			[this, State = EnterThrottlingScope(), OnComplete = MoveTemp(OnComplete)](FCacheChunkResponse&& Response)
+			[this, State = EnterThrottlingScope(), OnComplete = MoveTemp(OnComplete)](FCacheGetChunkResponse&& Response)
 			{
 				CloseThrottlingScope(State, FThrottlingState(this, Response.RawData.GetSize()));
 				if (OnComplete)
diff --git a/Engine/Source/Developer/DerivedDataCache/Private/DerivedDataBackendVerifyWrapper.h b/Engine/Source/Developer/DerivedDataCache/Private/DerivedDataBackendVerifyWrapper.h
--- a/Engine/Source/Developer/DerivedDataCache/Private/DerivedDataBackendVerifyWrapper.h
+++ b/Engine/Source/Developer/DerivedDataCache/Private/DerivedDataBackendVerifyWrapper.h
@@ -175,9 +175,9 @@ public:
 	}
 
 	virtual void GetChunks(
-		TConstArrayView<FCacheChunkRequest> Requests,
+		TConstArrayView<FCacheGetChunkRequest> Requests,
 		IRequestOwner& Owner,
-		FOnCacheChunkComplete&& OnComplete) override
+		FOnCacheGetChunkComplete&& OnComplete) override
 	{
 		return InnerBackend->GetChunks(Requests, Owner, MoveTemp(OnComplete));
 	}
diff --git a/Engine/Source/Developer/DerivedDataCache/Private/DerivedDataCache.cpp b/Engine/Source/Developer/DerivedDataCache/Private/DerivedDataCache.cpp
--- a/Engine/Source/Developer/DerivedDataCache/Private/DerivedDataCache.cpp
+++ b/Engine/Source/Developer/DerivedDataCache/Private/DerivedDataCache.cpp
@@ -943,9 +943,9 @@ public:
 	}
 
 	void GetChunks(
-		TConstArrayView<FCacheChunkRequest> Requests,
+		TConstArrayView<FCacheGetChunkRequest> Requests,
 		IRequestOwner& Owner,
-		FOnCacheChunkComplete&& OnComplete) final
+		FOnCacheGetChunkComplete&& OnComplete) final
 	{
 		return FDerivedDataBackend::Get().GetRoot().GetChunks(Requests, Owner, OnComplete ? MoveTemp(OnComplete) : [](auto&&){});
 	}
diff --git a/Engine/Source/Developer/DerivedDataCache/Private/FileSystemDerivedDataBackend.cpp b/Engine/Source/Developer/DerivedDataCache/Private/FileSystemDerivedDataBackend.cpp
--- a/Engine/Source/Developer/DerivedDataCache/Private/FileSystemDerivedDataBackend.cpp
+++ b/Engine/Source/Developer/DerivedDataCache/Private/FileSystemDerivedDataBackend.cpp
@@ -770,9 +770,9 @@ public:
 		IRequestOwner& Owner,
 		FOnCacheGetValueComplete&& OnComplete) final;
 	void GetChunks(
-		TConstArrayView<FCacheChunkRequest> Requests,
+		TConstArrayView<FCacheGetChunkRequest> Requests,
 		IRequestOwner& Owner,
-		FOnCacheChunkComplete&& OnComplete) final;
+		FOnCacheGetChunkComplete&& OnComplete) final;
 
 private:
 	[[nodiscard]] bool PutCacheRecord(FStringView Name, const FCacheRecord& Record, const FCacheRecordPolicy& Policy, uint64& OutWriteSize);
@@ -1680,11 +1680,11 @@ void FFileSystemCacheStore::GetValue(
 }
 
 void FFileSystemCacheStore::GetChunks(
-	const TConstArrayView<FCacheChunkRequest> Requests,
+	const TConstArrayView<FCacheGetChunkRequest> Requests,
 	IRequestOwner& Owner,
-	FOnCacheChunkComplete&& OnComplete)
+	FOnCacheGetChunkComplete&& OnComplete)
 {
-	TArray<FCacheChunkRequest, TInlineAllocator<16>> SortedRequests(Requests);
+	TArray<FCacheGetChunkRequest, TInlineAllocator<16>> SortedRequests(Requests);
 	SortedRequests.StableSort(TChunkLess());
 
 	bool bHasValue = false;
@@ -1695,7 +1695,7 @@ void FFileSystemCacheStore::GetChunks(
 	FCompressedBufferReader ValueReader;
 	EStatus ValueStatus = EStatus::Error;
 	FOptionalCacheRecord Record;
-	for (const FCacheChunkRequest& Request : SortedRequests)
+	for (const FCacheGetChunkRequest& Request : SortedRequests)
 	{
 		const bool bExistsOnly = EnumHasAnyFlags(Request.Policy, ECachePolicy::SkipData);
 		TRACE_CPUPROFILER_EVENT_SCOPE(FileSystemDDC_Get);
diff --git a/Engine/Source/Developer/DerivedDataCache/Private/HierarchicalDerivedDataBackend.h b/Engine/Source/Developer/DerivedDataCache/Private/HierarchicalDerivedDataBackend.h
--- a/Engine/Source/Developer/DerivedDataCache/Private/HierarchicalDerivedDataBackend.h
+++ b/Engine/Source/Developer/DerivedDataCache/Private/HierarchicalDerivedDataBackend.h
@@ -748,11 +748,11 @@ public:
 	}
 
 	virtual void GetChunks(
-		TConstArrayView<FCacheChunkRequest> Requests,
+		TConstArrayView<FCacheGetChunkRequest> Requests,
 		IRequestOwner& Owner,
-		FOnCacheChunkComplete&& OnComplete) override
+		FOnCacheGetChunkComplete&& OnComplete) override
 	{
-		TArray<FCacheChunkRequest, TInlineAllocator<16>> RemainingRequests(Requests);
+		TArray<FCacheGetChunkRequest, TInlineAllocator<16>> RemainingRequests(Requests);
 
 		{
 			FReadScopeLock LockScope(Lock);
@@ -764,10 +764,10 @@ public:
 					return;
 				}
 				RemainingRequests.StableSort(TChunkLess());
-				TArray<FCacheChunkRequest, TInlineAllocator<16>> ErrorChunks;
+				TArray<FCacheGetChunkRequest, TInlineAllocator<16>> ErrorChunks;
 				FRequestOwner BlockingOwner(EPriority::Blocking);
 				InnerBackend->GetChunks(RemainingRequests, BlockingOwner,
-					[&OnComplete, &RemainingRequests, &ErrorChunks](FCacheChunkResponse&& Response)
+					[&OnComplete, &RemainingRequests, &ErrorChunks](FCacheGetChunkResponse&& Response)
 					{
 						if (Response.Status == EStatus::Error)
 						{
@@ -789,7 +789,7 @@ public:
 
 		if (OnComplete)
 		{
-			for (const FCacheChunkRequest& Request : RemainingRequests)
+			for (const FCacheGetChunkRequest& Request : RemainingRequests)
 			{
 				OnComplete({Request.Name, Request.Key, Request.Id, Request.RawOffset, 0, {}, {}, Request.UserData, EStatus::Error});
 			}
diff --git a/Engine/Source/Developer/DerivedDataCache/Private/HttpDerivedDataBackend.cpp b/Engine/Source/Developer/DerivedDataCache/Private/HttpDerivedDataBackend.cpp
--- a/Engine/Source/Developer/DerivedDataCache/Private/HttpDerivedDataBackend.cpp
+++ b/Engine/Source/Developer/DerivedDataCache/Private/HttpDerivedDataBackend.cpp
@@ -2031,9 +2031,9 @@ public:
 		FOnCacheGetComplete&& OnComplete) override;
 
 	virtual void GetChunks(
-		TConstArrayView<FCacheChunkRequest> Requests,
+		TConstArrayView<FCacheGetChunkRequest> Requests,
 		IRequestOwner& Owner,
-		FOnCacheChunkComplete&& OnComplete) override;
+		FOnCacheGetChunkComplete&& OnComplete) override;
 
 	static FHttpDerivedDataBackend* GetAny()
 	{
@@ -3231,9 +3231,9 @@ void FHttpDerivedDataBackend::Get(
 }
 
 void FHttpDerivedDataBackend::GetChunks(
-	const TConstArrayView<FCacheChunkRequest> Requests,
+	const TConstArrayView<FCacheGetChunkRequest> Requests,
 	IRequestOwner& Owner,
-	FOnCacheChunkComplete&& OnComplete)
+	FOnCacheGetChunkComplete&& OnComplete)
 {
 	// TODO: This is inefficient because Jupiter doesn't allow us to get only part of a compressed blob, so we have to
 	//		 get the whole thing and then decompress only the portion we need.  Furthermore, because there is no propagation
@@ -3242,7 +3242,7 @@ void FHttpDerivedDataBackend::GetChunks(
 	//		 are missing in local/fast stores and have to be retrieved from slow stores into record requests instead.  That
 	//		 will make this code path unused/uncommon as Jupiter will most always be a slow store with a local/fast store in front of it.
 	//		 Regardless, to adhere to the functional contract, this implementation must exist.
-	TArray<FCacheChunkRequest, TInlineAllocator<16>> SortedRequests(Requests);
+	TArray<FCacheGetChunkRequest, TInlineAllocator<16>> SortedRequests(Requests);
 	SortedRequests.StableSort(TChunkLess());
 
 	bool bHasValue = false;
@@ -3253,7 +3253,7 @@ void FHttpDerivedDataBackend::GetChunks(
 	FCompressedBufferReader ValueReader;
 	EStatus ValueStatus = EStatus::Error;
 	FOptionalCacheRecord Record;
-	for (const FCacheChunkRequest& Request : SortedRequests)
+	for (const FCacheGetChunkRequest& Request : SortedRequests)
 	{
 		const bool bExistsOnly = EnumHasAnyFlags(Request.Policy, ECachePolicy::SkipData);
 		COOK_STAT(auto Timer = bExistsOnly ? UsageStats.TimeProbablyExists() : UsageStats.TimeGet());
diff --git a/Engine/Source/Developer/DerivedDataCache/Private/MemoryDerivedDataBackend.cpp b/Engine/Source/Developer/DerivedDataCache/Private/MemoryDerivedDataBackend.cpp
--- a/Engine/Source/Developer/DerivedDataCache/Private/MemoryDerivedDataBackend.cpp
+++ b/Engine/Source/Developer/DerivedDataCache/Private/MemoryDerivedDataBackend.cpp
@@ -107,9 +107,9 @@ public:
 		FOnCacheGetComplete&& OnComplete) override;
 
 	virtual void GetChunks(
-		TConstArrayView<FCacheChunkRequest> Requests,
+		TConstArrayView<FCacheGetChunkRequest> Requests,
 		IRequestOwner& Owner,
-		FOnCacheChunkComplete&& OnComplete) override;
+		FOnCacheGetChunkComplete&& OnComplete) override;
 
 private:
 	/** Name of the cache file loaded (if any). */
@@ -768,14 +768,14 @@ void FMemoryDerivedDataBackend::Get(
 }
 
 void FMemoryDerivedDataBackend::GetChunks(
-	const TConstArrayView<FCacheChunkRequest> Requests,
+	const TConstArrayView<FCacheGetChunkRequest> Requests,
 	IRequestOwner& Owner,
-	FOnCacheChunkComplete&& OnComplete)
+	FOnCacheGetChunkComplete&& OnComplete)
 {
 	FValueWithId Value;
 	FCacheKey ValueKey;
 	FCompressedBufferReader Reader;
-	for (const FCacheChunkRequest& Request : Requests)
+	for (const FCacheGetChunkRequest& Request : Requests)
 	{
 		const bool bExistsOnly = EnumHasAnyFlags(Request.Policy, ECachePolicy::SkipData);
 		COOK_STAT(auto Timer = bExistsOnly ? UsageStats.TimeProbablyExists() : UsageStats.TimeGet());
diff --git a/Engine/Source/Developer/DerivedDataCache/Private/PakFileDerivedDataBackend.cpp b/Engine/Source/Developer/DerivedDataCache/Private/PakFileDerivedDataBackend.cpp
--- a/Engine/Source/Developer/DerivedDataCache/Private/PakFileDerivedDataBackend.cpp
+++ b/Engine/Source/Developer/DerivedDataCache/Private/PakFileDerivedDataBackend.cpp
@@ -126,9 +126,9 @@ public:
 		FOnCacheGetComplete&& OnComplete) override;
 
 	virtual void GetChunks(
-		TConstArrayView<FCacheChunkRequest> Requests,
+		TConstArrayView<FCacheGetChunkRequest> Requests,
 		IRequestOwner& Owner,
-		FOnCacheChunkComplete&& OnComplete) override;
+		FOnCacheGetChunkComplete&& OnComplete) override;
 
 private:
 	bool PutCacheRecord(FStringView Name, const FCacheRecord& Record, const FCacheRecordPolicy& Policy);
@@ -710,16 +710,16 @@ void FPakFileDerivedDataBackend::Get(
 }
 
 void FPakFileDerivedDataBackend::GetChunks(
-	const TConstArrayView<FCacheChunkRequest> Requests,
+	const TConstArrayView<FCacheGetChunkRequest> Requests,
 	IRequestOwner& Owner,
-	FOnCacheChunkComplete&& OnComplete)
+	FOnCacheGetChunkComplete&& OnComplete)
 {
-	TArray<FCacheChunkRequest, TInlineAllocator<16>> SortedRequests(Requests);
+	TArray<FCacheGetChunkRequest, TInlineAllocator<16>> SortedRequests(Requests);
 	SortedRequests.StableSort(TChunkLess());
 
 	FOptionalCacheRecord Record;
 	FCompressedBufferReader Reader;
-	for (const FCacheChunkRequest& Request : SortedRequests)
+	for (const FCacheGetChunkRequest& Request : SortedRequests)
 	{
 		const bool bExistsOnly = EnumHasAnyFlags(Request.Policy, ECachePolicy::SkipData);
 		TRACE_CPUPROFILER_EVENT_SCOPE(PakFileDDC_Get);
diff --git a/Engine/Source/Developer/DerivedDataCache/Private/S3DerivedDataBackend.cpp b/Engine/Source/Developer/DerivedDataCache/Private/S3DerivedDataBackend.cpp
--- a/Engine/Source/Developer/DerivedDataCache/Private/S3DerivedDataBackend.cpp
+++ b/Engine/Source/Developer/DerivedDataCache/Private/S3DerivedDataBackend.cpp
@@ -242,9 +242,9 @@ public:
 		FOnCacheGetComplete&& OnComplete) override;
 
 	virtual void GetChunks(
-		TConstArrayView<FCacheChunkRequest> Requests,
+		TConstArrayView<FCacheGetChunkRequest> Requests,
 		IRequestOwner& Owner,
-		FOnCacheChunkComplete&& OnComplete) override;
+		FOnCacheGetChunkComplete&& OnComplete) override;
 
 private:
 	struct FBundle;
@@ -1343,13 +1343,13 @@ void FS3DerivedDataBackend::Get(
 }
 
 void FS3DerivedDataBackend::GetChunks(
-	const TConstArrayView<FCacheChunkRequest> Requests,
+	const TConstArrayView<FCacheGetChunkRequest> Requests,
 	IRequestOwner& Owner,
-	FOnCacheChunkComplete&& OnComplete)
+	FOnCacheGetChunkComplete&& OnComplete)
 {
 	if (OnComplete)
 	{
-		for (const FCacheChunkRequest& Request : Requests)
+		for (const FCacheGetChunkRequest& Request : Requests)
 		{
 			OnComplete({Request.Name, Request.Key, Request.Id, Request.RawOffset, 0, {}, {}, Request.UserData, EStatus::Error});
 		}
diff --git a/Engine/Source/Developer/DerivedDataCache/Private/Tests/CacheStoreTest.cpp b/Engine/Source/Developer/DerivedDataCache/Private/Tests/CacheStoreTest.cpp
--- a/Engine/Source/Developer/DerivedDataCache/Private/Tests/CacheStoreTest.cpp
+++ b/Engine/Source/Developer/DerivedDataCache/Private/Tests/CacheStoreTest.cpp
@@ -33,7 +33,7 @@ bool FCacheStoreTest::RunTest(const FString& Parameters)
 	TArray<FCachePutValueRequest> PutValueRequests;
 	TArray<FCacheGetRequest> GetRequests;
 	TArray<FCacheGetValueRequest> GetValueRequests;
-	TArray<FCacheChunkRequest> ChunkRequests;
+	TArray<FCacheGetChunkRequest> ChunkRequests;
 	struct FTestData
 	{
 		FValue Value;
@@ -84,7 +84,7 @@ bool FCacheStoreTest::RunTest(const FString& Parameters)
 			PutValueRequests.Add(FCachePutValueRequest{ Name, Key, TestData.Value, PutPolicy, UserData });
 			GetValueRequests.Add(FCacheGetValueRequest{ Name, Key, GetPolicy, UserData });
 		}
-		ChunkRequests.Add(FCacheChunkRequest{ Name, Key, TestData.bUseValueAPI ? FValueId() : ValueId,
+		ChunkRequests.Add(FCacheGetChunkRequest{ Name, Key, TestData.bUseValueAPI ? FValueId() : ValueId,
 			0, TestData.Value.GetRawSize(), FIoHash(), GetPolicy, UserData });
 	}
 
@@ -195,7 +195,7 @@ bool FCacheStoreTest::RunTest(const FString& Parameters)
 			});
 
 #if ZENSERVER_SUPPORTS_VALUEAPI
-		Cache.GetChunks(ChunkRequests, Owner, [&TestDatas, this](FCacheChunkResponse&& Response)
+		Cache.GetChunks(ChunkRequests, Owner, [&TestDatas, this](FCacheGetChunkResponse&& Response)
 			{
 				FTestData* TestData = TestDatas.Find(Response.UserData);
 				int32 n = (int32)Response.UserData;
diff --git a/Engine/Source/Developer/DerivedDataCache/Private/Tests/HttpDerivedDataBackendTest.cpp b/Engine/Source/Developer/DerivedDataCache/Private/Tests/HttpDerivedDataBackendTest.cpp
--- a/Engine/Source/Developer/DerivedDataCache/Private/Tests/HttpDerivedDataBackendTest.cpp
+++ b/Engine/Source/Developer/DerivedDataCache/Private/Tests/HttpDerivedDataBackendTest.cpp
@@ -227,7 +227,7 @@ protected:
 		using namespace UE::DerivedData;
 		UE::DerivedData::ICacheStore* TestBackend = GetTestBackend();
 
-		TArray<FCacheChunkRequest> Requests;
+		TArray<FCacheGetChunkRequest> Requests;
 
 		int32 OverallIndex = 0;
 		for (int32 RecordIndex = 0; RecordIndex < Records.Num(); ++RecordIndex)
@@ -251,7 +251,7 @@ protected:
 		TArray<TOptional<FGetChunksOutput>> GetOutputs;
 		GetOutputs.SetNum(Requests.Num());
 		FRequestOwner RequestOwner(EPriority::Blocking);
-		TestBackend->GetChunks(Requests, RequestOwner, [&GetOutputs](FCacheChunkResponse&& Response)
+		TestBackend->GetChunks(Requests, RequestOwner, [&GetOutputs](FCacheGetChunkResponse&& Response)
 			{
 				GetOutputs[Response.UserData].Emplace(FGetChunksOutput { Response.RawData, Response.Status });
 			});
diff --git a/Engine/Source/Developer/DerivedDataCache/Private/ZenDerivedDataBackend.cpp b/Engine/Source/Developer/DerivedDataCache/Private/ZenDerivedDataBackend.cpp
--- a/Engine/Source/Developer/DerivedDataCache/Private/ZenDerivedDataBackend.cpp
+++ b/Engine/Source/Developer/DerivedDataCache/Private/ZenDerivedDataBackend.cpp
@@ -134,9 +134,9 @@ public:
 		FOnCacheGetValueComplete&& OnComplete = FOnCacheGetValueComplete()) override;
 
 	virtual void GetChunks(
-		TConstArrayView<FCacheChunkRequest> Requests,
+		TConstArrayView<FCacheGetChunkRequest> Requests,
 		IRequestOwner& Owner,
-		FOnCacheChunkComplete&& OnComplete) override;
+		FOnCacheGetChunkComplete&& OnComplete) override;
 
 private:
 	enum class EGetResult
@@ -1083,14 +1083,14 @@ void FZenDerivedDataBackend::GetValue(
 
 
 void FZenDerivedDataBackend::GetChunks(
-	TConstArrayView<FCacheChunkRequest> Requests,
+	TConstArrayView<FCacheGetChunkRequest> Requests,
 	IRequestOwner& Owner,
-	FOnCacheChunkComplete&& OnComplete)
+	FOnCacheGetChunkComplete&& OnComplete)
 {
 	TRACE_CPUPROFILER_EVENT_SCOPE(ZenDDC::GetChunks);
 	TRACE_COUNTER_ADD(ZenDDC_ChunkRequestCount, int64(Requests.Num()));
 
-	TArray<FCacheChunkRequest, TInlineAllocator<16>> SortedRequests(Requests);
+	TArray<FCacheGetChunkRequest, TInlineAllocator<16>> SortedRequests(Requests);
 	SortedRequests.StableSort(TChunkLess());
 
 	int32 TotalCompleted = 0;
@@ -1110,7 +1110,7 @@ void FZenDerivedDataBackend::GetChunks(
 				BatchRequest.BeginArray("ChunkRequests"_ASV);
 				for (int32 ChunkIndex = BatchFirst; ChunkIndex <= BatchLast; ChunkIndex++)
 				{
-					const FCacheChunkRequest& Request = SortedRequests[ChunkIndex];
+					const FCacheGetChunkRequest& Request = SortedRequests[ChunkIndex];
 					
 					BatchRequest.BeginObject();
 					
@@ -1147,7 +1147,7 @@ void FZenDerivedDataBackend::GetChunks(
 			int32 ChunkIndex = BatchFirst;
 			for (FCbFieldView HashView : ResponseObj["Result"_ASV])
 			{
-				const FCacheChunkRequest& Request = SortedRequests[ChunkIndex++];
+				const FCacheGetChunkRequest& Request = SortedRequests[ChunkIndex++];
 
 				if (ShouldSimulateMiss(Request.Key))
 				{
@@ -1195,7 +1195,7 @@ void FZenDerivedDataBackend::GetChunks(
 		{
 			for (int32 ChunkIndex = BatchFirst; ChunkIndex <= BatchLast; ChunkIndex++)
 			{
-				const FCacheChunkRequest& Request = SortedRequests[ChunkIndex];
+				const FCacheGetChunkRequest& Request = SortedRequests[ChunkIndex];
 				
 				UE_LOG(LogDerivedDataCache, Display, TEXT("%s: Cache miss with missing value '%s' for '%s' from '%s'"),
 					*GetName(), *WriteToString<16>(Request.Id), *WriteToString<96>(Request.Key), *Request.Name);
diff --git a/Engine/Source/Developer/DerivedDataCache/Public/DerivedDataCache.h b/Engine/Source/Developer/DerivedDataCache/Public/DerivedDataCache.h
--- a/Engine/Source/Developer/DerivedDataCache/Public/DerivedDataCache.h
+++ b/Engine/Source/Developer/DerivedDataCache/Public/DerivedDataCache.h
@@ -20,8 +20,8 @@
 
 namespace UE::DerivedData { class ICacheStoreMaintainer; }
 namespace UE::DerivedData { class IRequestOwner; }
-namespace UE::DerivedData { struct FCacheChunkRequest; }
-namespace UE::DerivedData { struct FCacheChunkResponse; }
+namespace UE::DerivedData { struct FCacheGetChunkRequest; }
+namespace UE::DerivedData { struct FCacheGetChunkResponse; }
 namespace UE::DerivedData { struct FCacheGetRequest; }
 namespace UE::DerivedData { struct FCacheGetResponse; }
 namespace UE::DerivedData { struct FCacheGetValueRequest; }
@@ -35,13 +35,13 @@ namespace UE::DerivedData::Private { class ICacheRecordPolicyShared; }
 namespace UE::DerivedData
 {
 
-using FCacheRequestName UE_DEPRECATED(5.0, "Use FSharedString.") = FSharedString;
-
 using FOnCachePutComplete = TUniqueFunction<void (FCachePutResponse&& Response)>;
 using FOnCacheGetComplete = TUniqueFunction<void (FCacheGetResponse&& Response)>;
 using FOnCachePutValueComplete = TUniqueFunction<void (FCachePutValueResponse&& Response)>;
 using FOnCacheGetValueComplete = TUniqueFunction<void (FCacheGetValueResponse&& Response)>;
-using FOnCacheChunkComplete = TUniqueFunction<void (FCacheChunkResponse&& Response)>;
+using FOnCacheGetChunkComplete = TUniqueFunction<void (FCacheGetChunkResponse&& Response)>;
+
+using FOnCacheChunkComplete UE_DEPRECATED(5.0, "Use FOnCacheGetChunkComplete.") = FOnCacheGetChunkComplete;
 
 ///////////////////////////////////////////////////////////////////////////////////////////////////////////////////////
 
@@ -300,9 +300,9 @@ public:
 	 * @param OnComplete   A callback invoked for every chunk as it completes or is canceled.
 	 */
 	virtual void GetChunks(
-		TConstArrayView<FCacheChunkRequest> Requests,
+		TConstArrayView<FCacheGetChunkRequest> Requests,
 		IRequestOwner& Owner,
-		FOnCacheChunkComplete&& OnComplete) = 0;
+		FOnCacheGetChunkComplete&& OnComplete) = 0;
 };
 
 ///////////////////////////////////////////////////////////////////////////////////////////////////////////////////////
@@ -365,8 +365,6 @@ struct FCachePutResponse
 	EStatus Status = EStatus::Error;
 };
 
-using FCachePutCompleteParams UE_DEPRECATED(5.0, "Use FCachePutResponse.") = FCachePutResponse;
-
 /** Parameters to request to get a cache record. */
 struct FCacheGetRequest
 {
@@ -406,8 +404,6 @@ struct FCacheGetResponse
 	EStatus Status = EStatus::Error;
 };
 
-using FCacheGetCompleteParams UE_DEPRECATED(5.0, "Use FCacheGetResponse.") = FCacheGetResponse;
-
 /** Parameters to request to put a cache value. */
 struct FCachePutValueRequest
 {
@@ -443,8 +439,6 @@ struct FCachePutValueResponse
 	EStatus Status = EStatus::Error;
 };
 
-using FCachePutValueCompleteParams UE_DEPRECATED(5.0, "Use FCachePutValueResponse.") = FCachePutValueResponse;
-
 /** Parameters to request to get a cache value. */
 struct FCacheGetValueRequest
 {
@@ -485,10 +479,8 @@ struct FCacheGetValueResponse
 	EStatus Status = EStatus::Error;
 };
 
-using FCacheGetValueCompleteParams UE_DEPRECATED(5.0, "Use FCacheGetValueResponse.") = FCacheGetValueResponse;
-
 /** Parameters to request a chunk, which is a subset of a value, from a cache record or cache value. */
-struct FCacheChunkRequest
+struct FCacheGetChunkRequest
 {
 	/** A name to identify this request for logging and profiling. An object path is typically sufficient. */
 	FSharedString Name;
@@ -515,8 +507,10 @@ struct FCacheChunkRequest
 	uint64 UserData = 0;
 };
 
+using FCacheChunkRequest UE_DEPRECATED(5.0, "Use FCacheGetChunkRequest.") = FCacheGetChunkRequest;
+
 /** Parameters for the completion callback for cache chunk requests. */
-struct FCacheChunkResponse
+struct FCacheGetChunkResponse
 {
 	/** A copy of the name from the request. */
 	FSharedString Name;
@@ -546,7 +540,7 @@ struct FCacheChunkResponse
 	EStatus Status = EStatus::Error;
 };
 
-using FCacheChunkCompleteParams UE_DEPRECATED(5.0, "Use FCacheChunkResponse.") = FCacheChunkResponse;
+using FCacheChunkResponse UE_DEPRECATED(5.0, "Use FCacheGetChunkResponse.") = FCacheGetChunkResponse;
 
 ///////////////////////////////////////////////////////////////////////////////////////////////////////////////////////
 
diff --git a/Engine/Source/Editor/UnrealEd/Private/EditorDomain/EditorDomainArchive.cpp b/Engine/Source/Editor/UnrealEd/Private/EditorDomain/EditorDomainArchive.cpp
--- a/Engine/Source/Editor/UnrealEd/Private/EditorDomain/EditorDomainArchive.cpp
+++ b/Engine/Source/Editor/UnrealEd/Private/EditorDomain/EditorDomainArchive.cpp
@@ -393,10 +393,10 @@ void FEditorDomainPackageSegments::SendSegmentRequest(FSegment& Segment)
 
 	// Note that Segment.RequestOwner is Interface-only and so we can write it outside the lock
 	Segment.RequestOwner.Emplace(EPriority::Normal);
-	FCacheChunkRequest SegmentChunk{{PackagePath.GetDebugName()}, UE::EditorDomain::GetEditorDomainPackageKey(EditorDomainHash), Segment.ValueId};
+	FCacheGetChunkRequest SegmentChunk{{PackagePath.GetDebugName()}, UE::EditorDomain::GetEditorDomainPackageKey(EditorDomainHash), Segment.ValueId};
 	SegmentChunk.Policy = ECachePolicy::Local;
 	Cache.GetChunks({SegmentChunk}, *Segment.RequestOwner,
-		[this, &Segment](FCacheChunkResponse&& Response)
+		[this, &Segment](FCacheGetChunkResponse&& Response)
 		{
 			if (AsyncSource == ESource::Closed)
 			{
diff --git a/Engine/Source/Runtime/Engine/Private/Streaming/Texture2DMipDataProvider_DDC.cpp b/Engine/Source/Runtime/Engine/Private/Streaming/Texture2DMipDataProvider_DDC.cpp
--- a/Engine/Source/Runtime/Engine/Private/Streaming/Texture2DMipDataProvider_DDC.cpp
+++ b/Engine/Source/Runtime/Engine/Private/Streaming/Texture2DMipDataProvider_DDC.cpp
@@ -55,7 +55,7 @@ void FTexture2DMipDataProvider_DDC::Init(const FTextureUpdateContext& Context, c
 			else if (PlatformData->DerivedDataKey.IsType<UE::DerivedData::FCacheKeyProxy>())
 			{
 				using namespace UE::DerivedData;
-				TArray<FCacheChunkRequest> MipKeys;
+				TArray<FCacheGetChunkRequest> MipKeys;
 
 				TStringBuilder<256> MipNameBuilder;
 				Context.Texture->GetPathName(nullptr, MipNameBuilder);
@@ -67,7 +67,7 @@ void FTexture2DMipDataProvider_DDC::Init(const FTextureUpdateContext& Context, c
 					const FTexture2DMipMap& MipMap = *Context.MipsView[MipIndex];
 					if (MipMap.IsPagedToDerivedData())
 					{
-						FCacheChunkRequest& Request = MipKeys.AddDefaulted_GetRef();
+						FCacheGetChunkRequest& Request = MipKeys.AddDefaulted_GetRef();
 						MipNameBuilder.Appendf(TEXT(" [MIP 0]"), MipIndex + LODBias);
 						Request.Name = MipNameBuilder;
 						Request.Key = Key;
@@ -79,7 +79,7 @@ void FTexture2DMipDataProvider_DDC::Init(const FTextureUpdateContext& Context, c
 
 				if (MipKeys.Num())
 				{
-					GetCache().GetChunks(MipKeys, DDCRequestOwner, [this](FCacheChunkResponse&& Response)
+					GetCache().GetChunks(MipKeys, DDCRequestOwner, [this](FCacheGetChunkResponse&& Response)
 					{
 						if (Response.Status == EStatus::Ok)
 						{
diff --git a/Engine/Source/Runtime/Engine/Private/Streaming/Texture2DStreamIn_DDC.cpp b/Engine/Source/Runtime/Engine/Private/Streaming/Texture2DStreamIn_DDC.cpp
--- a/Engine/Source/Runtime/Engine/Private/Streaming/Texture2DStreamIn_DDC.cpp
+++ b/Engine/Source/Runtime/Engine/Private/Streaming/Texture2DStreamIn_DDC.cpp
@@ -158,7 +158,7 @@ void FTexture2DStreamIn_DDC::DoCreateAsyncDDCRequests(const FContext& Context)
 		else if (PlatformData->DerivedDataKey.IsType<UE::DerivedData::FCacheKeyProxy>())
 		{
 			using namespace UE::DerivedData;
-			TArray<FCacheChunkRequest> MipKeys;
+			TArray<FCacheGetChunkRequest> MipKeys;
 
 			TStringBuilder<256> MipNameBuilder;
 			Context.Texture->GetPathName(nullptr, MipNameBuilder);
@@ -170,7 +170,7 @@ void FTexture2DStreamIn_DDC::DoCreateAsyncDDCRequests(const FContext& Context)
 				const FTexture2DMipMap& MipMap = *Context.MipsView[MipIndex];
 				if (MipMap.IsPagedToDerivedData())
 				{
-					FCacheChunkRequest& Request = MipKeys.AddDefaulted_GetRef();
+					FCacheGetChunkRequest& Request = MipKeys.AddDefaulted_GetRef();
 					MipNameBuilder.Appendf(TEXT(" [MIP 0]"), MipIndex + LODBias);
 					Request.Name = MipNameBuilder;
 					Request.Key = Key;
@@ -182,7 +182,7 @@ void FTexture2DStreamIn_DDC::DoCreateAsyncDDCRequests(const FContext& Context)
 
 			if (MipKeys.Num())
 			{
-				GetCache().GetChunks(MipKeys, DDCRequestOwner, [this](FCacheChunkResponse&& Response)
+				GetCache().GetChunks(MipKeys, DDCRequestOwner, [this](FCacheGetChunkResponse&& Response)
 				{
 					if (Response.Status == EStatus::Ok)
 					{
diff --git a/Engine/Source/Runtime/Engine/Private/TextureDerivedData.cpp b/Engine/Source/Runtime/Engine/Private/TextureDerivedData.cpp
--- a/Engine/Source/Runtime/Engine/Private/TextureDerivedData.cpp
+++ b/Engine/Source/Runtime/Engine/Private/TextureDerivedData.cpp
@@ -1281,7 +1281,7 @@ static bool BeginLoadDerivedMips(FTexturePlatformData& PlatformData, int32 First
 	}
 	else if (PlatformData.DerivedDataKey.IsType<FCacheKeyProxy>())
 	{
-		TArray<FCacheChunkRequest> MipKeys;
+		TArray<FCacheGetChunkRequest> MipKeys;
 
 		const FCacheKey& Key = *PlatformData.DerivedDataKey.Get<UE::DerivedData::FCacheKeyProxy>().AsCacheKey();
 		for (int32 MipIndex = FirstMipToLoad; MipIndex < Mips.Num(); ++MipIndex)
@@ -1291,7 +1291,7 @@ static bool BeginLoadDerivedMips(FTexturePlatformData& PlatformData, int32 First
 			{
 				TStringBuilder<256> MipNameBuilder;
 				MipNameBuilder.Append(DebugContext).Appendf(TEXT(" [MIP 0]"), MipIndex);
-				FCacheChunkRequest& Request = MipKeys.AddDefaulted_GetRef();
+				FCacheGetChunkRequest& Request = MipKeys.AddDefaulted_GetRef();
 				Request.Name = MipNameBuilder;
 				Request.Key = Key;
 				Request.Id = FTexturePlatformData::MakeMipId(MipIndex);
@@ -1304,7 +1304,7 @@ static bool BeginLoadDerivedMips(FTexturePlatformData& PlatformData, int32 First
 			check(Callback);
 			bool bMiss = false;
 			FRequestOwner RequestOwner(EPriority::Blocking);
-			GetCache().GetChunks(MipKeys, RequestOwner, [Callback = MoveTemp(Callback), &bMiss](FCacheChunkResponse&& Response)
+			GetCache().GetChunks(MipKeys, RequestOwner, [Callback = MoveTemp(Callback), &bMiss](FCacheGetChunkResponse&& Response)
 			{
 				if (Response.Status == EStatus::Ok)
 				{
