Fix static analysis warnings: possible dereference of null. The analyzer did not know that TestTrue only returns true if its argument is true.
#rnx

diff --git a/Engine/Source/Developer/DerivedDataCache/Private/Tests/CacheStoreTest.cpp b/Engine/Source/Developer/DerivedDataCache/Private/Tests/CacheStoreTest.cpp
--- a/Engine/Source/Developer/DerivedDataCache/Private/Tests/CacheStoreTest.cpp
+++ b/Engine/Source/Developer/DerivedDataCache/Private/Tests/CacheStoreTest.cpp
@@ -93,7 +93,8 @@ bool FCacheStoreTest::RunTest(const FString& Parameters)
 		Cache.Put(PutRequests, Owner, [&TestDatas, this](FCachePutResponse&& Response)
 			{
 				FTestData* TestData = TestDatas.Find(Response.UserData);
-				if (TestTrue(TEXT("Valid UserData in Put Callback"), TestData != nullptr))
+				TestTrue(TEXT("Valid UserData in Put Callback"), TestData != nullptr);
+				if (TestData)
 				{
 					TestData->bReceivedPut = true;
 				}
@@ -101,7 +102,8 @@ bool FCacheStoreTest::RunTest(const FString& Parameters)
 		Cache.PutValue(PutValueRequests, Owner, [&TestDatas, this](FCachePutValueResponse&& Response)
 			{
 				FTestData* TestData = TestDatas.Find(Response.UserData);
-				if (TestTrue(TEXT("Valid UserData in PutValue Callback"), TestData != nullptr))
+				TestTrue(TEXT("Valid UserData in PutValue Callback"), TestData != nullptr);
+				if (TestData)
 				{
 					TestData->bReceivedPutValue = true;
 				}
@@ -128,7 +130,8 @@ bool FCacheStoreTest::RunTest(const FString& Parameters)
 			{
 				FTestData* TestData = TestDatas.Find(Response.UserData);
 				int32 n = (int32)Response.UserData;
-				if (TestTrue(TEXT("Valid UserData in Get Callback"), TestData != nullptr))
+				TestTrue(TEXT("Valid UserData in Get Callback"), TestData != nullptr);
+				if (TestData)
 				{
 					TestData->bReceivedGet = true;
 					if (TestTrue(*WriteToString<32>(TEXT("Get %d succeeded"), n), Response.Status == EStatus::Ok))
@@ -166,7 +169,8 @@ bool FCacheStoreTest::RunTest(const FString& Parameters)
 			{
 				FTestData* TestData = TestDatas.Find(Response.UserData);
 				int32 n = (int32)Response.UserData;
-				if (TestTrue(TEXT("Valid UserData in GetValue Callback"), TestData != nullptr))
+				TestTrue(TEXT("Valid UserData in GetValue Callback"), TestData != nullptr);
+				if (TestData)
 				{
 					TestData->bReceivedGetValue = true;
 					if (TestTrue(*WriteToString<32>(TEXT("GetValue %d succeeded"), n), Response.Status == EStatus::Ok))
