DDC: Added hashing and comparison to FSharedString and updated natvis

#rnx

diff --git a/Engine/Source/Developer/DerivedDataCache/Public/DerivedDataSharedString.h b/Engine/Source/Developer/DerivedDataCache/Public/DerivedDataSharedString.h
--- a/Engine/Source/Developer/DerivedDataCache/Public/DerivedDataSharedString.h
+++ b/Engine/Source/Developer/DerivedDataCache/Public/DerivedDataSharedString.h
@@ -157,6 +157,24 @@ void TSharedString<CharType>::Release(CharType* const MaybeNullChars)
 	}
 }
 
+template <typename CharType>
+inline uint32 GetTypeHash(const TSharedString<CharType>& String)
+{
+	return GetTypeHash(MakeStringView(String));
+}
+
+template <typename CharType>
+inline bool operator==(const TSharedString<CharType>& Lhs, const TSharedString<CharType>& Rhs)
+{
+	return MakeStringView(Lhs).Equals(MakeStringView(Rhs), ESearchCase::IgnoreCase);
+}
+
+template <typename CharType>
+inline bool operator<(const TSharedString<CharType>& Lhs, const TSharedString<CharType>& Rhs)
+{
+	return MakeStringView(Lhs).Compare(MakeStringView(Rhs), ESearchCase::IgnoreCase) < 0;
+}
+
 } // UE::DerivedData
 
 template <typename CharType>
