Implements a DDC.Graph cvar to stick DDC.Graph="DerivedDataBackendGraph_Jupiter" in ConsoleVariables.ini once and for all locally

diff --git a/Engine/Source/Developer/DerivedDataCache/Private/DerivedDataBackends.cpp b/Engine/Source/Developer/DerivedDataCache/Private/DerivedDataBackends.cpp
--- a/Engine/Source/Developer/DerivedDataCache/Private/DerivedDataBackends.cpp
+++ b/Engine/Source/Developer/DerivedDataCache/Private/DerivedDataBackends.cpp
@@ -33,6 +33,11 @@ DEFINE_LOG_CATEGORY(LogDerivedDataCache);
 #define MAX_BACKEND_KEY_LENGTH (120)
 #define LOCTEXT_NAMESPACE "DerivedDataBackendGraph"
 
+static TAutoConsoleVariable<FString> GDumpDDCName(
+	TEXT("DDC.Graph"), TEXT(""),
+	TEXT("Name of custom DDC Graph to use."),
+	ECVF_Default);
+
 namespace UE::DerivedData::CacheStore::AsyncPut
 {
 FDerivedDataBackendInterface* CreateAsyncPutDerivedDataBackend(FDerivedDataBackendInterface* InnerBackend, bool bCacheInFlightPuts);
@@ -113,19 +118,25 @@ public:
 		TMap<FString, FDerivedDataBackendInterface*> ParsedNodes;		
 
 		// Create the graph using ini settings. The string "default" forwards creation to use the default graph.
-		if( FParse::Value( FCommandLine::Get(), TEXT("-DDC="), GraphName ) && FCString::Stricmp(*GraphName, TEXT("default")) != 0 )
+
+		FString CustomDDCName;
+		FParse::Value(FCommandLine::Get(), TEXT("-DDC="), CustomDDCName);
+		if (CustomDDCName.IsEmpty() || FCString::Stricmp(*GraphName, TEXT("default")) == 0)
 		{
-			if( GraphName.Len() > 0 )
-			{
-				if ( GraphName == TEXT("None"))
-				{
-					UE_LOG(LogDerivedDataCache, Display, TEXT("DDC backend graph of 'None' specified, graph will be invalid for use."));
-					return;
-				}
+			CustomDDCName = GDumpDDCName.GetValueOnGameThread();
+		}
 
-				RootCache = ParseNode( TEXT("Root"), GEngineIni, *GraphName, ParsedNodes );				
+		if (!CustomDDCName.IsEmpty())
+		{
+			GraphName = CustomDDCName;
+			if (GraphName == TEXT("None"))
+			{
+				UE_LOG(LogDerivedDataCache, Display, TEXT("DDC backend graph of 'None' specified, graph will be invalid for use."));
+				return;
 			}
 
+			RootCache = ParseNode(TEXT("Root"), GEngineIni, *GraphName, ParsedNodes);
+
 			if( RootCache == NULL )
 			{
 				// Remove references to any backend instances that might have been created
