CachePolicy: Add Save/Load to StringBuilder.
CacheRecordPolicy: Add Save/Load to CompactBinary
#rnx

diff --git a/Engine/Source/Developer/DerivedDataCache/Private/DerivedDataCache.cpp b/Engine/Source/Developer/DerivedDataCache/Private/DerivedDataCache.cpp
--- a/Engine/Source/Developer/DerivedDataCache/Private/DerivedDataCache.cpp
+++ b/Engine/Source/Developer/DerivedDataCache/Private/DerivedDataCache.cpp
@@ -9,6 +9,8 @@
 #include "Algo/Sort.h"
 #include "Async/AsyncWork.h"
 #include "Async/TaskGraphInterfaces.h"
+#include "Containers/Map.h"
+#include "Containers/StringConv.h"
 #include "DDCCleanup.h"
 #include "DerivedDataBackendInterface.h"
 #include "DerivedDataCache.h"
@@ -23,8 +25,11 @@
 #include "Misc/CoreMisc.h"
 #include "Misc/ScopeLock.h"
 #include "ProfilingDebugging/CookStats.h"
+#include "Serialization/CompactBinary.h"
+#include "Serialization/CompactBinaryWriter.h"
 #include "Stats/Stats.h"
 #include "Stats/StatsMisc.h"
+#include "String/ParseTokens.h"
 #include "ZenServerInterface.h"
 #include <atomic>
 
@@ -164,6 +169,141 @@ namespace UE::DerivedData::Private { class FCacheRecordPolicyShared; }
 namespace UE::DerivedData
 {
 
+namespace Private::CachePolicy
+{
+	// TODO: Implement these as Ansi String instead to maximize the most-optimal path, see ParseCachePolicyImpl
+	constexpr TCHAR			  DelimiterChar = TEXT(',');
+	constexpr FStringView Delimiter = TEXTVIEW(",");
+	constexpr FStringView None = TEXTVIEW("None");
+	constexpr FStringView QueryLocal = TEXTVIEW("QueryLocal");
+	constexpr FStringView QueryRemote = TEXTVIEW("QueryRemote");
+	constexpr FStringView Query = TEXTVIEW("Query");
+	constexpr FStringView StoreLocal = TEXTVIEW("StoreLocal");
+	constexpr FStringView StoreRemote = TEXTVIEW("StoreRemote");
+	constexpr FStringView Store = TEXTVIEW("Store");
+	constexpr FStringView SkipMeta = TEXTVIEW("SkipMeta");
+	constexpr FStringView SkipData = TEXTVIEW("SkipData");
+	constexpr FStringView PartialRecord = TEXTVIEW("PartialRecord");
+	constexpr FStringView KeepAlive = TEXTVIEW("KeepAlive");
+	constexpr FStringView Local = TEXTVIEW("Local");
+	constexpr FStringView Remote = TEXTVIEW("Remote");
+	constexpr FStringView Default = TEXTVIEW("Default");
+	constexpr FStringView Disable = TEXTVIEW("Disable");
+
+	const TMap<FStringView, ECachePolicy> TextToPolicy 
+	{
+		{None,			ECachePolicy::None},
+		{QueryLocal,	ECachePolicy::QueryLocal},
+		{QueryRemote,	ECachePolicy::QueryRemote},
+		{Query,			ECachePolicy::Query},
+		{StoreLocal,	ECachePolicy::StoreLocal},
+		{StoreRemote,	ECachePolicy::StoreRemote},
+		{Store,			ECachePolicy::Store},
+		{SkipMeta,		ECachePolicy::SkipMeta},
+		{SkipData,		ECachePolicy::SkipData},
+		{PartialRecord, ECachePolicy::PartialRecord},
+		{KeepAlive,		ECachePolicy::KeepAlive},
+		{Local,			ECachePolicy::Local},
+		{Remote,		ECachePolicy::Remote},
+		{Default,		ECachePolicy::Default},
+		{Disable,		ECachePolicy::Disable}
+	};
+
+	using FPolicyTextPair = TPair<ECachePolicy, FStringView>;
+	const FPolicyTextPair FlagsToString[]
+	{
+		// Order of these Flags is important: we want the aliases before the atomic values,
+		// and the bigger aliases first, to reduce the number of tokens we add
+		{ ECachePolicy::Default, Default },
+		{ ECachePolicy::Remote, Remote },
+		{ ECachePolicy::Local, Local },
+		{ ECachePolicy::Store, Store },
+		{ ECachePolicy::Query, Query },
+
+		// Order of Atomics doesn't matter, so arbitrarily we list them in enum order
+		{ ECachePolicy::QueryLocal, QueryLocal },
+		{ ECachePolicy::QueryRemote, QueryRemote },
+		{ ECachePolicy::StoreLocal, StoreLocal },
+		{ ECachePolicy::StoreRemote, StoreRemote },
+		{ ECachePolicy::SkipMeta, SkipMeta },
+		{ ECachePolicy::SkipData, SkipData },
+		{ ECachePolicy::PartialRecord, PartialRecord },
+		{ ECachePolicy::KeepAlive, KeepAlive },
+
+		// None must come at the end of the array, to write out only if no others exist
+		{ ECachePolicy::None, None },
+	};
+	constexpr ECachePolicy KnownFlags = ECachePolicy::Default | ECachePolicy::SkipMeta | ECachePolicy::SkipData
+		| ECachePolicy::KeepAlive | ECachePolicy::PartialRecord;
+
+}	// namespace Private::CachePolicy
+
+template <typename CharType>
+TStringBuilderBase<CharType>& AppendToBuilderImpl(
+	TStringBuilderBase<CharType>& Builder, UE::DerivedData::ECachePolicy Policy)
+{
+	// Remove any bits we don't recognize; write None if there are not any bits we recognize
+	Policy = Policy & Private::CachePolicy::KnownFlags;
+	for (const Private::CachePolicy::FPolicyTextPair& Pair : Private::CachePolicy::FlagsToString)
+	{
+		if (EnumHasAllFlags(Policy, Pair.Key))
+		{
+			EnumRemoveFlags(Policy, Pair.Key);
+			Builder << Pair.Value << Private::CachePolicy::DelimiterChar;
+			if (Policy == ECachePolicy::None)
+			{
+				break;
+			}
+		}
+	}
+	Builder.RemoveSuffix(1); // Text will have been added by ECachePolicy::None if not by anything else
+	return Builder;
+}
+FAnsiStringBuilderBase& operator<<(FAnsiStringBuilderBase& Builder, UE::DerivedData::ECachePolicy Policy)
+{
+	return AppendToBuilderImpl(Builder, Policy);
+}
+FUtf8StringBuilderBase& operator<<(FUtf8StringBuilderBase& Builder, UE::DerivedData::ECachePolicy Policy)
+{
+	return AppendToBuilderImpl(Builder, Policy);
+}
+FWideStringBuilderBase& operator<<(FWideStringBuilderBase& Builder, UE::DerivedData::ECachePolicy Policy)
+{
+	return AppendToBuilderImpl(Builder, Policy);
+}
+
+template <typename CharType>
+ECachePolicy ParseCachePolicyImpl(TStringView<CharType> Text)
+{
+	checkf(!Text.IsEmpty(), TEXT("Empty string is not valid input to ParseCachePolicy"));
+
+	ECachePolicy Result = ECachePolicy::None;
+	// TODO: Implement ParseTokens for FAnsiStringView so we can convert to Ansi instead of Wide
+	auto WideText = StringCast<TCHAR, 128>(Text.GetData(), Text.Len());
+	UE::String::ParseTokens(WideText, Private::CachePolicy::DelimiterChar, [&Result](FStringView Token) {
+		const ECachePolicy* TokenPolicy = Private::CachePolicy::TextToPolicy.Find(Token);
+		if (TokenPolicy)
+		{
+			Result |= *TokenPolicy;
+		}
+	});
+
+	return Result;
+}
+ECachePolicy ParseCachePolicy(FAnsiStringView Text)
+{
+	return ParseCachePolicyImpl(Text);
+}
+ECachePolicy ParseCachePolicy(FUtf8StringView Text)
+{
+	return ParseCachePolicyImpl(Text);
+}
+ECachePolicy ParseCachePolicy(FWideStringView Text)
+{
+	return ParseCachePolicyImpl(Text);
+}
+
+
 class Private::FCacheRecordPolicyShared final : public Private::ICacheRecordPolicyShared
 {
 public:
@@ -229,6 +369,53 @@ FCacheRecordPolicy FCacheRecordPolicy::Transform(TFunctionRef<ECachePolicy (ECac
 	return Builder.Build();
 }
 
+void FCacheRecordPolicy::Save(FCbWriter& Writer) const
+{
+	Writer.BeginObject();
+	{
+		// The RecordPolicy is calculated from the ValuePolicies and does not need to be saved separately.
+		Writer << "DefaultValuePolicy"_ASV << WriteToUtf8String<128>(GetDefaultValuePolicy());
+		if (!IsUniform())
+		{
+			// FCacheRecordPolicyBuilder guarantees IsUniform -> non-empty GetValuePolicies. Small size penalty here if not.
+			Writer.BeginArray("ValuePolicies"_ASV);
+			{
+				for (const FCacheValuePolicy& ValuePolicy : GetValuePolicies())
+				{
+					// FCacheRecordPolicyBuilder is responsible for ensuring that each ValuePolicy != DefaultValuePolicy
+					// If it lets any duplicates through we will incur a small serialization size penalty here
+					Writer.BeginObject();
+					Writer << "Id"_ASV << ValuePolicy.Id;
+					Writer << "Policy"_ASV << WriteToUtf8String<128>(ValuePolicy.Policy);
+					Writer.EndObject();
+				}
+			}
+			Writer.EndArray();
+		}
+	}
+	Writer.EndObject();
+}
+
+FCacheRecordPolicy FCacheRecordPolicy::Load(FCbObjectView Object, ECachePolicy DefaultPolicy)
+{
+	FUtf8StringView PolicyText = Object["DefaultValuePolicy"_ASV].AsString();
+	ECachePolicy DefaultValuePolicy = !PolicyText.IsEmpty() ? ParseCachePolicy(PolicyText) : DefaultPolicy;
+
+	FCacheRecordPolicyBuilder Builder(DefaultValuePolicy);
+	for (FCbFieldView ValueObjectField : Object["ValuePolicies"_ASV])
+	{
+		FCbObjectView ValueObject = ValueObjectField.AsObjectView();
+		const FCbObjectId ValueId = ValueObject["Id"_ASV].AsObjectId();
+		PolicyText = ValueObject["Policy"_ASV].AsString();
+		ECachePolicy ValuePolicy = !PolicyText.IsEmpty() ? ParseCachePolicy(PolicyText) : DefaultValuePolicy;
+		// FCacheRecordPolicyBuilder should guarantee that FValueId(ValueId).IsValid and ValuePolicy != DefaultValuePolicy
+		// If it lets any through we will have unused data in the record we create.
+		Builder.AddValuePolicy(ValueId, ValuePolicy);
+	}
+
+	return Builder.Build();
+}
+
 void FCacheRecordPolicyBuilder::AddValuePolicy(const FCacheValuePolicy& Policy)
 {
 	if (!Shared)
diff --git a/Engine/Source/Developer/DerivedDataCache/Public/DerivedDataCache.h b/Engine/Source/Developer/DerivedDataCache/Public/DerivedDataCache.h
--- a/Engine/Source/Developer/DerivedDataCache/Public/DerivedDataCache.h
+++ b/Engine/Source/Developer/DerivedDataCache/Public/DerivedDataCache.h
@@ -18,6 +18,8 @@
 
 #define UE_API DERIVEDDATACACHE_API
 
+class FCbObjectView;
+class FCbWriter;
 namespace UE::DerivedData { class ICacheStoreMaintainer; }
 namespace UE::DerivedData { class IRequestOwner; }
 namespace UE::DerivedData { struct FCacheGetChunkRequest; }
@@ -120,6 +122,14 @@ enum class ECachePolicy : uint32
 };
 
 ENUM_CLASS_FLAGS(ECachePolicy);
+/** Serialize Policy to text and append to Builder. Appended text will not be empty. */
+UE_API FUtf8StringBuilderBase& operator<<(FUtf8StringBuilderBase& Builder, ECachePolicy Policy);
+UE_API FAnsiStringBuilderBase& operator<<(FAnsiStringBuilderBase& Builder, ECachePolicy Policy);
+UE_API FWideStringBuilderBase& operator<<(FWideStringBuilderBase& Builder, ECachePolicy Policy);
+/** Parse text written by operator<< back into an ECachePolicy. Text must not be empty. */
+UE_API ECachePolicy ParseCachePolicy(FUtf8StringView Text);
+UE_API ECachePolicy ParseCachePolicy(FAnsiStringView Text);
+UE_API ECachePolicy ParseCachePolicy(FWideStringView Text);
 
 ///////////////////////////////////////////////////////////////////////////////////////////////////////////////////////
 
@@ -185,6 +195,15 @@ public:
 	/** Returns a copy of this policy transformed by an operation. */
 	UE_API FCacheRecordPolicy Transform(TFunctionRef<ECachePolicy (ECachePolicy)> Op) const;
 
+	/** Save the values from *this into the given writer. */
+	UE_API void Save(FCbWriter& Writer) const;
+
+	/**
+	 * Returns a policy loaded from values on Object.
+	 * Invalid data will result in a uniform CacheRecordPolicy with defaultValuePolicy == DefaultPolicy.
+	 */
+	static UE_API FCacheRecordPolicy Load(FCbObjectView Object, ECachePolicy DefaultPolicy = ECachePolicy::Default);
+
 private:
 	friend class FCacheRecordPolicyBuilder;
 
