ZenDerivedDataBackend: CachePolicy use new string serialization. CacheRecordPolicy use CacheRecordPolicy::Save. GetCachePayloads change name to GetCacheValues.

diff --git a/Engine/Source/Developer/DerivedDataCache/Private/ZenDerivedDataBackend.cpp b/Engine/Source/Developer/DerivedDataCache/Private/ZenDerivedDataBackend.cpp
--- a/Engine/Source/Developer/DerivedDataCache/Private/ZenDerivedDataBackend.cpp
+++ b/Engine/Source/Developer/DerivedDataCache/Private/ZenDerivedDataBackend.cpp
@@ -487,63 +487,9 @@ void FZenDerivedDataBackend::AppendZenUri(const FCacheKey& CacheKey, const FValu
 
 void FZenDerivedDataBackend::AppendPolicyQueryString(ECachePolicy Policy, FStringBuilderBase& Uri)
 {
-	bool bQueryEmpty = true;
-	bool bValueEmpty = true;
-	auto AppendKey = [&Uri, &bQueryEmpty, &bValueEmpty](const TCHAR* Key)
+	if (Policy != ECachePolicy::Default)
 	{
-		if (bQueryEmpty)
-		{
-			TCHAR LastChar = Uri.Len() == 0 ? '\0' : Uri.LastChar();
-			if (LastChar != '?' && LastChar != '&')
-			{
-				Uri << '?';
-			}
-			bQueryEmpty = false;
-		}
-		else
-		{
-			Uri << '&';
-		}
-		bValueEmpty = true;
-		Uri << Key;
-	};
-	auto AppendValue = [&Uri, &bValueEmpty](const TCHAR* Value)
-	{
-		if (bValueEmpty)
-		{
-			bValueEmpty = false;
-		}
-		else
-		{
-			Uri << ',';
-		}
-		Uri << Value;
-	};
-
-	if (!EnumHasAllFlags(Policy, ECachePolicy::Query))
-	{
-		AppendKey(TEXT("query="));
-		if (EnumHasAnyFlags(Policy, ECachePolicy::QueryLocal)) { AppendValue(TEXT("local")); }
-		if (EnumHasAnyFlags(Policy, ECachePolicy::QueryRemote)) { AppendValue(TEXT("remote")); }
-		if (!EnumHasAnyFlags(Policy, ECachePolicy::Query)) { AppendValue(TEXT("none")); }
-	}
-	if (!EnumHasAllFlags(Policy, ECachePolicy::Store))
-	{
-		AppendKey(TEXT("store="));
-		if (EnumHasAnyFlags(Policy, ECachePolicy::StoreLocal)) { AppendValue(TEXT("local")); }
-		if (EnumHasAnyFlags(Policy, ECachePolicy::StoreRemote)) { AppendValue(TEXT("remote")); }
-		if (!EnumHasAnyFlags(Policy, ECachePolicy::Store)) { AppendValue(TEXT("none")); }
-	}
-	if (EnumHasAnyFlags(Policy, ECachePolicy::SkipMeta | ECachePolicy::SkipData))
-	{
-		AppendKey(TEXT("skip="));
-		if (EnumHasAllFlags(Policy, ECachePolicy::SkipMeta | ECachePolicy::SkipData)) { AppendValue(TEXT("data")); }
-		else
-		{
-			if (EnumHasAnyFlags(Policy, ECachePolicy::SkipMeta)) { AppendValue(TEXT("meta")); }
-			if (EnumHasAnyFlags(Policy, ECachePolicy::SkipData)) { AppendValue(TEXT("value")); }
-			if (EnumHasAnyFlags(Policy, ECachePolicy::SkipData)) { AppendValue(TEXT("attachments")); }
-		}
+		Uri << TEXT("?Policy=") << Policy;
 	}
 }
 
@@ -786,27 +732,8 @@ void FZenDerivedDataBackend::Get(
 				BatchRequest.EndArray();
 
 				// TODO: The policy needs to be sent with each key.
-				const FCacheRecordPolicy& Policy = Requests[BatchFirst].Policy;
-				BatchRequest.BeginObject("Policy"_ASV);
-				{
-					BatchRequest << "RecordPolicy"_ASV << static_cast<uint32>(Policy.GetRecordPolicy());
-					BatchRequest << "DefaultPayloadPolicy"_ASV << static_cast<uint32>(Policy.GetDefaultValuePolicy());
-					
-					TConstArrayView<FCacheValuePolicy> ValuePolicies = Policy.GetValuePolicies();
-					if (ValuePolicies.Num())
-					{
-						BatchRequest.BeginArray("PayloadPolicies"_ASV);
-						for (const FCacheValuePolicy& ValuePolicy : ValuePolicies)
-						{
-							BatchRequest.BeginObject();
-							BatchRequest.AddObjectId("Id"_ASV, ValuePolicy.Id);
-							BatchRequest << "Policy"_ASV << static_cast<uint32>(ValuePolicy.Policy);
-							BatchRequest.EndObject();
-						}
-						BatchRequest.EndArray();
-					}
-				}
-				BatchRequest.EndObject();
+				BatchRequest.SetName("Policy"_ASV);
+				Requests[BatchFirst].Policy.Save(BatchRequest);
 			}
 			BatchRequest.EndObject();
 		}
@@ -1104,7 +1031,7 @@ void FZenDerivedDataBackend::GetChunks(
 		FCbWriter BatchRequest;
 		BatchRequest.BeginObject();
 		{
-			BatchRequest << "Method"_ASV << "GetCachePayloads";
+			BatchRequest << "Method"_ASV << "GetCacheValues";
 			BatchRequest.BeginObject("Params"_ASV);
 			{
 				BatchRequest.BeginArray("ChunkRequests"_ASV);
@@ -1119,7 +1046,7 @@ void FZenDerivedDataBackend::GetChunks(
 					BatchRequest << "Hash"_ASV << Request.Key.Hash;
 					BatchRequest.EndObject();
 
-					BatchRequest.AddObjectId("PayloadId"_ASV, Request.Id);
+					BatchRequest.AddObjectId("ValueId"_ASV, Request.Id);
 					BatchRequest << "RawOffset"_ASV << Request.RawOffset;
 					BatchRequest << "RawSize"_ASV << Request.RawSize;
 					BatchRequest << "Policy"_ASV << static_cast<uint32>(Request.Policy);
