Change the DDC Http backend to ensure that when downloading from the /refs and /compressed-blob routes, that the Http return code 404 is not logged at display level, but instead treated as an expected error code because unlike the /c route, they use 404 to indicate a cache miss.

diff --git a/Engine/Source/Developer/DerivedDataCache/Private/HttpDerivedDataBackend.cpp b/Engine/Source/Developer/DerivedDataCache/Private/HttpDerivedDataBackend.cpp
--- a/Engine/Source/Developer/DerivedDataCache/Private/HttpDerivedDataBackend.cpp
+++ b/Engine/Source/Developer/DerivedDataCache/Private/HttpDerivedDataBackend.cpp
@@ -246,7 +246,7 @@ public:
 	 * @return Result of the request
 	 */
 	template<RequestVerb V>
-	Result PerformBlockingUpload(const TCHAR* Uri, TArrayView<const uint8> Buffer)
+	Result PerformBlockingUpload(const TCHAR* Uri, TArrayView<const uint8> Buffer, TConstArrayView<long> ExpectedErrorCodes = {})
 	{
 		static_assert(V == Put || V == PutCompactBinary || V == PutCompressedBlob || V == Post || V == PostJson, "Upload should use either Put or Post verbs.");
 		
@@ -284,7 +284,7 @@ public:
 			ReadDataView = Buffer;
 		}
 
-		return PerformBlocking(Uri, V, ContentLength);
+		return PerformBlocking(Uri, V, ContentLength, ExpectedErrorCodes);
 	}
 
 	/**
@@ -294,12 +294,12 @@ public:
 	 * be stored in an internal buffer and accessed GetResponse* methods.
 	 * @return Result of the request
 	 */
-	Result PerformBlockingDownload(const TCHAR* Uri, TArray<uint8>* Buffer)
+	Result PerformBlockingDownload(const TCHAR* Uri, TArray<uint8>* Buffer, TConstArrayView<long> ExpectedErrorCodes = {400})
 	{
 		curl_easy_setopt(Curl, CURLOPT_HTTPGET, 1L);
 		WriteDataBufferPtr = Buffer;
 
-		return PerformBlocking(Uri, Get, 0u);
+		return PerformBlocking(Uri, Get, 0u, ExpectedErrorCodes);
 	}
 
 	/**
@@ -308,7 +308,7 @@ public:
 	 * @return Result of the request
 	 */
 	template<RequestVerb V>
-	Result PerformBlockingQuery(const TCHAR* Uri)
+	Result PerformBlockingQuery(const TCHAR* Uri, TConstArrayView<long> ExpectedErrorCodes = {400})
 	{
 		static_assert(V == Head || V == Delete, "Queries should use either Head or Delete verbs.");
 
@@ -321,7 +321,7 @@ public:
 			curl_easy_setopt(Curl, CURLOPT_NOBODY, 1L);
 		}
 
-		return PerformBlocking(Uri, V, 0u);
+		return PerformBlocking(Uri, V, 0u, ExpectedErrorCodes);
 	}
 
 	/**
@@ -447,7 +447,7 @@ private:
 	 * @param Buffer Optional buffer to directly receive the result of the request.
 	 * If unset the response body will be stored in the request.
 	 */
-	Result PerformBlocking(const TCHAR* Uri, RequestVerb Verb, uint32 ContentLength)
+	Result PerformBlocking(const TCHAR* Uri, RequestVerb Verb, uint32 ContentLength, TConstArrayView<long> ExpectedErrorCodes)
 	{
 		static const char* CommonHeaders[] = {
 			"User-Agent: Unreal Engine",
@@ -500,7 +500,7 @@ private:
 			bRedirected = (ResponseCode >= 300 && ResponseCode < 400);
 		}
 
-		LogResult(CurlResult, Uri, Verb);
+		LogResult(CurlResult, Uri, Verb, ExpectedErrorCodes);
 
 		// Clean up
 		curl_slist_free_all(CurlHeaders);
@@ -508,7 +508,7 @@ private:
 		return CurlResult == CURLE_OK ? Success : Failed;
 	}
 
-	void LogResult(CURLcode Result, const TCHAR* Uri, RequestVerb Verb) const
+	void LogResult(CURLcode Result, const TCHAR* Uri, RequestVerb Verb, TConstArrayView<long> ExpectedErrorCodes) const
 	{
 		if (Result == CURLE_OK)
 		{
@@ -519,28 +519,28 @@ private:
 			switch (Verb)
 			{
 			case Head:
-				bSuccess = (ResponseCode == 400 || IsSuccessResponse(ResponseCode));
+				bSuccess = (ExpectedErrorCodes.Contains(ResponseCode) || IsSuccessResponse(ResponseCode));
 				VerbStr = TEXT("querying");
 				break;
 			case Get:
-				bSuccess = (ResponseCode == 400 || IsSuccessResponse(ResponseCode));
+				bSuccess = (ExpectedErrorCodes.Contains(ResponseCode) || IsSuccessResponse(ResponseCode));
 				VerbStr = TEXT("fetching");
 				AdditionalInfo = FString::Printf(TEXT("Received: %d bytes."), BytesReceived);
 				break;
 			case Put:
 			case PutCompactBinary:
 			case PutCompressedBlob:
-				bSuccess = IsSuccessResponse(ResponseCode);
+				bSuccess = (ExpectedErrorCodes.Contains(ResponseCode) || IsSuccessResponse(ResponseCode));
 				VerbStr = TEXT("updating");
 				AdditionalInfo = FString::Printf(TEXT("Sent: %d bytes."), BytesSent);
 				break;
 			case Post:
 			case PostJson:
-				bSuccess = IsSuccessResponse(ResponseCode);
+				bSuccess = (ExpectedErrorCodes.Contains(ResponseCode) || IsSuccessResponse(ResponseCode));
 				VerbStr = TEXT("posting");
 				break;
 			case Delete:
-				bSuccess = IsSuccessResponse(ResponseCode);
+				bSuccess = (ExpectedErrorCodes.Contains(ResponseCode) || IsSuccessResponse(ResponseCode));
 				VerbStr = TEXT("deleting");
 				break;
 			}
@@ -2590,7 +2590,7 @@ FOptionalCacheRecord FHttpDerivedDataBackend::GetCacheRecordOnly(
 
 			TArray<uint8> ByteArray;
 			Request->SetHeader(TEXT("Accept"), TEXT("application/x-ue-cb"));
-			Request->PerformBlockingDownload(*RefsUri, &ByteArray);
+			Request->PerformBlockingDownload(*RefsUri, &ByteArray, {401, 404});
 			ResponseCode = Request->GetResponseCode();
 
 			if (FHttpRequest::IsSuccessResponse(ResponseCode))
@@ -2740,7 +2740,7 @@ bool FHttpDerivedDataBackend::TryGetCachedDataBatch(
 			FScopedRequestPtr Request(PutRequestPools[IsInGameThread()].Get());
 
 			TArray<uint8> ByteArray;
-			const FHttpRequest::Result Result = Request->PerformBlockingDownload(*CompressedBlobsUri, &ByteArray);
+			const FHttpRequest::Result Result = Request->PerformBlockingDownload(*CompressedBlobsUri, &ByteArray, {404});
 			ResponseCode = Request->GetResponseCode();
 
 			if (FHttpRequest::IsSuccessResponse(ResponseCode))
@@ -2792,10 +2792,10 @@ bool FHttpDerivedDataBackend::CachedDataProbablyExistsBatch(
 		{
 			FScopedRequestPtr Request(PutRequestPools[IsInGameThread()].Get());
 
-			const FHttpRequest::Result Result = Request->PerformBlockingQuery<FHttpRequest::Head>(*CompressedBlobsUri);
+			const FHttpRequest::Result Result = Request->PerformBlockingQuery<FHttpRequest::Head>(*CompressedBlobsUri, {404});
 			ResponseCode = Request->GetResponseCode();
 
-			if (FHttpRequest::IsSuccessResponse(ResponseCode) || ResponseCode == 400)
+			if (FHttpRequest::IsSuccessResponse(ResponseCode) || ResponseCode == 404)
 			{
 				bIsHit = (Result == FHttpRequest::Success && FHttpRequest::IsSuccessResponse(ResponseCode));
 				break;
@@ -3143,7 +3143,7 @@ void FHttpDerivedDataBackend::RemoveCachedData(const TCHAR* CacheKey, bool bTran
 		FScopedRequestPtr Request(PutRequestPools[IsInGameThread()].Get());
 		if (Request.IsValid())
 		{
-			FHttpRequest::Result Result = Request->PerformBlockingQuery<FHttpRequest::Delete>(*Uri);
+			FHttpRequest::Result Result = Request->PerformBlockingQuery<FHttpRequest::Delete>(*Uri, {});
 			ResponseCode = Request->GetResponseCode();
 
 			if (ResponseCode == 200)
