FCacheKey: Add operator<< to write to an FCbWriter.
FCacheRecord: Add Save function that allows appending to an existing package.

diff --git a/Engine/Source/Developer/DerivedDataCache/Private/DerivedDataCacheKey.cpp b/Engine/Source/Developer/DerivedDataCache/Private/DerivedDataCacheKey.cpp
--- a/Engine/Source/Developer/DerivedDataCache/Private/DerivedDataCacheKey.cpp
+++ b/Engine/Source/Developer/DerivedDataCache/Private/DerivedDataCacheKey.cpp
@@ -8,6 +8,7 @@
 #include "Math/UnrealMathUtility.h"
 #include "Misc/ScopeRWLock.h"
 #include "Misc/StringBuilder.h"
+#include "Serialization/CompactBinaryWriter.h"
 
 namespace UE::DerivedData::Private
 {
@@ -121,4 +122,14 @@ FCacheBucket::FCacheBucket(FWideStringView InName)
 {
 }
 
+FCbWriter& operator<<(FCbWriter& Writer, const FCacheKey& Key)
+{
+	Writer.BeginObject();
+	Writer.AddString("Bucket"_ASV, Key.Bucket.ToString());
+	Writer.AddHash("Hash"_ASV, Key.Hash);
+	Writer.EndObject();
+	return Writer;
+}
+
+
 } // UE::DerivedData
diff --git a/Engine/Source/Developer/DerivedDataCache/Private/DerivedDataCacheRecord.cpp b/Engine/Source/Developer/DerivedDataCache/Private/DerivedDataCacheRecord.cpp
--- a/Engine/Source/Developer/DerivedDataCache/Private/DerivedDataCacheRecord.cpp
+++ b/Engine/Source/Developer/DerivedDataCache/Private/DerivedDataCacheRecord.cpp
@@ -160,14 +160,15 @@ FCbPackage FCacheRecord::Save() const
 {
 	FCbPackage Package;
 	FCbWriter Writer;
+	Save(Package, Writer);
+	Package.SetObject(Writer.Save().AsObject());
+	return Package;
+}
+
+void FCacheRecord::Save(FCbPackage& Attachments, FCbWriter& Writer) const
+{
 	Writer.BeginObject();
-	{
-		const FCacheKey& Key = GetKey();
-		Writer.BeginObject("Key"_ASV);
-		Writer.AddString("Bucket"_ASV, Key.Bucket.ToString());
-		Writer.AddHash("Hash"_ASV, Key.Hash);
-		Writer.EndObject();
-	}
+	Writer << "Key"_ASV << GetKey();
 
 	if (const FCbObject& Meta = GetMeta())
 	{
@@ -181,7 +182,7 @@ FCbPackage FCacheRecord::Save() const
 		{
 			if (Value.HasData())
 			{
-				Package.AddAttachment(FCbAttachment(Value.GetData()));
+				Attachments.AddAttachment(FCbAttachment(Value.GetData()));
 			}
 			Writer.BeginObject();
 			Writer.AddObjectId("Id"_ASV, Value.GetId());
@@ -192,9 +193,6 @@ FCbPackage FCacheRecord::Save() const
 		Writer.EndArray();
 	}
 	Writer.EndObject();
-
-	Package.SetObject(Writer.Save().AsObject());
-	return Package;
 }
 
 FOptionalCacheRecord FCacheRecord::Load(const FCbPackage& Attachments, const FCbObject& Object)
diff --git a/Engine/Source/Developer/DerivedDataCache/Public/DerivedDataCacheKey.h b/Engine/Source/Developer/DerivedDataCache/Public/DerivedDataCacheKey.h
--- a/Engine/Source/Developer/DerivedDataCache/Public/DerivedDataCacheKey.h
+++ b/Engine/Source/Developer/DerivedDataCache/Public/DerivedDataCacheKey.h
@@ -10,6 +10,8 @@
 
 #define UE_API DERIVEDDATACACHE_API
 
+class FCbWriter;
+
 namespace UE::DerivedData
 {
 
@@ -121,6 +123,8 @@ inline TStringBuilderBase<CharType>& operator<<(TStringBuilderBase<CharType>& Bu
 	return Builder << Key.Bucket << CharType('/') << Key.Hash;
 }
 
+UE_API FCbWriter& operator<<(FCbWriter& Writer, const FCacheKey& Key);
+
 ///////////////////////////////////////////////////////////////////////////////////////////////////////////////////////
 
 } // UE::DerivedData
diff --git a/Engine/Source/Developer/DerivedDataCache/Public/DerivedDataCacheRecord.h b/Engine/Source/Developer/DerivedDataCache/Public/DerivedDataCacheRecord.h
--- a/Engine/Source/Developer/DerivedDataCache/Public/DerivedDataCacheRecord.h
+++ b/Engine/Source/Developer/DerivedDataCache/Public/DerivedDataCacheRecord.h
@@ -82,6 +82,8 @@ public:
 
 	/** Save the cache record to a compact binary package. */
 	UE_API FCbPackage Save() const;
+	/** Append the cache record to an existing package and writer for e.g. a batch of records. */
+	UE_API void Save(FCbPackage& Attachments, FCbWriter& Writer) const;
 
 	/** Load a cache record from a compact binary package. Null on error. */
 	UE_API static FOptionalCacheRecord Load(const FCbPackage& Package);
