Fix bug where serialization of delimeter character for CachePolicy serialization (comma ',') was being serialized as an int32 ("44") due to mixing of UTF8 builder and TCHAR character.

diff --git a/Engine/Source/Developer/DerivedDataCache/Private/DerivedDataCache.cpp b/Engine/Source/Developer/DerivedDataCache/Private/DerivedDataCache.cpp
--- a/Engine/Source/Developer/DerivedDataCache/Private/DerivedDataCache.cpp
+++ b/Engine/Source/Developer/DerivedDataCache/Private/DerivedDataCache.cpp
@@ -172,7 +172,7 @@ namespace UE::DerivedData
 namespace Private::CachePolicy
 {
 	// TODO: Implement these as Ansi String instead to maximize the most-optimal path, see ParseCachePolicyImpl
-	constexpr TCHAR			  DelimiterChar = TEXT(',');
+	constexpr ANSICHAR	  DelimiterChar = ',';
 	constexpr FStringView Delimiter = TEXTVIEW(",");
 	constexpr FStringView None = TEXTVIEW("None");
 	constexpr FStringView QueryLocal = TEXTVIEW("QueryLocal");
