DDC: Moved FBackendDebugOptions, ESpeedClass, LogDerivedDataCache to prepare to extend ILegacyCacheStore

#rnx

diff --git a/Engine/Source/Developer/DerivedDataCache/Private/DerivedDataBackendInterface.h b/Engine/Source/Developer/DerivedDataCache/Private/DerivedDataBackendInterface.h
--- a/Engine/Source/Developer/DerivedDataCache/Private/DerivedDataBackendInterface.h
+++ b/Engine/Source/Developer/DerivedDataCache/Private/DerivedDataBackendInterface.h
@@ -14,8 +14,6 @@ class FDerivedDataCacheStatsNode;
 
 namespace UE::DerivedData { struct FCacheKey; }
 
-DECLARE_LOG_CATEGORY_EXTERN(LogDerivedDataCache, Log, All);
-
 DECLARE_DWORD_ACCUMULATOR_STAT_EXTERN(TEXT("Num Gets"),STAT_DDC_NumGets,STATGROUP_DDC, );
 DECLARE_DWORD_ACCUMULATOR_STAT_EXTERN(TEXT("Num Puts"),STAT_DDC_NumPuts,STATGROUP_DDC, );
 DECLARE_DWORD_ACCUMULATOR_STAT_EXTERN(TEXT("Num Build"),STAT_DDC_NumBuilds,STATGROUP_DDC, );
@@ -29,6 +27,48 @@ DECLARE_FLOAT_ACCUMULATOR_STAT_EXTERN(TEXT("Exists Time"),STAT_DDC_ExistTime,STA
 namespace UE::DerivedData
 {
 
+/**
+	* Speed classes. Higher values are faster so > / < comparisons make sense.
+	*/
+enum class EBackendSpeedClass
+{
+	Unknown,		/* Don't know yet*/
+	Slow,			/* Slow, likely a remote drive. Some benefit but handle with care */
+	Ok,				/* Ok but not great.  */
+	Fast,			/* Fast but seek times still have an impact */
+	Local			/* Little to no impact from seek times and extremely fast reads */
+};
+
+/** Debug options that can be applied to backends to simulate different behavior */
+struct FBackendDebugOptions
+{
+	/** Percentage of requests that should result in random misses */
+	int					RandomMissRate;
+
+	/** Apply behavior of this speed class */
+	EBackendSpeedClass	SpeedClass;
+
+	/** Types of DDC entries that should always be a miss */
+	TArray<FString>		SimulateMissTypes;
+
+	FBackendDebugOptions()
+		: RandomMissRate(0)
+		, SpeedClass(EBackendSpeedClass::Unknown)
+	{
+	}
+
+	/** Fill in the provided structure based on the name of the node (e.g. 'shared') and the provided token stream */
+	static bool ParseFromTokens(FBackendDebugOptions& OutOptions, const TCHAR* InNodeName, const TCHAR* InTokens);
+
+	/**
+		* Returns true if, according to the properties of this struct, the provided key should be treated as a miss.
+		* Implementing that miss and accounting for any behavior impact (e.g. skipping a subsequent put) is left to
+		* each backend.
+		*/
+	bool ShouldSimulateMiss(const TCHAR* InCacheKey);
+	bool ShouldSimulateMiss(const UE::DerivedData::FCacheKey& InCacheKey);
+};
+
 /**
  * Interface for cache server backends.
  * The entire API should be callable from any thread (except the singleton can be assumed to be called at least once before concurrent access).
@@ -36,6 +76,8 @@ namespace UE::DerivedData
 class FDerivedDataBackendInterface : public ILegacyCacheStore
 {
 public:
+	using ESpeedClass = UE::DerivedData::EBackendSpeedClass;
+	using FBackendDebugOptions = UE::DerivedData::FBackendDebugOptions;
 
 	/** Status of a put operation. */
 	enum class EPutStatus
@@ -50,48 +92,6 @@ public:
 		Skipped,
 	};
 
-	/**
-	 * Speed classes. Higher values are faster so > / < comparisons make sense.
-	 */
-	enum class ESpeedClass
-	{
-		Unknown,		/* Don't know yet*/
-		Slow,			/* Slow, likely a remote drive. Some benefit but handle with care */
-		Ok,				/* Ok but not great.  */
-		Fast,			/* Fast but seek times still have an impact */
-		Local			/* Little to no impact from seek times and extremely fast reads */
-	};
-
-	/** Debug options that can be applied to backends to simulate different behavior */
-	struct FBackendDebugOptions
-	{
-		/** Percentage of requests that should result in random misses */
-		int					RandomMissRate;
-
-		/** Apply behavior of this speed class */
-		ESpeedClass			SpeedClass;
-
-		/** Types of DDC entries that should always be a miss */
-		TArray<FString>		SimulateMissTypes;
-
-		FBackendDebugOptions()
-			: RandomMissRate(0)
-			, SpeedClass(ESpeedClass::Unknown)
-		{
-		}
-
-		/** Fill in the provided structure based on the name of the node (e.g. 'shared') and the provided token stream */
-		static bool ParseFromTokens(FBackendDebugOptions& OutOptions, const TCHAR* InNodeName, const TCHAR* InTokens);
-
-		/**
-		 * Returns true if, according to the properties of this struct, the provided key should be treated as a miss.
-		 * Implementing that miss and accounting for any behavior impact (e.g. skipping a subsequent put) is left to
-		 * each backend.
-		 */
-		bool ShouldSimulateMiss(const TCHAR* InCacheKey);
-		bool ShouldSimulateMiss(const UE::DerivedData::FCacheKey& InCacheKey);
-	};
-
 	virtual ~FDerivedDataBackendInterface()
 	{
 	}
diff --git a/Engine/Source/Developer/DerivedDataCache/Private/DerivedDataLegacyCacheStore.h b/Engine/Source/Developer/DerivedDataCache/Private/DerivedDataLegacyCacheStore.h
--- a/Engine/Source/Developer/DerivedDataCache/Private/DerivedDataLegacyCacheStore.h
+++ b/Engine/Source/Developer/DerivedDataCache/Private/DerivedDataLegacyCacheStore.h
@@ -3,6 +3,7 @@
 #pragma once
 
 #include "DerivedDataCache.h"
+#include "Logging/LogMacros.h"
 
 namespace UE::DerivedData { class FLegacyCacheKey; }
 namespace UE::DerivedData { struct FLegacyCacheDeleteRequest; }
@@ -12,6 +13,8 @@ namespace UE::DerivedData { struct FLegacyCacheGetResponse; }
 namespace UE::DerivedData { struct FLegacyCachePutRequest; }
 namespace UE::DerivedData { struct FLegacyCachePutResponse; }
 
+DECLARE_LOG_CATEGORY_EXTERN(LogDerivedDataCache, Log, All);
+
 namespace UE::DerivedData
 {
 
