DDC: Added MakeResponse(EStatus) to the cache request types

#rnx

diff --git a/Engine/Source/Developer/DerivedDataCache/Private/DerivedDataCache.cpp b/Engine/Source/Developer/DerivedDataCache/Private/DerivedDataCache.cpp
--- a/Engine/Source/Developer/DerivedDataCache/Private/DerivedDataCache.cpp
+++ b/Engine/Source/Developer/DerivedDataCache/Private/DerivedDataCache.cpp
@@ -467,6 +467,31 @@ void ICacheStore::GetValue(
 	}
 }
 
+FCachePutResponse FCachePutRequest::MakeResponse(const EStatus Status) const
+{
+	return {Name, Record.GetKey(), UserData, Status};
+}
+
+FCacheGetResponse FCacheGetRequest::MakeResponse(const EStatus Status) const
+{
+	return {Name, FCacheRecordBuilder(Key).Build(), UserData, Status};
+}
+
+FCachePutValueResponse FCachePutValueRequest::MakeResponse(const EStatus Status) const
+{
+	return {Name, Key, UserData, Status};
+}
+
+FCacheGetValueResponse FCacheGetValueRequest::MakeResponse(const EStatus Status) const
+{
+	return {Name, Key, {}, UserData, Status};
+}
+
+FCacheGetChunkResponse FCacheGetChunkRequest::MakeResponse(const EStatus Status) const
+{
+	return {Name, Key, Id, RawOffset, 0, {}, {}, UserData, Status};
+}
+
 } // UE::DerivedData
 
 namespace UE::DerivedData::Private
diff --git a/Engine/Source/Developer/DerivedDataCache/Private/DerivedDataLegacyCacheStore.cpp b/Engine/Source/Developer/DerivedDataCache/Private/DerivedDataLegacyCacheStore.cpp
--- a/Engine/Source/Developer/DerivedDataCache/Private/DerivedDataLegacyCacheStore.cpp
+++ b/Engine/Source/Developer/DerivedDataCache/Private/DerivedDataLegacyCacheStore.cpp
@@ -83,4 +83,19 @@ void FLegacyCacheKey::WriteValueTrailer(FCompositeBuffer& Value) const
 	}
 }
 
+FLegacyCachePutResponse FLegacyCachePutRequest::MakeResponse(const EStatus Status) const
+{
+	return {Name, Key, UserData, Status};
+}
+
+FLegacyCacheGetResponse FLegacyCacheGetRequest::MakeResponse(const EStatus Status) const
+{
+	return {Name, Key, {}, UserData, Status};
+}
+
+FLegacyCacheDeleteResponse FLegacyCacheDeleteRequest::MakeResponse(const EStatus Status) const
+{
+	return {Name, Key, UserData, Status};
+}
+
 } // UE::DerivedData
diff --git a/Engine/Source/Developer/DerivedDataCache/Private/DerivedDataLegacyCacheStore.h b/Engine/Source/Developer/DerivedDataCache/Private/DerivedDataLegacyCacheStore.h
--- a/Engine/Source/Developer/DerivedDataCache/Private/DerivedDataLegacyCacheStore.h
+++ b/Engine/Source/Developer/DerivedDataCache/Private/DerivedDataLegacyCacheStore.h
@@ -68,6 +68,8 @@ struct FLegacyCachePutRequest
 	FCompositeBuffer Value;
 	ECachePolicy Policy = ECachePolicy::Default;
 	uint64 UserData = 0;
+
+	FLegacyCachePutResponse MakeResponse(EStatus Status) const;
 };
 
 struct FLegacyCachePutResponse
@@ -84,6 +86,8 @@ struct FLegacyCacheGetRequest
 	FLegacyCacheKey Key;
 	ECachePolicy Policy = ECachePolicy::Default;
 	uint64 UserData = 0;
+
+	FLegacyCacheGetResponse MakeResponse(EStatus Status) const;
 };
 
 struct FLegacyCacheGetResponse
@@ -102,6 +106,8 @@ struct FLegacyCacheDeleteRequest
 	ECachePolicy Policy = ECachePolicy::Default;
 	bool bTransient = false;
 	uint64 UserData = 0;
+
+	FLegacyCacheDeleteResponse MakeResponse(EStatus Status) const;
 };
 
 struct FLegacyCacheDeleteResponse
diff --git a/Engine/Source/Developer/DerivedDataCache/Public/DerivedDataCache.h b/Engine/Source/Developer/DerivedDataCache/Public/DerivedDataCache.h
--- a/Engine/Source/Developer/DerivedDataCache/Public/DerivedDataCache.h
+++ b/Engine/Source/Developer/DerivedDataCache/Public/DerivedDataCache.h
@@ -43,8 +43,6 @@ using FOnCachePutValueComplete = TUniqueFunction<void (FCachePutValueResponse&&
 using FOnCacheGetValueComplete = TUniqueFunction<void (FCacheGetValueResponse&& Response)>;
 using FOnCacheGetChunkComplete = TUniqueFunction<void (FCacheGetChunkResponse&& Response)>;
 
-using FOnCacheChunkComplete UE_DEPRECATED(5.0, "Use FOnCacheGetChunkComplete.") = FOnCacheGetChunkComplete;
-
 ///////////////////////////////////////////////////////////////////////////////////////////////////////////////////////
 
 /**
@@ -366,6 +364,9 @@ struct FCachePutRequest
 
 	/** A value that will be returned in the completion callback. */
 	uint64 UserData = 0;
+
+	/** Make a default response for this request, with the provided status. */
+	UE_API FCachePutResponse MakeResponse(EStatus Status) const;
 };
 
 /** Parameters for the completion callback for cache put requests. */
@@ -398,6 +399,9 @@ struct FCacheGetRequest
 
 	/** A value that will be returned in the completion callback. */
 	uint64 UserData = 0;
+
+	/** Make a default response for this request, with the provided status. */
+	UE_API FCacheGetResponse MakeResponse(EStatus Status) const;
 };
 
 /** Parameters for the completion callback for cache get requests. */
@@ -440,6 +444,9 @@ struct FCachePutValueRequest
 
 	/** A value that will be returned in the completion callback. */
 	uint64 UserData = 0;
+
+	/** Make a default response for this request, with the provided status. */
+	UE_API FCachePutValueResponse MakeResponse(EStatus Status) const;
 };
 
 /** Parameters for the completion callback for cache value put requests. */
@@ -472,6 +479,9 @@ struct FCacheGetValueRequest
 
 	/** A value that will be returned in the completion callback. */
 	uint64 UserData = 0;
+
+	/** Make a default response for this request, with the provided status. */
+	UE_API FCacheGetValueResponse MakeResponse(EStatus Status) const;
 };
 
 /** Parameters for the completion callback for cache value get requests. */
@@ -524,9 +534,10 @@ struct FCacheGetChunkRequest
 
 	/** A value that will be returned in the completion callback. */
 	uint64 UserData = 0;
-};
 
-using FCacheChunkRequest UE_DEPRECATED(5.0, "Use FCacheGetChunkRequest.") = FCacheGetChunkRequest;
+	/** Make a default response for this request, with the provided status. */
+	UE_API FCacheGetChunkResponse MakeResponse(EStatus Status) const;
+};
 
 /** Parameters for the completion callback for cache chunk requests. */
 struct FCacheGetChunkResponse
@@ -559,8 +570,6 @@ struct FCacheGetChunkResponse
 	EStatus Status = EStatus::Error;
 };
 
-using FCacheChunkResponse UE_DEPRECATED(5.0, "Use FCacheGetChunkResponse.") = FCacheGetChunkResponse;
-
 ///////////////////////////////////////////////////////////////////////////////////////////////////////////////////////
 
 /** Returns a reference to the cache. Asserts if not available. */
