DDC: Added LegacyDebugOptions and LegacyStats to ILegacyCacheStore

#rnx

diff --git a/Engine/Source/Developer/DerivedDataCache/Private/DerivedDataBackendInterface.h b/Engine/Source/Developer/DerivedDataCache/Private/DerivedDataBackendInterface.h
--- a/Engine/Source/Developer/DerivedDataCache/Private/DerivedDataBackendInterface.h
+++ b/Engine/Source/Developer/DerivedDataCache/Private/DerivedDataBackendInterface.h
@@ -206,6 +206,10 @@ public:
 		TConstArrayView<FLegacyCacheDeleteRequest> Requests,
 		IRequestOwner& Owner,
 		FOnLegacyCacheDeleteComplete&& OnComplete) override;
+
+	virtual void LegacyStats(FDerivedDataCacheStatsNode& OutNode) override;
+
+	virtual bool LegacyDebugOptions(FBackendDebugOptions& Options) override;
 };
 
 class FDerivedDataBackend
diff --git a/Engine/Source/Developer/DerivedDataCache/Private/DerivedDataBackends.cpp b/Engine/Source/Developer/DerivedDataCache/Private/DerivedDataBackends.cpp
--- a/Engine/Source/Developer/DerivedDataCache/Private/DerivedDataBackends.cpp
+++ b/Engine/Source/Developer/DerivedDataCache/Private/DerivedDataBackends.cpp
@@ -1233,12 +1233,12 @@ public:
 
 	virtual TSharedRef<FDerivedDataCacheStatsNode> GatherUsageStats() const override
 	{
+		TSharedRef<FDerivedDataCacheStatsNode> Stats = MakeShared<FDerivedDataCacheStatsNode>();
 		if (RootCache)
 		{
-			return RootCache->GatherUsageStats();
+			RootCache->LegacyStats(Stats.Get());
 		}
-
-		return MakeShared<FDerivedDataCacheStatsNode>(TEXT(""), TEXT(""), /*bIsLocal*/ true);
+		return Stats;
 	}
 
 private:
@@ -1407,6 +1407,16 @@ void FDerivedDataBackendInterface::LegacyDelete(
 	}
 }
 
+void FDerivedDataBackendInterface::LegacyStats(FDerivedDataCacheStatsNode& OutNode)
+{
+	OutNode = MoveTemp(GatherUsageStats().Get());
+}
+
+bool FDerivedDataBackendInterface::LegacyDebugOptions(FBackendDebugOptions& Options)
+{
+	return ApplyDebugOptions(Options);
+}
+
 FDerivedDataBackend& FDerivedDataBackend::Get()
 {
 	return CacheStore::Graph::FDerivedDataBackendGraph::Get();
diff --git a/Engine/Source/Developer/DerivedDataCache/Private/DerivedDataLegacyCacheStore.h b/Engine/Source/Developer/DerivedDataCache/Private/DerivedDataLegacyCacheStore.h
--- a/Engine/Source/Developer/DerivedDataCache/Private/DerivedDataLegacyCacheStore.h
+++ b/Engine/Source/Developer/DerivedDataCache/Private/DerivedDataLegacyCacheStore.h
@@ -5,7 +5,10 @@
 #include "DerivedDataCache.h"
 #include "Logging/LogMacros.h"
 
+class FDerivedDataCacheStatsNode;
+
 namespace UE::DerivedData { class FLegacyCacheKey; }
+namespace UE::DerivedData { struct FBackendDebugOptions; }
 namespace UE::DerivedData { struct FLegacyCacheDeleteRequest; }
 namespace UE::DerivedData { struct FLegacyCacheDeleteResponse; }
 namespace UE::DerivedData { struct FLegacyCacheGetRequest; }
@@ -28,7 +31,7 @@ public:
 	virtual void LegacyPut(
 		TConstArrayView<FLegacyCachePutRequest> Requests,
 		IRequestOwner& Owner,
-		FOnLegacyCachePutComplete&& OnComplete) = 0;
+		FOnLegacyCachePutComplete&& OnComplete = {}) = 0;
 
 	virtual void LegacyGet(
 		TConstArrayView<FLegacyCacheGetRequest> Requests,
@@ -39,6 +42,10 @@ public:
 		TConstArrayView<FLegacyCacheDeleteRequest> Requests,
 		IRequestOwner& Owner,
 		FOnLegacyCacheDeleteComplete&& OnComplete) = 0;
+
+	virtual void LegacyStats(FDerivedDataCacheStatsNode& OutNode) = 0;
+
+	virtual bool LegacyDebugOptions(FBackendDebugOptions& Options) = 0;
 };
 
 class FLegacyCacheKey
diff --git a/Engine/Source/Developer/DerivedDataCache/Public/DerivedDataCacheUsageStats.h b/Engine/Source/Developer/DerivedDataCache/Public/DerivedDataCacheUsageStats.h
--- a/Engine/Source/Developer/DerivedDataCache/Public/DerivedDataCacheUsageStats.h
+++ b/Engine/Source/Developer/DerivedDataCache/Public/DerivedDataCacheUsageStats.h
@@ -109,6 +109,8 @@ public:
 class FDerivedDataCacheStatsNode : public TSharedFromThis<FDerivedDataCacheStatsNode>
 {
 public:
+	FDerivedDataCacheStatsNode() = default;
+
 	FDerivedDataCacheStatsNode(const FString& InCacheType, const FString& InCacheName, bool bInIsLocal)
 		: CacheType(InCacheType)
 		, CacheName(InCacheName)
