DDC: Made ICacheStore private and removed ICache::CancelAll

#rnx

diff --git a/Engine/Source/Developer/DerivedDataCache/Private/DerivedDataCache.cpp b/Engine/Source/Developer/DerivedDataCache/Private/DerivedDataCache.cpp
--- a/Engine/Source/Developer/DerivedDataCache/Private/DerivedDataCache.cpp
+++ b/Engine/Source/Developer/DerivedDataCache/Private/DerivedDataCache.cpp
@@ -1120,7 +1120,7 @@ private:
 	FOnDDCNotification DDCNotificationEvent;
 
 public:
-	// ICacheStore Interface
+	// ICache Interface
 
 	void Put(
 		TConstArrayView<FCachePutRequest> Requests,
@@ -1162,12 +1162,6 @@ public:
 		return FDerivedDataBackend::Get().GetRoot().GetChunks(Requests, Owner, OnComplete ? MoveTemp(OnComplete) : [](auto&&){});
 	}
 
-	// ICache Interface
-
-	void CancelAll() final
-	{
-	}
-
 	ICacheStoreMaintainer& GetMaintainer() final
 	{
 		return *this;
diff --git a/Engine/Source/Developer/DerivedDataCache/Private/DerivedDataCacheStore.h b/Engine/Source/Developer/DerivedDataCache/Private/DerivedDataCacheStore.h
new file mode 100644
--- /dev/null
+++ b/Engine/Source/Developer/DerivedDataCache/Private/DerivedDataCacheStore.h
@@ -0,0 +1,92 @@
+// Copyright Epic Games, Inc. All Rights Reserved.
+
+#pragma once
+
+#include "DerivedDataCache.h"
+
+namespace UE::DerivedData
+{
+
+/**
+ * Interface to a store of cache records and cache values.
+ *
+ * Functions on this interface may be called from any thread.
+ */
+class ICacheStore
+{
+public:
+	virtual ~ICacheStore() = default;
+
+	/**
+	 * Asynchronous request to put records in the cache.
+	 *
+	 * Behavior must match the cache record policy in each request. A few of those requirements:
+	 * - Store with Query is expected *not* to overwrite an existing record.
+	 * - Store without Query is permitted to overwrite an existing record.
+	 * - Response must be Error if any value has unrecoverable missing data, unless the policy has PartialRecord.
+	 *
+	 * @param Requests     Requests with the cache records to store. Records must have a key.
+	 * @param Owner        The owner to execute the request within. See IRequestOwner.
+	 * @param OnComplete   A callback invoked for every request as it completes or is canceled.
+	 */
+	virtual void Put(
+		TConstArrayView<FCachePutRequest> Requests,
+		IRequestOwner& Owner,
+		FOnCachePutComplete&& OnComplete = {}) = 0;
+
+	/**
+	 * Asynchronous request to get records from the cache.
+	 *
+	 * Behavior must match the cache record policy in each request. A few of those requirements:
+	 * - Data is required if the value policy has Query.
+	 * - Data must not be in the response if the value policy has SkipData.
+	 * - Response must contain values and available data if the record policy has PartialRecord.
+	 * - Response must be Error if any required data is missing, even if the record policy has PartialRecord.
+	 *
+	 * @param Requests     Requests with the keys of the cache records to fetch.
+	 * @param Owner        The owner to execute the request within. See IRequestOwner.
+	 * @param OnComplete   A callback invoked for every request as it completes or is canceled.
+	 */
+	virtual void Get(
+		TConstArrayView<FCacheGetRequest> Requests,
+		IRequestOwner& Owner,
+		FOnCacheGetComplete&& OnComplete) = 0;
+
+	/**
+	 * Asynchronous request to put values in the cache.
+	 *
+	 * @param Requests     Requests with the cache values to store. Requests must have a key.
+	 * @param Owner        The owner to execute the request within. See IRequestOwner.
+	 * @param OnComplete   A callback invoked for every request as it completes or is canceled.
+	 */
+	virtual void PutValue(
+		TConstArrayView<FCachePutValueRequest> Requests,
+		IRequestOwner& Owner,
+		FOnCachePutValueComplete&& OnComplete = {});
+
+	/**
+	 * Asynchronous request to get values from the cache.
+	 *
+	 * @param Requests     Requests with the keys of the cache values to fetch.
+	 * @param Owner        The owner to execute the request within. See IRequestOwner.
+	 * @param OnComplete   A callback invoked for every request as it completes or is canceled.
+	 */
+	virtual void GetValue(
+		TConstArrayView<FCacheGetValueRequest> Requests,
+		IRequestOwner& Owner,
+		FOnCacheGetValueComplete&& OnComplete);
+
+	/**
+	 * Asynchronous request to get chunks, which are subsets of values, from records or values.
+	 *
+	 * @param Requests     Requests with the key, ID, offset, and size of each chunk to fetch.
+	 * @param Owner        The owner to execute the request within. See IRequestOwner.
+	 * @param OnComplete   A callback invoked for every request as it completes or is canceled.
+	 */
+	virtual void GetChunks(
+		TConstArrayView<FCacheGetChunkRequest> Requests,
+		IRequestOwner& Owner,
+		FOnCacheGetChunkComplete&& OnComplete) = 0;
+};
+
+} // UE::DerivedData
diff --git a/Engine/Source/Developer/DerivedDataCache/Private/DerivedDataLegacyCacheStore.h b/Engine/Source/Developer/DerivedDataCache/Private/DerivedDataLegacyCacheStore.h
--- a/Engine/Source/Developer/DerivedDataCache/Private/DerivedDataLegacyCacheStore.h
+++ b/Engine/Source/Developer/DerivedDataCache/Private/DerivedDataLegacyCacheStore.h
@@ -3,6 +3,7 @@
 #pragma once
 
 #include "DerivedDataCache.h"
+#include "DerivedDataCacheStore.h"
 #include "Logging/LogMacros.h"
 
 class FDerivedDataCacheStatsNode;
diff --git a/Engine/Source/Developer/DerivedDataCache/Public/DerivedDataCache.h b/Engine/Source/Developer/DerivedDataCache/Public/DerivedDataCache.h
--- a/Engine/Source/Developer/DerivedDataCache/Public/DerivedDataCache.h
+++ b/Engine/Source/Developer/DerivedDataCache/Public/DerivedDataCache.h
@@ -238,111 +238,89 @@ private:
 ///////////////////////////////////////////////////////////////////////////////////////////////////////////////////////
 
 /**
- * Interface to a store of cache records and cache values.
+ * Interface to the cache.
+ *
+ * Functions on this interface may be called from any thread.
  *
- * Functions on this interface may be called from any thread. When a callback is provided, it may
- * be invoked from a different thread than the request was made on, and may be invoked before the
- * request function returns, both of which must be supported by the caller and the callback.
+ * Requests may complete out of order relative to the order that they were requested.
+ *
+ * Callbacks may be called from any thread, including the calling thread, may be called from more
+ * than one thread concurrently, and may be called before returning from the request function.
  */
-class ICacheStore
+class ICache
 {
 public:
-	virtual ~ICacheStore() = default;
+	virtual ~ICache() = default;
 
 	/**
-	 * Asynchronous request to put cache records according to the policy.
+	 * Asynchronous request to put records in the cache.
 	 *
-	 * The callback will always be called for every key, and may be called from an arbitrary thread.
-	 * Records may finish storing in any order, and from multiple threads concurrently.
-	 *
-	 * A cache store is free to interpret a record containing only a key as a request to delete that
-	 * record from the store. Records may contain values that do not have data, and such values must
-	 * reference an existing value in the store, if available. The PartialRecord policy may be used
-	 * to request that a partial record be stored when a value is missing.
+	 * @see FCachePutRequest
 	 *
 	 * @param Requests     Requests with the cache records to store. Records must have a key.
-	 * @param Owner        The owner to execute the request within.
-	 * @param OnComplete   A callback invoked for every record as it completes or is canceled.
+	 * @param Owner        The owner to execute the request within. See IRequestOwner.
+	 * @param OnComplete   A callback invoked for every request as it completes or is canceled.
 	 */
 	virtual void Put(
 		TConstArrayView<FCachePutRequest> Requests,
 		IRequestOwner& Owner,
-		FOnCachePutComplete&& OnComplete = FOnCachePutComplete()) = 0;
+		FOnCachePutComplete&& OnComplete = {}) = 0;
 
 	/**
-	 * Asynchronous request to get cache records according to the policy.
-	 *
-	 * The callback will always be called for every key, and may be called from an arbitrary thread.
-	 * Records may become available in any order, and from multiple threads concurrently.
+	 * Asynchronous request to get records from the cache.
 	 *
-	 * Records may propagate into other cache stores, in accordance with the policy. A propagated
-	 * record may be a partial record, containing values without data, depending on the policy. A
-	 * propagated put of a partial record will include the PartialRecord policy.
+	 * @see FCacheGetRequest
 	 *
-	 * When values are required by the policy but not available, the status must be Error. If the
-	 * PartialRecord policy applies to the missing values, the cache store must return a partial
-	 * record containing the available values.
-	 *
-	 * @param Requests     Requests with the keys identifying the cache records to fetch.
-	 * @param Owner        The owner to execute the request within.
-	 * @param OnComplete   A callback invoked for every key as it completes or is canceled.
+	 * @param Requests     Requests with the keys of the cache records to fetch.
+	 * @param Owner        The owner to execute the request within. See IRequestOwner.
+	 * @param OnComplete   A callback invoked for every request as it completes or is canceled.
 	 */
 	virtual void Get(
 		TConstArrayView<FCacheGetRequest> Requests,
 		IRequestOwner& Owner,
 		FOnCacheGetComplete&& OnComplete) = 0;
 
-	UE_API virtual void PutValue(
+	/**
+	 * Asynchronous request to put values in the cache.
+	 *
+	 * @see FCachePutValueRequest
+	 *
+	 * @param Requests     Requests with the cache values to store. Requests must have a key.
+	 * @param Owner        The owner to execute the request within. See IRequestOwner.
+	 * @param OnComplete   A callback invoked for every request as it completes or is canceled.
+	 */
+	virtual void PutValue(
 		TConstArrayView<FCachePutValueRequest> Requests,
 		IRequestOwner& Owner,
-		FOnCachePutValueComplete&& OnComplete = FOnCachePutValueComplete());
+		FOnCachePutValueComplete&& OnComplete = {}) = 0;
 
-	UE_API virtual void GetValue(
+	/**
+	 * Asynchronous request to get values from the cache.
+	 *
+	 * @see FCacheGetValueRequest
+	 *
+	 * @param Requests     Requests with the keys of the cache values to fetch.
+	 * @param Owner        The owner to execute the request within. See IRequestOwner.
+	 * @param OnComplete   A callback invoked for every request as it completes or is canceled.
+	 */
+	virtual void GetValue(
 		TConstArrayView<FCacheGetValueRequest> Requests,
 		IRequestOwner& Owner,
-		FOnCacheGetValueComplete&& OnComplete);
+		FOnCacheGetValueComplete&& OnComplete) = 0;
 
 	/**
 	 * Asynchronous request to get chunks, which are subsets of values, from records or values.
 	 *
-	 * The callback will always be called for every chunk, and may be called from an arbitrary thread.
-	 * Chunks may become available in any order, and from multiple threads concurrently.
-	 *
-	 * Values may propagate into other cache stores, in accordance with the policy, if the entire
-	 * value is requested, through a put of a partial record or a value. The cache store does not
-	 * need to verify the existence of unrequested values within requested records.
+	 * @see FCacheGetChunkRequest
 	 *
-	 * @param Requests     The keys, IDs, offsets, and sizes of the chunks to fetch.
-	 * @param Owner        The owner to execute the request within.
-	 * @param OnComplete   A callback invoked for every chunk as it completes or is canceled.
+	 * @param Requests     Requests with the key, ID, offset, and size of each chunk to fetch.
+	 * @param Owner        The owner to execute the request within. See IRequestOwner.
+	 * @param OnComplete   A callback invoked for every request as it completes or is canceled.
 	 */
 	virtual void GetChunks(
 		TConstArrayView<FCacheGetChunkRequest> Requests,
 		IRequestOwner& Owner,
 		FOnCacheGetChunkComplete&& OnComplete) = 0;
-};
-
-///////////////////////////////////////////////////////////////////////////////////////////////////////////////////////
-
-/**
- * Interface to the cache.
- *
- * @see ICacheStore for cache record and cache value storage.
- */
-class ICache : public ICacheStore
-{
-public:
-	virtual ~ICache() = default;
-
-	/**
-	 * Cancel all queued and active cache requests and invoke their callbacks.
-	 *
-	 * This is meant to be called before destroying the cache. Any new requests that are made during
-	 * execution of this function will also be canceled, because callbacks might queue requests when
-	 * they are invoked. This does not block any future cache requests, which must be handled by the
-	 * owner of the cache if future requests are meant to be avoided.
-	 */
-	virtual void CancelAll() = 0;
 
 	/** Returns the interface to the background cache store maintenance. */
 	virtual ICacheStoreMaintainer& GetMaintainer() = 0;
