DDC: Fixed TryToPrefetch when there is not one backend containing every key

#rnx

diff --git a/Engine/Source/Developer/DerivedDataCache/Private/DerivedDataBackendAsyncPutWrapper.cpp b/Engine/Source/Developer/DerivedDataCache/Private/DerivedDataBackendAsyncPutWrapper.cpp
--- a/Engine/Source/Developer/DerivedDataCache/Private/DerivedDataBackendAsyncPutWrapper.cpp
+++ b/Engine/Source/Developer/DerivedDataCache/Private/DerivedDataBackendAsyncPutWrapper.cpp
@@ -112,7 +112,7 @@ public:
 	 * @param	CacheKeys	Alphanumeric+underscore keys of the cache items
 	 * @return				true if the data will probably be found in a fast backend on a future request.
 	 */
-	virtual bool TryToPrefetch(TConstArrayView<FString> CacheKeys) override;
+	virtual TBitArray<> TryToPrefetch(TConstArrayView<FString> CacheKeys) override;
 
 	/**
 	 * Allows the DDC backend to determine if it wants to cache the provided data. Reasons for returning false could be a slow connection,
@@ -445,18 +445,23 @@ TBitArray<> FDerivedDataBackendAsyncPutWrapper::CachedDataProbablyExistsBatch(TC
 	return Result;
 }
 
-bool FDerivedDataBackendAsyncPutWrapper::TryToPrefetch(TConstArrayView<FString> CacheKeys)
+TBitArray<> FDerivedDataBackendAsyncPutWrapper::TryToPrefetch(TConstArrayView<FString> CacheKeys)
 {
 	COOK_STAT(auto Timer = UsageStats.TimePrefetch());
 
-	if ((InflightCache && InflightCache->CachedDataProbablyExistsBatch(CacheKeys).CountSetBits() == CacheKeys.Num()) ||
-		InnerBackend->TryToPrefetch(CacheKeys))
+	TBitArray<> Result;
+	if (InflightCache)
+	{
+		Result = InflightCache->CachedDataProbablyExistsBatch(CacheKeys);
+	}
+	Result = TBitArray<>::BitwiseOR(Result, InnerBackend->TryToPrefetch(CacheKeys), EBitwiseOperatorFlags::MaxSize);
+
+	if (Result.CountSetBits() == CacheKeys.Num())
 	{
 		COOK_STAT(Timer.AddHit(0));
-		return true;
 	}
 
-	return false;
+	return Result;
 }
 
 /*
diff --git a/Engine/Source/Developer/DerivedDataCache/Private/DerivedDataBackendInterface.h b/Engine/Source/Developer/DerivedDataCache/Private/DerivedDataBackendInterface.h
--- a/Engine/Source/Developer/DerivedDataCache/Private/DerivedDataBackendInterface.h
+++ b/Engine/Source/Developer/DerivedDataCache/Private/DerivedDataBackendInterface.h
@@ -176,9 +176,9 @@ public:
 	 * Synchronous attempt to make sure the cached data will be available as optimally as possible.
 	 *
 	 * @param	CacheKeys	Alphanumeric+underscore keys of the cache items
-	 * @return				true if the data will probably be found in a fast backend on a future request.
+	 * @return				A bit array with bits indicating whether the data for the corresponding key will probably be found in a fast backend on a future request.
 	 */
-	virtual bool TryToPrefetch(TConstArrayView<FString> CacheKeys) = 0;
+	virtual TBitArray<> TryToPrefetch(TConstArrayView<FString> CacheKeys) = 0;
 
 	/**
 	 * Allows the DDC backend to determine if it wants to cache the provided data. Reasons for returning false could be a slow connection,
diff --git a/Engine/Source/Developer/DerivedDataCache/Private/DerivedDataBackendThrottleWrapper.h b/Engine/Source/Developer/DerivedDataCache/Private/DerivedDataBackendThrottleWrapper.h
--- a/Engine/Source/Developer/DerivedDataCache/Private/DerivedDataBackendThrottleWrapper.h
+++ b/Engine/Source/Developer/DerivedDataCache/Private/DerivedDataBackendThrottleWrapper.h
@@ -148,7 +148,7 @@ public:
 		return Usage;
 	}
 
-	virtual bool TryToPrefetch(TConstArrayView<FString> CacheKeys) override
+	virtual TBitArray<> TryToPrefetch(TConstArrayView<FString> CacheKeys) override
 	{
 		FThrottlingScope Scope(this);
 
diff --git a/Engine/Source/Developer/DerivedDataCache/Private/DerivedDataBackendVerifyWrapper.h b/Engine/Source/Developer/DerivedDataCache/Private/DerivedDataBackendVerifyWrapper.h
--- a/Engine/Source/Developer/DerivedDataCache/Private/DerivedDataBackendVerifyWrapper.h
+++ b/Engine/Source/Developer/DerivedDataCache/Private/DerivedDataBackendVerifyWrapper.h
@@ -143,7 +143,7 @@ public:
 		return Usage;
 	}
 
-	virtual bool TryToPrefetch(TConstArrayView<FString> CacheKeys) override
+	virtual TBitArray<> TryToPrefetch(TConstArrayView<FString> CacheKeys) override
 	{
 		return InnerBackend->TryToPrefetch(CacheKeys);
 	}
diff --git a/Engine/Source/Developer/DerivedDataCache/Private/DerivedDataBackends.cpp b/Engine/Source/Developer/DerivedDataCache/Private/DerivedDataBackends.cpp
--- a/Engine/Source/Developer/DerivedDataCache/Private/DerivedDataBackends.cpp
+++ b/Engine/Source/Developer/DerivedDataCache/Private/DerivedDataBackends.cpp
@@ -1376,10 +1376,12 @@ void FDerivedDataBackendInterface::LegacyGet(
 
 	if (!PrefetchKeys.IsEmpty())
 	{
-		const bool bOk = TryToPrefetch(PrefetchKeys);
+		int32 Index = 0;
+		const TBitArray<> Exists = TryToPrefetch(PrefetchKeys);
 		for (const FLegacyCacheGetRequest* const Request : PrefetchRequests)
 		{
-			OnComplete({Request->Name, Request->Key, {}, Request->UserData, bOk ? EStatus::Ok : EStatus::Error});
+			OnComplete({Request->Name, Request->Key, {}, Request->UserData, Exists[Index] ? EStatus::Ok : EStatus::Error});
+			++Index;
 		}
 	}
 
diff --git a/Engine/Source/Developer/DerivedDataCache/Private/FileSystemDerivedDataBackend.cpp b/Engine/Source/Developer/DerivedDataCache/Private/FileSystemDerivedDataBackend.cpp
--- a/Engine/Source/Developer/DerivedDataCache/Private/FileSystemDerivedDataBackend.cpp
+++ b/Engine/Source/Developer/DerivedDataCache/Private/FileSystemDerivedDataBackend.cpp
@@ -747,7 +747,7 @@ public:
 	EPutStatus PutCachedData(const TCHAR* CacheKey, TConstArrayView<uint8> Data, bool bPutEvenIfExists) final;
 	void RemoveCachedData(const TCHAR* CacheKey, bool bTransient) final;
 	TSharedRef<FDerivedDataCacheStatsNode> GatherUsageStats() const final;
-	bool TryToPrefetch(TConstArrayView<FString> CacheKeys) final;
+	TBitArray<> TryToPrefetch(TConstArrayView<FString> CacheKeys) final;
 	bool WouldCache(const TCHAR* CacheKey, TConstArrayView<uint8> InData) final;
 	bool ApplyDebugOptions(FBackendDebugOptions& InOptions) final;
 
@@ -1533,9 +1533,9 @@ TSharedRef<FDerivedDataCacheStatsNode> FFileSystemCacheStore::GatherUsageStats()
 	return Usage;
 }
 
-bool FFileSystemCacheStore::TryToPrefetch(const TConstArrayView<FString> CacheKeys)
+TBitArray<> FFileSystemCacheStore::TryToPrefetch(const TConstArrayView<FString> CacheKeys)
 {
-	return CachedDataProbablyExistsBatch(CacheKeys).CountSetBits() == CacheKeys.Num();
+	return CachedDataProbablyExistsBatch(CacheKeys);
 }
 
 bool FFileSystemCacheStore::ApplyDebugOptions(FBackendDebugOptions& InOptions)
diff --git a/Engine/Source/Developer/DerivedDataCache/Private/HierarchicalDerivedDataBackend.h b/Engine/Source/Developer/DerivedDataCache/Private/HierarchicalDerivedDataBackend.h
--- a/Engine/Source/Developer/DerivedDataCache/Private/HierarchicalDerivedDataBackend.h
+++ b/Engine/Source/Developer/DerivedDataCache/Private/HierarchicalDerivedDataBackend.h
@@ -202,11 +202,19 @@ public:
 	 * @param	CacheKeys	Alphanumeric+underscore keys of the cache items
 	 * @return				true if the data will probably be found in a fast backend on a future request.
 	 */
-	virtual bool TryToPrefetch(TConstArrayView<FString> CacheKeys) override
+	virtual TBitArray<> TryToPrefetch(TConstArrayView<FString> CacheKeys) override
 	{
 		COOK_STAT(auto Timer = UsageStats.TimePrefetch());
 
 		TArray<FString, TInlineAllocator<16>> SearchKeys(CacheKeys);
+		TArray<int32, TInlineAllocator<16>> SearchKeysIndex;
+		SearchKeysIndex.Reserve(SearchKeys.Num());
+		for (int32 Index = 0; Index < SearchKeys.Num(); ++Index)
+		{
+			SearchKeysIndex.Add(Index);
+		}
+
+		TBitArray<> Result(false, SearchKeys.Num());
 
 		bool bHasFastBackendToWrite = false;
 		bool bHasSlowBackend = false;
@@ -232,6 +240,8 @@ public:
 					{
 						if (Hits[It.GetIndex()])
 						{
+							Result[SearchKeysIndex[It.GetIndex()]] = true;
+							SearchKeysIndex.RemoveAt(It.GetIndex());
 							// RemoveCurrent will decrement the iterator
 							It.RemoveCurrent();
 						}
@@ -245,7 +255,7 @@ public:
 					if (SearchKeys.IsEmpty())
 					{
 						COOK_STAT(Timer.AddHit(0));
-						return true;
+						return Result;
 					}
 				}
 			}
@@ -256,16 +266,26 @@ public:
 		int64 BytesFetched = 0;
 		if (bHasSlowBackend && bHasFastBackendToWrite)
 		{
+			int32 Index = 0;
 			for (const FString& CacheKey : SearchKeys)
 			{
 				TArray<uint8> Ignored;
-				bHit &= GetCachedData(*CacheKey, Ignored);
+				Result[SearchKeysIndex[Index]] = GetCachedData(*CacheKey, Ignored);
 				BytesFetched += Ignored.Num();
+				++Index;
 			}
 		}
 
-		COOK_STAT(if (bHit) { Timer.AddHit(BytesFetched); });
-		return bHit;
+		if (Result.CountSetBits() == CacheKeys.Num())
+		{
+			COOK_STAT(Timer.AddHit(BytesFetched));
+		}
+		else
+		{
+			COOK_STAT(Timer.AddMiss(BytesFetched));
+		}
+
+		return Result;
 	}
 
 	/*
diff --git a/Engine/Source/Developer/DerivedDataCache/Private/HttpDerivedDataBackend.cpp b/Engine/Source/Developer/DerivedDataCache/Private/HttpDerivedDataBackend.cpp
--- a/Engine/Source/Developer/DerivedDataCache/Private/HttpDerivedDataBackend.cpp
+++ b/Engine/Source/Developer/DerivedDataCache/Private/HttpDerivedDataBackend.cpp
@@ -2013,7 +2013,7 @@ public:
 	virtual TSharedRef<FDerivedDataCacheStatsNode> GatherUsageStats() const override;
 	
 	virtual FString GetName() const override;
-	virtual bool TryToPrefetch(TConstArrayView<FString> CacheKeys) override;
+	virtual TBitArray<> TryToPrefetch(TConstArrayView<FString> CacheKeys) override;
 	virtual bool WouldCache(const TCHAR* CacheKey, TArrayView<const uint8> InData) override;
 	virtual ESpeedClass GetSpeedClass() const override;
 	virtual bool ApplyDebugOptions(FBackendDebugOptions& InOptions) override;
@@ -2151,9 +2151,9 @@ FString FHttpDerivedDataBackend::GetName() const
 	return Domain;
 }
 
-bool FHttpDerivedDataBackend::TryToPrefetch(TConstArrayView<FString> CacheKeys)
+TBitArray<> FHttpDerivedDataBackend::TryToPrefetch(TConstArrayView<FString> CacheKeys)
 {
-	return CachedDataProbablyExistsBatch(CacheKeys).CountSetBits() == CacheKeys.Num();
+	return CachedDataProbablyExistsBatch(CacheKeys);
 }
 
 bool FHttpDerivedDataBackend::WouldCache(const TCHAR* CacheKey, TArrayView<const uint8> InData)
diff --git a/Engine/Source/Developer/DerivedDataCache/Private/MemoryDerivedDataBackend.cpp b/Engine/Source/Developer/DerivedDataCache/Private/MemoryDerivedDataBackend.cpp
--- a/Engine/Source/Developer/DerivedDataCache/Private/MemoryDerivedDataBackend.cpp
+++ b/Engine/Source/Developer/DerivedDataCache/Private/MemoryDerivedDataBackend.cpp
@@ -84,7 +84,7 @@ public:
 
 	virtual TSharedRef<FDerivedDataCacheStatsNode> GatherUsageStats() const override;
 
-	virtual bool TryToPrefetch(TConstArrayView<FString> CacheKeys) override;
+	virtual TBitArray<> TryToPrefetch(TConstArrayView<FString> CacheKeys) override;
 
 	/**
 	 *  Determines if we would cache the provided data
@@ -291,9 +291,9 @@ bool FMemoryDerivedDataBackend::GetCachedData(const TCHAR* CacheKey, TArray<uint
 	return false;
 }
 
-bool FMemoryDerivedDataBackend::TryToPrefetch(TConstArrayView<FString> CacheKeys)
+TBitArray<> FMemoryDerivedDataBackend::TryToPrefetch(TConstArrayView<FString> CacheKeys)
 {
-	return CachedDataProbablyExistsBatch(CacheKeys).CountSetBits() == CacheKeys.Num();
+	return CachedDataProbablyExistsBatch(CacheKeys);
 }
 
 bool FMemoryDerivedDataBackend::WouldCache(const TCHAR* CacheKey, TArrayView<const uint8> InData)
diff --git a/Engine/Source/Developer/DerivedDataCache/Private/PakFileDerivedDataBackend.cpp b/Engine/Source/Developer/DerivedDataCache/Private/PakFileDerivedDataBackend.cpp
--- a/Engine/Source/Developer/DerivedDataCache/Private/PakFileDerivedDataBackend.cpp
+++ b/Engine/Source/Developer/DerivedDataCache/Private/PakFileDerivedDataBackend.cpp
@@ -106,9 +106,9 @@ public:
 
 	virtual TSharedRef<FDerivedDataCacheStatsNode> GatherUsageStats() const override;
 
-	virtual bool TryToPrefetch(TConstArrayView<FString> CacheKeys) override
+	virtual TBitArray<> TryToPrefetch(TConstArrayView<FString> CacheKeys) override
 	{
-		return CachedDataProbablyExistsBatch(CacheKeys).CountSetBits() == CacheKeys.Num();
+		return CachedDataProbablyExistsBatch(CacheKeys);
 	}
 
 	virtual bool WouldCache(const TCHAR* CacheKey, TArrayView<const uint8> InData) override { return true; }
diff --git a/Engine/Source/Developer/DerivedDataCache/Private/S3DerivedDataBackend.cpp b/Engine/Source/Developer/DerivedDataCache/Private/S3DerivedDataBackend.cpp
--- a/Engine/Source/Developer/DerivedDataCache/Private/S3DerivedDataBackend.cpp
+++ b/Engine/Source/Developer/DerivedDataCache/Private/S3DerivedDataBackend.cpp
@@ -226,7 +226,7 @@ public:
 
 	FString GetName() const override;
 	ESpeedClass GetSpeedClass() const override;
-	bool TryToPrefetch(TConstArrayView<FString> CacheKeys) override;
+	TBitArray<> TryToPrefetch(TConstArrayView<FString> CacheKeys) override;
 	bool WouldCache(const TCHAR* CacheKey, TArrayView<const uint8> InData) override;
 
 	bool ApplyDebugOptions(FBackendDebugOptions& InOptions) override;
@@ -1096,9 +1096,9 @@ FDerivedDataBackendInterface::ESpeedClass FS3DerivedDataBackend::GetSpeedClass()
 	return ESpeedClass::Local;
 }
 
-bool FS3DerivedDataBackend::TryToPrefetch(TConstArrayView<FString> CacheKeys)
+TBitArray<> FS3DerivedDataBackend::TryToPrefetch(TConstArrayView<FString> CacheKeys)
 {
-	return CachedDataProbablyExistsBatch(CacheKeys).CountSetBits() == CacheKeys.Num();
+	return CachedDataProbablyExistsBatch(CacheKeys);
 }
 
 bool FS3DerivedDataBackend::WouldCache(const TCHAR* CacheKey, TArrayView<const uint8> InData)
diff --git a/Engine/Source/Developer/DerivedDataCache/Private/ZenDerivedDataBackend.cpp b/Engine/Source/Developer/DerivedDataCache/Private/ZenDerivedDataBackend.cpp
--- a/Engine/Source/Developer/DerivedDataCache/Private/ZenDerivedDataBackend.cpp
+++ b/Engine/Source/Developer/DerivedDataCache/Private/ZenDerivedDataBackend.cpp
@@ -99,7 +99,7 @@ public:
 	 * @param	CacheKeys	Alphanumeric+underscore keys of the cache items
 	 * @return				true if the data will probably be found in a fast backend on a future request.
 	 */
-	virtual bool TryToPrefetch(TConstArrayView<FString> CacheKeys) override;
+	virtual TBitArray<> TryToPrefetch(TConstArrayView<FString> CacheKeys) override;
 
 	virtual bool CachedDataProbablyExists(const TCHAR* CacheKey) override;
 	virtual bool GetCachedData(const TCHAR* CacheKey, TArray<uint8>& OutData) override;
@@ -587,9 +587,9 @@ TSharedRef<FDerivedDataCacheStatsNode> FZenDerivedDataBackend::GatherUsageStats(
 	return GroupNode;
 }
 
-bool FZenDerivedDataBackend::TryToPrefetch(TConstArrayView<FString> CacheKeys)
+TBitArray<> FZenDerivedDataBackend::TryToPrefetch(TConstArrayView<FString> CacheKeys)
 {
-	return CachedDataProbablyExistsBatch(CacheKeys).CountSetBits() == CacheKeys.Num();
+	return CachedDataProbablyExistsBatch(CacheKeys);
 }
 
 bool FZenDerivedDataBackend::WouldCache(const TCHAR* CacheKey, TArrayView<const uint8> InData)
