DDC: Added a new async cache hierarchy type

This new cache hierarchy does not require the cache stores within it to complete every operation synchronously. When the child cache stores start to do async I/O, this hierarchy will support that instead of blocking like the existing hierarchy does.

This is a snapshot of the work in progress that is only implemented for the functions in ILegacyCacheStore.

#rnx

diff --git a/Engine/Source/Developer/DerivedDataCache/Private/DerivedDataCacheStore.h b/Engine/Source/Developer/DerivedDataCache/Private/DerivedDataCacheStore.h
--- a/Engine/Source/Developer/DerivedDataCache/Private/DerivedDataCacheStore.h
+++ b/Engine/Source/Developer/DerivedDataCache/Private/DerivedDataCacheStore.h
@@ -4,6 +4,9 @@
 
 #include "DerivedDataCache.h"
 
+namespace UE::DerivedData { class ILegacyCacheStore; }
+namespace UE::DerivedData { enum class EStatus : uint8; }
+
 namespace UE::DerivedData
 {
 
@@ -89,4 +92,41 @@ public:
 		FOnCacheGetChunkComplete&& OnComplete) = 0;
 };
 
+enum class ECacheStoreFlags : uint32
+{
+	None            = 0,
+
+	/** Accepts requests with a policy of QueryLocal or StoreLocal. Needs matching Query/Store flag. */
+	Local           = 1 << 0,
+	/** Accepts requests with a policy of QueryRemote or StoreRemote. Needs matching Query/Store flag.*/
+	Remote          = 1 << 1,
+
+	/** Accepts requests with a policy of QueryLocal or QueryRemote. Needs matching Local/Remote flag. */
+	Query           = 1 << 2,
+	/** Accepts requests with a policy of StoreLocal or StoreRemote. Needs matching Local/Remote flag. */
+	Store           = 1 << 3,
+
+	/** Requests to store records or values contained by this cache store will not propagate any further. */
+	StopStore       = 1 << 4,
+};
+
+ENUM_CLASS_FLAGS(ECacheStoreFlags);
+
+class ICacheStoreOwner
+{
+public:
+	virtual void Add(ILegacyCacheStore* CacheStore, ECacheStoreFlags Flags) = 0;
+	virtual void SetFlags(ILegacyCacheStore* CacheStore, ECacheStoreFlags Flags) = 0;
+	virtual void RemoveNotSafe(ILegacyCacheStore* CacheStore) = 0;
+};
+
+template <typename RequestRangeType, typename OnCompleteType>
+inline void CompleteWithStatus(RequestRangeType&& Requests, OnCompleteType&& OnComplete, EStatus Status)
+{
+	for (const auto& Request : Requests)
+	{
+		OnComplete(Request.MakeResponse(Status));
+	}
+}
+
 } // UE::DerivedData
diff --git a/Engine/Source/Developer/DerivedDataCache/Private/DerivedDataCacheStoreHierarchy.cpp b/Engine/Source/Developer/DerivedDataCache/Private/DerivedDataCacheStoreHierarchy.cpp
new file mode 100644
--- /dev/null
+++ b/Engine/Source/Developer/DerivedDataCache/Private/DerivedDataCacheStoreHierarchy.cpp
@@ -0,0 +1,883 @@
+// Copyright Epic Games, Inc. All Rights Reserved.
+
+#include "Algo/Find.h"
+#include "Containers/Array.h"
+#include "DerivedDataCache.h"
+#include "DerivedDataCacheStore.h"
+#include "DerivedDataCacheUsageStats.h"
+#include "DerivedDataLegacyCacheStore.h"
+#include "DerivedDataRequest.h"
+#include "DerivedDataRequestOwner.h"
+#include "HAL/CriticalSection.h"
+#include "Misc/EnumClassFlags.h"
+#include "Misc/ScopeRWLock.h"
+#include "Templates/RefCounting.h"
+#include "Templates/UniquePtr.h"
+#include <atomic>
+
+namespace UE::DerivedData::CacheStore
+{
+
+ILegacyCacheStore* CreateAsyncPutCacheStore(ILegacyCacheStore* InnerCache, ECacheStoreFlags InnerFlags, bool bCacheInFlightPuts)
+{
+	return nullptr;
+}
+
+///////////////////////////////////////////////////////////////////////////////////////////////////////////////////////
+
+class FCacheStoreHierarchy final : public ILegacyCacheStore, public ICacheStoreOwner
+{
+public:
+	~FCacheStoreHierarchy() final = default;
+
+	void Add(ILegacyCacheStore* CacheStore, ECacheStoreFlags Flags) final;
+
+	void SetFlags(ILegacyCacheStore* CacheStore, ECacheStoreFlags Flags) final;
+
+	void RemoveNotSafe(ILegacyCacheStore* CacheStore) final;
+
+	void Put(
+		TConstArrayView<FCachePutRequest> Requests,
+		IRequestOwner& Owner,
+		FOnCachePutComplete&& OnComplete) final;
+
+	void Get(
+		TConstArrayView<FCacheGetRequest> Requests,
+		IRequestOwner& Owner,
+		FOnCacheGetComplete&& OnComplete) final;
+
+	void PutValue(
+		TConstArrayView<FCachePutValueRequest> Requests,
+		IRequestOwner& Owner,
+		FOnCachePutValueComplete&& OnComplete) final;
+
+	void GetValue(
+		TConstArrayView<FCacheGetValueRequest> Requests,
+		IRequestOwner& Owner,
+		FOnCacheGetValueComplete&& OnComplete) final;
+
+	void GetChunks(
+		TConstArrayView<FCacheGetChunkRequest> Requests,
+		IRequestOwner& Owner,
+		FOnCacheGetChunkComplete&& OnComplete) final;
+
+	void LegacyPut(
+		TConstArrayView<FLegacyCachePutRequest> Requests,
+		IRequestOwner& Owner,
+		FOnLegacyCachePutComplete&& OnComplete) final;
+
+	void LegacyGet(
+		TConstArrayView<FLegacyCacheGetRequest> Requests,
+		IRequestOwner& Owner,
+		FOnLegacyCacheGetComplete&& OnComplete) final;
+
+	void LegacyDelete(
+		TConstArrayView<FLegacyCacheDeleteRequest> Requests,
+		IRequestOwner& Owner,
+		FOnLegacyCacheDeleteComplete&& OnComplete) final;
+
+	void LegacyStats(FDerivedDataCacheStatsNode& OutNode) final;
+
+	bool LegacyDebugOptions(FBackendDebugOptions& Options) final;
+
+private:
+	// Caller must hold a write lock on CacheStoresLock.
+	void UpdateNodeFlags();
+
+	class FBatchBase;
+	class FCounterEvent;
+
+	class FLegacyPutBatch;
+	class FLegacyGetBatch;
+	class FLegacyDeleteBatch;
+
+	enum class ECacheStoreNodeFlags : uint32;
+	FRIEND_ENUM_CLASS_FLAGS(ECacheStoreNodeFlags);
+
+	static bool CanQuery(ECachePolicy Policy, ECacheStoreFlags Flags);
+	static bool CanStore(ECachePolicy Policy, ECacheStoreFlags Flags);
+	static bool CanStoreIfOk(ECachePolicy Policy, ECacheStoreNodeFlags Flags);
+	static bool CanQueryIfError(ECachePolicy Policy, ECacheStoreNodeFlags Flags);
+
+	struct FCacheStoreNode
+	{
+		ILegacyCacheStore* Cache{};
+		ECacheStoreFlags CacheFlags{};
+		ECacheStoreNodeFlags NodeFlags{};
+		TUniquePtr<ILegacyCacheStore> AsyncCache;
+	};
+
+	mutable FRWLock NodesLock;
+	ECacheStoreNodeFlags CombinedNodeFlags{};
+	TArray<FCacheStoreNode, TInlineAllocator<8>> Nodes;
+	FDerivedDataCacheUsageStats UsageStats;
+};
+
+///////////////////////////////////////////////////////////////////////////////////////////////////////////////////////
+
+class FCacheStoreHierarchy::FBatchBase
+{
+public:
+	void AddRef()
+	{
+		ReferenceCount.fetch_add(1, std::memory_order_relaxed);
+	}
+
+	void Release()
+	{
+		if (ReferenceCount.fetch_sub(1, std::memory_order_acq_rel) == 1)
+		{
+			delete this;
+		}
+	}
+
+private:
+	std::atomic<int32> ReferenceCount{0};
+};
+
+class FCacheStoreHierarchy::FCounterEvent
+{
+public:
+	void Reset(int32 NewCount)
+	{
+		Count.store(NewCount, std::memory_order_relaxed);
+	}
+
+	bool Signal()
+	{
+		return Count.fetch_sub(1, std::memory_order_acq_rel) == 1;
+	}
+
+private:
+	std::atomic<int32> Count{0};
+};
+
+///////////////////////////////////////////////////////////////////////////////////////////////////////////////////////
+
+enum class FCacheStoreHierarchy::ECacheStoreNodeFlags : uint32
+{
+	None                    = 0,
+
+	/** This node is preceded by a node that has the Store and Local flags. */
+	HasStoreLocalNode       = 1 << 0,
+	/** This node is preceded by a node that has the Store and Remote flags. */
+	HasStoreRemoteNode      = 1 << 1,
+	/** This node is preceded by a node that has the Store and (Local or Remote) flags. */
+	HasStoreNode            = HasStoreLocalNode | HasStoreRemoteNode,
+
+	/** This node is followed by a node that has the Query and Local flags. */
+	HasQueryLocalNode       = 1 << 2,
+	/** This node is followed by a node that has the Query and Remote flags. */
+	HasQueryRemoteNode      = 1 << 3,
+	/** This node is followed by a node that has the Query and (Local or Remote) flags. */
+	HasQueryNode            = HasQueryLocalNode | HasQueryRemoteNode,
+};
+
+ENUM_CLASS_FLAGS(FCacheStoreHierarchy::ECacheStoreNodeFlags);
+
+///////////////////////////////////////////////////////////////////////////////////////////////////////////////////////
+
+void FCacheStoreHierarchy::Add(ILegacyCacheStore* CacheStore, ECacheStoreFlags Flags)
+{
+	FWriteScopeLock Lock(NodesLock);
+	checkf(!Algo::FindBy(Nodes, CacheStore, &FCacheStoreNode::Cache),
+		TEXT("Attempting to add a cache store that was previously registered to the hierarchy."));
+	TUniquePtr<ILegacyCacheStore> AsyncCacheStore(CreateAsyncPutCacheStore(CacheStore, Flags, /*bCacheInFlightPuts*/ false));
+	Nodes.Add({CacheStore, Flags, {}, MoveTemp(AsyncCacheStore)});
+	UpdateNodeFlags();
+}
+
+void FCacheStoreHierarchy::SetFlags(ILegacyCacheStore* CacheStore, ECacheStoreFlags Flags)
+{
+	FWriteScopeLock Lock(NodesLock);
+	FCacheStoreNode* Node = Algo::FindBy(Nodes, CacheStore, &FCacheStoreNode::Cache);
+	checkf(!!Node, TEXT("Attempting to set flags on a cache store that is not registered to the hierarchy."));
+	Node->CacheFlags = Flags;
+	UpdateNodeFlags();
+}
+
+void FCacheStoreHierarchy::RemoveNotSafe(ILegacyCacheStore* CacheStore)
+{
+	FWriteScopeLock Lock(NodesLock);
+	FCacheStoreNode* Node = Algo::FindBy(Nodes, CacheStore, &FCacheStoreNode::Cache);
+	checkf(!!Node, TEXT("Attempting to remove a cache store that is not registered to the hierarchy."));
+	Nodes.RemoveAt(UE_PTRDIFF_TO_INT32(Node - Nodes.GetData()));
+	UpdateNodeFlags();
+}
+
+void FCacheStoreHierarchy::UpdateNodeFlags()
+{
+	ECacheStoreNodeFlags StoreFlags = ECacheStoreNodeFlags::None;
+	for (int32 Index = 0, Count = Nodes.Num(); Index < Count; ++Index)
+	{
+		FCacheStoreNode& Node = Nodes[Index];
+		Node.NodeFlags = StoreFlags;
+		if (EnumHasAllFlags(Node.CacheFlags, ECacheStoreFlags::Store | ECacheStoreFlags::Local))
+		{
+			StoreFlags |= ECacheStoreNodeFlags::HasStoreLocalNode;
+		}
+		if (EnumHasAllFlags(Node.CacheFlags, ECacheStoreFlags::Store | ECacheStoreFlags::Remote))
+		{
+			StoreFlags |= ECacheStoreNodeFlags::HasStoreRemoteNode;
+		}
+	}
+
+	ECacheStoreNodeFlags QueryFlags = ECacheStoreNodeFlags::None;
+	for (int32 Index = Nodes.Num() - 1; Index >= 0; --Index)
+	{
+		FCacheStoreNode& Node = Nodes[Index];
+		Node.NodeFlags |= QueryFlags;
+		if (EnumHasAllFlags(Node.CacheFlags, ECacheStoreFlags::Query | ECacheStoreFlags::Local))
+		{
+			QueryFlags |= ECacheStoreNodeFlags::HasQueryLocalNode;
+		}
+		if (EnumHasAllFlags(Node.CacheFlags, ECacheStoreFlags::Query | ECacheStoreFlags::Remote))
+		{
+			QueryFlags |= ECacheStoreNodeFlags::HasQueryRemoteNode;
+		}
+	}
+
+	CombinedNodeFlags = StoreFlags | QueryFlags;
+}
+
+bool FCacheStoreHierarchy::CanQuery(const ECachePolicy Policy, const ECacheStoreFlags Flags)
+{
+	const ECacheStoreFlags LocationFlags =
+		(EnumHasAnyFlags(Policy, ECachePolicy::QueryLocal) ? ECacheStoreFlags::Local : ECacheStoreFlags::None) |
+		(EnumHasAnyFlags(Policy, ECachePolicy::QueryRemote) ? ECacheStoreFlags::Remote : ECacheStoreFlags::None);
+	return EnumHasAnyFlags(Flags, LocationFlags) && EnumHasAnyFlags(Flags, ECacheStoreFlags::Query);
+}
+
+bool FCacheStoreHierarchy::CanStore(const ECachePolicy Policy, const ECacheStoreFlags Flags)
+{
+	const ECacheStoreFlags LocationFlags =
+		(EnumHasAnyFlags(Policy, ECachePolicy::StoreLocal) ? ECacheStoreFlags::Local : ECacheStoreFlags::None) |
+		(EnumHasAnyFlags(Policy, ECachePolicy::StoreRemote) ? ECacheStoreFlags::Remote : ECacheStoreFlags::None);
+	return EnumHasAnyFlags(Flags, LocationFlags) && EnumHasAnyFlags(Flags, ECacheStoreFlags::Store);
+}
+
+bool FCacheStoreHierarchy::CanStoreIfOk(const ECachePolicy Policy, const ECacheStoreNodeFlags Flags)
+{
+	const ECacheStoreNodeFlags LocationFlags =
+		(EnumHasAnyFlags(Policy, ECachePolicy::StoreLocal) ? ECacheStoreNodeFlags::HasStoreLocalNode : ECacheStoreNodeFlags::None) |
+		(EnumHasAnyFlags(Policy, ECachePolicy::StoreRemote) ? ECacheStoreNodeFlags::HasStoreRemoteNode : ECacheStoreNodeFlags::None);
+	return EnumHasAnyFlags(Flags, LocationFlags);
+}
+
+bool FCacheStoreHierarchy::CanQueryIfError(const ECachePolicy Policy, const ECacheStoreNodeFlags Flags)
+{
+	const ECacheStoreNodeFlags LocationFlags =
+		(EnumHasAnyFlags(Policy, ECachePolicy::QueryLocal) ? ECacheStoreNodeFlags::HasQueryLocalNode : ECacheStoreNodeFlags::None) |
+		(EnumHasAnyFlags(Policy, ECachePolicy::QueryRemote) ? ECacheStoreNodeFlags::HasQueryRemoteNode : ECacheStoreNodeFlags::None);
+	return EnumHasAnyFlags(Flags, LocationFlags);
+}
+
+///////////////////////////////////////////////////////////////////////////////////////////////////////////////////////
+
+void FCacheStoreHierarchy::Put(
+	TConstArrayView<FCachePutRequest> Requests,
+	IRequestOwner& Owner,
+	FOnCachePutComplete&& OnComplete)
+{
+	unimplemented();
+}
+
+///////////////////////////////////////////////////////////////////////////////////////////////////////////////////////
+
+void FCacheStoreHierarchy::Get(
+	TConstArrayView<FCacheGetRequest> Requests,
+	IRequestOwner& Owner,
+	FOnCacheGetComplete&& OnComplete)
+{
+	unimplemented();
+}
+
+///////////////////////////////////////////////////////////////////////////////////////////////////////////////////////
+
+void FCacheStoreHierarchy::PutValue(
+	TConstArrayView<FCachePutValueRequest> Requests,
+	IRequestOwner& Owner,
+	FOnCachePutValueComplete&& OnComplete)
+{
+	unimplemented();
+}
+
+///////////////////////////////////////////////////////////////////////////////////////////////////////////////////////
+
+void FCacheStoreHierarchy::GetValue(
+	TConstArrayView<FCacheGetValueRequest> Requests,
+	IRequestOwner& Owner,
+	FOnCacheGetValueComplete&& OnComplete)
+{
+	unimplemented();
+}
+
+///////////////////////////////////////////////////////////////////////////////////////////////////////////////////////
+
+void FCacheStoreHierarchy::GetChunks(
+	TConstArrayView<FCacheGetChunkRequest> Requests,
+	IRequestOwner& Owner,
+	FOnCacheGetChunkComplete&& OnComplete)
+{
+	unimplemented();
+}
+
+///////////////////////////////////////////////////////////////////////////////////////////////////////////////////////
+
+class FCacheStoreHierarchy::FLegacyPutBatch : public FBatchBase
+{
+public:
+	static void Begin(
+		const FCacheStoreHierarchy& Hierarchy,
+		TConstArrayView<FLegacyCachePutRequest> Requests,
+		IRequestOwner& Owner,
+		FOnLegacyCachePutComplete&& OnComplete);
+
+private:
+	FLegacyPutBatch(
+		const FCacheStoreHierarchy& InHierarchy,
+		const TConstArrayView<FLegacyCachePutRequest> InRequests,
+		IRequestOwner& InOwner,
+		FOnLegacyCachePutComplete&& InOnComplete)
+		: Hierarchy(InHierarchy)
+		, Requests(InRequests)
+		, BatchOwner(InOwner)
+		, OnComplete(MoveTemp(InOnComplete))
+		, AsyncOwner(FPlatformMath::Min(InOwner.GetPriority(), EPriority::Highest))
+	{
+		AsyncOwner.KeepAlive();
+		States.SetNum(Requests.Num());
+	}
+
+	void DispatchRequests();
+
+	bool DispatchGetRequests();
+	void CompleteGetRequest(FLegacyCacheGetResponse&& Response);
+
+	bool DispatchPutRequests();
+	void CompletePutRequest(FLegacyCachePutResponse&& Response);
+
+	struct FRequestState
+	{
+		bool bOk = false;
+		bool bStop = false;
+	};
+
+	const FCacheStoreHierarchy& Hierarchy;
+	TArray<FLegacyCachePutRequest> Requests;
+	IRequestOwner& BatchOwner;
+	FOnLegacyCachePutComplete OnComplete;
+
+	FRequestOwner AsyncOwner;
+	TArray<FRequestState, TInlineAllocator<1>> States;
+	FCounterEvent RemainingRequestCount;
+	int32 NodeGetIndex = -1;
+	int32 NodePutIndex = 0;
+};
+
+void FCacheStoreHierarchy::FLegacyPutBatch::Begin(
+	const FCacheStoreHierarchy& Hierarchy,
+	const TConstArrayView<FLegacyCachePutRequest> Requests,
+	IRequestOwner& Owner,
+	FOnLegacyCachePutComplete&& OnComplete)
+{
+	if (Requests.IsEmpty() || !EnumHasAnyFlags(Hierarchy.CombinedNodeFlags, ECacheStoreNodeFlags::HasStoreNode))
+	{
+		return CompleteWithStatus(Requests, OnComplete, EStatus::Error);
+	}
+
+	TRefCountPtr<FLegacyPutBatch> State = new FLegacyPutBatch(Hierarchy, Requests, Owner, MoveTemp(OnComplete));
+	State->DispatchRequests();
+}
+
+void FCacheStoreHierarchy::FLegacyPutBatch::DispatchRequests()
+{
+	FReadScopeLock Lock(Hierarchy.NodesLock);
+
+	for (const int32 NodeCount = Hierarchy.Nodes.Num(); NodePutIndex < NodeCount && !BatchOwner.IsCanceled(); ++NodePutIndex)
+	{
+		if (DispatchGetRequests() || DispatchPutRequests())
+		{
+			return;
+		}
+	}
+
+	int32 RequestIndex = 0;
+	for (const FLegacyCachePutRequest& Request : Requests)
+	{
+		const FRequestState& State = States[RequestIndex];
+		if (!State.bOk && !State.bStop)
+		{
+			OnComplete(Request.MakeResponse(BatchOwner.IsCanceled() ? EStatus::Canceled : EStatus::Error));
+		}
+	}
+}
+
+bool FCacheStoreHierarchy::FLegacyPutBatch::DispatchGetRequests()
+{
+	if (NodeGetIndex >= NodePutIndex)
+	{
+		return false;
+	}
+
+	NodeGetIndex = NodePutIndex;
+
+	const FCacheStoreNode& Node = Hierarchy.Nodes[NodeGetIndex];
+	if (!EnumHasAnyFlags(Node.CacheFlags, ECacheStoreFlags::StopStore))
+	{
+		return false;
+	}
+
+	TArray<FLegacyCacheGetRequest, TInlineAllocator<1>> NodeRequests;
+	NodeRequests.Reserve(Requests.Num());
+
+	uint64 RequestIndex = 0;
+	for (const FLegacyCachePutRequest& Request : Requests)
+	{
+		if (!States[RequestIndex].bStop && CanQuery(Request.Policy, Node.CacheFlags))
+		{
+			NodeRequests.Add({Request.Name, Request.Key, Request.Policy | ECachePolicy::SkipData, RequestIndex});
+		}
+		++RequestIndex;
+	}
+
+	if (const int32 NodeRequestsCount = NodeRequests.Num())
+	{
+		RemainingRequestCount.Reset(NodeRequestsCount + 1);
+		Node.Cache->LegacyGet(NodeRequests, BatchOwner,
+			[State = TRefCountPtr(this)](FLegacyCacheGetResponse&& Response)
+			{
+				State->CompleteGetRequest(MoveTemp(Response));
+			});
+		return !RemainingRequestCount.Signal();
+	}
+
+	return false;
+}
+
+void FCacheStoreHierarchy::FLegacyPutBatch::CompleteGetRequest(FLegacyCacheGetResponse&& Response)
+{
+	if (Response.Status == EStatus::Ok)
+	{
+		const int32 RequestIndex = int32(Response.UserData);
+		FRequestState& State = States[RequestIndex];
+		check(!State.bStop);
+		State.bStop = true;
+		if (!State.bOk)
+		{
+			OnComplete({Response.Name, Response.Key, Requests[RequestIndex].UserData, Response.Status});
+		}
+	}
+	if (RemainingRequestCount.Signal())
+	{
+		DispatchRequests();
+	}
+}
+
+bool FCacheStoreHierarchy::FLegacyPutBatch::DispatchPutRequests()
+{
+	const FCacheStoreNode& Node = Hierarchy.Nodes[NodePutIndex];
+	if (!EnumHasAnyFlags(Node.CacheFlags, ECacheStoreFlags::Store))
+	{
+		return false;
+	}
+
+	TArray<FLegacyCachePutRequest, TInlineAllocator<1>> NodeRequests;
+	TArray<FLegacyCachePutRequest, TInlineAllocator<1>> AsyncNodeRequests;
+
+	const int32 RequestCount = Requests.Num();
+	NodeRequests.Reserve(RequestCount);
+	AsyncNodeRequests.Reserve(RequestCount);
+
+	int32 RequestIndex = 0;
+	for (const FLegacyCachePutRequest& Request : Requests)
+	{
+		const FRequestState& State = States[RequestIndex];
+		if (!State.bStop && CanStore(Request.Policy, Node.CacheFlags))
+		{
+			(State.bOk ? AsyncNodeRequests : NodeRequests).Add_GetRef(Request).UserData = uint64(RequestIndex);
+		}
+		++RequestIndex;
+	}
+
+	if (!AsyncNodeRequests.IsEmpty())
+	{
+		Node.AsyncCache->LegacyPut(AsyncNodeRequests, AsyncOwner, [](auto&&){});
+	}
+
+	if (const int32 NodeRequestsCount = NodeRequests.Num())
+	{
+		RemainingRequestCount.Reset(NodeRequestsCount + 1);
+		Node.Cache->LegacyPut(NodeRequests, BatchOwner,
+			[State = TRefCountPtr(this)](FLegacyCachePutResponse&& Response)
+			{
+				State->CompletePutRequest(MoveTemp(Response));
+			});
+		return !RemainingRequestCount.Signal();
+	}
+
+	return false;
+}
+
+void FCacheStoreHierarchy::FLegacyPutBatch::CompletePutRequest(FLegacyCachePutResponse&& Response)
+{
+	if (Response.Status == EStatus::Ok)
+	{
+		const int32 RequestIndex = int32(Response.UserData);
+		FRequestState& State = States[RequestIndex];
+		check(!State.bOk && !State.bStop);
+		State.bOk = true;
+		Response.UserData = Requests[RequestIndex].UserData;
+		OnComplete(MoveTemp(Response));
+	}
+	if (RemainingRequestCount.Signal())
+	{
+		DispatchRequests();
+	}
+}
+
+void FCacheStoreHierarchy::LegacyPut(
+	const TConstArrayView<FLegacyCachePutRequest> Requests,
+	IRequestOwner& Owner,
+	FOnLegacyCachePutComplete&& OnComplete)
+{
+	FLegacyPutBatch::Begin(*this, Requests, Owner, MoveTemp(OnComplete));
+}
+
+///////////////////////////////////////////////////////////////////////////////////////////////////////////////////////
+
+class FCacheStoreHierarchy::FLegacyGetBatch : public FBatchBase
+{
+public:
+	static void Begin(
+		const FCacheStoreHierarchy& Hierarchy,
+		TConstArrayView<FLegacyCacheGetRequest> Requests,
+		IRequestOwner& Owner,
+		FOnLegacyCacheGetComplete&& OnComplete);
+
+private:
+	FLegacyGetBatch(
+		const FCacheStoreHierarchy& InHierarchy,
+		const TConstArrayView<FLegacyCacheGetRequest> InRequests,
+		IRequestOwner& InOwner,
+		FOnLegacyCacheGetComplete&& InOnComplete)
+		: Hierarchy(InHierarchy)
+		, OnComplete(MoveTemp(InOnComplete))
+		, Owner(InOwner)
+		, AsyncOwner(FPlatformMath::Min(InOwner.GetPriority(), EPriority::Highest))
+	{
+		AsyncOwner.KeepAlive();
+		States.Reserve(InRequests.Num());
+		for (const FLegacyCacheGetRequest& Request : InRequests)
+		{
+			States.Add({Request});
+		}
+	}
+
+	void DispatchRequests();
+	void CompleteRequest(FLegacyCacheGetResponse&& Response);
+
+	struct FState
+	{
+		FLegacyCacheGetRequest Request;
+		FLegacyCacheGetResponse Response;
+		bool bStop = false;
+	};
+
+	const FCacheStoreHierarchy& Hierarchy;
+	FOnLegacyCacheGetComplete OnComplete;
+	TArray<FState, TInlineAllocator<8>> States;
+
+	IRequestOwner& Owner;
+	FRequestOwner AsyncOwner;
+	FCounterEvent RemainingRequestCount;
+	int32 NodeIndex = 0;
+};
+
+void FCacheStoreHierarchy::FLegacyGetBatch::Begin(
+	const FCacheStoreHierarchy& InHierarchy,
+	const TConstArrayView<FLegacyCacheGetRequest> InRequests,
+	IRequestOwner& InOwner,
+	FOnLegacyCacheGetComplete&& InOnComplete)
+{
+	if (InRequests.IsEmpty() || !EnumHasAnyFlags(InHierarchy.CombinedNodeFlags, ECacheStoreNodeFlags::HasStoreNode))
+	{
+		return CompleteWithStatus(InRequests, InOnComplete, EStatus::Error);
+	}
+
+	TRefCountPtr<FLegacyGetBatch> State = new FLegacyGetBatch(InHierarchy, InRequests, InOwner, MoveTemp(InOnComplete));
+	State->DispatchRequests();
+}
+
+void FCacheStoreHierarchy::FLegacyGetBatch::DispatchRequests()
+{
+	FReadScopeLock Lock(Hierarchy.NodesLock);
+
+	TArray<FLegacyCacheGetRequest, TInlineAllocator<8>> NodeRequests;
+	TArray<FLegacyCachePutRequest, TInlineAllocator<8>> AsyncNodeRequests;
+
+	const int32 RequestCount = States.Num();
+	NodeRequests.Reserve(RequestCount);
+	AsyncNodeRequests.Reserve(RequestCount);
+
+	for (const int32 NodeCount = Hierarchy.Nodes.Num(); NodeIndex < NodeCount && !Owner.IsCanceled(); ++NodeIndex)
+	{
+		const FCacheStoreNode& Node = Hierarchy.Nodes[NodeIndex];
+
+		uint64 StateIndex = 0;
+		for (const FState& State : States)
+		{
+			const FLegacyCacheGetRequest& Request = State.Request;
+			const FLegacyCacheGetResponse& Response = State.Response;
+
+			// Fetch from any readable nodes.
+			if (Response.Status != EStatus::Ok && CanQuery(Request.Policy, Node.CacheFlags))
+			{
+				ECachePolicy Policy = Request.Policy;
+				if (EnumHasAnyFlags(Policy, ECachePolicy::SkipData) && CanStoreIfOk(Policy, Node.NodeFlags))
+				{
+					EnumRemoveFlags(Policy, ECachePolicy::SkipData);
+				}
+				NodeRequests.Add({Request.Name, Request.Key, Policy, StateIndex});
+			}
+
+			// Store to any writable local nodes unless a previous node contained the data and had the StopStore flag.
+			const ECachePolicy LocalPolicy = Request.Policy & ~ECachePolicy::Remote;
+			if (Response.Status == EStatus::Ok && Response.Value && !State.bStop && CanStore(LocalPolicy, Node.CacheFlags))
+			{
+				AsyncNodeRequests.Add({Response.Name, Response.Key, FCompositeBuffer(Response.Value), LocalPolicy});
+			}
+
+			++StateIndex;
+		}
+
+		if (!AsyncNodeRequests.IsEmpty())
+		{
+			Node.AsyncCache->LegacyPut(AsyncNodeRequests, AsyncOwner, [](auto&&){});
+			AsyncNodeRequests.Reset();
+		}
+
+		if (const int32 NodeRequestsCount = NodeRequests.Num())
+		{
+			RemainingRequestCount.Reset(NodeRequestsCount + 1);
+			Node.Cache->LegacyGet(NodeRequests, Owner,
+				[State = TRefCountPtr(this)](FLegacyCacheGetResponse&& Response)
+				{
+					State->CompleteRequest(MoveTemp(Response));
+				});
+			NodeRequests.Reset();
+			if (!RemainingRequestCount.Signal())
+			{
+				return;
+			}
+		}
+	}
+
+	for (const FState& State : States)
+	{
+		if (State.Response.Status != EStatus::Ok)
+		{
+			OnComplete(State.Request.MakeResponse(Owner.IsCanceled() ? EStatus::Canceled : EStatus::Error));
+		}
+	}
+}
+
+void FCacheStoreHierarchy::FLegacyGetBatch::CompleteRequest(FLegacyCacheGetResponse&& Response)
+{
+	if (Response.Status == EStatus::Ok)
+	{
+		FState& State = States[int32(Response.UserData)];
+		check(State.Response.Status == EStatus::Error);
+		Response.UserData = State.Request.UserData;
+
+		// Block any store to later nodes if this node has the StopStore flag.
+		const FCacheStoreNode& Node = Hierarchy.Nodes[NodeIndex];
+		if (EnumHasAnyFlags(Node.CacheFlags, ECacheStoreFlags::StopStore))
+		{
+			State.bStop = true;
+		}
+
+		// Store to any previous writable nodes that did not contain the data.
+		const ECachePolicy Policy = State.Request.Policy;
+		if (CanStoreIfOk(Policy, Node.NodeFlags) && Response.Value)
+		{
+			const FLegacyCachePutRequest PutRequest{Response.Name, Response.Key, FCompositeBuffer(Response.Value), Policy};
+			for (int32 PutNodeIndex = 0; PutNodeIndex < NodeIndex; ++PutNodeIndex)
+			{
+				const FCacheStoreNode& PutNode = Hierarchy.Nodes[PutNodeIndex];
+				if (CanStore(Policy, PutNode.CacheFlags))
+				{
+					PutNode.AsyncCache->LegacyPut({PutRequest}, AsyncOwner, [](auto&&){});
+				}
+			}
+		}
+
+		State.Response = Response;
+		OnComplete(MoveTemp(Response));
+	}
+
+	if (RemainingRequestCount.Signal())
+	{
+		DispatchRequests();
+	}
+}
+
+void FCacheStoreHierarchy::LegacyGet(
+	const TConstArrayView<FLegacyCacheGetRequest> Requests,
+	IRequestOwner& Owner,
+	FOnLegacyCacheGetComplete&& OnComplete)
+{
+	FLegacyGetBatch::Begin(*this, Requests, Owner, MoveTemp(OnComplete));
+}
+
+///////////////////////////////////////////////////////////////////////////////////////////////////////////////////////
+
+class FCacheStoreHierarchy::FLegacyDeleteBatch : public FBatchBase
+{
+public:
+	static void Begin(
+		const FCacheStoreHierarchy& Hierarchy,
+		TConstArrayView<FLegacyCacheDeleteRequest> Requests,
+		IRequestOwner& Owner,
+		FOnLegacyCacheDeleteComplete&& OnComplete);
+
+private:
+	FLegacyDeleteBatch(
+		const FCacheStoreHierarchy& InHierarchy,
+		const TConstArrayView<FLegacyCacheDeleteRequest> InRequests,
+		IRequestOwner& InOwner,
+		FOnLegacyCacheDeleteComplete&& InOnComplete)
+		: Hierarchy(InHierarchy)
+		, Requests(InRequests)
+		, BatchOwner(InOwner)
+		, OnComplete(MoveTemp(InOnComplete))
+	{
+		States.SetNum(Requests.Num());
+	}
+
+	void DispatchRequests();
+	void CompleteRequest(FLegacyCacheDeleteResponse&& Response);
+
+	struct FRequestState
+	{
+		bool bOk = false;
+	};
+
+	const FCacheStoreHierarchy& Hierarchy;
+	TArray<FLegacyCacheDeleteRequest> Requests;
+	IRequestOwner& BatchOwner;
+	FOnLegacyCacheDeleteComplete OnComplete;
+
+	TArray<FRequestState, TInlineAllocator<1>> States;
+	FCounterEvent RemainingRequestCount;
+	int32 NodeIndex = 0;
+};
+
+void FCacheStoreHierarchy::FLegacyDeleteBatch::Begin(
+	const FCacheStoreHierarchy& Hierarchy,
+	const TConstArrayView<FLegacyCacheDeleteRequest> Requests,
+	IRequestOwner& Owner,
+	FOnLegacyCacheDeleteComplete&& OnComplete)
+{
+	if (Requests.IsEmpty() || !EnumHasAnyFlags(Hierarchy.CombinedNodeFlags, ECacheStoreNodeFlags::HasStoreNode))
+	{
+		return CompleteWithStatus(Requests, OnComplete, EStatus::Error);
+	}
+
+	TRefCountPtr<FLegacyDeleteBatch> State = new FLegacyDeleteBatch(Hierarchy, Requests, Owner, MoveTemp(OnComplete));
+	State->DispatchRequests();
+}
+
+void FCacheStoreHierarchy::FLegacyDeleteBatch::DispatchRequests()
+{
+	FReadScopeLock Lock(Hierarchy.NodesLock);
+
+	TArray<FLegacyCacheDeleteRequest, TInlineAllocator<1>> NodeRequests;
+	NodeRequests.Reserve(Requests.Num());
+	for (const int32 NodeCount = Hierarchy.Nodes.Num(); NodeIndex < NodeCount && !BatchOwner.IsCanceled(); ++NodeIndex)
+	{
+		const FCacheStoreNode& Node = Hierarchy.Nodes[NodeIndex];
+
+		uint64 RequestIndex = 0;
+		for (const FLegacyCacheDeleteRequest& Request : Requests)
+		{
+			if (CanStore(Request.Policy, Node.CacheFlags))
+			{
+				NodeRequests.Add_GetRef(Request).UserData = RequestIndex;
+			}
+			++RequestIndex;
+		}
+
+		if (const int32 NodeRequestsCount = NodeRequests.Num())
+		{
+			RemainingRequestCount.Reset(NodeRequestsCount + 1);
+			Node.Cache->LegacyDelete(NodeRequests, BatchOwner,
+				[State = TRefCountPtr(this)](FLegacyCacheDeleteResponse&& Response)
+				{
+					State->CompleteRequest(MoveTemp(Response));
+				});
+			NodeRequests.Reset();
+			if (!RemainingRequestCount.Signal())
+			{
+				return;
+			}
+		}
+	}
+
+	int32 RequestIndex = 0;
+	for (const FLegacyCacheDeleteRequest& Request : Requests)
+	{
+		const bool bOk = States[RequestIndex++].bOk;
+		const EStatus Status = bOk ? EStatus::Ok : BatchOwner.IsCanceled() ? EStatus::Canceled : EStatus::Error;
+		OnComplete(Request.MakeResponse(Status));
+	}
+}
+
+void FCacheStoreHierarchy::FLegacyDeleteBatch::CompleteRequest(FLegacyCacheDeleteResponse&& Response)
+{
+	if (Response.Status == EStatus::Ok)
+	{
+		States[int32(Response.UserData)].bOk = true;
+	}
+	if (RemainingRequestCount.Signal())
+	{
+		DispatchRequests();
+	}
+}
+
+void FCacheStoreHierarchy::LegacyDelete(
+	const TConstArrayView<FLegacyCacheDeleteRequest> Requests,
+	IRequestOwner& Owner,
+	FOnLegacyCacheDeleteComplete&& OnComplete)
+{
+	FLegacyDeleteBatch::Begin(*this, Requests, Owner, MoveTemp(OnComplete));
+}
+
+///////////////////////////////////////////////////////////////////////////////////////////////////////////////////////
+
+void FCacheStoreHierarchy::LegacyStats(FDerivedDataCacheStatsNode& OutNode)
+{
+	FReadScopeLock Lock(NodesLock);
+	OutNode.Stats.Add(TEXT(""), UsageStats);
+	OutNode.Children.Reserve(Nodes.Num());
+	for (const FCacheStoreNode& Node : Nodes)
+	{
+		Node.Cache->LegacyStats(OutNode.Children.Add_GetRef(MakeShared<FDerivedDataCacheStatsNode>()).Get());
+	}
+}
+
+bool FCacheStoreHierarchy::LegacyDebugOptions(FBackendDebugOptions& Options)
+{
+	return false;
+}
+
+///////////////////////////////////////////////////////////////////////////////////////////////////////////////////////
+
+ILegacyCacheStore* CreateCacheStoreHierarchy(ICacheStoreOwner*& OutOwner)
+{
+	FCacheStoreHierarchy* Hierarchy = new FCacheStoreHierarchy;
+	OutOwner = Hierarchy;
+	return Hierarchy;
+}
+
+} // UE::DerivedData::CacheStore
