DDC: Updated the async cache wrapper to allow async execution in the wrapped cache store

#rnx

diff --git a/Engine/Source/Developer/DerivedDataCache/Private/DerivedDataBackendAsyncPutWrapper.cpp b/Engine/Source/Developer/DerivedDataCache/Private/DerivedDataBackendAsyncPutWrapper.cpp
deleted file mode 100644
diff --git a/Engine/Source/Developer/DerivedDataCache/Private/DerivedDataCacheStoreAsync.cpp b/Engine/Source/Developer/DerivedDataCache/Private/DerivedDataCacheStoreAsync.cpp
new file mode 100644
--- /dev/null
+++ b/Engine/Source/Developer/DerivedDataCache/Private/DerivedDataCacheStoreAsync.cpp
@@ -0,0 +1,750 @@
+// Copyright Epic Games, Inc. All Rights Reserved.
+
+#include "Async/AsyncWork.h"
+#include "DerivedDataBackendInterface.h"
+#include "DerivedDataCachePrivate.h"
+#include "DerivedDataCacheStore.h"
+#include "DerivedDataCacheUsageStats.h"
+#include "DerivedDataRequest.h"
+#include "DerivedDataRequestOwner.h"
+#include "DerivedDataValueId.h"
+#include "Experimental/Async/LazyEvent.h"
+#include "FileBackedDerivedDataBackend.h"
+#include "Memory/SharedBuffer.h"
+#include "Misc/ScopeLock.h"
+#include "ProfilingDebugging/CookStats.h"
+#include "Stats/Stats.h"
+#include "Tasks/Task.h"
+
+namespace UE::DerivedData::CacheStore::Memory
+{
+FFileBackedDerivedDataBackend* CreateMemoryDerivedDataBackend(const TCHAR* Name, int64 MaxCacheSize, bool bCanBeDisabled);
+} // UE::DerivedData::CacheStore::Memory
+
+namespace UE::DerivedData::CacheStore::AsyncPut
+{
+
+/** 
+ * Thread safe set helper
+**/
+struct FThreadSet
+{
+	FCriticalSection	SynchronizationObject;
+	TSet<FString>		FilesInFlight;
+
+	void Add(const FString& Key)
+	{
+		FScopeLock ScopeLock(&SynchronizationObject);
+		check(Key.Len());
+		FilesInFlight.Add(Key);
+	}
+	void Remove(const FString& Key)
+	{
+		FScopeLock ScopeLock(&SynchronizationObject);
+		FilesInFlight.Remove(Key);
+	}
+	bool Exists(const FString& Key)
+	{
+		FScopeLock ScopeLock(&SynchronizationObject);
+		return FilesInFlight.Contains(Key);
+	}
+	bool AddIfNotExists(const FString& Key)
+	{
+		FScopeLock ScopeLock(&SynchronizationObject);
+		check(Key.Len());
+		if (!FilesInFlight.Contains(Key))
+		{
+			FilesInFlight.Add(Key);
+			return true;
+		}
+		return false;
+	}
+};
+
+/** 
+ * A backend wrapper that coordinates async puts. This means that a Get will hit an in-memory cache while the async put is still in flight.
+**/
+class FDerivedDataBackendAsyncPutWrapper : public FDerivedDataBackendInterface
+{
+public:
+
+	/**
+	 * Constructor
+	 *
+	 * @param	InInnerBackend		Backend to use for storage, my responsibilities are about async puts
+	 * @param	bCacheInFlightPuts	if true, cache in-flight puts in a memory cache so that they hit immediately
+	 */
+	FDerivedDataBackendAsyncPutWrapper(FDerivedDataBackendInterface* InInnerBackend, bool bCacheInFlightPuts);
+
+	/** Return a name for this interface */
+	virtual FString GetName() const override
+	{
+		return FString::Printf(TEXT("AsyncPutWrapper (%s)"), *InnerBackend->GetName());
+	}
+
+	/** return true if this cache is writable **/
+	virtual bool IsWritable() const override;
+
+	/** Returns a class of speed for this interface **/
+	virtual ESpeedClass GetSpeedClass() const override;
+
+	/** Return true if hits on this cache should propagate to lower cache level. */
+	virtual bool BackfillLowerCacheLevels() const override;
+
+	/**
+	 * Synchronous test for the existence of a cache item
+	 *
+	 * @param	CacheKey	Alphanumeric+underscore key of this cache item
+	 * @return				true if the data probably will be found, this can't be guaranteed because of concurrency in the backends, corruption, etc
+	 */
+	virtual bool CachedDataProbablyExists(const TCHAR* CacheKey) override;
+
+	/**
+	 * Synchronous test for the existence of multiple cache items
+	 *
+	 * @param	CacheKeys	Alphanumeric+underscore key of the cache items
+	 * @return				A bit array with bits indicating whether the data for the corresponding key will probably be found
+	 */
+	virtual TBitArray<> CachedDataProbablyExistsBatch(TConstArrayView<FString> CacheKeys) override;
+
+	/**
+	 * Attempt to make sure the cached data will be available as optimally as possible.
+	 *
+	 * @param	CacheKeys	Alphanumeric+underscore keys of the cache items
+	 * @return				true if the data will probably be found in a fast backend on a future request.
+	 */
+	virtual TBitArray<> TryToPrefetch(TConstArrayView<FString> CacheKeys) override;
+
+	/**
+	 * Allows the DDC backend to determine if it wants to cache the provided data. Reasons for returning false could be a slow connection,
+	 * a file size limit, etc.
+	 */
+	virtual bool WouldCache(const TCHAR* CacheKey, TArrayView<const uint8> InData) override;
+
+	/**
+	 * Synchronous retrieve of a cache item
+	 *
+	 * @param	CacheKey	Alphanumeric+underscore key of this cache item
+	 * @param	OutData		Buffer to receive the results, if any were found
+	 * @return				true if any data was found, and in this case OutData is non-empty
+	 */
+	virtual bool GetCachedData(const TCHAR* CacheKey, TArray<uint8>& OutData) override;
+
+	/**
+	 * Asynchronous, fire-and-forget placement of a cache item
+	 *
+	 * @param	CacheKey	Alphanumeric+underscore key of this cache item
+	 * @param	InData		Buffer containing the data to cache, can be destroyed after the call returns, immediately
+	 * @param	bPutEvenIfExists	If true, then do not attempt skip the put even if CachedDataProbablyExists returns true
+	 */
+	virtual EPutStatus PutCachedData(const TCHAR* CacheKey, TArrayView<const uint8> InData, bool bPutEvenIfExists) override;
+
+	virtual void RemoveCachedData(const TCHAR* CacheKey, bool bTransient) override;
+
+	virtual TSharedRef<FDerivedDataCacheStatsNode> GatherUsageStats() const override;
+
+	virtual bool ApplyDebugOptions(FBackendDebugOptions& InOptions) override;
+
+	virtual void Put(
+		TConstArrayView<FCachePutRequest> Requests,
+		IRequestOwner& Owner,
+		FOnCachePutComplete&& OnComplete) override
+	{
+		Execute(COOK_STAT(&FDerivedDataCacheUsageStats::TimePut,) Requests, Owner, MoveTemp(OnComplete), &ICacheStore::Put);
+	}
+
+	virtual void Get(
+		TConstArrayView<FCacheGetRequest> Requests,
+		IRequestOwner& Owner,
+		FOnCacheGetComplete&& OnComplete) override
+	{
+		Execute(COOK_STAT(&FDerivedDataCacheUsageStats::TimeGet,) Requests, Owner, MoveTemp(OnComplete), &ICacheStore::Get);
+	}
+
+	virtual void PutValue(
+		TConstArrayView<FCachePutValueRequest> Requests,
+		IRequestOwner& Owner,
+		FOnCachePutValueComplete&& OnComplete) override
+	{
+		Execute(COOK_STAT(&FDerivedDataCacheUsageStats::TimePut,) Requests, Owner, MoveTemp(OnComplete), &ICacheStore::PutValue);
+	}
+
+	virtual void GetValue(
+		TConstArrayView<FCacheGetValueRequest> Requests,
+		IRequestOwner& Owner,
+		FOnCacheGetValueComplete&& OnComplete) override
+	{
+		Execute(COOK_STAT(&FDerivedDataCacheUsageStats::TimeGet,) Requests, Owner, MoveTemp(OnComplete), &ICacheStore::GetValue);
+	}
+
+	virtual void GetChunks(
+		TConstArrayView<FCacheGetChunkRequest> Requests,
+		IRequestOwner& Owner,
+		FOnCacheGetChunkComplete&& OnComplete) override
+	{
+		Execute(COOK_STAT(&FDerivedDataCacheUsageStats::TimeGet,) Requests, Owner, MoveTemp(OnComplete), &ICacheStore::GetChunks);
+	}
+
+	virtual void LegacyPut(
+		TConstArrayView<FLegacyCachePutRequest> Requests,
+		IRequestOwner& Owner,
+		FOnLegacyCachePutComplete&& OnComplete) override
+	{
+		Execute(COOK_STAT(&FDerivedDataCacheUsageStats::TimePut,) Requests, Owner, MoveTemp(OnComplete), &ILegacyCacheStore::LegacyPut);
+	}
+
+	virtual void LegacyGet(
+		TConstArrayView<FLegacyCacheGetRequest> Requests,
+		IRequestOwner& Owner,
+		FOnLegacyCacheGetComplete&& OnComplete) override
+	{
+		Execute(COOK_STAT(&FDerivedDataCacheUsageStats::TimeGet,) Requests, Owner, MoveTemp(OnComplete), &ILegacyCacheStore::LegacyGet);
+	}
+
+	virtual void LegacyDelete(
+		TConstArrayView<FLegacyCacheDeleteRequest> Requests,
+		IRequestOwner& Owner,
+		FOnLegacyCacheDeleteComplete&& OnComplete) override
+	{
+		Execute(COOK_STAT(&FDerivedDataCacheUsageStats::TimePut,) Requests, Owner, MoveTemp(OnComplete), &ILegacyCacheStore::LegacyDelete);
+	}
+
+	virtual void LegacyStats(FDerivedDataCacheStatsNode& OutNode) override;
+
+	virtual bool LegacyDebugOptions(FBackendDebugOptions& Options) override
+	{
+		return InnerBackend->LegacyDebugOptions(Options);
+	}
+
+private:
+	COOK_STAT(using CookStatsFunction = FCookStats::FScopedStatsCounter (FDerivedDataCacheUsageStats::*)());
+
+	template <typename RequestType, typename OnCompleteType, typename OnExecuteType>
+	void Execute(
+		COOK_STAT(CookStatsFunction OnAddStats,)
+		TConstArrayView<RequestType> Requests,
+		IRequestOwner& Owner,
+		OnCompleteType&& OnComplete,
+		OnExecuteType&& OnExecute);
+
+	FDerivedDataCacheUsageStats UsageStats;
+	FDerivedDataCacheUsageStats PutSyncUsageStats;
+
+	/** Backend to use for storage, my responsibilities are about async puts **/
+	FDerivedDataBackendInterface*					InnerBackend;
+	/** Memory based cache to deal with gets that happen while an async put is still in flight **/
+	TUniquePtr<FDerivedDataBackendInterface>		InflightCache;
+	/** We remember outstanding puts so that we don't do them redundantly **/
+	FThreadSet										FilesInFlight;
+};
+
+/** 
+ * Async task to handle the fire and forget async put
+ */
+class FCachePutAsyncWorker
+{
+public:
+	/** Cache Key for the put to InnerBackend **/
+	FString								CacheKey;
+	/** Data for the put to InnerBackend **/
+	TArray<uint8>						Data;
+	/** Backend to use for storage, my responsibilities are about async puts **/
+	FDerivedDataBackendInterface*		InnerBackend;
+	/** Memory based cache to clear once the put is finished **/
+	FDerivedDataBackendInterface*		InflightCache;
+	/** We remember outstanding puts so that we don't do them redundantly **/
+	FThreadSet*							FilesInFlight;
+	/**If true, then do not attempt skip the put even if CachedDataProbablyExists returns true **/
+	bool								bPutEvenIfExists;
+	/** Usage stats to track thread times. */
+	FDerivedDataCacheUsageStats&        UsageStats;
+
+	/** Constructor
+	*/
+	FCachePutAsyncWorker(const TCHAR* InCacheKey, TArrayView<const uint8> InData, FDerivedDataBackendInterface* InInnerBackend, bool InbPutEvenIfExists, FDerivedDataBackendInterface* InInflightCache, FThreadSet* InInFilesInFlight, FDerivedDataCacheUsageStats& InUsageStats)
+		: CacheKey(InCacheKey)
+		, InnerBackend(InInnerBackend)
+		, InflightCache(InInflightCache)
+		, FilesInFlight(InInFilesInFlight)
+		, bPutEvenIfExists(InbPutEvenIfExists)
+		, UsageStats(InUsageStats)
+	{
+		// Only make a copy if it's not going to be available from the Inflight cache
+		if (!InInflightCache || !InInflightCache->CachedDataProbablyExists(InCacheKey))
+		{
+			Data = InData;
+		}
+		check(InnerBackend);
+	}
+
+	bool ShouldAbortForShutdown()
+	{
+		using ESpeedClass = FDerivedDataBackendInterface::ESpeedClass;
+		ESpeedClass SpeedClass = InnerBackend->GetSpeedClass();
+		if (SpeedClass == ESpeedClass::Local)
+		{
+			return false;
+		}
+		return !GIsBuildMachine && FDerivedDataBackend::Get().IsShuttingDown();
+	}
+		
+	/** Call the inner backend and when that completes, remove the memory cache */
+	void DoWork()
+	{
+		TRACE_CPUPROFILER_EVENT_SCOPE(DDCPut_DoWork);
+		COOK_STAT(auto Timer = UsageStats.TimePut());
+
+		if (ShouldAbortForShutdown())
+		{
+			Abandon();
+			return;
+		}
+
+		using EPutStatus = FDerivedDataBackendInterface::EPutStatus;
+		EPutStatus Status = EPutStatus::NotCached;
+
+		if (!bPutEvenIfExists && InnerBackend->CachedDataProbablyExists(*CacheKey))
+		{
+			Status = EPutStatus::Cached;
+		}
+		else
+		{
+			if (InflightCache && Data.Num() == 0)
+			{
+				// We verified at construction time that we would be able to get the data from the Inflight cache
+				verify(InflightCache->GetCachedData(*CacheKey, Data));
+			}
+			Status = InnerBackend->PutCachedData(*CacheKey, Data, bPutEvenIfExists);
+			COOK_STAT(Timer.AddHit(Data.Num()));
+		}
+
+		if (InflightCache)
+		{
+			// if the data was not cached synchronously, retry
+			if (Status != EPutStatus::Cached && Status != EPutStatus::Skipped)
+			{
+				// retry after a brief wait
+				FPlatformProcess::SleepNoStats(0.2f);
+
+				if (Status == EPutStatus::Executing && InnerBackend->CachedDataProbablyExists(*CacheKey))
+				{
+					Status = EPutStatus::Cached;
+				}
+				else
+				{
+					if (Data.Num() == 0)
+					{
+						verify(InflightCache->GetCachedData(*CacheKey, Data));
+					}
+					Status = InnerBackend->PutCachedData(*CacheKey, Data, /*bPutEvenIfExists*/ false);
+				}
+			}
+
+			switch (Status)
+			{
+			case EPutStatus::Skipped:
+			case EPutStatus::Cached:
+				// remove this from the in-flight cache because the inner cache contains the data or it was intentionally skipped
+				InflightCache->RemoveCachedData(*CacheKey, /*bTransient*/ false);
+				break;
+			case EPutStatus::NotCached:
+				UE_LOG(LogDerivedDataCache, Display, TEXT("%s: Put failed, keeping in memory copy %s."), *InnerBackend->GetName(), *CacheKey);
+				if (uint32 ErrorCode = FPlatformMisc::GetLastError())
+				{
+					TCHAR ErrorBuffer[1024];
+					FPlatformMisc::GetSystemErrorMessage(ErrorBuffer, 1024, ErrorCode);
+					UE_LOG(LogDerivedDataCache, Display, TEXT("Failed to write %s to %s. Error: %u (%s)"), *CacheKey, *InnerBackend->GetName(), ErrorCode, ErrorBuffer);
+				}
+				break;
+			case EPutStatus::Executing:
+				UE_LOG(LogDerivedDataCache, Display, TEXT("%s: Put not finished executing, keeping in memory copy %s."), *InnerBackend->GetName(), *CacheKey);
+				break;
+			default:
+				break;
+			}
+		}
+
+		FilesInFlight->Remove(CacheKey);
+		FDerivedDataBackend::Get().AddToAsyncCompletionCounter(-1);
+		UE_LOG(LogDerivedDataCache, VeryVerbose, TEXT("%s: Completed AsyncPut of %s."), *InnerBackend->GetName(), *CacheKey);
+	}
+
+	FORCEINLINE TStatId GetStatId() const
+	{
+		RETURN_QUICK_DECLARE_CYCLE_STAT(FCachePutAsyncWorker, STATGROUP_ThreadPoolAsyncTasks);
+	}
+
+	/** Indicates to the thread pool that this task is abandonable */
+	bool CanAbandon()
+	{
+		return true;
+	}
+
+	/** Abandon routine, we need to remove the item from the in flight cache because something might be waiting for that */
+	void Abandon()
+	{
+		if (InflightCache)
+		{
+			InflightCache->RemoveCachedData(*CacheKey, /*bTransient=*/ false); // we can remove this from the temp cache, since the real cache will hit now
+		}
+		FilesInFlight->Remove(CacheKey);
+		FDerivedDataBackend::Get().AddToAsyncCompletionCounter(-1);
+		UE_LOG(LogDerivedDataCache, VeryVerbose, TEXT("%s: Abandoned AsyncPut of %s."), *InnerBackend->GetName(), *CacheKey);
+	}
+};
+
+FDerivedDataBackendAsyncPutWrapper::FDerivedDataBackendAsyncPutWrapper(FDerivedDataBackendInterface* InInnerBackend, bool bCacheInFlightPuts)
+	: InnerBackend(InInnerBackend)
+	, InflightCache(bCacheInFlightPuts ? Memory::CreateMemoryDerivedDataBackend(TEXT("InflightMemoryCache"), /*MaxCacheSize*/ -1, /*bCanBeDisabled*/ false) : nullptr)
+{
+	check(InnerBackend);
+}
+
+/** return true if this cache is writable **/
+bool FDerivedDataBackendAsyncPutWrapper::IsWritable() const
+{
+	return InnerBackend->IsWritable();
+}
+
+FDerivedDataBackendInterface::ESpeedClass FDerivedDataBackendAsyncPutWrapper::GetSpeedClass() const
+{
+	return InnerBackend->GetSpeedClass();
+}
+
+bool FDerivedDataBackendAsyncPutWrapper::BackfillLowerCacheLevels() const
+{
+	return InnerBackend->BackfillLowerCacheLevels();
+}
+
+bool FDerivedDataBackendAsyncPutWrapper::CachedDataProbablyExists(const TCHAR* CacheKey)
+{
+	COOK_STAT(auto Timer = UsageStats.TimeProbablyExists());
+	bool Result = (InflightCache && InflightCache->CachedDataProbablyExists(CacheKey)) || InnerBackend->CachedDataProbablyExists(CacheKey);
+	COOK_STAT(if (Result) {	Timer.AddHit(0); });
+
+	UE_LOG(LogDerivedDataCache, Verbose, TEXT("%s CachedDataProbablyExists=%d for %s"), *GetName(), Result, CacheKey);
+	return Result;
+}
+
+TBitArray<> FDerivedDataBackendAsyncPutWrapper::CachedDataProbablyExistsBatch(TConstArrayView<FString> CacheKeys)
+{
+	COOK_STAT(auto Timer = UsageStats.TimeProbablyExists());
+
+	TBitArray<> Result;
+	if (InflightCache)
+	{
+		Result = InflightCache->CachedDataProbablyExistsBatch(CacheKeys);
+		check(Result.Num() == CacheKeys.Num());
+		if (Result.CountSetBits() < CacheKeys.Num())
+		{
+			TBitArray<> InnerResult = InnerBackend->CachedDataProbablyExistsBatch(CacheKeys);
+			check(InnerResult.Num() == CacheKeys.Num());
+			Result.CombineWithBitwiseOR(InnerResult, EBitwiseOperatorFlags::MaintainSize);
+		}
+	}
+	else
+	{
+		Result = InnerBackend->CachedDataProbablyExistsBatch(CacheKeys);
+		check(Result.Num() == CacheKeys.Num());
+	}
+
+	COOK_STAT(if (Result.CountSetBits() == CacheKeys.Num()) { Timer.AddHit(0); });
+	UE_LOG(LogDerivedDataCache, Verbose, TEXT("%s CachedDataProbablyExists found %d/%d keys"), *GetName(), Result.CountSetBits(), CacheKeys.Num());
+	return Result;
+}
+
+TBitArray<> FDerivedDataBackendAsyncPutWrapper::TryToPrefetch(TConstArrayView<FString> CacheKeys)
+{
+	COOK_STAT(auto Timer = UsageStats.TimePrefetch());
+
+	TBitArray<> Result;
+	if (InflightCache)
+	{
+		Result = InflightCache->CachedDataProbablyExistsBatch(CacheKeys);
+	}
+	Result = TBitArray<>::BitwiseOR(Result, InnerBackend->TryToPrefetch(CacheKeys), EBitwiseOperatorFlags::MaxSize);
+
+	if (Result.CountSetBits() == CacheKeys.Num())
+	{
+		COOK_STAT(Timer.AddHit(0));
+	}
+
+	return Result;
+}
+
+/*
+	Determine if we would cache this by asking all our inner layers
+*/
+bool FDerivedDataBackendAsyncPutWrapper::WouldCache(const TCHAR* CacheKey, TArrayView<const uint8> InData)
+{
+	return InnerBackend->WouldCache(CacheKey, InData);
+}
+
+bool FDerivedDataBackendAsyncPutWrapper::ApplyDebugOptions(FBackendDebugOptions& InOptions)
+{
+	return InnerBackend->ApplyDebugOptions(InOptions);
+}
+
+bool FDerivedDataBackendAsyncPutWrapper::GetCachedData(const TCHAR* CacheKey, TArray<uint8>& OutData)
+{
+	COOK_STAT(auto Timer = UsageStats.TimeGet());
+	if (InflightCache && InflightCache->GetCachedData(CacheKey, OutData))
+	{
+		COOK_STAT(Timer.AddHit(OutData.Num()));
+		UE_LOG(LogDerivedDataCache, Verbose, TEXT("%s CacheHit from InFlightCache on %s"), *GetName(), CacheKey);
+		return true;
+	}
+
+	bool bSuccess = InnerBackend->GetCachedData(CacheKey, OutData);
+	if (bSuccess)
+	{
+		UE_LOG(LogDerivedDataCache, Verbose, TEXT("%s Cache hit on %s"), *GetName(), CacheKey);
+		COOK_STAT(Timer.AddHit(OutData.Num()));
+	}
+	else
+	{
+		UE_LOG(LogDerivedDataCache, Verbose, TEXT("%s Cache miss on %s"), *GetName(), CacheKey);
+	}
+	return bSuccess;
+}
+
+FDerivedDataBackendInterface::EPutStatus FDerivedDataBackendAsyncPutWrapper::PutCachedData(const TCHAR* CacheKey, TArrayView<const uint8> InData, bool bPutEvenIfExists)
+{
+	COOK_STAT(auto Timer = PutSyncUsageStats.TimePut());
+
+	if (!InnerBackend->IsWritable())
+	{
+		return EPutStatus::NotCached; // no point in continuing down the chain
+	}
+	const bool bAdded = FilesInFlight.AddIfNotExists(CacheKey);
+	if (!bAdded)
+	{
+		return EPutStatus::Executing; // if it is already on its way, we don't need to send it again
+	}
+	if (InflightCache)
+	{
+		if (InflightCache->CachedDataProbablyExists(CacheKey))
+		{
+			UE_LOG(LogDerivedDataCache, Verbose, TEXT("%s skipping out of key already in in-flight cache %s"), *GetName(), CacheKey);
+			return EPutStatus::Executing; // if it is already on its way, we don't need to send it again
+		}
+		InflightCache->PutCachedData(CacheKey, InData, true); // temp copy stored in memory while the async task waits to complete
+		COOK_STAT(Timer.AddHit(InData.Num()));
+	}
+	
+	UE_LOG(LogDerivedDataCache, Verbose, TEXT("%s queueing %s for put"), *GetName(), CacheKey);
+
+	FDerivedDataBackend::Get().AddToAsyncCompletionCounter(1);
+	(new FAutoDeleteAsyncTask<FCachePutAsyncWorker>(CacheKey, InData, InnerBackend, bPutEvenIfExists, InflightCache.Get(), &FilesInFlight, UsageStats))->StartBackgroundTask(Private::GCacheThreadPool, EQueuedWorkPriority::Low);
+
+	return EPutStatus::Executing;
+}
+
+void FDerivedDataBackendAsyncPutWrapper::RemoveCachedData(const TCHAR* CacheKey, bool bTransient)
+{
+	if (!InnerBackend->IsWritable())
+	{
+		return; // no point in continuing down the chain
+	}
+	while (FilesInFlight.Exists(CacheKey))
+	{
+		FPlatformProcess::Sleep(0.0f); // this is an exception condition (corruption), spin and wait for it to clear
+	}
+	if (InflightCache)
+	{
+		InflightCache->RemoveCachedData(CacheKey, bTransient);
+	}
+	InnerBackend->RemoveCachedData(CacheKey, bTransient);
+
+	UE_LOG(LogDerivedDataCache, Verbose, TEXT("%s removed %s"), *GetName(), CacheKey)
+}
+
+TSharedRef<FDerivedDataCacheStatsNode> FDerivedDataBackendAsyncPutWrapper::GatherUsageStats() const
+{
+	TSharedRef<FDerivedDataCacheStatsNode> Usage =
+		MakeShared<FDerivedDataCacheStatsNode>(TEXT("AsyncPutWrapper"), TEXT(""), InnerBackend->GetSpeedClass() == ESpeedClass::Local);
+	Usage->Stats.Add(TEXT("AsyncPut"), UsageStats);
+	Usage->Stats.Add(TEXT("AsyncPutSync"), PutSyncUsageStats);
+
+	Usage->Children.Add(InnerBackend->GatherUsageStats());
+	if (InflightCache)
+	{
+		Usage->Children.Add(InflightCache->GatherUsageStats());
+	}
+
+	return Usage;
+}
+
+void FDerivedDataBackendAsyncPutWrapper::LegacyStats(FDerivedDataCacheStatsNode& OutNode)
+{
+	OutNode = {TEXT("AsyncPutWrapper"), TEXT(""), InnerBackend->GetSpeedClass() == ESpeedClass::Local};
+	OutNode.Stats.Add(TEXT("AsyncPut"), UsageStats);
+	OutNode.Stats.Add(TEXT("AsyncPutSync"), PutSyncUsageStats);
+
+	InnerBackend->LegacyStats(OutNode.Children.Add_GetRef(MakeShared<FDerivedDataCacheStatsNode>()).Get());
+	if (InflightCache)
+	{
+		InflightCache->LegacyStats(OutNode.Children.Add_GetRef(MakeShared<FDerivedDataCacheStatsNode>()).Get());
+	}
+}
+
+class FDerivedDataAsyncWrapperRequest final : public FRequestBase, private IQueuedWork
+{
+public:
+	inline FDerivedDataAsyncWrapperRequest(
+		IRequestOwner& InOwner,
+		TUniqueFunction<void (IRequestOwner& Owner, bool bCancel)>&& InFunction)
+		: Owner(InOwner)
+		, Function(MoveTemp(InFunction))
+	{
+	}
+
+	inline void Start(EPriority Priority)
+	{
+		Owner.Begin(this);
+		DoneEvent.Reset();
+		Private::GCacheThreadPool->AddQueuedWork(this, GetPriority(Priority));
+	}
+
+	inline void Execute(bool bCancel)
+	{
+		FScopeCycleCounter Scope(GetStatId(), /*bAlways*/ true);
+		Owner.End(this, [this, bCancel]
+		{
+			Function(Owner, bCancel);
+			DoneEvent.Trigger();
+		});
+		// DO NOT ACCESS ANY MEMBERS PAST THIS POINT!
+	}
+
+	// IRequest Interface
+
+	inline void SetPriority(EPriority Priority) final
+	{
+		if (Private::GCacheThreadPool->RetractQueuedWork(this))
+		{
+			Private::GCacheThreadPool->AddQueuedWork(this, GetPriority(Priority));
+		}
+	}
+
+	inline void Cancel() final
+	{
+		if (!DoneEvent.Wait(0))
+		{
+			if (Private::GCacheThreadPool->RetractQueuedWork(this))
+			{
+				Abandon();
+			}
+			else
+			{
+				FScopeCycleCounter Scope(GetStatId());
+				DoneEvent.Wait();
+			}
+		}
+	}
+
+	inline void Wait() final
+	{
+		if (!DoneEvent.Wait(0))
+		{
+			if (Private::GCacheThreadPool->RetractQueuedWork(this))
+			{
+				DoThreadedWork();
+			}
+			else
+			{
+				FScopeCycleCounter Scope(GetStatId());
+				DoneEvent.Wait();
+			}
+		}
+	}
+
+private:
+	static EQueuedWorkPriority GetPriority(EPriority Priority)
+	{
+		switch (Priority)
+		{
+		case EPriority::Blocking: return EQueuedWorkPriority::Highest;
+		case EPriority::Highest:  return EQueuedWorkPriority::Highest;
+		case EPriority::High:     return EQueuedWorkPriority::High;
+		case EPriority::Normal:   return EQueuedWorkPriority::Normal;
+		case EPriority::Low:      return EQueuedWorkPriority::Low;
+		case EPriority::Lowest:   return EQueuedWorkPriority::Lowest;
+		default: checkNoEntry();  return EQueuedWorkPriority::Normal;
+		}
+	}
+
+	// IQueuedWork Interface
+
+	inline void DoThreadedWork() final
+	{
+		Execute(/*bCancel*/ false);
+	}
+
+	inline void Abandon() final
+	{
+		Execute(/*bCancel*/ true);
+	}
+
+	inline TStatId GetStatId() const
+	{
+		RETURN_QUICK_DECLARE_CYCLE_STAT(FDerivedDataAsyncWrapperRequest, STATGROUP_ThreadPoolAsyncTasks);
+	}
+
+private:
+	IRequestOwner& Owner;
+	TUniqueFunction<void (IRequestOwner& Owner, bool bCancel)> Function;
+	FLazyEvent DoneEvent{EEventMode::ManualReset};
+};
+
+template <typename RequestType, typename OnCompleteType, typename OnExecuteType>
+void FDerivedDataBackendAsyncPutWrapper::Execute(
+	COOK_STAT(CookStatsFunction OnAddStats,)
+	const TConstArrayView<RequestType> Requests,
+	IRequestOwner& Owner,
+	OnCompleteType&& OnComplete,
+	OnExecuteType&& OnExecute)
+{
+	auto ExecuteWithStats = [this, COOK_STAT(OnAddStats,) OnExecute](TConstArrayView<RequestType> Requests, IRequestOwner& Owner, OnCompleteType&& OnComplete) mutable
+	{
+		Invoke(OnExecute, *InnerBackend, Requests, Owner, [this, COOK_STAT(OnAddStats,) OnComplete = MoveTemp(OnComplete)](auto&& Response)
+		{
+			if (Response.Status == EStatus::Ok)
+			{
+				COOK_STAT((UsageStats.*OnAddStats)().AddHit(0));
+			}
+			if (OnComplete)
+			{
+				OnComplete(MoveTemp(Response));
+			}
+			FDerivedDataBackend::Get().AddToAsyncCompletionCounter(-1);
+		});
+	};
+
+	FDerivedDataBackend::Get().AddToAsyncCompletionCounter(Requests.Num());
+	if (Owner.GetPriority() == EPriority::Blocking || !Private::GCacheThreadPool)
+	{
+		return ExecuteWithStats(Requests, Owner, MoveTemp(OnComplete));
+	}
+
+	FDerivedDataAsyncWrapperRequest* Request = new FDerivedDataAsyncWrapperRequest(Owner,
+		[this, Requests = TArray<RequestType>(Requests), OnComplete = MoveTemp(OnComplete), ExecuteWithStats = MoveTemp(ExecuteWithStats)](IRequestOwner& Owner, bool bCancel) mutable
+		{
+			if (!bCancel)
+			{
+				ExecuteWithStats(Requests, Owner, MoveTemp(OnComplete));
+			}
+			else
+			{
+				CompleteWithStatus(Requests, MoveTemp(OnComplete), EStatus::Canceled);
+			}
+		});
+	Request->Start(Owner.GetPriority());
+}
+
+FDerivedDataBackendInterface* CreateAsyncPutDerivedDataBackend(FDerivedDataBackendInterface* InnerBackend, bool bCacheInFlightPuts)
+{
+	return new FDerivedDataBackendAsyncPutWrapper(InnerBackend, bCacheInFlightPuts);
+}
+
+} // UE::DerivedData::CacheStore::AsyncPut
