DDC: Removed default completion callbacks for private cache put functions

#rnx

diff --git a/Engine/Source/Developer/DerivedDataCache/Private/DerivedDataCacheStore.h b/Engine/Source/Developer/DerivedDataCache/Private/DerivedDataCacheStore.h
--- a/Engine/Source/Developer/DerivedDataCache/Private/DerivedDataCacheStore.h
+++ b/Engine/Source/Developer/DerivedDataCache/Private/DerivedDataCacheStore.h
@@ -35,7 +35,7 @@ public:
 	virtual void Put(
 		TConstArrayView<FCachePutRequest> Requests,
 		IRequestOwner& Owner,
-		FOnCachePutComplete&& OnComplete = {}) = 0;
+		FOnCachePutComplete&& OnComplete) = 0;
 
 	/**
 	 * Asynchronous request to get records from the cache.
@@ -65,7 +65,7 @@ public:
 	virtual void PutValue(
 		TConstArrayView<FCachePutValueRequest> Requests,
 		IRequestOwner& Owner,
-		FOnCachePutValueComplete&& OnComplete = {});
+		FOnCachePutValueComplete&& OnComplete);
 
 	/**
 	 * Asynchronous request to get values from the cache.
diff --git a/Engine/Source/Developer/DerivedDataCache/Private/DerivedDataLegacyCacheStore.h b/Engine/Source/Developer/DerivedDataCache/Private/DerivedDataLegacyCacheStore.h
--- a/Engine/Source/Developer/DerivedDataCache/Private/DerivedDataLegacyCacheStore.h
+++ b/Engine/Source/Developer/DerivedDataCache/Private/DerivedDataLegacyCacheStore.h
@@ -32,7 +32,7 @@ public:
 	virtual void LegacyPut(
 		TConstArrayView<FLegacyCachePutRequest> Requests,
 		IRequestOwner& Owner,
-		FOnLegacyCachePutComplete&& OnComplete = {}) = 0;
+		FOnLegacyCachePutComplete&& OnComplete) = 0;
 
 	virtual void LegacyGet(
 		TConstArrayView<FLegacyCacheGetRequest> Requests,
diff --git a/Engine/Source/Developer/DerivedDataCache/Private/HierarchicalDerivedDataBackend.h b/Engine/Source/Developer/DerivedDataCache/Private/HierarchicalDerivedDataBackend.h
--- a/Engine/Source/Developer/DerivedDataCache/Private/HierarchicalDerivedDataBackend.h
+++ b/Engine/Source/Developer/DerivedDataCache/Private/HierarchicalDerivedDataBackend.h
@@ -511,7 +511,7 @@ public:
 					}
 					else
 					{
-						AsyncPutInnerBackends[PutCacheIndex]->Put(Requests, AsyncOwner);
+						AsyncPutInnerBackends[PutCacheIndex]->Put(Requests, AsyncOwner, [](auto&&){});
 					}
 				}
 			}
@@ -616,7 +616,7 @@ public:
 								{
 									if (FillBackend->IsWritable())
 									{
-										FillBackend->Put({{Response.Name, Response.Record, PutPolicy}}, AsyncOwner);
+										FillBackend->Put({{Response.Name, Response.Record, PutPolicy}}, AsyncOwner, [](auto&&){});
 									}
 								}
 							}
@@ -677,7 +677,7 @@ public:
 					}
 					else
 					{
-						AsyncPutInnerBackends[PutCacheIndex]->PutValue(Requests, AsyncOwner);
+						AsyncPutInnerBackends[PutCacheIndex]->PutValue(Requests, AsyncOwner, [](auto&&){});
 					}
 				}
 			}
@@ -724,7 +724,7 @@ public:
 							{
 								if (GetCacheIndex != FillCacheIndex)
 								{
-									AsyncPutInnerBackends[FillCacheIndex]->PutValue({{Response.Name, Response.Key, Response.Value, ECachePolicy::Default}}, AsyncOwner);
+									AsyncPutInnerBackends[FillCacheIndex]->PutValue({{Response.Name, Response.Key, Response.Value, ECachePolicy::Default}}, AsyncOwner, [](auto&&){});
 								}
 							}
 
