DDC: Fixed legacy cache fill behavior in the new hierarchy

When a node with StopStore is present and contains the data, but comes later in the hierarchy than the node that first found the data, puts to subsequent nodes were not being disabled.

#rnx

diff --git a/Engine/Source/Developer/DerivedDataCache/Private/DerivedDataCacheStoreHierarchy.cpp b/Engine/Source/Developer/DerivedDataCache/Private/DerivedDataCacheStoreHierarchy.cpp
--- a/Engine/Source/Developer/DerivedDataCache/Private/DerivedDataCacheStoreHierarchy.cpp
+++ b/Engine/Source/Developer/DerivedDataCache/Private/DerivedDataCacheStoreHierarchy.cpp
@@ -502,6 +502,7 @@ bool FCacheStoreHierarchy::FLegacyPutBatch::DispatchPutRequests()
 
 	if (!AsyncNodeRequests.IsEmpty())
 	{
+		FRequestBarrier Barrier(AsyncOwner);
 		Node.AsyncCache->LegacyPut(AsyncNodeRequests, AsyncOwner, [](auto&&){});
 	}
 
@@ -581,7 +582,6 @@ private:
 	{
 		FLegacyCacheGetRequest Request;
 		FLegacyCacheGetResponse Response;
-		bool bStop = false;
 	};
 
 	const FCacheStoreHierarchy& Hierarchy;
@@ -629,30 +629,35 @@ void FCacheStoreHierarchy::FLegacyGetBatch::DispatchRequests()
 		{
 			const FLegacyCacheGetRequest& Request = State.Request;
 			const FLegacyCacheGetResponse& Response = State.Response;
-
-			// Fetch from any readable nodes.
-			if (Response.Status != EStatus::Ok && CanQuery(Request.Policy, Node.CacheFlags))
+			if (Response.Status == EStatus::Ok)
 			{
-				ECachePolicy Policy = Request.Policy;
-				if (EnumHasAnyFlags(Policy, ECachePolicy::SkipData) && CanStoreIfOk(Policy, Node.NodeFlags))
+				if (Response.Value && CanStore(Request.Policy, Node.CacheFlags))
+				{
+					AsyncNodeRequests.Add({Response.Name, Response.Key, FCompositeBuffer(Response.Value), Request.Policy});
+				}
+				else if (EnumHasAnyFlags(Node.CacheFlags, ECacheStoreFlags::StopStore) && CanQuery(Request.Policy, Node.CacheFlags))
 				{
-					EnumRemoveFlags(Policy, ECachePolicy::SkipData);
+					NodeRequests.Add({Request.Name, Request.Key, Request.Policy | ECachePolicy::SkipData, StateIndex});
 				}
-				NodeRequests.Add({Request.Name, Request.Key, Policy, StateIndex});
 			}
-
-			// Store to any writable local nodes unless a previous node contained the data and had the StopStore flag.
-			const ECachePolicy LocalPolicy = Request.Policy & ~ECachePolicy::Remote;
-			if (Response.Status == EStatus::Ok && Response.Value && !State.bStop && CanStore(LocalPolicy, Node.CacheFlags))
+			else
 			{
-				AsyncNodeRequests.Add({Response.Name, Response.Key, FCompositeBuffer(Response.Value), LocalPolicy});
+				if (CanQuery(Request.Policy, Node.CacheFlags))
+				{
+					ECachePolicy Policy = Request.Policy;
+					if (EnumHasAnyFlags(Policy, ECachePolicy::SkipData) && CanStoreIfOk(Policy, Node.NodeFlags))
+					{
+						EnumRemoveFlags(Policy, ECachePolicy::SkipData);
+					}
+					NodeRequests.Add({Request.Name, Request.Key, Policy, StateIndex});
+				}
 			}
-
 			++StateIndex;
 		}
 
 		if (!AsyncNodeRequests.IsEmpty())
 		{
+			FRequestBarrier Barrier(AsyncOwner);
 			Node.AsyncCache->LegacyPut(AsyncNodeRequests, AsyncOwner, [](auto&&){});
 			AsyncNodeRequests.Reset();
 		}
@@ -684,22 +689,17 @@ void FCacheStoreHierarchy::FLegacyGetBatch::DispatchRequests()
 
 void FCacheStoreHierarchy::FLegacyGetBatch::CompleteRequest(FLegacyCacheGetResponse&& Response)
 {
-	if (Response.Status == EStatus::Ok)
+	FState& State = States[int32(Response.UserData)];
+
+	if (Response.Status == EStatus::Ok && State.Response.Status == EStatus::Error)
 	{
-		FState& State = States[int32(Response.UserData)];
 		check(State.Response.Status == EStatus::Error);
 		Response.UserData = State.Request.UserData;
+		State.Response = Response;
 
-		// Block any store to later nodes if this node has the StopStore flag.
-		const FCacheStoreNode& Node = Hierarchy.Nodes[NodeIndex];
-		if (EnumHasAnyFlags(Node.CacheFlags, ECacheStoreFlags::StopStore))
-		{
-			State.bStop = true;
-		}
-
-		// Store to any previous writable nodes that did not contain the data.
+		// Store to previous writable nodes.
 		const ECachePolicy Policy = State.Request.Policy;
-		if (CanStoreIfOk(Policy, Node.NodeFlags) && Response.Value)
+		if (CanStoreIfOk(Policy, Hierarchy.Nodes[NodeIndex].NodeFlags) && Response.Value)
 		{
 			const FLegacyCachePutRequest PutRequest{Response.Name, Response.Key, FCompositeBuffer(Response.Value), Policy};
 			for (int32 PutNodeIndex = 0; PutNodeIndex < NodeIndex; ++PutNodeIndex)
@@ -707,15 +707,34 @@ void FCacheStoreHierarchy::FLegacyGetBatch::CompleteRequest(FLegacyCacheGetRespo
 				const FCacheStoreNode& PutNode = Hierarchy.Nodes[PutNodeIndex];
 				if (CanStore(Policy, PutNode.CacheFlags))
 				{
+					FRequestBarrier Barrier(AsyncOwner);
 					PutNode.AsyncCache->LegacyPut({PutRequest}, AsyncOwner, [](auto&&){});
 				}
 			}
 		}
 
-		State.Response = Response;
+		// Value may be fetched to fill other nodes. Remove the value if requested.
+		if (Response.Value && EnumHasAnyFlags(Policy, ECachePolicy::SkipData))
+		{
+			Response.Value.Reset();
+		}
+
 		OnComplete(MoveTemp(Response));
 	}
 
+	if (Response.Status == EStatus::Ok)
+	{
+		// Never store to later remote nodes.
+		// This is a necessary optimization until these speculative stores can be optimized.
+		EnumRemoveFlags(State.Request.Policy, ECachePolicy::Remote);
+
+		// Never store to later later nodes if this node has StopStore.
+		if (EnumHasAnyFlags(Hierarchy.Nodes[NodeIndex].CacheFlags, ECacheStoreFlags::StopStore))
+		{
+			EnumRemoveFlags(State.Request.Policy, ECachePolicy::Default);
+		}
+	}
+
 	if (RemainingRequestCount.Signal())
 	{
 		DispatchRequests();
