Fix accidental copy of data being fetched by [REDACTED_PERSON] due to the behavior of TArray64->TArray move operations.

diff --git a/Engine/Source/Developer/DerivedDataCache/Private/ZenDerivedDataBackend.cpp b/Engine/Source/Developer/DerivedDataCache/Private/ZenDerivedDataBackend.cpp
--- a/Engine/Source/Developer/DerivedDataCache/Private/ZenDerivedDataBackend.cpp
+++ b/Engine/Source/Developer/DerivedDataCache/Private/ZenDerivedDataBackend.cpp
@@ -303,7 +303,7 @@ bool FZenDerivedDataBackend::GetCachedData(const TCHAR* CacheKey, TArray<uint8>&
 	TArray64<uint8> ArrayBuffer;
 	EGetResult Result = GetZenData(MakeLegacyZenKey(CacheKey), &ArrayBuffer, Zen::EContentType::Binary);
 	check(ArrayBuffer.Num() <= UINT32_MAX);
-	OutData = MoveTemp(ArrayBuffer);
+	OutData = TArray<uint8>(MoveTemp(ArrayBuffer));
 	if (Result != EGetResult::Success)
 	{
 		switch (Result)
