Avoid synchronization bug involving DDC memory backend and value-locking that can occur during disable of the memory backend.

diff --git a/Engine/Source/Developer/DerivedDataCache/Private/MemoryDerivedDataBackend.cpp b/Engine/Source/Developer/DerivedDataCache/Private/MemoryDerivedDataBackend.cpp
--- a/Engine/Source/Developer/DerivedDataCache/Private/MemoryDerivedDataBackend.cpp
+++ b/Engine/Source/Developer/DerivedDataCache/Private/MemoryDerivedDataBackend.cpp
@@ -127,6 +127,12 @@ private:
 			, Size(InSize)
 		{
 		}
+
+		~FCacheValue()
+		{
+			// Ensure we don't begin destruction of this value while our own PutCachedData function is holding this lock
+			FWriteScopeLock WriteLock(DataLock);
+		}
 	};
 
 	FORCEINLINE int32 CalcSerializedCacheValueSize(const FString& Key, const FCacheValue& Val)
