DDC: Reorganized cache policy code into its own files

- FCacheRecordPolicyBuilder now disallows adding value policies with flags that only make sense for records.
- FCacheRecordPolicyBuilder now disallows adding multiple value policies for the same value.
- FCacheRecordPolicyBuilder now skips addition of value policies that match the default policy.

#rnx

diff --git a/Engine/Source/Developer/DerivedDataCache/Private/DerivedDataBuild.cpp b/Engine/Source/Developer/DerivedDataCache/Private/DerivedDataBuild.cpp
--- a/Engine/Source/Developer/DerivedDataCache/Private/DerivedDataBuild.cpp
+++ b/Engine/Source/Developer/DerivedDataCache/Private/DerivedDataBuild.cpp
@@ -185,7 +185,7 @@ EBuildPolicy FBuildPolicy::GetValuePolicy(const FValueId& Id) const
 			}
 		}
 	}
-	return DefaultValuePolicy;
+	return DefaultPolicy;
 }
 
 FBuildPolicy FBuildPolicy::Transform(TFunctionRef<EBuildPolicy (EBuildPolicy)> Op) const
@@ -194,7 +194,7 @@ FBuildPolicy FBuildPolicy::Transform(TFunctionRef<EBuildPolicy (EBuildPolicy)> O
 	{
 		return Op(CombinedPolicy);
 	}
-	FBuildPolicyBuilder Builder(Op(DefaultValuePolicy));
+	FBuildPolicyBuilder Builder(Op(DefaultPolicy));
 	for (const FBuildValuePolicy& Value : GetValuePolicies())
 	{
 		Builder.AddValuePolicy({Value.Id, Op(Value.Policy)});
diff --git a/Engine/Source/Developer/DerivedDataCache/Private/DerivedDataBuildJob.cpp b/Engine/Source/Developer/DerivedDataCache/Private/DerivedDataBuildJob.cpp
--- a/Engine/Source/Developer/DerivedDataCache/Private/DerivedDataBuildJob.cpp
+++ b/Engine/Source/Developer/DerivedDataCache/Private/DerivedDataBuildJob.cpp
@@ -490,7 +490,7 @@ static ECachePolicy MakeCacheQueryPolicy(EBuildPolicy BuildPolicy, const FBuildJ
 
 static FCacheRecordPolicy MakeCacheRecordQueryPolicy(const FBuildPolicy& BuildPolicy, const FBuildJobContext& Context)
 {
-	const ECachePolicy DefaultCachePolicy = MakeCacheQueryPolicy(BuildPolicy.GetDefaultValuePolicy(), Context);
+	const ECachePolicy DefaultCachePolicy = MakeCacheQueryPolicy(BuildPolicy.GetDefaultPolicy(), Context);
 	const ECachePolicy CombinedCachePolicy = MakeCacheQueryPolicy(BuildPolicy.GetCombinedPolicy(), Context);
 
 	FCacheRecordPolicyBuilder Builder(DefaultCachePolicy);
diff --git a/Engine/Source/Developer/DerivedDataCache/Private/DerivedDataBuildJobContext.cpp b/Engine/Source/Developer/DerivedDataCache/Private/DerivedDataBuildJobContext.cpp
--- a/Engine/Source/Developer/DerivedDataCache/Private/DerivedDataBuildJobContext.cpp
+++ b/Engine/Source/Developer/DerivedDataCache/Private/DerivedDataBuildJobContext.cpp
@@ -164,7 +164,7 @@ void FBuildJobContext::EndAsyncBuild()
 void FBuildJobContext::SetCacheBucket(FCacheBucket Bucket)
 {
 	checkf(!Bucket.IsNull(), TEXT("Null cache bucket not allowed for build of '%s' by %s. ")
-		TEXT("The cache can be disabled by calling SetCachePolicy(ECachePolicy::Disable)."),
+		TEXT("The cache can be disabled by calling SetCachePolicyMask(~ECachePolicy::Default)."),
 		*Job.GetName(), *WriteToString<32>(Job.GetFunction()));
 	CacheKey.Bucket = Bucket;
 }
@@ -173,11 +173,11 @@ void FBuildJobContext::SetCachePolicyMask(ECachePolicy Policy)
 {
 	checkf(EnumHasAllFlags(Policy, ECachePolicy::SkipData),
 		TEXT("SkipData flags may not be masked out on the cache policy for build of '%s' by %s. ")
-		TEXT("Flags for skipping data may be set indirectly through EBuildPolicy."),
+		TEXT("Flags for skipping data may be set indirectly through FBuildPolicy."),
 		*Job.GetName(), *WriteToString<32>(Job.GetFunction()));
 	checkf(EnumHasAllFlags(Policy, ECachePolicy::KeepAlive),
 		TEXT("KeepAlive flag may not be masked out on the cache policy for build of '%s' by %s. ")
-		TEXT("Flags for cache record lifetime may be set indirectly through EBuildPolicy."),
+		TEXT("Flags for cache record lifetime may be set indirectly through FBuildPolicy."),
 		*Job.GetName(), *WriteToString<32>(Job.GetFunction()));
 	checkf(EnumHasAllFlags(Policy, ECachePolicy::PartialRecord),
 		TEXT("PartialRecord flag may not be masked out on the cache policy for build of '%s' by %s."),
@@ -189,14 +189,14 @@ void FBuildJobContext::SetBuildPolicyMask(EBuildPolicy Policy)
 {
 	checkf(EnumHasAllFlags(Policy, EBuildPolicy::Cache),
 		TEXT("Cache flags may not be masked out on the build policy for build of '%s' by %s. ")
-		TEXT("Flags for modifying cache operations may be set through ECachePolicy."),
+		TEXT("Flags for modifying cache operations may be set through SetCachePolicyMask."),
 		*Job.GetName(), *WriteToString<32>(Job.GetFunction()));
 	checkf(EnumHasAllFlags(Policy, EBuildPolicy::CacheKeepAlive),
 		TEXT("CacheKeepAlive flag may not be masked out on the build policy for build of '%s' by %s. ")
 		TEXT("Flags for cache record lifetime may only be set through the build session."),
 		*Job.GetName(), *WriteToString<32>(Job.GetFunction()));
 	checkf(EnumHasAllFlags(Policy, EBuildPolicy::SkipData),
-		TEXT("SkipData flags may not be masked out on the build policy for build of '%s' by %s. ")
+		TEXT("SkipData flag may not be masked out on the build policy for build of '%s' by %s. ")
 		TEXT("Flags for skipping the data may only be set through the build session."),
 		*Job.GetName(), *WriteToString<32>(Job.GetFunction()));
 	BuildPolicyMask = Policy;
diff --git a/Engine/Source/Developer/DerivedDataCache/Private/DerivedDataCache.cpp b/Engine/Source/Developer/DerivedDataCache/Private/DerivedDataCache.cpp
--- a/Engine/Source/Developer/DerivedDataCache/Private/DerivedDataCache.cpp
+++ b/Engine/Source/Developer/DerivedDataCache/Private/DerivedDataCache.cpp
@@ -3,14 +3,10 @@
 #include "DerivedDataCache.h"
 #include "DerivedDataCacheInterface.h"
 
-#include "Algo/Accumulate.h"
 #include "Algo/AllOf.h"
-#include "Algo/BinarySearch.h"
-#include "Algo/Sort.h"
 #include "Async/AsyncWork.h"
 #include "Async/TaskGraphInterfaces.h"
 #include "Containers/Map.h"
-#include "Containers/StringConv.h"
 #include "DDCCleanup.h"
 #include "DerivedDataBackendInterface.h"
 #include "DerivedDataCache.h"
@@ -25,11 +21,8 @@
 #include "Misc/CoreMisc.h"
 #include "Misc/ScopeLock.h"
 #include "ProfilingDebugging/CookStats.h"
-#include "Serialization/CompactBinary.h"
-#include "Serialization/CompactBinaryWriter.h"
 #include "Stats/Stats.h"
 #include "Stats/StatsMisc.h"
-#include "String/ParseTokens.h"
 #include "ZenServerInterface.h"
 #include <atomic>
 
@@ -164,281 +157,9 @@ void GatherDerivedDataCacheSummaryStats(FDerivedDataCacheSummaryStats& DDCSummar
 /** Whether we want to verify the DDC (pass in -VerifyDDC on the command line)*/
 bool GVerifyDDC = false;
 
-namespace UE::DerivedData::Private { class FCacheRecordPolicyShared; }
-
 namespace UE::DerivedData
 {
 
-namespace Private::CachePolicy
-{
-	// TODO: Implement these as Ansi String instead to maximize the most-optimal path, see ParseCachePolicyImpl
-	constexpr ANSICHAR	  DelimiterChar = ',';
-	constexpr FStringView Delimiter = TEXTVIEW(",");
-	constexpr FStringView None = TEXTVIEW("None");
-	constexpr FStringView QueryLocal = TEXTVIEW("QueryLocal");
-	constexpr FStringView QueryRemote = TEXTVIEW("QueryRemote");
-	constexpr FStringView Query = TEXTVIEW("Query");
-	constexpr FStringView StoreLocal = TEXTVIEW("StoreLocal");
-	constexpr FStringView StoreRemote = TEXTVIEW("StoreRemote");
-	constexpr FStringView Store = TEXTVIEW("Store");
-	constexpr FStringView SkipMeta = TEXTVIEW("SkipMeta");
-	constexpr FStringView SkipData = TEXTVIEW("SkipData");
-	constexpr FStringView PartialRecord = TEXTVIEW("PartialRecord");
-	constexpr FStringView KeepAlive = TEXTVIEW("KeepAlive");
-	constexpr FStringView Local = TEXTVIEW("Local");
-	constexpr FStringView Remote = TEXTVIEW("Remote");
-	constexpr FStringView Default = TEXTVIEW("Default");
-	constexpr FStringView Disable = TEXTVIEW("Disable");
-
-	const TMap<FStringView, ECachePolicy> TextToPolicy 
-	{
-		{None,			ECachePolicy::None},
-		{QueryLocal,	ECachePolicy::QueryLocal},
-		{QueryRemote,	ECachePolicy::QueryRemote},
-		{Query,			ECachePolicy::Query},
-		{StoreLocal,	ECachePolicy::StoreLocal},
-		{StoreRemote,	ECachePolicy::StoreRemote},
-		{Store,			ECachePolicy::Store},
-		{SkipMeta,		ECachePolicy::SkipMeta},
-		{SkipData,		ECachePolicy::SkipData},
-		{PartialRecord, ECachePolicy::PartialRecord},
-		{KeepAlive,		ECachePolicy::KeepAlive},
-		{Local,			ECachePolicy::Local},
-		{Remote,		ECachePolicy::Remote},
-		{Default,		ECachePolicy::Default},
-		{Disable,		ECachePolicy::Disable}
-	};
-
-	using FPolicyTextPair = TPair<ECachePolicy, FStringView>;
-	const FPolicyTextPair FlagsToString[]
-	{
-		// Order of these Flags is important: we want the aliases before the atomic values,
-		// and the bigger aliases first, to reduce the number of tokens we add
-		{ ECachePolicy::Default, Default },
-		{ ECachePolicy::Remote, Remote },
-		{ ECachePolicy::Local, Local },
-		{ ECachePolicy::Store, Store },
-		{ ECachePolicy::Query, Query },
-
-		// Order of Atomics doesn't matter, so arbitrarily we list them in enum order
-		{ ECachePolicy::QueryLocal, QueryLocal },
-		{ ECachePolicy::QueryRemote, QueryRemote },
-		{ ECachePolicy::StoreLocal, StoreLocal },
-		{ ECachePolicy::StoreRemote, StoreRemote },
-		{ ECachePolicy::SkipMeta, SkipMeta },
-		{ ECachePolicy::SkipData, SkipData },
-		{ ECachePolicy::PartialRecord, PartialRecord },
-		{ ECachePolicy::KeepAlive, KeepAlive },
-
-		// None must come at the end of the array, to write out only if no others exist
-		{ ECachePolicy::None, None },
-	};
-	constexpr ECachePolicy KnownFlags = ECachePolicy::Default | ECachePolicy::SkipMeta | ECachePolicy::SkipData
-		| ECachePolicy::KeepAlive | ECachePolicy::PartialRecord;
-
-}	// namespace Private::CachePolicy
-
-template <typename CharType>
-TStringBuilderBase<CharType>& AppendToBuilderImpl(
-	TStringBuilderBase<CharType>& Builder, UE::DerivedData::ECachePolicy Policy)
-{
-	// Remove any bits we don't recognize; write None if there are not any bits we recognize
-	Policy = Policy & Private::CachePolicy::KnownFlags;
-	for (const Private::CachePolicy::FPolicyTextPair& Pair : Private::CachePolicy::FlagsToString)
-	{
-		if (EnumHasAllFlags(Policy, Pair.Key))
-		{
-			EnumRemoveFlags(Policy, Pair.Key);
-			Builder << Pair.Value << Private::CachePolicy::DelimiterChar;
-			if (Policy == ECachePolicy::None)
-			{
-				break;
-			}
-		}
-	}
-	Builder.RemoveSuffix(1); // Text will have been added by ECachePolicy::None if not by anything else
-	return Builder;
-}
-FAnsiStringBuilderBase& operator<<(FAnsiStringBuilderBase& Builder, UE::DerivedData::ECachePolicy Policy)
-{
-	return AppendToBuilderImpl(Builder, Policy);
-}
-FUtf8StringBuilderBase& operator<<(FUtf8StringBuilderBase& Builder, UE::DerivedData::ECachePolicy Policy)
-{
-	return AppendToBuilderImpl(Builder, Policy);
-}
-FWideStringBuilderBase& operator<<(FWideStringBuilderBase& Builder, UE::DerivedData::ECachePolicy Policy)
-{
-	return AppendToBuilderImpl(Builder, Policy);
-}
-
-template <typename CharType>
-ECachePolicy ParseCachePolicyImpl(TStringView<CharType> Text)
-{
-	checkf(!Text.IsEmpty(), TEXT("Empty string is not valid input to ParseCachePolicy"));
-
-	ECachePolicy Result = ECachePolicy::None;
-	// TODO: Implement ParseTokens for FAnsiStringView so we can convert to Ansi instead of Wide
-	auto WideText = StringCast<TCHAR, 128>(Text.GetData(), Text.Len());
-	UE::String::ParseTokens(WideText, Private::CachePolicy::DelimiterChar, [&Result](FStringView Token) {
-		const ECachePolicy* TokenPolicy = Private::CachePolicy::TextToPolicy.Find(Token);
-		if (TokenPolicy)
-		{
-			Result |= *TokenPolicy;
-		}
-	});
-
-	return Result;
-}
-ECachePolicy ParseCachePolicy(FAnsiStringView Text)
-{
-	return ParseCachePolicyImpl(Text);
-}
-ECachePolicy ParseCachePolicy(FUtf8StringView Text)
-{
-	return ParseCachePolicyImpl(Text);
-}
-ECachePolicy ParseCachePolicy(FWideStringView Text)
-{
-	return ParseCachePolicyImpl(Text);
-}
-
-
-class Private::FCacheRecordPolicyShared final : public Private::ICacheRecordPolicyShared
-{
-public:
-	inline void AddRef() const final
-	{
-		ReferenceCount.fetch_add(1, std::memory_order_relaxed);
-	}
-
-	inline void Release() const final
-	{
-		if (ReferenceCount.fetch_sub(1, std::memory_order_acq_rel) == 1)
-		{
-			delete this;
-		}
-	}
-
-	inline TConstArrayView<FCacheValuePolicy> GetValuePolicies() const final
-	{
-		return Values;
-	}
-
-	inline void AddValuePolicy(const FCacheValuePolicy& Policy) final
-	{
-		Values.Add(Policy);
-	}
-
-	inline void Build() final
-	{
-		Algo::SortBy(Values, &FCacheValuePolicy::Id);
-	}
-
-private:
-	TArray<FCacheValuePolicy, TInlineAllocator<14>> Values;
-	mutable std::atomic<uint32> ReferenceCount{0};
-};
-
-ECachePolicy FCacheRecordPolicy::GetValuePolicy(const FValueId& Id) const
-{
-	if (Shared)
-	{
-		if (TConstArrayView<FCacheValuePolicy> Values = Shared->GetValuePolicies(); !Values.IsEmpty())
-		{
-			if (int32 Index = Algo::BinarySearchBy(Values, Id, &FCacheValuePolicy::Id); Index != INDEX_NONE)
-			{
-				return Values[Index].Policy;
-			}
-		}
-	}
-	return DefaultValuePolicy;
-}
-
-FCacheRecordPolicy FCacheRecordPolicy::Transform(TFunctionRef<ECachePolicy (ECachePolicy)> Op) const
-{
-	if (IsUniform())
-	{
-		return Op(RecordPolicy);
-	}
-	FCacheRecordPolicyBuilder Builder(Op(DefaultValuePolicy));
-	for (const FCacheValuePolicy& Value : GetValuePolicies())
-	{
-		Builder.AddValuePolicy({Value.Id, Op(Value.Policy)});
-	}
-	return Builder.Build();
-}
-
-void FCacheRecordPolicy::Save(FCbWriter& Writer) const
-{
-	Writer.BeginObject();
-	{
-		// The RecordPolicy is calculated from the ValuePolicies and does not need to be saved separately.
-		Writer << "DefaultValuePolicy"_ASV << WriteToUtf8String<128>(GetDefaultValuePolicy());
-		if (!IsUniform())
-		{
-			// FCacheRecordPolicyBuilder guarantees IsUniform -> non-empty GetValuePolicies. Small size penalty here if not.
-			Writer.BeginArray("ValuePolicies"_ASV);
-			{
-				for (const FCacheValuePolicy& ValuePolicy : GetValuePolicies())
-				{
-					// FCacheRecordPolicyBuilder is responsible for ensuring that each ValuePolicy != DefaultValuePolicy
-					// If it lets any duplicates through we will incur a small serialization size penalty here
-					Writer.BeginObject();
-					Writer << "Id"_ASV << ValuePolicy.Id;
-					Writer << "Policy"_ASV << WriteToUtf8String<128>(ValuePolicy.Policy);
-					Writer.EndObject();
-				}
-			}
-			Writer.EndArray();
-		}
-	}
-	Writer.EndObject();
-}
-
-FCacheRecordPolicy FCacheRecordPolicy::Load(FCbObjectView Object, ECachePolicy DefaultPolicy)
-{
-	FUtf8StringView PolicyText = Object["DefaultValuePolicy"_ASV].AsString();
-	ECachePolicy DefaultValuePolicy = !PolicyText.IsEmpty() ? ParseCachePolicy(PolicyText) : DefaultPolicy;
-
-	FCacheRecordPolicyBuilder Builder(DefaultValuePolicy);
-	for (FCbFieldView ValueObjectField : Object["ValuePolicies"_ASV])
-	{
-		FCbObjectView ValueObject = ValueObjectField.AsObjectView();
-		const FCbObjectId ValueId = ValueObject["Id"_ASV].AsObjectId();
-		PolicyText = ValueObject["Policy"_ASV].AsString();
-		ECachePolicy ValuePolicy = !PolicyText.IsEmpty() ? ParseCachePolicy(PolicyText) : DefaultValuePolicy;
-		// FCacheRecordPolicyBuilder should guarantee that FValueId(ValueId).IsValid and ValuePolicy != DefaultValuePolicy
-		// If it lets any through we will have unused data in the record we create.
-		Builder.AddValuePolicy(ValueId, ValuePolicy);
-	}
-
-	return Builder.Build();
-}
-
-void FCacheRecordPolicyBuilder::AddValuePolicy(const FCacheValuePolicy& Policy)
-{
-	if (!Shared)
-	{
-		Shared = new Private::FCacheRecordPolicyShared;
-	}
-	Shared->AddValuePolicy(Policy);
-}
-
-FCacheRecordPolicy FCacheRecordPolicyBuilder::Build()
-{
-	FCacheRecordPolicy Policy(BasePolicy);
-	if (Shared)
-	{
-		Shared->Build();
-		const auto PolicyOr = [](ECachePolicy A, ECachePolicy B) { return A | (B & ~ECachePolicy::SkipData); };
-		const TConstArrayView<FCacheValuePolicy> Values = Shared->GetValuePolicies();
-		Policy.RecordPolicy = Algo::TransformAccumulate(Values, &FCacheValuePolicy::Policy, BasePolicy, PolicyOr);
-		Policy.Shared = MoveTemp(Shared);
-	}
-	return Policy;
-}
-
 void ICacheStore::PutValue(
 	const TConstArrayView<FCachePutValueRequest> Requests,
 	IRequestOwner& Owner,
diff --git a/Engine/Source/Developer/DerivedDataCache/Private/DerivedDataCachePolicy.cpp b/Engine/Source/Developer/DerivedDataCache/Private/DerivedDataCachePolicy.cpp
new file mode 100644
--- /dev/null
+++ b/Engine/Source/Developer/DerivedDataCache/Private/DerivedDataCachePolicy.cpp
@@ -0,0 +1,258 @@
+// Copyright Epic Games, Inc. All Rights Reserved.
+
+#include "DerivedDataCachePolicy.h"
+
+#include "Algo/Accumulate.h"
+#include "Containers/StringConv.h"
+#include "Containers/StringView.h"
+#include "Misc/StringBuilder.h"
+#include "Serialization/CompactBinary.h"
+#include "Serialization/CompactBinaryWriter.h"
+#include "String/ParseTokens.h"
+#include "Templates/Function.h"
+
+namespace UE::DerivedData::Private { class FCacheRecordPolicyShared; }
+
+namespace UE::DerivedData::Private
+{
+
+constexpr ANSICHAR CachePolicyDelimiter = ',';
+
+struct FCachePolicyToText
+{
+	ECachePolicy Policy;
+	FStringView Text;
+};
+
+const FCachePolicyToText CachePolicyToText[]
+{
+	// Flags with multiple bits are ordered by bit count to minimize token count in the text format.
+	{ ECachePolicy::Default,       TEXTVIEW("Default") },
+	{ ECachePolicy::Remote,        TEXTVIEW("Remote") },
+	{ ECachePolicy::Local,         TEXTVIEW("Local") },
+	{ ECachePolicy::Store,         TEXTVIEW("Store") },
+	{ ECachePolicy::Query,         TEXTVIEW("Query") },
+	// Flags with only one bit can be in any order. Match the order in ECachePolicy.
+	{ ECachePolicy::QueryLocal,    TEXTVIEW("QueryLocal") },
+	{ ECachePolicy::QueryRemote,   TEXTVIEW("QueryRemote") },
+	{ ECachePolicy::StoreLocal,    TEXTVIEW("StoreLocal") },
+	{ ECachePolicy::StoreRemote,   TEXTVIEW("StoreRemote") },
+	{ ECachePolicy::SkipMeta,      TEXTVIEW("SkipMeta") },
+	{ ECachePolicy::SkipData,      TEXTVIEW("SkipData") },
+	{ ECachePolicy::PartialRecord, TEXTVIEW("PartialRecord") },
+	{ ECachePolicy::KeepAlive,     TEXTVIEW("KeepAlive") },
+	// None must be last because it matches every policy.
+	{ ECachePolicy::None,          TEXTVIEW("None") },
+};
+
+constexpr ECachePolicy CachePolicyKnownFlags
+	= ECachePolicy::Default
+	| ECachePolicy::SkipMeta
+	| ECachePolicy::SkipData
+	| ECachePolicy::PartialRecord
+	| ECachePolicy::KeepAlive;
+
+template <typename CharType>
+TStringBuilderBase<CharType>& CachePolicyToString(TStringBuilderBase<CharType>& Builder, ECachePolicy Policy)
+{
+	// Mask out unknown flags. None will be written if no flags are known.
+	Policy &= CachePolicyKnownFlags;
+	for (const FCachePolicyToText& Pair : CachePolicyToText)
+	{
+		if (EnumHasAllFlags(Policy, Pair.Policy))
+		{
+			EnumRemoveFlags(Policy, Pair.Policy);
+			Builder << Pair.Text << CachePolicyDelimiter;
+			if (Policy == ECachePolicy::None)
+			{
+				break;
+			}
+		}
+	}
+	Builder.RemoveSuffix(1);
+	return Builder;
+}
+
+template <typename CharType>
+ECachePolicy ParseCachePolicy(const TStringView<CharType> Text)
+{
+	checkf(!Text.IsEmpty(), TEXT("ParseCachePolicy requires a non-empty string."));
+	ECachePolicy Policy = ECachePolicy::None;
+	String::ParseTokens(StringCast<TCHAR, 128>(Text.GetData(), Text.Len()), CachePolicyDelimiter,
+		[&Policy, Index = int32(0)](FStringView Token) mutable
+		{
+			const int32 EndIndex = Index;
+			for (; Index < UE_ARRAY_COUNT(CachePolicyToText); ++Index)
+			{
+				if (CachePolicyToText[Index].Text == Token)
+				{
+					Policy |= CachePolicyToText[Index].Policy;
+					++Index;
+					return;
+				}
+			}
+			for (Index = 0; Index < EndIndex; ++Index)
+			{
+				if (CachePolicyToText[Index].Text == Token)
+				{
+					Policy |= CachePolicyToText[Index].Policy;
+					++Index;
+					return;
+				}
+			}
+		});
+	return Policy;
+}
+
+} // UE::DerivedData::Private
+
+namespace UE::DerivedData
+{
+
+FAnsiStringBuilderBase& operator<<(FAnsiStringBuilderBase& Builder, ECachePolicy Policy) { return Private::CachePolicyToString(Builder, Policy); }
+FWideStringBuilderBase& operator<<(FWideStringBuilderBase& Builder, ECachePolicy Policy) { return Private::CachePolicyToString(Builder, Policy); }
+FUtf8StringBuilderBase& operator<<(FUtf8StringBuilderBase& Builder, ECachePolicy Policy) { return Private::CachePolicyToString(Builder, Policy); }
+
+ECachePolicy ParseCachePolicy(FAnsiStringView Text) { return Private::ParseCachePolicy(Text); }
+ECachePolicy ParseCachePolicy(FWideStringView Text) { return Private::ParseCachePolicy(Text); }
+ECachePolicy ParseCachePolicy(FUtf8StringView Text) { return Private::ParseCachePolicy(Text); }
+
+class Private::FCacheRecordPolicyShared final : public Private::ICacheRecordPolicyShared
+{
+public:
+	inline void AddRef() const final
+	{
+		ReferenceCount.fetch_add(1, std::memory_order_relaxed);
+	}
+
+	inline void Release() const final
+	{
+		if (ReferenceCount.fetch_sub(1, std::memory_order_acq_rel) == 1)
+		{
+			delete this;
+		}
+	}
+
+	inline void AddValuePolicy(const FCacheValuePolicy& Value) final
+	{
+		checkf(Value.Id.IsValid(), TEXT("Failed to add value policy because the ID is null."));
+		const int32 Index = Algo::LowerBoundBy(Values, Value.Id, &FCacheValuePolicy::Id);
+		checkf(!(Values.IsValidIndex(Index) && Values[Index].Id == Value.Id),
+			TEXT("Failed to add value policy with ID %s because it has an existing value policy with that ID. ")
+			TEXT("New: %s. Existing: %s."),
+			*WriteToString<32>(Value.Id), *WriteToString<128>(Value.Policy), *WriteToString<128>(Values[Index].Policy));
+		Values.Insert(Value, Index);
+	}
+
+	inline TConstArrayView<FCacheValuePolicy> GetValuePolicies() const final
+	{
+		return Values;
+	}
+
+private:
+	TArray<FCacheValuePolicy, TInlineAllocator<14>> Values;
+	mutable std::atomic<uint32> ReferenceCount{0};
+};
+
+ECachePolicy FCacheRecordPolicy::GetValuePolicy(const FValueId& Id) const
+{
+	if (Shared)
+	{
+		if (TConstArrayView<FCacheValuePolicy> Values = Shared->GetValuePolicies(); !Values.IsEmpty())
+		{
+			if (int32 Index = Algo::BinarySearchBy(Values, Id, &FCacheValuePolicy::Id); Index != INDEX_NONE)
+			{
+				return Values[Index].Policy;
+			}
+		}
+	}
+	return DefaultPolicy;
+}
+
+FCacheRecordPolicy FCacheRecordPolicy::Transform(const TFunctionRef<ECachePolicy (ECachePolicy)> Op) const
+{
+	if (IsUniform())
+	{
+		return Op(RecordPolicy);
+	}
+
+	FCacheRecordPolicyBuilder Builder(Op(DefaultPolicy));
+	for (const FCacheValuePolicy& Value : GetValuePolicies())
+	{
+		Builder.AddValuePolicy({Value.Id, Op(Value.Policy)});
+	}
+	return Builder.Build();
+}
+
+void FCacheRecordPolicy::Save(FCbWriter& Writer) const
+{
+	Writer.BeginObject();
+	Writer.AddString("DefaultValuePolicy"_ASV, WriteToUtf8String<128>(GetDefaultPolicy()));
+	if (!IsUniform())
+	{
+		Writer.BeginArray("ValuePolicies"_ASV);
+		for (const FCacheValuePolicy& Value : GetValuePolicies())
+		{
+			Writer.BeginObject();
+			Writer.AddObjectId("Id"_ASV, Value.Id);
+			Writer.AddString("Policy"_ASV, WriteToUtf8String<128>(Value.Policy));
+			Writer.EndObject();
+		}
+		Writer.EndArray();
+	}
+	Writer.EndObject();
+}
+
+FOptionalCacheRecordPolicy FCacheRecordPolicy::Load(const FCbObjectView Object)
+{
+	const FUtf8StringView DefaultPolicyText = Object["DefaultValuePolicy"_ASV].AsString();
+	if (DefaultPolicyText.IsEmpty())
+	{
+		return {};
+	}
+
+	FCacheRecordPolicyBuilder Builder(ParseCachePolicy(DefaultPolicyText));
+	for (const FCbFieldView Value : Object["ValuePolicies"_ASV])
+	{
+		const FValueId Id = Value["Id"_ASV].AsObjectId();
+		const FUtf8StringView PolicyText = Value["Policy"_ASV].AsString();
+		if (Id.IsNull() || PolicyText.IsEmpty())
+		{
+			return {};
+		}
+		Builder.AddValuePolicy(Id, ParseCachePolicy(PolicyText));
+	}
+	return Builder.Build();
+}
+
+void FCacheRecordPolicyBuilder::AddValuePolicy(const FCacheValuePolicy& Value)
+{
+	constexpr ECachePolicy RecordOnlyPolicyFlags = ECachePolicy::SkipMeta | ECachePolicy::PartialRecord | ECachePolicy::KeepAlive;
+	checkf(!EnumHasAnyFlags(Value.Policy, RecordOnlyPolicyFlags),
+		TEXT("Value policy contains flags that only make sense on the record policy. Policy: %s"),
+		*WriteToString<128>(Value.Policy));
+	if (Value.Policy == (BasePolicy & ~RecordOnlyPolicyFlags))
+	{
+		return;
+	}
+	if (!Shared)
+	{
+		Shared = new Private::FCacheRecordPolicyShared;
+	}
+	Shared->AddValuePolicy(Value);
+}
+
+FCacheRecordPolicy FCacheRecordPolicyBuilder::Build()
+{
+	FCacheRecordPolicy Policy(BasePolicy);
+	if (Shared)
+	{
+		const auto PolicyOr = [](ECachePolicy A, ECachePolicy B) { return A | (B & ~ECachePolicy::SkipData); };
+		const TConstArrayView<FCacheValuePolicy> Values = Shared->GetValuePolicies();
+		Policy.RecordPolicy = Algo::TransformAccumulate(Values, &FCacheValuePolicy::Policy, BasePolicy, PolicyOr);
+		Policy.Shared = MoveTemp(Shared);
+	}
+	return Policy;
+}
+
+} // UE::DerivedData
diff --git a/Engine/Source/Developer/DerivedDataCache/Private/DerivedDataCacheRecord.cpp b/Engine/Source/Developer/DerivedDataCache/Private/DerivedDataCacheRecord.cpp
--- a/Engine/Source/Developer/DerivedDataCache/Private/DerivedDataCacheRecord.cpp
+++ b/Engine/Source/Developer/DerivedDataCache/Private/DerivedDataCacheRecord.cpp
@@ -3,6 +3,7 @@
 #include "DerivedDataCacheRecord.h"
 
 #include "Algo/Accumulate.h"
+#include "Algo/BinarySearch.h"
 #include "Algo/IsSorted.h"
 #include "DerivedDataCacheKey.h"
 #include "DerivedDataCachePrivate.h"
diff --git a/Engine/Source/Developer/DerivedDataCache/Private/DerivedDataLegacyCacheStore.h b/Engine/Source/Developer/DerivedDataCache/Private/DerivedDataLegacyCacheStore.h
--- a/Engine/Source/Developer/DerivedDataCache/Private/DerivedDataLegacyCacheStore.h
+++ b/Engine/Source/Developer/DerivedDataCache/Private/DerivedDataLegacyCacheStore.h
@@ -2,7 +2,6 @@
 
 #pragma once
 
-#include "DerivedDataCache.h"
 #include "DerivedDataCacheStore.h"
 #include "Logging/LogMacros.h"
 
diff --git a/Engine/Source/Developer/DerivedDataCache/Private/FileSystemDerivedDataBackend.cpp b/Engine/Source/Developer/DerivedDataCache/Private/FileSystemDerivedDataBackend.cpp
--- a/Engine/Source/Developer/DerivedDataCache/Private/FileSystemDerivedDataBackend.cpp
+++ b/Engine/Source/Developer/DerivedDataCache/Private/FileSystemDerivedDataBackend.cpp
@@ -1780,7 +1780,7 @@ bool FFileSystemCacheStore::PutCacheRecord(
 
 	const FCacheKey& Key = Record.GetKey();
 	const ECachePolicy CombinedValuePolicy = Algo::TransformAccumulate(
-		Policy.GetValuePolicies(), &FCacheValuePolicy::Policy, Policy.GetDefaultValuePolicy(), UE_PROJECTION(operator|));
+		Policy.GetValuePolicies(), &FCacheValuePolicy::Policy, Policy.GetDefaultPolicy(), UE_PROJECTION(operator|));
 	const ECachePolicy CombinedPolicy = Policy.GetRecordPolicy() | CombinedValuePolicy;
 
 	// Skip the request if storing to the cache is disabled.
diff --git a/Engine/Source/Developer/DerivedDataCache/Private/HttpDerivedDataBackend.cpp b/Engine/Source/Developer/DerivedDataCache/Private/HttpDerivedDataBackend.cpp
--- a/Engine/Source/Developer/DerivedDataCache/Private/HttpDerivedDataBackend.cpp
+++ b/Engine/Source/Developer/DerivedDataCache/Private/HttpDerivedDataBackend.cpp
@@ -2451,7 +2451,7 @@ bool FHttpDerivedDataBackend::PutCacheRecord(
 
 	const FCacheKey& Key = Record.GetKey();
 	const ECachePolicy CombinedValuePolicy = Algo::TransformAccumulate(
-		Policy.GetValuePolicies(), &FCacheValuePolicy::Policy, Policy.GetDefaultValuePolicy(), UE_PROJECTION(operator|));
+		Policy.GetValuePolicies(), &FCacheValuePolicy::Policy, Policy.GetDefaultPolicy(), UE_PROJECTION(operator|));
 	const ECachePolicy CombinedPolicy = Policy.GetRecordPolicy() | CombinedValuePolicy;
 
 	// Skip the request if storing to the cache is disabled.
diff --git a/Engine/Source/Developer/DerivedDataCache/Private/PakFileDerivedDataBackend.cpp b/Engine/Source/Developer/DerivedDataCache/Private/PakFileDerivedDataBackend.cpp
--- a/Engine/Source/Developer/DerivedDataCache/Private/PakFileDerivedDataBackend.cpp
+++ b/Engine/Source/Developer/DerivedDataCache/Private/PakFileDerivedDataBackend.cpp
@@ -798,7 +798,7 @@ bool FPakFileDerivedDataBackend::PutCacheRecord(
 	TStringBuilder<256> Path;
 	FPathViews::Append(Path, TEXT("Buckets"), Key);
 	const ECachePolicy CombinedValuePolicy = Algo::TransformAccumulate(
-		Policy.GetValuePolicies(), &FCacheValuePolicy::Policy, Policy.GetDefaultValuePolicy(), UE_PROJECTION(operator|));
+		Policy.GetValuePolicies(), &FCacheValuePolicy::Policy, Policy.GetDefaultPolicy(), UE_PROJECTION(operator|));
 	if (EnumHasAnyFlags(CombinedValuePolicy, ECachePolicy::SkipData))
 	{
 		bRecordExists = FileExists(Path);
diff --git a/Engine/Source/Developer/DerivedDataCache/Public/DerivedDataBuildTypes.h b/Engine/Source/Developer/DerivedDataCache/Public/DerivedDataBuildTypes.h
--- a/Engine/Source/Developer/DerivedDataCache/Public/DerivedDataBuildTypes.h
+++ b/Engine/Source/Developer/DerivedDataCache/Public/DerivedDataBuildTypes.h
@@ -97,7 +97,7 @@ public:
 	/** Construct a build policy with a single policy for every value. */
 	inline FBuildPolicy(EBuildPolicy Policy)
 		: CombinedPolicy(Policy)
-		, DefaultValuePolicy(Policy)
+		, DefaultPolicy(Policy)
 	{
 	}
 
@@ -107,12 +107,12 @@ public:
 	/** Returns the build policy combined from the value policies. */
 	inline EBuildPolicy GetCombinedPolicy() const { return CombinedPolicy; }
 
+	/** Returns the build policy to use for values with no override. */
+	inline EBuildPolicy GetDefaultPolicy() const { return DefaultPolicy; }
+
 	/** Returns the build policy to use for the value. */
 	UE_API EBuildPolicy GetValuePolicy(const FValueId& Id) const;
 
-	/** Returns the build policy to use for values with no override. */
-	inline EBuildPolicy GetDefaultValuePolicy() const { return DefaultValuePolicy; }
-
 	/** Returns the array of build policy overrides for values, sorted by ID. */
 	inline TConstArrayView<FBuildValuePolicy> GetValuePolicies() const
 	{
@@ -126,7 +126,7 @@ private:
 	friend class FBuildPolicyBuilder;
 
 	EBuildPolicy CombinedPolicy = EBuildPolicy::Default;
-	EBuildPolicy DefaultValuePolicy = EBuildPolicy::Default;
+	EBuildPolicy DefaultPolicy = EBuildPolicy::Default;
 	TRefCountPtr<const Private::IBuildPolicyShared> Shared;
 };
 
diff --git a/Engine/Source/Developer/DerivedDataCache/Public/DerivedDataCache.h b/Engine/Source/Developer/DerivedDataCache/Public/DerivedDataCache.h
--- a/Engine/Source/Developer/DerivedDataCache/Public/DerivedDataCache.h
+++ b/Engine/Source/Developer/DerivedDataCache/Public/DerivedDataCache.h
@@ -5,6 +5,7 @@
 #include "CoreTypes.h"
 #include "Containers/StringFwd.h"
 #include "DerivedDataCacheKey.h"
+#include "DerivedDataCachePolicy.h"
 #include "DerivedDataCacheRecord.h"
 #include "DerivedDataRequestTypes.h"
 #include "DerivedDataSharedString.h"
@@ -12,14 +13,10 @@
 #include "DerivedDataValueId.h"
 #include "Math/NumericLimits.h"
 #include "Memory/SharedBuffer.h"
-#include "Misc/EnumClassFlags.h"
 #include "Templates/Function.h"
-#include "Templates/RefCounting.h"
 
 #define UE_API DERIVEDDATACACHE_API
 
-class FCbObjectView;
-class FCbWriter;
 namespace UE::DerivedData { class ICacheStoreMaintainer; }
 namespace UE::DerivedData { class IRequestOwner; }
 namespace UE::DerivedData { struct FCacheGetChunkRequest; }
@@ -32,7 +29,6 @@ namespace UE::DerivedData { struct FCachePutRequest; }
 namespace UE::DerivedData { struct FCachePutResponse; }
 namespace UE::DerivedData { struct FCachePutValueRequest; }
 namespace UE::DerivedData { struct FCachePutValueResponse; }
-namespace UE::DerivedData::Private { class ICacheRecordPolicyShared; }
 
 namespace UE::DerivedData
 {
@@ -45,198 +41,6 @@ using FOnCacheGetChunkComplete = TUniqueFunction<void (FCacheGetChunkResponse&&
 
 ///////////////////////////////////////////////////////////////////////////////////////////////////////////////////////
 
-/**
- * Flags to control the behavior of cache requests.
- *
- * The cache policy flags can be combined to support a variety of usage patterns. Examples:
- *
- * Get(Default): read from any cache; put to writable caches if missing.
- * Get(Remote): read any remote cache; put to writable remote caches if missing.
- * Get(Local): read any local cache; put to writable local caches if missing.
- * Get(Query | StoreRemote): read from any cache; put to writable remote caches if missing.
- * Get(Query | StoreLocal): read from any cache; put to writable local caches if missing.
- * Get(Query | SkipData): check for existence in any cache; do not modify any cache.
- * Get(Default | SkipData): check for existence in any cache; put to writable caches if missing.
- *
- * Put(Default): write to every writable cache.
- * Put(Remote): write to every writable remote cache, skipping local caches.
- * Put(Local): write to every writable local cache, skipping remote caches.
- */
-enum class ECachePolicy : uint32
-{
-	/** A value without any flags set. */
-	None            = 0,
-
-	/** Allow a cache request to query local caches. */
-	QueryLocal      = 1 << 0,
-	/** Allow a cache request to query remote caches. */
-	QueryRemote     = 1 << 1,
-	/** Allow a cache request to query any caches. */
-	Query           = QueryLocal | QueryRemote,
-
-	/** Allow cache records and values to be stored in local caches. */
-	StoreLocal      = 1 << 2,
-	/** Allow cache records and values to be stored in remote caches. */
-	StoreRemote     = 1 << 3,
-	/** Allow cache records and values to be stored in any caches. */
-	Store           = StoreLocal | StoreRemote,
-
-	/** Skip fetching the metadata for record requests. */
-	SkipMeta        = 1 << 4,
-	/** Skip fetching the data for values. */
-	SkipData        = 1 << 5,
-
-	/**
-	 * Partial output will be provided with the error status when a required value is missing.
-	 *
-	 * This is meant for cases when the missing values can be individually recovered, or rebuilt,
-	 * without rebuilding the whole record. The cache automatically adds this flag when there are
-	 * other cache stores that it may be able to recover missing values from.
-	 *
-	 * Missing values will be returned in the records or chunks, but with only the hash and size.
-	 *
-	 * Applying this flag for a put of a record allows a partial record to be stored.
-	 */
-	PartialRecord   = 1 << 6,
-
-	/**
-	 * Keep records in the cache for at least the duration of the session.
-	 *
-	 * This is a hint that the record may be accessed again in this session. This is mainly meant
-	 * to be used when subsequent accesses will not tolerate a cache miss.
-	 */
-	KeepAlive       = 1 << 7,
-
-	/** Allow cache requests to query and store records and values in local caches. */
-	Local           = QueryLocal | StoreLocal,
-	/** Allow cache requests to query and store records and values in remote caches. */
-	Remote          = QueryRemote | StoreRemote,
-
-	/** Allow cache requests to query and store records and values in any caches. */
-	Default         = Query | Store,
-
-	/** Do not allow cache requests to query or store records and values in any caches. */
-	Disable         = None,
-};
-
-ENUM_CLASS_FLAGS(ECachePolicy);
-/** Serialize Policy to text and append to Builder. Appended text will not be empty. */
-UE_API FUtf8StringBuilderBase& operator<<(FUtf8StringBuilderBase& Builder, ECachePolicy Policy);
-UE_API FAnsiStringBuilderBase& operator<<(FAnsiStringBuilderBase& Builder, ECachePolicy Policy);
-UE_API FWideStringBuilderBase& operator<<(FWideStringBuilderBase& Builder, ECachePolicy Policy);
-/** Parse text written by operator<< back into an ECachePolicy. Text must not be empty. */
-UE_API ECachePolicy ParseCachePolicy(FUtf8StringView Text);
-UE_API ECachePolicy ParseCachePolicy(FAnsiStringView Text);
-UE_API ECachePolicy ParseCachePolicy(FWideStringView Text);
-
-///////////////////////////////////////////////////////////////////////////////////////////////////////////////////////
-
-/** A value ID and the cache policy to use for that value. */
-struct FCacheValuePolicy
-{
-	FValueId Id;
-	ECachePolicy Policy = ECachePolicy::Default;
-};
-
-/** Interface for the private implementation of the cache record policy. */
-class Private::ICacheRecordPolicyShared
-{
-public:
-	virtual ~ICacheRecordPolicyShared() = default;
-	virtual void AddRef() const = 0;
-	virtual void Release() const = 0;
-	virtual TConstArrayView<FCacheValuePolicy> GetValuePolicies() const = 0;
-	virtual void AddValuePolicy(const FCacheValuePolicy& Policy) = 0;
-	virtual void Build() = 0;
-};
-
-/**
- * Flags to control the behavior of cache record requests, with optional overrides by value.
- *
- * Examples:
- * - A base policy of Disable, with value policy overrides of Default, will fetch those values if
- *   they exist in the record, and skip data for any other values.
- * - A base policy of Default, with value policy overrides of (Query | SkipData), will skip those
- *   values, but still check if they exist, and will load any other values.
- */
-class FCacheRecordPolicy
-{
-public:
-	/** Construct a cache record policy that uses the default policy. */
-	FCacheRecordPolicy() = default;
-
-	/** Construct a cache record policy with a uniform policy for the record and every value. */
-	inline FCacheRecordPolicy(ECachePolicy Policy)
-		: RecordPolicy(Policy)
-		, DefaultValuePolicy(Policy)
-	{
-	}
-
-	/** Returns true if the record and every value use the same cache policy. */
-	inline bool IsUniform() const { return !Shared && RecordPolicy == DefaultValuePolicy; }
-
-	/** Returns the cache policy to use for the record. */
-	inline ECachePolicy GetRecordPolicy() const { return RecordPolicy; }
-
-	/** Returns the cache policy to use for the value. */
-	UE_API ECachePolicy GetValuePolicy(const FValueId& Id) const;
-
-	/** Returns the cache policy to use for values with no override. */
-	inline ECachePolicy GetDefaultValuePolicy() const { return DefaultValuePolicy; }
-
-	/** Returns the array of cache policy overrides for values, sorted by ID. */
-	inline TConstArrayView<FCacheValuePolicy> GetValuePolicies() const
-	{
-		return Shared ? Shared->GetValuePolicies() : TConstArrayView<FCacheValuePolicy>();
-	}
-
-	/** Returns a copy of this policy transformed by an operation. */
-	UE_API FCacheRecordPolicy Transform(TFunctionRef<ECachePolicy (ECachePolicy)> Op) const;
-
-	/** Save the values from *this into the given writer. */
-	UE_API void Save(FCbWriter& Writer) const;
-
-	/**
-	 * Returns a policy loaded from values on Object.
-	 * Invalid data will result in a uniform CacheRecordPolicy with defaultValuePolicy == DefaultPolicy.
-	 */
-	static UE_API FCacheRecordPolicy Load(FCbObjectView Object, ECachePolicy DefaultPolicy = ECachePolicy::Default);
-
-private:
-	friend class FCacheRecordPolicyBuilder;
-
-	ECachePolicy RecordPolicy = ECachePolicy::Default;
-	ECachePolicy DefaultValuePolicy = ECachePolicy::Default;
-	TRefCountPtr<const Private::ICacheRecordPolicyShared> Shared;
-};
-
-/** A cache record policy builder is used to construct a cache record policy. */
-class FCacheRecordPolicyBuilder
-{
-public:
-	/** Construct a policy builder that uses the default policy as its base policy. */
-	FCacheRecordPolicyBuilder() = default;
-
-	/** Construct a policy builder that uses the provided policy for the record and values with no override. */
-	inline explicit FCacheRecordPolicyBuilder(ECachePolicy Policy)
-		: BasePolicy(Policy)
-	{
-	}
-
-	/** Adds a cache policy override for a value. */
-	UE_API void AddValuePolicy(const FCacheValuePolicy& Policy);
-	inline void AddValuePolicy(const FValueId& Id, ECachePolicy Policy) { AddValuePolicy({Id, Policy}); }
-
-	/** Build a cache record policy, which makes this builder subsequently unusable. */
-	UE_API FCacheRecordPolicy Build();
-
-private:
-	ECachePolicy BasePolicy = ECachePolicy::Default;
-	TRefCountPtr<Private::ICacheRecordPolicyShared> Shared;
-};
-
-///////////////////////////////////////////////////////////////////////////////////////////////////////////////////////
-
 /**
  * Interface to the cache.
  *
diff --git a/Engine/Source/Developer/DerivedDataCache/Public/DerivedDataCachePolicy.h b/Engine/Source/Developer/DerivedDataCache/Public/DerivedDataCachePolicy.h
new file mode 100644
--- /dev/null
+++ b/Engine/Source/Developer/DerivedDataCache/Public/DerivedDataCachePolicy.h
@@ -0,0 +1,236 @@
+// Copyright Epic Games, Inc. All Rights Reserved.
+
+#pragma once
+
+#include "CoreTypes.h"
+#include "Containers/ArrayView.h"
+#include "Containers/StringFwd.h"
+#include "DerivedDataValueId.h"
+#include "Misc/EnumClassFlags.h"
+#include "Templates/RefCounting.h"
+#include "Templates/UnrealTemplate.h"
+
+#define UE_API DERIVEDDATACACHE_API
+
+class FCbObjectView;
+class FCbWriter;
+
+template <typename FuncType> class TFunctionRef;
+
+namespace UE::DerivedData { class FOptionalCacheRecordPolicy; }
+namespace UE::DerivedData::Private { class ICacheRecordPolicyShared; }
+
+namespace UE::DerivedData
+{
+
+/**
+ * Flags to control the behavior of cache requests.
+ *
+ * Flags can be combined to cover a variety of requirements. Examples:
+ *
+ * Get(Default): Fetch from any cache. Store the response to any caches if missing.
+ * Get(Local): Fetch from any local cache. Store the response to any local caches if missing.
+ * Get(Query | StoreLocal): Fetch from any cache. Store response to any local caches if missing.
+ * Get(Query | SkipData): Check for existence in any cache. Do not store to any caches if missing.
+ * Get(Default | SkipData): Check for existence in any cache. Store the response to any caches if missing.
+ *
+ * Put(Default): Store to every cache.
+ * Put(Local): Store to every local cache, skipping remote caches.
+ */
+enum class ECachePolicy : uint32
+{
+	/** A value with no flags. Disables access to the cache unless combined with other flags. */
+	None            = 0,
+
+	/** Allow a cache request to query local caches. */
+	QueryLocal      = 1 << 0,
+	/** Allow a cache request to query remote caches. */
+	QueryRemote     = 1 << 1,
+	/** Allow a cache request to query any caches. */
+	Query           = QueryLocal | QueryRemote,
+
+	/** Allow cache records and values to be stored in local caches. */
+	StoreLocal      = 1 << 2,
+	/** Allow cache records and values to be stored in remote caches. */
+	StoreRemote     = 1 << 3,
+	/** Allow cache records and values to be stored in any caches. */
+	Store           = StoreLocal | StoreRemote,
+
+	/** Allow cache requests to query and store records and values in local caches. */
+	Local           = QueryLocal | StoreLocal,
+	/** Allow cache requests to query and store records and values in remote caches. */
+	Remote          = QueryRemote | StoreRemote,
+
+	/** Allow cache requests to query and store records and values in any caches. */
+	Default         = Query | Store,
+
+	/** Skip fetching the data for values. */
+	SkipData        = 1 << 4,
+
+	/** Skip fetching the metadata for record requests. */
+	SkipMeta        = 1 << 5,
+
+	/**
+	 * Partial output will be provided with the error status when a required value is missing.
+	 *
+	 * This is meant for cases when the missing values can be individually recovered, or rebuilt,
+	 * without rebuilding the whole record. The cache automatically adds this flag when there are
+	 * other cache stores that it may be able to recover missing values from.
+	 *
+	 * Missing values will be returned in the records, but with only the hash and size.
+	 *
+	 * Applying this flag for a put of a record allows a partial record to be stored.
+	 */
+	PartialRecord   = 1 << 6,
+
+	/**
+	 * Keep records in the cache for at least the duration of the session.
+	 *
+	 * This is a hint that the record may be accessed again in this session. This is mainly meant
+	 * to be used when subsequent accesses will not tolerate a cache miss.
+	 */
+	KeepAlive       = 1 << 7,
+};
+
+ENUM_CLASS_FLAGS(ECachePolicy);
+
+/** Append a non-empty text version of the policy to the builder. */
+UE_API FAnsiStringBuilderBase& operator<<(FAnsiStringBuilderBase& Builder, ECachePolicy Policy);
+UE_API FWideStringBuilderBase& operator<<(FWideStringBuilderBase& Builder, ECachePolicy Policy);
+UE_API FUtf8StringBuilderBase& operator<<(FUtf8StringBuilderBase& Builder, ECachePolicy Policy);
+
+/** Parse non-empty text written by operator<< into a policy. */
+UE_API ECachePolicy ParseCachePolicy(FAnsiStringView Text);
+UE_API ECachePolicy ParseCachePolicy(FWideStringView Text);
+UE_API ECachePolicy ParseCachePolicy(FUtf8StringView Text);
+
+/** A value ID and the cache policy to use for that value. */
+struct FCacheValuePolicy
+{
+	FValueId Id;
+	ECachePolicy Policy = ECachePolicy::Default;
+};
+
+/** Interface for the private implementation of the cache record policy. */
+class Private::ICacheRecordPolicyShared
+{
+public:
+	virtual ~ICacheRecordPolicyShared() = default;
+	virtual void AddRef() const = 0;
+	virtual void Release() const = 0;
+	virtual void AddValuePolicy(const FCacheValuePolicy& Policy) = 0;
+	virtual TConstArrayView<FCacheValuePolicy> GetValuePolicies() const = 0;
+};
+
+/**
+ * Flags to control the behavior of cache record requests, with optional overrides by value.
+ *
+ * Examples:
+ * - A base policy of None with value policy overrides of Default will fetch those values if they
+ *   exist in the record, and skip data for any other values.
+ * - A base policy of Default, with value policy overrides of (Query | SkipData), will skip those
+ *   values, but still check if they exist, and will load any other values.
+ */
+class FCacheRecordPolicy
+{
+public:
+	/** Construct a cache record policy that uses the default policy. */
+	FCacheRecordPolicy() = default;
+
+	/** Construct a cache record policy with a uniform policy for the record and every value. */
+	inline FCacheRecordPolicy(ECachePolicy Policy)
+		: RecordPolicy(Policy)
+		, DefaultPolicy(Policy)
+	{
+	}
+
+	/** Returns true if the record and every value use the same cache policy. */
+	inline bool IsUniform() const { return !Shared; }
+
+	/** Returns the cache policy to use for the record. */
+	inline ECachePolicy GetRecordPolicy() const { return RecordPolicy; }
+
+	/** Returns the cache policy to use for values with no override. */
+	inline ECachePolicy GetDefaultPolicy() const { return DefaultPolicy; }
+
+	/** Returns the cache policy to use for the value. */
+	UE_API ECachePolicy GetValuePolicy(const FValueId& Id) const;
+
+	/** Returns the array of cache policy overrides for values, sorted by ID. */
+	inline TConstArrayView<FCacheValuePolicy> GetValuePolicies() const
+	{
+		return Shared ? Shared->GetValuePolicies() : TConstArrayView<FCacheValuePolicy>();
+	}
+
+	/** Returns a copy of this policy transformed by an operation. */
+	UE_API FCacheRecordPolicy Transform(TFunctionRef<ECachePolicy (ECachePolicy)> Op) const;
+
+	/** Saves the cache record policy to a compact binary object. */
+	UE_API void Save(FCbWriter& Writer) const;
+
+	/** Loads a cache record policy from an object. */
+	UE_API static FOptionalCacheRecordPolicy Load(FCbObjectView Object);
+
+private:
+	friend class FCacheRecordPolicyBuilder;
+	friend class FOptionalCacheRecordPolicy;
+
+	ECachePolicy RecordPolicy = ECachePolicy::Default;
+	ECachePolicy DefaultPolicy = ECachePolicy::Default;
+	TRefCountPtr<const Private::ICacheRecordPolicyShared> Shared;
+};
+
+/** A cache record policy builder is used to construct a cache record policy. */
+class FCacheRecordPolicyBuilder
+{
+public:
+	/** Construct a policy builder that uses the default policy as its base policy. */
+	FCacheRecordPolicyBuilder() = default;
+
+	/** Construct a policy builder that uses the provided policy for the record and values with no override. */
+	inline explicit FCacheRecordPolicyBuilder(ECachePolicy Policy)
+		: BasePolicy(Policy)
+	{
+	}
+
+	/** Adds a cache policy override for a value. */
+	UE_API void AddValuePolicy(const FCacheValuePolicy& Policy);
+	inline void AddValuePolicy(const FValueId& Id, ECachePolicy Policy) { AddValuePolicy({Id, Policy}); }
+
+	/** Build a cache record policy, which makes this builder subsequently unusable. */
+	UE_API FCacheRecordPolicy Build();
+
+private:
+	ECachePolicy BasePolicy = ECachePolicy::Default;
+	TRefCountPtr<Private::ICacheRecordPolicyShared> Shared;
+};
+
+/**
+ * A cache record policy that can be null.
+ *
+ * @see FCacheRecordPolicy
+ */
+class FOptionalCacheRecordPolicy : private FCacheRecordPolicy
+{
+public:
+	inline FOptionalCacheRecordPolicy() : FCacheRecordPolicy(~ECachePolicy::None) {}
+
+	inline FOptionalCacheRecordPolicy(FCacheRecordPolicy&& InOutput) : FCacheRecordPolicy(MoveTemp(InOutput)) {}
+	inline FOptionalCacheRecordPolicy(const FCacheRecordPolicy& InOutput) : FCacheRecordPolicy(InOutput) {}
+	inline FOptionalCacheRecordPolicy& operator=(FCacheRecordPolicy&& InOutput) { FCacheRecordPolicy::operator=(MoveTemp(InOutput)); return *this; }
+	inline FOptionalCacheRecordPolicy& operator=(const FCacheRecordPolicy& InOutput) { FCacheRecordPolicy::operator=(InOutput); return *this; }
+
+	/** Returns the cache record policy. The caller must check for null before using this accessor. */
+	inline const FCacheRecordPolicy& Get() const & { return *this; }
+	inline FCacheRecordPolicy Get() && { return MoveTemp(*this); }
+
+	inline bool IsNull() const { return RecordPolicy == ~ECachePolicy::None; }
+	inline bool IsValid() const { return !IsNull(); }
+	inline explicit operator bool() const { return !IsNull(); }
+
+	inline void Reset() { *this = FOptionalCacheRecordPolicy(); }
+};
+
+} // namespace UE::DerivedData
+
+#undef UE_API
diff --git a/Engine/Source/Runtime/Engine/Private/StaticMesh.cpp b/Engine/Source/Runtime/Engine/Private/StaticMesh.cpp
--- a/Engine/Source/Runtime/Engine/Private/StaticMesh.cpp
+++ b/Engine/Source/Runtime/Engine/Private/StaticMesh.cpp
@@ -2743,7 +2743,7 @@ void FStaticMeshRenderData::Cache(const ITargetPlatform* TargetPlatform, UStatic
 		FSharedBuffer MeshDataBuffer;
 		FIoHash NaniteStreamingDataHash;
 		{
-			FCacheRecordPolicyBuilder PolicyBuilder(ECachePolicy::Default);
+			FCacheRecordPolicyBuilder PolicyBuilder(ECachePolicy::Default | ECachePolicy::KeepAlive);
 			PolicyBuilder.AddValuePolicy(NaniteStreamingDataId, ECachePolicy::Default | ECachePolicy::SkipData);
 			
 			FCacheGetRequest Request;
@@ -2891,11 +2891,8 @@ void FStaticMeshRenderData::Cache(const ITargetPlatform* TargetPlatform, UStatic
 				RecordBuilder.AddValue(MeshDataId, FSharedBuffer::MakeView(Ar.GetData(), Ar.TotalSize()));
 				TotalPushedBytes += Ar.TotalSize();
 
-				FCacheRecordPolicyBuilder PolicyBuilder(ECachePolicy::Default);
-				PolicyBuilder.AddValuePolicy(NaniteStreamingDataId, ECachePolicy::Default | ECachePolicy::KeepAlive);
-
 				FRequestOwner RequestOwner(UE::DerivedData::EPriority::Blocking);
-				const FCachePutRequest PutRequest = { FSharedString(TEXT("StaticMesh")), RecordBuilder.Build(), PolicyBuilder.Build() };
+				const FCachePutRequest PutRequest = { FSharedString(TEXT("StaticMesh")), RecordBuilder.Build(), ECachePolicy::Default | ECachePolicy::KeepAlive };
 				GetCache().Put(MakeArrayView(&PutRequest, 1),
 					RequestOwner,
 					[](FCachePutResponse&& Response)
