DDC: Masked record-only cache policy flags out of the value policy

Fixing backends to respect the masked flags includes a few fixes to handling of ~Query by Put and PutValue.

#rnx

diff --git a/Engine/Source/Developer/DerivedDataCache/Private/DerivedDataBuild.cpp b/Engine/Source/Developer/DerivedDataCache/Private/DerivedDataBuild.cpp
--- a/Engine/Source/Developer/DerivedDataCache/Private/DerivedDataBuild.cpp
+++ b/Engine/Source/Developer/DerivedDataCache/Private/DerivedDataBuild.cpp
@@ -3,6 +3,7 @@
 #include "DerivedDataBuild.h"
 
 #include "Algo/Accumulate.h"
+#include "Algo/BinarySearch.h"
 #include "DerivedDataBuildAction.h"
 #include "DerivedDataBuildDefinition.h"
 #include "DerivedDataBuildFunctionRegistry.h"
@@ -158,14 +159,14 @@ public:
 		return Values;
 	}
 
-	inline void AddValuePolicy(const FBuildValuePolicy& Policy) final
+	inline void AddValuePolicy(const FBuildValuePolicy& Value) final
 	{
-		Values.Add(Policy);
-	}
-
-	inline void Build() final
-	{
-		Algo::SortBy(Values, &FBuildValuePolicy::Id);
+		checkf(Value.Id.IsValid(), TEXT("Failed to add value policy because the ID is null."));
+		const int32 Index = Algo::LowerBoundBy(Values, Value.Id, &FBuildValuePolicy::Id);
+		checkf(!(Values.IsValidIndex(Index) && Values[Index].Id == Value.Id),
+			TEXT("Failed to add value policy with ID %s because it has an existing value policy with that ID."),
+			*WriteToString<32>(Value.Id));
+		Values.Insert(Value, Index);
 	}
 
 private:
@@ -177,12 +178,10 @@ EBuildPolicy FBuildPolicy::GetValuePolicy(const FValueId& Id) const
 {
 	if (Shared)
 	{
-		if (TConstArrayView<FBuildValuePolicy> Values = Shared->GetValuePolicies(); !Values.IsEmpty())
+		const TConstArrayView<FBuildValuePolicy> Values = Shared->GetValuePolicies();
+		if (const int32 Index = Algo::BinarySearchBy(Values, Id, &FBuildValuePolicy::Id); Index != INDEX_NONE)
 		{
-			if (int32 Index = Algo::BinarySearchBy(Values, Id, &FBuildValuePolicy::Id); Index != INDEX_NONE)
-			{
-				return Values[Index].Policy;
-			}
+			return Values[Index].Policy;
 		}
 	}
 	return DefaultPolicy;
@@ -202,13 +201,17 @@ FBuildPolicy FBuildPolicy::Transform(TFunctionRef<EBuildPolicy (EBuildPolicy)> O
 	return Builder.Build();
 }
 
-void FBuildPolicyBuilder::AddValuePolicy(const FBuildValuePolicy& Policy)
+void FBuildPolicyBuilder::AddValuePolicy(const FBuildValuePolicy& Value)
 {
+	if (Value.Policy == BasePolicy)
+	{
+		return;
+	}
 	if (!Shared)
 	{
 		Shared = new Private::FBuildPolicyShared;
 	}
-	Shared->AddValuePolicy(Policy);
+	Shared->AddValuePolicy(Value);
 }
 
 FBuildPolicy FBuildPolicyBuilder::Build()
@@ -216,10 +219,12 @@ FBuildPolicy FBuildPolicyBuilder::Build()
 	FBuildPolicy Policy(BasePolicy);
 	if (Shared)
 	{
-		Shared->Build();
-		const auto PolicyOr = [](EBuildPolicy A, EBuildPolicy B) { return A | (B & EBuildPolicy::Default); };
+		const auto Add = [](const EBuildPolicy A, const EBuildPolicy B)
+		{
+			return ((A | B) & ~EBuildPolicy::SkipData) | ((A & B) & EBuildPolicy::SkipData);
+		};
 		const TConstArrayView<FBuildValuePolicy> Values = Shared->GetValuePolicies();
-		Policy.CombinedPolicy = Algo::TransformAccumulate(Values, &FBuildValuePolicy::Policy, BasePolicy, PolicyOr);
+		Policy.CombinedPolicy = Algo::TransformAccumulate(Values, &FBuildValuePolicy::Policy, BasePolicy, Add);
 		Policy.Shared = MoveTemp(Shared);
 	}
 	return Policy;
diff --git a/Engine/Source/Developer/DerivedDataCache/Private/DerivedDataBuildJob.cpp b/Engine/Source/Developer/DerivedDataCache/Private/DerivedDataBuildJob.cpp
--- a/Engine/Source/Developer/DerivedDataCache/Private/DerivedDataBuildJob.cpp
+++ b/Engine/Source/Developer/DerivedDataCache/Private/DerivedDataBuildJob.cpp
@@ -493,10 +493,10 @@ static FCacheRecordPolicy MakeCacheRecordQueryPolicy(const FBuildPolicy& BuildPo
 	const ECachePolicy DefaultCachePolicy = MakeCacheQueryPolicy(BuildPolicy.GetDefaultPolicy(), Context);
 	const ECachePolicy CombinedCachePolicy = MakeCacheQueryPolicy(BuildPolicy.GetCombinedPolicy(), Context);
 
-	FCacheRecordPolicyBuilder Builder(DefaultCachePolicy);
-	for (const FBuildValuePolicy& ValuePolicy : BuildPolicy.GetValuePolicies())
+	FCacheRecordPolicyBuilder Builder(DefaultCachePolicy | (CombinedCachePolicy & ECachePolicy::KeepAlive));
+	for (const FBuildValuePolicy& Value : BuildPolicy.GetValuePolicies())
 	{
-		Builder.AddValuePolicy(ValuePolicy.Id, MakeCacheQueryPolicy(ValuePolicy.Policy, Context));
+		Builder.AddValuePolicy(Value.Id, MakeCacheQueryPolicy(Value.Policy, Context) & ~ECachePolicy::KeepAlive);
 	}
 
 	// Output can only be loaded if its value is present. Add a policy for the output value which
diff --git a/Engine/Source/Developer/DerivedDataCache/Private/DerivedDataCachePolicy.cpp b/Engine/Source/Developer/DerivedDataCache/Private/DerivedDataCachePolicy.cpp
--- a/Engine/Source/Developer/DerivedDataCache/Private/DerivedDataCachePolicy.cpp
+++ b/Engine/Source/Developer/DerivedDataCache/Private/DerivedDataCachePolicy.cpp
@@ -3,6 +3,7 @@
 #include "DerivedDataCachePolicy.h"
 
 #include "Algo/Accumulate.h"
+#include "Algo/BinarySearch.h"
 #include "Containers/StringConv.h"
 #include "Containers/StringView.h"
 #include "Misc/StringBuilder.h"
@@ -158,15 +159,13 @@ ECachePolicy FCacheRecordPolicy::GetValuePolicy(const FValueId& Id) const
 {
 	if (Shared)
 	{
-		if (TConstArrayView<FCacheValuePolicy> Values = Shared->GetValuePolicies(); !Values.IsEmpty())
+		const TConstArrayView<FCacheValuePolicy> Values = Shared->GetValuePolicies();
+		if (const int32 Index = Algo::BinarySearchBy(Values, Id, &FCacheValuePolicy::Id); Index != INDEX_NONE)
 		{
-			if (int32 Index = Algo::BinarySearchBy(Values, Id, &FCacheValuePolicy::Id); Index != INDEX_NONE)
-			{
-				return Values[Index].Policy;
-			}
+			return Values[Index].Policy;
 		}
 	}
-	return DefaultPolicy;
+	return DefaultValuePolicy;
 }
 
 FCacheRecordPolicy FCacheRecordPolicy::Transform(const TFunctionRef<ECachePolicy (ECachePolicy)> Op) const
@@ -176,7 +175,7 @@ FCacheRecordPolicy FCacheRecordPolicy::Transform(const TFunctionRef<ECachePolicy
 		return Op(RecordPolicy);
 	}
 
-	FCacheRecordPolicyBuilder Builder(Op(DefaultPolicy));
+	FCacheRecordPolicyBuilder Builder(Op(GetBasePolicy()));
 	for (const FCacheValuePolicy& Value : GetValuePolicies())
 	{
 		Builder.AddValuePolicy({Value.Id, Op(Value.Policy)});
@@ -187,7 +186,7 @@ FCacheRecordPolicy FCacheRecordPolicy::Transform(const TFunctionRef<ECachePolicy
 void FCacheRecordPolicy::Save(FCbWriter& Writer) const
 {
 	Writer.BeginObject();
-	Writer.AddString("DefaultValuePolicy"_ASV, WriteToUtf8String<128>(GetDefaultPolicy()));
+	Writer.AddString("DefaultValuePolicy"_ASV, WriteToUtf8String<128>(GetBasePolicy()));
 	if (!IsUniform())
 	{
 		Writer.BeginArray("ValuePolicies"_ASV);
@@ -205,13 +204,13 @@ void FCacheRecordPolicy::Save(FCbWriter& Writer) const
 
 FOptionalCacheRecordPolicy FCacheRecordPolicy::Load(const FCbObjectView Object)
 {
-	const FUtf8StringView DefaultPolicyText = Object["DefaultValuePolicy"_ASV].AsString();
-	if (DefaultPolicyText.IsEmpty())
+	const FUtf8StringView BasePolicyText = Object["DefaultValuePolicy"_ASV].AsString();
+	if (BasePolicyText.IsEmpty())
 	{
 		return {};
 	}
 
-	FCacheRecordPolicyBuilder Builder(ParseCachePolicy(DefaultPolicyText));
+	FCacheRecordPolicyBuilder Builder(ParseCachePolicy(BasePolicyText));
 	for (const FCbFieldView Value : Object["ValuePolicies"_ASV])
 	{
 		const FValueId Id = Value["Id"_ASV].AsObjectId();
@@ -220,18 +219,22 @@ FOptionalCacheRecordPolicy FCacheRecordPolicy::Load(const FCbObjectView Object)
 		{
 			return {};
 		}
-		Builder.AddValuePolicy(Id, ParseCachePolicy(PolicyText));
+		const ECachePolicy Policy = ParseCachePolicy(PolicyText);
+		if (EnumHasAnyFlags(Policy, ~FCacheValuePolicy::PolicyMask))
+		{
+			return {};
+		}
+		Builder.AddValuePolicy(Id, Policy);
 	}
 	return Builder.Build();
 }
 
 void FCacheRecordPolicyBuilder::AddValuePolicy(const FCacheValuePolicy& Value)
 {
-	constexpr ECachePolicy RecordOnlyPolicyFlags = ECachePolicy::SkipMeta | ECachePolicy::PartialRecord | ECachePolicy::KeepAlive;
-	checkf(!EnumHasAnyFlags(Value.Policy, RecordOnlyPolicyFlags),
+	checkf(!EnumHasAnyFlags(Value.Policy, ~Value.PolicyMask),
 		TEXT("Value policy contains flags that only make sense on the record policy. Policy: %s"),
 		*WriteToString<128>(Value.Policy));
-	if (Value.Policy == (BasePolicy & ~RecordOnlyPolicyFlags))
+	if (Value.Policy == (BasePolicy & Value.PolicyMask))
 	{
 		return;
 	}
@@ -247,9 +250,12 @@ FCacheRecordPolicy FCacheRecordPolicyBuilder::Build()
 	FCacheRecordPolicy Policy(BasePolicy);
 	if (Shared)
 	{
-		const auto PolicyOr = [](ECachePolicy A, ECachePolicy B) { return A | (B & ~ECachePolicy::SkipData); };
+		const auto Add = [](const ECachePolicy A, const ECachePolicy B)
+		{
+			return ((A | B) & ~ECachePolicy::SkipData) | ((A & B) & ECachePolicy::SkipData);
+		};
 		const TConstArrayView<FCacheValuePolicy> Values = Shared->GetValuePolicies();
-		Policy.RecordPolicy = Algo::TransformAccumulate(Values, &FCacheValuePolicy::Policy, BasePolicy, PolicyOr);
+		Policy.RecordPolicy = Algo::TransformAccumulate(Values, &FCacheValuePolicy::Policy, BasePolicy, Add);
 		Policy.Shared = MoveTemp(Shared);
 	}
 	return Policy;
diff --git a/Engine/Source/Developer/DerivedDataCache/Private/FileSystemDerivedDataBackend.cpp b/Engine/Source/Developer/DerivedDataCache/Private/FileSystemDerivedDataBackend.cpp
--- a/Engine/Source/Developer/DerivedDataCache/Private/FileSystemDerivedDataBackend.cpp
+++ b/Engine/Source/Developer/DerivedDataCache/Private/FileSystemDerivedDataBackend.cpp
@@ -814,9 +814,9 @@ private:
 	void BuildCachePackagePath(const FCacheKey& CacheKey, FStringBuilderBase& Path) const;
 	void BuildCacheContentPath(const FIoHash& RawHash, FStringBuilderBase& Path) const;
 
-	[[nodiscard]] bool SaveFileWithHash(FStringBuilderBase& Path, FStringView DebugName, TFunctionRef<void (FArchive&)> WriteFunction) const;
+	[[nodiscard]] bool SaveFileWithHash(FStringBuilderBase& Path, FStringView DebugName, TFunctionRef<void (FArchive&)> WriteFunction, bool bReplaceExisting = false) const;
 	[[nodiscard]] bool LoadFileWithHash(FStringBuilderBase& Path, FStringView DebugName, TFunctionRef<void (FArchive&)> ReadFunction) const;
-	[[nodiscard]] bool SaveFile(FStringBuilderBase& Path, FStringView DebugName, TFunctionRef<void (FArchive&)> WriteFunction) const;
+	[[nodiscard]] bool SaveFile(FStringBuilderBase& Path, FStringView DebugName, TFunctionRef<void (FArchive&)> WriteFunction, bool bReplaceExisting = false) const;
 	[[nodiscard]] bool LoadFile(FStringBuilderBase& Path, FStringView DebugName, TFunctionRef<void (FArchive&)> ReadFunction) const;
 	[[nodiscard]] TUniquePtr<FArchive> OpenFile(FStringBuilderBase& Path, FStringView DebugName) const;
 
@@ -1779,13 +1779,11 @@ bool FFileSystemCacheStore::PutCacheRecord(
 	}
 
 	const FCacheKey& Key = Record.GetKey();
-	const ECachePolicy CombinedValuePolicy = Algo::TransformAccumulate(
-		Policy.GetValuePolicies(), &FCacheValuePolicy::Policy, Policy.GetDefaultPolicy(), UE_PROJECTION(operator|));
-	const ECachePolicy CombinedPolicy = Policy.GetRecordPolicy() | CombinedValuePolicy;
+	const ECachePolicy RecordPolicy = Policy.GetRecordPolicy();
 
 	// Skip the request if storing to the cache is disabled.
 	const ECachePolicy StoreFlag = SpeedClass == ESpeedClass::Local ? ECachePolicy::StoreLocal : ECachePolicy::StoreRemote;
-	if (!EnumHasAnyFlags(CombinedPolicy, StoreFlag))
+	if (!EnumHasAnyFlags(RecordPolicy, StoreFlag))
 	{
 		UE_LOG(LogDerivedDataCache, VeryVerbose, TEXT("%s: Skipped put of %s from '%.*s' due to cache policy"),
 			*CachePath, *WriteToString<96>(Key), Name.Len(), Name.GetData());
@@ -1804,11 +1802,9 @@ bool FFileSystemCacheStore::PutCacheRecord(
 	FCbPackage ExistingPackage;
 	TStringBuilder<256> Path;
 	BuildCachePackagePath(Key, Path);
-	if (EnumHasAnyFlags(CombinedValuePolicy, ECachePolicy::SkipData))
-	{
-		bRecordExists = FileExists(Path);
-	}
-	else
+	const ECachePolicy QueryFlag = SpeedClass == ESpeedClass::Local ? ECachePolicy::QueryLocal : ECachePolicy::QueryRemote;
+	const bool bReplaceExisting = !EnumHasAnyFlags(RecordPolicy, QueryFlag);
+	if (!bReplaceExisting)
 	{
 		bRecordExists = LoadFileWithHash(Path, Name, [&ExistingPackage](FArchive& Ar) { ExistingPackage.TryLoad(Ar); });
 	}
@@ -1868,7 +1864,7 @@ bool FFileSystemCacheStore::PutCacheRecord(
 
 	// Save the record package to storage.
 	const auto WriteRecord = [&](FArchive& Ar) { Package.Save(Ar); OutWriteSize += uint64(Ar.TotalSize()); };
-	if (!bRecordExists && !SaveFileWithHash(Path, Name, WriteRecord))
+	if (!bRecordExists && !SaveFileWithHash(Path, Name, WriteRecord, bReplaceExisting))
 	{
 		return false;
 	}
@@ -1974,7 +1970,8 @@ FOptionalCacheRecord FFileSystemCacheStore::GetCacheRecord(
 
 	FCacheRecordBuilder RecordBuilder(Key);
 
-	if (!EnumHasAnyFlags(Policy.GetRecordPolicy(), ECachePolicy::SkipMeta))
+	const ECachePolicy RecordPolicy = Policy.GetRecordPolicy();
+	if (!EnumHasAnyFlags(RecordPolicy, ECachePolicy::SkipMeta))
 	{
 		RecordBuilder.SetMeta(FCbObject(Record.Get().GetMeta()));
 	}
@@ -1988,7 +1985,7 @@ FOptionalCacheRecord FFileSystemCacheStore::GetCacheRecord(
 		{
 			RecordBuilder.AddValue(Id, MoveTemp(Content));
 		}
-		else if (EnumHasAnyFlags(ValuePolicy, ECachePolicy::PartialRecord))
+		else if (EnumHasAnyFlags(RecordPolicy, ECachePolicy::PartialRecord))
 		{
 			OutStatus = EStatus::Error;
 			RecordBuilder.AddValue(Value);
@@ -2039,11 +2036,9 @@ bool FFileSystemCacheStore::PutCacheValue(
 	FCbPackage ExistingPackage;
 	TStringBuilder<256> Path;
 	BuildCachePackagePath(Key, Path);
-	if (EnumHasAnyFlags(Policy, ECachePolicy::SkipData))
-	{
-		bValueExists = FileExists(Path);
-	}
-	else
+	const ECachePolicy QueryFlag = SpeedClass == ESpeedClass::Local ? ECachePolicy::QueryLocal : ECachePolicy::QueryRemote;
+	const bool bReplaceExisting = !EnumHasAnyFlags(Policy, QueryFlag);
+	if (!bReplaceExisting)
 	{
 		bValueExists = LoadFileWithHash(Path, Name, [&ExistingPackage](FArchive& Ar) { ExistingPackage.TryLoad(Ar); })
 			&& (ExistingPackage.FindAttachment(Value.GetRawHash()) || GetCacheContentExists(Key, Value.GetRawHash()));
@@ -2087,7 +2082,7 @@ bool FFileSystemCacheStore::PutCacheValue(
 
 		// Save the value package to storage.
 		const auto WritePackage = [&](FArchive& Ar) { Package.Save(Ar); OutWriteSize += uint64(Ar.TotalSize()); };
-		if (!SaveFileWithHash(Path, Name, WritePackage))
+		if (!SaveFileWithHash(Path, Name, WritePackage, bReplaceExisting))
 		{
 			return false;
 		}
@@ -2383,7 +2378,8 @@ void FFileSystemCacheStore::BuildCacheContentPath(const FIoHash& RawHash, FStrin
 bool FFileSystemCacheStore::SaveFileWithHash(
 	FStringBuilderBase& Path,
 	const FStringView DebugName,
-	const TFunctionRef<void (FArchive&)> WriteFunction) const
+	const TFunctionRef<void (FArchive&)> WriteFunction,
+	const bool bReplaceExisting) const
 {
 	return SaveFile(Path, DebugName, [&WriteFunction](FArchive& Ar)
 	{
@@ -2391,7 +2387,7 @@ bool FFileSystemCacheStore::SaveFileWithHash(
 		WriteFunction(HashAr);
 		FBlake3Hash Hash = HashAr.GetHash();
 		Ar << Hash;
-	});
+	}, bReplaceExisting);
 }
 
 bool FFileSystemCacheStore::LoadFileWithHash(
@@ -2421,7 +2417,8 @@ bool FFileSystemCacheStore::LoadFileWithHash(
 bool FFileSystemCacheStore::SaveFile(
 	FStringBuilderBase& Path,
 	const FStringView DebugName,
-	const TFunctionRef<void (FArchive&)> WriteFunction) const
+	const TFunctionRef<void (FArchive&)> WriteFunction,
+	const bool bReplaceExisting) const
 {
 	TStringBuilder<256> TempPath;
 	TempPath << FPathViews::GetPath(Path) << TEXT("/Temp.") << FGuid::NewGuid();
@@ -2460,7 +2457,7 @@ bool FFileSystemCacheStore::SaveFile(
 	}
 
 	if (!IFileManager::Get().Move(*Path, *TempPath,
-		/*bReplace*/ false, /*bEvenIfReadOnly*/ false, /*bAttributes*/ false, /*bDoNotRetryOrError*/ true))
+		bReplaceExisting, /*bEvenIfReadOnly*/ false, /*bAttributes*/ false, /*bDoNotRetryOrError*/ true))
 	{
 		UE_LOG(LogDerivedDataCache, Log,
 			TEXT("%s: Move collision when writing file %s from '%.*s'."),
diff --git a/Engine/Source/Developer/DerivedDataCache/Private/HttpDerivedDataBackend.cpp b/Engine/Source/Developer/DerivedDataCache/Private/HttpDerivedDataBackend.cpp
--- a/Engine/Source/Developer/DerivedDataCache/Private/HttpDerivedDataBackend.cpp
+++ b/Engine/Source/Developer/DerivedDataCache/Private/HttpDerivedDataBackend.cpp
@@ -2450,14 +2450,12 @@ bool FHttpDerivedDataBackend::PutCacheRecord(
 	}
 
 	const FCacheKey& Key = Record.GetKey();
-	const ECachePolicy CombinedValuePolicy = Algo::TransformAccumulate(
-		Policy.GetValuePolicies(), &FCacheValuePolicy::Policy, Policy.GetDefaultPolicy(), UE_PROJECTION(operator|));
-	const ECachePolicy CombinedPolicy = Policy.GetRecordPolicy() | CombinedValuePolicy;
+	const ECachePolicy RecordPolicy = Policy.GetRecordPolicy();
 
 	// Skip the request if storing to the cache is disabled.
 	// Http backends won't generally be "local" but including handling for this possibility for consistency
 	const ECachePolicy StoreFlag = SpeedClass == ESpeedClass::Local ? ECachePolicy::StoreLocal : ECachePolicy::StoreRemote;
-	if (!EnumHasAnyFlags(CombinedPolicy, StoreFlag))
+	if (!EnumHasAnyFlags(RecordPolicy, StoreFlag))
 	{
 		UE_LOG(LogDerivedDataCache, VeryVerbose, TEXT("%s: Skipped put of %s from '%.*s' due to cache policy"),
 			*GetName(), *WriteToString<96>(Key), Name.Len(), Name.GetData());
diff --git a/Engine/Source/Developer/DerivedDataCache/Private/MemoryDerivedDataBackend.cpp b/Engine/Source/Developer/DerivedDataCache/Private/MemoryDerivedDataBackend.cpp
--- a/Engine/Source/Developer/DerivedDataCache/Private/MemoryDerivedDataBackend.cpp
+++ b/Engine/Source/Developer/DerivedDataCache/Private/MemoryDerivedDataBackend.cpp
@@ -662,6 +662,11 @@ void FMemoryDerivedDataBackend::Put(
 			}
 		};
 
+		if (!EnumHasAnyFlags(Request.Policy.GetRecordPolicy(), ECachePolicy::StoreLocal))
+		{
+			continue;
+		}
+
 		if (ShouldSimulateMiss(Key))
 		{
 			UE_LOG(LogDerivedDataCache, Verbose, TEXT("%s: Simulated miss for put of %s from '%s'"),
@@ -676,44 +681,39 @@ void FMemoryDerivedDataBackend::Put(
 			continue;
 		}
 
-		if (Values.IsEmpty())
+		COOK_STAT(auto Timer = UsageStats.TimePut());
+		const int64 RecordSize = Private::GetCacheRecordCompressedSize(Record);
+		const bool bReplaceExisting = !EnumHasAnyFlags(Request.Policy.GetRecordPolicy(), ECachePolicy::QueryLocal);
+
+		FWriteScopeLock ScopeLock(SynchronizationObject);
+		FCacheRecord* const ExistingRecord = CacheRecords.Find(Key);
+		Status = ExistingRecord && !bReplaceExisting ? EStatus::Ok : EStatus::Error;
+		if (bDisabled || Status == EStatus::Ok)
 		{
-			FWriteScopeLock ScopeLock(SynchronizationObject);
-			if (bDisabled)
-			{
-				continue;
-			}
-			if (const FCacheRecord* Existing = CacheRecords.Find(Key))
-			{
-				CacheRecords.Remove(Key);
-				CurrentCacheSize -= Private::GetCacheRecordCompressedSize(*Existing);
-				bMaxSizeExceeded = false;
-			}
-			Status = EStatus::Ok;
+			continue;
 		}
-		else
-		{
-			COOK_STAT(auto Timer = UsageStats.TimePut());
-			const int64 RecordSize = Private::GetCacheRecordCompressedSize(Record);
 
-			FWriteScopeLock ScopeLock(SynchronizationObject);
-			Status = CacheRecords.Contains(Key) ? EStatus::Ok : EStatus::Error;
-			if (bDisabled || Status == EStatus::Ok)
-			{
-				continue;
-			}
-			if (MaxCacheSize > 0 && (CurrentCacheSize + RecordSize) > MaxCacheSize)
-			{
-				UE_LOG(LogDerivedDataCache, Display, TEXT("Failed to cache data. Maximum cache size reached. CurrentSize %" INT64_FMT " KiB / MaxSize: %" INT64_FMT " KiB"), CurrentCacheSize / 1024, MaxCacheSize / 1024);
-				bMaxSizeExceeded = true;
-				continue;
-			}
+		const int64 ExistingRecordSize = ExistingRecord ? Private::GetCacheRecordCompressedSize(*ExistingRecord) : 0;
+		const int64 RequiredSize = RecordSize - ExistingRecordSize;
+
+		if (MaxCacheSize > 0 && (CurrentCacheSize + RequiredSize) > MaxCacheSize)
+		{
+			UE_LOG(LogDerivedDataCache, Display, TEXT("Failed to cache data. Maximum cache size reached. CurrentSize %" INT64_FMT " KiB / MaxSize: %" INT64_FMT " KiB"), CurrentCacheSize / 1024, MaxCacheSize / 1024);
+			bMaxSizeExceeded = true;
+			continue;
+		}
 
-			CurrentCacheSize += RecordSize;
+		CurrentCacheSize += RequiredSize;
+		if (ExistingRecord)
+		{
+			*ExistingRecord = Record;
+		}
+		else
+		{
 			CacheRecords.Add(Record);
-			COOK_STAT(Timer.AddHit(RecordSize));
-			Status = EStatus::Ok;
 		}
+		COOK_STAT(Timer.AddHit(RecordSize));
+		Status = EStatus::Ok;
 	}
 }
 
@@ -747,7 +747,7 @@ void FMemoryDerivedDataBackend::Get(
 				if (!Value.HasData() && !EnumHasAnyFlags(Policy.GetValuePolicy(Value.GetId()), ECachePolicy::SkipData))
 				{
 					Status = EStatus::Error;
-					if (!EnumHasAllFlags(Policy.GetValuePolicy(Value.GetId()), ECachePolicy::PartialRecord))
+					if (!EnumHasAllFlags(Policy.GetRecordPolicy(), ECachePolicy::PartialRecord))
 					{
 						Record.Reset();
 						break;
diff --git a/Engine/Source/Developer/DerivedDataCache/Private/PakFileDerivedDataBackend.cpp b/Engine/Source/Developer/DerivedDataCache/Private/PakFileDerivedDataBackend.cpp
--- a/Engine/Source/Developer/DerivedDataCache/Private/PakFileDerivedDataBackend.cpp
+++ b/Engine/Source/Developer/DerivedDataCache/Private/PakFileDerivedDataBackend.cpp
@@ -776,6 +776,7 @@ bool FPakFileDerivedDataBackend::PutCacheRecord(
 	}
 
 	const FCacheKey& Key = Record.GetKey();
+	const ECachePolicy RecordPolicy = Policy.GetRecordPolicy();
 
 	// Skip the request if storing to the cache is disabled.
 	if (!EnumHasAnyFlags(Policy.GetRecordPolicy(), ECachePolicy::StoreLocal))
@@ -797,16 +798,14 @@ bool FPakFileDerivedDataBackend::PutCacheRecord(
 	FCbPackage ExistingPackage;
 	TStringBuilder<256> Path;
 	FPathViews::Append(Path, TEXT("Buckets"), Key);
-	const ECachePolicy CombinedValuePolicy = Algo::TransformAccumulate(
-		Policy.GetValuePolicies(), &FCacheValuePolicy::Policy, Policy.GetDefaultPolicy(), UE_PROJECTION(operator|));
-	if (EnumHasAnyFlags(CombinedValuePolicy, ECachePolicy::SkipData))
+	const bool bReplaceExisting = !EnumHasAnyFlags(RecordPolicy, ECachePolicy::QueryLocal);
+	if (!bReplaceExisting)
 	{
-		bRecordExists = FileExists(Path);
-	}
-	else if (FSharedBuffer Buffer = LoadFile(Path, Name))
-	{
-		FCbFieldIterator It = FCbFieldIterator::MakeRange(MoveTemp(Buffer));
-		bRecordExists = ExistingPackage.TryLoad(It);
+		if (FSharedBuffer Buffer = LoadFile(Path, Name))
+		{
+			FCbFieldIterator It = FCbFieldIterator::MakeRange(MoveTemp(Buffer));
+			bRecordExists = ExistingPackage.TryLoad(It);
+		}
 	}
 
 	// Save the record to a package and remove attachments that will be stored externally.
@@ -951,6 +950,7 @@ FOptionalCacheRecord FPakFileDerivedDataBackend::GetCacheRecord(
 	OutStatus = EStatus::Ok;
 
 	FCacheRecordBuilder RecordBuilder(Key);
+	const ECachePolicy RecordOnlyPolicy = Policy.GetRecordPolicy() & ~FCacheValuePolicy::PolicyMask;
 
 	if (!EnumHasAnyFlags(Policy.GetRecordPolicy(), ECachePolicy::SkipMeta))
 	{
@@ -960,7 +960,7 @@ FOptionalCacheRecord FPakFileDerivedDataBackend::GetCacheRecord(
 	for (FValueWithId Value : Record.Get().GetValues())
 	{
 		const ECachePolicy ValuePolicy = Policy.GetValuePolicy(Value.GetId());
-		GetCacheContent(Name, Key, ValuePolicy, Value, OutStatus);
+		GetCacheContent(Name, Key, ValuePolicy | RecordOnlyPolicy, Value, OutStatus);
 		if (Value.IsNull())
 		{
 			return FOptionalCacheRecord();
diff --git a/Engine/Source/Developer/DerivedDataCache/Public/DerivedDataBuildTypes.h b/Engine/Source/Developer/DerivedDataCache/Public/DerivedDataBuildTypes.h
--- a/Engine/Source/Developer/DerivedDataCache/Public/DerivedDataBuildTypes.h
+++ b/Engine/Source/Developer/DerivedDataCache/Public/DerivedDataBuildTypes.h
@@ -84,7 +84,6 @@ public:
 	virtual void Release() const = 0;
 	virtual TConstArrayView<FBuildValuePolicy> GetValuePolicies() const = 0;
 	virtual void AddValuePolicy(const FBuildValuePolicy& Policy) = 0;
-	virtual void Build() = 0;
 };
 
 /** Flags to control the behavior of build requests, with optional overrides by value. */
@@ -144,7 +143,7 @@ public:
 	}
 
 	/** Adds a build policy override for a value. */
-	UE_API void AddValuePolicy(const FBuildValuePolicy& Policy);
+	UE_API void AddValuePolicy(const FBuildValuePolicy& Value);
 	inline void AddValuePolicy(const FValueId& Id, EBuildPolicy Policy) { AddValuePolicy({Id, Policy}); }
 
 	/** Build a build policy, which makes this builder subsequently unusable. */
diff --git a/Engine/Source/Developer/DerivedDataCache/Public/DerivedDataCachePolicy.h b/Engine/Source/Developer/DerivedDataCache/Public/DerivedDataCachePolicy.h
--- a/Engine/Source/Developer/DerivedDataCache/Public/DerivedDataCachePolicy.h
+++ b/Engine/Source/Developer/DerivedDataCache/Public/DerivedDataCachePolicy.h
@@ -109,6 +109,9 @@ struct FCacheValuePolicy
 {
 	FValueId Id;
 	ECachePolicy Policy = ECachePolicy::Default;
+
+	/** Flags that are valid on a value policy. */
+	static constexpr ECachePolicy PolicyMask = ECachePolicy::Default | ECachePolicy::SkipData;
 };
 
 /** Interface for the private implementation of the cache record policy. */
@@ -138,9 +141,9 @@ public:
 	FCacheRecordPolicy() = default;
 
 	/** Construct a cache record policy with a uniform policy for the record and every value. */
-	inline FCacheRecordPolicy(ECachePolicy Policy)
-		: RecordPolicy(Policy)
-		, DefaultPolicy(Policy)
+	inline FCacheRecordPolicy(ECachePolicy BasePolicy)
+		: RecordPolicy(BasePolicy)
+		, DefaultValuePolicy(BasePolicy & FCacheValuePolicy::PolicyMask)
 	{
 	}
 
@@ -150,8 +153,11 @@ public:
 	/** Returns the cache policy to use for the record. */
 	inline ECachePolicy GetRecordPolicy() const { return RecordPolicy; }
 
-	/** Returns the cache policy to use for values with no override. */
-	inline ECachePolicy GetDefaultPolicy() const { return DefaultPolicy; }
+	/** Returns the base cache policy that this was constructed from. */
+	inline ECachePolicy GetBasePolicy() const
+	{
+		return DefaultValuePolicy | (RecordPolicy & ~FCacheValuePolicy::PolicyMask);
+	}
 
 	/** Returns the cache policy to use for the value. */
 	UE_API ECachePolicy GetValuePolicy(const FValueId& Id) const;
@@ -176,7 +182,7 @@ private:
 	friend class FOptionalCacheRecordPolicy;
 
 	ECachePolicy RecordPolicy = ECachePolicy::Default;
-	ECachePolicy DefaultPolicy = ECachePolicy::Default;
+	ECachePolicy DefaultValuePolicy = ECachePolicy::Default;
 	TRefCountPtr<const Private::ICacheRecordPolicyShared> Shared;
 };
 
@@ -194,7 +200,7 @@ public:
 	}
 
 	/** Adds a cache policy override for a value. */
-	UE_API void AddValuePolicy(const FCacheValuePolicy& Policy);
+	UE_API void AddValuePolicy(const FCacheValuePolicy& Value);
 	inline void AddValuePolicy(const FValueId& Id, ECachePolicy Policy) { AddValuePolicy({Id, Policy}); }
 
 	/** Build a cache record policy, which makes this builder subsequently unusable. */
