DDC: Added PutValue and GetValue support to the new cache hierarchy

#rnx

diff --git a/Engine/Source/Developer/DerivedDataCache/Private/DerivedDataCacheStoreHierarchy.cpp b/Engine/Source/Developer/DerivedDataCache/Private/DerivedDataCacheStoreHierarchy.cpp
--- a/Engine/Source/Developer/DerivedDataCache/Private/DerivedDataCacheStoreHierarchy.cpp
+++ b/Engine/Source/Developer/DerivedDataCache/Private/DerivedDataCacheStoreHierarchy.cpp
@@ -11,6 +11,7 @@
 #include "HAL/CriticalSection.h"
 #include "Misc/EnumClassFlags.h"
 #include "Misc/ScopeRWLock.h"
+#include "Templates/Invoke.h"
 #include "Templates/RefCounting.h"
 #include "Templates/UniquePtr.h"
 #include <atomic>
@@ -80,15 +81,20 @@ public:
 
 	bool LegacyDebugOptions(FBackendDebugOptions& Options) final;
 
+	template <typename PutRequestType, typename GetRequestType> struct TBatchParams;
+	using FCacheValueBatchParams = TBatchParams<FCachePutValueRequest, FCacheGetValueRequest>;
+	using FLegacyCacheBatchParams = TBatchParams<FLegacyCachePutRequest, FLegacyCacheGetRequest>;
+
 private:
 	// Caller must hold a write lock on CacheStoresLock.
 	void UpdateNodeFlags();
 
-	class FBatchBase;
 	class FCounterEvent;
 
-	class FLegacyPutBatch;
-	class FLegacyGetBatch;
+	class FBatchBase;
+	template <typename Params> class TPutBatch;
+	template <typename Params> class TGetBatch;
+
 	class FLegacyDeleteBatch;
 
 	enum class ECacheStoreNodeFlags : uint32;
@@ -115,6 +121,23 @@ private:
 
 ///////////////////////////////////////////////////////////////////////////////////////////////////////////////////////
 
+class FCacheStoreHierarchy::FCounterEvent
+{
+public:
+	void Reset(int32 NewCount)
+	{
+		Count.store(NewCount, std::memory_order_relaxed);
+	}
+
+	bool Signal()
+	{
+		return Count.fetch_sub(1, std::memory_order_acq_rel) == 1;
+	}
+
+private:
+	std::atomic<int32> Count{0};
+};
+
 class FCacheStoreHierarchy::FBatchBase
 {
 public:
@@ -135,21 +158,22 @@ private:
 	std::atomic<int32> ReferenceCount{0};
 };
 
-class FCacheStoreHierarchy::FCounterEvent
+template <typename PutRequestType, typename GetRequestType>
+struct FCacheStoreHierarchy::TBatchParams
 {
-public:
-	void Reset(int32 NewCount)
-	{
-		Count.store(NewCount, std::memory_order_relaxed);
-	}
-
-	bool Signal()
-	{
-		return Count.fetch_sub(1, std::memory_order_acq_rel) == 1;
-	}
-
-private:
-	std::atomic<int32> Count{0};
+	using FPutRequest = PutRequestType;
+	using FGetRequest = GetRequestType;
+	using FPutResponse = decltype(DeclVal<FPutRequest>().MakeResponse(EStatus::Ok));
+	using FGetResponse = decltype(DeclVal<FGetRequest>().MakeResponse(EStatus::Ok));
+	using FOnPutComplete = TUniqueFunction<void (FPutResponse&& Response)>;
+	using FOnGetComplete = TUniqueFunction<void (FGetResponse&& Response)>;
+	using PutFunctionType = void (ILegacyCacheStore::*)(TConstArrayView<FPutRequest>, IRequestOwner&, FOnPutComplete&&);
+	using GetFunctionType = void (ILegacyCacheStore::*)(TConstArrayView<FGetRequest>, IRequestOwner&, FOnGetComplete&&);
+	static PutFunctionType Put();
+	static GetFunctionType Get();
+	static bool HasResponseData(const FGetResponse& Response);
+	static void FilterResponseByRequest(FGetResponse& Response, const FGetRequest& Request);
+	static FPutRequest MakePutRequest(const FGetResponse& Response, const FGetRequest& Request);
 };
 
 ///////////////////////////////////////////////////////////////////////////////////////////////////////////////////////
@@ -274,71 +298,31 @@ bool FCacheStoreHierarchy::CanQueryIfError(const ECachePolicy Policy, const ECac
 
 ///////////////////////////////////////////////////////////////////////////////////////////////////////////////////////
 
-void FCacheStoreHierarchy::Put(
-	TConstArrayView<FCachePutRequest> Requests,
-	IRequestOwner& Owner,
-	FOnCachePutComplete&& OnComplete)
-{
-	unimplemented();
-}
-
-///////////////////////////////////////////////////////////////////////////////////////////////////////////////////////
-
-void FCacheStoreHierarchy::Get(
-	TConstArrayView<FCacheGetRequest> Requests,
-	IRequestOwner& Owner,
-	FOnCacheGetComplete&& OnComplete)
-{
-	unimplemented();
-}
-
-///////////////////////////////////////////////////////////////////////////////////////////////////////////////////////
-
-void FCacheStoreHierarchy::PutValue(
-	TConstArrayView<FCachePutValueRequest> Requests,
-	IRequestOwner& Owner,
-	FOnCachePutValueComplete&& OnComplete)
-{
-	unimplemented();
-}
-
-///////////////////////////////////////////////////////////////////////////////////////////////////////////////////////
-
-void FCacheStoreHierarchy::GetValue(
-	TConstArrayView<FCacheGetValueRequest> Requests,
-	IRequestOwner& Owner,
-	FOnCacheGetValueComplete&& OnComplete)
-{
-	unimplemented();
-}
-
-///////////////////////////////////////////////////////////////////////////////////////////////////////////////////////
-
-void FCacheStoreHierarchy::GetChunks(
-	TConstArrayView<FCacheGetChunkRequest> Requests,
-	IRequestOwner& Owner,
-	FOnCacheGetChunkComplete&& OnComplete)
+template <typename Params>
+class FCacheStoreHierarchy::TPutBatch : public FBatchBase, public Params
 {
-	unimplemented();
-}
-
-///////////////////////////////////////////////////////////////////////////////////////////////////////////////////////
+	using FPutRequest = typename Params::FPutRequest;
+	using FGetRequest = typename Params::FGetRequest;
+	using FPutResponse = typename Params::FPutResponse;
+	using FGetResponse = typename Params::FGetResponse;
+	using FOnPutComplete = typename Params::FOnPutComplete;
+	using FOnGetComplete = typename Params::FOnGetComplete;
+	using Params::Put;
+	using Params::Get;
 
-class FCacheStoreHierarchy::FLegacyPutBatch : public FBatchBase
-{
 public:
 	static void Begin(
 		const FCacheStoreHierarchy& Hierarchy,
-		TConstArrayView<FLegacyCachePutRequest> Requests,
+		TConstArrayView<FPutRequest> Requests,
 		IRequestOwner& Owner,
-		FOnLegacyCachePutComplete&& OnComplete);
+		FOnPutComplete&& OnComplete);
 
 private:
-	FLegacyPutBatch(
+	TPutBatch(
 		const FCacheStoreHierarchy& InHierarchy,
-		const TConstArrayView<FLegacyCachePutRequest> InRequests,
+		const TConstArrayView<FPutRequest> InRequests,
 		IRequestOwner& InOwner,
-		FOnLegacyCachePutComplete&& InOnComplete)
+		FOnPutComplete&& InOnComplete)
 		: Hierarchy(InHierarchy)
 		, Requests(InRequests)
 		, BatchOwner(InOwner)
@@ -352,10 +336,10 @@ private:
 	void DispatchRequests();
 
 	bool DispatchGetRequests();
-	void CompleteGetRequest(FLegacyCacheGetResponse&& Response);
+	void CompleteGetRequest(FGetResponse&& Response);
 
 	bool DispatchPutRequests();
-	void CompletePutRequest(FLegacyCachePutResponse&& Response);
+	void CompletePutRequest(FPutResponse&& Response);
 
 	struct FRequestState
 	{
@@ -364,9 +348,9 @@ private:
 	};
 
 	const FCacheStoreHierarchy& Hierarchy;
-	TArray<FLegacyCachePutRequest> Requests;
+	TArray<FPutRequest> Requests;
 	IRequestOwner& BatchOwner;
-	FOnLegacyCachePutComplete OnComplete;
+	FOnPutComplete OnComplete;
 
 	FRequestOwner AsyncOwner;
 	TArray<FRequestState, TInlineAllocator<1>> States;
@@ -375,22 +359,24 @@ private:
 	int32 NodePutIndex = 0;
 };
 
-void FCacheStoreHierarchy::FLegacyPutBatch::Begin(
+template <typename Params>
+void FCacheStoreHierarchy::TPutBatch<Params>::Begin(
 	const FCacheStoreHierarchy& Hierarchy,
-	const TConstArrayView<FLegacyCachePutRequest> Requests,
+	const TConstArrayView<FPutRequest> Requests,
 	IRequestOwner& Owner,
-	FOnLegacyCachePutComplete&& OnComplete)
+	FOnPutComplete&& OnComplete)
 {
 	if (Requests.IsEmpty() || !EnumHasAnyFlags(Hierarchy.CombinedNodeFlags, ECacheStoreNodeFlags::HasStoreNode))
 	{
 		return CompleteWithStatus(Requests, OnComplete, EStatus::Error);
 	}
 
-	TRefCountPtr<FLegacyPutBatch> State = new FLegacyPutBatch(Hierarchy, Requests, Owner, MoveTemp(OnComplete));
+	TRefCountPtr<TPutBatch> State = new TPutBatch(Hierarchy, Requests, Owner, MoveTemp(OnComplete));
 	State->DispatchRequests();
 }
 
-void FCacheStoreHierarchy::FLegacyPutBatch::DispatchRequests()
+template <typename Params>
+void FCacheStoreHierarchy::TPutBatch<Params>::DispatchRequests()
 {
 	FReadScopeLock Lock(Hierarchy.NodesLock);
 
@@ -403,7 +389,7 @@ void FCacheStoreHierarchy::FLegacyPutBatch::DispatchRequests()
 	}
 
 	int32 RequestIndex = 0;
-	for (const FLegacyCachePutRequest& Request : Requests)
+	for (const FPutRequest& Request : Requests)
 	{
 		const FRequestState& State = States[RequestIndex];
 		if (!State.bOk && !State.bStop)
@@ -413,7 +399,8 @@ void FCacheStoreHierarchy::FLegacyPutBatch::DispatchRequests()
 	}
 }
 
-bool FCacheStoreHierarchy::FLegacyPutBatch::DispatchGetRequests()
+template <typename Params>
+bool FCacheStoreHierarchy::TPutBatch<Params>::DispatchGetRequests()
 {
 	if (NodeGetIndex >= NodePutIndex)
 	{
@@ -428,11 +415,11 @@ bool FCacheStoreHierarchy::FLegacyPutBatch::DispatchGetRequests()
 		return false;
 	}
 
-	TArray<FLegacyCacheGetRequest, TInlineAllocator<1>> NodeRequests;
+	TArray<FGetRequest, TInlineAllocator<1>> NodeRequests;
 	NodeRequests.Reserve(Requests.Num());
 
 	uint64 RequestIndex = 0;
-	for (const FLegacyCachePutRequest& Request : Requests)
+	for (const FPutRequest& Request : Requests)
 	{
 		if (!States[RequestIndex].bStop && CanQuery(Request.Policy, Node.CacheFlags))
 		{
@@ -444,8 +431,8 @@ bool FCacheStoreHierarchy::FLegacyPutBatch::DispatchGetRequests()
 	if (const int32 NodeRequestsCount = NodeRequests.Num())
 	{
 		RemainingRequestCount.Reset(NodeRequestsCount + 1);
-		Node.Cache->LegacyGet(NodeRequests, BatchOwner,
-			[State = TRefCountPtr(this)](FLegacyCacheGetResponse&& Response)
+		Invoke(Get(), Node.Cache, NodeRequests, BatchOwner,
+			[State = TRefCountPtr(this)](FGetResponse&& Response)
 			{
 				State->CompleteGetRequest(MoveTemp(Response));
 			});
@@ -455,7 +442,8 @@ bool FCacheStoreHierarchy::FLegacyPutBatch::DispatchGetRequests()
 	return false;
 }
 
-void FCacheStoreHierarchy::FLegacyPutBatch::CompleteGetRequest(FLegacyCacheGetResponse&& Response)
+template <typename Params>
+void FCacheStoreHierarchy::TPutBatch<Params>::CompleteGetRequest(FGetResponse&& Response)
 {
 	if (Response.Status == EStatus::Ok)
 	{
@@ -474,7 +462,8 @@ void FCacheStoreHierarchy::FLegacyPutBatch::CompleteGetRequest(FLegacyCacheGetRe
 	}
 }
 
-bool FCacheStoreHierarchy::FLegacyPutBatch::DispatchPutRequests()
+template <typename Params>
+bool FCacheStoreHierarchy::TPutBatch<Params>::DispatchPutRequests()
 {
 	const FCacheStoreNode& Node = Hierarchy.Nodes[NodePutIndex];
 	if (!EnumHasAnyFlags(Node.CacheFlags, ECacheStoreFlags::Store))
@@ -482,15 +471,15 @@ bool FCacheStoreHierarchy::FLegacyPutBatch::DispatchPutRequests()
 		return false;
 	}
 
-	TArray<FLegacyCachePutRequest, TInlineAllocator<1>> NodeRequests;
-	TArray<FLegacyCachePutRequest, TInlineAllocator<1>> AsyncNodeRequests;
+	TArray<FPutRequest, TInlineAllocator<1>> NodeRequests;
+	TArray<FPutRequest, TInlineAllocator<1>> AsyncNodeRequests;
 
 	const int32 RequestCount = Requests.Num();
 	NodeRequests.Reserve(RequestCount);
 	AsyncNodeRequests.Reserve(RequestCount);
 
 	int32 RequestIndex = 0;
-	for (const FLegacyCachePutRequest& Request : Requests)
+	for (const FPutRequest& Request : Requests)
 	{
 		const FRequestState& State = States[RequestIndex];
 		if (!State.bStop && CanStore(Request.Policy, Node.CacheFlags))
@@ -503,14 +492,14 @@ bool FCacheStoreHierarchy::FLegacyPutBatch::DispatchPutRequests()
 	if (!AsyncNodeRequests.IsEmpty())
 	{
 		FRequestBarrier Barrier(AsyncOwner);
-		Node.AsyncCache->LegacyPut(AsyncNodeRequests, AsyncOwner, [](auto&&){});
+		Invoke(Put(), Node.AsyncCache, AsyncNodeRequests, AsyncOwner, [](auto&&){});
 	}
 
 	if (const int32 NodeRequestsCount = NodeRequests.Num())
 	{
 		RemainingRequestCount.Reset(NodeRequestsCount + 1);
-		Node.Cache->LegacyPut(NodeRequests, BatchOwner,
-			[State = TRefCountPtr(this)](FLegacyCachePutResponse&& Response)
+		Invoke(Put(), Node.Cache, NodeRequests, BatchOwner,
+			[State = TRefCountPtr(this)](FPutResponse&& Response)
 			{
 				State->CompletePutRequest(MoveTemp(Response));
 			});
@@ -520,7 +509,8 @@ bool FCacheStoreHierarchy::FLegacyPutBatch::DispatchPutRequests()
 	return false;
 }
 
-void FCacheStoreHierarchy::FLegacyPutBatch::CompletePutRequest(FLegacyCachePutResponse&& Response)
+template <typename Params>
+void FCacheStoreHierarchy::TPutBatch<Params>::CompletePutRequest(FPutResponse&& Response)
 {
 	if (Response.Status == EStatus::Ok)
 	{
@@ -537,31 +527,36 @@ void FCacheStoreHierarchy::FLegacyPutBatch::CompletePutRequest(FLegacyCachePutRe
 	}
 }
 
-void FCacheStoreHierarchy::LegacyPut(
-	const TConstArrayView<FLegacyCachePutRequest> Requests,
-	IRequestOwner& Owner,
-	FOnLegacyCachePutComplete&& OnComplete)
-{
-	FLegacyPutBatch::Begin(*this, Requests, Owner, MoveTemp(OnComplete));
-}
-
 ///////////////////////////////////////////////////////////////////////////////////////////////////////////////////////
 
-class FCacheStoreHierarchy::FLegacyGetBatch : public FBatchBase
+template <typename Params>
+class FCacheStoreHierarchy::TGetBatch : public FBatchBase, public Params
 {
+	using FPutRequest = typename Params::FPutRequest;
+	using FGetRequest = typename Params::FGetRequest;
+	using FPutResponse = typename Params::FPutResponse;
+	using FGetResponse = typename Params::FGetResponse;
+	using FOnPutComplete = typename Params::FOnPutComplete;
+	using FOnGetComplete = typename Params::FOnGetComplete;
+	using Params::Put;
+	using Params::Get;
+	using Params::HasResponseData;
+	using Params::FilterResponseByRequest;
+	using Params::MakePutRequest;
+
 public:
 	static void Begin(
 		const FCacheStoreHierarchy& Hierarchy,
-		TConstArrayView<FLegacyCacheGetRequest> Requests,
+		TConstArrayView<FGetRequest> Requests,
 		IRequestOwner& Owner,
-		FOnLegacyCacheGetComplete&& OnComplete);
+		FOnGetComplete&& OnComplete);
 
 private:
-	FLegacyGetBatch(
+	TGetBatch(
 		const FCacheStoreHierarchy& InHierarchy,
-		const TConstArrayView<FLegacyCacheGetRequest> InRequests,
+		const TConstArrayView<FGetRequest> InRequests,
 		IRequestOwner& InOwner,
-		FOnLegacyCacheGetComplete&& InOnComplete)
+		FOnGetComplete&& InOnComplete)
 		: Hierarchy(InHierarchy)
 		, OnComplete(MoveTemp(InOnComplete))
 		, Owner(InOwner)
@@ -569,23 +564,23 @@ private:
 	{
 		AsyncOwner.KeepAlive();
 		States.Reserve(InRequests.Num());
-		for (const FLegacyCacheGetRequest& Request : InRequests)
+		for (const FGetRequest& Request : InRequests)
 		{
 			States.Add({Request});
 		}
 	}
 
 	void DispatchRequests();
-	void CompleteRequest(FLegacyCacheGetResponse&& Response);
+	void CompleteRequest(FGetResponse&& Response);
 
 	struct FState
 	{
-		FLegacyCacheGetRequest Request;
-		FLegacyCacheGetResponse Response;
+		FGetRequest Request;
+		FGetResponse Response;
 	};
 
 	const FCacheStoreHierarchy& Hierarchy;
-	FOnLegacyCacheGetComplete OnComplete;
+	FOnGetComplete OnComplete;
 	TArray<FState, TInlineAllocator<8>> States;
 
 	IRequestOwner& Owner;
@@ -594,27 +589,29 @@ private:
 	int32 NodeIndex = 0;
 };
 
-void FCacheStoreHierarchy::FLegacyGetBatch::Begin(
+template <typename Params>
+void FCacheStoreHierarchy::TGetBatch<Params>::Begin(
 	const FCacheStoreHierarchy& InHierarchy,
-	const TConstArrayView<FLegacyCacheGetRequest> InRequests,
+	const TConstArrayView<FGetRequest> InRequests,
 	IRequestOwner& InOwner,
-	FOnLegacyCacheGetComplete&& InOnComplete)
+	FOnGetComplete&& InOnComplete)
 {
 	if (InRequests.IsEmpty() || !EnumHasAnyFlags(InHierarchy.CombinedNodeFlags, ECacheStoreNodeFlags::HasStoreNode))
 	{
 		return CompleteWithStatus(InRequests, InOnComplete, EStatus::Error);
 	}
 
-	TRefCountPtr<FLegacyGetBatch> State = new FLegacyGetBatch(InHierarchy, InRequests, InOwner, MoveTemp(InOnComplete));
+	TRefCountPtr<TGetBatch> State = new TGetBatch(InHierarchy, InRequests, InOwner, MoveTemp(InOnComplete));
 	State->DispatchRequests();
 }
 
-void FCacheStoreHierarchy::FLegacyGetBatch::DispatchRequests()
+template <typename Params>
+void FCacheStoreHierarchy::TGetBatch<Params>::DispatchRequests()
 {
 	FReadScopeLock Lock(Hierarchy.NodesLock);
 
-	TArray<FLegacyCacheGetRequest, TInlineAllocator<8>> NodeRequests;
-	TArray<FLegacyCachePutRequest, TInlineAllocator<8>> AsyncNodeRequests;
+	TArray<FGetRequest, TInlineAllocator<8>> NodeRequests;
+	TArray<FPutRequest, TInlineAllocator<8>> AsyncNodeRequests;
 
 	const int32 RequestCount = States.Num();
 	NodeRequests.Reserve(RequestCount);
@@ -627,13 +624,13 @@ void FCacheStoreHierarchy::FLegacyGetBatch::DispatchRequests()
 		uint64 StateIndex = 0;
 		for (const FState& State : States)
 		{
-			const FLegacyCacheGetRequest& Request = State.Request;
-			const FLegacyCacheGetResponse& Response = State.Response;
+			const FGetRequest& Request = State.Request;
+			const FGetResponse& Response = State.Response;
 			if (Response.Status == EStatus::Ok)
 			{
-				if (Response.Value && CanStore(Request.Policy, Node.CacheFlags))
+				if (HasResponseData(Response) && CanStore(Request.Policy, Node.CacheFlags))
 				{
-					AsyncNodeRequests.Add({Response.Name, Response.Key, FCompositeBuffer(Response.Value), Request.Policy});
+					AsyncNodeRequests.Add(MakePutRequest(Response, Request));
 				}
 				else if (EnumHasAnyFlags(Node.CacheFlags, ECacheStoreFlags::StopStore) && CanQuery(Request.Policy, Node.CacheFlags))
 				{
@@ -658,15 +655,15 @@ void FCacheStoreHierarchy::FLegacyGetBatch::DispatchRequests()
 		if (!AsyncNodeRequests.IsEmpty())
 		{
 			FRequestBarrier Barrier(AsyncOwner);
-			Node.AsyncCache->LegacyPut(AsyncNodeRequests, AsyncOwner, [](auto&&){});
+			Invoke(Put(), Node.AsyncCache, AsyncNodeRequests, AsyncOwner, [](auto&&){});
 			AsyncNodeRequests.Reset();
 		}
 
 		if (const int32 NodeRequestsCount = NodeRequests.Num())
 		{
 			RemainingRequestCount.Reset(NodeRequestsCount + 1);
-			Node.Cache->LegacyGet(NodeRequests, Owner,
-				[State = TRefCountPtr(this)](FLegacyCacheGetResponse&& Response)
+			Invoke(Get(), Node.Cache, NodeRequests, Owner,
+				[State = TRefCountPtr(this)](FGetResponse&& Response)
 				{
 					State->CompleteRequest(MoveTemp(Response));
 				});
@@ -687,7 +684,8 @@ void FCacheStoreHierarchy::FLegacyGetBatch::DispatchRequests()
 	}
 }
 
-void FCacheStoreHierarchy::FLegacyGetBatch::CompleteRequest(FLegacyCacheGetResponse&& Response)
+template <typename Params>
+void FCacheStoreHierarchy::TGetBatch<Params>::CompleteRequest(FGetResponse&& Response)
 {
 	FState& State = States[int32(Response.UserData)];
 
@@ -699,26 +697,22 @@ void FCacheStoreHierarchy::FLegacyGetBatch::CompleteRequest(FLegacyCacheGetRespo
 
 		// Store to previous writable nodes.
 		const ECachePolicy Policy = State.Request.Policy;
-		if (CanStoreIfOk(Policy, Hierarchy.Nodes[NodeIndex].NodeFlags) && Response.Value)
+		if (CanStoreIfOk(Policy, Hierarchy.Nodes[NodeIndex].NodeFlags) && HasResponseData(Response))
 		{
-			const FLegacyCachePutRequest PutRequest{Response.Name, Response.Key, FCompositeBuffer(Response.Value), Policy};
+			const FPutRequest PutRequest = MakePutRequest(Response, State.Request);
 			for (int32 PutNodeIndex = 0; PutNodeIndex < NodeIndex; ++PutNodeIndex)
 			{
 				const FCacheStoreNode& PutNode = Hierarchy.Nodes[PutNodeIndex];
 				if (CanStore(Policy, PutNode.CacheFlags))
 				{
 					FRequestBarrier Barrier(AsyncOwner);
-					PutNode.AsyncCache->LegacyPut({PutRequest}, AsyncOwner, [](auto&&){});
+					Invoke(Put(), PutNode.AsyncCache, MakeArrayView(&PutRequest, 1), AsyncOwner, [](auto&&){});
 				}
 			}
 		}
 
 		// Value may be fetched to fill other nodes. Remove the value if requested.
-		if (Response.Value && EnumHasAnyFlags(Policy, ECachePolicy::SkipData))
-		{
-			Response.Value.Reset();
-		}
-
+		FilterResponseByRequest(Response, State.Request);
 		OnComplete(MoveTemp(Response));
 	}
 
@@ -741,14 +735,6 @@ void FCacheStoreHierarchy::FLegacyGetBatch::CompleteRequest(FLegacyCacheGetRespo
 	}
 }
 
-void FCacheStoreHierarchy::LegacyGet(
-	const TConstArrayView<FLegacyCacheGetRequest> Requests,
-	IRequestOwner& Owner,
-	FOnLegacyCacheGetComplete&& OnComplete)
-{
-	FLegacyGetBatch::Begin(*this, Requests, Owner, MoveTemp(OnComplete));
-}
-
 ///////////////////////////////////////////////////////////////////////////////////////////////////////////////////////
 
 class FCacheStoreHierarchy::FLegacyDeleteBatch : public FBatchBase
@@ -864,6 +850,146 @@ void FCacheStoreHierarchy::FLegacyDeleteBatch::CompleteRequest(FLegacyCacheDelet
 	}
 }
 
+///////////////////////////////////////////////////////////////////////////////////////////////////////////////////////
+
+void FCacheStoreHierarchy::Put(
+	TConstArrayView<FCachePutRequest> Requests,
+	IRequestOwner& Owner,
+	FOnCachePutComplete&& OnComplete)
+{
+	unimplemented();
+}
+
+void FCacheStoreHierarchy::Get(
+	TConstArrayView<FCacheGetRequest> Requests,
+	IRequestOwner& Owner,
+	FOnCacheGetComplete&& OnComplete)
+{
+	unimplemented();
+}
+
+///////////////////////////////////////////////////////////////////////////////////////////////////////////////////////
+
+template <>
+auto FCacheStoreHierarchy::FCacheValueBatchParams::Put() -> PutFunctionType
+{
+	return &ICacheStore::PutValue;
+};
+
+template <>
+auto FCacheStoreHierarchy::FCacheValueBatchParams::Get() -> GetFunctionType
+{
+	return &ICacheStore::GetValue;
+};
+
+template <>
+bool FCacheStoreHierarchy::FCacheValueBatchParams::HasResponseData(const FCacheGetValueResponse& Response)
+{
+	return Response.Value.HasData();
+}
+
+template <>
+void FCacheStoreHierarchy::FCacheValueBatchParams::FilterResponseByRequest(
+	FCacheGetValueResponse& Response,
+	const FCacheGetValueRequest& Request)
+{
+	if (HasResponseData(Response) && EnumHasAnyFlags(Request.Policy, ECachePolicy::SkipData))
+	{
+		Response.Value = Response.Value.RemoveData();
+	}
+}
+
+template <>
+FCachePutValueRequest FCacheStoreHierarchy::FCacheValueBatchParams::MakePutRequest(
+	const FCacheGetValueResponse& Response,
+	const FCacheGetValueRequest& Request)
+{
+	return {Response.Name, Response.Key, Response.Value, Request.Policy};
+}
+
+void FCacheStoreHierarchy::PutValue(
+	const TConstArrayView<FCachePutValueRequest> Requests,
+	IRequestOwner& Owner,
+	FOnCachePutValueComplete&& OnComplete)
+{
+	using Params = TBatchParams<FCachePutValueRequest, FCacheGetValueRequest>;
+	TPutBatch<Params>::Begin(*this, Requests, Owner, MoveTemp(OnComplete));
+}
+
+void FCacheStoreHierarchy::GetValue(
+	TConstArrayView<FCacheGetValueRequest> Requests,
+	IRequestOwner& Owner,
+	FOnCacheGetValueComplete&& OnComplete)
+{
+	using Params = TBatchParams<FCachePutValueRequest, FCacheGetValueRequest>;
+	TGetBatch<Params>::Begin(*this, Requests, Owner, MoveTemp(OnComplete));
+}
+
+///////////////////////////////////////////////////////////////////////////////////////////////////////////////////////
+
+void FCacheStoreHierarchy::GetChunks(
+	TConstArrayView<FCacheGetChunkRequest> Requests,
+	IRequestOwner& Owner,
+	FOnCacheGetChunkComplete&& OnComplete)
+{
+	unimplemented();
+}
+
+///////////////////////////////////////////////////////////////////////////////////////////////////////////////////////
+
+template <>
+auto FCacheStoreHierarchy::FLegacyCacheBatchParams::Put() -> PutFunctionType
+{
+	return &ILegacyCacheStore::LegacyPut;
+};
+
+template <>
+auto FCacheStoreHierarchy::FLegacyCacheBatchParams::Get() -> GetFunctionType
+{
+	return &ILegacyCacheStore::LegacyGet;
+};
+
+template <>
+bool FCacheStoreHierarchy::FLegacyCacheBatchParams::HasResponseData(const FLegacyCacheGetResponse& Response)
+{
+	return !Response.Value.IsNull();
+}
+
+template <>
+void FCacheStoreHierarchy::FLegacyCacheBatchParams::FilterResponseByRequest(
+	FLegacyCacheGetResponse& Response,
+	const FLegacyCacheGetRequest& Request)
+{
+	if (HasResponseData(Response) && EnumHasAnyFlags(Request.Policy, ECachePolicy::SkipData))
+	{
+		Response.Value.Reset();
+	}
+}
+
+template <>
+FLegacyCachePutRequest FCacheStoreHierarchy::FLegacyCacheBatchParams::MakePutRequest(
+	const FLegacyCacheGetResponse& Response,
+	const FLegacyCacheGetRequest& Request)
+{
+	return {Response.Name, Response.Key, FCompositeBuffer(Response.Value), Request.Policy};
+}
+
+void FCacheStoreHierarchy::LegacyPut(
+	const TConstArrayView<FLegacyCachePutRequest> Requests,
+	IRequestOwner& Owner,
+	FOnLegacyCachePutComplete&& OnComplete)
+{
+	TPutBatch<FLegacyCacheBatchParams>::Begin(*this, Requests, Owner, MoveTemp(OnComplete));
+}
+
+void FCacheStoreHierarchy::LegacyGet(
+	const TConstArrayView<FLegacyCacheGetRequest> Requests,
+	IRequestOwner& Owner,
+	FOnLegacyCacheGetComplete&& OnComplete)
+{
+	TGetBatch<FLegacyCacheBatchParams>::Begin(*this, Requests, Owner, MoveTemp(OnComplete));
+}
+
 void FCacheStoreHierarchy::LegacyDelete(
 	const TConstArrayView<FLegacyCacheDeleteRequest> Requests,
 	IRequestOwner& Owner,
@@ -872,8 +998,6 @@ void FCacheStoreHierarchy::LegacyDelete(
 	FLegacyDeleteBatch::Begin(*this, Requests, Owner, MoveTemp(OnComplete));
 }
 
-///////////////////////////////////////////////////////////////////////////////////////////////////////////////////////
-
 void FCacheStoreHierarchy::LegacyStats(FDerivedDataCacheStatsNode& OutNode)
 {
 	FReadScopeLock Lock(NodesLock);
