DDC: Changed "Type&& Get() &&" accessors to "Type Get() &&"

#rnx

diff --git a/Engine/Source/Developer/DerivedDataCache/Public/DerivedDataBuildAction.h b/Engine/Source/Developer/DerivedDataCache/Public/DerivedDataBuildAction.h
--- a/Engine/Source/Developer/DerivedDataCache/Public/DerivedDataBuildAction.h
+++ b/Engine/Source/Developer/DerivedDataCache/Public/DerivedDataBuildAction.h
@@ -197,7 +197,7 @@ public:
 
 	/** Returns the build action. The caller must check for null before using this accessor. */
 	inline const FBuildAction& Get() const & { return *this; }
-	inline FBuildAction&& Get() && { return MoveTemp(*this); }
+	inline FBuildAction Get() && { return MoveTemp(*this); }
 
 	inline bool IsNull() const { return !IsValid(); }
 	inline bool IsValid() const { return Action.IsValid(); }
diff --git a/Engine/Source/Developer/DerivedDataCache/Public/DerivedDataBuildDefinition.h b/Engine/Source/Developer/DerivedDataCache/Public/DerivedDataBuildDefinition.h
--- a/Engine/Source/Developer/DerivedDataCache/Public/DerivedDataBuildDefinition.h
+++ b/Engine/Source/Developer/DerivedDataCache/Public/DerivedDataBuildDefinition.h
@@ -251,7 +251,7 @@ public:
 
 	/** Returns the build definition. The caller must check for null before using this accessor. */
 	inline const FBuildDefinition& Get() const & { return *this; }
-	inline FBuildDefinition&& Get() && { return MoveTemp(*this); }
+	inline FBuildDefinition Get() && { return MoveTemp(*this); }
 
 	inline bool IsNull() const { return !IsValid(); }
 	inline bool IsValid() const { return Definition.IsValid(); }
diff --git a/Engine/Source/Developer/DerivedDataCache/Public/DerivedDataBuildInputs.h b/Engine/Source/Developer/DerivedDataCache/Public/DerivedDataBuildInputs.h
--- a/Engine/Source/Developer/DerivedDataCache/Public/DerivedDataBuildInputs.h
+++ b/Engine/Source/Developer/DerivedDataCache/Public/DerivedDataBuildInputs.h
@@ -132,7 +132,7 @@ public:
 
 	/** Returns the build inputs. The caller must check for null before using this accessor. */
 	inline const FBuildInputs& Get() const & { return *this; }
-	inline FBuildInputs&& Get() && { return MoveTemp(*this); }
+	inline FBuildInputs Get() && { return MoveTemp(*this); }
 
 	inline bool IsNull() const { return !IsValid(); }
 	inline bool IsValid() const { return Inputs.IsValid(); }
diff --git a/Engine/Source/Developer/DerivedDataCache/Public/DerivedDataBuildOutput.h b/Engine/Source/Developer/DerivedDataCache/Public/DerivedDataBuildOutput.h
--- a/Engine/Source/Developer/DerivedDataCache/Public/DerivedDataBuildOutput.h
+++ b/Engine/Source/Developer/DerivedDataCache/Public/DerivedDataBuildOutput.h
@@ -252,7 +252,7 @@ public:
 
 	/** Returns the build output. The caller must check for null before using this accessor. */
 	inline const FBuildOutput& Get() const & { return *this; }
-	inline FBuildOutput&& Get() && { return MoveTemp(*this); }
+	inline FBuildOutput Get() && { return MoveTemp(*this); }
 
 	inline bool IsNull() const { return !IsValid(); }
 	inline bool IsValid() const { return Output.IsValid(); }
diff --git a/Engine/Source/Developer/DerivedDataCache/Public/DerivedDataBuildSession.h b/Engine/Source/Developer/DerivedDataCache/Public/DerivedDataBuildSession.h
--- a/Engine/Source/Developer/DerivedDataCache/Public/DerivedDataBuildSession.h
+++ b/Engine/Source/Developer/DerivedDataCache/Public/DerivedDataBuildSession.h
@@ -157,7 +157,7 @@ public:
 
 	/** Returns the build session. The caller must check for null before using this accessor. */
 	inline FBuildSession& Get() & { return *this; }
-	inline FBuildSession&& Get() && { return MoveTemp(*this); }
+	inline FBuildSession Get() && { return MoveTemp(*this); }
 
 	inline bool IsNull() const { return !IsValid(); }
 	inline bool IsValid() const { return Session.IsValid(); }
diff --git a/Engine/Source/Developer/DerivedDataCache/Public/DerivedDataCacheRecord.h b/Engine/Source/Developer/DerivedDataCache/Public/DerivedDataCacheRecord.h
--- a/Engine/Source/Developer/DerivedDataCache/Public/DerivedDataCacheRecord.h
+++ b/Engine/Source/Developer/DerivedDataCache/Public/DerivedDataCacheRecord.h
@@ -199,7 +199,7 @@ public:
 
 	/** Returns the cache record. The caller must check for null before using this accessor. */
 	inline const FCacheRecord& Get() const & { return *this; }
-	inline FCacheRecord&& Get() && { return MoveTemp(*this); }
+	inline FCacheRecord Get() && { return MoveTemp(*this); }
 
 	inline bool IsNull() const { return !IsValid(); }
 	inline bool IsValid() const { return Record.IsValid(); }
