DDC: Added Put and Get support to the new cache hierarchy

#rnx

diff --git a/Engine/Source/Developer/DerivedDataCache/Private/DerivedDataCacheStoreHierarchy.cpp b/Engine/Source/Developer/DerivedDataCache/Private/DerivedDataCacheStoreHierarchy.cpp
--- a/Engine/Source/Developer/DerivedDataCache/Private/DerivedDataCacheStoreHierarchy.cpp
+++ b/Engine/Source/Developer/DerivedDataCache/Private/DerivedDataCacheStoreHierarchy.cpp
@@ -1,5 +1,7 @@
 // Copyright Epic Games, Inc. All Rights Reserved.
 
+#include "Algo/AllOf.h"
+#include "Algo/AnyOf.h"
 #include "Algo/Find.h"
 #include "Containers/Array.h"
 #include "DerivedDataCache.h"
@@ -11,6 +13,7 @@
 #include "HAL/CriticalSection.h"
 #include "Misc/EnumClassFlags.h"
 #include "Misc/ScopeRWLock.h"
+#include "Serialization/CompactBinary.h"
 #include "Templates/Invoke.h"
 #include "Templates/RefCounting.h"
 #include "Templates/UniquePtr.h"
@@ -82,6 +85,7 @@ public:
 	bool LegacyDebugOptions(FBackendDebugOptions& Options) final;
 
 	template <typename PutRequestType, typename GetRequestType> struct TBatchParams;
+	using FCacheRecordBatchParams = TBatchParams<FCachePutRequest, FCacheGetRequest>;
 	using FCacheValueBatchParams = TBatchParams<FCachePutValueRequest, FCacheGetValueRequest>;
 	using FLegacyCacheBatchParams = TBatchParams<FLegacyCachePutRequest, FLegacyCacheGetRequest>;
 
@@ -100,6 +104,14 @@ private:
 	enum class ECacheStoreNodeFlags : uint32;
 	FRIEND_ENUM_CLASS_FLAGS(ECacheStoreNodeFlags);
 
+	static ECachePolicy GetCombinedPolicy(const ECachePolicy Policy) { return Policy; }
+	static ECachePolicy GetCombinedPolicy(const FCacheRecordPolicy& Policy) { return Policy.GetRecordPolicy(); }
+
+	static ECachePolicy AddPolicy(ECachePolicy BasePolicy, ECachePolicy Policy) { return BasePolicy | Policy; }
+	static ECachePolicy RemovePolicy(ECachePolicy BasePolicy, ECachePolicy Policy) { return BasePolicy & ~Policy; }
+	static FCacheRecordPolicy AddPolicy(const FCacheRecordPolicy& BasePolicy, ECachePolicy Policy);
+	static FCacheRecordPolicy RemovePolicy(const FCacheRecordPolicy& BasePolicy, ECachePolicy Policy);
+
 	static bool CanQuery(ECachePolicy Policy, ECacheStoreFlags Flags);
 	static bool CanStore(ECachePolicy Policy, ECacheStoreFlags Flags);
 	static bool CanStoreIfOk(ECachePolicy Policy, ECacheStoreNodeFlags Flags);
@@ -174,6 +186,7 @@ struct FCacheStoreHierarchy::TBatchParams
 	static bool HasResponseData(const FGetResponse& Response);
 	static void FilterResponseByRequest(FGetResponse& Response, const FGetRequest& Request);
 	static FPutRequest MakePutRequest(const FGetResponse& Response, const FGetRequest& Request);
+	static FGetRequest MakeGetRequest(const FPutRequest& Request, uint64 UserData);
 };
 
 ///////////////////////////////////////////////////////////////////////////////////////////////////////////////////////
@@ -264,6 +277,16 @@ void FCacheStoreHierarchy::UpdateNodeFlags()
 	CombinedNodeFlags = StoreFlags | QueryFlags;
 }
 
+FCacheRecordPolicy FCacheStoreHierarchy::AddPolicy(const FCacheRecordPolicy& BasePolicy, ECachePolicy Policy)
+{
+	return BasePolicy.Transform([Policy](ECachePolicy P) { return P | Policy; });
+}
+
+FCacheRecordPolicy FCacheStoreHierarchy::RemovePolicy(const FCacheRecordPolicy& BasePolicy, ECachePolicy Policy)
+{
+	return BasePolicy.Transform([Policy](ECachePolicy P) { return P & ~Policy; });
+}
+
 bool FCacheStoreHierarchy::CanQuery(const ECachePolicy Policy, const ECacheStoreFlags Flags)
 {
 	const ECacheStoreFlags LocationFlags =
@@ -309,6 +332,7 @@ class FCacheStoreHierarchy::TPutBatch : public FBatchBase, public Params
 	using FOnGetComplete = typename Params::FOnGetComplete;
 	using Params::Put;
 	using Params::Get;
+	using Params::MakeGetRequest;
 
 public:
 	static void Begin(
@@ -421,9 +445,9 @@ bool FCacheStoreHierarchy::TPutBatch<Params>::DispatchGetRequests()
 	uint64 RequestIndex = 0;
 	for (const FPutRequest& Request : Requests)
 	{
-		if (!States[RequestIndex].bStop && CanQuery(Request.Policy, Node.CacheFlags))
+		if (!States[RequestIndex].bStop && CanQuery(GetCombinedPolicy(Request.Policy), Node.CacheFlags))
 		{
-			NodeRequests.Add({Request.Name, Request.Key, Request.Policy | ECachePolicy::SkipData, RequestIndex});
+			NodeRequests.Add(MakeGetRequest(Request, RequestIndex));
 		}
 		++RequestIndex;
 	}
@@ -453,7 +477,7 @@ void FCacheStoreHierarchy::TPutBatch<Params>::CompleteGetRequest(FGetResponse&&
 		State.bStop = true;
 		if (!State.bOk)
 		{
-			OnComplete({Response.Name, Response.Key, Requests[RequestIndex].UserData, Response.Status});
+			OnComplete(Requests[RequestIndex].MakeResponse(Response.Status));
 		}
 	}
 	if (RemainingRequestCount.Signal())
@@ -482,7 +506,7 @@ bool FCacheStoreHierarchy::TPutBatch<Params>::DispatchPutRequests()
 	for (const FPutRequest& Request : Requests)
 	{
 		const FRequestState& State = States[RequestIndex];
-		if (!State.bStop && CanStore(Request.Policy, Node.CacheFlags))
+		if (!State.bStop && CanStore(GetCombinedPolicy(Request.Policy), Node.CacheFlags))
 		{
 			(State.bOk ? AsyncNodeRequests : NodeRequests).Add_GetRef(Request).UserData = uint64(RequestIndex);
 		}
@@ -566,7 +590,7 @@ private:
 		States.Reserve(InRequests.Num());
 		for (const FGetRequest& Request : InRequests)
 		{
-			States.Add({Request});
+			States.Add({Request, Request.MakeResponse(EStatus::Error)});
 		}
 	}
 
@@ -628,25 +652,27 @@ void FCacheStoreHierarchy::TGetBatch<Params>::DispatchRequests()
 			const FGetResponse& Response = State.Response;
 			if (Response.Status == EStatus::Ok)
 			{
-				if (HasResponseData(Response) && CanStore(Request.Policy, Node.CacheFlags))
+				if (HasResponseData(Response) && CanStore(GetCombinedPolicy(Request.Policy), Node.CacheFlags))
 				{
 					AsyncNodeRequests.Add(MakePutRequest(Response, Request));
 				}
-				else if (EnumHasAnyFlags(Node.CacheFlags, ECacheStoreFlags::StopStore) && CanQuery(Request.Policy, Node.CacheFlags))
+				else if (EnumHasAnyFlags(Node.CacheFlags, ECacheStoreFlags::StopStore) && CanQuery(GetCombinedPolicy(Request.Policy), Node.CacheFlags))
 				{
-					NodeRequests.Add({Request.Name, Request.Key, Request.Policy | ECachePolicy::SkipData, StateIndex});
+					NodeRequests.Add({Request.Name, Request.Key, AddPolicy(Request.Policy, ECachePolicy::SkipData), StateIndex});
 				}
 			}
 			else
 			{
-				if (CanQuery(Request.Policy, Node.CacheFlags))
+				if (const ECachePolicy Policy = GetCombinedPolicy(Request.Policy); CanQuery(Policy, Node.CacheFlags))
 				{
-					ECachePolicy Policy = Request.Policy;
 					if (EnumHasAnyFlags(Policy, ECachePolicy::SkipData) && CanStoreIfOk(Policy, Node.NodeFlags))
 					{
-						EnumRemoveFlags(Policy, ECachePolicy::SkipData);
+						NodeRequests.Add({Request.Name, Request.Key, RemovePolicy(Request.Policy, ECachePolicy::SkipData), StateIndex});
+					}
+					else
+					{
+						NodeRequests.Add({Request.Name, Request.Key, Request.Policy, StateIndex});
 					}
-					NodeRequests.Add({Request.Name, Request.Key, Policy, StateIndex});
 				}
 			}
 			++StateIndex;
@@ -675,11 +701,11 @@ void FCacheStoreHierarchy::TGetBatch<Params>::DispatchRequests()
 		}
 	}
 
-	for (const FState& State : States)
+	for (FState& State : States)
 	{
 		if (State.Response.Status != EStatus::Ok)
 		{
-			OnComplete(State.Request.MakeResponse(Owner.IsCanceled() ? EStatus::Canceled : EStatus::Error));
+			OnComplete(MoveTemp(State.Response));
 		}
 	}
 }
@@ -688,47 +714,55 @@ template <typename Params>
 void FCacheStoreHierarchy::TGetBatch<Params>::CompleteRequest(FGetResponse&& Response)
 {
 	FState& State = States[int32(Response.UserData)];
+	const FCacheStoreNode& Node = Hierarchy.Nodes[NodeIndex];
 
-	if (Response.Status == EStatus::Ok && State.Response.Status == EStatus::Error)
+	const bool bFirstOk = Response.Status == EStatus::Ok && State.Response.Status == EStatus::Error;
+	const bool bLastQuery = bFirstOk || !CanQueryIfError(GetCombinedPolicy(State.Request.Policy), Node.NodeFlags);
+
+	if (State.Response.Status == EStatus::Error)
 	{
-		check(State.Response.Status == EStatus::Error);
 		Response.UserData = State.Request.UserData;
+		// TODO: Merge values from partial records.
 		State.Response = Response;
+	}
 
-		// Store to previous writable nodes.
-		const ECachePolicy Policy = State.Request.Policy;
-		if (CanStoreIfOk(Policy, Hierarchy.Nodes[NodeIndex].NodeFlags) && HasResponseData(Response))
+	if (bLastQuery && CanStoreIfOk(GetCombinedPolicy(State.Request.Policy), Node.NodeFlags) && HasResponseData(Response))
+	{
+		// Store any retrieved values to previous writable nodes if Ok or there are no remaining nodes to query.
+		const FPutRequest PutRequest = MakePutRequest(Response, State.Request);
+		for (int32 PutNodeIndex = 0; PutNodeIndex < NodeIndex; ++PutNodeIndex)
 		{
-			const FPutRequest PutRequest = MakePutRequest(Response, State.Request);
-			for (int32 PutNodeIndex = 0; PutNodeIndex < NodeIndex; ++PutNodeIndex)
+			const FCacheStoreNode& PutNode = Hierarchy.Nodes[PutNodeIndex];
+			if (CanStore(GetCombinedPolicy(State.Request.Policy), PutNode.CacheFlags))
 			{
-				const FCacheStoreNode& PutNode = Hierarchy.Nodes[PutNodeIndex];
-				if (CanStore(Policy, PutNode.CacheFlags))
-				{
-					FRequestBarrier Barrier(AsyncOwner);
-					Invoke(Put(), PutNode.AsyncCache, MakeArrayView(&PutRequest, 1), AsyncOwner, [](auto&&){});
-				}
+				FRequestBarrier Barrier(AsyncOwner);
+				Invoke(Put(), PutNode.AsyncCache, MakeArrayView(&PutRequest, 1), AsyncOwner, [](auto&&){});
 			}
 		}
-
-		// Value may be fetched to fill other nodes. Remove the value if requested.
-		FilterResponseByRequest(Response, State.Request);
-		OnComplete(MoveTemp(Response));
 	}
 
-	if (Response.Status == EStatus::Ok)
+	if (State.Response.Status == EStatus::Ok)
 	{
-		// Never store to later remote nodes.
-		// This is a necessary optimization until these speculative stores can be optimized.
-		EnumRemoveFlags(State.Request.Policy, ECachePolicy::Remote);
-
-		// Never store to later later nodes if this node has StopStore.
-		if (EnumHasAnyFlags(Hierarchy.Nodes[NodeIndex].CacheFlags, ECacheStoreFlags::StopStore))
+		if (EnumHasAnyFlags(Node.CacheFlags, ECacheStoreFlags::StopStore))
+		{
+			// Never store to later later nodes if this node has StopStore.
+			State.Request.Policy = RemovePolicy(State.Request.Policy, ECachePolicy::Default);
+		}
+		else
 		{
-			EnumRemoveFlags(State.Request.Policy, ECachePolicy::Default);
+			// Never store to later remote nodes.
+			// This is a necessary optimization until speculative stores can be optimized.
+			State.Request.Policy = RemovePolicy(State.Request.Policy, ECachePolicy::Remote);
 		}
 	}
 
+	if (bFirstOk)
+	{
+		// Values may be fetched to fill previous nodes. Remove values if requested.
+		FilterResponseByRequest(Response, State.Request);
+		OnComplete(MoveTemp(Response));
+	}
+
 	if (RemainingRequestCount.Signal())
 	{
 		DispatchRequests();
@@ -852,12 +886,80 @@ void FCacheStoreHierarchy::FLegacyDeleteBatch::CompleteRequest(FLegacyCacheDelet
 
 ///////////////////////////////////////////////////////////////////////////////////////////////////////////////////////
 
+template <>
+auto FCacheStoreHierarchy::FCacheRecordBatchParams::Put() -> PutFunctionType
+{
+	return &ICacheStore::Put;
+};
+
+template <>
+auto FCacheStoreHierarchy::FCacheRecordBatchParams::Get() -> GetFunctionType
+{
+	return &ICacheStore::Get;
+};
+
+template <>
+bool FCacheStoreHierarchy::FCacheRecordBatchParams::HasResponseData(const FCacheGetResponse& Response)
+{
+	return Algo::AnyOf(Response.Record.GetValues(), &FValue::HasData);
+}
+
+template <>
+void FCacheStoreHierarchy::FCacheRecordBatchParams::FilterResponseByRequest(
+	FCacheGetResponse& Response,
+	const FCacheGetRequest& Request)
+{
+	if ((!Request.Policy.IsUniform() && Algo::AnyOf(Response.Record.GetValues(), &FValue::HasData)) ||
+		(EnumHasAnyFlags(Request.Policy.GetRecordPolicy(), ECachePolicy::SkipMeta) && Response.Record.GetMeta()))
+	{
+		FCacheRecordBuilder Builder(Response.Record.GetKey());
+		if (!EnumHasAnyFlags(Request.Policy.GetRecordPolicy(), ECachePolicy::SkipMeta))
+		{
+			Builder.SetMeta(CopyTemp(Response.Record.GetMeta()));
+		}
+		for (const FValueWithId& Value : Response.Record.GetValues())
+		{
+			if (EnumHasAnyFlags(Request.Policy.GetValuePolicy(Value.GetId()), ECachePolicy::SkipData))
+			{
+				Builder.AddValue(Value.GetId(), Value.RemoveData());
+			}
+			else
+			{
+				Builder.AddValue(Value);
+			}
+		}
+		Response.Record = Builder.Build();
+	}
+}
+
+template <>
+FCachePutRequest FCacheStoreHierarchy::FCacheRecordBatchParams::MakePutRequest(
+	const FCacheGetResponse& Response,
+	const FCacheGetRequest& Request)
+{
+	FCacheRecordPolicy Policy = Request.Policy;
+	if (!Algo::AllOf(Response.Record.GetValues(), &FValue::HasData) &&
+		!EnumHasAnyFlags(Policy.GetRecordPolicy(), ECachePolicy::PartialRecord))
+	{
+		Policy = Policy.Transform([](ECachePolicy P) { return P | ECachePolicy::PartialRecord; });
+	}
+	return {Response.Name, Response.Record, MoveTemp(Policy)};
+}
+
+template <>
+FCacheGetRequest FCacheStoreHierarchy::FCacheRecordBatchParams::MakeGetRequest(
+	const FCachePutRequest& Request,
+	const uint64 UserData)
+{
+	return {Request.Name, Request.Record.GetKey(), AddPolicy(Request.Policy, ECachePolicy::SkipData), UserData};
+}
+
 void FCacheStoreHierarchy::Put(
 	TConstArrayView<FCachePutRequest> Requests,
 	IRequestOwner& Owner,
 	FOnCachePutComplete&& OnComplete)
 {
-	unimplemented();
+	TPutBatch<FCacheRecordBatchParams>::Begin(*this, Requests, Owner, MoveTemp(OnComplete));
 }
 
 void FCacheStoreHierarchy::Get(
@@ -865,7 +967,7 @@ void FCacheStoreHierarchy::Get(
 	IRequestOwner& Owner,
 	FOnCacheGetComplete&& OnComplete)
 {
-	unimplemented();
+	TGetBatch<FCacheRecordBatchParams>::Begin(*this, Requests, Owner, MoveTemp(OnComplete));
 }
 
 ///////////////////////////////////////////////////////////////////////////////////////////////////////////////////////
@@ -893,7 +995,7 @@ void FCacheStoreHierarchy::FCacheValueBatchParams::FilterResponseByRequest(
 	FCacheGetValueResponse& Response,
 	const FCacheGetValueRequest& Request)
 {
-	if (HasResponseData(Response) && EnumHasAnyFlags(Request.Policy, ECachePolicy::SkipData))
+	if (Response.Value.HasData() && EnumHasAnyFlags(Request.Policy, ECachePolicy::SkipData))
 	{
 		Response.Value = Response.Value.RemoveData();
 	}
@@ -907,13 +1009,20 @@ FCachePutValueRequest FCacheStoreHierarchy::FCacheValueBatchParams::MakePutReque
 	return {Response.Name, Response.Key, Response.Value, Request.Policy};
 }
 
+template <>
+FCacheGetValueRequest FCacheStoreHierarchy::FCacheValueBatchParams::MakeGetRequest(
+	const FCachePutValueRequest& Request,
+	const uint64 UserData)
+{
+	return {Request.Name, Request.Key, AddPolicy(Request.Policy, ECachePolicy::SkipData), UserData};
+}
+
 void FCacheStoreHierarchy::PutValue(
 	const TConstArrayView<FCachePutValueRequest> Requests,
 	IRequestOwner& Owner,
 	FOnCachePutValueComplete&& OnComplete)
 {
-	using Params = TBatchParams<FCachePutValueRequest, FCacheGetValueRequest>;
-	TPutBatch<Params>::Begin(*this, Requests, Owner, MoveTemp(OnComplete));
+	TPutBatch<FCacheValueBatchParams>::Begin(*this, Requests, Owner, MoveTemp(OnComplete));
 }
 
 void FCacheStoreHierarchy::GetValue(
@@ -921,8 +1030,7 @@ void FCacheStoreHierarchy::GetValue(
 	IRequestOwner& Owner,
 	FOnCacheGetValueComplete&& OnComplete)
 {
-	using Params = TBatchParams<FCachePutValueRequest, FCacheGetValueRequest>;
-	TGetBatch<Params>::Begin(*this, Requests, Owner, MoveTemp(OnComplete));
+	TGetBatch<FCacheValueBatchParams>::Begin(*this, Requests, Owner, MoveTemp(OnComplete));
 }
 
 ///////////////////////////////////////////////////////////////////////////////////////////////////////////////////////
@@ -960,7 +1068,7 @@ void FCacheStoreHierarchy::FLegacyCacheBatchParams::FilterResponseByRequest(
 	FLegacyCacheGetResponse& Response,
 	const FLegacyCacheGetRequest& Request)
 {
-	if (HasResponseData(Response) && EnumHasAnyFlags(Request.Policy, ECachePolicy::SkipData))
+	if (Response.Value && EnumHasAnyFlags(Request.Policy, ECachePolicy::SkipData))
 	{
 		Response.Value.Reset();
 	}
@@ -974,6 +1082,14 @@ FLegacyCachePutRequest FCacheStoreHierarchy::FLegacyCacheBatchParams::MakePutReq
 	return {Response.Name, Response.Key, FCompositeBuffer(Response.Value), Request.Policy};
 }
 
+template <>
+FLegacyCacheGetRequest FCacheStoreHierarchy::FLegacyCacheBatchParams::MakeGetRequest(
+	const FLegacyCachePutRequest& Request,
+	const uint64 UserData)
+{
+	return {Request.Name, Request.Key, AddPolicy(Request.Policy, ECachePolicy::SkipData), UserData};
+}
+
 void FCacheStoreHierarchy::LegacyPut(
 	const TConstArrayView<FLegacyCachePutRequest> Requests,
 	IRequestOwner& Owner,
