DDC: Moved the delete batch in the new hierarchy next to the delete function

#rnx

diff --git a/Engine/Source/Developer/DerivedDataCache/Private/DerivedDataCacheStoreHierarchy.cpp b/Engine/Source/Developer/DerivedDataCache/Private/DerivedDataCacheStoreHierarchy.cpp
--- a/Engine/Source/Developer/DerivedDataCache/Private/DerivedDataCacheStoreHierarchy.cpp
+++ b/Engine/Source/Developer/DerivedDataCache/Private/DerivedDataCacheStoreHierarchy.cpp
@@ -771,121 +771,6 @@
 
 ///////////////////////////////////////////////////////////////////////////////////////////////////////////////////////
 
-class FCacheStoreHierarchy::FLegacyDeleteBatch : public FBatchBase
-{
-public:
-	static void Begin(
-		const FCacheStoreHierarchy& Hierarchy,
-		TConstArrayView<FLegacyCacheDeleteRequest> Requests,
-		IRequestOwner& Owner,
-		FOnLegacyCacheDeleteComplete&& OnComplete);
-
-private:
-	FLegacyDeleteBatch(
-		const FCacheStoreHierarchy& InHierarchy,
-		const TConstArrayView<FLegacyCacheDeleteRequest> InRequests,
-		IRequestOwner& InOwner,
-		FOnLegacyCacheDeleteComplete&& InOnComplete)
-		: Hierarchy(InHierarchy)
-		, Requests(InRequests)
-		, BatchOwner(InOwner)
-		, OnComplete(MoveTemp(InOnComplete))
-	{
-		States.SetNum(Requests.Num());
-	}
-
-	void DispatchRequests();
-	void CompleteRequest(FLegacyCacheDeleteResponse&& Response);
-
-	struct FRequestState
-	{
-		bool bOk = false;
-	};
-
-	const FCacheStoreHierarchy& Hierarchy;
-	TArray<FLegacyCacheDeleteRequest> Requests;
-	IRequestOwner& BatchOwner;
-	FOnLegacyCacheDeleteComplete OnComplete;
-
-	TArray<FRequestState, TInlineAllocator<1>> States;
-	FCounterEvent RemainingRequestCount;
-	int32 NodeIndex = 0;
-};
-
-void FCacheStoreHierarchy::FLegacyDeleteBatch::Begin(
-	const FCacheStoreHierarchy& Hierarchy,
-	const TConstArrayView<FLegacyCacheDeleteRequest> Requests,
-	IRequestOwner& Owner,
-	FOnLegacyCacheDeleteComplete&& OnComplete)
-{
-	if (Requests.IsEmpty() || !EnumHasAnyFlags(Hierarchy.CombinedNodeFlags, ECacheStoreNodeFlags::HasStoreNode))
-	{
-		return CompleteWithStatus(Requests, OnComplete, EStatus::Error);
-	}
-
-	TRefCountPtr<FLegacyDeleteBatch> State = new FLegacyDeleteBatch(Hierarchy, Requests, Owner, MoveTemp(OnComplete));
-	State->DispatchRequests();
-}
-
-void FCacheStoreHierarchy::FLegacyDeleteBatch::DispatchRequests()
-{
-	FReadScopeLock Lock(Hierarchy.NodesLock);
-
-	TArray<FLegacyCacheDeleteRequest, TInlineAllocator<1>> NodeRequests;
-	NodeRequests.Reserve(Requests.Num());
-	for (const int32 NodeCount = Hierarchy.Nodes.Num(); NodeIndex < NodeCount && !BatchOwner.IsCanceled(); ++NodeIndex)
-	{
-		const FCacheStoreNode& Node = Hierarchy.Nodes[NodeIndex];
-
-		uint64 RequestIndex = 0;
-		for (const FLegacyCacheDeleteRequest& Request : Requests)
-		{
-			if (CanStore(Request.Policy, Node.CacheFlags))
-			{
-				NodeRequests.Add_GetRef(Request).UserData = RequestIndex;
-			}
-			++RequestIndex;
-		}
-
-		if (const int32 NodeRequestsCount = NodeRequests.Num())
-		{
-			RemainingRequestCount.Reset(NodeRequestsCount + 1);
-			Node.Cache->LegacyDelete(NodeRequests, BatchOwner,
-				[State = TRefCountPtr(this)](FLegacyCacheDeleteResponse&& Response)
-				{
-					State->CompleteRequest(MoveTemp(Response));
-				});
-			NodeRequests.Reset();
-			if (!RemainingRequestCount.Signal())
-			{
-				return;
-			}
-		}
-	}
-
-	int32 RequestIndex = 0;
-	for (const FLegacyCacheDeleteRequest& Request : Requests)
-	{
-		const bool bOk = States[RequestIndex++].bOk;
-		const EStatus Status = bOk ? EStatus::Ok : BatchOwner.IsCanceled() ? EStatus::Canceled : EStatus::Error;
-		OnComplete(Request.MakeResponse(Status));
-	}
-}
-
-void FCacheStoreHierarchy::FLegacyDeleteBatch::CompleteRequest(FLegacyCacheDeleteResponse&& Response)
-{
-	if (Response.Status == EStatus::Ok)
-	{
-		States[int32(Response.UserData)].bOk = true;
-	}
-	if (RemainingRequestCount.Signal())
-	{
-		DispatchRequests();
-	}
-}
-
-///////////////////////////////////////////////////////////////////////////////////////////////////////////////////////
-
 template <>
 auto FCacheStoreHierarchy::FCacheRecordBatchParams::Put() -> PutFunctionType
 {
@@ -1106,6 +991,121 @@
 	TGetBatch<FLegacyCacheBatchParams>::Begin(*this, Requests, Owner, MoveTemp(OnComplete));
 }
 
+///////////////////////////////////////////////////////////////////////////////////////////////////////////////////////
+
+class FCacheStoreHierarchy::FLegacyDeleteBatch : public FBatchBase
+{
+public:
+	static void Begin(
+		const FCacheStoreHierarchy& Hierarchy,
+		TConstArrayView<FLegacyCacheDeleteRequest> Requests,
+		IRequestOwner& Owner,
+		FOnLegacyCacheDeleteComplete&& OnComplete);
+
+private:
+	FLegacyDeleteBatch(
+		const FCacheStoreHierarchy& InHierarchy,
+		const TConstArrayView<FLegacyCacheDeleteRequest> InRequests,
+		IRequestOwner& InOwner,
+		FOnLegacyCacheDeleteComplete&& InOnComplete)
+		: Hierarchy(InHierarchy)
+		, Requests(InRequests)
+		, BatchOwner(InOwner)
+		, OnComplete(MoveTemp(InOnComplete))
+	{
+		States.SetNum(Requests.Num());
+	}
+
+	void DispatchRequests();
+	void CompleteRequest(FLegacyCacheDeleteResponse&& Response);
+
+	struct FRequestState
+	{
+		bool bOk = false;
+	};
+
+	const FCacheStoreHierarchy& Hierarchy;
+	TArray<FLegacyCacheDeleteRequest> Requests;
+	IRequestOwner& BatchOwner;
+	FOnLegacyCacheDeleteComplete OnComplete;
+
+	TArray<FRequestState, TInlineAllocator<1>> States;
+	FCounterEvent RemainingRequestCount;
+	int32 NodeIndex = 0;
+};
+
+void FCacheStoreHierarchy::FLegacyDeleteBatch::Begin(
+	const FCacheStoreHierarchy& Hierarchy,
+	const TConstArrayView<FLegacyCacheDeleteRequest> Requests,
+	IRequestOwner& Owner,
+	FOnLegacyCacheDeleteComplete&& OnComplete)
+{
+	if (Requests.IsEmpty() || !EnumHasAnyFlags(Hierarchy.CombinedNodeFlags, ECacheStoreNodeFlags::HasStoreNode))
+	{
+		return CompleteWithStatus(Requests, OnComplete, EStatus::Error);
+	}
+
+	TRefCountPtr<FLegacyDeleteBatch> State = new FLegacyDeleteBatch(Hierarchy, Requests, Owner, MoveTemp(OnComplete));
+	State->DispatchRequests();
+}
+
+void FCacheStoreHierarchy::FLegacyDeleteBatch::DispatchRequests()
+{
+	FReadScopeLock Lock(Hierarchy.NodesLock);
+
+	TArray<FLegacyCacheDeleteRequest, TInlineAllocator<1>> NodeRequests;
+	NodeRequests.Reserve(Requests.Num());
+	for (const int32 NodeCount = Hierarchy.Nodes.Num(); NodeIndex < NodeCount && !BatchOwner.IsCanceled(); ++NodeIndex)
+	{
+		const FCacheStoreNode& Node = Hierarchy.Nodes[NodeIndex];
+
+		uint64 RequestIndex = 0;
+		for (const FLegacyCacheDeleteRequest& Request : Requests)
+		{
+			if (CanStore(Request.Policy, Node.CacheFlags))
+			{
+				NodeRequests.Add_GetRef(Request).UserData = RequestIndex;
+			}
+			++RequestIndex;
+		}
+
+		if (const int32 NodeRequestsCount = NodeRequests.Num())
+		{
+			RemainingRequestCount.Reset(NodeRequestsCount + 1);
+			Node.Cache->LegacyDelete(NodeRequests, BatchOwner,
+				[State = TRefCountPtr(this)](FLegacyCacheDeleteResponse&& Response)
+				{
+					State->CompleteRequest(MoveTemp(Response));
+				});
+			NodeRequests.Reset();
+			if (!RemainingRequestCount.Signal())
+			{
+				return;
+			}
+		}
+	}
+
+	int32 RequestIndex = 0;
+	for (const FLegacyCacheDeleteRequest& Request : Requests)
+	{
+		const bool bOk = States[RequestIndex++].bOk;
+		const EStatus Status = bOk ? EStatus::Ok : BatchOwner.IsCanceled() ? EStatus::Canceled : EStatus::Error;
+		OnComplete(Request.MakeResponse(Status));
+	}
+}
+
+void FCacheStoreHierarchy::FLegacyDeleteBatch::CompleteRequest(FLegacyCacheDeleteResponse&& Response)
+{
+	if (Response.Status == EStatus::Ok)
+	{
+		States[int32(Response.UserData)].bOk = true;
+	}
+	if (RemainingRequestCount.Signal())
+	{
+		DispatchRequests();
+	}
+}
+
 void FCacheStoreHierarchy::LegacyDelete(
 	const TConstArrayView<FLegacyCacheDeleteRequest> Requests,
 	IRequestOwner& Owner,
@@ -1114,6 +1114,8 @@
 	FLegacyDeleteBatch::Begin(*this, Requests, Owner, MoveTemp(OnComplete));
 }
 
+///////////////////////////////////////////////////////////////////////////////////////////////////////////////////////
+
 void FCacheStoreHierarchy::LegacyStats(FDerivedDataCacheStatsNode& OutNode)
 {
 	FReadScopeLock Lock(NodesLock);
