ZenServer CachePolicy support.

Add batched Put and PutValue.
Update Batched Get to send the RecordPolicy with each key.
Update GetChunks to send the CachePolicy as string and to handle SkipData in the response.
Add more unit tests  for upstream behavior and GetChunks.
Update the CacheRecordPolicy Save/Load to match the new variable names on FCacheRecordPolicy.
Deploy new zenserver binaries with the required matching server-side changes.

#rnx

diff --git a/Engine/Source/Developer/DerivedDataCache/Private/BatchView.h b/Engine/Source/Developer/DerivedDataCache/Private/BatchView.h
new file mode 100644
--- /dev/null
+++ b/Engine/Source/Developer/DerivedDataCache/Private/BatchView.h
@@ -0,0 +1,88 @@
+// Copyright Epic Games, Inc. All Rights Reserved.
+
+#pragma once
+
+#include "Containers/Array.h"
+#include "Containers/ArrayView.h"
+#include "HAL/Platform.h"
+#include "Misc/AssertionMacros.h"
+
+enum class EBatchView
+{
+	/** Add the item to the previous batch. If no previous batch, treated the same as NewBatch. */
+	Continue,
+	/** End the previous batch (if it exists), start a new batch, add the item to the new batch. */
+	NewBatch,
+};
+
+/**
+ * Takes a list of items and a ShouldAddToBatch function and provides
+ * a ranged-for over the set of batches created from the list.
+ */
+template <typename T>
+class TBatchView
+{
+public:
+	/**
+	 * Constructor from list of elements and a function that is called on each element
+	 * to decide whether to add the element to the current batch.
+	 *
+	 * EBatchView ShouldAddToBatch(const T& Item);
+	 */
+	template <typename ShouldAddToBatchFunc>
+	TBatchView(TArrayView<T> InItems, ShouldAddToBatchFunc&& ShouldAddToBatch)
+		: Items(InItems)
+	{
+		if (Items.IsEmpty())
+		{
+			return;
+		}
+		ShouldAddToBatch(Items[0]);
+		BatchStarts.Add(0);
+		for (int32 Index = 1; Index < Items.Num(); ++Index)
+		{
+			if (ShouldAddToBatch(Items[Index]) == EBatchView::NewBatch)
+			{
+				BatchStarts.Add(Index);
+			}
+		}
+	}
+
+	struct TIterator
+	{
+		TIterator(const TBatchView& InBatchList, bool bEnd)
+			: BatchList(InBatchList)
+			, BatchIndex(bEnd ? BatchList.BatchStarts.Num() : 0)
+		{
+		}
+		TIterator& operator++()
+		{
+			BatchIndex++;
+			return *this;
+		}
+		TArrayView<T> operator*() const
+		{
+			int32 Start = BatchList.BatchStarts[BatchIndex];
+			int32 End = BatchIndex < BatchList.BatchStarts.Num() - 1 ? BatchList.BatchStarts[BatchIndex + 1] : BatchList.Items.Num();
+			return TArrayView<T>(BatchList.Items.GetData() + Start, End - Start);
+		}
+		bool operator!=(TIterator& Other) const
+		{
+			return BatchIndex != Other.BatchIndex;
+		}
+		const TBatchView& BatchList;
+		int32 BatchIndex;
+	};
+	TIterator begin() const
+	{
+		return TIterator(*this, false /* bEnd */);
+	}
+	TIterator end() const
+	{
+		return TIterator(*this, true /* bEnd */);
+	}
+
+private:
+	TArrayView<T> Items;
+	TArray<int32> BatchStarts;
+};
\ No newline at end of file
diff --git a/Engine/Source/Developer/DerivedDataCache/Private/DerivedDataCachePolicy.cpp b/Engine/Source/Developer/DerivedDataCache/Private/DerivedDataCachePolicy.cpp
--- a/Engine/Source/Developer/DerivedDataCache/Private/DerivedDataCachePolicy.cpp
+++ b/Engine/Source/Developer/DerivedDataCache/Private/DerivedDataCachePolicy.cpp
@@ -186,7 +186,7 @@ FCacheRecordPolicy FCacheRecordPolicy::Transform(const TFunctionRef<ECachePolicy
 void FCacheRecordPolicy::Save(FCbWriter& Writer) const
 {
 	Writer.BeginObject();
-	Writer.AddString("DefaultValuePolicy"_ASV, WriteToUtf8String<128>(GetBasePolicy()));
+	Writer.AddString("BasePolicy"_ASV, WriteToUtf8String<128>(GetBasePolicy()));
 	if (!IsUniform())
 	{
 		Writer.BeginArray("ValuePolicies"_ASV);
@@ -204,7 +204,7 @@ void FCacheRecordPolicy::Save(FCbWriter& Writer) const
 
 FOptionalCacheRecordPolicy FCacheRecordPolicy::Load(const FCbObjectView Object)
 {
-	const FUtf8StringView BasePolicyText = Object["DefaultValuePolicy"_ASV].AsString();
+	const FUtf8StringView BasePolicyText = Object["BasePolicy"_ASV].AsString();
 	if (BasePolicyText.IsEmpty())
 	{
 		return {};
diff --git a/Engine/Source/Developer/DerivedDataCache/Private/Tests/CacheStoreTest.cpp b/Engine/Source/Developer/DerivedDataCache/Private/Tests/CacheStoreTest.cpp
--- a/Engine/Source/Developer/DerivedDataCache/Private/Tests/CacheStoreTest.cpp
+++ b/Engine/Source/Developer/DerivedDataCache/Private/Tests/CacheStoreTest.cpp
@@ -5,159 +5,244 @@
 #if WITH_DEV_AUTOMATION_TESTS
 
 #include "Containers/Array.h"
+#include "Containers/ArrayView.h"
 #include "Containers/Map.h"
 #include "DerivedDataCache.h"
 #include "DerivedDataRequestOwner.h"
 #include "Hash/Blake3.h"
+#include "Misc/CommandLine.h"
+#include "Misc/Parse.h"
 #include "Misc/StringBuilder.h"
 
 IMPLEMENT_SIMPLE_AUTOMATION_TEST(FCacheStoreTest, "System.DerivedDataCache.CacheStore",
 	EAutomationTestFlags::EditorContext | EAutomationTestFlags::EngineFilter);
 
-// TODO: remove this ifdef and enable the code it surrounds when we submit zenserver.exe
-// with the change to HandleRpcGetCachePayloads to support the ValueAPI
-#define ZENSERVER_SUPPORTS_VALUEAPI 0
-
 bool FCacheStoreTest::RunTest(const FString& Parameters)
 {
 	using namespace UE::DerivedData;
 	ICache& Cache = GetCache();
+	FStringView TestVersion = TEXTVIEW("F72150A02AE34B57A9EC91D36BA1CE08");
 
 	FRequestOwner Owner(EPriority::Blocking);
 	FCacheBucket DDCTestBucket(TEXT("DDCTest"));
+	static bool bExpectUpstream = FParse::Param(FCommandLine::Get(), TEXT("CacheStoreTestUpstream"));
+	static bool bExpectWarm = FParse::Param(FCommandLine::Get(), TEXT("CacheStoreTestWarm"));
 
-	// NumRequests/2 must be larger than the batch size in ZenDerivedData,
-	// and NumRequests >= 8 so that we get two results for each possible combination of flags
-	const int32 NumRequests = 24;
+	// NumKeys = (2 Value vs Record)*(2 SkipData vs Default)*(2 ForceMiss vs Not)*(2 use local)*(2 use remote)*(4 cases per type)
+	constexpr int32 NumKeys = 128;
+	constexpr int32 NumValues = 4;
 	TArray<FCachePutRequest> PutRequests;
 	TArray<FCachePutValueRequest> PutValueRequests;
 	TArray<FCacheGetRequest> GetRequests;
 	TArray<FCacheGetValueRequest> GetValueRequests;
 	TArray<FCacheGetChunkRequest> ChunkRequests;
-	struct FTestData
+	FValueId ValueIds[NumValues];
+	for (int32 ValueIndex = 0; ValueIndex < NumValues; ++ValueIndex)
+	{
+		ValueIds[ValueIndex] = FValueId::FromName(*WriteToString<16>(TEXT("ValueId_"), ValueIndex));
+	}
+
+	struct FKeyData;
+	struct FUserData
+	{
+		FUserData& Set(FKeyData* InKeyData, int32 InValueIndex)
+		{
+			KeyData = InKeyData;
+			ValueIndex = InValueIndex;
+			return *this;
+		}
+		FKeyData* KeyData = nullptr;
+		int32 ValueIndex = 0;
+	};
+	struct FKeyData
 	{
-		FValue Value;
-		uint8 ByteValue;
+		FValue BufferValues[NumValues];
+		uint64 IntValues[NumValues];
+		FUserData ValueUserData[NumValues];
+		bool ReceivedChunk[NumValues];
+		FCacheKey Key;
+		FUserData KeyUserData;
+		uint32 KeyIndex = 0;
 		bool bGetRequestsData = true;
-		bool bUseValueAPI = true;
+		bool bUseValueAPI = false;
+		bool bForceMiss = false;
+		bool bUseLocal = true;
+		bool bUseRemote = true;
+		bool bShouldBeHit = true;
 		bool bReceivedPut = false;
 		bool bReceivedGet = false;
 		bool bReceivedPutValue = false;
 		bool bReceivedGetValue = false;
-		bool bReceivedChunk = false;
 	};
 
-	TMap<uint64, FTestData> TestDatas;
-	FValueId ValueId = FValueId::FromName(TEXT("ValueName"));
-	for (int32 n = 0; n < NumRequests; ++n)
+	FKeyData KeyDatas[NumKeys];
+	for (uint32 KeyIndex = 0; KeyIndex < NumKeys; ++KeyIndex)
 	{
 		FBlake3 KeyWriter;
 		KeyWriter.Update(*TestName, TestName.Len() * sizeof(TestName[0]));
-		KeyWriter.Update(&n, sizeof(n));
+		KeyWriter.Update(TestVersion.GetData(), TestVersion.Len() * sizeof(TestVersion[0]));
+		KeyWriter.Update(&KeyIndex, sizeof(KeyIndex));
 		FIoHash KeyHash = KeyWriter.Finalize();
-		FCacheKey Key{ DDCTestBucket, KeyHash };
-		FSharedString Name(WriteToString<16>(TEXT("Request"), n));
-		uint64 UserData = (uint64)n;
-		FTestData& TestData = TestDatas.FindOrAdd(UserData);
-
-		ECachePolicy PutPolicy = ECachePolicy::Default;
-		ECachePolicy GetPolicy = ECachePolicy::Default;
-		TestData.bGetRequestsData = (n & 0x1) == 0;
-		TestData.bUseValueAPI = (n & 0x2) != 0;
-		if (!TestData.bGetRequestsData)
+		FSharedString Name(WriteToString<16>(TEXT("Key_"), KeyIndex));
+		FKeyData& KeyData = KeyDatas[KeyIndex];
+
+		KeyData.Key = FCacheKey{ DDCTestBucket, KeyHash };
+		KeyData.KeyIndex = KeyIndex;
+		KeyData.bGetRequestsData = (KeyIndex & (1 << 1)) == 0;
+		KeyData.bUseValueAPI = (KeyIndex & (1 << 2)) != 0;
+		KeyData.bForceMiss = (KeyIndex & (1 << 3)) == 0;
+		KeyData.bUseLocal = (KeyIndex & (1 << 4)) == 0;
+		KeyData.bUseRemote = (KeyIndex & (1 << 5)) == 0;
+		KeyData.bShouldBeHit = !KeyData.bForceMiss && (KeyData.bUseLocal || (KeyData.bUseRemote && bExpectUpstream));
+		ECachePolicy SharedPolicy = KeyData.bUseLocal ? ECachePolicy::Local : ECachePolicy::None;
+		SharedPolicy |= KeyData.bUseRemote ? ECachePolicy::Remote : ECachePolicy::None;
+		ECachePolicy PutPolicy = SharedPolicy;
+		ECachePolicy GetPolicy = SharedPolicy;
+		GetPolicy |= !KeyData.bGetRequestsData ? ECachePolicy::SkipData : ECachePolicy::None;
+		FCacheKey& Key = KeyData.Key;
+
+		for (int32 ValueIndex = 0; ValueIndex < NumValues; ++ValueIndex)
 		{
-			GetPolicy |= ECachePolicy::SkipData;
+			KeyData.IntValues[ValueIndex] = static_cast<uint64>(KeyIndex) | (static_cast<uint64>(ValueIndex) << 32);
+			KeyData.BufferValues[ValueIndex] = FValue::Compress(FSharedBuffer::MakeView(&KeyData.IntValues[ValueIndex], sizeof(KeyData.IntValues[ValueIndex])));
+			KeyData.ReceivedChunk[ValueIndex] = false;
 		}
-		TestData.ByteValue = (uint8)n;
-		TestData.Value = FValue::Compress(FSharedBuffer::MakeView(&TestData.ByteValue, 1));
 
-		if (!TestData.bUseValueAPI)
+		static_assert(sizeof(uint64) == sizeof(FUserData*), "Storing pointers in the UserData");
+		FUserData& KeyUserData = KeyData.KeyUserData.Set(&KeyData, -1);
+		for (uint32 ValueIndex = 0; ValueIndex < NumValues; ++ValueIndex)
+		{
+			KeyData.ValueUserData[ValueIndex].Set(&KeyData, ValueIndex);
+		}
+		if (!KeyData.bUseValueAPI)
 		{
 			FCacheRecordBuilder Builder(Key);
-			Builder.AddValue(ValueId, TestData.Value);
+			for (uint32 ValueIndex = 0; ValueIndex < NumValues; ++ValueIndex)
+			{
+				Builder.AddValue(ValueIds[ValueIndex], KeyData.BufferValues[ValueIndex]);
+			}
 
-			PutRequests.Add(FCachePutRequest{ Name, Builder.Build(), FCacheRecordPolicy(PutPolicy), UserData });
-			GetRequests.Add(FCacheGetRequest{ Name, Key, FCacheRecordPolicy(GetPolicy), UserData });
+			if (!KeyData.bForceMiss)
+			{
+				PutRequests.Add(FCachePutRequest{ Name, Builder.Build(), FCacheRecordPolicy(PutPolicy), reinterpret_cast<uint64>(&KeyUserData) });
+			}
+			GetRequests.Add(FCacheGetRequest{ Name, Key, FCacheRecordPolicy(GetPolicy), reinterpret_cast<uint64>(&KeyUserData) });
+			for (uint32 ValueIndex = 0; ValueIndex < NumValues; ++ValueIndex)
+			{
+				FUserData& ValueUserData = KeyData.ValueUserData[ValueIndex];
+				ChunkRequests.Add(FCacheGetChunkRequest{ Name, Key, ValueIds[ValueIndex],
+					0, MAX_uint64, FIoHash(), GetPolicy, reinterpret_cast<uint64>(&ValueUserData) });
+			}
 		}
 		else
 		{
-			PutValueRequests.Add(FCachePutValueRequest{ Name, Key, TestData.Value, PutPolicy, UserData });
-			GetValueRequests.Add(FCacheGetValueRequest{ Name, Key, GetPolicy, UserData });
+			if (!KeyData.bForceMiss)
+			{
+				PutValueRequests.Add(FCachePutValueRequest{ Name, Key, KeyData.BufferValues[0], PutPolicy, reinterpret_cast<uint64>(&KeyUserData) });
+			}
+			GetValueRequests.Add(FCacheGetValueRequest{ Name, Key, GetPolicy, reinterpret_cast<uint64>(&KeyUserData) });
+			ChunkRequests.Add(FCacheGetChunkRequest{ Name, Key, FValueId(),
+				0, MAX_uint64, FIoHash(), GetPolicy, reinterpret_cast<uint64>(&KeyUserData) });
 		}
-		ChunkRequests.Add(FCacheGetChunkRequest{ Name, Key, TestData.bUseValueAPI ? FValueId() : ValueId,
-			0, TestData.Value.GetRawSize(), FIoHash(), GetPolicy, UserData });
 	}
 
+	if (!bExpectWarm)
 	{
-		FRequestBarrier Barrier(Owner);
-		Cache.Put(PutRequests, Owner, [&TestDatas, this](FCachePutResponse&& Response)
+		{
+			FRequestBarrier Barrier(Owner);
+			Cache.Put(PutRequests, Owner, [this](FCachePutResponse&& Response)
+				{
+					FUserData* UserData = reinterpret_cast<FUserData*>(Response.UserData);
+					FKeyData* KeyData = UserData->KeyData;
+					TestTrue(TEXT("Valid UserData in Put Callback"), KeyData != nullptr);
+					if (KeyData)
+					{
+						KeyData->bReceivedPut = true;
+					}
+				});
+			Cache.PutValue(PutValueRequests, Owner, [this](FCachePutValueResponse&& Response)
+				{
+					FUserData* UserData = reinterpret_cast<FUserData*>(Response.UserData);
+					FKeyData* KeyData = UserData->KeyData;
+					TestTrue(TEXT("Valid UserData in PutValue Callback"), KeyData != nullptr);
+					if (KeyData)
+					{
+						KeyData->bReceivedPutValue = true;
+					}
+				});
+		}
+		Owner.Wait();
+		for (FKeyData& KeyData : KeyDatas)
+		{
+			int32 n = KeyData.KeyIndex;
+			if (!KeyData.bForceMiss)
 			{
-				FTestData* TestData = TestDatas.Find(Response.UserData);
-				TestTrue(TEXT("Valid UserData in Put Callback"), TestData != nullptr);
-				if (TestData)
+				if (!KeyData.bUseValueAPI)
 				{
-					TestData->bReceivedPut = true;
+					TestTrue(*WriteToString<16>(TEXT("Put "), n, TEXT(" received")), KeyData.bReceivedPut);
 				}
-			});
-		Cache.PutValue(PutValueRequests, Owner, [&TestDatas, this](FCachePutValueResponse&& Response)
-			{
-				FTestData* TestData = TestDatas.Find(Response.UserData);
-				TestTrue(TEXT("Valid UserData in PutValue Callback"), TestData != nullptr);
-				if (TestData)
+				else
 				{
-					TestData->bReceivedPutValue = true;
+					TestTrue(*WriteToString<16>(TEXT("PutValue "), n, TEXT(" received")), KeyData.bReceivedPutValue);
 				}
-			});
-	}
-	Owner.Wait();
-	for (TPair<uint64, FTestData>& Pair : TestDatas)
-	{
-		int32 n = (int32)Pair.Key;
-		if (!Pair.Value.bUseValueAPI)
-		{
-			TestTrue(*WriteToString<16>(TEXT("Put %d received"), n), Pair.Value.bReceivedPut);
+			}
 		}
-		else
-		{
-			TestTrue(*WriteToString<16>(TEXT("PutValue %d received"), n), Pair.Value.bReceivedPutValue);
-		}
-
 	}
 
 	{
 		FRequestBarrier Barrier(Owner);
-		Cache.Get(GetRequests, Owner, [&TestDatas, &ValueId, this](FCacheGetResponse&& Response)
+		Cache.Get(GetRequests, Owner, [&ValueIds, this, NumValues](FCacheGetResponse&& Response)
 			{
-				FTestData* TestData = TestDatas.Find(Response.UserData);
-				int32 n = (int32)Response.UserData;
-				TestTrue(TEXT("Valid UserData in Get Callback"), TestData != nullptr);
-				if (TestData)
+				FUserData* UserData = reinterpret_cast<FUserData*>(Response.UserData);
+				FKeyData* KeyData = UserData->KeyData;
+				TestTrue(TEXT("Valid UserData in Get Callback"), KeyData != nullptr);
+				if (KeyData)
 				{
-					TestData->bReceivedGet = true;
-					if (TestTrue(*WriteToString<32>(TEXT("Get %d succeeded"), n), Response.Status == EStatus::Ok))
+					int32 n = (int32)KeyData->KeyIndex;
+					KeyData->bReceivedGet = true;
+					if (KeyData->bShouldBeHit)
+					{
+						TestEqual(*WriteToString<32>(TEXT("Get "), n, TEXT(" succeeded")), Response.Status, EStatus::Ok);
+					}
+					else if (KeyData->bForceMiss)
+					{
+						TestEqual(*WriteToString<32>(TEXT("Get "), n, TEXT(" failed as expected")), Response.Status, EStatus::Error);
+					}
+					if (!KeyData->bForceMiss && Response.Status == EStatus::Ok)
 					{
 						FCacheRecord& Record = Response.Record;
 						TConstArrayView<FValueWithId> Values = Record.GetValues();
-						if (TestEqual(*WriteToString<32>(TEXT("Get %d ValuesLen"), n), Values.Num(), 1))
+						if (TestEqual(*WriteToString<32>(TEXT("Get "), n, TEXT(" ValuesLen")), Values.Num(), NumValues))
 						{
-							const FValueWithId& ActualValue = Values[0];
-							TestEqual(*WriteToString<32>(TEXT("Get %d ValueId"), n), ActualValue.GetId(), ValueId);
-							FValue& ExpectedValue = TestData->Value;
-							TestEqual(*WriteToString<32>(TEXT("Get %d Hash"), n),
-								ActualValue.GetRawHash(), ExpectedValue.GetRawHash());
-							TestEqual(*WriteToString<32>(TEXT("Get %d Size"), n),
-								ActualValue.GetRawSize(), ExpectedValue.GetRawSize());
-
-							if (TestData->bGetRequestsData)
+							for (int32 ActualValueIndex = 0; ActualValueIndex < NumValues; ++ActualValueIndex)
 							{
-								const FCompressedBuffer& Compressed = ActualValue.GetData();
-								FSharedBuffer Buffer = Compressed.Decompress();
-								TestEqual(*WriteToString<32>(TEXT("Get %d Data Size"), n),
-									Buffer.GetSize(), ActualValue.GetRawSize());
-								if (Buffer.GetSize())
+								const FValueWithId& ActualValue = Values[ActualValueIndex];
+								int32 ExpectedValueIndex = TArrayView<const FValueId>(ValueIds).Find(ActualValue.GetId());
+								const FValueId& ExpectedValueId = ExpectedValueIndex >= 0 ? ValueIds[ExpectedValueIndex] : FValueId();
+
+								TestTrue(*WriteToString<32>(TEXT("Get "), n, TEXT(" ValueId")), ExpectedValueIndex >= 0);
+								if (ExpectedValueIndex >= 0)
 								{
-									TestEqual(*WriteToString<32>(TEXT("Get %d Data Equals"), n),
-										((const uint8*)Buffer.GetData())[0], TestData->ByteValue);
+									FValue& ExpectedValue = KeyData->BufferValues[ExpectedValueIndex];
+									TestEqual(*WriteToString<32>(TEXT("Get "), n, TEXT(" Hash")),
+										ActualValue.GetRawHash(), ExpectedValue.GetRawHash());
+									TestEqual(*WriteToString<32>(TEXT("Get "), n, TEXT(" Size")),
+										ActualValue.GetRawSize(), ExpectedValue.GetRawSize());
+
+									if (KeyData->bGetRequestsData)
+									{
+										const FCompressedBuffer& Compressed = ActualValue.GetData();
+										FSharedBuffer Buffer = Compressed.Decompress();
+										TestEqual(*WriteToString<32>(TEXT("Get "), n, TEXT(" Data Size")),
+											Buffer.GetSize(), ActualValue.GetRawSize());
+										if (Buffer.GetSize())
+										{
+											uint64 ActualIntValue = ((const uint64*)Buffer.GetData())[0];
+											uint64 ExpectedIntValue = KeyData->IntValues[ExpectedValueIndex];
+											TestEqual(*WriteToString<32>(TEXT("Get "), n, TEXT(" Data Equals"), n),
+												ActualIntValue, ExpectedIntValue);
+										}
+									}
 								}
 							}
 						}
@@ -165,85 +250,109 @@ bool FCacheStoreTest::RunTest(const FString& Parameters)
 				}
 			});
 
-		Cache.GetValue(GetValueRequests, Owner, [&TestDatas, this](FCacheGetValueResponse&& Response)
+		Cache.GetValue(GetValueRequests, Owner, [this](FCacheGetValueResponse&& Response)
 			{
-				FTestData* TestData = TestDatas.Find(Response.UserData);
-				int32 n = (int32)Response.UserData;
-				TestTrue(TEXT("Valid UserData in GetValue Callback"), TestData != nullptr);
-				if (TestData)
+				FUserData* UserData = reinterpret_cast<FUserData*>(Response.UserData);
+				FKeyData* KeyData = UserData->KeyData;
+				TestTrue(TEXT("Valid UserData in GetValue Callback"), KeyData != nullptr);
+				if (KeyData)
 				{
-					TestData->bReceivedGetValue = true;
-					if (TestTrue(*WriteToString<32>(TEXT("GetValue %d succeeded"), n), Response.Status == EStatus::Ok))
+					int32 n = KeyData->KeyIndex;
+					KeyData->bReceivedGetValue = true;
+					if (KeyData->bShouldBeHit)
+					{
+						TestEqual(*WriteToString<32>(TEXT("GetValue "), n, TEXT(" succeeded")), Response.Status, EStatus::Ok);
+					}
+					else if (KeyData->bForceMiss)
+					{
+						TestEqual(*WriteToString<32>(TEXT("GetValue "), n, TEXT(" failed as expected")), Response.Status, EStatus::Error);
+					}
+					if (!KeyData->bForceMiss && Response.Status == EStatus::Ok)
 					{
 						FValue& ActualValue = Response.Value;
-						FValue& ExpectedValue = TestData->Value;
-						TestEqual(*WriteToString<32>(TEXT("GetValue %d Hash"), n),
+						FValue& ExpectedValue = KeyData->BufferValues[0];
+						TestEqual(*WriteToString<32>(TEXT("GetValue "), n, TEXT(" Hash")),
 							ActualValue.GetRawHash(), ExpectedValue.GetRawHash());
-						TestEqual(*WriteToString<32>(TEXT("GetValue %d Size"), n),
+						TestEqual(*WriteToString<32>(TEXT("GetValue "), n, TEXT(" Size")),
 							ActualValue.GetRawSize(), ExpectedValue.GetRawSize());
 
-						if (TestData->bGetRequestsData)
+						if (KeyData->bGetRequestsData)
 						{
 							const FCompressedBuffer& Compressed = ActualValue.GetData();
 							FSharedBuffer Buffer = Compressed.Decompress();
-							TestEqual(*WriteToString<32>(TEXT("GetValue %d Data Size"), n),
+							TestEqual(*WriteToString<32>(TEXT("GetValue "), n, TEXT(" Data Size")),
 								Buffer.GetSize(), ActualValue.GetRawSize());
 							if (Buffer.GetSize())
 							{
-								TestEqual(*WriteToString<32>(TEXT("GetValue %d Data Equals"), n),
-									((const uint8*)Buffer.GetData())[0], TestData->ByteValue);
+								uint64 Value = ((const uint64*)Buffer.GetData())[0];
+								TestEqual(*WriteToString<32>(TEXT("GetValue "), n, TEXT(" Data Equals")),
+									Value, KeyData->IntValues[0]);
 							}
 						}
 					}
 				}
 			});
 
-#if ZENSERVER_SUPPORTS_VALUEAPI
-		Cache.GetChunks(ChunkRequests, Owner, [&TestDatas, this](FCacheGetChunkResponse&& Response)
+		Cache.GetChunks(ChunkRequests, Owner, [this](FCacheGetChunkResponse&& Response)
 			{
-				FTestData* TestData = TestDatas.Find(Response.UserData);
-				int32 n = (int32)Response.UserData;
-				if (TestTrue(TEXT("Valid UserData in GetChunks Callback"), TestData != nullptr))
+				FUserData* UserData = reinterpret_cast<FUserData*>(Response.UserData);
+				FKeyData* KeyData = UserData->KeyData;
+				if (TestTrue(TEXT("Valid UserData in GetChunks Callback"), KeyData != nullptr))
 				{
-					TestData->bReceivedChunk = true;
-					if (TestTrue(*WriteToString<32>(TEXT("GetChunks %d succeeded"), n), Response.Status == EStatus::Ok))
+					int32 n = KeyData->KeyIndex;
+					int32 ValueIndex = UserData->ValueIndex >= 0 ? UserData->ValueIndex : 0;
+					TStringBuilder<32> Name;
+					Name << TEXT("GetChunks (") << n << TEXT(",") << ValueIndex << TEXT(")");
+
+					KeyData->ReceivedChunk[ValueIndex] = true;
+					if (KeyData->bShouldBeHit)
+					{
+						TestEqual(*WriteToString<32>(*Name, TEXT(" succeeded")), Response.Status, EStatus::Ok);
+					}
+					else if (KeyData->bForceMiss)
 					{
-						FValue& ExpectedValue = TestData->Value;
-						TestEqual(*WriteToString<32>(TEXT("GetChunks %d Hash"), n),
+						TestEqual(*WriteToString<32>(*Name, TEXT(" failed as expected")), Response.Status, EStatus::Error);
+					}
+					if (KeyData->bShouldBeHit && Response.Status == EStatus::Ok)
+					{
+						FValue& ExpectedValue = KeyData->BufferValues[ValueIndex];
+						TestEqual(*WriteToString<32>(*Name, TEXT(" Hash")),
 							Response.RawHash, ExpectedValue.GetRawHash());
-						TestEqual(*WriteToString<32>(TEXT("GetChunks %d Size"), n),
+						TestEqual(*WriteToString<32>(*Name, TEXT(" Size")),
 							Response.RawSize, ExpectedValue.GetRawSize());
-						if (TestData->bGetRequestsData)
+						if (KeyData->bGetRequestsData)
 						{
 							FSharedBuffer Buffer = Response.RawData;
-							TestEqual(*WriteToString<32>(TEXT("GetChunks %d Data Size"), n),
+							TestEqual(*WriteToString<32>(*Name, TEXT(" Data Size")),
 								Buffer.GetSize(), Response.RawSize);
 							if (Buffer.GetSize())
 							{
-								TestEqual(*WriteToString<32>(TEXT("GetChunks %d Data Equals"), n),
-									((const uint8*)Buffer.GetData())[0], TestData->ByteValue);
+								uint64 Value = ((const uint64*)Buffer.GetData())[0];
+								TestEqual(*WriteToString<32>(*Name, TEXT(" Data Equals")),
+									Value, KeyData->IntValues[ValueIndex]);
 							}
 						}
 					}
 				}
 			});
-#endif
 	}
 	Owner.Wait();
-	for (TPair<uint64, FTestData>& Pair : TestDatas)
+	for (FKeyData& KeyData : KeyDatas)
 	{
-		int32 n = (int32)Pair.Key;
-		if (!Pair.Value.bUseValueAPI)
+		int32 n = KeyData.KeyIndex;
+		if (!KeyData.bUseValueAPI)
 		{
-			TestTrue(*WriteToString<16>(TEXT("Get %d received"), n), Pair.Value.bReceivedGet);
+			TestTrue(*WriteToString<32>(TEXT("Get "), n, TEXT(" received")), KeyData.bReceivedGet);
+			for (int32 ValueIndex = 0; ValueIndex < NumValues; ++ValueIndex)
+			{
+				TestTrue(*WriteToString<32>(TEXT("GetChunk ("), n, TEXT(","), ValueIndex, TEXT(") received")), KeyData.ReceivedChunk[ValueIndex]);
+			}
 		}
 		else
 		{
-			TestTrue(*WriteToString<16>(TEXT("GetValue %d received"), n), TestDatas[n].bReceivedGetValue);
+			TestTrue(*WriteToString<32>(TEXT("GetValue "), n, TEXT(" received")), KeyData.bReceivedGetValue);
+			TestTrue(*WriteToString<32>(TEXT("GetChunk ("), n, TEXT(",0) received")), KeyData.ReceivedChunk[0]);
 		}
-#if ZENSERVER_SUPPORTS_VALUEAPI
-		TestTrue(*WriteToString<16>(TEXT("GetChunk %d received"), n), TestDatas[n].bReceivedChunk);
-#endif
 	}
 
 	return true;
diff --git a/Engine/Source/Developer/DerivedDataCache/Private/ZenDerivedDataBackend.cpp b/Engine/Source/Developer/DerivedDataCache/Private/ZenDerivedDataBackend.cpp
--- a/Engine/Source/Developer/DerivedDataCache/Private/ZenDerivedDataBackend.cpp
+++ b/Engine/Source/Developer/DerivedDataCache/Private/ZenDerivedDataBackend.cpp
@@ -5,6 +5,7 @@
 
 #if UE_WITH_ZEN
 
+#include "BatchView.h"
 #include "Containers/Set.h"
 #include "Containers/StringFwd.h"
 #include "Containers/StringFwd.h"
@@ -14,6 +15,7 @@
 #include "DerivedDataChunk.h"
 #include "HAL/CriticalSection.h"
 #include "Misc/ConfigCacheIni.h"
+#include "Misc/Optional.h"
 #include "Misc/ScopeLock.h"
 #include "ProfilingDebugging/CountersTrace.h"
 #include "ProfilingDebugging/CpuProfilerTrace.h"
@@ -151,9 +153,6 @@ private:
 	FDerivedDataBackendInterface::EPutStatus PutZenData(const TCHAR* Uri, const FCompositeBuffer& InData, Zen::EContentType ContentType);
 	EGetResult GetZenData(const FCacheKey& Key, ECachePolicy CachePolicy, FCbPackage& OutPackage) const;
 
-	bool PutCacheRecord(const FCacheRecord& Record, const FCacheRecordPolicy& Policy);
-	bool PutCacheValue(const FCacheKey& Key, const FValue& Value, const ECachePolicy& ValuePolicy);
-
 	bool IsServiceReady();
 	static FString MakeLegacyZenKey(const TCHAR* CacheKey);
 	static void AppendZenUri(const FCacheKey& CacheKey, FStringBuilderBase& Out);
@@ -175,6 +174,7 @@ private:
 	bool bIsRemote = false;
 	uint32 FailedLoginAttempts = 0;
 	uint32 MaxAttempts = 4;
+	int32 BatchPutMaxBytes = 1024*1024;
 	int32 CacheRecordBatchSize = 8;
 	int32 CacheChunksBatchSize = 8;
 
@@ -207,6 +207,7 @@ FZenDerivedDataBackend::FZenDerivedDataBackend(
 		}
 	}
 
+	GConfig->GetInt(TEXT("Zen"), TEXT("BatchPutMaxBytes"), BatchPutMaxBytes, GEngineIni);
 	GConfig->GetInt(TEXT("Zen"), TEXT("CacheRecordBatchSize"), CacheRecordBatchSize, GEngineIni);
 	GConfig->GetInt(TEXT("Zen"), TEXT("CacheChunksBatchSize"), CacheChunksBatchSize, GEngineIni);
 }
@@ -659,41 +660,127 @@ void FZenDerivedDataBackend::Put(
 	FOnCachePutComplete&& OnComplete)
 {
 	TRACE_CPUPROFILER_EVENT_SCOPE(ZenDDC::PutCachedRecord);
+	COOK_STAT(auto Timer = UsageStats.TimePut());
+
+	uint64 BatchSize = 0;
+	TBatchView<const FCachePutRequest> Batches(Requests,
+		[&BatchSize, this](const FCachePutRequest& NextRequest) -> EBatchView
+		{
+			const FCacheRecord& Record = NextRequest.Record;
+			uint64 RecordSize = sizeof(FCacheKey) + Record.GetMeta().GetSize();
+			for (const FValueWithId& Value : Record.GetValues())
+			{
+				RecordSize += Value.GetData().GetCompressedSize();
+			}
+			BatchSize += RecordSize;
+			if (BatchSize > BatchPutMaxBytes)
+			{
+				BatchSize = RecordSize;
+				return EBatchView::NewBatch;
+			}
+			return EBatchView::Continue;
+		});
 
-	for (const FCachePutRequest& Request : Requests)
+	auto OnHit = [this, &OnComplete COOK_STAT(, &Timer)](const FCachePutRequest& Request)
 	{
-		COOK_STAT(auto Timer = UsageStats.TimePut());
-		const FCacheRecord& Record = Request.Record;
-		bool bResult;
-		if (ShouldSimulateMiss(Record.GetKey()))
+		const FCacheKey& Key = Request.Record.GetKey();
+		UE_LOG(LogDerivedDataCache, Verbose, TEXT("%s: Cache Put complete for %s from '%s'"),
+			*GetName(), *WriteToString<96>(Key), *Request.Name);
+		COOK_STAT(Timer.AddHit(Private::GetCacheRecordCompressedSize(Request.Record)));
+		if (OnComplete)
 		{
-			UE_LOG(LogDerivedDataCache, Verbose, TEXT("%s: Simulated miss for put of %s from '%s'"),
-				*GetName(), *WriteToString<96>(Record.GetKey()), *Request.Name);
-			bResult = false;
+			OnComplete({ Request.Name, Key, Request.UserData, EStatus::Ok });
 		}
-		else
+	};
+	auto OnMiss = [this, &OnComplete COOK_STAT(, &Timer)](const FCachePutRequest& Request)
+	{
+		const FCacheKey& Key = Request.Record.GetKey();
+		UE_LOG(LogDerivedDataCache, Verbose, TEXT("%s: Cache Put miss for '%s' from '%s'"),
+			*GetName(), *WriteToString<96>(Key), *Request.Name);
+		COOK_STAT(Timer.AddMiss());
+		if (OnComplete)
 		{
-			bResult = PutCacheRecord(Record, Request.Policy);
+			OnComplete({ Request.Name, Key, Request.UserData, EStatus::Error });
 		}
+	};
 
-		if (bResult)
+	for (TArrayView<const FCachePutRequest> Batch : Batches)
+	{
+		FCbPackage BatchPackage;
+		FCbWriter BatchWriter;
+		BatchWriter.BeginObject();
 		{
-			UE_LOG(LogDerivedDataCache, Verbose, TEXT("%s: Cache put complete for %s from '%s'"),
-				*GetName(), *WriteToString<96>(Record.GetKey()), *Request.Name);
-			COOK_STAT(Timer.AddHit(Private::GetCacheRecordCompressedSize(Record)));
-			if (OnComplete)
+			BatchWriter << "Method"_ASV << "PutCacheRecords";
+			BatchWriter.BeginObject("Params"_ASV);
 			{
-				OnComplete({ Request.Name, Record.GetKey(), Request.UserData, EStatus::Ok });
+				ECachePolicy BatchDefaultPolicy = Batch[0].Policy.GetRecordPolicy();
+				BatchWriter << "DefaultPolicy"_ASV << *WriteToString<128>(BatchDefaultPolicy);
+
+				BatchWriter.BeginArray("Requests"_ASV);
+				for (const FCachePutRequest& Request : Batch)
+				{
+					const FCacheRecord& Record = Request.Record;
+
+					BatchWriter.BeginObject();
+					{
+						BatchWriter.SetName("Record"_ASV);
+						Record.Save(BatchPackage, BatchWriter);
+						if (!Request.Policy.IsUniform() || Request.Policy.GetRecordPolicy() != BatchDefaultPolicy)
+						{
+							BatchWriter.SetName("Policy"_ASV);
+							Request.Policy.Save(BatchWriter);
+						}
+					}
+					BatchWriter.EndObject();
+				}
+				BatchWriter.EndArray();
 			}
+			BatchWriter.EndObject();
+		}
+		BatchWriter.EndObject();
+		BatchPackage.SetObject(BatchWriter.Save().AsObject());
+
+		FCbPackage BatchResponse;
+		Zen::FZenHttpRequest::Result HttpResult = Zen::FZenHttpRequest::Result::Failed;
+
+		{
+			Zen::FZenScopedRequestPtr Request(RequestPool.Get());
+			HttpResult = Request->PerformRpc(TEXT("/z$/$rpc"_SV), BatchPackage, BatchResponse);
 		}
-		else
+
+		int32 RequestIndex = 0;
+		if (HttpResult == Zen::FZenHttpRequest::Result::Success)
 		{
-			COOK_STAT(Timer.AddMiss());
-			if (OnComplete)
+			const FCbObject& ResponseObj = BatchResponse.GetObject();
+			for (FCbField ResponseField : ResponseObj["Result"_ASV])
 			{
-				OnComplete({ Request.Name, Record.GetKey(), Request.UserData, EStatus::Error });
+				if (RequestIndex >= Batch.Num())
+				{
+					++RequestIndex;
+					continue;
+				}
+				const FCachePutRequest& Request = Batch[RequestIndex++];
+
+				const FCacheKey& Key = Request.Record.GetKey();
+				bool bPutSucceeded = ResponseField.AsBool();
+				if (ShouldSimulateMiss(Key))
+				{
+					UE_LOG(LogDerivedDataCache, Verbose, TEXT("%s: Simulated miss for put of %s from '%s'"),
+						*GetName(), *WriteToString<96>(Key), *Request.Name);
+					bPutSucceeded = false;
+				}
+				bPutSucceeded ? OnHit(Request) : OnMiss(Request);
+			}
+			if (RequestIndex != Batch.Num())
+			{
+				UE_LOG(LogDerivedDataCache, Warning,
+					TEXT("Invalid response received from PutCacheRecords rpc: %d results expected, received %d."), Batch.Num(), RequestIndex);
 			}
 		}
+		for (const FCachePutRequest& Request : Batch.RightChop(RequestIndex))
+		{
+			OnMiss(Request);
+		}
 	}
 }
 
@@ -704,36 +791,64 @@ void FZenDerivedDataBackend::Get(
 {
 	TRACE_CPUPROFILER_EVENT_SCOPE(ZenDDC::GetCacheRecord);
 	TRACE_COUNTER_ADD(ZenDDC_CacheRecordRequestCount, int64(Requests.Num()));
+	TRACE_COUNTER_ADD(ZenDDC_Get, int64(Requests.Num()));
+	COOK_STAT(auto Timer = UsageStats.TimeGet());
+
+	auto OnHit = [this, &OnComplete COOK_STAT(, &Timer)](const FCacheGetRequest& Request, FCacheRecord&& Record)
+	{
+		UE_LOG(LogDerivedDataCache, Verbose, TEXT("%s: Cache hit for '%s' from '%s'"),
+			*GetName(), *WriteToString<96>(Request.Key), *Request.Name);
 
-	int32 TotalCompleted = 0;
+		TRACE_COUNTER_ADD(ZenDDC_GetHit, int64(1));
+		int64 ReceivedSize = Private::GetCacheRecordCompressedSize(Record);
+		TRACE_COUNTER_ADD(ZenDDC_BytesReceived, ReceivedSize);
+		COOK_STAT(Timer.AddHit(ReceivedSize));
+
+		if (OnComplete)
+		{
+			OnComplete({ Request.Name, MoveTemp(Record), Request.UserData, EStatus::Ok });
+		}
+	};
+	auto OnMiss = [this, &OnComplete](const FCacheGetRequest& Request)
+	{
+		UE_LOG(LogDerivedDataCache, Verbose, TEXT("%s: Cache miss for '%s' from '%s'"),
+			*GetName(), *WriteToString<96>(Request.Key), *Request.Name);
+
+		if (OnComplete)
+		{
+			OnComplete({ Request.Name, FCacheRecordBuilder(Request.Key).Build(), Request.UserData, EStatus::Error });
+		}
+	};
 
 	ForEachBatch(CacheRecordBatchSize, Requests.Num(),
-		[this, &Requests, &Owner, &OnComplete, &TotalCompleted](int32 BatchFirst, int32 BatchLast)
+		[this, &Requests, &Owner, &OnComplete, &OnHit, &OnMiss](int32 BatchFirst, int32 BatchLast)
 	{
-		TRACE_COUNTER_ADD(ZenDDC_Get, int64(BatchLast) - int64(BatchFirst) + 1);
-		COOK_STAT(auto Timer = UsageStats.TimeGet());
-		
+		TArrayView<const FCacheGetRequest> Batch(Requests.GetData() + BatchFirst, BatchLast - BatchFirst + 1);
+
 		FCbWriter BatchRequest;
 		BatchRequest.BeginObject();
 		{
-			BatchRequest << "Method"_ASV << "GetCacheRecords";
+			BatchRequest << "Method"_ASV << "GetCacheRecords"_ASV;
 			BatchRequest.BeginObject("Params"_ASV);
 			{
-				BatchRequest.BeginArray("CacheKeys"_ASV);
-				for (int32 KeyIndex = BatchFirst; KeyIndex <= BatchLast; KeyIndex++)
-				{
-					const FCacheKey& Key = Requests[KeyIndex].Key;
+				ECachePolicy BatchDefaultPolicy = Batch[0].Policy.GetRecordPolicy();
+				BatchRequest << "DefaultPolicy"_ASV << *WriteToString<128>(BatchDefaultPolicy);
 
+				BatchRequest.BeginArray("Requests"_ASV);
+				for (const FCacheGetRequest& Request : Batch)
+				{
 					BatchRequest.BeginObject();
-					BatchRequest << "Bucket"_ASV << Key.Bucket.ToString();
-					BatchRequest << "Hash"_ASV << Key.Hash;
+					{
+						BatchRequest << "Key"_ASV << Request.Key;
+						if (!Request.Policy.IsUniform() || Request.Policy.GetRecordPolicy() != BatchDefaultPolicy)
+						{
+							BatchRequest.SetName("Policy"_ASV);
+							Request.Policy.Save(BatchRequest);
+						}
+					}
 					BatchRequest.EndObject();
 				}
 				BatchRequest.EndArray();
-
-				// TODO: The policy needs to be sent with each key.
-				BatchRequest.SetName("Policy"_ASV);
-				Requests[BatchFirst].Policy.Save(BatchRequest);
 			}
 			BatchRequest.EndObject();
 		}
@@ -747,81 +862,46 @@ void FZenDerivedDataBackend::Get(
 			HttpResult = Request->PerformRpc(TEXT("/z$/$rpc"_SV), BatchRequest.Save().AsObject(), BatchResponse);
 		}
 
+		int32 RequestIndex = 0;
 		if (HttpResult == Zen::FZenHttpRequest::Result::Success)
 		{
 			const FCbObject& ResponseObj = BatchResponse.GetObject();
 			
-			int32 KeyIndex = BatchFirst;
 			for (FCbField RecordField : ResponseObj["Result"_ASV])
 			{
-				const FCacheGetRequest& Request = Requests[KeyIndex++];
-				const FCacheKey& Key = Request.Key;
-				
-				FOptionalCacheRecord Record;
-
-				if (!RecordField.IsNull())
+				if (RequestIndex >= Batch.Num())
 				{
-					if (ShouldSimulateMiss(Key))
-					{
-						UE_LOG(LogDerivedDataCache, Verbose, TEXT("%s: Simulated miss for put of '%s' from '%s'"),
-							*GetName(), *WriteToString<96>(Key), *Request.Name);
-					}
-					else
-					{
-						Record = FCacheRecord::Load(BatchResponse, RecordField.AsObject());
-					}
+					++RequestIndex;
+					continue;
 				}
+				const FCacheGetRequest& Request = Batch[RequestIndex++];
+
+				const FCacheKey& Key = Request.Key;
+				FOptionalCacheRecord Record;
 
-				if (Record)
+				if (ShouldSimulateMiss(Key))
 				{
-					UE_LOG(LogDerivedDataCache, Verbose, TEXT("%s: Cache hit for '%s' from '%s'"),
+					UE_LOG(LogDerivedDataCache, Verbose, TEXT("%s: Simulated miss for put of '%s' from '%s'"),
 						*GetName(), *WriteToString<96>(Key), *Request.Name);
-					
-					TRACE_COUNTER_ADD(ZenDDC_GetHit, int64(1));
-					int64 ReceivedSize = Private::GetCacheRecordCompressedSize(Record.Get());
-					TRACE_COUNTER_ADD(ZenDDC_BytesReceived, ReceivedSize);
-					COOK_STAT(Timer.AddHit(ReceivedSize));
-
-					if (OnComplete)
-					{
-						OnComplete({ Request.Name, MoveTemp(Record).Get(), Request.UserData, EStatus::Ok });
-					}
 				}
-				else
+				else if (!RecordField.IsNull())
 				{
-					UE_LOG(LogDerivedDataCache, Verbose, TEXT("%s: Cache miss for '%s' from '%s'"),
-						*GetName(), *WriteToString<96>(Key), *Request.Name);
-					
-					if (OnComplete)
-					{
-						OnComplete({ Request.Name, FCacheRecordBuilder(Key).Build(), Request.UserData, EStatus::Error });
-					}
+					Record = FCacheRecord::Load(BatchResponse, RecordField.AsObject());
 				}
-				
-				TotalCompleted++;
+				Record ? OnHit(Request, MoveTemp(Record).Get()) : OnMiss(Request);
 			}
-		}
-		else
-		{
-			for (int32 KeyIndex = BatchFirst; KeyIndex <= BatchLast; KeyIndex++)
+			if (RequestIndex != Batch.Num())
 			{
-				const FCacheGetRequest& Request = Requests[KeyIndex];
-				const FCacheKey& Key = Request.Key;
-
-				UE_LOG(LogDerivedDataCache, Verbose, TEXT("%s: Cache miss for '%s' from '%s'"),
-					*GetName(), *WriteToString<96>(Key), *Request.Name);
-					
-				if (OnComplete)
-				{
-					OnComplete({ Request.Name, FCacheRecordBuilder(Key).Build(), Request.UserData, EStatus::Error });
-				}
-				
-				TotalCompleted++;
+				UE_LOG(LogDerivedDataCache, Warning,
+					TEXT("Invalid response received from GetCacheRecords rpc: %d results expected, received %d."), Batch.Num(), RequestIndex);
 			}
 		}
+		for (const FCacheGetRequest& Request : Batch.RightChop(RequestIndex))
+		{
+			OnMiss(Request);
+		}
 	});
 	
-	UE_CLOG(TotalCompleted != Requests.Num(), LogDerivedDataCache, Warning, TEXT("Only '%d/%d' cache record request(s) completed"), TotalCompleted, Requests.Num());
 	TRACE_COUNTER_SUBTRACT(ZenDDC_CacheRecordRequestCount, int64(Requests.Num()));
 }
 
@@ -831,39 +911,119 @@ void FZenDerivedDataBackend::PutValue(
 	FOnCachePutValueComplete&& OnComplete)
 {
 	TRACE_CPUPROFILER_EVENT_SCOPE(ZenDDC::PutValue);
+	COOK_STAT(auto Timer = UsageStats.TimePut());
+
+	uint64 BatchSize = 0;
+	TBatchView<const FCachePutValueRequest> Batches(Requests,
+		[&BatchSize, this](const FCachePutValueRequest& NextValue) -> EBatchView
+		{
+			uint64 ValueSize = sizeof(FCacheKey) + NextValue.Value.GetData().GetCompressedSize();
+			BatchSize += ValueSize;
+			if (BatchSize > BatchPutMaxBytes)
+			{
+				BatchSize = ValueSize;
+				return EBatchView::NewBatch;
+			}
+			return EBatchView::Continue;
+		});
 
-	for (const FCachePutValueRequest& Request : Requests)
+	auto OnHit = [this, &OnComplete COOK_STAT(, &Timer)](const FCachePutValueRequest& Request)
 	{
-		COOK_STAT(auto Timer = UsageStats.TimePut());
-		bool bResult;
-		if (ShouldSimulateMiss(Request.Key))
+		UE_LOG(LogDerivedDataCache, Verbose, TEXT("%s: Cache PutValue complete for %s from '%s'"),
+			*GetName(), *WriteToString<96>(Request.Key), *Request.Name);
+		COOK_STAT(Timer.AddHit(Request.Value.GetData().GetCompressedSize()));
+		if (OnComplete)
 		{
-			UE_LOG(LogDerivedDataCache, Verbose, TEXT("%s: Simulated miss for PutValue of %s from '%s'"),
-				*GetName(), *WriteToString<96>(Request.Key), *Request.Name);
-			bResult = false;
+			OnComplete({ Request.Name, Request.Key, Request.UserData, EStatus::Ok });
 		}
-		else
+	};
+	auto OnMiss = [this, &OnComplete COOK_STAT(, &Timer)](const FCachePutValueRequest& Request)
+	{
+		UE_LOG(LogDerivedDataCache, Verbose, TEXT("%s: Cache PutValue miss for '%s' from '%s'"),
+			*GetName(), *WriteToString<96>(Request.Key), *Request.Name);
+		COOK_STAT(Timer.AddMiss());
+		if (OnComplete)
 		{
-			bResult = PutCacheValue(Request.Key, Request.Value, Request.Policy);
+			OnComplete({ Request.Name, Request.Key, Request.UserData, EStatus::Error });
 		}
+	};
 
-		if (bResult)
+	for (TArrayView<const FCachePutValueRequest> Batch : Batches)
+	{
+		FCbPackage BatchPackage;
+		FCbWriter BatchWriter;
+		BatchWriter.BeginObject();
 		{
-			UE_LOG(LogDerivedDataCache, Verbose, TEXT("%s: Cache PutValue complete for %s from '%s'"),
-				*GetName(), *WriteToString<96>(Request.Key), *Request.Name);
-			COOK_STAT(Timer.AddHit(Request.Value.GetData().GetCompressedSize()));
-			if (OnComplete)
+			BatchWriter << "Method"_ASV << "PutCacheValues"_ASV;
+			BatchWriter.BeginObject("Params"_ASV);
 			{
-				OnComplete({ Request.Name, Request.Key, Request.UserData, EStatus::Ok });
+				ECachePolicy BatchDefaultPolicy = Batch[0].Policy;
+				BatchWriter << "DefaultPolicy"_ASV << *WriteToString<128>(BatchDefaultPolicy);
+
+				BatchWriter.BeginArray("Requests");
+				for (const FCachePutValueRequest& Request : Batch)
+				{
+					BatchWriter.BeginObject();
+					{
+						BatchWriter << "Key"_ASV << Request.Key;
+						const FValue& Value = Request.Value;
+						BatchWriter.AddBinaryAttachment("RawHash", Value.GetRawHash());
+						if (Value.HasData())
+						{
+							BatchPackage.AddAttachment(FCbAttachment(Value.GetData()));
+						}
+						if (Request.Policy != BatchDefaultPolicy)
+						{
+							BatchWriter << "Policy"_ASV << WriteToString<128>(Request.Policy);
+						}
+					}
+					BatchWriter.EndObject();
+				}
+				BatchWriter.EndArray();
 			}
+			BatchWriter.EndObject();
 		}
-		else
+		BatchWriter.EndObject();
+		BatchPackage.SetObject(BatchWriter.Save().AsObject());
+
+		FCbPackage BatchResponse;
+		Zen::FZenHttpRequest::Result HttpResult = Zen::FZenHttpRequest::Result::Failed;
+		{
+			Zen::FZenScopedRequestPtr Request(RequestPool.Get());
+			HttpResult = Request->PerformRpc(TEXT("/z$/$rpc"_SV), BatchPackage, BatchResponse);
+		}
+
+		int32 RequestIndex = 0;
+		if (HttpResult == Zen::FZenHttpRequest::Result::Success)
 		{
-			COOK_STAT(Timer.AddMiss());
-			if (OnComplete)
+			const FCbObject& ResponseObj = BatchResponse.GetObject();
+			for (FCbField ResponseField : ResponseObj["Result"_ASV])
 			{
-				OnComplete({ Request.Name, Request.Key, Request.UserData, EStatus::Error });
+				if (RequestIndex >= Batch.Num())
+				{
+					++RequestIndex;
+					continue;
+				}
+				const FCachePutValueRequest& Request = Batch[RequestIndex++];
+
+				bool bPutSucceeded = ResponseField.AsBool();
+				if (ShouldSimulateMiss(Request.Key))
+				{
+					UE_LOG(LogDerivedDataCache, Verbose, TEXT("%s: Simulated miss for PutValue of %s from '%s'"),
+						*GetName(), *WriteToString<96>(Request.Key), *Request.Name);
+					bPutSucceeded = false;
+				}
+				bPutSucceeded ? OnHit(Request) : OnMiss(Request);
 			}
+			if (RequestIndex != Batch.Num())
+			{
+				UE_LOG(LogDerivedDataCache, Warning,
+					TEXT("Invalid response received from PutCacheValues rpc: %d results expected, received %d."), Batch.Num(), RequestIndex);
+			}
+		}
+		for (const FCachePutValueRequest& Request : Batch.RightChop(RequestIndex))
+		{
+			OnMiss(Request);
 		}
 	}
 }
@@ -875,40 +1035,62 @@ void FZenDerivedDataBackend::GetValue(
 {
 	TRACE_CPUPROFILER_EVENT_SCOPE(ZenDDC::GetValue);
 	TRACE_COUNTER_ADD(ZenDDC_CacheRecordRequestCount, int64(Requests.Num()));
+	COOK_STAT(auto Timer = UsageStats.TimeGet());
+	TRACE_COUNTER_ADD(ZenDDC_Get, (int64)Requests.Num());
+
+	auto OnHit = [this, &OnComplete COOK_STAT(, &Timer)](const FCacheGetValueRequest& Request, FValue&& Value)
+	{
+		UE_LOG(LogDerivedDataCache, Verbose, TEXT("%s: Cache hit for '%s' from '%s'"),
+			*GetName(), *WriteToString<96>(Request.Key), *Request.Name);
+		TRACE_COUNTER_ADD(ZenDDC_GetHit, int64(1));
+		int64 ReceivedSize = Value.GetData().GetCompressedSize();
+		TRACE_COUNTER_ADD(ZenDDC_BytesReceived, ReceivedSize);
+		COOK_STAT(Timer.AddHit(ReceivedSize));
+
+		if (OnComplete)
+		{
+			OnComplete({ Request.Name, Request.Key, MoveTemp(Value), Request.UserData, EStatus::Ok });
+		}
+	};
+	auto OnMiss = [this, &OnComplete](const FCacheGetValueRequest& Request)
+	{
+		UE_LOG(LogDerivedDataCache, Verbose, TEXT("%s: Cache miss for '%s' from '%s'"),
+			*GetName(), *WriteToString<96>(Request.Key), *Request.Name);
 
-	int32 TotalCompleted = 0;
+		if (OnComplete)
+		{
+			OnComplete({ Request.Name, Request.Key, {}, Request.UserData, EStatus::Error });
+		}
+	};
 
 	ForEachBatch(CacheRecordBatchSize, Requests.Num(),
-		[this, &Requests, &Owner, &OnComplete, &TotalCompleted](int32 BatchFirst, int32 BatchLast)
+		[this, &Requests, &Owner, &OnComplete, &OnHit, &OnMiss](int32 BatchFirst, int32 BatchLast)
 		{
-			TRACE_COUNTER_ADD(ZenDDC_Get, int64(BatchLast) - int64(BatchFirst) + 1);
-			COOK_STAT(auto Timer = UsageStats.TimeGet());
+			TArrayView<const FCacheGetValueRequest> Batch(Requests.GetData() + BatchFirst, BatchLast - BatchFirst + 1);
 
 			FCbWriter BatchRequest;
 			BatchRequest.BeginObject();
 			{
-				BatchRequest << "Method"_ASV << "GetCacheRecords";
+				BatchRequest << "Method"_ASV << "GetCacheValues"_ASV;
 				BatchRequest.BeginObject("Params"_ASV);
 				{
-					BatchRequest.BeginArray("CacheKeys"_ASV);
-					for (int32 KeyIndex = BatchFirst; KeyIndex <= BatchLast; KeyIndex++)
-					{
-						const FCacheKey& Key = Requests[KeyIndex].Key;
+					ECachePolicy BatchDefaultPolicy = Batch[0].Policy;
+					BatchRequest << "DefaultPolicy"_ASV << *WriteToString<128>(BatchDefaultPolicy);
 
+					BatchRequest.BeginArray("Requests");
+					for (const FCacheGetValueRequest& Request : Batch)
+					{
 						BatchRequest.BeginObject();
-						BatchRequest << "Bucket"_ASV << Key.Bucket.ToString();
-						BatchRequest << "Hash"_ASV << Key.Hash;
+						{
+							BatchRequest << "Key"_ASV << Request.Key;
+							if (Request.Policy != BatchDefaultPolicy)
+							{
+								BatchRequest << "Policy"_ASV << WriteToString<128>(Request.Policy);
+							}
+						}
 						BatchRequest.EndObject();
 					}
 					BatchRequest.EndArray();
-
-					// TODO: The policy needs to be sent with each key.
-					const ECachePolicy& Policy = Requests[BatchFirst].Policy;
-					BatchRequest.BeginObject("Policy"_ASV);
-					{
-						BatchRequest << "RecordPolicy"_ASV << static_cast<uint32>(Policy);
-					}
-					BatchRequest.EndObject();
 				}
 				BatchRequest.EndObject();
 			}
@@ -922,89 +1104,58 @@ void FZenDerivedDataBackend::GetValue(
 				HttpResult = Request->PerformRpc(TEXT("/z$/$rpc"_SV), BatchRequest.Save().AsObject(), BatchResponse);
 			}
 
+			int32 RequestIndex = 0;
 			if (HttpResult == Zen::FZenHttpRequest::Result::Success)
 			{
 				const FCbObject& ResponseObj = BatchResponse.GetObject();
 
-				int32 KeyIndex = BatchFirst;
-				for (FCbField ValueReferenceField : ResponseObj["Result"_ASV])
+				for (FCbFieldView ResultField : ResponseObj["Result"_ASV])
 				{
-					const FCacheGetValueRequest& Request = Requests[KeyIndex++];
-					const FCacheKey& Key = Request.Key;
-					FIoHash RawHash = ValueReferenceField["RawHash"_ASV].AsBinaryAttachment();
-					TOptional<FValue> ResultValue;
-
-					if (!RawHash.IsZero())
+					if (RequestIndex >= Batch.Num())
 					{
-						if (ShouldSimulateMiss(Key))
-						{
-							UE_LOG(LogDerivedDataCache, Verbose, TEXT("%s: Simulated miss for put of '%s' from '%s'"),
-								*GetName(), *WriteToString<96>(Key), *Request.Name);
-						}
-						else if (EnumHasAnyFlags(Request.Policy, UE::DerivedData::ECachePolicy::SkipData))
-						{
-							ResultValue.Emplace(RawHash, ValueReferenceField["RawSize"].AsUInt64());
-						}
-						else
-						{
-							const FCbAttachment* Attachment = BatchResponse.FindAttachment(RawHash);
-							if (Attachment)
-							{
-								ResultValue.Emplace(Attachment->AsCompressedBinary());
-							}
-						}
+						++RequestIndex;
+						continue;
 					}
+					const FCacheGetValueRequest& Request = Batch[RequestIndex++];
 
-					if (ResultValue)
+					FCbObjectView ResultObj = ResultField.AsObjectView();
+					TOptional<FValue> Value;
+					if (ShouldSimulateMiss(Request.Key))
 					{
-						UE_LOG(LogDerivedDataCache, Verbose, TEXT("%s: Cache hit for '%s' from '%s'"),
-							*GetName(), *WriteToString<96>(Key), *Request.Name);
-
-						TRACE_COUNTER_ADD(ZenDDC_GetHit, int64(1));
-						int64 ReceivedSize = ResultValue->GetData().GetCompressedSize();
-						TRACE_COUNTER_ADD(ZenDDC_BytesReceived, ReceivedSize);
-						COOK_STAT(Timer.AddHit(ReceivedSize));
-
-						if (OnComplete)
-						{
-							OnComplete({ Request.Name, Key, MoveTemp(*ResultValue), Request.UserData, EStatus::Ok });
-						}
+						UE_LOG(LogDerivedDataCache, Verbose, TEXT("%s: Simulated miss for GetValue of '%s' from '%s'"),
+							*GetName(), *WriteToString<96>(Request.Key), *Request.Name);
 					}
 					else
 					{
-						UE_LOG(LogDerivedDataCache, Verbose, TEXT("%s: Cache miss for '%s' from '%s'"),
-							*GetName(), *WriteToString<96>(Key), *Request.Name);
-
-						if (OnComplete)
+						FCbFieldView RawHashField = ResultObj["RawHash"];
+						FIoHash RawHash = RawHashField.AsHash();
+						if (const FCbAttachment* Attachment = BatchResponse.FindAttachment(RawHash))
+						{
+							Value.Emplace(Attachment->AsCompressedBinary());
+						}
+						else
 						{
-							OnComplete({ Request.Name, Key, {}, Request.UserData, EStatus::Error });
+							FCbFieldView RawSizeField = ResultObj["RawSize"];
+							uint64 RawSize = RawSizeField.AsUInt64();
+							if (!RawSizeField.HasError() && !RawHashField.HasError())
+							{
+								Value.Emplace(RawHash, RawSize);
+							}
 						}
 					}
-
-					TotalCompleted++;
+					(bool)Value ? OnHit(Request, MoveTemp(*Value)) : OnMiss(Request);
 				}
-			}
-			else
-			{
-				for (int32 KeyIndex = BatchFirst; KeyIndex <= BatchLast; KeyIndex++)
+				if (RequestIndex != Batch.Num())
 				{
-					const FCacheGetValueRequest& Request = Requests[KeyIndex];
-					const FCacheKey& Key = Request.Key;
-
-					UE_LOG(LogDerivedDataCache, Verbose, TEXT("%s: Cache miss for '%s' from '%s'"),
-						*GetName(), *WriteToString<96>(Key), *Request.Name);
-
-					if (OnComplete)
-					{
-						OnComplete({ Request.Name, Key, {}, Request.UserData, EStatus::Error });
-					}
-
-					TotalCompleted++;
+					UE_LOG(LogDerivedDataCache, Warning,
+						TEXT("Invalid response received from GetCacheValues rpc: %d results expected, received %d."), Batch.Num(), RequestIndex);
 				}
 			}
+			for (const FCacheGetValueRequest& Request : Batch.RightChop(RequestIndex))
+			{
+				OnMiss(Request);
+			}
 		});
-
-	UE_CLOG(TotalCompleted != Requests.Num(), LogDerivedDataCache, Warning, TEXT("Only '%d/%d' GetValue request(s) completed"), TotalCompleted, Requests.Num());
 	TRACE_COUNTER_SUBTRACT(ZenDDC_CacheRecordRequestCount, int64(Requests.Num()));
 }
 
@@ -1016,40 +1167,75 @@ void FZenDerivedDataBackend::GetChunks(
 {
 	TRACE_CPUPROFILER_EVENT_SCOPE(ZenDDC::GetChunks);
 	TRACE_COUNTER_ADD(ZenDDC_ChunkRequestCount, int64(Requests.Num()));
+	TRACE_COUNTER_ADD(ZenDDC_Get, int64(Requests.Num()));
+	COOK_STAT(auto Timer = UsageStats.TimeGet());
 
 	TArray<FCacheGetChunkRequest, TInlineAllocator<16>> SortedRequests(Requests);
 	SortedRequests.StableSort(TChunkLess());
+	auto OnHit = [this, &OnComplete COOK_STAT(, &Timer)](const FCacheGetChunkRequest& Request, FIoHash&& RawHash, uint64 RawSize, FSharedBuffer&& RequestedBytes)
+	{
+		UE_LOG(LogDerivedDataCache, Verbose, TEXT("%s: CacheChunk hit for '%s' from '%s'"),
+			*GetName(), *WriteToString<96>(Request.Key, '/', Request.Id), *Request.Name);
+		TRACE_COUNTER_ADD(ZenDDC_GetHit, int64(1));
+		COOK_STAT(Timer.AddHit(RawSize));
+		if (OnComplete)
+		{
+			OnComplete({ Request.Name, Request.Key, Request.Id, Request.RawOffset, RawSize, MoveTemp(RawHash), MoveTemp(RequestedBytes), Request.UserData, EStatus::Ok });
+		}
+	};
+
+	auto OnMiss = [this, &OnComplete COOK_STAT(, &Timer)](const FCacheGetChunkRequest& Request)
+	{
+		UE_LOG(LogDerivedDataCache, Display, TEXT("%s: CacheChunk miss with missing value '%s' for '%s' from '%s'"),
+			*GetName(), *WriteToString<16>(Request.Id), *WriteToString<96>(Request.Key), *Request.Name);
+		if (OnComplete)
+		{
+			OnComplete({ Request.Name, Request.Key, Request.Id, Request.RawOffset, 0, {}, {}, Request.UserData, EStatus::Error });
+		}
+	};
 
-	int32 TotalCompleted = 0;
 
 	ForEachBatch(CacheChunksBatchSize, SortedRequests.Num(),
-		[this, &SortedRequests, &Owner, &OnComplete, &TotalCompleted](int32 BatchFirst, int32 BatchLast)
+		[this, &SortedRequests, &Owner, &OnHit, &OnMiss](int32 BatchFirst, int32 BatchLast)
 	{
-		TRACE_COUNTER_ADD(ZenDDC_Get, int64(BatchLast) - int64(BatchFirst) + 1);
-		COOK_STAT(auto Timer = UsageStats.TimeGet());
+		TArrayView<const FCacheGetChunkRequest> Batch(SortedRequests.GetData() + BatchFirst, BatchLast - BatchFirst + 1);
 
 		FCbWriter BatchRequest;
 		BatchRequest.BeginObject();
 		{
-			BatchRequest << "Method"_ASV << "GetCacheValues";
+			BatchRequest << "Method"_ASV << "GetCacheChunks";
+			BatchRequest.AddInteger("MethodVersion"_ASV, 1);
 			BatchRequest.BeginObject("Params"_ASV);
 			{
+				ECachePolicy DefaultPolicy = Batch[0].Policy;
+				BatchRequest << "DefaultPolicy"_ASV << WriteToString<128>(DefaultPolicy);
 				BatchRequest.BeginArray("ChunkRequests"_ASV);
-				for (int32 ChunkIndex = BatchFirst; ChunkIndex <= BatchLast; ChunkIndex++)
+				for (const FCacheGetChunkRequest& Request : Batch)
 				{
-					const FCacheGetChunkRequest& Request = SortedRequests[ChunkIndex];
-					
 					BatchRequest.BeginObject();
 					
-					BatchRequest.BeginObject("Key"_ASV);
-					BatchRequest << "Bucket"_ASV << Request.Key.Bucket.ToString();
-					BatchRequest << "Hash"_ASV << Request.Key.Hash;
-					BatchRequest.EndObject();
+					BatchRequest << "Key"_ASV << Request.Key;
 
-					BatchRequest.AddObjectId("ValueId"_ASV, Request.Id);
-					BatchRequest << "RawOffset"_ASV << Request.RawOffset;
-					BatchRequest << "RawSize"_ASV << Request.RawSize;
-					BatchRequest << "Policy"_ASV << static_cast<uint32>(Request.Policy);
+					if (Request.Id.IsValid())
+					{
+						BatchRequest.AddObjectId("ValueId"_ASV, Request.Id);
+					}
+					if (Request.RawOffset != 0)
+					{
+						BatchRequest << "RawOffset"_ASV << Request.RawOffset;
+					}
+					if (Request.RawSize != MAX_uint64)
+					{
+						BatchRequest << "RawSize"_ASV << Request.RawSize;
+					}
+					if (!Request.RawHash.IsZero())
+					{
+						BatchRequest << "ChunkId"_ASV << Request.RawHash;
+					}
+					if (Request.Policy != DefaultPolicy)
+					{
+						BatchRequest << "Policy"_ASV << WriteToString<128>(Request.Policy);
+					}
 
 					BatchRequest.EndObject();
 				}
@@ -1067,128 +1253,71 @@ void FZenDerivedDataBackend::GetChunks(
 			HttpResult = Request->PerformRpc(TEXT("/z$/$rpc"_SV), BatchRequest.Save().AsObject(), BatchResponse);
 		}
 
+		int32 RequestIndex = 0;
 		if (HttpResult == Zen::FZenHttpRequest::Result::Success)
 		{
 			const FCbObject& ResponseObj = BatchResponse.GetObject();
 
-			int32 ChunkIndex = BatchFirst;
-			for (FCbFieldView HashView : ResponseObj["Result"_ASV])
+			for (FCbFieldView ResultView : ResponseObj["Result"_ASV])
 			{
-				const FCacheGetChunkRequest& Request = SortedRequests[ChunkIndex++];
+				if (RequestIndex >= Batch.Num())
+				{
+					++RequestIndex;
+					continue;
+				}
+				const FCacheGetChunkRequest& Request = Batch[RequestIndex++];
 
+				FIoHash RawHash;
+				bool Succeeded = false;
+				uint64 RawSize = 0;
+				FCbObjectView ResultObject = ResultView.AsObjectView();
+				FSharedBuffer RequestedBytes;
 				if (ShouldSimulateMiss(Request.Key))
 				{
 					UE_LOG(LogDerivedDataCache, Verbose, TEXT("%s: Simulated miss for get of '%s' from '%s'"),
-						*GetName(), *WriteToString<96>(Request.Key ,'/', Request.Id), *Request.Name);
-					
-					if (OnComplete)
-					{
-						OnComplete({ Request.Name, Request.Key, Request.Id, Request.RawOffset, 0, {}, {}, Request.UserData, EStatus::Error });
-					}
-				}
-				else if (const FCbAttachment* Attachment = BatchResponse.FindAttachment(HashView.AsHash()))
-				{
-					const FCompressedBuffer& CompressedBuffer = Attachment->AsCompressedBinary();
-					FSharedBuffer Buffer = FCompressedBufferReader(CompressedBuffer).Decompress(Request.RawOffset, Request.RawSize);
-					
-					UE_LOG(LogDerivedDataCache, Verbose, TEXT("%s: Cache hit for '%s' from '%s'"),
 						*GetName(), *WriteToString<96>(Request.Key, '/', Request.Id), *Request.Name);
-
-					const uint64 RawSize = Buffer.GetSize();
-					TRACE_COUNTER_ADD(ZenDDC_GetHit, int64(1));
-					TRACE_COUNTER_ADD(ZenDDC_BytesReceived, RawSize);
-					COOK_STAT(Timer.AddHit(RawSize));
-					
-					if (OnComplete)
-					{
-						OnComplete({ Request.Name, Request.Key, Request.Id, Request.RawOffset, RawSize, CompressedBuffer.GetRawHash(), MoveTemp(Buffer), Request.UserData, EStatus::Ok });
-					}
 				}
 				else
 				{
-					UE_LOG(LogDerivedDataCache, Display, TEXT("%s: Cache miss with missing value '%s' for '%s' from '%s'"),
-						*GetName(), *WriteToString<16>(Request.Id), *WriteToString<96>(Request.Key), *Request.Name);
-					
-					if (OnComplete)
+					FCbFieldView HashView = ResultObject["RawHash"_ASV];
+					RawHash = HashView.AsHash();
+					if (!HashView.HasError())
 					{
-						OnComplete({ Request.Name, Request.Key, Request.Id, Request.RawOffset, 0, {}, {}, Request.UserData, EStatus::Error });
+						if (const FCbAttachment* Attachment = BatchResponse.FindAttachment(HashView.AsHash()))
+						{
+							FCompressedBuffer CompressedBuffer = Attachment->AsCompressedBinary();
+							if (CompressedBuffer)
+							{
+								TRACE_COUNTER_ADD(ZenDDC_BytesReceived, CompressedBuffer.GetCompressedSize());
+								RequestedBytes = FCompressedBufferReader(CompressedBuffer).Decompress(Request.RawOffset, Request.RawSize);
+								RawSize = RequestedBytes.GetSize();
+								Succeeded = true;
+							}
+						}
+						else
+						{
+							FCbFieldView RawSizeField = ResultObject["RawSize"_ASV];
+							uint64 TotalSize = RawSizeField.AsUInt64();
+							Succeeded = !RawSizeField.HasError();
+							if (Succeeded)
+							{
+								RawSize = FMath::Min(Request.RawSize, TotalSize - FMath::Min(Request.RawOffset, TotalSize));
+							}
+						}
 					}
 				}
-
-				TotalCompleted++;
+				Succeeded ? OnHit(Request, MoveTemp(RawHash), RawSize, MoveTemp(RequestedBytes)) : OnMiss(Request);
 			}
 		}
-		else
+		for (const FCacheGetChunkRequest& Request : Batch.RightChop(RequestIndex))
 		{
-			for (int32 ChunkIndex = BatchFirst; ChunkIndex <= BatchLast; ChunkIndex++)
-			{
-				const FCacheGetChunkRequest& Request = SortedRequests[ChunkIndex];
-				
-				UE_LOG(LogDerivedDataCache, Display, TEXT("%s: Cache miss with missing value '%s' for '%s' from '%s'"),
-					*GetName(), *WriteToString<16>(Request.Id), *WriteToString<96>(Request.Key), *Request.Name);
-				
-				if (OnComplete)
-				{
-					OnComplete({ Request.Name, Request.Key, Request.Id, Request.RawOffset, 0, {}, {}, Request.UserData, EStatus::Error });
-				}
-				
-				TotalCompleted++;
-			}
+			OnMiss(Request);
 		}
 	});
 
-	UE_CLOG(TotalCompleted != SortedRequests.Num(), LogDerivedDataCache, Warning, TEXT("Only '%d/%d' cache chunk request(s) completed"), TotalCompleted, SortedRequests.Num());
 	TRACE_COUNTER_SUBTRACT(ZenDDC_ChunkRequestCount, int64(Requests.Num()));
 }
 
-bool FZenDerivedDataBackend::PutCacheRecord(const FCacheRecord& Record, const FCacheRecordPolicy& Policy)
-{
-	const FCacheKey& Key = Record.GetKey();
-	FCbPackage Package = Record.Save();
-	FBufferArchive Ar;
-	Package.Save(Ar);
-	FCompositeBuffer Buffer = FCompositeBuffer(FSharedBuffer::MakeView(Ar.GetData(), Ar.Num()));
-	TStringBuilder<256> Uri;
-	AppendZenUri(Record.GetKey(), Uri);
-	AppendPolicyQueryString(Policy.GetRecordPolicy(), Uri);
-	if (PutZenData(Uri.ToString(), Buffer, Zen::EContentType::CbPackage)
-		!= FDerivedDataBackendInterface::EPutStatus::Cached)
-	{
-		return false;
-	}
-
-	return true;
-}
-
-bool FZenDerivedDataBackend::PutCacheValue(const FCacheKey& Key, const FValue& Value, const ECachePolicy& ValuePolicy)
-{
-	FCbWriter Writer;
-	Writer.BeginObject();
-	Writer.AddBinaryAttachment("RawHash", Value.GetRawHash());
-	Writer.AddInteger("RawSize", Value.GetRawSize());
-	Writer.EndObject();
-
-	FCbPackage Package(Writer.Save().AsObject());
-	if (Value.HasData())
-	{
-		Package.AddAttachment(FCbAttachment(Value.GetData()));
-	}
-
-	FBufferArchive Ar;
-	Package.Save(Ar);
-	FCompositeBuffer Buffer = FCompositeBuffer(FSharedBuffer::MakeView(Ar.GetData(), Ar.Num()));
-	TStringBuilder<256> Uri;
-	AppendZenUri(Key, Uri);
-	AppendPolicyQueryString(ValuePolicy, Uri);
-	if (PutZenData(Uri.ToString(), Buffer, Zen::EContentType::CbPackage)
-		!= FDerivedDataBackendInterface::EPutStatus::Cached)
-	{
-		return false;
-	}
-
-	return true;
-}
-
 FDerivedDataBackendInterface* CreateZenDerivedDataBackend(const TCHAR* NodeName, const TCHAR* ServiceUrl, const TCHAR* Namespace)
 {
 	FZenDerivedDataBackend* Backend = new FZenDerivedDataBackend(ServiceUrl, Namespace);
@@ -1215,4 +1344,4 @@ FDerivedDataBackendInterface* CreateZenDerivedDataBackend(const TCHAR* NodeName,
 
 } // UE::DerivedData::CacheStore::ZenCache
 
-#endif // UE_WITH_ZEN
+#endif // UE_WITH_ZEN
\ No newline at end of file
