DDC: Converted modular feature names to constant lazy names

#rnx

diff --git a/Engine/Source/Developer/DerivedDataCache/Private/DerivedDataBuildFunctionRegistry.cpp b/Engine/Source/Developer/DerivedDataCache/Private/DerivedDataBuildFunctionRegistry.cpp
--- a/Engine/Source/Developer/DerivedDataCache/Private/DerivedDataBuildFunctionRegistry.cpp
+++ b/Engine/Source/Developer/DerivedDataCache/Private/DerivedDataBuildFunctionRegistry.cpp
@@ -42,7 +42,7 @@ private:
 FBuildFunctionRegistry::FBuildFunctionRegistry()
 {
 	IModularFeatures& ModularFeatures = IModularFeatures::Get();
-	for (IBuildFunctionFactory* Factory : ModularFeatures.GetModularFeatureImplementations<IBuildFunctionFactory>(IBuildFunctionFactory::GetFeatureName()))
+	for (IBuildFunctionFactory* Factory : ModularFeatures.GetModularFeatureImplementations<IBuildFunctionFactory>(IBuildFunctionFactory::FeatureName))
 	{
 		AddFunctionNoLock(Factory);
 	}
@@ -59,7 +59,7 @@ FBuildFunctionRegistry::~FBuildFunctionRegistry()
 
 void FBuildFunctionRegistry::OnModularFeatureRegistered(const FName& Type, IModularFeature* ModularFeature)
 {
-	if (Type == IBuildFunctionFactory::GetFeatureName())
+	if (Type == IBuildFunctionFactory::FeatureName)
 	{
 		FWriteScopeLock WriteLock(Lock);
 		AddFunctionNoLock(static_cast<IBuildFunctionFactory*>(ModularFeature));
@@ -68,7 +68,7 @@ void FBuildFunctionRegistry::OnModularFeatureRegistered(const FName& Type, IModu
 
 void FBuildFunctionRegistry::OnModularFeatureUnregistered(const FName& Type, IModularFeature* ModularFeature)
 {
-	if (Type == IBuildFunctionFactory::GetFeatureName())
+	if (Type == IBuildFunctionFactory::FeatureName)
 	{
 		RemoveFunction(static_cast<IBuildFunctionFactory*>(ModularFeature));
 	}
diff --git a/Engine/Source/Developer/DerivedDataCache/Private/DerivedDataBuildWorkerRegistry.cpp b/Engine/Source/Developer/DerivedDataCache/Private/DerivedDataBuildWorkerRegistry.cpp
--- a/Engine/Source/Developer/DerivedDataCache/Private/DerivedDataBuildWorkerRegistry.cpp
+++ b/Engine/Source/Developer/DerivedDataCache/Private/DerivedDataBuildWorkerRegistry.cpp
@@ -164,11 +164,11 @@ private:
 FBuildWorkerRegistry::FBuildWorkerRegistry()
 {
 	IModularFeatures& ModularFeatures = IModularFeatures::Get();
-	if (ModularFeatures.IsModularFeatureAvailable(IBuildWorkerExecutor::GetFeatureName()))
+	if (ModularFeatures.IsModularFeatureAvailable(IBuildWorkerExecutor::FeatureName))
 	{
-		Executor = &ModularFeatures.GetModularFeature<IBuildWorkerExecutor>(IBuildWorkerExecutor::GetFeatureName());
+		Executor = &ModularFeatures.GetModularFeature<IBuildWorkerExecutor>(IBuildWorkerExecutor::FeatureName);
 	}
-	for (IBuildWorkerFactory* Worker : ModularFeatures.GetModularFeatureImplementations<IBuildWorkerFactory>(IBuildWorkerFactory::GetFeatureName()))
+	for (IBuildWorkerFactory* Worker : ModularFeatures.GetModularFeatureImplementations<IBuildWorkerFactory>(IBuildWorkerFactory::FeatureName))
 	{
 		AddWorker(Worker);
 	}
@@ -185,12 +185,12 @@ FBuildWorkerRegistry::~FBuildWorkerRegistry()
 
 void FBuildWorkerRegistry::OnModularFeatureRegistered(const FName& Type, IModularFeature* ModularFeature)
 {
-	if (!Executor && Type == IBuildWorkerExecutor::GetFeatureName())
+	if (!Executor && Type == IBuildWorkerExecutor::FeatureName)
 	{
 		FWriteScopeLock WriteLock(Lock);
 		Executor = static_cast<IBuildWorkerExecutor*>(ModularFeature);
 	}
-	else if (Type == IBuildWorkerFactory::GetFeatureName())
+	else if (Type == IBuildWorkerFactory::FeatureName)
 	{
 		AddWorker(static_cast<IBuildWorkerFactory*>(ModularFeature));
 	}
@@ -198,18 +198,18 @@ void FBuildWorkerRegistry::OnModularFeatureRegistered(const FName& Type, IModula
 
 void FBuildWorkerRegistry::OnModularFeatureUnregistered(const FName& Type, IModularFeature* ModularFeature)
 {
-	if (Executor == ModularFeature && Type == IBuildWorkerExecutor::GetFeatureName())
+	if (Executor == ModularFeature && Type == IBuildWorkerExecutor::FeatureName)
 	{
 		IModularFeature* NextExecutor = nullptr;
 		IModularFeatures& ModularFeatures = IModularFeatures::Get();
-		if (ModularFeatures.IsModularFeatureAvailable(IBuildWorkerExecutor::GetFeatureName()))
+		if (ModularFeatures.IsModularFeatureAvailable(IBuildWorkerExecutor::FeatureName))
 		{
-			NextExecutor = &ModularFeatures.GetModularFeature<IBuildWorkerExecutor>(IBuildWorkerExecutor::GetFeatureName());
+			NextExecutor = &ModularFeatures.GetModularFeature<IBuildWorkerExecutor>(IBuildWorkerExecutor::FeatureName);
 		}
 		FWriteScopeLock WriteLock(Lock);
 		Executor = static_cast<IBuildWorkerExecutor*>(NextExecutor);
 	}
-	else if (Type == IBuildWorkerFactory::GetFeatureName())
+	else if (Type == IBuildWorkerFactory::FeatureName)
 	{
 		RemoveWorker(static_cast<IBuildWorkerFactory*>(ModularFeature));
 	}
diff --git a/Engine/Source/Developer/DerivedDataCache/Public/DerivedDataBuildFunctionFactory.h b/Engine/Source/Developer/DerivedDataCache/Public/DerivedDataBuildFunctionFactory.h
--- a/Engine/Source/Developer/DerivedDataCache/Public/DerivedDataBuildFunctionFactory.h
+++ b/Engine/Source/Developer/DerivedDataCache/Public/DerivedDataBuildFunctionFactory.h
@@ -13,14 +13,10 @@ namespace UE::DerivedData
 class IBuildFunctionFactory : public IModularFeature
 {
 public:
+	static inline const FLazyName FeatureName{"BuildFunctionFactory"};
+
 	/** Returns the build function associated with this factory. */
 	virtual const IBuildFunction& GetFunction() const = 0;
-
-	/** Returns the name of the build function factory modular feature. */
-	static FName GetFeatureName()
-	{
-		return FName("BuildFunctionFactory");
-	}
 };
 
 /**
@@ -43,12 +39,12 @@ public:
 	explicit TBuildFunctionFactory(ArgTypes&&... Args)
 		: Function(Forward<ArgTypes>(Args)...)
 	{
-		IModularFeatures::Get().RegisterModularFeature(GetFeatureName(), this);
+		IModularFeatures::Get().RegisterModularFeature(FeatureName, this);
 	}
 
 	~TBuildFunctionFactory()
 	{
-		IModularFeatures::Get().UnregisterModularFeature(GetFeatureName(), this);
+		IModularFeatures::Get().UnregisterModularFeature(FeatureName, this);
 	}
 
 	const IBuildFunction& GetFunction() const final
diff --git a/Engine/Source/Developer/DerivedDataCache/Public/DerivedDataBuildWorker.h b/Engine/Source/Developer/DerivedDataCache/Public/DerivedDataBuildWorker.h
--- a/Engine/Source/Developer/DerivedDataCache/Public/DerivedDataBuildWorker.h
+++ b/Engine/Source/Developer/DerivedDataCache/Public/DerivedDataBuildWorker.h
@@ -64,21 +64,19 @@ public:
 class IBuildWorkerFactory : public IModularFeature
 {
 public:
+	static inline const FLazyName FeatureName{"BuildWorkerFactory"};
+
 	virtual ~IBuildWorkerFactory() = default;
 
 	virtual void Build(FBuildWorkerBuilder& Builder) = 0;
 	virtual void FindFileData(TConstArrayView<FIoHash> RawHashes, IRequestOwner& Owner, FOnBuildWorkerFileDataComplete&& OnComplete) = 0;
-
-	/** Returns the name of the build worker factory modular feature. */
-	static FName GetFeatureName()
-	{
-		return FName("BuildWorkerFactory");
-	}
 };
 
 class IBuildWorkerExecutor : public IModularFeature
 {
 public:
+	static inline const FLazyName FeatureName{"BuildWorkerExecutor"};
+
 	virtual ~IBuildWorkerExecutor() = default;
 
 	virtual void Build(
@@ -91,12 +89,6 @@ public:
 		FOnBuildWorkerActionComplete&& OnComplete) = 0;
 
 	virtual TConstArrayView<FStringView> GetHostPlatforms() const = 0;
-
-	/** Returns the name of the build worker executor modular feature. */
-	static FName GetFeatureName()
-	{
-		return FName("BuildWorkerExecutor");
-	}
 };
 
 struct FBuildWorkerActionCompleteParams
diff --git a/Engine/Source/Developer/DesktopPlatform/Private/TargetReceiptBuildWorker.cpp b/Engine/Source/Developer/DesktopPlatform/Private/TargetReceiptBuildWorker.cpp
--- a/Engine/Source/Developer/DesktopPlatform/Private/TargetReceiptBuildWorker.cpp
+++ b/Engine/Source/Developer/DesktopPlatform/Private/TargetReceiptBuildWorker.cpp
@@ -159,7 +159,7 @@ FTargetReceiptBuildWorker::~FTargetReceiptBuildWorker()
 	FTargetReceiptBuildWorkerFactory* WorkerFactory = GetWorkerFactory();
 	if (bEnabled)
 	{
-		IModularFeatures::Get().UnregisterModularFeature(UE::DerivedData::IBuildWorkerFactory::GetFeatureName(), WorkerFactory);
+		IModularFeatures::Get().UnregisterModularFeature(UE::DerivedData::IBuildWorkerFactory::FeatureName, WorkerFactory);
 	}
 	WorkerFactory->~FTargetReceiptBuildWorkerFactory();
 }
@@ -494,6 +494,6 @@ void FTargetReceiptBuildWorker::PopulateWorkerFromReceipt(const FString& TargetR
 		}
 	}
 
-	IModularFeatures::Get().RegisterModularFeature(UE::DerivedData::IBuildWorkerFactory::GetFeatureName(), GetWorkerFactory());
+	IModularFeatures::Get().RegisterModularFeature(UE::DerivedData::IBuildWorkerFactory::FeatureName, GetWorkerFactory());
 	bEnabled = true;
 }
diff --git a/Engine/Source/Editor/UnrealEd/Private/DerivedDataBuildLocalExecutor.cpp b/Engine/Source/Editor/UnrealEd/Private/DerivedDataBuildLocalExecutor.cpp
--- a/Engine/Source/Editor/UnrealEd/Private/DerivedDataBuildLocalExecutor.cpp
+++ b/Engine/Source/Editor/UnrealEd/Private/DerivedDataBuildLocalExecutor.cpp
@@ -51,7 +51,7 @@ public:
 		constexpr bool Tree = true;
 		IFileManager::Get().DeleteDirectory(*SandboxRootDir, RequireExists, Tree);
 
-		IModularFeatures::Get().RegisterModularFeature(IBuildWorkerExecutor::GetFeatureName(), this);
+		IModularFeatures::Get().RegisterModularFeature(IBuildWorkerExecutor::FeatureName, this);
 	}
 
 	~FLocalBuildWorkerExecutor() final = default;
diff --git a/Engine/Source/Editor/UnrealEd/Private/DerivedDataBuildRemoteExecutor.cpp b/Engine/Source/Editor/UnrealEd/Private/DerivedDataBuildRemoteExecutor.cpp
--- a/Engine/Source/Editor/UnrealEd/Private/DerivedDataBuildRemoteExecutor.cpp
+++ b/Engine/Source/Editor/UnrealEd/Private/DerivedDataBuildRemoteExecutor.cpp
@@ -166,7 +166,7 @@ public:
 			if (bEnabled)
 			{
 				RequestPool = MakeUnique<UE::Zen::FZenHttpRequestPool>(ZenServiceInstance.GetURL());
-				IModularFeatures::Get().RegisterModularFeature(IBuildWorkerExecutor::GetFeatureName(), this);
+				IModularFeatures::Get().RegisterModularFeature(IBuildWorkerExecutor::FeatureName, this);
 			}
 			else
 			{
@@ -185,7 +185,7 @@ public:
 	{
 		if (bEnabled)
 		{
-			IModularFeatures::Get().UnregisterModularFeature(IBuildWorkerExecutor::GetFeatureName(), this);
+			IModularFeatures::Get().UnregisterModularFeature(IBuildWorkerExecutor::FeatureName, this);
 			bProcessingThreadRunning = false;
 			ProcessingThreadEvent->Trigger();
 			ProcessingThread->WaitForCompletion();
