Implement DDC2 Value API for Jupiter backend.

diff --git a/Engine/Source/Developer/DerivedDataCache/Private/HttpDerivedDataBackend.cpp b/Engine/Source/Developer/DerivedDataCache/Private/HttpDerivedDataBackend.cpp
--- a/Engine/Source/Developer/DerivedDataCache/Private/HttpDerivedDataBackend.cpp
+++ b/Engine/Source/Developer/DerivedDataCache/Private/HttpDerivedDataBackend.cpp
@@ -42,6 +42,7 @@
 #include "Serialization/CompactBinary.h"
 #include "Serialization/CompactBinaryPackage.h"
 #include "Serialization/CompactBinaryValidation.h"
+#include "Serialization/CompactBinaryWriter.h"
 #include "Serialization/JsonReader.h"
 #include "Serialization/JsonSerializer.h"
 #include "Serialization/JsonWriter.h"
@@ -1936,6 +1937,25 @@ TConstArrayView<uint8> MakeArrayView(FSharedBuffer Buffer)
 	return TConstArrayView<uint8>(reinterpret_cast<const uint8*>(Buffer.GetData()), Buffer.GetSize());
 }
 
+static bool IsValueDataReady(FValue& Value, const ECachePolicy Policy)
+{
+	if (!EnumHasAnyFlags(Policy, ECachePolicy::Query))
+	{
+		Value = Value.RemoveData();
+		return true;
+	}
+
+	if (Value.HasData())
+	{
+		if (EnumHasAnyFlags(Policy, ECachePolicy::SkipData))
+		{
+			Value = Value.RemoveData();
+		}
+		return true;
+	}
+	return false;
+};
+
 //----------------------------------------------------------------------------------------------------------
 // FHttpAccessToken
 //----------------------------------------------------------------------------------------------------------
@@ -2029,7 +2049,14 @@ public:
 		TConstArrayView<FCacheGetRequest> Requests,
 		IRequestOwner& Owner,
 		FOnCacheGetComplete&& OnComplete) override;
-
+	virtual void PutValue(
+		TConstArrayView<FCachePutValueRequest> Requests,
+		IRequestOwner& Owner,
+		FOnCachePutValueComplete&& OnComplete) override;
+	virtual void GetValue(
+		TConstArrayView<FCacheGetValueRequest> Requests,
+		IRequestOwner& Owner,
+		FOnCacheGetValueComplete&& OnComplete) override;
 	virtual void GetChunks(
 		TConstArrayView<FCacheGetChunkRequest> Requests,
 		IRequestOwner& Owner,
@@ -2077,7 +2104,7 @@ private:
 	bool ShouldSimulateMiss(const FCacheKey& Key);
 
 	bool PutCacheRecord(FStringView Name, const FCacheRecord& Record, const FCacheRecordPolicy& Policy, uint64& OutWriteSize);
-	uint64 PutRef(const FCacheRecord& Record, const FCbPackage& Package, FStringView Bucket, bool bFinalize, TArray<FIoHash>& OutNeededBlobHashes, bool& bOutPutCompletedSuccessfully);
+	uint64 PutRef(const FCbPackage& Package, const FCacheKey& Key, FStringView Bucket, bool bFinalize, TArray<FIoHash>& OutNeededBlobHashes, bool& bOutPutCompletedSuccessfully);
 
 	FOptionalCacheRecord GetCacheRecordOnly(
 		FStringView Name,
@@ -2088,15 +2115,38 @@ private:
 		const FCacheKey& Key,
 		const FCacheRecordPolicy& Policy,
 		EStatus& OutStatus);
+
+	bool PutCacheValue(
+		FStringView Name,
+		const FCacheKey& Key,
+		const FValue& Value,
+		ECachePolicy Policy,
+		uint64& OutWriteSize);
+
+	bool GetCacheValueOnly(
+		FStringView Name,
+		const FCacheKey& Key,
+		ECachePolicy Policy,
+		FValue& OutValue);
+	bool GetCacheValue(
+		FStringView Name,
+		const FCacheKey& Key,
+		ECachePolicy Policy,
+		FValue& OutValue);
+
+	template<typename ValueType, typename ValuePrinterType>
 	bool TryGetCachedDataBatch(
 		FStringView Name,
 		const FCacheKey& Key,
-		TConstArrayView<FValueWithId> Values,
-		TArray<FCompressedBuffer>& OutBuffers);
+		TConstArrayView<ValueType> Values,
+		TArray<FCompressedBuffer>& OutBuffers,
+		ValuePrinterType ValuePrinter);
+	template<typename ValueType, typename ValuePrinterType>
 	bool CachedDataProbablyExistsBatch(
 		FStringView Name,
 		const FCacheKey& Key,
-		TConstArrayView<FValueWithId> Values);
+		TConstArrayView<ValueType> Values,
+		ValuePrinterType ValuePrinter);
 };
 
 FHttpDerivedDataBackend::FHttpDerivedDataBackend(
@@ -2365,12 +2415,11 @@ bool FHttpDerivedDataBackend::ShouldSimulateMiss(const FCacheKey& Key)
 	return false;
 }
 
-uint64 FHttpDerivedDataBackend::PutRef(const FCacheRecord& Record, const FCbPackage& Package, FStringView Bucket, bool bFinalize, TArray<FIoHash>& OutNeededBlobHashes, bool& bOutPutCompletedSuccessfully)
+uint64 FHttpDerivedDataBackend::PutRef(const FCbPackage& Package, const FCacheKey& Key, FStringView Bucket, bool bFinalize, TArray<FIoHash>& OutNeededBlobHashes, bool& bOutPutCompletedSuccessfully)
 {
 	bOutPutCompletedSuccessfully = false;
 
 	uint64 BytesSent = 0;
-	const FCacheKey& Key = Record.GetKey();
 	TStringBuilder<256> RefsUri;
 	RefsUri << "api/v1/refs/" << StructuredNamespace << "/" << Bucket << "/" << Key.Hash;
 	if (bFinalize)
@@ -2472,8 +2521,7 @@ bool FHttpDerivedDataBackend::PutCacheRecord(
 	// TODO: Jupiter currently always overwrites.  It doesn't have a "write if not present" feature (for records or attachments),
 	//		 but would require one to implement all policy correctly.
 
-	TStringConversion<TStringConvert<ANSICHAR, TCHAR>> BucketStringConversion(Key.Bucket.ToString());
-	FString Bucket(FStringView(BucketStringConversion.Get(), BucketStringConversion.Length()));
+	FString Bucket(Key.Bucket.ToString());
 	Bucket.ToLowerInline();
 
 	bool bPutCompletedSuccessfully = false;
@@ -2482,7 +2530,7 @@ bool FHttpDerivedDataBackend::PutCacheRecord(
 	TArray<FIoHash> NeededBlobHashes;
 
 	// Initial record upload
-	size_t PutRefBytesSent = PutRef(Record, Package, Bucket, false, NeededBlobHashes, bPutCompletedSuccessfully);
+	size_t PutRefBytesSent = PutRef(Package, Record.GetKey(), Bucket, false, NeededBlobHashes, bPutCompletedSuccessfully);
 	OutWriteSize += PutRefBytesSent;
 
 	// TODO: blob uploading and finalization should be replaced with a single batch compressed blob upload endpoint in the future.
@@ -2533,7 +2581,7 @@ bool FHttpDerivedDataBackend::PutCacheRecord(
 	// Finalization (if any blobs needed)
 	if (!NeededBlobHashes.IsEmpty())
 	{
-		size_t FinalizeRefBytesSent = PutRef(Record, Package, Bucket, true, NeededBlobHashes, bPutCompletedSuccessfully);
+		size_t FinalizeRefBytesSent = PutRef(Package, Record.GetKey(), Bucket, true, NeededBlobHashes, bPutCompletedSuccessfully);
 		OutWriteSize += FinalizeRefBytesSent;
 	}
 
@@ -2572,8 +2620,7 @@ FOptionalCacheRecord FHttpDerivedDataBackend::GetCacheRecordOnly(
 
 	FOptionalCacheRecord Record;
 	{
-		TStringConversion<TStringConvert<ANSICHAR, TCHAR>> BucketStringConversion(Key.Bucket.ToString());
-		FString Bucket(FStringView(BucketStringConversion.Get(), BucketStringConversion.Length()));
+		FString Bucket(Key.Bucket.ToString());
 		Bucket.ToLowerInline();
 
 		TStringBuilder<256> RefsUri;
@@ -2625,24 +2672,211 @@ FOptionalCacheRecord FHttpDerivedDataBackend::GetCacheRecordOnly(
 	return Record;
 }
 
-static bool IsValueDataReady(FValue& Value, const ECachePolicy Policy)
+bool FHttpDerivedDataBackend::PutCacheValue(
+	const FStringView Name,
+	const FCacheKey& Key,
+	const FValue& Value,
+	const ECachePolicy Policy,
+	uint64& OutWriteSize)
 {
-	if (!EnumHasAnyFlags(Policy, ECachePolicy::Query))
+	if (!IsWritable())
 	{
-		Value = Value.RemoveData();
-		return true;
+		UE_LOG(LogDerivedDataCache, VeryVerbose,
+			TEXT("%s: Skipped put of %s from '%.*s' because this cache store is read-only"),
+			*GetName(), *WriteToString<96>(Key), Name.Len(), Name.GetData());
+		return false;
 	}
 
-	if (Value.HasData())
+	// Skip the request if storing to the cache is disabled.
+	const ECachePolicy StoreFlag = SpeedClass == ESpeedClass::Local ? ECachePolicy::StoreLocal : ECachePolicy::StoreRemote;
+	if (!EnumHasAnyFlags(Policy, StoreFlag))
+	{
+		UE_LOG(LogDerivedDataCache, VeryVerbose, TEXT("%s: Skipped put of %s from '%.*s' due to cache policy"),
+			*GetName(), *WriteToString<96>(Key), Name.Len(), Name.GetData());
+		return false;
+	}
+
+	if (ShouldSimulateMiss(Key))
 	{
+		UE_LOG(LogDerivedDataCache, Verbose, TEXT("%s: Simulated miss for put of %s from '%.*s'"),
+			*GetName(), *WriteToString<96>(Key), Name.Len(), Name.GetData());
+		return false;
+	}
+
+	// TODO: Jupiter currently always overwrites.  It doesn't have a "write if not present" feature (for records or attachments),
+	//		 but would require one to implement all policy correctly.
+
+	FString Bucket(Key.Bucket.ToString());
+	Bucket.ToLowerInline();
+
+	bool bPutCompletedSuccessfully = false;
+
+	FCbWriter Writer;
+	Writer.BeginObject();
+	Writer.AddBinaryAttachment("RawHash", Value.GetRawHash());
+	Writer.AddInteger("RawSize", Value.GetRawSize());
+	Writer.EndObject();
+
+	FCbPackage Package(Writer.Save().AsObject());
+	TArray<FIoHash> NeededBlobHashes;
+
+	// Initial record upload
+	size_t PutRefBytesSent = PutRef(Package, Key, Bucket, false, NeededBlobHashes, bPutCompletedSuccessfully);
+	OutWriteSize += PutRefBytesSent;
+
+	// Needed blob upload (if any missing)
+	if (!NeededBlobHashes.IsEmpty())
+	{
+		check(NeededBlobHashes.Num() == 1);
+		check(NeededBlobHashes[0] == Value.GetRawHash());
+		TStringBuilder<256> CompressedBlobsUri;
+		CompressedBlobsUri << "api/v1/compressed-blobs/" << StructuredNamespace << "/" << Value.GetRawHash();
+
+		FSharedBuffer TempBuffer = Value.GetData().GetCompressed().ToShared();
+
+		int64 ResponseCode = 0;
+		for (uint32 Attempts = 0; (Attempts < UE_HTTPDDC_MAX_ATTEMPTS) && !ShouldAbortForShutdown() && (Attempts == 0 || ShouldRetryOnError(ResponseCode)); ++Attempts)
+		{
+			FScopedRequestPtr Request(PutRequestPools[IsInGameThread()].Get());
+			Request->PerformBlockingUpload<FHttpRequest::PutCompressedBlob>(*CompressedBlobsUri, MakeArrayView(TempBuffer));
+
+			ResponseCode = Request->GetResponseCode();
+
+			if (FHttpRequest::IsSuccessResponse(ResponseCode))
+			{
+				OutWriteSize += Request->GetBytesSent();
+				break;
+			}
+		}
+
+		OutWriteSize += PutRef(Package, Key, Bucket, true, NeededBlobHashes, bPutCompletedSuccessfully);
+	}
+
+
+	return true;
+}
+
+bool FHttpDerivedDataBackend::GetCacheValueOnly(
+	const FStringView Name,
+	const FCacheKey& Key,
+	const ECachePolicy Policy,
+	FValue& OutValue)
+{
+	if (!IsUsable())
+	{
+		UE_LOG(LogDerivedDataCache, VeryVerbose,
+			TEXT("%s: Skipped get of %s from '%.*s' because this cache store is not available"),
+			*GetName(), *WriteToString<96>(Key), Name.Len(), Name.GetData());
+		return false;
+	}
+
+	// Skip the request if querying the cache is disabled.
+	const ECachePolicy QueryFlag = SpeedClass == ESpeedClass::Local ? ECachePolicy::QueryLocal : ECachePolicy::QueryRemote;
+	if (!EnumHasAnyFlags(Policy, QueryFlag))
+	{
+		UE_LOG(LogDerivedDataCache, VeryVerbose, TEXT("%s: Skipped get of %s from '%.*s' due to cache policy"),
+			*GetName(), *WriteToString<96>(Key), Name.Len(), Name.GetData());
+		return false;
+	}
+
+	if (ShouldSimulateMiss(Key))
+	{
+		UE_LOG(LogDerivedDataCache, Verbose, TEXT("%s: Simulated miss for get of %s from '%.*s'"),
+			*GetName(), *WriteToString<96>(Key), Name.Len(), Name.GetData());
+		return false;
+	}
+
+
+	FString Bucket(Key.Bucket.ToString());
+	Bucket.ToLowerInline();
+
+	TStringBuilder<256> RefsUri;
+	RefsUri << "api/v1/refs/" << StructuredNamespace << "/" << Bucket << "/" << Key.Hash;
+
+	bool bIsSuccessfulResponse = false;
+	FSharedBuffer ResponseBuffer;
+	int64 ResponseCode = 0;
+	for (uint32 Attempts = 0; (Attempts < UE_HTTPDDC_MAX_ATTEMPTS) && !ShouldAbortForShutdown() && (Attempts == 0 || ShouldRetryOnError(ResponseCode)); ++Attempts)
+	{
+		FScopedRequestPtr Request(PutRequestPools[IsInGameThread()].Get());
+
+		TArray<uint8> ByteArray;
+		Request->SetHeader(TEXT("Accept"), TEXT("application/x-ue-cb"));
+		Request->PerformBlockingDownload(*RefsUri, &ByteArray, { 401, 404 });
+		ResponseCode = Request->GetResponseCode();
+
+		if (FHttpRequest::IsSuccessResponse(ResponseCode))
+		{
+			ResponseBuffer = MakeSharedBufferFromArray(MoveTemp(ByteArray));
+			bIsSuccessfulResponse = true;
+			break;
+		}
+	}
+
+	if (!bIsSuccessfulResponse)
+	{
+		UE_LOG(LogDerivedDataCache, Verbose, TEXT("%s: Cache miss with missing package for %s from '%.*s'"),
+			*GetName(), *WriteToString<96>(Key), Name.Len(), Name.GetData());
+		return false;
+	}
+
+	if (ValidateCompactBinary(ResponseBuffer, ECbValidateMode::Default) != ECbValidateError::None)
+	{
+		UE_LOG(LogDerivedDataCache, Display, TEXT("%s: Cache miss with invalid package for %s from '%.*s'"),
+			*GetName(), *WriteToString<96>(Key), Name.Len(), Name.GetData());
+		return false;
+	}
+
+	FCbPackage Package;
+	Package.SetObject(FCbObject(ResponseBuffer));
+	const FCbObjectView Object = Package.GetObject();
+	const FIoHash RawHash = Object["RawHash"].AsHash();
+	const uint64 RawSize = Object["RawSize"].AsUInt64(MAX_uint64);
+	if (RawHash.IsZero() || RawSize == MAX_uint64)
+	{
+		UE_LOG(LogDerivedDataCache, Display, TEXT("%s: Cache miss with invalid value for %s from '%.*s'"),
+			*GetName(), *WriteToString<96>(Key), Name.Len(), Name.GetData());
+		return false;
+	}
+
+	OutValue = FValue(RawHash, RawSize);
+
+	return true;
+}
+
+bool FHttpDerivedDataBackend::GetCacheValue(
+	const FStringView Name,
+	const FCacheKey& Key,
+	const ECachePolicy Policy,
+	FValue& OutValue)
+{
+	if (GetCacheValueOnly(Name, Key, Policy, OutValue))
+	{
+		if (IsValueDataReady(OutValue, Policy))
+		{
+			return true;
+		}
+
 		if (EnumHasAnyFlags(Policy, ECachePolicy::SkipData))
 		{
-			Value = Value.RemoveData();
+			if (CachedDataProbablyExistsBatch<FValue>(Name, Key, {OutValue}, [](const FValue& Value) { return TEXT("Default"); }))
+			{
+				return true;
+			}
+		}
+		else
+		{
+			TArray<FCompressedBuffer> Buffers;
+			if (TryGetCachedDataBatch<FValue>(Name, Key, {OutValue}, Buffers, [](const FValue& Value) { return TEXT("Default"); }))
+			{
+				check(Buffers.Num() == 1);
+				OutValue = FValue(MoveTemp(Buffers[0]));
+				return true;
+			}
 		}
-		return true;
 	}
 	return false;
-};
+}
 
 FOptionalCacheRecord FHttpDerivedDataBackend::GetCacheRecord(
 	const FStringView Name,
@@ -2692,14 +2926,14 @@ FOptionalCacheRecord FHttpDerivedDataBackend::GetCacheRecord(
 		}
 	}
 
-	if (!CachedDataProbablyExistsBatch(Name, Key, RequiredHeads))
+	if (!CachedDataProbablyExistsBatch<FValueWithId>(Name, Key, RequiredHeads, [](const FValueWithId& Value) { return *WriteToString<16>(Value.GetId()); }))
 	{
 		OutStatus = EStatus::Error;
 		return FOptionalCacheRecord();
 	}
 
 	TArray<FCompressedBuffer> FetchedBuffers;
-	if (!TryGetCachedDataBatch(Name, Key, RequiredGets, FetchedBuffers))
+	if (!TryGetCachedDataBatch<FValueWithId>(Name, Key, RequiredGets, FetchedBuffers, [](const FValueWithId& Value) { return *WriteToString<16>(Value.GetId()); }))
 	{
 		OutStatus = EStatus::Error;
 		return FOptionalCacheRecord();
@@ -2718,13 +2952,15 @@ FOptionalCacheRecord FHttpDerivedDataBackend::GetCacheRecord(
 	return RecordBuilder.Build();
 }
 
+template<typename ValueType, typename ValuePrinterType>
 bool FHttpDerivedDataBackend::TryGetCachedDataBatch(
 	const FStringView Name,
 	const FCacheKey& Key,
-	TConstArrayView<FValueWithId> Values,
-	TArray<FCompressedBuffer>& OutBuffers)
+	TConstArrayView<ValueType> Values,
+	TArray<FCompressedBuffer>& OutBuffers,
+	ValuePrinterType ValuePrinter)
 {
-	for (const FValueWithId& Value : Values)
+	for (const ValueType& Value : Values)
 	{
 		const FIoHash& RawHash = Value.GetRawHash();
 		TStringBuilder<256> CompressedBlobsUri;
@@ -2753,7 +2989,7 @@ bool FHttpDerivedDataBackend::TryGetCachedDataBatch(
 		{
 			UE_LOG(LogDerivedDataCache, Verbose,
 				TEXT("%s: Cache miss with missing value %s with hash %s for %s from '%.*s'"),
-				*GetName(), *WriteToString<16>(Value.GetId()), *WriteToString<48>(RawHash), *WriteToString<96>(Key),
+				*GetName(), ValuePrinter(Value), *WriteToString<48>(RawHash), *WriteToString<96>(Key),
 				Name.Len(), Name.GetData());
 			return false;
 		}
@@ -2761,7 +2997,7 @@ bool FHttpDerivedDataBackend::TryGetCachedDataBatch(
 		{
 			UE_LOG(LogDerivedDataCache, Display,
 				TEXT("%s: Cache miss with corrupted value %s with hash %s for %s from '%.*s'"),
-				*GetName(), *WriteToString<16>(Value.GetId()), *WriteToString<48>(RawHash),
+				*GetName(), ValuePrinter(Value), *WriteToString<48>(RawHash),
 				*WriteToString<96>(Key), Name.Len(), Name.GetData());
 			return false;
 		}
@@ -2773,12 +3009,14 @@ bool FHttpDerivedDataBackend::TryGetCachedDataBatch(
 	return true;
 }
 
+template<typename ValueType, typename ValuePrinterType>
 bool FHttpDerivedDataBackend::CachedDataProbablyExistsBatch(
 	const FStringView Name,
 	const FCacheKey& Key,
-	TConstArrayView<FValueWithId> Values)
+	TConstArrayView<ValueType> Values,
+	ValuePrinterType ValuePrinter)
 {
-	for (const FValueWithId& Value : Values)
+	for (const ValueType& Value : Values)
 	{
 		const FIoHash& RawHash = Value.GetRawHash();
 		TStringBuilder<256> CompressedBlobsUri;
@@ -2804,7 +3042,7 @@ bool FHttpDerivedDataBackend::CachedDataProbablyExistsBatch(
 		{
 			UE_LOG(LogDerivedDataCache, Verbose,
 				TEXT("%s: Cache miss with missing value %s with hash %s for %s from '%.*s'"),
-				*GetName(), *WriteToString<16>(Value.GetId()), *WriteToString<48>(RawHash), *WriteToString<96>(Key),
+				*GetName(), ValuePrinter(Value), *WriteToString<48>(RawHash), *WriteToString<96>(Key),
 				Name.Len(), Name.GetData());
 			return false;
 		}
@@ -3228,6 +3466,66 @@ void FHttpDerivedDataBackend::Get(
 	}
 }
 
+void FHttpDerivedDataBackend::PutValue(
+	const TConstArrayView<FCachePutValueRequest> Requests,
+	IRequestOwner& Owner,
+	FOnCachePutValueComplete&& OnComplete)
+{
+	for (const FCachePutValueRequest& Request : Requests)
+	{
+		COOK_STAT(auto Timer = UsageStats.TimePut());
+		uint64 WriteSize = 0;
+		if (PutCacheValue(Request.Name, Request.Key, Request.Value, Request.Policy, WriteSize))
+		{
+			UE_LOG(LogDerivedDataCache, Verbose, TEXT("%s: Cache put complete for %s from '%s'"),
+				*GetName(), *WriteToString<96>(Request.Key), *Request.Name);
+			TRACE_COUNTER_ADD(HttpDDC_BytesSent, WriteSize);
+			COOK_STAT(if (WriteSize) { Timer.AddHit(WriteSize); });
+			if (OnComplete)
+			{
+				OnComplete({Request.Name, Request.Key, Request.UserData, EStatus::Ok});
+			}
+		}
+		else
+		{
+			if (OnComplete)
+			{
+				OnComplete({Request.Name, Request.Key, Request.UserData, EStatus::Error});
+			}
+		}
+	}
+}
+
+void FHttpDerivedDataBackend::GetValue(
+	const TConstArrayView<FCacheGetValueRequest> Requests,
+	IRequestOwner& Owner,
+	FOnCacheGetValueComplete&& OnComplete)
+{
+	for (const FCacheGetValueRequest& Request : Requests)
+	{
+		COOK_STAT(auto Timer = UsageStats.TimeGet());
+		FValue Value;
+		if (GetCacheValue(Request.Name, Request.Key, Request.Policy, Value))
+		{
+			UE_LOG(LogDerivedDataCache, Verbose, TEXT("%s: Cache hit for %s from '%s'"),
+				*GetName(), *WriteToString<96>(Request.Key), *Request.Name);
+			TRACE_COUNTER_ADD(HttpDDC_BytesReceived, Value.GetData().GetCompressedSize());
+			COOK_STAT(Timer.AddHit(Value.GetData().GetCompressedSize()));
+			if (OnComplete)
+			{
+				OnComplete({ Request.Name, Request.Key, Value, Request.UserData, EStatus::Ok });
+			}
+		}
+		else
+		{
+			if (OnComplete)
+			{
+				OnComplete({ Request.Name, Request.Key, {}, Request.UserData, EStatus::Error });
+			}
+		}
+	}
+}
+
 void FHttpDerivedDataBackend::GetChunks(
 	const TConstArrayView<FCacheGetChunkRequest> Requests,
 	IRequestOwner& Owner,
@@ -3287,7 +3585,7 @@ void FHttpDerivedDataBackend::GetChunks(
 					else
 					{
 						TArray<FCompressedBuffer> ValueBuffers;
-						if (TryGetCachedDataBatch(Request.Name, Request.Key, ::MakeArrayView({ ValueWithId }), ValueBuffers))
+						if (TryGetCachedDataBatch<FValueWithId>(Request.Name, Request.Key, ::MakeArrayView({ ValueWithId }), ValueBuffers, [](const FValueWithId& Value) { return *WriteToString<16>(Value.GetId()); }))
 						{
 							ValueBuffer = ValueBuffers[0];
 							ValueReader.SetSource(ValueBuffer);
diff --git a/Engine/Source/Developer/DerivedDataCache/Private/Tests/HttpDerivedDataBackendTest.cpp b/Engine/Source/Developer/DerivedDataCache/Private/Tests/HttpDerivedDataBackendTest.cpp
--- a/Engine/Source/Developer/DerivedDataCache/Private/Tests/HttpDerivedDataBackendTest.cpp
+++ b/Engine/Source/Developer/DerivedDataCache/Private/Tests/HttpDerivedDataBackendTest.cpp
@@ -222,6 +222,52 @@ protected:
 		return true;
 	}
 
+	bool GetValues(TConstArrayView<UE::DerivedData::FValue> Values, UE::DerivedData::ECachePolicy Policy, TArray<UE::DerivedData::FValue>& OutValues, const char* BucketName = nullptr)
+	{
+		using namespace UE::DerivedData;
+		UE::DerivedData::ICacheStore* TestBackend = GetTestBackend();
+		FCacheBucket TestCacheBucket(BucketName ? BucketName : "AutoTestDummy");
+
+		TArray<FCacheGetValueRequest> Requests;
+		Requests.Reserve(Values.Num());
+
+		for (int32 ValueIndex = 0; ValueIndex < Values.Num(); ++ValueIndex)
+		{
+			const FValue& Value = Values[ValueIndex];
+			FCacheKey Key;
+			Key.Bucket = TestCacheBucket;
+			Key.Hash = Value.GetRawHash();
+			Requests.Add({ {TEXT("FHttpDerivedDataTestBase")}, Key, Policy, static_cast<uint64>(ValueIndex) });
+		}
+
+		struct FGetValueOutput
+		{
+			FValue Value;
+			EStatus Status = EStatus::Error;
+		};
+
+		TArray<TOptional<FGetValueOutput>> GetValueOutputs;
+		GetValueOutputs.SetNum(Values.Num());
+		FRequestOwner RequestOwner(EPriority::Blocking);
+		TestBackend->GetValue(Requests, RequestOwner, [&GetValueOutputs](FCacheGetValueResponse&& Response)
+			{
+				GetValueOutputs[Response.UserData].Emplace(FGetValueOutput{ Response.Value, Response.Status });
+			});
+		RequestOwner.Wait();
+
+		for (int32 ValueIndex = 0; ValueIndex < Values.Num(); ++ValueIndex)
+		{
+			FGetValueOutput& ReceivedOutput = GetValueOutputs[ValueIndex].GetValue();
+			if (ReceivedOutput.Status != EStatus::Ok)
+			{
+				return false;
+			}
+			OutValues.Add(MoveTemp(ReceivedOutput.Value));
+		}
+
+		return true;
+	}
+
 	bool GetRecordChunks(TConstArrayView<UE::DerivedData::FCacheRecord> Records, UE::DerivedData::FCacheRecordPolicy Policy, uint64 Offset, uint64 Size, TArray<FSharedBuffer>& OutChunks)
 	{
 		using namespace UE::DerivedData;
@@ -314,6 +360,33 @@ protected:
 		}
 	}
 
+	void ValidateValues(const TCHAR* Name, TConstArrayView<UE::DerivedData::FValue> ValuesToTest, TConstArrayView<UE::DerivedData::FValue> ReferenceValues, UE::DerivedData::ECachePolicy Policy)
+	{
+		using namespace UE::DerivedData;
+
+		if (!TestEqual(FString::Printf(TEXT("%s::Value quantity"), Name), ValuesToTest.Num(), ReferenceValues.Num()))
+		{
+			return;
+		}
+
+		for (int32 ValueIndex = 0; ValueIndex < ValuesToTest.Num(); ++ValueIndex)
+		{
+			const FValue& ExpectedValue = ReferenceValues[ValueIndex];
+			const FValue& ValueToTest = ValuesToTest[ValueIndex];
+
+			if (EnumHasAnyFlags(Policy, ECachePolicy::SkipData))
+			{
+				TestTrue(FString::Printf(TEXT("%s::Get value[%d] !HasData"), Name, ValueIndex), !ValueToTest.HasData());
+			}
+			else
+			{
+				TestTrue(FString::Printf(TEXT("%s::Get value[%d] HasData"), Name, ValueIndex), ValueToTest.HasData());
+				TestTrue(FString::Printf(TEXT("%s::Get value[%d] equality"), Name, ValueIndex), ExpectedValue == ValueToTest);
+				TestTrue(FString::Printf(TEXT("%s::Get value[%d] data equality"), Name, ValueIndex), FIoHash::HashBuffer(ValueToTest.GetData().GetCompressed()) == FIoHash::HashBuffer(ExpectedValue.GetData().GetCompressed()));
+			}
+		}
+	}
+
 	void ValidateRecordChunks(const TCHAR* Name, TConstArrayView<FSharedBuffer> RecordChunksToTest, TConstArrayView<UE::DerivedData::FCacheRecord> ReferenceRecords, UE::DerivedData::FCacheRecordPolicy Policy, uint64 Offset, uint64 Size)
 	{
 		using namespace UE::DerivedData;
@@ -375,6 +448,22 @@ protected:
 		return ReceivedRecords;
 	}
 
+	TArray<UE::DerivedData::FValue> GetAndValidateValues(const TCHAR* Name, TConstArrayView<UE::DerivedData::FValue> Values, UE::DerivedData::ECachePolicy Policy)
+	{
+		using namespace UE::DerivedData;
+		TArray<FValue> ReceivedValues;
+		bool bGetSuccessful = GetValues(Values, Policy, ReceivedValues);
+		TestTrue(FString::Printf(TEXT("%s::Get status"), Name), bGetSuccessful);
+
+		if (!bGetSuccessful)
+		{
+			return TArray<FValue>();
+		}
+
+		ValidateValues(Name, ReceivedValues, Values, Policy);
+		return ReceivedValues;
+	}
+
 	TArray<FSharedBuffer> GetAndValidateRecordChunks(const TCHAR* Name, TConstArrayView<UE::DerivedData::FCacheRecord> Records, UE::DerivedData::FCacheRecordPolicy Policy, uint64 Offset, uint64 Size)
 	{
 		using namespace UE::DerivedData;
@@ -488,6 +577,52 @@ TArray<UE::DerivedData::FCacheRecord> CreateTestCacheRecords(UE::DerivedData::IC
 	return CacheRecords;
 }
 
+TArray<UE::DerivedData::FValue> CreateTestCacheValues(UE::DerivedData::ICacheStore* InTestBackend, uint32 InNumValues, const char* BucketName = nullptr)
+{
+	using namespace UE::DerivedData;
+	FCacheBucket TestCacheBucket(BucketName ? BucketName : "AutoTestDummy");
+
+	TArray<FCachePutValueRequest> PutValueRequests;
+	PutValueRequests.Reserve(InNumValues);
+
+	TArray<FSharedBuffer> ValueBuffers;
+	for (uint32 ValueIndex = 0; ValueIndex < InNumValues; ++ValueIndex)
+	{
+		TArray<uint8> ValueContents;
+		// Add N zeroed bytes where N corresponds to the value index times 10.
+		const int32 NumBytes = (ValueIndex+1)*10;
+		ValueContents.AddUninitialized(NumBytes);
+		for (int32 ContentIndex = 0; ContentIndex < NumBytes; ++ContentIndex)
+		{
+			ValueContents[ContentIndex] = (uint8)(ValueIndex + ContentIndex + 42); // offset of 42 to keep the contents distinct from record test data
+		}
+		ValueBuffers.Emplace(MakeSharedBufferFromArray(MoveTemp(ValueContents)));
+	}
+
+	uint64 KeyIndex = 0;
+	for (const FSharedBuffer& ValueBuffer : ValueBuffers)
+	{
+		FIoHash ValueHash(FIoHash::HashBuffer(ValueBuffer));
+
+		FCacheKey Key;
+		Key.Bucket = TestCacheBucket;
+		Key.Hash = ValueHash;
+
+		PutValueRequests.Add({ {TEXT("AutoTest")}, Key, FValue::Compress(ValueBuffer), ECachePolicy::Default, KeyIndex++ });
+	}
+
+	TArray<FValue> Values;
+	FRequestOwner Owner(EPriority::Blocking);
+	InTestBackend->PutValue(PutValueRequests, Owner, [&Values, &PutValueRequests](FCachePutValueResponse&& Response)
+		{
+			check(Response.Status == EStatus::Ok);
+			Values.Add(PutValueRequests[Response.UserData].Value);
+		});
+	Owner.Wait();
+
+	return Values;
+}
+
 IMPLEMENT_HTTPDERIVEDDATA_AUTOMATION_TEST(FConcurrentCachedDataProbablyExistsBatch, TEXT(".FConcurrentCachedDataProbablyExistsBatch"), EAutomationTestFlags::EditorContext | EAutomationTestFlags::ProductFilter)
 bool FConcurrentCachedDataProbablyExistsBatch::RunTest(const FString& Parameters)
 {
@@ -594,6 +729,7 @@ bool CacheStore::RunTest(const FString& Parameters)
 #endif // UE_WITH_ZEN
 
 	const uint32 RecordsInBatch = 3;
+	const uint32 ValuesInBatch = RecordsInBatch;
 
 	{
 		TArray<FCacheRecord> PutRecords = CreateTestCacheRecords(TestBackend, RecordsInBatch, 1);
@@ -655,6 +791,12 @@ bool CacheStore::RunTest(const FString& Parameters)
 #endif // UE_WITH_ZEN
 	}
 
+	{
+		TArray<FValue> PutValues = CreateTestCacheValues(TestBackend, ValuesInBatch);
+		GetAndValidateValues(TEXT("SimpleValue"), PutValues, ECachePolicy::Default);
+		GetAndValidateValues(TEXT("SimpleValueSkipData"), PutValues, ECachePolicy::Default | ECachePolicy::SkipData);
+	}
+
 	return true;
 }
 
