CacheStoreTest: Add a test case that catches ValuePolicies being parsed correctly.
#rnx

diff --git a/Engine/Source/Developer/DerivedDataCache/Private/Tests/CacheStoreTest.cpp b/Engine/Source/Developer/DerivedDataCache/Private/Tests/CacheStoreTest.cpp
--- a/Engine/Source/Developer/DerivedDataCache/Private/Tests/CacheStoreTest.cpp
+++ b/Engine/Source/Developer/DerivedDataCache/Private/Tests/CacheStoreTest.cpp
@@ -28,8 +28,9 @@ bool FCacheStoreTest::RunTest(const FString& Parameters)
 	static bool bExpectUpstream = FParse::Param(FCommandLine::Get(), TEXT("CacheStoreTestUpstream"));
 	static bool bExpectWarm = FParse::Param(FCommandLine::Get(), TEXT("CacheStoreTestWarm"));
 
-	// NumKeys = (2 Value vs Record)*(2 SkipData vs Default)*(2 ForceMiss vs Not)*(2 use local)*(2 use remote)*(4 cases per type)
-	constexpr int32 NumKeys = 128;
+	// NumKeys = (2 Value vs Record)*(2 SkipData vs Default)*(2 ForceMiss vs Not)*(2 use local)
+	// *(2 use remote)*(2 UseValue Policy vs not)*(4 cases per type)
+	constexpr int32 NumKeys = 256;
 	constexpr int32 NumValues = 4;
 	TArray<FCachePutRequest> PutRequests;
 	TArray<FCachePutValueRequest> PutValueRequests;
@@ -65,6 +66,7 @@ bool FCacheStoreTest::RunTest(const FString& Parameters)
 		uint32 KeyIndex = 0;
 		bool bGetRequestsData = true;
 		bool bUseValueAPI = false;
+		bool bUseValuePolicy = false;
 		bool bForceMiss = false;
 		bool bUseLocal = true;
 		bool bUseRemote = true;
@@ -90,9 +92,10 @@ bool FCacheStoreTest::RunTest(const FString& Parameters)
 		KeyData.KeyIndex = KeyIndex;
 		KeyData.bGetRequestsData = (KeyIndex & (1 << 1)) == 0;
 		KeyData.bUseValueAPI = (KeyIndex & (1 << 2)) != 0;
-		KeyData.bForceMiss = (KeyIndex & (1 << 3)) == 0;
-		KeyData.bUseLocal = (KeyIndex & (1 << 4)) == 0;
-		KeyData.bUseRemote = (KeyIndex & (1 << 5)) == 0;
+		KeyData.bUseValuePolicy = (KeyIndex & (1 << 3)) != 0;
+		KeyData.bForceMiss = (KeyIndex & (1 << 4)) == 0;
+		KeyData.bUseLocal = (KeyIndex & (1 << 5)) == 0;
+		KeyData.bUseRemote = (KeyIndex & (1 << 6)) == 0;
 		KeyData.bShouldBeHit = !KeyData.bForceMiss && (KeyData.bUseLocal || (KeyData.bUseRemote && bExpectUpstream));
 		ECachePolicy SharedPolicy = KeyData.bUseLocal ? ECachePolicy::Local : ECachePolicy::None;
 		SharedPolicy |= KeyData.bUseRemote ? ECachePolicy::Remote : ECachePolicy::None;
@@ -122,11 +125,32 @@ bool FCacheStoreTest::RunTest(const FString& Parameters)
 				Builder.AddValue(ValueIds[ValueIndex], KeyData.BufferValues[ValueIndex]);
 			}
 
+			FCacheRecordPolicy PutRecordPolicy;
+			FCacheRecordPolicy GetRecordPolicy;
+			if (!KeyData.bUseValuePolicy)
+			{
+				PutRecordPolicy = FCacheRecordPolicy(PutPolicy);
+				GetRecordPolicy = FCacheRecordPolicy(GetPolicy);
+			}
+			else
+			{
+				// Switch the SkipData field in the Record policy so that if the CacheStore ignores the ValuePolicies
+				// it will use the wrong value for SkipData and fail our tests.
+				FCacheRecordPolicyBuilder PutBuilder(PutPolicy ^ ECachePolicy::SkipData);
+				FCacheRecordPolicyBuilder GetBuilder(GetPolicy ^ ECachePolicy::SkipData);
+				for (int32 ValueIndex = 0; ValueIndex < NumValues; ++ValueIndex)
+				{
+					PutBuilder.AddValuePolicy(ValueIds[ValueIndex], PutPolicy);
+					GetBuilder.AddValuePolicy(ValueIds[ValueIndex], GetPolicy);
+				}
+				PutRecordPolicy = PutBuilder.Build();
+				GetRecordPolicy = GetBuilder.Build();
+			}
 			if (!KeyData.bForceMiss)
 			{
-				PutRequests.Add(FCachePutRequest{ Name, Builder.Build(), FCacheRecordPolicy(PutPolicy), reinterpret_cast<uint64>(&KeyUserData) });
+				PutRequests.Add(FCachePutRequest{ Name, Builder.Build(), PutRecordPolicy, reinterpret_cast<uint64>(&KeyUserData) });
 			}
-			GetRequests.Add(FCacheGetRequest{ Name, Key, FCacheRecordPolicy(GetPolicy), reinterpret_cast<uint64>(&KeyUserData) });
+			GetRequests.Add(FCacheGetRequest{ Name, Key, GetRecordPolicy, reinterpret_cast<uint64>(&KeyUserData) });
 			for (uint32 ValueIndex = 0; ValueIndex < NumValues; ++ValueIndex)
 			{
 				FUserData& ValueUserData = KeyData.ValueUserData[ValueIndex];
