DDC: Changed cache policy string conversion to use optimal encoding

Parsing uses UTF-8 to avoid conversion of UTF-8 and ANSI strings, and conversion to a string uses ANSI to avoid expensive conversion from UTF-16.

Includes a change to mask out record-only flags from values in FCacheRecordPolicy::Transform.

#rnx

diff --git a/Engine/Source/Developer/DerivedDataCache/Private/DerivedDataCachePolicy.cpp b/Engine/Source/Developer/DerivedDataCache/Private/DerivedDataCachePolicy.cpp
--- a/Engine/Source/Developer/DerivedDataCache/Private/DerivedDataCachePolicy.cpp
+++ b/Engine/Source/Developer/DerivedDataCache/Private/DerivedDataCachePolicy.cpp
@@ -22,28 +22,28 @@ constexpr ANSICHAR CachePolicyDelimiter = ',';
 struct FCachePolicyToText
 {
 	ECachePolicy Policy;
-	FStringView Text;
+	FAnsiStringView Text;
 };
 
 const FCachePolicyToText CachePolicyToText[]
 {
 	// Flags with multiple bits are ordered by bit count to minimize token count in the text format.
-	{ ECachePolicy::Default,       TEXTVIEW("Default") },
-	{ ECachePolicy::Remote,        TEXTVIEW("Remote") },
-	{ ECachePolicy::Local,         TEXTVIEW("Local") },
-	{ ECachePolicy::Store,         TEXTVIEW("Store") },
-	{ ECachePolicy::Query,         TEXTVIEW("Query") },
+	{ ECachePolicy::Default,       ANSITEXTVIEW("Default") },
+	{ ECachePolicy::Remote,        ANSITEXTVIEW("Remote") },
+	{ ECachePolicy::Local,         ANSITEXTVIEW("Local") },
+	{ ECachePolicy::Store,         ANSITEXTVIEW("Store") },
+	{ ECachePolicy::Query,         ANSITEXTVIEW("Query") },
 	// Flags with only one bit can be in any order. Match the order in ECachePolicy.
-	{ ECachePolicy::QueryLocal,    TEXTVIEW("QueryLocal") },
-	{ ECachePolicy::QueryRemote,   TEXTVIEW("QueryRemote") },
-	{ ECachePolicy::StoreLocal,    TEXTVIEW("StoreLocal") },
-	{ ECachePolicy::StoreRemote,   TEXTVIEW("StoreRemote") },
-	{ ECachePolicy::SkipMeta,      TEXTVIEW("SkipMeta") },
-	{ ECachePolicy::SkipData,      TEXTVIEW("SkipData") },
-	{ ECachePolicy::PartialRecord, TEXTVIEW("PartialRecord") },
-	{ ECachePolicy::KeepAlive,     TEXTVIEW("KeepAlive") },
+	{ ECachePolicy::QueryLocal,    ANSITEXTVIEW("QueryLocal") },
+	{ ECachePolicy::QueryRemote,   ANSITEXTVIEW("QueryRemote") },
+	{ ECachePolicy::StoreLocal,    ANSITEXTVIEW("StoreLocal") },
+	{ ECachePolicy::StoreRemote,   ANSITEXTVIEW("StoreRemote") },
+	{ ECachePolicy::SkipMeta,      ANSITEXTVIEW("SkipMeta") },
+	{ ECachePolicy::SkipData,      ANSITEXTVIEW("SkipData") },
+	{ ECachePolicy::PartialRecord, ANSITEXTVIEW("PartialRecord") },
+	{ ECachePolicy::KeepAlive,     ANSITEXTVIEW("KeepAlive") },
 	// None must be last because it matches every policy.
-	{ ECachePolicy::None,          TEXTVIEW("None") },
+	{ ECachePolicy::None,          ANSITEXTVIEW("None") },
 };
 
 constexpr ECachePolicy CachePolicyKnownFlags
@@ -79,8 +79,8 @@ ECachePolicy ParseCachePolicy(const TStringView<CharType> Text)
 {
 	checkf(!Text.IsEmpty(), TEXT("ParseCachePolicy requires a non-empty string."));
 	ECachePolicy Policy = ECachePolicy::None;
-	String::ParseTokens(StringCast<TCHAR, 128>(Text.GetData(), Text.Len()), CachePolicyDelimiter,
-		[&Policy, Index = int32(0)](FStringView Token) mutable
+	String::ParseTokens(StringCast<UTF8CHAR, 128>(Text.GetData(), Text.Len()), UTF8CHAR(CachePolicyDelimiter),
+		[&Policy, Index = int32(0)](FUtf8StringView Token) mutable
 		{
 			const int32 EndIndex = Index;
 			for (; Index < UE_ARRAY_COUNT(CachePolicyToText); ++Index)
@@ -178,7 +178,7 @@ FCacheRecordPolicy FCacheRecordPolicy::Transform(const TFunctionRef<ECachePolicy
 	FCacheRecordPolicyBuilder Builder(Op(GetBasePolicy()));
 	for (const FCacheValuePolicy& Value : GetValuePolicies())
 	{
-		Builder.AddValuePolicy({Value.Id, Op(Value.Policy)});
+		Builder.AddValuePolicy({Value.Id, Op(Value.Policy) & FCacheValuePolicy::PolicyMask});
 	}
 	return Builder.Build();
 }
