DDC: Fixed GetChunks in the pak file cache to fail gracefully when it receives a value request

#rnx

diff --git a/Engine/Source/Developer/DerivedDataCache/Private/PakFileDerivedDataBackend.cpp b/Engine/Source/Developer/DerivedDataCache/Private/PakFileDerivedDataBackend.cpp
--- a/Engine/Source/Developer/DerivedDataCache/Private/PakFileDerivedDataBackend.cpp
+++ b/Engine/Source/Developer/DerivedDataCache/Private/PakFileDerivedDataBackend.cpp
@@ -717,25 +717,45 @@ void FPakFileDerivedDataBackend::GetChunks(
 	TArray<FCacheGetChunkRequest, TInlineAllocator<16>> SortedRequests(Requests);
 	SortedRequests.StableSort(TChunkLess());
 
+	bool bHasValue = false;
+	FValueWithId Value;
+	FValueId ValueId;
+	FCacheKey ValueKey;
+	FCompressedBufferReader ValueReader;
+	EStatus ValueStatus = EStatus::Error;
 	FOptionalCacheRecord Record;
-	FCompressedBufferReader Reader;
 	for (const FCacheGetChunkRequest& Request : SortedRequests)
 	{
 		const bool bExistsOnly = EnumHasAnyFlags(Request.Policy, ECachePolicy::SkipData);
 		TRACE_CPUPROFILER_EVENT_SCOPE(PakFileDDC_Get);
 		COOK_STAT(auto Timer = bExistsOnly ? UsageStats.TimeProbablyExists() : UsageStats.TimeGet());
-		if (!Record || Record.Get().GetKey() != Request.Key)
+		if (!(bHasValue && ValueKey == Request.Key && ValueId == Request.Id) || ValueReader.HasSource() < !bExistsOnly)
 		{
-			FCacheRecordPolicyBuilder PolicyBuilder(ECachePolicy::None);
-			PolicyBuilder.AddValuePolicy(Request.Id, Request.Policy);
-			Record = GetCacheRecordOnly(Request.Name, Request.Key, PolicyBuilder.Build());
-		}
-		if (Record)
-		{
-			EStatus ValueStatus = EStatus::Ok;
-			FValueWithId Value = Record.Get().GetValue(Request.Id);
-			GetCacheContent(Request.Name, Request.Key, Request.Policy, Value, ValueStatus);
-			if (Value)
+			ValueStatus = EStatus::Error;
+			ValueReader.ResetSource();
+			ValueKey = {};
+			ValueId.Reset();
+			Value.Reset();
+			bHasValue = false;
+			if (Request.Id.IsValid())
+			{
+				if (!(Record && Record.Get().GetKey() == Request.Key))
+				{
+					FCacheRecordPolicyBuilder PolicyBuilder(ECachePolicy::None);
+					PolicyBuilder.AddValuePolicy(Request.Id, Request.Policy);
+					Record = GetCacheRecordOnly(Request.Name, Request.Key, PolicyBuilder.Build());
+				}
+				if (Record)
+				{
+					const FValueWithId& ValueWithId = Record.Get().GetValue(Request.Id);
+					bHasValue = ValueWithId.IsValid();
+					Value = ValueWithId;
+					ValueId = Request.Id;
+					ValueKey = Request.Key;
+					GetCacheContent(Request.Name, Request.Key, Request.Policy, Value, ValueStatus);
+				}
+			}
+			if (bHasValue)
 			{
 				const uint64 RawOffset = FMath::Min(Value.GetRawSize(), Request.RawOffset);
 				const uint64 RawSize = FMath::Min(Value.GetRawSize() - RawOffset, Request.RawSize);
@@ -747,8 +767,8 @@ void FPakFileDerivedDataBackend::GetChunks(
 					FSharedBuffer Buffer;
 					if (Value.HasData() && !bExistsOnly)
 					{
-						FCompressedBufferReaderSourceScope Source(Reader, Value.GetData());
-						Buffer = Reader.Decompress(RawOffset, RawSize);
+						FCompressedBufferReaderSourceScope Source(ValueReader, Value.GetData());
+						Buffer = ValueReader.Decompress(RawOffset, RawSize);
 					}
 					OnComplete({Request.Name, Request.Key, Request.Id, Request.RawOffset,
 						RawSize, Value.GetRawHash(), MoveTemp(Buffer), Request.UserData, ValueStatus});
