DDC: Added GetChunks support to the new cache hierarchy

Includes policy fixes for other functions.

#rnx

diff --git a/Engine/Source/Developer/DerivedDataCache/Private/DerivedDataCacheStoreHierarchy.cpp b/Engine/Source/Developer/DerivedDataCache/Private/DerivedDataCacheStoreHierarchy.cpp
--- a/Engine/Source/Developer/DerivedDataCache/Private/DerivedDataCacheStoreHierarchy.cpp
+++ b/Engine/Source/Developer/DerivedDataCache/Private/DerivedDataCacheStoreHierarchy.cpp
@@ -99,6 +99,7 @@ private:
 	template <typename Params> class TPutBatch;
 	template <typename Params> class TGetBatch;
 
+	class FGetChunksBatch;
 	class FLegacyDeleteBatch;
 
 	enum class ECacheStoreNodeFlags : uint32;
@@ -385,17 +386,17 @@ private:
 
 template <typename Params>
 void FCacheStoreHierarchy::TPutBatch<Params>::Begin(
-	const FCacheStoreHierarchy& Hierarchy,
-	const TConstArrayView<FPutRequest> Requests,
-	IRequestOwner& Owner,
-	FOnPutComplete&& OnComplete)
+	const FCacheStoreHierarchy& InHierarchy,
+	const TConstArrayView<FPutRequest> InRequests,
+	IRequestOwner& InOwner,
+	FOnPutComplete&& InOnComplete)
 {
-	if (Requests.IsEmpty() || !EnumHasAnyFlags(Hierarchy.CombinedNodeFlags, ECacheStoreNodeFlags::HasStoreNode))
+	if (InRequests.IsEmpty() || !EnumHasAnyFlags(InHierarchy.CombinedNodeFlags, ECacheStoreNodeFlags::HasStoreNode))
 	{
-		return CompleteWithStatus(Requests, OnComplete, EStatus::Error);
+		return CompleteWithStatus(InRequests, InOnComplete, EStatus::Error);
 	}
 
-	TRefCountPtr<TPutBatch> State = new TPutBatch(Hierarchy, Requests, Owner, MoveTemp(OnComplete));
+	TRefCountPtr<TPutBatch> State = new TPutBatch(InHierarchy, InRequests, InOwner, MoveTemp(InOnComplete));
 	State->DispatchRequests();
 }
 
@@ -620,7 +621,7 @@ void FCacheStoreHierarchy::TGetBatch<Params>::Begin(
 	IRequestOwner& InOwner,
 	FOnGetComplete&& InOnComplete)
 {
-	if (InRequests.IsEmpty() || !EnumHasAnyFlags(InHierarchy.CombinedNodeFlags, ECacheStoreNodeFlags::HasStoreNode))
+	if (InRequests.IsEmpty() || !EnumHasAnyFlags(InHierarchy.CombinedNodeFlags, ECacheStoreNodeFlags::HasQueryNode))
 	{
 		return CompleteWithStatus(InRequests, InOnComplete, EStatus::Error);
 	}
@@ -667,7 +668,7 @@ void FCacheStoreHierarchy::TGetBatch<Params>::DispatchRequests()
 				{
 					if (EnumHasAnyFlags(Policy, ECachePolicy::SkipData) && CanStoreIfOk(Policy, Node.NodeFlags))
 					{
-						NodeRequests.Add({Request.Name, Request.Key, RemovePolicy(Request.Policy, ECachePolicy::SkipData), StateIndex});
+						NodeRequests.Add({Request.Name, Request.Key, RemovePolicy(Request.Policy, ECachePolicy::SkipData | ECachePolicy::SkipMeta), StateIndex});
 					}
 					else
 					{
@@ -840,7 +841,7 @@ FCacheGetRequest FCacheStoreHierarchy::FCacheRecordBatchParams::MakeGetRequest(
 }
 
 void FCacheStoreHierarchy::Put(
-	TConstArrayView<FCachePutRequest> Requests,
+	const TConstArrayView<FCachePutRequest> Requests,
 	IRequestOwner& Owner,
 	FOnCachePutComplete&& OnComplete)
 {
@@ -848,7 +849,7 @@ void FCacheStoreHierarchy::Put(
 }
 
 void FCacheStoreHierarchy::Get(
-	TConstArrayView<FCacheGetRequest> Requests,
+	const TConstArrayView<FCacheGetRequest> Requests,
 	IRequestOwner& Owner,
 	FOnCacheGetComplete&& OnComplete)
 {
@@ -920,12 +921,135 @@ void FCacheStoreHierarchy::GetValue(
 
 ///////////////////////////////////////////////////////////////////////////////////////////////////////////////////////
 
+class FCacheStoreHierarchy::FGetChunksBatch : public FBatchBase
+{
+public:
+	static void Begin(
+		const FCacheStoreHierarchy& Hierarchy,
+		TConstArrayView<FCacheGetChunkRequest> Requests,
+		IRequestOwner& Owner,
+		FOnCacheGetChunkComplete&& OnComplete);
+
+private:
+	FGetChunksBatch(
+		const FCacheStoreHierarchy& InHierarchy,
+		const TConstArrayView<FCacheGetChunkRequest> InRequests,
+		IRequestOwner& InOwner,
+		FOnCacheGetChunkComplete&& InOnComplete)
+		: Hierarchy(InHierarchy)
+		, OnComplete(MoveTemp(InOnComplete))
+		, Owner(InOwner)
+	{
+		States.Reserve(InRequests.Num());
+		for (const FCacheGetChunkRequest& Request : InRequests)
+		{
+			States.Add({Request});
+		}
+	}
+
+	void DispatchRequests();
+	void CompleteRequest(FCacheGetChunkResponse&& Response);
+
+	struct FState
+	{
+		FCacheGetChunkRequest Request;
+		EStatus Status = EStatus::Error;
+	};
+
+	const FCacheStoreHierarchy& Hierarchy;
+	FOnCacheGetChunkComplete OnComplete;
+	TArray<FState, TInlineAllocator<8>> States;
+
+	IRequestOwner& Owner;
+	FCounterEvent RemainingRequestCount;
+	int32 NodeIndex = 0;
+};
+
+void FCacheStoreHierarchy::FGetChunksBatch::Begin(
+	const FCacheStoreHierarchy& InHierarchy,
+	const TConstArrayView<FCacheGetChunkRequest> InRequests,
+	IRequestOwner& InOwner,
+	FOnCacheGetChunkComplete&& InOnComplete)
+{
+	if (InRequests.IsEmpty() || !EnumHasAnyFlags(InHierarchy.CombinedNodeFlags, ECacheStoreNodeFlags::HasQueryNode))
+	{
+		return CompleteWithStatus(InRequests, InOnComplete, EStatus::Error);
+	}
+
+	TRefCountPtr<FGetChunksBatch> State = new FGetChunksBatch(InHierarchy, InRequests, InOwner, MoveTemp(InOnComplete));
+	State->DispatchRequests();
+}
+
+void FCacheStoreHierarchy::FGetChunksBatch::DispatchRequests()
+{
+	FReadScopeLock Lock(Hierarchy.NodesLock);
+
+	TArray<FCacheGetChunkRequest, TInlineAllocator<8>> NodeRequests;
+	NodeRequests.Reserve(States.Num());
+
+	for (const int32 NodeCount = Hierarchy.Nodes.Num(); NodeIndex < NodeCount && !Owner.IsCanceled(); ++NodeIndex)
+	{
+		const FCacheStoreNode& Node = Hierarchy.Nodes[NodeIndex];
+
+		uint64 StateIndex = 0;
+		for (const FState& State : States)
+		{
+			const FCacheGetChunkRequest& Request = State.Request;
+			if (State.Status == EStatus::Error && CanQuery(Request.Policy, Node.CacheFlags))
+			{
+				NodeRequests.Add_GetRef(Request).UserData = StateIndex;
+			}
+			++StateIndex;
+		}
+
+		if (const int32 NodeRequestsCount = NodeRequests.Num())
+		{
+			RemainingRequestCount.Reset(NodeRequestsCount + 1);
+			Node.Cache->GetChunks(NodeRequests, Owner,
+				[State = TRefCountPtr(this)](FCacheGetChunkResponse&& Response)
+				{
+					State->CompleteRequest(MoveTemp(Response));
+				});
+			NodeRequests.Reset();
+			if (!RemainingRequestCount.Signal())
+			{
+				return;
+			}
+		}
+	}
+
+	for (const FState& State : States)
+	{
+		if (State.Status != EStatus::Ok)
+		{
+			OnComplete(State.Request.MakeResponse(Owner.IsCanceled() ? EStatus::Canceled : EStatus::Error));
+		}
+	}
+}
+
+void FCacheStoreHierarchy::FGetChunksBatch::CompleteRequest(FCacheGetChunkResponse&& Response)
+{
+	FState& State = States[int32(Response.UserData)];
+	if (Response.Status == EStatus::Ok)
+	{
+		check(State.Status == EStatus::Error);
+		Response.UserData = State.Request.UserData;
+		OnComplete(MoveTemp(Response));
+	}
+	State.Status = Response.Status;
+
+	if (RemainingRequestCount.Signal())
+	{
+		DispatchRequests();
+	}
+}
+
 void FCacheStoreHierarchy::GetChunks(
 	TConstArrayView<FCacheGetChunkRequest> Requests,
 	IRequestOwner& Owner,
 	FOnCacheGetChunkComplete&& OnComplete)
 {
-	unimplemented();
+	FGetChunksBatch::Begin(*this, Requests, Owner, MoveTemp(OnComplete));
 }
 
 ///////////////////////////////////////////////////////////////////////////////////////////////////////////////////////
@@ -1035,17 +1159,17 @@ private:
 };
 
 void FCacheStoreHierarchy::FLegacyDeleteBatch::Begin(
-	const FCacheStoreHierarchy& Hierarchy,
-	const TConstArrayView<FLegacyCacheDeleteRequest> Requests,
-	IRequestOwner& Owner,
-	FOnLegacyCacheDeleteComplete&& OnComplete)
+	const FCacheStoreHierarchy& InHierarchy,
+	const TConstArrayView<FLegacyCacheDeleteRequest> InRequests,
+	IRequestOwner& InOwner,
+	FOnLegacyCacheDeleteComplete&& InOnComplete)
 {
-	if (Requests.IsEmpty() || !EnumHasAnyFlags(Hierarchy.CombinedNodeFlags, ECacheStoreNodeFlags::HasStoreNode))
+	if (InRequests.IsEmpty() || !EnumHasAnyFlags(InHierarchy.CombinedNodeFlags, ECacheStoreNodeFlags::HasStoreNode))
 	{
-		return CompleteWithStatus(Requests, OnComplete, EStatus::Error);
+		return CompleteWithStatus(InRequests, InOnComplete, EStatus::Error);
 	}
 
-	TRefCountPtr<FLegacyDeleteBatch> State = new FLegacyDeleteBatch(Hierarchy, Requests, Owner, MoveTemp(OnComplete));
+	TRefCountPtr<FLegacyDeleteBatch> State = new FLegacyDeleteBatch(InHierarchy, InRequests, InOwner, MoveTemp(InOnComplete));
 	State->DispatchRequests();
 }
 
