GetChunks value implementation for Jupiter.

diff --git a/Engine/Source/Developer/DerivedDataCache/Private/HttpDerivedDataBackend.cpp b/Engine/Source/Developer/DerivedDataCache/Private/HttpDerivedDataBackend.cpp
--- a/Engine/Source/Developer/DerivedDataCache/Private/HttpDerivedDataBackend.cpp
+++ b/Engine/Source/Developer/DerivedDataCache/Private/HttpDerivedDataBackend.cpp
@@ -3600,9 +3600,26 @@ void FHttpDerivedDataBackend::GetChunks(
 			}
 			else
 			{
-				UE_LOG(LogDerivedDataCache, Verbose, TEXT("%s: Cache miss with unimplemented value support for %s from '%s'"),
-					*GetName(), *WriteToString<96>(Request.Key, '/', Request.Id), *Request.Name);
-				checkf(false, TEXT("Value API not implemented for HTTP backend"));
+				ValueKey = Request.Key;
+				bHasValue = GetCacheValueOnly(Request.Name, Request.Key, Request.Policy, Value);
+				if (IsValueDataReady(Value, Request.Policy))
+				{
+					ValueReader.SetSource(Value.GetData());
+				}
+				else
+				{
+					TArray<FCompressedBuffer> ValueBuffers;
+					if (TryGetCachedDataBatch<FValue>(Request.Name, Request.Key, { Value }, ValueBuffers, [](const FValue& Value) { return TEXT("Default"); }))
+					{
+						ValueBuffer = ValueBuffers[0];
+						ValueReader.SetSource(ValueBuffer);
+					}
+					else
+					{
+						ValueBuffer.Reset();
+						ValueReader.ResetSource();
+					}
+				}
 			}
 		}
 		if (bHasValue)
