DDC: Fixed a missing decrement of the async completion counter

Cancellation of requets in the async put wrapper was not decrementing the async completion counter, which causes a hang on exit if any cancellations have occurred. Any cook that builds a lot of meshes changes the timing of other operations enough that they are often canceled.

#rnx

diff --git a/Engine/Source/Developer/DerivedDataCache/Private/DerivedDataCacheStoreAsync.cpp b/Engine/Source/Developer/DerivedDataCache/Private/DerivedDataCacheStoreAsync.cpp
--- a/Engine/Source/Developer/DerivedDataCache/Private/DerivedDataCacheStoreAsync.cpp
+++ b/Engine/Source/Developer/DerivedDataCache/Private/DerivedDataCacheStoreAsync.cpp
@@ -737,6 +737,7 @@ void FDerivedDataBackendAsyncPutWrapper::Execute(
 			else
 			{
 				CompleteWithStatus(Requests, MoveTemp(OnComplete), EStatus::Canceled);
+				FDerivedDataBackend::Get().AddToAsyncCompletionCounter(-Requests.Num());
 			}
 		});
 	Request->Start(Owner.GetPriority());
