DDC: Switched to the async cache hierarchy and fully implemented the verify wrapper

The async cache hierarchy:
- Is required to add compression to the legacy cache by forwarding LegacyPut/LegacyGet to PutValue/GetValue.
- Is always present in the graph, unlike the previous cache hierarchy, which will allow significant simplification of leaf cache stores.
- Allows for the leaf cache store nodes to operate asynchronously without blocking a worker thread like the previous cache hierarchy.
- Shifts from controlling cache behavior by speed class to controlling cache behavior by local/remote classification, which is a critical distinction for a cache like Zen that can identify as both local and remote.
- Respects the local/remote and query/store flags in the cache policy.
- Does not fully implement the partial record cache policy at this time.
- Does not propagate records or values in GetChunks, which will be added in a future release.

The verify wrapper was previously missing an implementation for the new cache interface. This version is more efficient than the previous because requests through the new cache interface can compare data without loading it from storage, and load it only when a mismatch has been detected, to dump it to disk.

#rnx

diff --git a/Engine/Source/Developer/DerivedDataCache/Private/DerivedDataBackendThrottleWrapper.h b/Engine/Source/Developer/DerivedDataCache/Private/DerivedDataBackendThrottleWrapper.h
deleted file mode 100644
diff --git a/Engine/Source/Developer/DerivedDataCache/Private/DerivedDataBackendVerifyWrapper.h b/Engine/Source/Developer/DerivedDataCache/Private/DerivedDataBackendVerifyWrapper.h
deleted file mode 100644
diff --git a/Engine/Source/Developer/DerivedDataCache/Private/DerivedDataBackends.cpp b/Engine/Source/Developer/DerivedDataCache/Private/DerivedDataBackends.cpp
--- a/Engine/Source/Developer/DerivedDataCache/Private/DerivedDataBackends.cpp
+++ b/Engine/Source/Developer/DerivedDataCache/Private/DerivedDataBackends.cpp
@@ -3,15 +3,13 @@
 #include "CoreTypes.h"
 #include "Containers/StringView.h"
 #include "DerivedDataBackendInterface.h"
-#include "DerivedDataBackendThrottleWrapper.h"
-#include "DerivedDataBackendVerifyWrapper.h"
+#include "DerivedDataCacheStore.h"
 #include "DerivedDataCacheUsageStats.h"
 #include "FileBackedDerivedDataBackend.h"
 #include "HAL/FileManager.h"
 #include "HAL/IConsoleManager.h"
 #include "HAL/PlatformFileManager.h"
 #include "HAL/ThreadSafeCounter.h"
-#include "HierarchicalDerivedDataBackend.h"
 #include "Internationalization/FastDecimalFormat.h"
 #include "Math/BasicMathExpressionEvaluator.h"
 #include "Math/UnitConversion.h"
@@ -33,15 +31,19 @@ DEFINE_LOG_CATEGORY(LogDerivedDataCache);
 #define MAX_BACKEND_KEY_LENGTH (120)
 #define LOCTEXT_NAMESPACE "DerivedDataBackendGraph"
 
-static TAutoConsoleVariable<FString> GDumpDDCName(
-	TEXT("DDC.Graph"), TEXT(""),
-	TEXT("Name of custom DDC Graph to use."),
-	ECVF_Default);
+static TAutoConsoleVariable<FString> GDerivedDataCacheGraphName(
+	TEXT("DDC.Graph"),
+	TEXT("Default"),
+	TEXT("Name of the graph to use for the Derived Data Cache."),
+	ECVF_ReadOnly);
 
-namespace UE::DerivedData::CacheStore::AsyncPut
+namespace UE::DerivedData
 {
-FDerivedDataBackendInterface* CreateAsyncPutDerivedDataBackend(FDerivedDataBackendInterface* InnerBackend, bool bCacheInFlightPuts);
-} // UE::DerivedData::CacheStore::AsyncPut
+ILegacyCacheStore* CreateCacheStoreAsync(ILegacyCacheStore* InnerBackend, ECacheStoreFlags InnerFlags, bool bCacheInFlightPuts);
+ILegacyCacheStore* CreateCacheStoreHierarchy(ICacheStoreOwner*& OutOwner);
+ILegacyCacheStore* CreateCacheStoreThrottle(ILegacyCacheStore* InnerCache, uint32 LatencyMS, uint32 MaxBytesPerSecond);
+ILegacyCacheStore* CreateCacheStoreVerify(ILegacyCacheStore* InnerCache, bool bPutOnError);
+} // UE::DerivedData
 
 namespace UE::DerivedData::CacheStore::FileSystem
 {
@@ -82,100 +84,122 @@ namespace UE::DerivedData::CacheStore::ZenCache
 FDerivedDataBackendInterface* CreateZenDerivedDataBackend(const TCHAR* NodeName, const TCHAR* ServiceUrl, const TCHAR* Namespace);
 } // UE::DerivedData::CacheStore::ZenCache
 
-namespace UE::DerivedData::CacheStore::Graph
+namespace UE::DerivedData
 {
 
 /**
-  * This class is used to create a singleton that represents the derived data cache hierarchy and all of the wrappers necessary
-  * ideally this would be data driven and the backends would be plugins...
-**/
+ * This class is used to create a singleton that represents the derived data cache hierarchy and all of the wrappers necessary
+ * ideally this would be data driven and the backends would be plugins...
+ */
 class FDerivedDataBackendGraph : public FDerivedDataBackend
 {
 public:
+	using FParsedNode = TPair<ILegacyCacheStore*, ECacheStoreFlags>;
+	using FParsedNodeMap = TMap<FString, FParsedNode>;
+
 	/**
-	* constructor, builds the cache tree
-	**/
+	 * constructor, builds the cache tree
+	 */
 	FDerivedDataBackendGraph()
-		: RootCache(NULL)
-		, BootCache(NULL)
-		, WritePakCache(NULL)
-		, AsyncPutWrapper(NULL)
-		, HierarchicalWrapper(NULL)
+		: RootCache(nullptr)
+		, BootCache(nullptr)
+		, WritePakCache(nullptr)
+		, AsyncNode(nullptr)
+		, Hierarchy(nullptr)
 		, bUsingSharedDDC(false)
 		, bIsShuttingDown(false)
 		, MountPakCommand(
-		TEXT( "DDC.MountPak" ),
-		*LOCTEXT("CommandText_DDCMountPak", "Mounts read-only pak file").ToString(),
-		FConsoleCommandWithArgsDelegate::CreateRaw( this, &FDerivedDataBackendGraph::MountPakCommandHandler ) )
+			TEXT("DDC.MountPak"),
+			*LOCTEXT("CommandText_DDCMountPak", "Mounts read-only pak file").ToString(),
+			FConsoleCommandWithArgsDelegate::CreateRaw(this, &FDerivedDataBackendGraph::MountPakCommandHandler))
 		, UnountPakCommand(
-		TEXT( "DDC.UnmountPak" ),
-		*LOCTEXT("CommandText_DDCUnmountPak", "Unmounts read-only pak file").ToString(),
-		FConsoleCommandWithArgsDelegate::CreateRaw( this, &FDerivedDataBackendGraph::UnmountPakCommandHandler ) )
+			TEXT("DDC.UnmountPak"),
+			*LOCTEXT("CommandText_DDCUnmountPak", "Unmounts read-only pak file").ToString(),
+			FConsoleCommandWithArgsDelegate::CreateRaw(this, &FDerivedDataBackendGraph::UnmountPakCommandHandler))
 	{
 		check(IsInGameThread()); // we pretty much need this to be initialized from the main thread...it uses GConfig, etc
 		check(GConfig && GConfig->IsReadyForUse());
-		RootCache = NULL;
-		TMap<FString, FDerivedDataBackendInterface*> ParsedNodes;		
+		RootCache = nullptr;
+		FParsedNodeMap ParsedNodes;
+
+		FParsedNode RootNode(nullptr, ECacheStoreFlags::None);
 
 		// Create the graph using ini settings. The string "default" forwards creation to use the default graph.
 
-		FString CustomDDCName;
-		FParse::Value(FCommandLine::Get(), TEXT("-DDC="), CustomDDCName);
-		if (CustomDDCName.IsEmpty() || FCString::Stricmp(*GraphName, TEXT("default")) == 0)
+		if (!FParse::Value(FCommandLine::Get(), TEXT("-DDC="), GraphName))
 		{
-			CustomDDCName = GDumpDDCName.GetValueOnGameThread();
+			GraphName = GDerivedDataCacheGraphName.GetValueOnGameThread();
 		}
 
-		if (!CustomDDCName.IsEmpty())
+		if (GraphName == TEXT("None"))
+		{
+			UE_LOG(LogDerivedDataCache, Display, TEXT("Requested cache graph of 'None'. Every cache operation will fail."));
+		}
+		else
 		{
-			GraphName = CustomDDCName;
-			if (GraphName == TEXT("None"))
+			if (!GraphName.IsEmpty() && GraphName != TEXT("Default"))
 			{
-				UE_LOG(LogDerivedDataCache, Display, TEXT("DDC backend graph of 'None' specified, graph will be invalid for use."));
-				return;
-			}
+				RootNode = ParseNode(TEXT("Root"), GEngineIni, *GraphName, ParsedNodes);
 
-			RootCache = ParseNode(TEXT("Root"), GEngineIni, *GraphName, ParsedNodes);
+				if (!RootNode.Key)
+				{
+					// Destroy any cache stores that have been created.
+					ParsedNodes.Empty();
+					DestroyCreatedBackends();
+					UE_LOG(LogDerivedDataCache, Warning,
+						TEXT("Unable to create cache graph using the requested graph settings (%s). "),
+						TEXT("Reverting to the default graph."), *GraphName);
+				}
+			}
 
-			if( RootCache == NULL )
+			if (!RootNode.Key)
 			{
-				// Remove references to any backend instances that might have been created
-				ParsedNodes.Empty();
-				DestroyCreatedBackends();
-				UE_LOG( LogDerivedDataCache, Warning, TEXT("FDerivedDataBackendGraph: Unable to create backend graph using the specified graph settings (\"%s\"). Reverting to the default graph."), *GraphName );
+				// Try to use the default graph.
+				GraphName = FApp::IsEngineInstalled() ? TEXT("InstalledDerivedDataBackendGraph") : TEXT("DerivedDataBackendGraph");
+				FString Entry;
+				if (!GConfig->GetString(*GraphName, TEXT("Root"), Entry, GEngineIni) || !Entry.Len())
+				{
+					UE_LOG(LogDerivedDataCache, Fatal,
+						TEXT("Unable to create cache graph using the default graph settings (%s) ini=%s."),
+						*GraphName, *GEngineIni);
+				}
+				RootNode = ParseNode(TEXT("Root"), GEngineIni, *GraphName, ParsedNodes);
+				if (!RootNode.Key)
+				{
+					UE_LOG(LogDerivedDataCache, Fatal,
+						TEXT("Unable to create cache graph using the default graph settings (%s) ini=%s. ")
+						TEXT("At least one cache store in the graph must be available."),
+						*GraphName, *GEngineIni);
+				}
 			}
 		}
 
-		if( !RootCache )
+		// Hierarchy must exist in the graph.
+		if (!Hierarchy)
 		{
-			// Use default graph
-			GraphName = FApp::IsEngineInstalled() ? TEXT("InstalledDerivedDataBackendGraph") : TEXT("DerivedDataBackendGraph");
-			FString Entry;
-			if( !GConfig->GetString( *GraphName, TEXT("Root"), Entry, GEngineIni ) || !Entry.Len())
+			ILegacyCacheStore* HierarchyStore = CreateCacheStoreHierarchy(Hierarchy);
+			if (RootNode.Key)
 			{
-				UE_LOG( LogDerivedDataCache, Fatal, TEXT("Unable to create backend graph using the default graph settings (%s) ini=%s."), *GraphName, *GEngineIni );
-			}
-			RootCache = ParseNode( TEXT("Root"), GEngineIni, *GraphName, ParsedNodes );
-			if (!RootCache)
-			{
-				UE_LOG(LogDerivedDataCache, Fatal, TEXT("Unable to create backend graph using the default graph settings (%s) ini=%s. ")
-					TEXT("At least one backend in the graph must be available."), *GraphName, *GEngineIni);
+				Hierarchy->Add(RootNode.Key, RootNode.Value);
 			}
+			CreatedNodes.AddUnique(HierarchyStore);
+			RootNode.Key = HierarchyStore;
 		}
 
-		// Make sure AsyncPutWrapper is created
-		if( !AsyncPutWrapper )
+		// Async must exist in the graph.
+		if (!AsyncNode)
 		{
-			AsyncPutWrapper = AsyncPut::CreateAsyncPutDerivedDataBackend( RootCache, /*bCacheInFlightPuts*/ true );
-			check(AsyncPutWrapper);
-			CreatedBackends.Add( AsyncPutWrapper );
-			RootCache = AsyncPutWrapper;
+			AsyncNode = CreateCacheStoreAsync(RootNode.Key, RootNode.Value, /*bCacheInFlightPuts*/ true);
+			CreatedNodes.AddUnique(AsyncNode);
+			RootNode.Key = AsyncNode;
 		}
 
 		if (MaxKeyLength == 0)
 		{
 			MaxKeyLength = MAX_BACKEND_KEY_LENGTH;
 		}
+
+		RootCache = RootNode.Key;
 	}
 
 	/**
@@ -197,94 +221,99 @@ public:
 	 * @param InParsedNodes Map of parsed nodes and their names to be able to find already parsed nodes
 	 * @return Derived data backend interface instance created from ini settings
 	 */
-	FDerivedDataBackendInterface* ParseNode( const TCHAR* NodeName, const FString& IniFilename, const TCHAR* IniSection, TMap<FString, FDerivedDataBackendInterface*>& InParsedNodes  )
+	FParsedNode ParseNode(const FString& NodeName, const FString& IniFilename, const TCHAR* IniSection, FParsedNodeMap& InParsedNodes)
 	{
-		FDerivedDataBackendInterface* ParsedNode = NULL;
+		if (const FParsedNode* ParsedNode = InParsedNodes.Find(NodeName))
+		{
+			UE_LOG(LogDerivedDataCache, Warning, TEXT("Node %s was referenced more than once in the graph. Nodes may not be shared."), *NodeName);
+			return *ParsedNode;
+		}
+
+		FParsedNode ParsedNode(nullptr, ECacheStoreFlags::None);
 		FString Entry;
-		if( GConfig->GetString( IniSection, NodeName, Entry, IniFilename ) )
+		if (GConfig->GetString(IniSection, *NodeName, Entry, IniFilename))
 		{
-			// Trim whitespace at the beginning.
 			Entry.TrimStartInline();
-			// Remove brackets.
 			Entry.RemoveFromStart(TEXT("("));
 			Entry.RemoveFromEnd(TEXT(")"));
 
 			FString	NodeType;
-			if( FParse::Value( *Entry, TEXT("Type="), NodeType ) )
+			if (FParse::Value(*Entry, TEXT("Type="), NodeType))
 			{
-				if( NodeType == TEXT("FileSystem") )
+				if (NodeType == TEXT("FileSystem"))
 				{
-					ParsedNode = ParseDataCache( NodeName, *Entry );
+					ParsedNode = ParseDataCache(*NodeName, *Entry);
 				}
-				else if( NodeType == TEXT("Boot") )
+				else if (NodeType == TEXT("Boot"))
 				{
-					if( BootCache == NULL )
+					UE_LOG(LogDerivedDataCache, Display, TEXT("Boot nodes are deprecated. Please remove the Boot node from the cache graph."));
+					if (BootCache == nullptr)
 					{
-						BootCache = ParseBootCache( NodeName, *Entry, BootCacheFilename );
-						ParsedNode = BootCache;
+						BootCache = ParseBootCache(*NodeName, *Entry, BootCacheFilename);
+						ParsedNode = MakeTuple(BootCache, ECacheStoreFlags::Local | ECacheStoreFlags::Query | ECacheStoreFlags::Store);
 					}
 					else
 					{
-						UE_LOG( LogDerivedDataCache, Warning, TEXT("FDerivedDataBackendGraph:  Unable to create %s Boot cache because only one Boot or Sparse cache node is supported."), NodeName );
+						UE_LOG(LogDerivedDataCache, Warning, TEXT("Unable to create %s Boot cache because only one Boot node is supported."), *NodeName);
 					}
 				}
-				else if( NodeType == TEXT("Memory") )
+				else if (NodeType == TEXT("Memory"))
 				{
-					ParsedNode = ParseMemoryCache( NodeName, *Entry );
+					ParsedNode = MakeTuple(ParseMemoryCache(*NodeName, *Entry), ECacheStoreFlags::Local | ECacheStoreFlags::Query | ECacheStoreFlags::Store);
 				}
-				else if( NodeType == TEXT("Hierarchical") )
+				else if (NodeType == TEXT("Hierarchical"))
 				{
-					ParsedNode = ParseHierarchicalCache( NodeName, *Entry, IniFilename, IniSection, InParsedNodes );
+					ParsedNode = ParseHierarchyNode(*NodeName, *Entry, IniFilename, IniSection, InParsedNodes);
 				}
-				else if( NodeType == TEXT("KeyLength") )
+				else if (NodeType == TEXT("KeyLength"))
 				{
-					if( MaxKeyLength == 0 )
+					if (MaxKeyLength == 0)
 					{
-						ParsedNode = ParseKeyLength( NodeName, *Entry, IniFilename, IniSection, InParsedNodes );
+						ParsedNode = ParseKeyLength(*NodeName, *Entry, IniFilename, IniSection, InParsedNodes);
 					}
 					else
 					{
-						UE_LOG( LogDerivedDataCache, Warning, TEXT("FDerivedDataBackendGraph:  Unable to create %s KeyLength because only one KeyLength node is supported."), NodeName );
+						UE_LOG(LogDerivedDataCache, Warning, TEXT("Unable to create %s KeyLength node because only one KeyLength node is supported."), *NodeName);
 					}
 				}
-				else if( NodeType == TEXT("AsyncPut") )
+				else if (NodeType == TEXT("AsyncPut"))
 				{
-					if( AsyncPutWrapper == NULL )
+					if (AsyncNode == nullptr)
 					{
-						AsyncPutWrapper = ParseAsyncPut( NodeName, *Entry, IniFilename, IniSection, InParsedNodes );
-						ParsedNode = AsyncPutWrapper;
+						ParsedNode = ParseAsyncNode(*NodeName, *Entry, IniFilename, IniSection, InParsedNodes);
+						AsyncNode = ParsedNode.Key;
 					}
 					else
 					{
-						UE_LOG( LogDerivedDataCache, Warning, TEXT("FDerivedDataBackendGraph:  Unable to create %s AsyncPutWrapper because only one AsyncPutWrapper node is supported."), NodeName );
+						UE_LOG(LogDerivedDataCache, Warning, TEXT("Unable to create %s AsyncPut because only one AsyncPut node is supported."), *NodeName);
 					}
 				}
-				else if( NodeType == TEXT("Verify") )
+				else if (NodeType == TEXT("Verify"))
 				{
-					ParsedNode = ParseVerify( NodeName, *Entry, IniFilename, IniSection, InParsedNodes );
+					ParsedNode = ParseVerify(*NodeName, *Entry, IniFilename, IniSection, InParsedNodes);
 				}
-				else if( NodeType == TEXT("ReadPak") )
+				else if (NodeType == TEXT("ReadPak"))
 				{
-					ParsedNode = ParsePak( NodeName, *Entry, false );
+					ParsedNode = MakeTuple(ParsePak(*NodeName, *Entry, /*bWriting*/ false), ECacheStoreFlags::Local | ECacheStoreFlags::Query | ECacheStoreFlags::StopStore);
 				}
-				else if( NodeType == TEXT("WritePak") )
+				else if (NodeType == TEXT("WritePak"))
 				{
-					ParsedNode = ParsePak( NodeName, *Entry, true );
+					ParsedNode = MakeTuple(ParsePak(*NodeName, *Entry, /*bWriting*/ true), ECacheStoreFlags::Local | ECacheStoreFlags::Store);
 				}
 				else if (NodeType == TEXT("S3"))
 				{
-					ParsedNode = ParseS3Cache(NodeName, *Entry);
+					ParsedNode = MakeTuple(ParseS3Cache(*NodeName, *Entry), ECacheStoreFlags::Local | ECacheStoreFlags::Query | ECacheStoreFlags::StopStore);
 				}
 				else if (NodeType == TEXT("Http"))
 				{
-					ParsedNode = ParseHttpCache(NodeName, *Entry, IniFilename, IniSection);
+					ParsedNode = ParseHttpCache(*NodeName, *Entry, IniFilename, IniSection);
 				}
 				else if (NodeType == TEXT("Zen"))
 				{
-					ParsedNode = ParseZenCache(NodeName, *Entry);
+					ParsedNode = MakeTuple(ParseZenCache(*NodeName, *Entry), ECacheStoreFlags::Local | ECacheStoreFlags::Remote | ECacheStoreFlags::Query | ECacheStoreFlags::Remote);
 				}
 				
-				if (ParsedNode)
+				if (ParsedNode.Key)
 				{
 					// Add a throttling layer if parameters are found
 					uint32 LatencyMS = 0;
@@ -295,27 +324,27 @@ public:
 
 					if (LatencyMS != 0 || MaxBytesPerSecond != 0)
 					{
-						ParsedNode = new Throttle::FDerivedDataBackendThrottleWrapper(ParsedNode, LatencyMS, MaxBytesPerSecond, *Entry);
+						CreatedNodes.AddUnique(ParsedNode.Key);
+						ParsedNode = MakeTuple(CreateCacheStoreThrottle(ParsedNode.Key, LatencyMS, MaxBytesPerSecond), ParsedNode.Value);
 					}
 				}
 			}
 		}
 
-		if( ParsedNode )
+		if (ParsedNode.Key)
 		{
-			// Store this node so that we don't require any order of adding backend nodes
-			InParsedNodes.Add( NodeName, ParsedNode );
+			// Store this node so that we don't require any order of adding nodes
+			InParsedNodes.Add(NodeName, ParsedNode);
 			// Keep references to all created nodes.
-			CreatedBackends.AddUnique( ParsedNode );
+			CreatedNodes.AddUnique(ParsedNode.Key);
 
-			// parse any debug options for this backend. E.g. -ddc-<name>-missrate
+			// Parse any debug options for this node. E.g. -DDC-<Name>-MissRate
 			FDerivedDataBackendInterface::FBackendDebugOptions DebugOptions;
-
-			if (FDerivedDataBackendInterface::FBackendDebugOptions::ParseFromTokens(DebugOptions, NodeName, FCommandLine::Get()))
+			if (FDerivedDataBackendInterface::FBackendDebugOptions::ParseFromTokens(DebugOptions, *NodeName, FCommandLine::Get()))
 			{
-				if (!ParsedNode->ApplyDebugOptions(DebugOptions))
+				if (!ParsedNode.Key->LegacyDebugOptions(DebugOptions))
 				{
-					UE_LOG(LogDerivedDataCache, Warning, TEXT("Node %s is ignoring one or mode -ddc-<nodename>-opt debug options"), NodeName);
+					UE_LOG(LogDerivedDataCache, Warning, TEXT("Node %s is ignoring one or more -DDC-<NodeName>-Option debug options"), *NodeName);
 				}
 			}
 		}
@@ -329,11 +358,11 @@ public:
 	 * @param NodeName Node name
 	 * @param Entry Node definition
 	 * @param bWriting true to create pak interface for writing
-	 * @return Pak file data backend interface instance or NULL if unsuccessful
+	 * @return Pak file data backend interface instance or nullptr if unsuccessful
 	 */
-	FDerivedDataBackendInterface* ParsePak( const TCHAR* NodeName, const TCHAR* Entry, const bool bWriting )
+	ILegacyCacheStore* ParsePak( const TCHAR* NodeName, const TCHAR* Entry, const bool bWriting )
 	{
-		FDerivedDataBackendInterface* PakNode = NULL;
+		ILegacyCacheStore* PakNode = nullptr;
 		FString PakFilename;
 		FParse::Value( Entry, TEXT("Filename="), PakFilename );
 		bool bCompressed = GetParsedBool(Entry, TEXT("Compressed="));
@@ -350,7 +379,7 @@ public:
 				FGuid Temp = FGuid::NewGuid();
 				ReadPakFilename = PakFilename;
 				WritePakFilename = PakFilename + TEXT(".") + Temp.ToString();
-				WritePakCache = PakFile::CreatePakFileDerivedDataBackend(*WritePakFilename, /*bWriting*/ true, bCompressed);
+				WritePakCache = CacheStore::PakFile::CreatePakFileDerivedDataBackend(*WritePakFilename, /*bWriting*/ true, bCompressed);
 				PakNode = WritePakCache;
 			}
 			else
@@ -358,7 +387,7 @@ public:
 				bool bReadPak = FPlatformFileManager::Get().GetPlatformFile().FileExists( *PakFilename );
 				if( bReadPak )
 				{
-					PakFile::IPakFileDerivedDataBackend* ReadPak = PakFile::CreatePakFileDerivedDataBackend(*PakFilename, /*bWriting*/ false, bCompressed);
+					CacheStore::PakFile::IPakFileDerivedDataBackend* ReadPak = CacheStore::PakFile::CreatePakFileDerivedDataBackend(*PakFilename, /*bWriting*/ false, bCompressed);
 					ReadPakFilename = PakFilename;
 					PakNode = ReadPak;
 					ReadPakCache.Add(ReadPak);
@@ -380,40 +409,30 @@ public:
 	 * @param IniFilename ini filename.
 	 * @param IniSection ini section containing graph definition
 	 * @param InParsedNodes map of nodes and their names which have already been parsed
-	 * @return Verify wrapper backend interface instance or NULL if unsuccessful
+	 * @return Verify wrapper backend interface instance or nullptr if unsuccessful
 	 */
-	FDerivedDataBackendInterface* ParseVerify( const TCHAR* NodeName, const TCHAR* Entry, const FString& IniFilename, const TCHAR* IniSection, TMap<FString, FDerivedDataBackendInterface*>& InParsedNodes )
+	FParsedNode ParseVerify(const TCHAR* NodeName, const TCHAR* Entry, const FString& IniFilename, const TCHAR* IniSection, FParsedNodeMap& InParsedNodes)
 	{
-		FDerivedDataBackendInterface* InnerNode = NULL;
+		FParsedNode InnerNode(nullptr, ECacheStoreFlags::None);
 		FString InnerName;
-		if( FParse::Value( Entry, TEXT("Inner="), InnerName ) )
+		if (FParse::Value(Entry, TEXT("Inner="), InnerName))
 		{
-			FDerivedDataBackendInterface** ParsedInnerNode = InParsedNodes.Find( InnerName );
-			if( ParsedInnerNode )
-			{
-				UE_LOG( LogDerivedDataCache, Warning, TEXT("Inner node %s for Verify node %s already exists. Nodes can only be used once."), *InnerName, NodeName );
-				return NULL;
-			}
-			else
-			{
-				InnerNode = ParseNode( *InnerName, IniFilename, IniSection, InParsedNodes );
-			}
+			InnerNode = ParseNode(InnerName, IniFilename, IniSection, InParsedNodes);
 		}
 
-		FDerivedDataBackendInterface* VerifyNode = NULL;
-		if( InnerNode )
+		if (InnerNode.Key)
 		{
-			IFileManager::Get().DeleteDirectory(*(FPaths::ProjectSavedDir() / TEXT("VerifyDDC/")), false, true);
+			IFileManager::Get().DeleteDirectory(*(FPaths::ProjectSavedDir() / TEXT("VerifyDDC/")), /*bRequireExists*/ false, /*bTree*/ true);
 
-			const bool bFix = GetParsedBool( Entry, TEXT("Fix=") );
-			VerifyNode = new Verify::FDerivedDataBackendVerifyWrapper( InnerNode, bFix );
+			const bool bFix = GetParsedBool(Entry, TEXT("Fix="));
+			InnerNode = MakeTuple(CreateCacheStoreVerify(InnerNode.Key, bFix), InnerNode.Value);
 		}
 		else
 		{
-			UE_LOG( LogDerivedDataCache, Warning, TEXT("Unable to find inner node %s for Verify node %s. Verify node will not be created."), *InnerName, NodeName );
+			UE_LOG(LogDerivedDataCache, Warning, TEXT("Unable to find inner node %s for Verify node %s. Verify node will not be created."), *InnerName, NodeName);
 		}
 
-		return VerifyNode;
+		return InnerNode;
 	}
 
 	/**
@@ -424,37 +443,36 @@ public:
 	 * @param IniFilename ini filename.
 	 * @param IniSection ini section containing graph definition
 	 * @param InParsedNodes map of nodes and their names which have already been parsed
-	 * @return AsyncPut wrapper backend interface instance or NULL if unsuccessfull
+	 * @return AsyncPut wrapper backend interface instance or nullptr if unsuccessful
 	 */
-	FDerivedDataBackendInterface* ParseAsyncPut( const TCHAR* NodeName, const TCHAR* Entry, const FString& IniFilename, const TCHAR* IniSection, TMap<FString, FDerivedDataBackendInterface*>& InParsedNodes )
+	FParsedNode ParseAsyncNode(const TCHAR* NodeName, const TCHAR* Entry, const FString& IniFilename, const TCHAR* IniSection, FParsedNodeMap& InParsedNodes)
 	{
-		FDerivedDataBackendInterface* InnerNode = NULL;
+		FParsedNode InnerNode(nullptr, ECacheStoreFlags::None);
 		FString InnerName;
-		if( FParse::Value( Entry, TEXT("Inner="), InnerName ) )
+		if (FParse::Value(Entry, TEXT("Inner="), InnerName))
 		{
-			FDerivedDataBackendInterface** ParsedInnerNode = InParsedNodes.Find( InnerName );
-			if( ParsedInnerNode )
-			{
-				UE_LOG( LogDerivedDataCache, Warning, TEXT("Inner node %s for AsyncPut node %s already exists. Nodes can only be used once."), *InnerName, NodeName );
-				return NULL;
-			}
-			else
-			{
-				InnerNode = ParseNode( *InnerName, IniFilename, IniSection, InParsedNodes );
-			}
+			InnerNode = ParseNode(InnerName, IniFilename, IniSection, InParsedNodes);
 		}
 
-		FDerivedDataBackendInterface* AsyncNode = NULL;
-		if( InnerNode )
+		if (InnerNode.Key)
 		{
-			AsyncNode = AsyncPut::CreateAsyncPutDerivedDataBackend( InnerNode, /*bCacheInFlightPuts*/ true );
+			if (!Hierarchy)
+			{
+				ILegacyCacheStore* HierarchyStore = CreateCacheStoreHierarchy(Hierarchy);
+				Hierarchy->Add(InnerNode.Key, InnerNode.Value);
+				CreatedNodes.AddUnique(HierarchyStore);
+				InnerNode.Key = HierarchyStore;
+			}
+
+			ILegacyCacheStore* AsyncStore = CreateCacheStoreAsync(InnerNode.Key, InnerNode.Value, /*bCacheInFlightPuts*/ true);
+			InnerNode = MakeTuple(AsyncStore, InnerNode.Value);
 		}
 		else
 		{
-			UE_LOG( LogDerivedDataCache, Warning, TEXT("Unable to find inner node %s for AsyncPut node %s. AsyncPut node will not be created."), *InnerName, NodeName );
+			UE_LOG(LogDerivedDataCache, Warning, TEXT("Unable to find inner node %s for AsyncPut node %s. AsyncPut node will not be created."), *InnerName, NodeName);
 		}
 
-		return AsyncNode;
+		return InnerNode;
 	}
 
 	/**
@@ -465,35 +483,32 @@ public:
 	 * @param IniFilename ini filename.
 	 * @param IniSection ini section containing graph definition
 	 * @param InParsedNodes map of nodes and their names which have already been parsed
-	 * @return KeyLength wrapper backend interface instance or NULL if unsuccessfull
+	 * @return KeyLength wrapper backend interface instance or nullptr if unsuccessful
 	 */
-	FDerivedDataBackendInterface* ParseKeyLength( const TCHAR* NodeName, const TCHAR* Entry, const FString& IniFilename, const TCHAR* IniSection, TMap<FString, FDerivedDataBackendInterface*>& InParsedNodes )
+	FParsedNode ParseKeyLength(const TCHAR* NodeName, const TCHAR* Entry, const FString& IniFilename, const TCHAR* IniSection, FParsedNodeMap& InParsedNodes)
 	{
-		FDerivedDataBackendInterface* InnerNode = nullptr;
+		if (MaxKeyLength)
+		{
+			UE_LOG(LogDerivedDataCache, Warning, TEXT("Node %s is disabled because there may be only one key length node."), NodeName);
+			return MakeTuple(nullptr, ECacheStoreFlags::None);
+		}
+
+		FParsedNode InnerNode(nullptr, ECacheStoreFlags::None);
 		FString InnerName;
-		if( FParse::Value( Entry, TEXT("Inner="), InnerName ) )
+		if (FParse::Value(Entry, TEXT("Inner="), InnerName))
 		{
-			FDerivedDataBackendInterface** ParsedInnerNode = InParsedNodes.Find( InnerName );
-			if( ParsedInnerNode )
-			{
-				UE_LOG( LogDerivedDataCache, Warning, TEXT("Inner node %s for KeyLength node %s already exists. Nodes can only be used once."), *InnerName, NodeName );
-				return NULL;
-			}
-			else
-			{
-				InnerNode = ParseNode( *InnerName, IniFilename, IniSection, InParsedNodes );
-			}
+			InnerNode = ParseNode(InnerName, IniFilename, IniSection, InParsedNodes);
 		}
 
-		if( InnerNode )
+		if (InnerNode.Key)
 		{
 			int32 KeyLength = MAX_BACKEND_KEY_LENGTH;
-			FParse::Value( Entry, TEXT("Length="), KeyLength );
-			MaxKeyLength = FMath::Clamp( KeyLength, 0, MAX_BACKEND_KEY_LENGTH );
+			FParse::Value(Entry, TEXT("Length="), KeyLength);
+			MaxKeyLength = FMath::Clamp(KeyLength, 0, MAX_BACKEND_KEY_LENGTH);
 		}
 		else
 		{
-			UE_LOG( LogDerivedDataCache, Warning, TEXT("Unable to find inner node %s for KeyLength node %s. KeyLength node will not be created."), *InnerName, NodeName );
+			UE_LOG(LogDerivedDataCache, Warning, TEXT("Unable to find inner node %s for KeyLength node %s. KeyLength node will not be created."), *InnerName, NodeName);
 		}
 
 		return InnerNode;
@@ -507,62 +522,53 @@ public:
 	 * @param IniFilename ini filename.
 	 * @param IniSection ini section containing graph definition
 	 * @param InParsedNodes map of nodes and their names which have already been parsed
-	 * @return Hierarchical backend interface instance or NULL if unsuccessfull
+	 * @return Hierarchical backend interface instance or nullptr if unsuccessful
 	 */
-	FDerivedDataBackendInterface* ParseHierarchicalCache( const TCHAR* NodeName, const TCHAR* Entry, const FString& IniFilename, const TCHAR* IniSection, TMap<FString, FDerivedDataBackendInterface*>& InParsedNodes )
+	FParsedNode ParseHierarchyNode(const TCHAR* NodeName, const TCHAR* Entry, const FString& IniFilename, const TCHAR* IniSection, FParsedNodeMap& InParsedNodes)
 	{
 		const TCHAR* InnerMatch = TEXT("Inner=");
-		const int32 InnerMatchLength = FCString::Strlen( InnerMatch );
+		const int32 InnerMatchLength = FCString::Strlen(InnerMatch);
 
-		TArray< FDerivedDataBackendInterface* > InnerNodes;
+		TArray<FParsedNode> InnerNodes;
 		FString InnerName;
-		while ( FParse::Value( Entry, InnerMatch, InnerName ) )
+		while (FParse::Value(Entry, InnerMatch, InnerName))
 		{
-			// Check if the child has already been parsed
-			FDerivedDataBackendInterface** ParsedInnerNode = InParsedNodes.Find( InnerName );
-			if( ParsedInnerNode )
+			FParsedNode InnerNode = ParseNode(InnerName, IniFilename, IniSection, InParsedNodes);
+			if (InnerNode.Key)
 			{
-				UE_LOG( LogDerivedDataCache, Warning, TEXT("Inner node %s for hierarchical node %s already exists. Nodes can only be used once."), *InnerName, NodeName );
+				InnerNodes.Add(InnerNode);
 			}
 			else
 			{
-				FDerivedDataBackendInterface* InnerNode = ParseNode( *InnerName, IniFilename, IniSection, InParsedNodes );
-				if( InnerNode )
-				{
-					InnerNodes.Add( InnerNode );
-				}
-				else
-				{
-					UE_LOG( LogDerivedDataCache, Log, TEXT("Unable to find inner node %s for hierarchical cache %s."), *InnerName, NodeName );
-				}
+				UE_LOG(LogDerivedDataCache, Log, TEXT("Unable to find inner node %s for hierarchy %s."), *InnerName, NodeName);
 			}
 
 			// Move the Entry pointer forward so that we can find more children
-			Entry = FCString::Strifind( Entry, InnerMatch );
+			Entry = FCString::Strifind(Entry, InnerMatch);
 			Entry += InnerMatchLength;
 		}
 
-		FDerivedDataBackendInterface* Hierarchy = NULL;
-		if( InnerNodes.Num() > 1 )
+		if (InnerNodes.IsEmpty())
 		{
-			Hierarchical::FHierarchicalDerivedDataBackend* HierarchyBackend = new Hierarchical::FHierarchicalDerivedDataBackend( InnerNodes );
-			Hierarchy = HierarchyBackend;
-			if (HierarchicalWrapper == NULL)
-			{
-				HierarchicalWrapper = HierarchyBackend;
-			}
+			UE_LOG(LogDerivedDataCache, Warning, TEXT("Hierarchical cache %s has no inner backends and will not be created."), NodeName);
+			return MakeTuple(nullptr, ECacheStoreFlags::None);
 		}
-		else if ( InnerNodes.Num() == 1 )
+
+		if (Hierarchy)
 		{
-			Hierarchy = InnerNodes[ 0 ];
-			InnerNodes.Empty();
+			UE_LOG(LogDerivedDataCache, Warning, TEXT("Node %s is disabled because there may be only one hierarchy node. ")
+				TEXT("Confirm there is only one hierarchy in the cache graph and that it is inside of any async node."), NodeName);
+			return MakeTuple(nullptr, ECacheStoreFlags::None);
 		}
-		else
+
+		ILegacyCacheStore* HierarchyStore = CreateCacheStoreHierarchy(Hierarchy);
+		ECacheStoreFlags Flags = ECacheStoreFlags::None;
+		for (const FParsedNode& Node : InnerNodes)
 		{
-			UE_LOG( LogDerivedDataCache, Warning, TEXT("Hierarchical cache %s has no inner backends and will not be created."), NodeName );
+			Hierarchy->Add(Node.Key, Node.Value);
+			Flags |= Node.Value;
 		}
-
-		return Hierarchy;
+		return MakeTuple(HierarchyStore, Flags & ~ECacheStoreFlags::StopStore);
 	}
 
 	/**
@@ -570,13 +576,13 @@ public:
 	 *
 	 * @param NodeName Node name.
 	 * @param Entry Node definition.
-	 * @return Filesystem data cache backend interface instance or NULL if unsuccessfull
+	 * @return Filesystem data cache backend interface instance or nullptr if unsuccessful
 	 */
-	FDerivedDataBackendInterface* ParseDataCache( const TCHAR* NodeName, const TCHAR* Entry )
+	FParsedNode ParseDataCache( const TCHAR* NodeName, const TCHAR* Entry )
 	{
-		FDerivedDataBackendInterface* DataCache = NULL;
+		FParsedNode DataCache(nullptr, ECacheStoreFlags::None);
 
-		// Parse Path by default, it may be overwriten by EnvPathOverride
+		// Parse Path by default, it may be overwritten by EnvPathOverride
 		FString Path;
 		FParse::Value( Entry, TEXT("Path="), Path );
 
@@ -652,7 +658,7 @@ public:
 		}
 		else
 		{
-			FDerivedDataBackendInterface* InnerFileSystem = NULL;
+			FDerivedDataBackendInterface* InnerFileSystem = nullptr;
 
 			// Try to set up the shared drive, allow user to correct any issues that may exist.
 			bool RetryOnFailure = false;
@@ -672,14 +678,18 @@ public:
 
 				if (!bShared || IFileManager::Get().DirectoryExists(*Path))
 				{
-					InnerFileSystem = FileSystem::CreateFileSystemDerivedDataBackend(*Path, Entry, *WriteAccessLog);
+					InnerFileSystem = CacheStore::FileSystem::CreateFileSystemDerivedDataBackend(*Path, Entry, *WriteAccessLog);
 				}
 
 				if (InnerFileSystem)
 				{
 					bUsingSharedDDC |= bShared;
 
-					DataCache = InnerFileSystem;
+					ECacheStoreFlags Flags = ECacheStoreFlags::Query;
+					Flags |= InnerFileSystem->IsWritable() ? ECacheStoreFlags::Store : ECacheStoreFlags::None;
+					Flags |= InnerFileSystem->GetSpeedClass() == EBackendSpeedClass::Local ? ECacheStoreFlags::Local : ECacheStoreFlags::Remote;
+
+					DataCache = MakeTuple(InnerFileSystem, Flags);
 					UE_LOG(LogDerivedDataCache, Log, TEXT("Using %s data cache path %s: %s"), NodeName, *Path, !InnerFileSystem->IsWritable() ? TEXT("ReadOnly") : TEXT("Writable"));
 					Directories.AddUnique(Path);
 				}
@@ -706,7 +716,7 @@ public:
 	/**
 	 * Creates an S3 data cache interface.
 	 */
-	FDerivedDataBackendInterface* ParseS3Cache(const TCHAR* NodeName, const TCHAR* Entry)
+	ILegacyCacheStore* ParseS3Cache(const TCHAR* NodeName, const TCHAR* Entry)
 	{
 		FString ManifestPath;
 		if (!FParse::Value(Entry, TEXT("Manifest="), ManifestPath))
@@ -766,7 +776,7 @@ public:
 			}
 		}
 
-		if (FDerivedDataBackendInterface* Backend = S3::CreateS3DerivedDataBackend(*ManifestPath, *BaseUrl, *Region, *CanaryObjectKey, *CachePath))
+		if (ILegacyCacheStore* Backend = CacheStore::S3::CreateS3DerivedDataBackend(*ManifestPath, *BaseUrl, *Region, *CanaryObjectKey, *CachePath))
 		{
 			return Backend;
 		}
@@ -836,7 +846,7 @@ public:
 	/**
 	 * Creates a HTTP data cache interface.
 	 */
-	FDerivedDataBackendInterface* ParseHttpCache(
+	FParsedNode ParseHttpCache(
 		const TCHAR* NodeName,
 		const TCHAR* Entry,
 		const FString& IniFilename,
@@ -855,13 +865,13 @@ public:
 		if (Host.IsEmpty())
 		{
 			UE_LOG(LogDerivedDataCache, Error, TEXT("Node %s does not specify 'Host'"), NodeName);
-			return nullptr;
+			return MakeTuple(nullptr, ECacheStoreFlags::None);
 		}
 
 		if (Host == TEXT("None"))
 		{
 			UE_LOG(LogDerivedDataCache, Log, TEXT("Node %s is disabled because Host is set to 'None'"), NodeName);
-			return nullptr;
+			return MakeTuple(nullptr, ECacheStoreFlags::None);
 		}
 
 		if (Namespace.IsEmpty())
@@ -878,19 +888,19 @@ public:
 		if (OAuthProvider.IsEmpty())
 		{
 			UE_LOG(LogDerivedDataCache, Error, TEXT("Node %s does not specify 'OAuthProvider'"), NodeName);
-			return nullptr;
+			return MakeTuple(nullptr, ECacheStoreFlags::None);
 		}
 
 		if (OAuthClientId.IsEmpty())
 		{
 			UE_LOG(LogDerivedDataCache, Error, TEXT("Node %s does not specify 'OAuthClientId'"), NodeName);
-			return nullptr;
+			return MakeTuple(nullptr, ECacheStoreFlags::None);
 		}
 
 		if (OAuthSecret.IsEmpty())
 		{
 			UE_LOG(LogDerivedDataCache, Error, TEXT("Node %s does not specify 'OAuthSecret'"), NodeName);
-			return nullptr;
+			return MakeTuple(nullptr, ECacheStoreFlags::None);
 		}
 
 		FDerivedDataBackendInterface::ESpeedClass ForceSpeedClass = FDerivedDataBackendInterface::ESpeedClass::Unknown;
@@ -924,15 +934,16 @@ public:
 			UE_LOG(LogDerivedDataCache, Log, TEXT("Node %s found speed class override ForceSpeedClass=%s"), NodeName, *ForceSpeedClassValue);
 		}
 
-		return Http::CreateHttpDerivedDataBackend(
+		return MakeTuple(CacheStore::Http::CreateHttpDerivedDataBackend(
 			NodeName, *Host, *Namespace, *StructuredNamespace, *OAuthProvider, *OAuthClientId, *OAuthSecret,
-			ForceSpeedClass == FDerivedDataBackendInterface::ESpeedClass::Unknown ? nullptr : &ForceSpeedClass, bReadOnly);
+			ForceSpeedClass == FDerivedDataBackendInterface::ESpeedClass::Unknown ? nullptr : &ForceSpeedClass, bReadOnly),
+			ECacheStoreFlags::Remote | ECacheStoreFlags::Query | (bReadOnly ? ECacheStoreFlags::None : ECacheStoreFlags::Store));
 	}
 
 	/**
 	 * Creates a Zen structured data cache interface
 	 */
-	FDerivedDataBackendInterface* ParseZenCache(const TCHAR* NodeName, const TCHAR* Entry)
+	ILegacyCacheStore* ParseZenCache(const TCHAR* NodeName, const TCHAR* Entry)
 	{
 		FString ServiceUrl;
 		FParse::Value(Entry, TEXT("Host="), ServiceUrl);
@@ -944,7 +955,7 @@ public:
 			UE_LOG(LogDerivedDataCache, Warning, TEXT("Node %s does not specify 'Namespace', falling back to '%s'"), NodeName, *Namespace);
 		}
 
-		if (FDerivedDataBackendInterface* Backend = ZenCache::CreateZenDerivedDataBackend(NodeName, *ServiceUrl, *Namespace))
+		if (ILegacyCacheStore* Backend = CacheStore::ZenCache::CreateZenDerivedDataBackend(NodeName, *ServiceUrl, *Namespace))
 		{
 			return Backend;
 		}
@@ -959,7 +970,7 @@ public:
 	 * @param NodeName Node name.
 	 * @param Entry Node definition.
 	 * @param OutFilename filename specified for the cache
-	 * @return Boot data cache backend interface instance or NULL if unsuccessful
+	 * @return Boot data cache backend interface instance or nullptr if unsuccessful
 	 */
 	FFileBackedDerivedDataBackend* ParseBootCache( const TCHAR* NodeName, const TCHAR* Entry, FString& OutFilename )
 	{
@@ -983,7 +994,7 @@ public:
 			MaxCacheSize = FMath::Min(MaxCacheSize, MaxSupportedCacheSize);
 
 			UE_LOG( LogDerivedDataCache, Display, TEXT("Max Cache Size: %d MB"), MaxCacheSize);
-			Cache = Memory::CreateMemoryDerivedDataBackend(TEXT("Boot"), MaxCacheSize * 1024 * 1024, /*bCanBeDisabled*/ true);
+			Cache = CacheStore::Memory::CreateMemoryDerivedDataBackend(TEXT("Boot"), MaxCacheSize * 1024 * 1024, /*bCanBeDisabled*/ true);
 
 			if( Cache && Filename.Len() )
 			{
@@ -1017,7 +1028,7 @@ public:
 	 *
 	 * @param NodeName Node name.
 	 * @param Entry Node definition.
-	 * @return Memory data cache backend interface instance or NULL if unsuccessfull
+	 * @return Memory data cache backend interface instance or nullptr if unsuccessful
 	 */
 	FFileBackedDerivedDataBackend* ParseMemoryCache( const TCHAR* NodeName, const TCHAR* Entry )
 	{
@@ -1025,7 +1036,7 @@ public:
 		FString Filename;
 
 		FParse::Value( Entry, TEXT("Filename="), Filename );
-		Cache = Memory::CreateMemoryDerivedDataBackend(NodeName, /*MaxCacheSize*/ -1, /*bCanBeDisabled*/ false);
+		Cache = CacheStore::Memory::CreateMemoryDerivedDataBackend(NodeName, /*MaxCacheSize*/ -1, /*bCanBeDisabled*/ false);
 		if( Cache && Filename.Len() )
 		{
 			if( Cache->LoadCache( *Filename ) )
@@ -1043,7 +1054,7 @@ public:
 
 	virtual ~FDerivedDataBackendGraph()
 	{
-		RootCache = NULL;
+		RootCache = nullptr;
 		DestroyCreatedBackends();
 	}
 
@@ -1101,8 +1112,8 @@ public:
 
 				for(const FString& MergePakName : MergePakList)
 				{
-					TUniquePtr<PakFile::IPakFileDerivedDataBackend> ReadPak(
-						PakFile::CreatePakFileDerivedDataBackend(*FPaths::Combine(*FPaths::GetPath(WritePakFilename), *MergePakName), /*bWriting*/ false, /*bCompressed*/ false));
+					TUniquePtr<CacheStore::PakFile::IPakFileDerivedDataBackend> ReadPak(
+						CacheStore::PakFile::CreatePakFileDerivedDataBackend(*FPaths::Combine(*FPaths::GetPath(WritePakFilename), *MergePakName), /*bWriting*/ false, /*bCompressed*/ false));
 					WritePakCache->MergeCache(ReadPak.Get());
 				}
 			}
@@ -1127,7 +1138,7 @@ public:
 							UE_LOG(LogDerivedDataCache, Error, TEXT("Could not delete the pak file %s to overwrite it with a new one."), *ReadPakFilename);
 						}
 					}
-					if (!PakFile::IPakFileDerivedDataBackend::SortAndCopy(WritePakFilename, ReadPakFilename))
+					if (!CacheStore::PakFile::IPakFileDerivedDataBackend::SortAndCopy(WritePakFilename, ReadPakFilename))
 					{
 						UE_LOG(LogDerivedDataCache, Error, TEXT("Couldn't sort pak file (%s)"), *WritePakFilename);
 					}
@@ -1191,13 +1202,13 @@ public:
 	{
 		// Assumptions: there's at least one read-only pak backend in the hierarchy
 		// and its parent is a hierarchical backend.
-		PakFile::IPakFileDerivedDataBackend* ReadPak = nullptr;
-		if (HierarchicalWrapper && FPlatformFileManager::Get().GetPlatformFile().FileExists(PakFilename))
+		CacheStore::PakFile::IPakFileDerivedDataBackend* ReadPak = nullptr;
+		if (Hierarchy && FPlatformFileManager::Get().GetPlatformFile().FileExists(PakFilename))
 		{
-			ReadPak = PakFile::CreatePakFileDerivedDataBackend(PakFilename, /*bWriting*/ false, /*bCompressed*/ false);
+			ReadPak = CacheStore::PakFile::CreatePakFileDerivedDataBackend(PakFilename, /*bWriting*/ false, /*bCompressed*/ false);
 
-			HierarchicalWrapper->AddInnerBackend(ReadPak);
-			CreatedBackends.Add(ReadPak);
+			Hierarchy->Add(ReadPak, ECacheStoreFlags::Local | ECacheStoreFlags::Store | ECacheStoreFlags::StopStore);
+			CreatedNodes.AddUnique(ReadPak);
 			ReadPakCache.Add(ReadPak);
 		}
 		else
@@ -1212,17 +1223,17 @@ public:
 	{
 		for (int PakIndex = 0; PakIndex < ReadPakCache.Num(); ++PakIndex)
 		{
-			PakFile::IPakFileDerivedDataBackend* ReadPak = ReadPakCache[PakIndex];
+			CacheStore::PakFile::IPakFileDerivedDataBackend* ReadPak = ReadPakCache[PakIndex];
 			if (ReadPak->GetFilename() == PakFilename)
 			{
-				check(HierarchicalWrapper);
+				check(Hierarchy);
 
 				// Wait until all async requests are complete.
 				WaitForQuiescence(false);
 
-				HierarchicalWrapper->RemoveInnerBackend(ReadPak);
+				Hierarchy->RemoveNotSafe(ReadPak);
 				ReadPakCache.RemoveAt(PakIndex);
-				CreatedBackends.Remove(ReadPak);
+				CreatedNodes.Remove(ReadPak);
 				ReadPak->Close();
 				delete ReadPak;
 				return true;
@@ -1246,11 +1257,11 @@ private:
 	/** Delete all created backends in the reversed order they were created. */
 	void DestroyCreatedBackends()
 	{
-		for (int32 BackendIndex = CreatedBackends.Num() - 1; BackendIndex >= 0; --BackendIndex)
+		for (int32 BackendIndex = CreatedNodes.Num() - 1; BackendIndex >= 0; --BackendIndex)
 		{
-			delete CreatedBackends[BackendIndex];
+			delete CreatedNodes[BackendIndex];
 		}
-		CreatedBackends.Empty();
+		CreatedNodes.Empty();
 	}
 
 	/** MountPak console command handler. */
@@ -1282,18 +1293,18 @@ private:
 	FString											WritePakFilename;
 
 	/** Root of the graph */
-	FDerivedDataBackendInterface*					RootCache;
+	ILegacyCacheStore*					RootCache;
 
 	/** References to all created backed interfaces */
-	TArray< FDerivedDataBackendInterface* > CreatedBackends;
+	TArray< ILegacyCacheStore* > CreatedNodes;
 
 	/** Instances of backend interfaces which exist in only one copy */
 	FFileBackedDerivedDataBackend*	BootCache;
-	PakFile::IPakFileDerivedDataBackend* WritePakCache;
-	FDerivedDataBackendInterface*	AsyncPutWrapper;
-	Hierarchical::FHierarchicalDerivedDataBackend* HierarchicalWrapper;
+	CacheStore::PakFile::IPakFileDerivedDataBackend* WritePakCache;
+	ILegacyCacheStore*	AsyncNode;
+	ICacheStoreOwner* Hierarchy;
 	/** Support for multiple read only pak files. */
-	TArray<PakFile::IPakFileDerivedDataBackend*> ReadPakCache;
+	TArray<CacheStore::PakFile::IPakFileDerivedDataBackend*> ReadPakCache;
 
 	/** List of directories used by the DDC */
 	TArray<FString> Directories;
@@ -1312,7 +1323,7 @@ private:
 	FAutoConsoleCommand UnountPakCommand;
 };
 
-} // UE::DerivedData::CacheStore::Graph
+} // UE::DerivedData
 
 namespace UE::DerivedData
 {
@@ -1366,11 +1377,18 @@ void FDerivedDataBackendInterface::LegacyGet(
 		}
 		else
 		{
+			FSharedBuffer Value;
 			TArray<uint8> Data;
-			const bool bGetOk = GetCachedData(*Request.Key.GetShortKey(), Data);
-			FCompositeBuffer Value(MakeSharedBufferFromArray(MoveTemp(Data)));
-			const bool bKeyOk = bGetOk && Request.Key.ReadValueTrailer(Value);
-			OnComplete({Request.Name, Request.Key, MoveTemp(Value).ToShared(), Request.UserData, bKeyOk ? EStatus::Ok : EStatus::Error});
+			if (const bool bGetOk = GetCachedData(*Request.Key.GetShortKey(), Data))
+			{
+				FCompositeBuffer CompositeValue(MakeSharedBufferFromArray(MoveTemp(Data)));
+				if (const bool bKeyOk = Request.Key.ReadValueTrailer(CompositeValue))
+				{
+					Value = MoveTemp(CompositeValue).ToShared();
+				}
+			}
+			const EStatus Status = Value ? EStatus::Ok : EStatus::Error;
+			OnComplete({Request.Name, Request.Key, MoveTemp(Value), Request.UserData, Status});
 		}
 	}
 
@@ -1421,7 +1439,7 @@ bool FDerivedDataBackendInterface::LegacyDebugOptions(FBackendDebugOptions& Opti
 
 FDerivedDataBackend& FDerivedDataBackend::Get()
 {
-	return CacheStore::Graph::FDerivedDataBackendGraph::Get();
+	return FDerivedDataBackendGraph::Get();
 }
 
 /**
diff --git a/Engine/Source/Developer/DerivedDataCache/Private/DerivedDataCacheStoreAsync.cpp b/Engine/Source/Developer/DerivedDataCache/Private/DerivedDataCacheStoreAsync.cpp
--- a/Engine/Source/Developer/DerivedDataCache/Private/DerivedDataCacheStoreAsync.cpp
+++ b/Engine/Source/Developer/DerivedDataCache/Private/DerivedDataCacheStoreAsync.cpp
@@ -1,149 +1,38 @@
 // Copyright Epic Games, Inc. All Rights Reserved.
 
 #include "Async/AsyncWork.h"
-#include "DerivedDataBackendInterface.h"
 #include "DerivedDataCachePrivate.h"
 #include "DerivedDataCacheStore.h"
 #include "DerivedDataCacheUsageStats.h"
+#include "DerivedDataLegacyCacheStore.h"
 #include "DerivedDataRequest.h"
 #include "DerivedDataRequestOwner.h"
 #include "DerivedDataValueId.h"
 #include "Experimental/Async/LazyEvent.h"
 #include "FileBackedDerivedDataBackend.h"
 #include "Memory/SharedBuffer.h"
-#include "Misc/ScopeLock.h"
 #include "ProfilingDebugging/CookStats.h"
 #include "Stats/Stats.h"
 #include "Tasks/Task.h"
+#include "Templates/Invoke.h"
 
 namespace UE::DerivedData::CacheStore::Memory
 {
 FFileBackedDerivedDataBackend* CreateMemoryDerivedDataBackend(const TCHAR* Name, int64 MaxCacheSize, bool bCanBeDisabled);
 } // UE::DerivedData::CacheStore::Memory
 
-namespace UE::DerivedData::CacheStore::AsyncPut
+namespace UE::DerivedData
 {
 
-/** 
- * Thread safe set helper
-**/
-struct FThreadSet
-{
-	FCriticalSection	SynchronizationObject;
-	TSet<FString>		FilesInFlight;
-
-	void Add(const FString& Key)
-	{
-		FScopeLock ScopeLock(&SynchronizationObject);
-		check(Key.Len());
-		FilesInFlight.Add(Key);
-	}
-	void Remove(const FString& Key)
-	{
-		FScopeLock ScopeLock(&SynchronizationObject);
-		FilesInFlight.Remove(Key);
-	}
-	bool Exists(const FString& Key)
-	{
-		FScopeLock ScopeLock(&SynchronizationObject);
-		return FilesInFlight.Contains(Key);
-	}
-	bool AddIfNotExists(const FString& Key)
-	{
-		FScopeLock ScopeLock(&SynchronizationObject);
-		check(Key.Len());
-		if (!FilesInFlight.Contains(Key))
-		{
-			FilesInFlight.Add(Key);
-			return true;
-		}
-		return false;
-	}
-};
-
-/** 
- * A backend wrapper that coordinates async puts. This means that a Get will hit an in-memory cache while the async put is still in flight.
-**/
-class FDerivedDataBackendAsyncPutWrapper : public FDerivedDataBackendInterface
+/**
+ * A cache store that executes non-blocking requests in a dedicated thread pool.
+ *
+ * Puts can be stored in a memory cache while they are in flight.
+ */
+class FCacheStoreAsync : public ILegacyCacheStore
 {
 public:
-
-	/**
-	 * Constructor
-	 *
-	 * @param	InInnerBackend		Backend to use for storage, my responsibilities are about async puts
-	 * @param	bCacheInFlightPuts	if true, cache in-flight puts in a memory cache so that they hit immediately
-	 */
-	FDerivedDataBackendAsyncPutWrapper(FDerivedDataBackendInterface* InInnerBackend, bool bCacheInFlightPuts);
-
-	/** Return a name for this interface */
-	virtual FString GetName() const override
-	{
-		return FString::Printf(TEXT("AsyncPutWrapper (%s)"), *InnerBackend->GetName());
-	}
-
-	/** return true if this cache is writable **/
-	virtual bool IsWritable() const override;
-
-	/** Returns a class of speed for this interface **/
-	virtual ESpeedClass GetSpeedClass() const override;
-
-	/** Return true if hits on this cache should propagate to lower cache level. */
-	virtual bool BackfillLowerCacheLevels() const override;
-
-	/**
-	 * Synchronous test for the existence of a cache item
-	 *
-	 * @param	CacheKey	Alphanumeric+underscore key of this cache item
-	 * @return				true if the data probably will be found, this can't be guaranteed because of concurrency in the backends, corruption, etc
-	 */
-	virtual bool CachedDataProbablyExists(const TCHAR* CacheKey) override;
-
-	/**
-	 * Synchronous test for the existence of multiple cache items
-	 *
-	 * @param	CacheKeys	Alphanumeric+underscore key of the cache items
-	 * @return				A bit array with bits indicating whether the data for the corresponding key will probably be found
-	 */
-	virtual TBitArray<> CachedDataProbablyExistsBatch(TConstArrayView<FString> CacheKeys) override;
-
-	/**
-	 * Attempt to make sure the cached data will be available as optimally as possible.
-	 *
-	 * @param	CacheKeys	Alphanumeric+underscore keys of the cache items
-	 * @return				true if the data will probably be found in a fast backend on a future request.
-	 */
-	virtual TBitArray<> TryToPrefetch(TConstArrayView<FString> CacheKeys) override;
-
-	/**
-	 * Allows the DDC backend to determine if it wants to cache the provided data. Reasons for returning false could be a slow connection,
-	 * a file size limit, etc.
-	 */
-	virtual bool WouldCache(const TCHAR* CacheKey, TArrayView<const uint8> InData) override;
-
-	/**
-	 * Synchronous retrieve of a cache item
-	 *
-	 * @param	CacheKey	Alphanumeric+underscore key of this cache item
-	 * @param	OutData		Buffer to receive the results, if any were found
-	 * @return				true if any data was found, and in this case OutData is non-empty
-	 */
-	virtual bool GetCachedData(const TCHAR* CacheKey, TArray<uint8>& OutData) override;
-
-	/**
-	 * Asynchronous, fire-and-forget placement of a cache item
-	 *
-	 * @param	CacheKey	Alphanumeric+underscore key of this cache item
-	 * @param	InData		Buffer containing the data to cache, can be destroyed after the call returns, immediately
-	 * @param	bPutEvenIfExists	If true, then do not attempt skip the put even if CachedDataProbablyExists returns true
-	 */
-	virtual EPutStatus PutCachedData(const TCHAR* CacheKey, TArrayView<const uint8> InData, bool bPutEvenIfExists) override;
-
-	virtual void RemoveCachedData(const TCHAR* CacheKey, bool bTransient) override;
-
-	virtual TSharedRef<FDerivedDataCacheStatsNode> GatherUsageStats() const override;
-
-	virtual bool ApplyDebugOptions(FBackendDebugOptions& InOptions) override;
+	FCacheStoreAsync(ILegacyCacheStore* InInnerCache, ECacheStoreFlags InInnerFlags, bool bCacheInFlightPuts);
 
 	virtual void Put(
 		TConstArrayView<FCachePutRequest> Requests,
@@ -213,7 +102,7 @@ public:
 
 	virtual bool LegacyDebugOptions(FBackendDebugOptions& Options) override
 	{
-		return InnerBackend->LegacyDebugOptions(Options);
+		return InnerCache->LegacyDebugOptions(Options);
 	}
 
 private:
@@ -227,364 +116,29 @@ private:
 		OnCompleteType&& OnComplete,
 		OnExecuteType&& OnExecute);
 
+	ILegacyCacheStore* InnerCache;
+	ECacheStoreFlags InnerFlags;
+	TUniquePtr<ILegacyCacheStore> MemoryCache;
 	FDerivedDataCacheUsageStats UsageStats;
-	FDerivedDataCacheUsageStats PutSyncUsageStats;
-
-	/** Backend to use for storage, my responsibilities are about async puts **/
-	FDerivedDataBackendInterface*					InnerBackend;
-	/** Memory based cache to deal with gets that happen while an async put is still in flight **/
-	TUniquePtr<FDerivedDataBackendInterface>		InflightCache;
-	/** We remember outstanding puts so that we don't do them redundantly **/
-	FThreadSet										FilesInFlight;
 };
 
-/** 
- * Async task to handle the fire and forget async put
- */
-class FCachePutAsyncWorker
-{
-public:
-	/** Cache Key for the put to InnerBackend **/
-	FString								CacheKey;
-	/** Data for the put to InnerBackend **/
-	TArray<uint8>						Data;
-	/** Backend to use for storage, my responsibilities are about async puts **/
-	FDerivedDataBackendInterface*		InnerBackend;
-	/** Memory based cache to clear once the put is finished **/
-	FDerivedDataBackendInterface*		InflightCache;
-	/** We remember outstanding puts so that we don't do them redundantly **/
-	FThreadSet*							FilesInFlight;
-	/**If true, then do not attempt skip the put even if CachedDataProbablyExists returns true **/
-	bool								bPutEvenIfExists;
-	/** Usage stats to track thread times. */
-	FDerivedDataCacheUsageStats&        UsageStats;
-
-	/** Constructor
-	*/
-	FCachePutAsyncWorker(const TCHAR* InCacheKey, TArrayView<const uint8> InData, FDerivedDataBackendInterface* InInnerBackend, bool InbPutEvenIfExists, FDerivedDataBackendInterface* InInflightCache, FThreadSet* InInFilesInFlight, FDerivedDataCacheUsageStats& InUsageStats)
-		: CacheKey(InCacheKey)
-		, InnerBackend(InInnerBackend)
-		, InflightCache(InInflightCache)
-		, FilesInFlight(InInFilesInFlight)
-		, bPutEvenIfExists(InbPutEvenIfExists)
-		, UsageStats(InUsageStats)
-	{
-		// Only make a copy if it's not going to be available from the Inflight cache
-		if (!InInflightCache || !InInflightCache->CachedDataProbablyExists(InCacheKey))
-		{
-			Data = InData;
-		}
-		check(InnerBackend);
-	}
-
-	bool ShouldAbortForShutdown()
-	{
-		using ESpeedClass = FDerivedDataBackendInterface::ESpeedClass;
-		ESpeedClass SpeedClass = InnerBackend->GetSpeedClass();
-		if (SpeedClass == ESpeedClass::Local)
-		{
-			return false;
-		}
-		return !GIsBuildMachine && FDerivedDataBackend::Get().IsShuttingDown();
-	}
-		
-	/** Call the inner backend and when that completes, remove the memory cache */
-	void DoWork()
-	{
-		TRACE_CPUPROFILER_EVENT_SCOPE(DDCPut_DoWork);
-		COOK_STAT(auto Timer = UsageStats.TimePut());
-
-		if (ShouldAbortForShutdown())
-		{
-			Abandon();
-			return;
-		}
-
-		using EPutStatus = FDerivedDataBackendInterface::EPutStatus;
-		EPutStatus Status = EPutStatus::NotCached;
-
-		if (!bPutEvenIfExists && InnerBackend->CachedDataProbablyExists(*CacheKey))
-		{
-			Status = EPutStatus::Cached;
-		}
-		else
-		{
-			if (InflightCache && Data.Num() == 0)
-			{
-				// We verified at construction time that we would be able to get the data from the Inflight cache
-				verify(InflightCache->GetCachedData(*CacheKey, Data));
-			}
-			Status = InnerBackend->PutCachedData(*CacheKey, Data, bPutEvenIfExists);
-			COOK_STAT(Timer.AddHit(Data.Num()));
-		}
-
-		if (InflightCache)
-		{
-			// if the data was not cached synchronously, retry
-			if (Status != EPutStatus::Cached && Status != EPutStatus::Skipped)
-			{
-				// retry after a brief wait
-				FPlatformProcess::SleepNoStats(0.2f);
-
-				if (Status == EPutStatus::Executing && InnerBackend->CachedDataProbablyExists(*CacheKey))
-				{
-					Status = EPutStatus::Cached;
-				}
-				else
-				{
-					if (Data.Num() == 0)
-					{
-						verify(InflightCache->GetCachedData(*CacheKey, Data));
-					}
-					Status = InnerBackend->PutCachedData(*CacheKey, Data, /*bPutEvenIfExists*/ false);
-				}
-			}
-
-			switch (Status)
-			{
-			case EPutStatus::Skipped:
-			case EPutStatus::Cached:
-				// remove this from the in-flight cache because the inner cache contains the data or it was intentionally skipped
-				InflightCache->RemoveCachedData(*CacheKey, /*bTransient*/ false);
-				break;
-			case EPutStatus::NotCached:
-				UE_LOG(LogDerivedDataCache, Display, TEXT("%s: Put failed, keeping in memory copy %s."), *InnerBackend->GetName(), *CacheKey);
-				if (uint32 ErrorCode = FPlatformMisc::GetLastError())
-				{
-					TCHAR ErrorBuffer[1024];
-					FPlatformMisc::GetSystemErrorMessage(ErrorBuffer, 1024, ErrorCode);
-					UE_LOG(LogDerivedDataCache, Display, TEXT("Failed to write %s to %s. Error: %u (%s)"), *CacheKey, *InnerBackend->GetName(), ErrorCode, ErrorBuffer);
-				}
-				break;
-			case EPutStatus::Executing:
-				UE_LOG(LogDerivedDataCache, Display, TEXT("%s: Put not finished executing, keeping in memory copy %s."), *InnerBackend->GetName(), *CacheKey);
-				break;
-			default:
-				break;
-			}
-		}
-
-		FilesInFlight->Remove(CacheKey);
-		FDerivedDataBackend::Get().AddToAsyncCompletionCounter(-1);
-		UE_LOG(LogDerivedDataCache, VeryVerbose, TEXT("%s: Completed AsyncPut of %s."), *InnerBackend->GetName(), *CacheKey);
-	}
-
-	FORCEINLINE TStatId GetStatId() const
-	{
-		RETURN_QUICK_DECLARE_CYCLE_STAT(FCachePutAsyncWorker, STATGROUP_ThreadPoolAsyncTasks);
-	}
-
-	/** Indicates to the thread pool that this task is abandonable */
-	bool CanAbandon()
-	{
-		return true;
-	}
-
-	/** Abandon routine, we need to remove the item from the in flight cache because something might be waiting for that */
-	void Abandon()
-	{
-		if (InflightCache)
-		{
-			InflightCache->RemoveCachedData(*CacheKey, /*bTransient=*/ false); // we can remove this from the temp cache, since the real cache will hit now
-		}
-		FilesInFlight->Remove(CacheKey);
-		FDerivedDataBackend::Get().AddToAsyncCompletionCounter(-1);
-		UE_LOG(LogDerivedDataCache, VeryVerbose, TEXT("%s: Abandoned AsyncPut of %s."), *InnerBackend->GetName(), *CacheKey);
-	}
-};
-
-FDerivedDataBackendAsyncPutWrapper::FDerivedDataBackendAsyncPutWrapper(FDerivedDataBackendInterface* InInnerBackend, bool bCacheInFlightPuts)
-	: InnerBackend(InInnerBackend)
-	, InflightCache(bCacheInFlightPuts ? Memory::CreateMemoryDerivedDataBackend(TEXT("InflightMemoryCache"), /*MaxCacheSize*/ -1, /*bCanBeDisabled*/ false) : nullptr)
-{
-	check(InnerBackend);
-}
-
-/** return true if this cache is writable **/
-bool FDerivedDataBackendAsyncPutWrapper::IsWritable() const
-{
-	return InnerBackend->IsWritable();
-}
-
-FDerivedDataBackendInterface::ESpeedClass FDerivedDataBackendAsyncPutWrapper::GetSpeedClass() const
-{
-	return InnerBackend->GetSpeedClass();
-}
-
-bool FDerivedDataBackendAsyncPutWrapper::BackfillLowerCacheLevels() const
-{
-	return InnerBackend->BackfillLowerCacheLevels();
-}
-
-bool FDerivedDataBackendAsyncPutWrapper::CachedDataProbablyExists(const TCHAR* CacheKey)
-{
-	COOK_STAT(auto Timer = UsageStats.TimeProbablyExists());
-	bool Result = (InflightCache && InflightCache->CachedDataProbablyExists(CacheKey)) || InnerBackend->CachedDataProbablyExists(CacheKey);
-	COOK_STAT(if (Result) {	Timer.AddHit(0); });
-
-	UE_LOG(LogDerivedDataCache, Verbose, TEXT("%s CachedDataProbablyExists=%d for %s"), *GetName(), Result, CacheKey);
-	return Result;
-}
-
-TBitArray<> FDerivedDataBackendAsyncPutWrapper::CachedDataProbablyExistsBatch(TConstArrayView<FString> CacheKeys)
-{
-	COOK_STAT(auto Timer = UsageStats.TimeProbablyExists());
-
-	TBitArray<> Result;
-	if (InflightCache)
-	{
-		Result = InflightCache->CachedDataProbablyExistsBatch(CacheKeys);
-		check(Result.Num() == CacheKeys.Num());
-		if (Result.CountSetBits() < CacheKeys.Num())
-		{
-			TBitArray<> InnerResult = InnerBackend->CachedDataProbablyExistsBatch(CacheKeys);
-			check(InnerResult.Num() == CacheKeys.Num());
-			Result.CombineWithBitwiseOR(InnerResult, EBitwiseOperatorFlags::MaintainSize);
-		}
-	}
-	else
-	{
-		Result = InnerBackend->CachedDataProbablyExistsBatch(CacheKeys);
-		check(Result.Num() == CacheKeys.Num());
-	}
-
-	COOK_STAT(if (Result.CountSetBits() == CacheKeys.Num()) { Timer.AddHit(0); });
-	UE_LOG(LogDerivedDataCache, Verbose, TEXT("%s CachedDataProbablyExists found %d/%d keys"), *GetName(), Result.CountSetBits(), CacheKeys.Num());
-	return Result;
-}
-
-TBitArray<> FDerivedDataBackendAsyncPutWrapper::TryToPrefetch(TConstArrayView<FString> CacheKeys)
+FCacheStoreAsync::FCacheStoreAsync(ILegacyCacheStore* InInnerCache, ECacheStoreFlags InInnerFlags, bool bCacheInFlightPuts)
+	: InnerCache(InInnerCache)
+	, InnerFlags(InInnerFlags)
+	, MemoryCache(bCacheInFlightPuts ? CacheStore::Memory::CreateMemoryDerivedDataBackend(TEXT("InflightMemoryCache"), /*MaxCacheSize*/ -1, /*bCanBeDisabled*/ false) : nullptr)
 {
-	COOK_STAT(auto Timer = UsageStats.TimePrefetch());
-
-	TBitArray<> Result;
-	if (InflightCache)
-	{
-		Result = InflightCache->CachedDataProbablyExistsBatch(CacheKeys);
-	}
-	Result = TBitArray<>::BitwiseOR(Result, InnerBackend->TryToPrefetch(CacheKeys), EBitwiseOperatorFlags::MaxSize);
-
-	if (Result.CountSetBits() == CacheKeys.Num())
-	{
-		COOK_STAT(Timer.AddHit(0));
-	}
-
-	return Result;
+	check(InnerCache);
 }
 
-/*
-	Determine if we would cache this by asking all our inner layers
-*/
-bool FDerivedDataBackendAsyncPutWrapper::WouldCache(const TCHAR* CacheKey, TArrayView<const uint8> InData)
+void FCacheStoreAsync::LegacyStats(FDerivedDataCacheStatsNode& OutNode)
 {
-	return InnerBackend->WouldCache(CacheKey, InData);
-}
+	OutNode = {TEXT("AsyncWrapper"), TEXT(""), EnumHasAnyFlags(InnerFlags, ECacheStoreFlags::Local)};
+	OutNode.Stats.Add(TEXT(""), UsageStats);
 
-bool FDerivedDataBackendAsyncPutWrapper::ApplyDebugOptions(FBackendDebugOptions& InOptions)
-{
-	return InnerBackend->ApplyDebugOptions(InOptions);
-}
-
-bool FDerivedDataBackendAsyncPutWrapper::GetCachedData(const TCHAR* CacheKey, TArray<uint8>& OutData)
-{
-	COOK_STAT(auto Timer = UsageStats.TimeGet());
-	if (InflightCache && InflightCache->GetCachedData(CacheKey, OutData))
-	{
-		COOK_STAT(Timer.AddHit(OutData.Num()));
-		UE_LOG(LogDerivedDataCache, Verbose, TEXT("%s CacheHit from InFlightCache on %s"), *GetName(), CacheKey);
-		return true;
-	}
-
-	bool bSuccess = InnerBackend->GetCachedData(CacheKey, OutData);
-	if (bSuccess)
+	InnerCache->LegacyStats(OutNode.Children.Add_GetRef(MakeShared<FDerivedDataCacheStatsNode>()).Get());
+	if (MemoryCache)
 	{
-		UE_LOG(LogDerivedDataCache, Verbose, TEXT("%s Cache hit on %s"), *GetName(), CacheKey);
-		COOK_STAT(Timer.AddHit(OutData.Num()));
-	}
-	else
-	{
-		UE_LOG(LogDerivedDataCache, Verbose, TEXT("%s Cache miss on %s"), *GetName(), CacheKey);
-	}
-	return bSuccess;
-}
-
-FDerivedDataBackendInterface::EPutStatus FDerivedDataBackendAsyncPutWrapper::PutCachedData(const TCHAR* CacheKey, TArrayView<const uint8> InData, bool bPutEvenIfExists)
-{
-	COOK_STAT(auto Timer = PutSyncUsageStats.TimePut());
-
-	if (!InnerBackend->IsWritable())
-	{
-		return EPutStatus::NotCached; // no point in continuing down the chain
-	}
-	const bool bAdded = FilesInFlight.AddIfNotExists(CacheKey);
-	if (!bAdded)
-	{
-		return EPutStatus::Executing; // if it is already on its way, we don't need to send it again
-	}
-	if (InflightCache)
-	{
-		if (InflightCache->CachedDataProbablyExists(CacheKey))
-		{
-			UE_LOG(LogDerivedDataCache, Verbose, TEXT("%s skipping out of key already in in-flight cache %s"), *GetName(), CacheKey);
-			return EPutStatus::Executing; // if it is already on its way, we don't need to send it again
-		}
-		InflightCache->PutCachedData(CacheKey, InData, true); // temp copy stored in memory while the async task waits to complete
-		COOK_STAT(Timer.AddHit(InData.Num()));
-	}
-	
-	UE_LOG(LogDerivedDataCache, Verbose, TEXT("%s queueing %s for put"), *GetName(), CacheKey);
-
-	FDerivedDataBackend::Get().AddToAsyncCompletionCounter(1);
-	(new FAutoDeleteAsyncTask<FCachePutAsyncWorker>(CacheKey, InData, InnerBackend, bPutEvenIfExists, InflightCache.Get(), &FilesInFlight, UsageStats))->StartBackgroundTask(Private::GCacheThreadPool, EQueuedWorkPriority::Low);
-
-	return EPutStatus::Executing;
-}
-
-void FDerivedDataBackendAsyncPutWrapper::RemoveCachedData(const TCHAR* CacheKey, bool bTransient)
-{
-	if (!InnerBackend->IsWritable())
-	{
-		return; // no point in continuing down the chain
-	}
-	while (FilesInFlight.Exists(CacheKey))
-	{
-		FPlatformProcess::Sleep(0.0f); // this is an exception condition (corruption), spin and wait for it to clear
-	}
-	if (InflightCache)
-	{
-		InflightCache->RemoveCachedData(CacheKey, bTransient);
-	}
-	InnerBackend->RemoveCachedData(CacheKey, bTransient);
-
-	UE_LOG(LogDerivedDataCache, Verbose, TEXT("%s removed %s"), *GetName(), CacheKey)
-}
-
-TSharedRef<FDerivedDataCacheStatsNode> FDerivedDataBackendAsyncPutWrapper::GatherUsageStats() const
-{
-	TSharedRef<FDerivedDataCacheStatsNode> Usage =
-		MakeShared<FDerivedDataCacheStatsNode>(TEXT("AsyncPutWrapper"), TEXT(""), InnerBackend->GetSpeedClass() == ESpeedClass::Local);
-	Usage->Stats.Add(TEXT("AsyncPut"), UsageStats);
-	Usage->Stats.Add(TEXT("AsyncPutSync"), PutSyncUsageStats);
-
-	Usage->Children.Add(InnerBackend->GatherUsageStats());
-	if (InflightCache)
-	{
-		Usage->Children.Add(InflightCache->GatherUsageStats());
-	}
-
-	return Usage;
-}
-
-void FDerivedDataBackendAsyncPutWrapper::LegacyStats(FDerivedDataCacheStatsNode& OutNode)
-{
-	OutNode = {TEXT("AsyncPutWrapper"), TEXT(""), InnerBackend->GetSpeedClass() == ESpeedClass::Local};
-	OutNode.Stats.Add(TEXT("AsyncPut"), UsageStats);
-	OutNode.Stats.Add(TEXT("AsyncPutSync"), PutSyncUsageStats);
-
-	InnerBackend->LegacyStats(OutNode.Children.Add_GetRef(MakeShared<FDerivedDataCacheStatsNode>()).Get());
-	if (InflightCache)
-	{
-		InflightCache->LegacyStats(OutNode.Children.Add_GetRef(MakeShared<FDerivedDataCacheStatsNode>()).Get());
+		MemoryCache->LegacyStats(OutNode.Children.Add_GetRef(MakeShared<FDerivedDataCacheStatsNode>()).Get());
 	}
 }
 
@@ -698,7 +252,7 @@ private:
 };
 
 template <typename RequestType, typename OnCompleteType, typename OnExecuteType>
-void FDerivedDataBackendAsyncPutWrapper::Execute(
+void FCacheStoreAsync::Execute(
 	COOK_STAT(CookStatsFunction OnAddStats,)
 	const TConstArrayView<RequestType> Requests,
 	IRequestOwner& Owner,
@@ -707,16 +261,13 @@ void FDerivedDataBackendAsyncPutWrapper::Execute(
 {
 	auto ExecuteWithStats = [this, COOK_STAT(OnAddStats,) OnExecute](TConstArrayView<RequestType> Requests, IRequestOwner& Owner, OnCompleteType&& OnComplete) mutable
 	{
-		Invoke(OnExecute, *InnerBackend, Requests, Owner, [this, COOK_STAT(OnAddStats,) OnComplete = MoveTemp(OnComplete)](auto&& Response)
+		Invoke(OnExecute, *InnerCache, Requests, Owner, [this, COOK_STAT(OnAddStats,) OnComplete = MoveTemp(OnComplete)](auto&& Response)
 		{
 			if (Response.Status == EStatus::Ok)
 			{
-				COOK_STAT((UsageStats.*OnAddStats)().AddHit(0));
-			}
-			if (OnComplete)
-			{
-				OnComplete(MoveTemp(Response));
+				COOK_STAT(Invoke(OnAddStats, UsageStats).AddHit(0));
 			}
+			OnComplete(MoveTemp(Response));
 			FDerivedDataBackend::Get().AddToAsyncCompletionCounter(-1);
 		});
 	};
@@ -743,9 +294,9 @@ void FDerivedDataBackendAsyncPutWrapper::Execute(
 	Request->Start(Owner.GetPriority());
 }
 
-FDerivedDataBackendInterface* CreateAsyncPutDerivedDataBackend(FDerivedDataBackendInterface* InnerBackend, bool bCacheInFlightPuts)
+ILegacyCacheStore* CreateCacheStoreAsync(ILegacyCacheStore* InnerCache, ECacheStoreFlags InnerFlags, bool bCacheInFlightPuts)
 {
-	return new FDerivedDataBackendAsyncPutWrapper(InnerBackend, bCacheInFlightPuts);
+	return new FCacheStoreAsync(InnerCache, InnerFlags, bCacheInFlightPuts);
 }
 
-} // UE::DerivedData::CacheStore::AsyncPut
+} // UE::DerivedData
diff --git a/Engine/Source/Developer/DerivedDataCache/Private/DerivedDataCacheStoreHierarchy.cpp b/Engine/Source/Developer/DerivedDataCache/Private/DerivedDataCacheStoreHierarchy.cpp
--- a/Engine/Source/Developer/DerivedDataCache/Private/DerivedDataCacheStoreHierarchy.cpp
+++ b/Engine/Source/Developer/DerivedDataCache/Private/DerivedDataCacheStoreHierarchy.cpp
@@ -19,13 +19,10 @@
 #include "Templates/UniquePtr.h"
 #include <atomic>
 
-namespace UE::DerivedData::CacheStore
+namespace UE::DerivedData
 {
 
-ILegacyCacheStore* CreateAsyncPutCacheStore(ILegacyCacheStore* InnerCache, ECacheStoreFlags InnerFlags, bool bCacheInFlightPuts)
-{
-	return nullptr;
-}
+ILegacyCacheStore* CreateCacheStoreAsync(ILegacyCacheStore* InnerCache, ECacheStoreFlags InnerFlags, bool bCacheInFlightPuts);
 
 ///////////////////////////////////////////////////////////////////////////////////////////////////////////////////////
 
@@ -129,7 +126,6 @@ private:
 	mutable FRWLock NodesLock;
 	ECacheStoreNodeFlags CombinedNodeFlags{};
 	TArray<FCacheStoreNode, TInlineAllocator<8>> Nodes;
-	FDerivedDataCacheUsageStats UsageStats;
 };
 
 ///////////////////////////////////////////////////////////////////////////////////////////////////////////////////////
@@ -220,7 +216,7 @@ void FCacheStoreHierarchy::Add(ILegacyCacheStore* CacheStore, ECacheStoreFlags F
 	FWriteScopeLock Lock(NodesLock);
 	checkf(!Algo::FindBy(Nodes, CacheStore, &FCacheStoreNode::Cache),
 		TEXT("Attempting to add a cache store that was previously registered to the hierarchy."));
-	TUniquePtr<ILegacyCacheStore> AsyncCacheStore(CreateAsyncPutCacheStore(CacheStore, Flags, /*bCacheInFlightPuts*/ false));
+	TUniquePtr<ILegacyCacheStore> AsyncCacheStore(CreateCacheStoreAsync(CacheStore, Flags, /*bCacheInFlightPuts*/ false));
 	Nodes.Add({CacheStore, Flags, {}, MoveTemp(AsyncCacheStore)});
 	UpdateNodeFlags();
 }
@@ -1243,7 +1239,6 @@ void FCacheStoreHierarchy::LegacyDelete(
 void FCacheStoreHierarchy::LegacyStats(FDerivedDataCacheStatsNode& OutNode)
 {
 	FReadScopeLock Lock(NodesLock);
-	OutNode.Stats.Add(TEXT(""), UsageStats);
 	OutNode.Children.Reserve(Nodes.Num());
 	for (const FCacheStoreNode& Node : Nodes)
 	{
@@ -1265,4 +1260,4 @@ ILegacyCacheStore* CreateCacheStoreHierarchy(ICacheStoreOwner*& OutOwner)
 	return Hierarchy;
 }
 
-} // UE::DerivedData::CacheStore
+} // UE::DerivedData
diff --git a/Engine/Source/Developer/DerivedDataCache/Private/DerivedDataCacheStoreThrottle.cpp b/Engine/Source/Developer/DerivedDataCache/Private/DerivedDataCacheStoreThrottle.cpp
new file mode 100644
--- /dev/null
+++ b/Engine/Source/Developer/DerivedDataCache/Private/DerivedDataCacheStoreThrottle.cpp
@@ -0,0 +1,265 @@
+// Copyright Epic Games, Inc. All Rights Reserved.
+
+#include "Algo/Find.h"
+#include "Algo/Transform.h"
+#include "DerivedDataCachePrivate.h"
+#include "DerivedDataCacheRecord.h"
+#include "DerivedDataLegacyCacheStore.h"
+#include "HAL/PlatformProcess.h"
+#include "HAL/PlatformTime.h"
+#include <atomic>
+
+namespace UE::DerivedData
+{
+
+/**
+ * A cache store that increases the latency and reduces the throughput of another cache store.
+ * 1. Reproduce timings for a remote cache with a local cache, to reduce both network usage and measurement noise.
+ * 2. Reproduce HDD latency and throughput even when data is stored on SSD.
+ */
+class FCacheStoreThrottle final : public ILegacyCacheStore
+{
+public:
+	FCacheStoreThrottle(ILegacyCacheStore* InInnerCache, uint32 InLatencyMS, uint32 InMaxBytesPerSecond)
+		: InnerCache(InInnerCache)
+		, Latency(float(InLatencyMS) / 1000.0f)
+		, MaxBytesPerSecond(InMaxBytesPerSecond)
+	{
+		check(InnerCache);
+	}
+
+	void Put(
+		const TConstArrayView<FCachePutRequest> Requests,
+		IRequestOwner& Owner,
+		FOnCachePutComplete&& OnComplete) final
+	{
+		struct FRecordSize
+		{
+			FCacheKey Key;
+			uint64 Size;
+		};
+		TArray<FRecordSize, TInlineAllocator<1>> RecordSizes;
+		RecordSizes.Reserve(Requests.Num());
+		Algo::Transform(Requests, RecordSizes, [](const FCachePutRequest& Request) -> FRecordSize
+		{
+			return {Request.Record.GetKey(), Private::GetCacheRecordCompressedSize(Request.Record)};
+		});
+
+		InnerCache->Put(Requests, Owner,
+			[this, RecordSizes = MoveTemp(RecordSizes), State = EnterThrottlingScope(), OnComplete = MoveTemp(OnComplete)](FCachePutResponse&& Response)
+			{
+				const FRecordSize* Size = Algo::FindBy(RecordSizes, Response.Key, &FRecordSize::Key);
+				CloseThrottlingScope(State, FThrottlingState(this, Size ? Size->Size : 0));
+				if (OnComplete)
+				{
+					OnComplete(MoveTemp(Response));
+				}
+			});
+	}
+
+	void Get(
+		const TConstArrayView<FCacheGetRequest> Requests,
+		IRequestOwner& Owner,
+		FOnCacheGetComplete&& OnComplete) final
+	{
+		InnerCache->Get(Requests, Owner,
+			[this, State = EnterThrottlingScope(), OnComplete = MoveTemp(OnComplete)](FCacheGetResponse&& Response)
+			{
+				CloseThrottlingScope(State, FThrottlingState(this, Private::GetCacheRecordCompressedSize(Response.Record)));
+				if (OnComplete)
+				{
+					OnComplete(MoveTemp(Response));
+				}
+			});
+	}
+
+	void PutValue(
+		const TConstArrayView<FCachePutValueRequest> Requests,
+		IRequestOwner& Owner,
+		FOnCachePutValueComplete&& OnComplete) final
+	{
+		struct FValueSize
+		{
+			FCacheKey Key;
+			uint64 Size;
+		};
+		TArray<FValueSize, TInlineAllocator<1>> ValueSizes;
+		ValueSizes.Reserve(Requests.Num());
+		Algo::Transform(Requests, ValueSizes, [](const FCachePutValueRequest& Request) -> FValueSize
+		{
+			return {Request.Key, Request.Value.GetData().GetCompressedSize()};
+		});
+
+		InnerCache->PutValue(Requests, Owner,
+			[this, ValueSizes = MoveTemp(ValueSizes), State = EnterThrottlingScope(), OnComplete = MoveTemp(OnComplete)](FCachePutValueResponse&& Response)
+			{
+				const FValueSize* Size = Algo::FindBy(ValueSizes, Response.Key, &FValueSize::Key);
+				CloseThrottlingScope(State, FThrottlingState(this, Size ? Size->Size : 0));
+				if (OnComplete)
+				{
+					OnComplete(MoveTemp(Response));
+				}
+			});
+	}
+
+	void GetValue(
+		const TConstArrayView<FCacheGetValueRequest> Requests,
+		IRequestOwner& Owner,
+		FOnCacheGetValueComplete&& OnComplete) final
+	{
+		InnerCache->GetValue(Requests, Owner,
+			[this, State = EnterThrottlingScope(), OnComplete = MoveTemp(OnComplete)](FCacheGetValueResponse&& Response)
+			{
+				CloseThrottlingScope(State, FThrottlingState(this, Response.Value.GetData().GetCompressedSize()));
+				if (OnComplete)
+				{
+					OnComplete(MoveTemp(Response));
+				}
+			});
+	}
+
+	void GetChunks(
+		const TConstArrayView<FCacheGetChunkRequest> Requests,
+		IRequestOwner& Owner,
+		FOnCacheGetChunkComplete&& OnComplete) final
+	{
+		InnerCache->GetChunks(Requests, Owner,
+			[this, State = EnterThrottlingScope(), OnComplete = MoveTemp(OnComplete)](FCacheGetChunkResponse&& Response)
+			{
+				CloseThrottlingScope(State, FThrottlingState(this, Response.RawData.GetSize()));
+				if (OnComplete)
+				{
+					OnComplete(MoveTemp(Response));
+				}
+			});
+
+	}
+
+	void LegacyPut(
+		const TConstArrayView<FLegacyCachePutRequest> Requests,
+		IRequestOwner& Owner,
+		FOnLegacyCachePutComplete&& OnComplete) final
+	{
+		struct FValueSize
+		{
+			FCacheKey Key;
+			uint64 Size;
+		};
+		TArray<FValueSize, TInlineAllocator<1>> ValueSizes;
+		ValueSizes.Reserve(Requests.Num());
+		Algo::Transform(Requests, ValueSizes, [](const FLegacyCachePutRequest& Request) -> FValueSize
+		{
+			return {Request.Key.GetKey(), Request.Value.GetSize()};
+		});
+
+		InnerCache->LegacyPut(Requests, Owner,
+			[this, ValueSizes = MoveTemp(ValueSizes), State = EnterThrottlingScope(), OnComplete = MoveTemp(OnComplete)](FLegacyCachePutResponse&& Response)
+			{
+				const FValueSize* Size = Algo::FindBy(ValueSizes, Response.Key.GetKey(), &FValueSize::Key);
+				CloseThrottlingScope(State, FThrottlingState(this, Size ? Size->Size : 0));
+				if (OnComplete)
+				{
+					OnComplete(MoveTemp(Response));
+				}
+			});
+	}
+
+	void LegacyGet(
+		const TConstArrayView<FLegacyCacheGetRequest> Requests,
+		IRequestOwner& Owner,
+		FOnLegacyCacheGetComplete&& OnComplete) final
+	{
+		InnerCache->LegacyGet(Requests, Owner,
+			[this, State = EnterThrottlingScope(), OnComplete = MoveTemp(OnComplete)](FLegacyCacheGetResponse&& Response)
+			{
+				CloseThrottlingScope(State, FThrottlingState(this, Response.Value.GetSize()));
+				if (OnComplete)
+				{
+					OnComplete(MoveTemp(Response));
+				}
+			});
+	}
+
+	void LegacyDelete(
+		const TConstArrayView<FLegacyCacheDeleteRequest> Requests,
+		IRequestOwner& Owner,
+		FOnLegacyCacheDeleteComplete&& OnComplete) final
+	{
+		InnerCache->LegacyDelete(Requests, Owner,
+			[this, State = EnterThrottlingScope(), OnComplete = MoveTemp(OnComplete)](FLegacyCacheDeleteResponse&& Response)
+			{
+				CloseThrottlingScope(State, FThrottlingState(this, 0));
+				if (OnComplete)
+				{
+					OnComplete(MoveTemp(Response));
+				}
+			});
+	}
+
+	void LegacyStats(FDerivedDataCacheStatsNode& OutNode) final
+	{
+		InnerCache->LegacyStats(OutNode);
+	}
+
+	bool LegacyDebugOptions(FBackendDebugOptions& Options) final
+	{
+		return InnerCache->LegacyDebugOptions(Options);
+	}
+
+private:
+	struct FThrottlingState
+	{
+		double Time;
+		uint64 TotalBytesTransferred;
+
+		explicit FThrottlingState(FCacheStoreThrottle* ThrottleWrapper)
+		: Time(FPlatformTime::Seconds())
+		, TotalBytesTransferred(ThrottleWrapper->TotalBytesTransferred.load(std::memory_order_relaxed))
+		{
+		}
+
+		explicit FThrottlingState(FCacheStoreThrottle* ThrottleWrapper, uint64 BytesTransferred)
+		: Time(FPlatformTime::Seconds())
+		, TotalBytesTransferred(ThrottleWrapper->TotalBytesTransferred.fetch_add(BytesTransferred, std::memory_order_relaxed) + BytesTransferred)
+		{
+		}
+	};
+
+	FThrottlingState EnterThrottlingScope()
+	{
+		if (Latency > 0)
+		{
+			TRACE_CPUPROFILER_EVENT_SCOPE(ThrottlingLatency);
+			FPlatformProcess::Sleep(Latency);
+		}
+		return FThrottlingState(this);
+	}
+
+	void CloseThrottlingScope(FThrottlingState PreviousState, FThrottlingState CurrentState)
+	{
+		if (MaxBytesPerSecond)
+		{
+			// Take into account any other transfer that might have happened during that time from any other thread so we have a global limit
+			const double ExpectedTime = double(CurrentState.TotalBytesTransferred - PreviousState.TotalBytesTransferred) / MaxBytesPerSecond;
+			const double ActualTime = CurrentState.Time - PreviousState.Time;
+			if (ExpectedTime > ActualTime)
+			{
+				TRACE_CPUPROFILER_EVENT_SCOPE(ThrottlingBandwidth);
+				FPlatformProcess::Sleep(ExpectedTime - ActualTime);
+			}
+		}
+	}
+
+	/** Backend to use for storage, my responsibilities are about throttling **/
+	ILegacyCacheStore* InnerCache;
+	float Latency;
+	uint32 MaxBytesPerSecond;
+	std::atomic<uint64> TotalBytesTransferred{0};
+};
+
+ILegacyCacheStore* CreateCacheStoreThrottle(ILegacyCacheStore* InnerCache, uint32 LatencyMS, uint32 MaxBytesPerSecond)
+{
+	return new FCacheStoreThrottle(InnerCache, LatencyMS, MaxBytesPerSecond);
+}
+
+} // UE::DerivedData
diff --git a/Engine/Source/Developer/DerivedDataCache/Private/DerivedDataCacheStoreVerify.cpp b/Engine/Source/Developer/DerivedDataCache/Private/DerivedDataCacheStoreVerify.cpp
new file mode 100644
--- /dev/null
+++ b/Engine/Source/Developer/DerivedDataCache/Private/DerivedDataCacheStoreVerify.cpp
@@ -0,0 +1,721 @@
+// Copyright Epic Games, Inc. All Rights Reserved.
+
+#include "Algo/Compare.h"
+#include "Compression/OodleDataCompression.h"
+#include "Containers/Set.h"
+#include "DerivedDataCacheKey.h"
+#include "DerivedDataLegacyCacheStore.h"
+#include "HAL/CriticalSection.h"
+#include "Misc/Paths.h"
+#include "Misc/PathViews.h"
+#include "Misc/ScopeLock.h"
+#include "Misc/ScopeRWLock.h"
+#include "Misc/FileHelper.h"
+#include "Templates/Tuple.h"
+#include "Templates/UniquePtr.h"
+
+namespace UE::DerivedData
+{
+
+/**
+ * A cache store that verifies that derived data is generated deterministically.
+ *
+ * This wraps a cache store and fails every get until a matching put occurs, then compares the derived data.
+ */
+class FCacheStoreVerify : public ILegacyCacheStore
+{
+public:
+	FCacheStoreVerify(ILegacyCacheStore* InInnerCache, bool bInPutOnError)
+		: InnerCache(InInnerCache)
+		, bPutOnError(bInPutOnError)
+	{
+		check(InnerCache);
+	}
+
+	void Put(
+		TConstArrayView<FCachePutRequest> Requests,
+		IRequestOwner& Owner,
+		FOnCachePutComplete&& OnComplete) final;
+
+	void Get(
+		TConstArrayView<FCacheGetRequest> Requests,
+		IRequestOwner& Owner,
+		FOnCacheGetComplete&& OnComplete) final;
+
+	void PutValue(
+		TConstArrayView<FCachePutValueRequest> Requests,
+		IRequestOwner& Owner,
+		FOnCachePutValueComplete&& OnComplete) final;
+
+	void GetValue(
+		TConstArrayView<FCacheGetValueRequest> Requests,
+		IRequestOwner& Owner,
+		FOnCacheGetValueComplete&& OnComplete) final;
+
+	void GetChunks(
+		TConstArrayView<FCacheGetChunkRequest> Requests,
+		IRequestOwner& Owner,
+		FOnCacheGetChunkComplete&& OnComplete) final;
+
+	void LegacyPut(
+		TConstArrayView<FLegacyCachePutRequest> Requests,
+		IRequestOwner& Owner,
+		FOnLegacyCachePutComplete&& OnComplete) final;
+
+	void LegacyGet(
+		TConstArrayView<FLegacyCacheGetRequest> Requests,
+		IRequestOwner& Owner,
+		FOnLegacyCacheGetComplete&& OnComplete) final;
+
+	void LegacyDelete(
+		TConstArrayView<FLegacyCacheDeleteRequest> Requests,
+		IRequestOwner& Owner,
+		FOnLegacyCacheDeleteComplete&& OnComplete) final
+	{
+		InnerCache->LegacyDelete(Requests, Owner, MoveTemp(OnComplete));
+	}
+
+	void LegacyStats(FDerivedDataCacheStatsNode& OutNode) final
+	{
+		InnerCache->LegacyStats(OutNode);
+	}
+
+	bool LegacyDebugOptions(FBackendDebugOptions& Options) final
+	{
+		return InnerCache->LegacyDebugOptions(Options);
+	}
+
+private:
+	struct FVerifyPutState
+	{
+		TArray<FCachePutRequest> ForwardRequests;
+		TArray<FCachePutRequest> VerifyRequests;
+		FOnCachePutComplete OnComplete;
+		int32 ActiveRequests = 0;
+		FRWLock Lock;
+	};
+
+	struct FVerifyPutValueState
+	{
+		TArray<FCachePutValueRequest> ForwardRequests;
+		TArray<FCachePutValueRequest> VerifyRequests;
+		FOnCachePutValueComplete OnComplete;
+		int32 ActiveRequests = 0;
+		FRWLock Lock;
+	};
+
+	struct FVerifyLegacyPutState
+	{
+		TArray<FLegacyCachePutRequest> ForwardRequests;
+		TArray<FLegacyCachePutRequest> VerifyRequests;
+		FOnLegacyCachePutComplete OnComplete;
+		int32 ActiveRequests = 0;
+		FRWLock Lock;
+	};
+
+	void GetMetaComplete(IRequestOwner& Owner, FVerifyPutState* State, FCacheGetResponse&& Response);
+	void GetDataComplete(IRequestOwner& Owner, FVerifyPutState* State, FCacheGetResponse&& Response);
+	void GetComplete(IRequestOwner& Owner, FVerifyPutState* State);
+
+	void GetMetaComplete(IRequestOwner& Owner, FVerifyPutValueState* State, FCacheGetValueResponse&& Response);
+	void GetDataComplete(IRequestOwner& Owner, FVerifyPutValueState* State, FCacheGetValueResponse&& Response);
+	void GetComplete(IRequestOwner& Owner, FVerifyPutValueState* State);
+
+	void GetMetaComplete(IRequestOwner& Owner, FVerifyLegacyPutState* State, FLegacyCacheGetResponse&& Response);
+	void GetDataComplete(IRequestOwner& Owner, FVerifyLegacyPutState* State, FLegacyCacheGetResponse&& Response);
+	void GetComplete(IRequestOwner& Owner, FVerifyLegacyPutState* State);
+
+	static bool CompareRecords(const FCacheRecord& PutRecord, const FCacheRecord& GetRecord, const FSharedString& Name);
+
+	static void LogChangedValue(const FSharedString& Name, const FCacheKey& Key, const FValueId& Id, const FValue& NewValue, const FValue& OldValue);
+
+private:
+	ILegacyCacheStore* InnerCache;
+	FCriticalSection SynchronizationObject;
+	TSet<FCacheKey> AlreadyTested;
+	bool bPutOnError;
+};
+
+void FCacheStoreVerify::Put(
+	const TConstArrayView<FCachePutRequest> Requests,
+	IRequestOwner& Owner,
+	FOnCachePutComplete&& OnComplete)
+{
+	TUniquePtr<FVerifyPutState> State = MakeUnique<FVerifyPutState>();
+	State->VerifyRequests.Reserve(Requests.Num());
+	{
+		FScopeLock ScopeLock(&SynchronizationObject);
+		for (const FCachePutRequest& Request : Requests)
+		{
+			bool bAlreadyTested = false;
+			AlreadyTested.Add(Request.Record.GetKey(), &bAlreadyTested);
+			(bAlreadyTested ? State->ForwardRequests : State->VerifyRequests).Add(Request);
+		}
+	}
+
+	if (State->VerifyRequests.IsEmpty())
+	{
+		return InnerCache->Put(State->ForwardRequests, Owner, MoveTemp(OnComplete));
+	}
+
+	TArray<FCacheGetRequest> GetMetaRequests;
+	GetMetaRequests.Reserve(State->VerifyRequests.Num());
+	{
+		uint64 PutIndex = 0;
+		const ECachePolicy GetPolicy = ECachePolicy::Query | ECachePolicy::PartialRecord | ECachePolicy::SkipData;
+		for (const FCachePutRequest& PutRequest : State->VerifyRequests)
+		{
+			GetMetaRequests.Add({PutRequest.Name, PutRequest.Record.GetKey(), GetPolicy, PutIndex++});
+		}
+	}
+
+	State->OnComplete = MoveTemp(OnComplete);
+	State->ActiveRequests = GetMetaRequests.Num();
+	InnerCache->Get(GetMetaRequests, Owner, [this, &Owner, State = State.Release()](FCacheGetResponse&& MetaResponse)
+	{
+		GetMetaComplete(Owner, State, MoveTemp(MetaResponse));
+	});
+}
+
+void FCacheStoreVerify::GetMetaComplete(IRequestOwner& Owner, FVerifyPutState* State, FCacheGetResponse&& Response)
+{
+	FCachePutRequest& Request = State->VerifyRequests[int32(Response.UserData)];
+
+	if ((Response.Status == EStatus::Ok) ||
+		(Response.Status == EStatus::Error && !Response.Record.GetValues().IsEmpty()))
+	{
+		const auto MakeValueTuple = [](const FValueWithId& Value) -> TTuple<FValueId, FIoHash>
+		{
+			return MakeTuple(Value.GetId(), Value.GetRawHash());
+		};
+		if (Algo::CompareBy(Request.Record.GetValues(), Response.Record.GetValues(), MakeValueTuple))
+		{
+			UE_LOG(LogDerivedDataCache, Log,
+				TEXT("Verify: Data in the cache matches newly generated data for %s from '%s'."),
+				*WriteToString<96>(Request.Record.GetKey()), *Request.Name);
+			State->OnComplete(Request.MakeResponse(EStatus::Ok));
+		}
+		else
+		{
+			const ECachePolicy Policy = ECachePolicy::Default | ECachePolicy::PartialRecord;
+			const FCacheGetRequest GetDataRequests[]{{Response.Name, Response.Record.GetKey(), Policy, Response.UserData}};
+			return InnerCache->Get(GetDataRequests, Owner, [this, &Owner, State](FCacheGetResponse&& DataResponse)
+			{
+				GetDataComplete(Owner, State, MoveTemp(DataResponse));
+			});
+		}
+	}
+	else
+	{
+		UE_LOG(LogDerivedDataCache, Warning,
+			TEXT("Verify: Cache did not contain a record for %s from '%s'."),
+			*WriteToString<96>(Request.Record.GetKey()), *Request.Name);
+		FWriteScopeLock Lock(State->Lock);
+		State->ForwardRequests.Add(MoveTemp(Request));
+	}
+
+	GetComplete(Owner, State);
+}
+
+void FCacheStoreVerify::GetDataComplete(IRequestOwner& Owner, FVerifyPutState* State, FCacheGetResponse&& Response)
+{
+	FCachePutRequest& Request = State->VerifyRequests[int32(Response.UserData)];
+
+	if ((Response.Status == EStatus::Ok) ||
+		(Response.Status == EStatus::Error && !Response.Record.GetValues().IsEmpty()))
+	{
+		if (CompareRecords(Request.Record, Response.Record, Request.Name))
+		{
+			UE_LOG(LogDerivedDataCache, Log,
+				TEXT("Verify: Data in the cache matches newly generated data for %s from '%s'."),
+				*WriteToString<96>(Request.Record.GetKey()), *Request.Name);
+			State->OnComplete(Request.MakeResponse(EStatus::Ok));
+		}
+		else if (bPutOnError)
+		{
+			// Ask to overwrite existing records to potentially eliminate the mismatch.
+			UE_LOG(LogDerivedDataCache, Display,
+				TEXT("Verify: Writing newly generated data to the cache for %s from '%s'."),
+				*WriteToString<96>(Request.Record.GetKey()), *Request.Name);
+			Request.Policy = Request.Policy.Transform([](ECachePolicy P) { return P & ~ECachePolicy::Query; });
+			FWriteScopeLock Lock(State->Lock);
+			State->ForwardRequests.Add(MoveTemp(Request));
+		}
+		else
+		{
+			State->OnComplete(Request.MakeResponse(EStatus::Ok));
+		}
+	}
+	else
+	{
+		UE_LOG(LogDerivedDataCache, Warning,
+			TEXT("Verify: Cache did not contain a record for %s from '%s'."),
+			*WriteToString<96>(Request.Record.GetKey()), *Request.Name);
+		FWriteScopeLock Lock(State->Lock);
+		State->ForwardRequests.Add(MoveTemp(Request));
+	}
+
+	GetComplete(Owner, State);
+}
+
+void FCacheStoreVerify::GetComplete(IRequestOwner& Owner, FVerifyPutState* State)
+{
+	if (FWriteScopeLock Lock(State->Lock); --State->ActiveRequests > 0)
+	{
+		return;
+	}
+	if (!State->ForwardRequests.IsEmpty())
+	{
+		InnerCache->Put(State->ForwardRequests, Owner, MoveTemp(State->OnComplete));
+	}
+	delete State;
+}
+
+void FCacheStoreVerify::Get(
+	const TConstArrayView<FCacheGetRequest> Requests,
+	IRequestOwner& Owner,
+	FOnCacheGetComplete&& OnComplete)
+{
+	TArray<FCacheGetRequest, TInlineAllocator<8>> ForwardRequests;
+	TArray<FCacheGetRequest, TInlineAllocator<8>> VerifyRequests;
+	ForwardRequests.Reserve(Requests.Num());
+	VerifyRequests.Reserve(Requests.Num());
+	{
+		FScopeLock ScopeLock(&SynchronizationObject);
+		for (const FCacheGetRequest& Request : Requests)
+		{
+			bool bAlreadyTested = AlreadyTested.Contains(Request.Key);
+			(bAlreadyTested ? ForwardRequests : VerifyRequests).Add(Request);
+		}
+	}
+
+	CompleteWithStatus(VerifyRequests, OnComplete, EStatus::Error);
+
+	if (!ForwardRequests.IsEmpty())
+	{
+		InnerCache->Get(ForwardRequests, Owner, MoveTemp(OnComplete));
+	}
+}
+
+void FCacheStoreVerify::PutValue(
+	const TConstArrayView<FCachePutValueRequest> Requests,
+	IRequestOwner& Owner,
+	FOnCachePutValueComplete&& OnComplete)
+{
+	TUniquePtr<FVerifyPutValueState> State = MakeUnique<FVerifyPutValueState>();
+	State->VerifyRequests.Reserve(Requests.Num());
+	{
+		FScopeLock ScopeLock(&SynchronizationObject);
+		for (const FCachePutValueRequest& Request : Requests)
+		{
+			bool bAlreadyTested = false;
+			AlreadyTested.Add(Request.Key, &bAlreadyTested);
+			(bAlreadyTested ? State->ForwardRequests : State->VerifyRequests).Add(Request);
+		}
+	}
+
+	if (State->VerifyRequests.IsEmpty())
+	{
+		return InnerCache->PutValue(State->ForwardRequests, Owner, MoveTemp(OnComplete));
+	}
+
+	TArray<FCacheGetValueRequest> GetMetaRequests;
+	GetMetaRequests.Reserve(State->VerifyRequests.Num());
+	{
+		uint64 PutIndex = 0;
+		const ECachePolicy GetPolicy = ECachePolicy::Query | ECachePolicy::SkipData;
+		for (const FCachePutValueRequest& PutRequest : State->VerifyRequests)
+		{
+			GetMetaRequests.Add({PutRequest.Name, PutRequest.Key, GetPolicy, PutIndex++});
+		}
+	}
+
+	State->OnComplete = MoveTemp(OnComplete);
+	State->ActiveRequests = GetMetaRequests.Num();
+	InnerCache->GetValue(GetMetaRequests, Owner, [this, &Owner, State = State.Release()](FCacheGetValueResponse&& MetaResponse)
+	{
+		GetMetaComplete(Owner, State, MoveTemp(MetaResponse));
+	});
+}
+
+void FCacheStoreVerify::GetMetaComplete(IRequestOwner& Owner, FVerifyPutValueState* State, FCacheGetValueResponse&& Response)
+{
+	FCachePutValueRequest& Request = State->VerifyRequests[int32(Response.UserData)];
+
+	if (Response.Status == EStatus::Ok)
+	{
+		if (Request.Value.GetRawHash() == Response.Value.GetRawHash())
+		{
+			UE_LOG(LogDerivedDataCache, Log,
+				TEXT("Verify: Data in the cache matches newly generated data for %s from '%s'."),
+				*WriteToString<96>(Request.Key), *Request.Name);
+			State->OnComplete(Request.MakeResponse(EStatus::Ok));
+		}
+		else
+		{
+			const FCacheGetValueRequest GetDataRequests[]{{Response.Name, Response.Key, ECachePolicy::Default, Response.UserData}};
+			return InnerCache->GetValue(GetDataRequests, Owner, [this, &Owner, State](FCacheGetValueResponse&& DataResponse)
+			{
+				GetDataComplete(Owner, State, MoveTemp(DataResponse));
+			});
+		}
+	}
+	else
+	{
+		UE_LOG(LogDerivedDataCache, Warning,
+			TEXT("Verify: Cache did not contain a value for %s from '%s'."),
+			*WriteToString<96>(Request.Key), *Request.Name);
+		FWriteScopeLock Lock(State->Lock);
+		State->ForwardRequests.Add(MoveTemp(Request));
+	}
+
+	GetComplete(Owner, State);
+}
+
+void FCacheStoreVerify::GetDataComplete(IRequestOwner& Owner, FVerifyPutValueState* State, FCacheGetValueResponse&& Response)
+{
+	FCachePutValueRequest& Request = State->VerifyRequests[int32(Response.UserData)];
+
+	if (Response.Status == EStatus::Ok)
+	{
+		if (Request.Value.GetRawHash() == Response.Value.GetRawHash())
+		{
+			UE_LOG(LogDerivedDataCache, Log,
+				TEXT("Verify: Data in the cache matches newly generated data for %s from '%s'."),
+				*WriteToString<96>(Request.Key), *Request.Name);
+			State->OnComplete(Request.MakeResponse(EStatus::Ok));
+		}
+		else
+		{
+			LogChangedValue(Request.Name, Request.Key, FValueId::Null, Request.Value, Response.Value);
+			if (bPutOnError)
+			{
+				// Ask to overwrite existing values to potentially eliminate the mismatch.
+				UE_LOG(LogDerivedDataCache, Display,
+					TEXT("Verify: Writing newly generated data to the cache for %s from '%s'."),
+					*WriteToString<96>(Request.Key), *Request.Name);
+				Request.Policy &= ~ECachePolicy::Query;
+				FWriteScopeLock Lock(State->Lock);
+				State->ForwardRequests.Add(MoveTemp(Request));
+			}
+			else
+			{
+				State->OnComplete(Request.MakeResponse(EStatus::Ok));
+			}
+		}
+	}
+	else
+	{
+		UE_LOG(LogDerivedDataCache, Warning,
+			TEXT("Verify: Cache did not contain a value for %s from '%s'."),
+			*WriteToString<96>(Request.Key), *Request.Name);
+		FWriteScopeLock Lock(State->Lock);
+		State->ForwardRequests.Add(MoveTemp(Request));
+	}
+
+	GetComplete(Owner, State);
+}
+
+void FCacheStoreVerify::GetComplete(IRequestOwner& Owner, FVerifyPutValueState* State)
+{
+	if (FWriteScopeLock Lock(State->Lock); --State->ActiveRequests > 0)
+	{
+		return;
+	}
+	if (!State->ForwardRequests.IsEmpty())
+	{
+		InnerCache->PutValue(State->ForwardRequests, Owner, MoveTemp(State->OnComplete));
+	}
+	delete State;
+}
+
+void FCacheStoreVerify::GetValue(
+	const TConstArrayView<FCacheGetValueRequest> Requests,
+	IRequestOwner& Owner,
+	FOnCacheGetValueComplete&& OnComplete)
+{
+	TArray<FCacheGetValueRequest, TInlineAllocator<8>> ForwardRequests;
+	TArray<FCacheGetValueRequest, TInlineAllocator<8>> VerifyRequests;
+	ForwardRequests.Reserve(Requests.Num());
+	VerifyRequests.Reserve(Requests.Num());
+	{
+		FScopeLock ScopeLock(&SynchronizationObject);
+		for (const FCacheGetValueRequest& Request : Requests)
+		{
+			bool bAlreadyTested = AlreadyTested.Contains(Request.Key);
+			(bAlreadyTested ? ForwardRequests : VerifyRequests).Add(Request);
+		}
+	}
+
+	CompleteWithStatus(VerifyRequests, OnComplete, EStatus::Error);
+
+	if (!ForwardRequests.IsEmpty())
+	{
+		InnerCache->GetValue(ForwardRequests, Owner, MoveTemp(OnComplete));
+	}
+}
+
+void FCacheStoreVerify::GetChunks(
+	const TConstArrayView<FCacheGetChunkRequest> Requests,
+	IRequestOwner& Owner,
+	FOnCacheGetChunkComplete&& OnComplete)
+{
+	TArray<FCacheGetChunkRequest, TInlineAllocator<8>> ForwardRequests;
+	TArray<FCacheGetChunkRequest, TInlineAllocator<8>> VerifyRequests;
+	ForwardRequests.Reserve(Requests.Num());
+	VerifyRequests.Reserve(Requests.Num());
+	{
+		FScopeLock ScopeLock(&SynchronizationObject);
+		for (const FCacheGetChunkRequest& Request : Requests)
+		{
+			bool bAlreadyTested = AlreadyTested.Contains(Request.Key);
+			(bAlreadyTested ? ForwardRequests : VerifyRequests).Add(Request);
+		}
+	}
+
+	CompleteWithStatus(VerifyRequests, OnComplete, EStatus::Error);
+
+	if (!ForwardRequests.IsEmpty())
+	{
+		InnerCache->GetChunks(ForwardRequests, Owner, MoveTemp(OnComplete));
+	}
+}
+
+void FCacheStoreVerify::LegacyPut(
+	const TConstArrayView<FLegacyCachePutRequest> Requests,
+	IRequestOwner& Owner,
+	FOnLegacyCachePutComplete&& OnComplete)
+{
+	TUniquePtr<FVerifyLegacyPutState> State = MakeUnique<FVerifyLegacyPutState>();
+	State->VerifyRequests.Reserve(Requests.Num());
+	{
+		FScopeLock ScopeLock(&SynchronizationObject);
+		for (const FLegacyCachePutRequest& Request : Requests)
+		{
+			bool bAlreadyTested = false;
+			AlreadyTested.Add(Request.Key.GetKey(), &bAlreadyTested);
+			(bAlreadyTested ? State->ForwardRequests : State->VerifyRequests).Add(Request);
+		}
+	}
+
+	if (State->VerifyRequests.IsEmpty())
+	{
+		return InnerCache->LegacyPut(State->ForwardRequests, Owner, MoveTemp(OnComplete));
+	}
+
+	TArray<FLegacyCacheGetRequest> GetDataRequests;
+	GetDataRequests.Reserve(State->VerifyRequests.Num());
+	{
+		uint64 PutIndex = 0;
+		const ECachePolicy GetPolicy = ECachePolicy::Query;
+		for (const FLegacyCachePutRequest& PutRequest : State->VerifyRequests)
+		{
+			GetDataRequests.Add({PutRequest.Name, PutRequest.Key, GetPolicy, PutIndex++});
+		}
+	}
+
+	State->OnComplete = MoveTemp(OnComplete);
+	State->ActiveRequests = GetDataRequests.Num();
+	InnerCache->LegacyGet(GetDataRequests, Owner, [this, &Owner, State = State.Release()](FLegacyCacheGetResponse&& MetaResponse)
+	{
+		GetDataComplete(Owner, State, MoveTemp(MetaResponse));
+	});
+}
+
+void FCacheStoreVerify::GetDataComplete(IRequestOwner& Owner, FVerifyLegacyPutState* State, FLegacyCacheGetResponse&& Response)
+{
+	FLegacyCachePutRequest& Request = State->VerifyRequests[int32(Response.UserData)];
+
+	if (Response.Status == EStatus::Ok)
+	{
+		const FValue NewValue(FCompressedBuffer::Compress(Request.Value, ECompressedBufferCompressor::NotSet, ECompressedBufferCompressionLevel::None));
+		const FValue OldValue(FCompressedBuffer::Compress(Response.Value, ECompressedBufferCompressor::NotSet, ECompressedBufferCompressionLevel::None));
+		if (NewValue.GetRawHash() == OldValue.GetRawHash())
+		{
+			UE_LOG(LogDerivedDataCache, Log,
+				TEXT("Verify: Data in the cache matches newly generated data for %s from '%s'."),
+				*WriteToString<96>(Request.Key.GetKey()), *Request.Name);
+			State->OnComplete(Request.MakeResponse(EStatus::Ok));
+		}
+		else
+		{
+			LogChangedValue(Request.Name, Request.Key.GetKey(), FValueId::Null, NewValue, OldValue);
+			if (bPutOnError)
+			{
+				// Ask to overwrite existing values to potentially eliminate the mismatch.
+				UE_LOG(LogDerivedDataCache, Display,
+					TEXT("Verify: Writing newly generated data to the cache for %s from '%s'."),
+					*WriteToString<96>(Request.Key.GetKey()), *Request.Name);
+				Request.Policy &= ~ECachePolicy::Query;
+				FWriteScopeLock Lock(State->Lock);
+				State->ForwardRequests.Add(MoveTemp(Request));
+			}
+			else
+			{
+				State->OnComplete(Request.MakeResponse(EStatus::Ok));
+			}
+		}
+	}
+	else
+	{
+		UE_LOG(LogDerivedDataCache, Warning,
+			TEXT("Verify: Cache did not contain a value for %s from '%s'."),
+			*WriteToString<96>(Request.Key.GetKey()), *Request.Name);
+		FWriteScopeLock Lock(State->Lock);
+		State->ForwardRequests.Add(MoveTemp(Request));
+	}
+
+	GetComplete(Owner, State);
+}
+
+void FCacheStoreVerify::GetComplete(IRequestOwner& Owner, FVerifyLegacyPutState* State)
+{
+	if (FWriteScopeLock Lock(State->Lock); --State->ActiveRequests > 0)
+	{
+		return;
+	}
+	if (!State->ForwardRequests.IsEmpty())
+	{
+		InnerCache->LegacyPut(State->ForwardRequests, Owner, MoveTemp(State->OnComplete));
+	}
+	delete State;
+}
+
+void FCacheStoreVerify::LegacyGet(
+	const TConstArrayView<FLegacyCacheGetRequest> Requests,
+	IRequestOwner& Owner,
+	FOnLegacyCacheGetComplete&& OnComplete)
+{
+	TArray<FLegacyCacheGetRequest, TInlineAllocator<8>> ForwardRequests;
+	TArray<FLegacyCacheGetRequest, TInlineAllocator<8>> VerifyRequests;
+	ForwardRequests.Reserve(Requests.Num());
+	VerifyRequests.Reserve(Requests.Num());
+	{
+		FScopeLock ScopeLock(&SynchronizationObject);
+		for (const FLegacyCacheGetRequest& Request : Requests)
+		{
+			bool bAlreadyTested = AlreadyTested.Contains(Request.Key.GetKey());
+			(bAlreadyTested ? ForwardRequests : VerifyRequests).Add(Request);
+		}
+	}
+
+	CompleteWithStatus(VerifyRequests, OnComplete, EStatus::Error);
+
+	if (!ForwardRequests.IsEmpty())
+	{
+		InnerCache->LegacyGet(ForwardRequests, Owner, MoveTemp(OnComplete));
+	}
+}
+
+bool FCacheStoreVerify::CompareRecords(const FCacheRecord& PutRecord, const FCacheRecord& GetRecord, const FSharedString& Name)
+{
+	bool bEqual = true;
+
+	const FCacheKey& Key = PutRecord.GetKey();
+	const TConstArrayView<FValueWithId> PutValues = PutRecord.GetValues();
+	const TConstArrayView<FValueWithId> GetValues = GetRecord.GetValues();
+	const FValueWithId* PutIt = PutValues.GetData();
+	const FValueWithId* GetIt = GetValues.GetData();
+	const FValueWithId* const PutEnd = PutIt + PutValues.Num();
+	const FValueWithId* const GetEnd = GetIt + GetValues.Num();
+
+	const auto LogNewValue = [&Name, &Key](const FValueWithId& Value)
+	{
+		UE_LOG(LogDerivedDataCache, Error,
+			TEXT("Verify: Value %s with hash %s is in the new record but does not exist in the cache for %s from '%s'."),
+			*WriteToString<32>(Value.GetId()), *WriteToString<48>(Value.GetRawHash()), *WriteToString<96>(Key), *Name);
+	};
+
+	const auto LogOldValue = [&Name, &Key](const FValueWithId& Value)
+	{
+		UE_LOG(LogDerivedDataCache, Error,
+			TEXT("Verify: Value %s with hash %s is in the cache but does not exist in the new record for %s from '%s'."),
+			*WriteToString<32>(Value.GetId()), *WriteToString<48>(Value.GetRawHash()), *WriteToString<96>(Key), *Name);
+	};
+
+	while (PutIt != PutEnd && GetIt != GetEnd)
+	{
+		if (PutIt->GetId() == GetIt->GetId())
+		{
+			if (PutIt->GetRawHash() != GetIt->GetRawHash())
+			{
+				LogChangedValue(Name, Key, PutIt->GetId(), *PutIt, *GetIt);
+				bEqual = false;
+			}
+			++PutIt;
+			++GetIt;
+		}
+		else if (PutIt->GetId() < GetIt->GetId())
+		{
+			LogNewValue(*PutIt++);
+			bEqual = false;
+		}
+		else
+		{
+			LogOldValue(*GetIt++);
+			bEqual = false;
+		}
+	}
+
+	while (PutIt != PutEnd)
+	{
+		LogNewValue(*PutIt++);
+		bEqual = false;
+	}
+
+	while (GetIt != GetEnd)
+	{
+		LogOldValue(*GetIt++);
+		bEqual = false;
+	}
+
+	return bEqual;
+}
+
+void FCacheStoreVerify::LogChangedValue(const FSharedString& Name, const FCacheKey& Key, const FValueId& Id, const FValue& NewValue, const FValue& OldValue)
+{
+	TStringBuilder<32> IdString;
+	if (Id.IsValid())
+	{
+		IdString << TEXT(' ') << Id;
+	}
+
+	const auto LogDataToFile = [&Name, &Key, &Id, &IdString](const FValue& Value, FStringView Extension)
+	{
+		if (FSharedBuffer RawData = Value.GetData().Decompress())
+		{
+			TStringBuilder<256> Path;
+			FPathViews::Append(Path, FPaths::ProjectSavedDir(), TEXT("VerifyDDC"), TEXT(""));
+			Path << Key.Bucket << TEXT('_') << Key.Hash;
+			if (Id.IsValid())
+			{
+				Path << TEXT('_') << Id;
+			}
+			Path << Extension;
+			if (TUniquePtr<FArchive> Ar{IFileManager::Get().CreateFileWriter(*Path, FILEWRITE_Silent)})
+			{
+				Ar->Serialize(const_cast<void*>(RawData.GetData()), int64(RawData.GetSize()));
+			}
+		}
+		else
+		{
+			UE_LOG(LogDerivedDataCache, Log,
+				TEXT("Verify: Value%s does not have data with hash %s to save to disk for %s from '%s'."),
+				*IdString, *WriteToString<48>(Value.GetRawHash()), *WriteToString<96>(Key), *Name);
+		}
+	};
+
+	UE_LOG(LogDerivedDataCache, Error,
+		TEXT("Verify: Value%s has hash %s in the newly generated data and hash %s in the cache for %s from '%s'."),
+		*IdString, *WriteToString<48>(NewValue.GetRawHash()), *WriteToString<48>(OldValue.GetRawHash()),
+		*WriteToString<96>(Key), *Name);
+	LogDataToFile(NewValue, TEXTVIEW(".verify"));
+	LogDataToFile(OldValue, TEXTVIEW(".fromcache"));
+}
+
+ILegacyCacheStore* CreateCacheStoreVerify(ILegacyCacheStore* InnerCache, bool bPutOnError)
+{
+	return new FCacheStoreVerify(InnerCache, bPutOnError);
+}
+
+} // UE::DerivedData
diff --git a/Engine/Source/Developer/DerivedDataCache/Private/HierarchicalDerivedDataBackend.h b/Engine/Source/Developer/DerivedDataCache/Private/HierarchicalDerivedDataBackend.h
deleted file mode 100644
