Fix static analysis error in CacheStoreTest.cpp.

diff --git a/Engine/Source/Developer/DerivedDataCache/Private/Tests/CacheStoreTest.cpp b/Engine/Source/Developer/DerivedDataCache/Private/Tests/CacheStoreTest.cpp
--- a/Engine/Source/Developer/DerivedDataCache/Private/Tests/CacheStoreTest.cpp
+++ b/Engine/Source/Developer/DerivedDataCache/Private/Tests/CacheStoreTest.cpp
@@ -321,7 +321,8 @@ bool FCacheStoreTest::RunTest(const FString& Parameters)
 			{
 				FUserData* UserData = reinterpret_cast<FUserData*>(Response.UserData);
 				FKeyData* KeyData = UserData->KeyData;
-				if (TestTrue(TEXT("Valid UserData in GetChunks Callback"), KeyData != nullptr))
+				TestTrue(TEXT("Valid UserData in GetChunks Callback"), KeyData != nullptr);
+				if (KeyData != nullptr)
 				{
 					int32 n = KeyData->KeyIndex;
 					int32 ValueIndex = UserData->ValueIndex >= 0 ? UserData->ValueIndex : 0;
