DDC: Fixed the Zen cache being accidentally made read-only

#rnx

diff --git a/Engine/Source/Developer/DerivedDataCache/Private/DerivedDataBackends.cpp b/Engine/Source/Developer/DerivedDataCache/Private/DerivedDataBackends.cpp
--- a/Engine/Source/Developer/DerivedDataCache/Private/DerivedDataBackends.cpp
+++ b/Engine/Source/Developer/DerivedDataCache/Private/DerivedDataBackends.cpp
@@ -310,7 +310,7 @@ public:
 				}
 				else if (NodeType == TEXT("Zen"))
 				{
-					ParsedNode = MakeTuple(ParseZenCache(*NodeName, *Entry), ECacheStoreFlags::Local | ECacheStoreFlags::Remote | ECacheStoreFlags::Query | ECacheStoreFlags::Remote);
+					ParsedNode = MakeTuple(ParseZenCache(*NodeName, *Entry), ECacheStoreFlags::Local | ECacheStoreFlags::Remote | ECacheStoreFlags::Query | ECacheStoreFlags::Store);
 				}
 				
 				if (ParsedNode.Key)
