DDC: Fixed incorrect callbacks in put batches with failures

Also adds the intended inline allocators for arrays that were missing them.

#rnx

diff --git a/Engine/Source/Developer/DerivedDataCache/Private/DerivedDataCacheStoreHierarchy.cpp b/Engine/Source/Developer/DerivedDataCache/Private/DerivedDataCacheStoreHierarchy.cpp
--- a/Engine/Source/Developer/DerivedDataCache/Private/DerivedDataCacheStoreHierarchy.cpp
+++ b/Engine/Source/Developer/DerivedDataCache/Private/DerivedDataCacheStoreHierarchy.cpp
@@ -369,7 +369,7 @@ private:
 	};
 
 	const FCacheStoreHierarchy& Hierarchy;
-	TArray<FPutRequest> Requests;
+	TArray<FPutRequest, TInlineAllocator<1>> Requests;
 	IRequestOwner& BatchOwner;
 	FOnPutComplete OnComplete;
 
@@ -417,6 +417,7 @@ void FCacheStoreHierarchy::TPutBatch<Params>::DispatchRequests()
 		{
 			OnComplete(Request.MakeResponse(BatchOwner.IsCanceled() ? EStatus::Canceled : EStatus::Error));
 		}
+		++RequestIndex;
 	}
 }
 
@@ -1145,7 +1146,7 @@ private:
 	};
 
 	const FCacheStoreHierarchy& Hierarchy;
-	TArray<FLegacyCacheDeleteRequest> Requests;
+	TArray<FLegacyCacheDeleteRequest, TInlineAllocator<1>> Requests;
 	IRequestOwner& BatchOwner;
 	FOnLegacyCacheDeleteComplete OnComplete;
 
@@ -1208,9 +1209,10 @@ void FCacheStoreHierarchy::FLegacyDeleteBatch::DispatchRequests()
 	int32 RequestIndex = 0;
 	for (const FLegacyCacheDeleteRequest& Request : Requests)
 	{
-		const bool bOk = States[RequestIndex++].bOk;
+		const bool bOk = States[RequestIndex].bOk;
 		const EStatus Status = bOk ? EStatus::Ok : BatchOwner.IsCanceled() ? EStatus::Canceled : EStatus::Error;
 		OnComplete(Request.MakeResponse(Status));
+		++RequestIndex;
 	}
 }
 
