DDC: Fixed propagation of cache values when a value policy contains SkipData

#rnx

diff --git a/Engine/Source/Developer/DerivedDataCache/Private/DerivedDataCacheStoreHierarchy.cpp b/Engine/Source/Developer/DerivedDataCache/Private/DerivedDataCacheStoreHierarchy.cpp
--- a/Engine/Source/Developer/DerivedDataCache/Private/DerivedDataCacheStoreHierarchy.cpp
+++ b/Engine/Source/Developer/DerivedDataCache/Private/DerivedDataCacheStoreHierarchy.cpp
@@ -663,7 +663,7 @@ void FCacheStoreHierarchy::TGetBatch<Params>::DispatchRequests()
 			{
 				if (const ECachePolicy Policy = GetCombinedPolicy(Request.Policy); CanQuery(Policy, Node.CacheFlags))
 				{
-					if (EnumHasAnyFlags(Policy, ECachePolicy::SkipData) && CanStoreIfOk(Policy, Node.NodeFlags))
+					if (CanStoreIfOk(Policy, Node.NodeFlags))
 					{
 						NodeRequests.Add({Request.Name, Request.Key, RemovePolicy(Request.Policy, ECachePolicy::SkipData | ECachePolicy::SkipMeta), StateIndex});
 					}
@@ -727,7 +727,8 @@ void FCacheStoreHierarchy::TGetBatch<Params>::CompleteRequest(FGetResponse&& Res
 	if (bLastQuery && CanStoreIfOk(GetCombinedPolicy(State.Request.Policy), Node.NodeFlags) && HasResponseData(Response))
 	{
 		// Store any retrieved values to previous writable nodes if Ok or there are no remaining nodes to query.
-		const FPutRequest PutRequest = MakePutRequest(Response, State.Request);
+		FPutRequest PutRequest = MakePutRequest(Response, State.Request);
+		PutRequest.Policy = RemovePolicy(PutRequest.Policy, ECachePolicy::Query);
 		for (int32 PutNodeIndex = 0; PutNodeIndex < NodeIndex; ++PutNodeIndex)
 		{
 			const FCacheStoreNode& PutNode = Hierarchy.Nodes[PutNodeIndex];
@@ -739,6 +740,13 @@ void FCacheStoreHierarchy::TGetBatch<Params>::CompleteRequest(FGetResponse&& Res
 		}
 	}
 
+	if (bFirstOk)
+	{
+		// Values may be fetched to fill previous nodes. Remove values if requested.
+		FilterResponseByRequest(Response, State.Request);
+		OnComplete(MoveTemp(Response));
+	}
+
 	if (State.Response.Status == EStatus::Ok)
 	{
 		if (EnumHasAnyFlags(Node.CacheFlags, ECacheStoreFlags::StopStore))
@@ -749,18 +757,11 @@ void FCacheStoreHierarchy::TGetBatch<Params>::CompleteRequest(FGetResponse&& Res
 		else
 		{
 			// Never store to later remote nodes.
-			// This is a necessary optimization until speculative stores can be optimized.
+			// This is a necessary optimization until speculative stores have been optimized.
 			State.Request.Policy = RemovePolicy(State.Request.Policy, ECachePolicy::Remote);
 		}
 	}
 
-	if (bFirstOk)
-	{
-		// Values may be fetched to fill previous nodes. Remove values if requested.
-		FilterResponseByRequest(Response, State.Request);
-		OnComplete(MoveTemp(Response));
-	}
-
 	if (RemainingRequestCount.Signal())
 	{
 		DispatchRequests();
diff --git a/Engine/Source/Developer/DerivedDataCache/Private/FileSystemDerivedDataBackend.cpp b/Engine/Source/Developer/DerivedDataCache/Private/FileSystemDerivedDataBackend.cpp
--- a/Engine/Source/Developer/DerivedDataCache/Private/FileSystemDerivedDataBackend.cpp
+++ b/Engine/Source/Developer/DerivedDataCache/Private/FileSystemDerivedDataBackend.cpp
@@ -2,6 +2,7 @@
 
 #include "Algo/Accumulate.h"
 #include "Algo/AllOf.h"
+#include "Algo/Compare.h"
 #include "Algo/StableSort.h"
 #include "Algo/Transform.h"
 #include "Async/Async.h"
@@ -1797,22 +1798,56 @@ bool FFileSystemCacheStore::PutCacheRecord(
 		return false;
 	}
 
-	// Check if there is an existing record package.
-	bool bRecordExists = false;
-	FCbPackage ExistingPackage;
 	TStringBuilder<256> Path;
 	BuildCachePackagePath(Key, Path);
+
+	// Check if there is an existing record package.
+	FCbPackage ExistingPackage;
 	const ECachePolicy QueryFlag = SpeedClass == ESpeedClass::Local ? ECachePolicy::QueryLocal : ECachePolicy::QueryRemote;
-	const bool bReplaceExisting = !EnumHasAnyFlags(RecordPolicy, QueryFlag);
-	if (!bReplaceExisting)
-	{
-		bRecordExists = LoadFileWithHash(Path, Name, [&ExistingPackage](FArchive& Ar) { ExistingPackage.TryLoad(Ar); });
+	bool bReplaceExisting = !EnumHasAnyFlags(RecordPolicy, QueryFlag);
+	bool bSaveRecord = bReplaceExisting;
+	if (const bool bLoadRecord = !bReplaceExisting || !Algo::AllOf(Record.GetValues(), &FValue::HasData))
+	{
+		// Load the existing package to take its inline attachments into account.
+		bSaveRecord |= !LoadFileWithHash(Path, Name, [&ExistingPackage](FArchive& Ar) { ExistingPackage.TryLoad(Ar); });
+		if (!bSaveRecord)
+		{
+			// Load the existing record to detect determinism issues.
+			const FOptionalCacheRecord ExistingRecord = FCacheRecord::Load(ExistingPackage);
+			bSaveRecord |= !ExistingRecord;
+			const auto MakeValueTuple = [](const FValueWithId& Value) -> TTuple<FValueId, FIoHash>
+			{
+				return MakeTuple(Value.GetId(), Value.GetRawHash());
+			};
+			if (ExistingRecord && !Algo::CompareBy(ExistingRecord.Get().GetValues(), Record.GetValues(), MakeValueTuple))
+			{
+				// Values differ between the existing record and the new record.
+				// Overwrite the existing record if the cache is missing content for any of its values.
+				UE_LOG(LogDerivedDataCache, Display, TEXT("%s: Cache put found non-deterministic record for %s from '%.*s'"),
+					*CachePath, *WriteToString<96>(Key), Name.Len(), Name.GetData());
+				const auto HasValueContent = [this, Name, &Key](const FValueWithId& Value) -> bool
+				{
+					if (!Value.HasData() && !GetCacheContentExists(Key, Value.GetRawHash()))
+					{
+						UE_LOG(LogDerivedDataCache, Log,
+							TEXT("%s: Cache put of non-deterministic record will overwrite existing record due to ")
+							TEXT("missing value %s with hash %s for %s from '%.*s'"),
+							*CachePath, *WriteToString<16>(Value.GetId()), *WriteToString<48>(Value.GetRawHash()),
+							*WriteToString<96>(Key), Name.Len(), Name.GetData());
+						return false;
+					}
+					return true;
+				};
+				bSaveRecord |= !Algo::AllOf(ExistingRecord.Get().GetValues(), HasValueContent);
+				bReplaceExisting |= bSaveRecord;
+			}
+		}
 	}
 
-	// Save the record to a package and remove attachments that will be stored externally.
+	// Serialize the record to a package and remove attachments that will be stored externally.
 	FCbPackage Package = Record.Save();
 	TArray<FCompressedBuffer, TInlineAllocator<8>> ExternalContent;
-	if (ExistingPackage)
+	if (ExistingPackage && !bSaveRecord)
 	{
 		// Mirror the existing internal/external attachment storage.
 		TArray<FCompressedBuffer, TInlineAllocator<8>> AllContent;
@@ -1829,11 +1864,26 @@ bool FFileSystemCacheStore::PutCacheRecord(
 	}
 	else
 	{
+		// Attempt to copy missing attachments from the existing package.
+		if (ExistingPackage)
+		{
+			for (const FValue& Value : Record.GetValues())
+			{
+				if (!Value.HasData())
+				{
+					if (const FCbAttachment* Attachment = ExistingPackage.FindAttachment(Value.GetRawHash()))
+					{
+						Package.AddAttachment(*Attachment);
+					}
+				}
+			}
+		}
+
 		// Remove the largest attachments from the package until it fits within the size limits.
 		TArray<FCompressedBuffer, TInlineAllocator<8>> AllContent;
 		Algo::Transform(Package.GetAttachments(), AllContent, &FCbAttachment::AsCompressedBinary);
 		uint64 TotalSize = Algo::TransformAccumulate(AllContent, &FCompressedBuffer::GetCompressedSize, uint64(0));
-		const uint64 MaxSize = (AllContent.Num() == 1 ? MaxValueSizeKB : MaxRecordSizeKB) * 1024;
+		const uint64 MaxSize = (Record.GetValues().Num() == 1 ? MaxValueSizeKB : MaxRecordSizeKB) * 1024;
 		if (TotalSize > MaxSize)
 		{
 			Algo::StableSortBy(AllContent, &FCompressedBuffer::GetCompressedSize, TGreater<>());
@@ -1864,7 +1914,7 @@ bool FFileSystemCacheStore::PutCacheRecord(
 
 	// Save the record package to storage.
 	const auto WriteRecord = [&](FArchive& Ar) { Package.Save(Ar); OutWriteSize += uint64(Ar.TotalSize()); };
-	if (!bRecordExists && !SaveFileWithHash(Path, Name, WriteRecord, bReplaceExisting))
+	if (bSaveRecord && !SaveFileWithHash(Path, Name, WriteRecord, bReplaceExisting))
 	{
 		return false;
 	}
@@ -2420,6 +2470,8 @@ bool FFileSystemCacheStore::SaveFile(
 	const TFunctionRef<void (FArchive&)> WriteFunction,
 	const bool bReplaceExisting) const
 {
+	const double StartTime = FPlatformTime::Seconds();
+
 	TStringBuilder<256> TempPath;
 	TempPath << FPathViews::GetPath(Path) << TEXT("/Temp.") << FGuid::NewGuid();
 
@@ -2464,6 +2516,12 @@ bool FFileSystemCacheStore::SaveFile(
 			*CachePath, *Path, DebugName.Len(), DebugName.GetData());
 	}
 
+	const double WriteDuration = FPlatformTime::Seconds() - StartTime;
+	const double WriteSpeed = WriteDuration > 0.001 ? (WriteSize / WriteDuration) / (1024.0 * 1024.0) : 0.0;
+	UE_LOG(LogDerivedDataCache, VeryVerbose,
+		TEXT("%s: Saved %s from '%.*s' (%" INT64_FMT " bytes, %.02f secs, %.2f MiB/s)"),
+		*CachePath, *Path, DebugName.Len(), DebugName.GetData(), WriteSize, WriteDuration, WriteSpeed);
+
 	return true;
 }
 
