Allow for region pinning when connecting to Jupiter.  Fix value printing error.  Increase concurrent PUT operations from 2 to 4 when batching is enabled.

diff --git a/Engine/Source/Developer/DerivedDataCache/Private/HttpDerivedDataBackend.cpp b/Engine/Source/Developer/DerivedDataCache/Private/HttpDerivedDataBackend.cpp
--- a/Engine/Source/Developer/DerivedDataCache/Private/HttpDerivedDataBackend.cpp
+++ b/Engine/Source/Developer/DerivedDataCache/Private/HttpDerivedDataBackend.cpp
@@ -66,7 +66,7 @@
 #define UE_HTTPDDC_HTTP_DEBUG 0
 #if WITH_DATAREQUEST_HELPER
 	#define UE_HTTPDDC_GET_REQUEST_POOL_SIZE 16 
-	#define UE_HTTPDDC_PUT_REQUEST_POOL_SIZE 2
+	#define UE_HTTPDDC_PUT_REQUEST_POOL_SIZE 4
 #else
 	#define UE_HTTPDDC_GET_REQUEST_POOL_SIZE 48
 	#define UE_HTTPDDC_PUT_REQUEST_POOL_SIZE 16
@@ -148,9 +148,10 @@ public:
 		FailedTimeout
 	};
 
-	FHttpRequest(const TCHAR* InDomain, FHttpAccessToken* InAuthorizationToken, bool bInLogErrors)
+	FHttpRequest(const TCHAR* InDomain, const TCHAR* InEffectiveDomain, FHttpAccessToken* InAuthorizationToken, bool bInLogErrors)
 		: bLogErrors(bInLogErrors)
 		, Domain(InDomain)
+		, EffectiveDomain(InEffectiveDomain)
 		, AuthorizationToken(InAuthorizationToken)
 	{
 		Curl = curl_easy_init();
@@ -216,12 +217,24 @@ public:
 #endif
 	}
 
+	/** Gets the domain name for this request */
+	const FString& GetName() const
+	{
+		return Domain;
+	}
+
 	/** Gets the domain name for this request */
 	const FString& GetDomain() const
 	{
 		return Domain;
 	}
 
+	/** Gets the effective domain name for this request */
+	const FString& GetEffectiveDomain() const
+	{
+		return EffectiveDomain;
+	}
+
 	/** Returns the HTTP response code.*/
 	const int64 GetResponseCode() const
 	{
@@ -439,6 +452,7 @@ private:
 	TArray<uint8>			ResponseBuffer;
 	TArray<FString>			Headers;
 	FString					Domain;
+	FString					EffectiveDomain;
 	FHttpAccessToken*		AuthorizationToken;
 
 	/**
@@ -458,10 +472,10 @@ private:
 		TRACE_CPUPROFILER_EVENT_SCOPE(HttpDDC_CurlPerform);
 
 		// Setup request options
-		FString Url = FString::Printf(TEXT("%s/%s"), *Domain, Uri);
+		FString Url = FString::Printf(TEXT("%s/%s"), *EffectiveDomain, Uri);
 		curl_easy_setopt(Curl, CURLOPT_URL, TCHAR_TO_ANSI(*Url));
 
-		// Setup response header buffer. If caller has not setup a response data buffer, use interal.
+		// Setup response header buffer. If caller has not setup a response data buffer, use internal.
 		WriteHeaderBufferPtr = &ResponseHeader;
 		if (WriteDataBufferPtr == nullptr)
 		{
@@ -551,8 +565,9 @@ private:
 				UE_LOG(
 					LogDerivedDataCache, 
 					Verbose, 
-					TEXT("Finished %s HTTP cache entry (response %d) from %s. %s"), 
-					VerbStr, 
+					TEXT("%s: Finished %s HTTP cache entry (response %d) from %s. %s"), 
+					*GetName(),
+					VerbStr,
 					ResponseCode, 
 					Uri,
 					*AdditionalInfo
@@ -570,7 +585,8 @@ private:
 					UE_LOG(
 						LogDerivedDataCache,
 						Verbose,
-						TEXT("Failed %s HTTP cache entry (response %d) from %s. Response: %s"),
+						TEXT("%s: Failed %s HTTP cache entry (response %d) from %s. Response: %s"),
+						*GetName(),
 						VerbStr,
 						ResponseCode,
 						Uri,
@@ -582,7 +598,8 @@ private:
 					UE_LOG(
 						LogDerivedDataCache,
 						Display,
-						TEXT("Failed %s HTTP cache entry (response %d) from %s. Response: %s"),
+						TEXT("%s: Failed %s HTTP cache entry (response %d) from %s. Response: %s"),
+						*GetName(),
 						VerbStr,
 						ResponseCode,
 						Uri,
@@ -598,8 +615,9 @@ private:
 			UE_LOG(
 				LogDerivedDataCache, 
 				Display, 
-				TEXT("Error while connecting to %s: %s"), 
-				*Domain, 
+				TEXT("%s: Error while connecting to %s: %s"), 
+				*GetName(),
+				*EffectiveDomain,
 				ANSI_TO_TCHAR(curl_easy_strerror(Result))
 			);
 		}
@@ -632,28 +650,28 @@ private:
 			FString DebugText(ConvertedString.Length(), ConvertedString.Get());
 			DebugText.ReplaceInline(TEXT("\n"), TEXT(""), ESearchCase::CaseSensitive);
 			DebugText.ReplaceInline(TEXT("\r"), TEXT(""), ESearchCase::CaseSensitive);
-			UE_LOG(LogDerivedDataCache, VeryVerbose, TEXT("%p: '%s'"), Request, *DebugText);
+			UE_LOG(LogDerivedDataCache, VeryVerbose, TEXT("%s: %p: '%s'"), *Request->GetName(), Request, *DebugText);
 		}
 		break;
 
 		case CURLINFO_HEADER_IN:
-			UE_LOG(LogDerivedDataCache, VeryVerbose, TEXT("%p: Received header (%d bytes)"), Request, DebugInfoSize);
+			UE_LOG(LogDerivedDataCache, VeryVerbose, TEXT("%s: %p: Received header (%d bytes)"), *Request->GetName(), Request, DebugInfoSize);
 			break;
 
 		case CURLINFO_DATA_IN:
-			UE_LOG(LogDerivedDataCache, VeryVerbose, TEXT("%p: Received data (%d bytes)"), Request, DebugInfoSize);
+			UE_LOG(LogDerivedDataCache, VeryVerbose, TEXT("%s: %p: Received data (%d bytes)"), *Request->GetName(), Request, DebugInfoSize);
 			break;
 
 		case CURLINFO_DATA_OUT:
-			UE_LOG(LogDerivedDataCache, VeryVerbose, TEXT("%p: Sent data (%d bytes)"), Request, DebugInfoSize);
+			UE_LOG(LogDerivedDataCache, VeryVerbose, TEXT("%s: %p: Sent data (%d bytes)"), *Request->GetName(), Request, DebugInfoSize);
 			break;
 
 		case CURLINFO_SSL_DATA_IN:
-			UE_LOG(LogDerivedDataCache, VeryVerbose, TEXT("%p: Received SSL data (%d bytes)"), Request, DebugInfoSize);
+			UE_LOG(LogDerivedDataCache, VeryVerbose, TEXT("%s: %p: Received SSL data (%d bytes)"), *Request->GetName(), Request, DebugInfoSize);
 			break;
 
 		case CURLINFO_SSL_DATA_OUT:
-			UE_LOG(LogDerivedDataCache, VeryVerbose, TEXT("%p: Sent SSL data (%d bytes)"), Request, DebugInfoSize);
+			UE_LOG(LogDerivedDataCache, VeryVerbose, TEXT("%s: %p: Sent SSL data (%d bytes)"), *Request->GetName(), Request, DebugInfoSize);
 			break;
 		}
 
@@ -778,13 +796,13 @@ bool ShouldAbortForShutdown();
  */
 struct FRequestPool
 {
-	FRequestPool(const TCHAR* InServiceUrl, FHttpAccessToken* InAuthorizationToken, uint32 PoolSize)
+	FRequestPool(const TCHAR* InServiceUrl, const TCHAR* InEffectiveServiceUrl, FHttpAccessToken* InAuthorizationToken, uint32 PoolSize)
 	{
 		Pool.AddUninitialized(PoolSize);
 		for (uint8 i = 0; i < Pool.Num(); ++i)
 		{
 			Pool[i].Usage = 0u;
-			Pool[i].Request = new FHttpRequest(InServiceUrl, InAuthorizationToken, true);
+			Pool[i].Request = new FHttpRequest(InServiceUrl, InEffectiveServiceUrl, InAuthorizationToken, true);
 		}
 		
 	}
@@ -1391,14 +1409,14 @@ private:
 				// Parse the response and move the data to the target requests.
 				if (ParseBatchedResponse(Response, ResponseSize, Entries))
 				{
-					UE_LOG(LogDerivedDataCache, VeryVerbose, TEXT("Batch query with %d operations completed."), Entries.Num());
+					UE_LOG(LogDerivedDataCache, VeryVerbose, TEXT("%s: Batch query with %d operations completed."), *Request->GetName(), Entries.Num());
 					return;
 				}
 			}
 		}
 		
 		// If we get here the request failed.
-		UE_LOG(LogDerivedDataCache, Display, TEXT("Batch query failed. Query: %s"), ANSI_TO_TCHAR((ANSICHAR*)RequestData.GetData()));
+		UE_LOG(LogDerivedDataCache, Display, TEXT("%s: Batch query failed. Query: %s"), *Request->GetName(), ANSI_TO_TCHAR((ANSICHAR*)RequestData.GetData()));
 
 		// Set all batch operations to failures
 		for (FQueuedBatchEntry Entry : Entries)
@@ -1909,7 +1927,7 @@ bool VerifyRequest(const FHttpRequest* Request, const TCHAR* Namespace, const TC
 		FIoHash ReceivedHash(ReceivedHashStr);
 		return VerifyPayload(ReceivedHash, Namespace, Bucket, CacheKey, Payload);
 	}
-	UE_LOG(LogDerivedDataCache, Warning, TEXT("HTTP server did not send a content hash. Wrong server version?"));
+	UE_LOG(LogDerivedDataCache, Warning, TEXT("%s: HTTP server did not send a content hash. Wrong server version?"), *Request->GetName());
 	return true;
 }
 
@@ -2175,10 +2193,57 @@ FHttpDerivedDataBackend::FHttpDerivedDataBackend(
 #endif
 	if (IsServiceReady() && AcquireAccessToken())
 	{
-		GetRequestPools[0] = MakeUnique<FRequestPool>(InServiceUrl, Access.Get(), UE_HTTPDDC_GET_REQUEST_POOL_SIZE);
-		GetRequestPools[1] = MakeUnique<FRequestPool>(InServiceUrl, Access.Get(), 1);
-		PutRequestPools[0] = MakeUnique<FRequestPool>(InServiceUrl, Access.Get(), UE_HTTPDDC_PUT_REQUEST_POOL_SIZE);
-		PutRequestPools[1] = MakeUnique<FRequestPool>(InServiceUrl, Access.Get(), 1);
+		FString EffectiveDomain;
+		FString OriginalDomainPrefix;
+		TAnsiStringBuilder<64> DomainResolveName;
+
+		if (Domain.StartsWith(TEXT("http://")))
+		{
+			DomainResolveName << Domain.RightChop(7);
+			OriginalDomainPrefix = TEXT("http://");
+		}
+		else if (Domain.StartsWith(TEXT("https://")))
+		{
+			DomainResolveName << Domain.RightChop(8);
+			OriginalDomainPrefix = TEXT("https://");
+		}
+		else
+		{
+			DomainResolveName << Domain;
+		}
+
+		addrinfo* AddrResult = nullptr;
+		addrinfo AddrHints;
+		FMemory::Memset(&AddrHints, 0, sizeof(AddrHints));
+		AddrHints.ai_flags = AI_CANONNAME;
+		AddrHints.ai_family = AF_UNSPEC;
+		if (!::getaddrinfo(*DomainResolveName, nullptr, &AddrHints, &AddrResult))
+		{
+			if (AddrResult->ai_canonname)
+			{
+				// Swap the domain with a canonical name from DNS so that if we are using regional redirection, we pin to a region.
+				EffectiveDomain = OriginalDomainPrefix + ANSI_TO_TCHAR(AddrResult->ai_canonname);
+
+				UE_LOG(LogDerivedDataCache, Display,
+					TEXT("%s: Pinned to %s based on DNS canonical name."),
+					*Domain, *EffectiveDomain);
+			}
+			else
+			{
+				EffectiveDomain = Domain;
+			}
+
+			::freeaddrinfo(AddrResult);
+		}
+		else
+		{
+			EffectiveDomain = Domain;
+		}
+
+		GetRequestPools[0] = MakeUnique<FRequestPool>(*Domain, *EffectiveDomain, Access.Get(), UE_HTTPDDC_GET_REQUEST_POOL_SIZE);
+		GetRequestPools[1] = MakeUnique<FRequestPool>(*Domain, *EffectiveDomain, Access.Get(), 1);
+		PutRequestPools[0] = MakeUnique<FRequestPool>(*Domain, *EffectiveDomain, Access.Get(), UE_HTTPDDC_PUT_REQUEST_POOL_SIZE);
+		PutRequestPools[1] = MakeUnique<FRequestPool>(*Domain, *EffectiveDomain, Access.Get(), 1);
 		bIsUsable = true;
 	}
 
@@ -2225,17 +2290,17 @@ bool FHttpDerivedDataBackend::ApplyDebugOptions(FBackendDebugOptions& InOptions)
 
 bool FHttpDerivedDataBackend::IsServiceReady()
 {
-	FHttpRequest Request(*Domain, nullptr, false);
+	FHttpRequest Request(*Domain, *Domain, nullptr, false);
 	FHttpRequest::Result Result = Request.PerformBlockingDownload(TEXT("health/ready"), nullptr);
 	
 	if (Result == FHttpRequest::Success && Request.GetResponseCode() == 200)
 	{
-		UE_LOG(LogDerivedDataCache, Display, TEXT("HTTP DDC service status: %s."), *Request.GetResponseAsString());
+		UE_LOG(LogDerivedDataCache, Display, TEXT("%s: HTTP DDC service status: %s."), *Request.GetName(), *Request.GetResponseAsString());
 		return true;
 	}
 	else
 	{
-		UE_LOG(LogDerivedDataCache, Warning, TEXT("Unable to reach HTTP DDC service at %s. Status: %d . Response: %s"), *Domain, Request.GetResponseCode(), *Request.GetResponseAsString());
+		UE_LOG(LogDerivedDataCache, Warning, TEXT("%s: Unable to reach HTTP DDC service at %s. Status: %d . Response: %s"), *Request.GetName(), *Domain, Request.GetResponseCode(), *Request.GetResponseAsString());
 	}
 
 	return false;
@@ -2273,7 +2338,7 @@ bool FHttpDerivedDataBackend::AcquireAccessToken()
 		FString AuthDomain(DomainEnd, *OAuthProvider);
 		FString Uri(*OAuthProvider + DomainEnd + 1);
 
-		FHttpRequest Request(*AuthDomain, nullptr, false);
+		FHttpRequest Request(*AuthDomain, *AuthDomain, nullptr, false);
 
 		// If contents of the secret string is a file path, resolve and read form data.
 		if (OAuthSecret.StartsWith(TEXT("file://")))
@@ -2287,7 +2352,7 @@ bool FHttpDerivedDataBackend::AcquireAccessToken()
 			}
 			else
 			{
-				UE_LOG(LogDerivedDataCache, Warning, TEXT("Failed to read OAuth form data file (%s)."), *OAuthSecret);
+				UE_LOG(LogDerivedDataCache, Warning, TEXT("%s: Failed to read OAuth form data file (%s)."), *Request.GetName(), *OAuthSecret);
 				return false;
 			}
 		}
@@ -2321,7 +2386,7 @@ bool FHttpDerivedDataBackend::AcquireAccessToken()
 						Access = MakeUnique<FHttpAccessToken>();
 					}
 					Access->SetHeader(*AccessTokenString);
-					UE_LOG(LogDerivedDataCache, Display, TEXT("Logged in to HTTP DDC services. Expires in %d seconds."), ExpiryTimeSeconds);
+					UE_LOG(LogDerivedDataCache, Display, TEXT("%s: Logged in to HTTP DDC services. Expires in %d seconds."), *Request.GetName(), ExpiryTimeSeconds);
 
 					//Schedule a refresh of the token ahead of expiry time (this will not work in commandlets)
 					if (!IsRunningCommandlet())
@@ -2342,7 +2407,7 @@ bool FHttpDerivedDataBackend::AcquireAccessToken()
 		}
 		else
 		{
-			UE_LOG(LogDerivedDataCache, Warning, TEXT("Failed to log in to HTTP services. Server responed with code %d."), Request.GetResponseCode());
+			UE_LOG(LogDerivedDataCache, Warning, TEXT("%s: Failed to log in to HTTP services. Server responed with code %d."), *Request.GetName(), Request.GetResponseCode());
 			FailedLoginAttempts++;
 		}
 	}
@@ -2859,7 +2924,7 @@ bool FHttpDerivedDataBackend::GetCacheValue(
 
 		if (EnumHasAnyFlags(Policy, ECachePolicy::SkipData))
 		{
-			if (CachedDataProbablyExistsBatch<FValue>(Name, Key, {OutValue}, [](const FValue& Value) { return TEXT("Default"); }))
+			if (CachedDataProbablyExistsBatch<FValue>(Name, Key, {OutValue}, [](const FValue& Value) { return FString(TEXT("Default")); }))
 			{
 				return true;
 			}
@@ -2867,7 +2932,7 @@ bool FHttpDerivedDataBackend::GetCacheValue(
 		else
 		{
 			TArray<FCompressedBuffer> Buffers;
-			if (TryGetCachedDataBatch<FValue>(Name, Key, {OutValue}, Buffers, [](const FValue& Value) { return TEXT("Default"); }))
+			if (TryGetCachedDataBatch<FValue>(Name, Key, {OutValue}, Buffers, [](const FValue& Value) { return FString(TEXT("Default")); }))
 			{
 				check(Buffers.Num() == 1);
 				OutValue = FValue(MoveTemp(Buffers[0]));
@@ -2926,14 +2991,14 @@ FOptionalCacheRecord FHttpDerivedDataBackend::GetCacheRecord(
 		}
 	}
 
-	if (!CachedDataProbablyExistsBatch<FValueWithId>(Name, Key, RequiredHeads, [](const FValueWithId& Value) { return *WriteToString<16>(Value.GetId()); }))
+	if (!CachedDataProbablyExistsBatch<FValueWithId>(Name, Key, RequiredHeads, [](const FValueWithId& Value) { return WriteToString<16>(Value.GetId()); }))
 	{
 		OutStatus = EStatus::Error;
 		return FOptionalCacheRecord();
 	}
 
 	TArray<FCompressedBuffer> FetchedBuffers;
-	if (!TryGetCachedDataBatch<FValueWithId>(Name, Key, RequiredGets, FetchedBuffers, [](const FValueWithId& Value) { return *WriteToString<16>(Value.GetId()); }))
+	if (!TryGetCachedDataBatch<FValueWithId>(Name, Key, RequiredGets, FetchedBuffers, [](const FValueWithId& Value) { return WriteToString<16>(Value.GetId()); }))
 	{
 		OutStatus = EStatus::Error;
 		return FOptionalCacheRecord();
@@ -2989,7 +3054,7 @@ bool FHttpDerivedDataBackend::TryGetCachedDataBatch(
 		{
 			UE_LOG(LogDerivedDataCache, Verbose,
 				TEXT("%s: Cache miss with missing value %s with hash %s for %s from '%.*s'"),
-				*GetName(), ValuePrinter(Value), *WriteToString<48>(RawHash), *WriteToString<96>(Key),
+				*GetName(), *ValuePrinter(Value), *WriteToString<48>(RawHash), *WriteToString<96>(Key),
 				Name.Len(), Name.GetData());
 			return false;
 		}
@@ -2997,7 +3062,7 @@ bool FHttpDerivedDataBackend::TryGetCachedDataBatch(
 		{
 			UE_LOG(LogDerivedDataCache, Display,
 				TEXT("%s: Cache miss with corrupted value %s with hash %s for %s from '%.*s'"),
-				*GetName(), ValuePrinter(Value), *WriteToString<48>(RawHash),
+				*GetName(), *ValuePrinter(Value), *WriteToString<48>(RawHash),
 				*WriteToString<96>(Key), Name.Len(), Name.GetData());
 			return false;
 		}
@@ -3042,7 +3107,7 @@ bool FHttpDerivedDataBackend::CachedDataProbablyExistsBatch(
 		{
 			UE_LOG(LogDerivedDataCache, Verbose,
 				TEXT("%s: Cache miss with missing value %s with hash %s for %s from '%.*s'"),
-				*GetName(), ValuePrinter(Value), *WriteToString<48>(RawHash), *WriteToString<96>(Key),
+				*GetName(), *ValuePrinter(Value), *WriteToString<48>(RawHash), *WriteToString<96>(Key),
 				Name.Len(), Name.GetData());
 			return false;
 		}
@@ -3585,7 +3650,7 @@ void FHttpDerivedDataBackend::GetChunks(
 					else
 					{
 						TArray<FCompressedBuffer> ValueBuffers;
-						if (TryGetCachedDataBatch<FValueWithId>(Request.Name, Request.Key, ::MakeArrayView({ ValueWithId }), ValueBuffers, [](const FValueWithId& Value) { return *WriteToString<16>(Value.GetId()); }))
+						if (TryGetCachedDataBatch<FValueWithId>(Request.Name, Request.Key, ::MakeArrayView({ ValueWithId }), ValueBuffers, [](const FValueWithId& Value) { return WriteToString<16>(Value.GetId()); }))
 						{
 							ValueBuffer = ValueBuffers[0];
 							ValueReader.SetSource(ValueBuffer);
@@ -3609,7 +3674,7 @@ void FHttpDerivedDataBackend::GetChunks(
 				else
 				{
 					TArray<FCompressedBuffer> ValueBuffers;
-					if (TryGetCachedDataBatch<FValue>(Request.Name, Request.Key, { Value }, ValueBuffers, [](const FValue& Value) { return TEXT("Default"); }))
+					if (TryGetCachedDataBatch<FValue>(Request.Name, Request.Key, { Value }, ValueBuffers, [](const FValue& Value) { return FString(TEXT("Default")); }))
 					{
 						ValueBuffer = ValueBuffers[0];
 						ValueReader.SetSource(ValueBuffer);
