DDC: Fixed a memory leak in the cache

#rnx

diff --git a/Engine/Source/Developer/DerivedDataCache/Private/DerivedDataCacheStoreHierarchy.cpp b/Engine/Source/Developer/DerivedDataCache/Private/DerivedDataCacheStoreHierarchy.cpp
--- a/Engine/Source/Developer/DerivedDataCache/Private/DerivedDataCacheStoreHierarchy.cpp
+++ b/Engine/Source/Developer/DerivedDataCache/Private/DerivedDataCacheStoreHierarchy.cpp
@@ -150,6 +150,8 @@ private:
 class FCacheStoreHierarchy::FBatchBase
 {
 public:
+	virtual ~FBatchBase() = default;
+
 	void AddRef()
 	{
 		ReferenceCount.fetch_add(1, std::memory_order_relaxed);
@@ -319,7 +321,7 @@ bool FCacheStoreHierarchy::CanQueryIfError(const ECachePolicy Policy, const ECac
 ///////////////////////////////////////////////////////////////////////////////////////////////////////////////////////
 
 template <typename Params>
-class FCacheStoreHierarchy::TPutBatch : public FBatchBase, public Params
+class FCacheStoreHierarchy::TPutBatch final : public FBatchBase, public Params
 {
 	using FPutRequest = typename Params::FPutRequest;
 	using FGetRequest = typename Params::FGetRequest;
@@ -552,7 +554,7 @@ void FCacheStoreHierarchy::TPutBatch<Params>::CompletePutRequest(FPutResponse&&
 ///////////////////////////////////////////////////////////////////////////////////////////////////////////////////////
 
 template <typename Params>
-class FCacheStoreHierarchy::TGetBatch : public FBatchBase, public Params
+class FCacheStoreHierarchy::TGetBatch final : public FBatchBase, public Params
 {
 	using FPutRequest = typename Params::FPutRequest;
 	using FGetRequest = typename Params::FGetRequest;
@@ -919,7 +921,7 @@ void FCacheStoreHierarchy::GetValue(
 
 ///////////////////////////////////////////////////////////////////////////////////////////////////////////////////////
 
-class FCacheStoreHierarchy::FGetChunksBatch : public FBatchBase
+class FCacheStoreHierarchy::FGetChunksBatch final : public FBatchBase
 {
 public:
 	static void Begin(
@@ -1115,7 +1117,7 @@ void FCacheStoreHierarchy::LegacyGet(
 
 ///////////////////////////////////////////////////////////////////////////////////////////////////////////////////////
 
-class FCacheStoreHierarchy::FLegacyDeleteBatch : public FBatchBase
+class FCacheStoreHierarchy::FLegacyDeleteBatch final : public FBatchBase
 {
 public:
 	static void Begin(
