DDC: Fixed a test failure for small values in the file system cache

#rnx

diff --git a/Engine/Source/Developer/DerivedDataCache/Private/FileSystemDerivedDataBackend.cpp b/Engine/Source/Developer/DerivedDataCache/Private/FileSystemDerivedDataBackend.cpp
--- a/Engine/Source/Developer/DerivedDataCache/Private/FileSystemDerivedDataBackend.cpp
+++ b/Engine/Source/Developer/DerivedDataCache/Private/FileSystemDerivedDataBackend.cpp
@@ -2213,11 +2213,7 @@ bool FFileSystemCacheStore::GetCacheValueOnly(
 		return false;
 	}
 
-	if (EnumHasAnyFlags(Policy, ECachePolicy::SkipData))
-	{
-		OutValue = FValue(RawHash, RawSize);
-	}
-	else if (const FCbAttachment* const Attachment = Package.FindAttachment(RawHash))
+	if (const FCbAttachment* const Attachment = Package.FindAttachment(RawHash))
 	{
 		const FCompressedBuffer& Data = Attachment->AsCompressedBinary();
 		if (Data.GetRawHash() != RawHash || Data.GetRawSize() != RawSize)
