DDC: Fixed an error in GetChunks in the pak file cache

#rnx

diff --git a/Engine/Source/Developer/DerivedDataCache/Private/PakFileDerivedDataBackend.cpp b/Engine/Source/Developer/DerivedDataCache/Private/PakFileDerivedDataBackend.cpp
--- a/Engine/Source/Developer/DerivedDataCache/Private/PakFileDerivedDataBackend.cpp
+++ b/Engine/Source/Developer/DerivedDataCache/Private/PakFileDerivedDataBackend.cpp
@@ -722,7 +722,6 @@ void FPakFileDerivedDataBackend::GetChunks(
 	FValueId ValueId;
 	FCacheKey ValueKey;
 	FCompressedBufferReader ValueReader;
-	EStatus ValueStatus = EStatus::Error;
 	FOptionalCacheRecord Record;
 	for (const FCacheGetChunkRequest& Request : SortedRequests)
 	{
@@ -731,7 +730,6 @@ void FPakFileDerivedDataBackend::GetChunks(
 		COOK_STAT(auto Timer = bExistsOnly ? UsageStats.TimeProbablyExists() : UsageStats.TimeGet());
 		if (!(bHasValue && ValueKey == Request.Key && ValueId == Request.Id) || ValueReader.HasSource() < !bExistsOnly)
 		{
-			ValueStatus = EStatus::Error;
 			ValueReader.ResetSource();
 			ValueKey = {};
 			ValueId.Reset();
@@ -752,6 +750,7 @@ void FPakFileDerivedDataBackend::GetChunks(
 					Value = ValueWithId;
 					ValueId = Request.Id;
 					ValueKey = Request.Key;
+					EStatus ValueStatus;
 					GetCacheContent(Request.Name, Request.Key, Request.Policy, Value, ValueStatus);
 				}
 			}
@@ -762,26 +761,20 @@ void FPakFileDerivedDataBackend::GetChunks(
 				UE_LOG(LogDerivedDataCache, Verbose, TEXT("%s: Cache hit for %s from '%s'"),
 					*CachePath, *WriteToString<96>(Request.Key, '/', Request.Id), *Request.Name);
 				COOK_STAT(Timer.AddHit(Value.HasData() ? RawSize : 0));
-				if (OnComplete)
+				FSharedBuffer Buffer;
+				if (Value.HasData() && !bExistsOnly)
 				{
-					FSharedBuffer Buffer;
-					if (Value.HasData() && !bExistsOnly)
-					{
-						FCompressedBufferReaderSourceScope Source(ValueReader, Value.GetData());
-						Buffer = ValueReader.Decompress(RawOffset, RawSize);
-					}
-					OnComplete({Request.Name, Request.Key, Request.Id, Request.RawOffset,
-						RawSize, Value.GetRawHash(), MoveTemp(Buffer), Request.UserData, ValueStatus});
+					FCompressedBufferReaderSourceScope Source(ValueReader, Value.GetData());
+					Buffer = ValueReader.Decompress(RawOffset, RawSize);
 				}
+				const EStatus ChunkStatus = bExistsOnly || Buffer.GetSize() == RawSize ? EStatus::Ok : EStatus::Error;
+				OnComplete({Request.Name, Request.Key, Request.Id, Request.RawOffset,
+					RawSize, Value.GetRawHash(), MoveTemp(Buffer), Request.UserData, ChunkStatus});
 				continue;
 			}
 		}
 
-		if (OnComplete)
-		{
-			OnComplete({Request.Name, Request.Key, Request.Id, Request.RawOffset,
-				0, {}, {}, Request.UserData, EStatus::Error});
-		}
+		OnComplete(Request.MakeResponse(EStatus::Error));
 	}
 }
 
