DDC: Added compression support to legacy cache values

The cache stores can individually control their "legacy mode" now, which offers a choice of ValueOnly, ValueWithLegacyFallback, LegacyOnly.
- Using ValueOnly will forward every legacy request through the Value API, which compresses values by default and supports values larger than 2 GiB.
- Using ValueWithLegacyFallback behaves like ValueOnly, but misses from GetValue are forwarded to GetCachedData, and stored with PutValue if found.
- Using LegacyOnly will forward every legacy request through the Legacy API.

FLegacyCacheValue uses shared state to ensure that each value is compressed or decompressed at most once.

#rnx

diff --git a/Engine/Source/Developer/DerivedDataCache/Private/DerivedDataBackendInterface.h b/Engine/Source/Developer/DerivedDataCache/Private/DerivedDataBackendInterface.h
--- a/Engine/Source/Developer/DerivedDataCache/Private/DerivedDataBackendInterface.h
+++ b/Engine/Source/Developer/DerivedDataCache/Private/DerivedDataBackendInterface.h
@@ -5,6 +5,7 @@
 
 #include "CoreTypes.h"
 #include "Containers/BitArray.h"
+#include "Containers/StringView.h"
 #include "DerivedDataCache.h"
 #include "DerivedDataLegacyCacheStore.h"
 #include "HAL/LowLevelMemTracker.h"
@@ -34,9 +35,19 @@ LLM_DECLARE_TAG_API(UntaggedDDCResult, DERIVEDDATACACHE_API);
 namespace UE::DerivedData
 {
 
+enum class EBackendLegacyMode
+{
+	/** Use only GetValue/PutValue. */
+	ValueOnly,
+	/** Use GetValue/PutValue with a fallback to GetCachedData+PutValue for misses. */
+	ValueWithLegacyFallback,
+	/** Use only GetCachedData/PutCachedData. */
+	LegacyOnly,
+};
+
 /**
-	* Speed classes. Higher values are faster so > / < comparisons make sense.
-	*/
+ * Speed classes. Higher values are faster so > / < comparisons make sense.
+ */
 enum class EBackendSpeedClass
 {
 	Unknown,		/* Don't know yet*/
@@ -199,24 +210,26 @@ public:
 	 */
 	virtual bool ApplyDebugOptions(FBackendDebugOptions& InOptions) = 0;
 
+	virtual EBackendLegacyMode GetLegacyMode() const { return EBackendLegacyMode::LegacyOnly; }
+
 	virtual void LegacyPut(
 		TConstArrayView<FLegacyCachePutRequest> Requests,
 		IRequestOwner& Owner,
-		FOnLegacyCachePutComplete&& OnComplete) override;
+		FOnLegacyCachePutComplete&& OnComplete) final;
 
 	virtual void LegacyGet(
 		TConstArrayView<FLegacyCacheGetRequest> Requests,
 		IRequestOwner& Owner,
-		FOnLegacyCacheGetComplete&& OnComplete) override;
+		FOnLegacyCacheGetComplete&& OnComplete) final;
 
 	virtual void LegacyDelete(
 		TConstArrayView<FLegacyCacheDeleteRequest> Requests,
 		IRequestOwner& Owner,
-		FOnLegacyCacheDeleteComplete&& OnComplete) override;
+		FOnLegacyCacheDeleteComplete&& OnComplete) final;
 
-	virtual void LegacyStats(FDerivedDataCacheStatsNode& OutNode) override;
+	virtual void LegacyStats(FDerivedDataCacheStatsNode& OutNode) final;
 
-	virtual bool LegacyDebugOptions(FBackendDebugOptions& Options) override;
+	virtual bool LegacyDebugOptions(FBackendDebugOptions& Options) final;
 };
 
 class FDerivedDataBackend
@@ -273,6 +286,42 @@ public:
 
 /** Lexical conversions from and to enums */
 
+[[nodiscard]] inline const TCHAR* LexToString(const EBackendLegacyMode Value)
+{
+	switch (Value)
+	{
+	case EBackendLegacyMode::ValueOnly:
+		return TEXT("ValueOnly");
+	case EBackendLegacyMode::ValueWithLegacyFallback:
+		return TEXT("ValueWithLegacyFallback");
+	case EBackendLegacyMode::LegacyOnly:
+		return TEXT("LegacyOnly");
+	default:
+		checkNoEntry();
+		return TEXT("Unknown");
+	}
+}
+
+[[nodiscard]] inline bool TryLexFromString(EBackendLegacyMode& OutValue, const FStringView String)
+{
+	if (String == TEXTVIEW("ValueOnly"))
+	{
+		OutValue = EBackendLegacyMode::ValueOnly;
+		return true;
+	}
+	if (String == TEXTVIEW("ValueWithLegacyFallback"))
+	{
+		OutValue = EBackendLegacyMode::ValueWithLegacyFallback;
+		return true;
+	}
+	if (String == TEXTVIEW("LegacyOnly"))
+	{
+		OutValue = EBackendLegacyMode::LegacyOnly;
+		return true;
+	}
+	return false;
+}
+
 inline const TCHAR* LexToString(FDerivedDataBackendInterface::ESpeedClass SpeedClass)
 {
 	switch (SpeedClass)
diff --git a/Engine/Source/Developer/DerivedDataCache/Private/DerivedDataBackends.cpp b/Engine/Source/Developer/DerivedDataCache/Private/DerivedDataBackends.cpp
--- a/Engine/Source/Developer/DerivedDataCache/Private/DerivedDataBackends.cpp
+++ b/Engine/Source/Developer/DerivedDataCache/Private/DerivedDataBackends.cpp
@@ -3,8 +3,10 @@
 #include "CoreTypes.h"
 #include "Containers/StringView.h"
 #include "DerivedDataBackendInterface.h"
+#include "DerivedDataCachePrivate.h"
 #include "DerivedDataCacheStore.h"
 #include "DerivedDataCacheUsageStats.h"
+#include "DerivedDataRequestOwner.h"
 #include "FileBackedDerivedDataBackend.h"
 #include "HAL/FileManager.h"
 #include "HAL/IConsoleManager.h"
@@ -1335,9 +1337,14 @@ void FDerivedDataBackendInterface::LegacyPut(
 	IRequestOwner& Owner,
 	FOnLegacyCachePutComplete&& OnComplete)
 {
+	if (GetLegacyMode() != EBackendLegacyMode::LegacyOnly)
+	{
+		return ILegacyCacheStore::LegacyPut(Requests, Owner, MoveTemp(OnComplete));
+	}
+
 	for (const FLegacyCachePutRequest& Request : Requests)
 	{
-		FCompositeBuffer CompositeValue = Request.Value;
+		FCompositeBuffer CompositeValue = Request.Value.GetRawData();
 		Request.Key.WriteValueTrailer(CompositeValue);
 
 		checkf(CompositeValue.GetSize() < MAX_int32,
@@ -1355,10 +1362,61 @@ void FDerivedDataBackendInterface::LegacyPut(
 }
 
 void FDerivedDataBackendInterface::LegacyGet(
-	const TConstArrayView<FLegacyCacheGetRequest> Requests,
+	TConstArrayView<FLegacyCacheGetRequest> Requests,
 	IRequestOwner& Owner,
 	FOnLegacyCacheGetComplete&& OnComplete)
 {
+	const EBackendLegacyMode LegacyMode = GetLegacyMode();
+	if (LegacyMode == EBackendLegacyMode::ValueOnly)
+	{
+		return ILegacyCacheStore::LegacyGet(Requests, Owner, MoveTemp(OnComplete));
+	}
+
+	// Make a blocking query to the value cache and fall back to the legacy cache for requests with errors.
+
+	TArray<FLegacyCacheGetRequest> LegacyRequests;
+	if (LegacyMode == EBackendLegacyMode::ValueWithLegacyFallback)
+	{
+		TArray<FLegacyCacheGetRequest, TInlineAllocator<8>> ValueRequests;
+		ValueRequests.Reserve(Requests.Num());
+		uint64 RequestIndex = 0;
+		for (const FLegacyCacheGetRequest& Request : Requests)
+		{
+			ValueRequests.Add_GetRef(Request).UserData = RequestIndex++;
+		}
+
+		FRequestOwner BlockingOwner(EPriority::Blocking);
+		ILegacyCacheStore::LegacyGet(ValueRequests, BlockingOwner, [this, &OnComplete, &Requests, &ValueRequests](FLegacyCacheGetResponse&& Response)
+		{
+			if (Response.Status != EStatus::Error)
+			{
+				const int32 Index = int32(Response.UserData);
+				Response.UserData = Requests[Index].UserData;
+				ValueRequests[Index].UserData = MAX_uint64;
+				OnComplete(MoveTemp(Response));
+			}
+		});
+		BlockingOwner.Wait();
+
+		for (const FLegacyCacheGetRequest& Request : ValueRequests)
+		{
+			if (Request.UserData != MAX_uint64)
+			{
+				LegacyRequests.Add(Requests[int32(Request.UserData)]);
+			}
+		}
+		if (LegacyRequests.IsEmpty())
+		{
+			return;
+		}
+		Requests = LegacyRequests;
+	}
+
+	// Query the legacy cache by translating the requests to legacy cache functions.
+
+	FRequestOwner AsyncOwner(FPlatformMath::Min(Owner.GetPriority(), EPriority::Highest));
+	AsyncOwner.KeepAlive();
+
 	TArray<FString, TInlineAllocator<8>> ExistsKeys;
 	TArray<const FLegacyCacheGetRequest*, TInlineAllocator<8>> ExistsRequests;
 
@@ -1390,7 +1448,15 @@ void FDerivedDataBackendInterface::LegacyGet(
 				}
 			}
 			const EStatus Status = Value ? EStatus::Ok : EStatus::Error;
-			OnComplete({Request.Name, Request.Key, MoveTemp(Value), Request.UserData, Status});
+			FLegacyCacheValue LegacyValue(FCompositeBuffer(MoveTemp(Value)));
+			if (LegacyValue.HasData() && LegacyMode == EBackendLegacyMode::ValueWithLegacyFallback)
+			{
+				Private::ExecuteInCacheThreadPool(AsyncOwner, [this, Request, LegacyValue](IRequestOwner& AsyncOwner, bool bCancel)
+				{
+					ILegacyCacheStore::LegacyPut({{Request.Name, Request.Key, LegacyValue}}, AsyncOwner, [](auto&&){});
+				});
+			}
+			OnComplete({Request.Name, Request.Key, MoveTemp(LegacyValue), Request.UserData, Status});
 		}
 	}
 
@@ -1422,6 +1488,11 @@ void FDerivedDataBackendInterface::LegacyDelete(
 	IRequestOwner& Owner,
 	FOnLegacyCacheDeleteComplete&& OnComplete)
 {
+	if (GetLegacyMode() != EBackendLegacyMode::LegacyOnly)
+	{
+		return ILegacyCacheStore::LegacyDelete(Requests, Owner, MoveTemp(OnComplete));
+	}
+
 	for (const FLegacyCacheDeleteRequest& Request : Requests)
 	{
 		RemoveCachedData(*Request.Key.GetShortKey(), Request.bTransient);
diff --git a/Engine/Source/Developer/DerivedDataCache/Private/DerivedDataCache.cpp b/Engine/Source/Developer/DerivedDataCache/Private/DerivedDataCache.cpp
--- a/Engine/Source/Developer/DerivedDataCache/Private/DerivedDataCache.cpp
+++ b/Engine/Source/Developer/DerivedDataCache/Private/DerivedDataCache.cpp
@@ -300,10 +300,16 @@ class FDerivedDataCache final
 					FDerivedDataBackend::Get().GetRoot().LegacyGet({LegacyRequest}, BlockingOwner,
 						[this, &bGetResult](FLegacyCacheGetResponse&& Response)
 						{
-							bGetResult = Response.Status == EStatus::Ok && Response.Value.GetSize() < MAX_int32;
+							const uint64 RawSize = Response.Value.GetRawSize();
+							bGetResult = Response.Status == EStatus::Ok && RawSize < MAX_int32;
 							if (bGetResult)
 							{
-								Data = MakeArrayView(static_cast<const uint8*>(Response.Value.GetData()), int32(Response.Value.GetSize()));
+								const FCompositeBuffer& RawData = Response.Value.GetRawData();
+								Data.Reset(int32(RawSize));
+								for (const FSharedBuffer& Segment : RawData.GetSegments())
+								{
+									Data.Append(static_cast<const uint8*>(Segment.GetData()), int32(Segment.GetSize()));
+								}
 							}
 						});
 					BlockingOwner.Wait();
@@ -380,7 +386,7 @@ class FDerivedDataCache final
 						FLegacyCachePutRequest LegacyRequest;
 						LegacyRequest.Name = DebugContext;
 						LegacyRequest.Key = FLegacyCacheKey(CacheKey, FDerivedDataBackend::Get().GetMaxKeyLength());
-						LegacyRequest.Value = FCompositeBuffer(FSharedBuffer::Clone(MakeMemoryView(Data)));
+						LegacyRequest.Value = FLegacyCacheValue(FCompositeBuffer(FSharedBuffer::Clone(MakeMemoryView(Data))));
 						FRequestOwner BlockingOwner(EPriority::Blocking);
 						FDerivedDataBackend::Get().GetRoot().LegacyPut({LegacyRequest}, BlockingOwner, [](auto&&){});
 						BlockingOwner.Wait();
@@ -599,7 +605,7 @@ public:
 			FLegacyCachePutRequest LegacyRequest;
 			LegacyRequest.Name = DebugContext;
 			LegacyRequest.Key = FLegacyCacheKey(CacheKey, FDerivedDataBackend::Get().GetMaxKeyLength());
-			LegacyRequest.Value = FCompositeBuffer(FSharedBuffer::Clone(MakeMemoryView(Data)));
+			LegacyRequest.Value = FLegacyCacheValue(FCompositeBuffer(FSharedBuffer::Clone(MakeMemoryView(Data))));
 			FRequestOwner BlockingOwner(EPriority::Blocking);
 			FDerivedDataBackend::Get().GetRoot().LegacyPut({LegacyRequest}, BlockingOwner, [](auto&&){});
 			BlockingOwner.Wait();
diff --git a/Engine/Source/Developer/DerivedDataCache/Private/DerivedDataCachePrivate.h b/Engine/Source/Developer/DerivedDataCache/Private/DerivedDataCachePrivate.h
--- a/Engine/Source/Developer/DerivedDataCache/Private/DerivedDataCachePrivate.h
+++ b/Engine/Source/Developer/DerivedDataCache/Private/DerivedDataCachePrivate.h
@@ -9,8 +9,11 @@
 class FDerivedDataCacheInterface;
 class FQueuedThreadPool;
 
+template <typename FuncType> class TUniqueFunction;
+
 namespace UE::DerivedData { class FCacheRecord; }
 namespace UE::DerivedData { class ICache; }
+namespace UE::DerivedData { class IRequestOwner; }
 namespace UE::DerivedData { struct FCacheKey; }
 
 namespace UE::DerivedData::Private
@@ -27,6 +30,11 @@ uint64 GetCacheRecordCompressedSize(const FCacheRecord& Record);
 uint64 GetCacheRecordTotalRawSize(const FCacheRecord& Record);
 uint64 GetCacheRecordRawSize(const FCacheRecord& Record);
 
+// Implemented in DerivedDataCacheStoreAsync.cpp
+void ExecuteInCacheThreadPool(
+	IRequestOwner& Owner,
+	TUniqueFunction<void (IRequestOwner& Owner, bool bCancel)>&& Function);
+
 ///////////////////////////////////////////////////////////////////////////////////////////////////////////////////////
 
 template <typename CharType>
diff --git a/Engine/Source/Developer/DerivedDataCache/Private/DerivedDataCacheStoreAsync.cpp b/Engine/Source/Developer/DerivedDataCache/Private/DerivedDataCacheStoreAsync.cpp
--- a/Engine/Source/Developer/DerivedDataCache/Private/DerivedDataCacheStoreAsync.cpp
+++ b/Engine/Source/Developer/DerivedDataCache/Private/DerivedDataCacheStoreAsync.cpp
@@ -251,6 +251,14 @@ private:
 	FLazyEvent DoneEvent{EEventMode::ManualReset};
 };
 
+void Private::ExecuteInCacheThreadPool(
+	IRequestOwner& Owner,
+	TUniqueFunction<void (IRequestOwner& Owner, bool bCancel)>&& Function)
+{
+	FDerivedDataAsyncWrapperRequest* Request = new FDerivedDataAsyncWrapperRequest(Owner, MoveTemp(Function));
+	Request->Start(Owner.GetPriority());
+}
+
 template <typename RequestType, typename OnCompleteType, typename OnExecuteType>
 void FCacheStoreAsync::Execute(
 	COOK_STAT(CookStatsFunction OnAddStats,)
@@ -278,8 +286,12 @@ void FCacheStoreAsync::Execute(
 		return ExecuteWithStats(Requests, Owner, MoveTemp(OnComplete));
 	}
 
-	FDerivedDataAsyncWrapperRequest* Request = new FDerivedDataAsyncWrapperRequest(Owner,
-		[this, Requests = TArray<RequestType>(Requests), OnComplete = MoveTemp(OnComplete), ExecuteWithStats = MoveTemp(ExecuteWithStats)](IRequestOwner& Owner, bool bCancel) mutable
+	Private::ExecuteInCacheThreadPool(Owner,
+		[this,
+		Requests = TArray<RequestType>(Requests),
+		OnComplete = MoveTemp(OnComplete),
+		ExecuteWithStats = MoveTemp(ExecuteWithStats)]
+		(IRequestOwner& Owner, bool bCancel) mutable
 		{
 			if (!bCancel)
 			{
@@ -291,7 +303,6 @@ void FCacheStoreAsync::Execute(
 				FDerivedDataBackend::Get().AddToAsyncCompletionCounter(-Requests.Num());
 			}
 		});
-	Request->Start(Owner.GetPriority());
 }
 
 ILegacyCacheStore* CreateCacheStoreAsync(ILegacyCacheStore* InnerCache, ECacheStoreFlags InnerFlags, bool bCacheInFlightPuts)
diff --git a/Engine/Source/Developer/DerivedDataCache/Private/DerivedDataCacheStoreHierarchy.cpp b/Engine/Source/Developer/DerivedDataCache/Private/DerivedDataCacheStoreHierarchy.cpp
--- a/Engine/Source/Developer/DerivedDataCache/Private/DerivedDataCacheStoreHierarchy.cpp
+++ b/Engine/Source/Developer/DerivedDataCache/Private/DerivedDataCacheStoreHierarchy.cpp
@@ -795,11 +795,13 @@ void FCacheStoreHierarchy::FCacheRecordBatchParams::FilterResponseByRequest(
 	FCacheGetResponse& Response,
 	const FCacheGetRequest& Request)
 {
-	if ((!Request.Policy.IsUniform() && Algo::AnyOf(Response.Record.GetValues(), &FValue::HasData)) ||
-		(EnumHasAnyFlags(Request.Policy.GetRecordPolicy(), ECachePolicy::SkipMeta) && Response.Record.GetMeta()))
+	const ECachePolicy RecordPolicy = Request.Policy.GetRecordPolicy();
+	const bool bMightSkipData = EnumHasAnyFlags(RecordPolicy, ECachePolicy::SkipData) || !Request.Policy.IsUniform();
+	if ((bMightSkipData && Algo::AnyOf(Response.Record.GetValues(), &FValue::HasData)) ||
+		(EnumHasAnyFlags(RecordPolicy, ECachePolicy::SkipMeta) && Response.Record.GetMeta()))
 	{
 		FCacheRecordBuilder Builder(Response.Record.GetKey());
-		if (!EnumHasAnyFlags(Request.Policy.GetRecordPolicy(), ECachePolicy::SkipMeta))
+		if (!EnumHasAnyFlags(RecordPolicy, ECachePolicy::SkipMeta))
 		{
 			Builder.SetMeta(CopyTemp(Response.Record.GetMeta()));
 		}
@@ -1069,7 +1071,7 @@ auto FCacheStoreHierarchy::FLegacyCacheBatchParams::Get() -> GetFunctionType
 template <>
 bool FCacheStoreHierarchy::FLegacyCacheBatchParams::HasResponseData(const FLegacyCacheGetResponse& Response)
 {
-	return !Response.Value.IsNull();
+	return Response.Value.HasData();
 }
 
 template <>
@@ -1088,7 +1090,7 @@ FLegacyCachePutRequest FCacheStoreHierarchy::FLegacyCacheBatchParams::MakePutReq
 	const FLegacyCacheGetResponse& Response,
 	const FLegacyCacheGetRequest& Request)
 {
-	return {Response.Name, Response.Key, FCompositeBuffer(Response.Value), Request.Policy};
+	return {Response.Name, Response.Key, Response.Value, Request.Policy};
 }
 
 template <>
diff --git a/Engine/Source/Developer/DerivedDataCache/Private/DerivedDataCacheStoreThrottle.cpp b/Engine/Source/Developer/DerivedDataCache/Private/DerivedDataCacheStoreThrottle.cpp
--- a/Engine/Source/Developer/DerivedDataCache/Private/DerivedDataCacheStoreThrottle.cpp
+++ b/Engine/Source/Developer/DerivedDataCache/Private/DerivedDataCacheStoreThrottle.cpp
@@ -149,7 +149,7 @@ public:
 		ValueSizes.Reserve(Requests.Num());
 		Algo::Transform(Requests, ValueSizes, [](const FLegacyCachePutRequest& Request) -> FValueSize
 		{
-			return {Request.Key.GetKey(), Request.Value.GetSize()};
+			return {Request.Key.GetKey(), Request.Value.GetRawSize()};
 		});
 
 		InnerCache->LegacyPut(Requests, Owner,
@@ -172,7 +172,7 @@ public:
 		InnerCache->LegacyGet(Requests, Owner,
 			[this, State = EnterThrottlingScope(), OnComplete = MoveTemp(OnComplete)](FLegacyCacheGetResponse&& Response)
 			{
-				CloseThrottlingScope(State, FThrottlingState(this, Response.Value.GetSize()));
+				CloseThrottlingScope(State, FThrottlingState(this, Response.Value.GetRawSize()));
 				if (OnComplete)
 				{
 					OnComplete(MoveTemp(Response));
diff --git a/Engine/Source/Developer/DerivedDataCache/Private/DerivedDataCacheStoreVerify.cpp b/Engine/Source/Developer/DerivedDataCache/Private/DerivedDataCacheStoreVerify.cpp
--- a/Engine/Source/Developer/DerivedDataCache/Private/DerivedDataCacheStoreVerify.cpp
+++ b/Engine/Source/Developer/DerivedDataCache/Private/DerivedDataCacheStoreVerify.cpp
@@ -127,7 +127,14 @@ private:
 
 	static bool CompareRecords(const FCacheRecord& PutRecord, const FCacheRecord& GetRecord, const FSharedString& Name);
 
-	static void LogChangedValue(const FSharedString& Name, const FCacheKey& Key, const FValueId& Id, const FValue& NewValue, const FValue& OldValue);
+	static void LogChangedValue(
+		const FSharedString& Name,
+		const FCacheKey& Key,
+		const FValueId& Id,
+		const FIoHash& NewRawHash,
+		const FIoHash& OldRawHash,
+		const FCompositeBuffer& NewRawData,
+		const FCompositeBuffer& OldRawData);
 
 private:
 	ILegacyCacheStore* InnerCache;
@@ -387,7 +394,9 @@ void FCacheStoreVerify::GetDataComplete(IRequestOwner& Owner, FVerifyPutValueSta
 		}
 		else
 		{
-			LogChangedValue(Request.Name, Request.Key, FValueId::Null, Request.Value, Response.Value);
+			LogChangedValue(Request.Name, Request.Key, FValueId::Null,
+				Request.Value.GetRawHash(), Response.Value.GetRawHash(),
+				Request.Value.GetData().DecompressToComposite(), Response.Value.GetData().DecompressToComposite());
 			if (bPutOnError)
 			{
 				// Ask to overwrite existing values to potentially eliminate the mismatch.
@@ -528,8 +537,8 @@ void FCacheStoreVerify::GetDataComplete(IRequestOwner& Owner, FVerifyLegacyPutSt
 
 	if (Response.Status == EStatus::Ok)
 	{
-		const FValue NewValue(FCompressedBuffer::Compress(Request.Value, ECompressedBufferCompressor::NotSet, ECompressedBufferCompressionLevel::None));
-		const FValue OldValue(FCompressedBuffer::Compress(Response.Value, ECompressedBufferCompressor::NotSet, ECompressedBufferCompressionLevel::None));
+		const FLegacyCacheValue& NewValue = Request.Value;
+		const FLegacyCacheValue& OldValue = Response.Value;
 		if (NewValue.GetRawHash() == OldValue.GetRawHash())
 		{
 			UE_LOG(LogDerivedDataCache, Log,
@@ -539,7 +548,9 @@ void FCacheStoreVerify::GetDataComplete(IRequestOwner& Owner, FVerifyLegacyPutSt
 		}
 		else
 		{
-			LogChangedValue(Request.Name, Request.Key.GetKey(), FValueId::Null, NewValue, OldValue);
+			LogChangedValue(Request.Name, Request.Key.GetKey(), FValueId::Null,
+				NewValue.GetRawHash(), OldValue.GetRawHash(),
+				NewValue.GetRawData(), OldValue.GetRawData());
 			if (bPutOnError)
 			{
 				// Ask to overwrite existing values to potentially eliminate the mismatch.
@@ -639,7 +650,9 @@ bool FCacheStoreVerify::CompareRecords(const FCacheRecord& PutRecord, const FCac
 		{
 			if (PutIt->GetRawHash() != GetIt->GetRawHash())
 			{
-				LogChangedValue(Name, Key, PutIt->GetId(), *PutIt, *GetIt);
+				LogChangedValue(Name, Key, PutIt->GetId(),
+					PutIt->GetRawHash(), GetIt->GetRawHash(),
+					PutIt->GetData().DecompressToComposite(), GetIt->GetData().DecompressToComposite());
 				bEqual = false;
 			}
 			++PutIt;
@@ -672,7 +685,14 @@ bool FCacheStoreVerify::CompareRecords(const FCacheRecord& PutRecord, const FCac
 	return bEqual;
 }
 
-void FCacheStoreVerify::LogChangedValue(const FSharedString& Name, const FCacheKey& Key, const FValueId& Id, const FValue& NewValue, const FValue& OldValue)
+void FCacheStoreVerify::LogChangedValue(
+	const FSharedString& Name,
+	const FCacheKey& Key,
+	const FValueId& Id,
+	const FIoHash& NewRawHash,
+	const FIoHash& OldRawHash,
+	const FCompositeBuffer& NewRawData,
+	const FCompositeBuffer& OldRawData)
 {
 	TStringBuilder<32> IdString;
 	if (Id.IsValid())
@@ -680,9 +700,9 @@ void FCacheStoreVerify::LogChangedValue(const FSharedString& Name, const FCacheK
 		IdString << TEXT(' ') << Id;
 	}
 
-	const auto LogDataToFile = [&Name, &Key, &Id, &IdString](const FValue& Value, FStringView Extension)
+	const auto LogDataToFile = [&Name, &Key, &Id, &IdString](const FIoHash& RawHash, const FCompositeBuffer& RawData, FStringView Extension)
 	{
-		if (FSharedBuffer RawData = Value.GetData().Decompress())
+		if (!RawData.IsNull())
 		{
 			TStringBuilder<256> Path;
 			FPathViews::Append(Path, FPaths::ProjectSavedDir(), TEXT("VerifyDDC"), TEXT(""));
@@ -694,23 +714,25 @@ void FCacheStoreVerify::LogChangedValue(const FSharedString& Name, const FCacheK
 			Path << Extension;
 			if (TUniquePtr<FArchive> Ar{IFileManager::Get().CreateFileWriter(*Path, FILEWRITE_Silent)})
 			{
-				Ar->Serialize(const_cast<void*>(RawData.GetData()), int64(RawData.GetSize()));
+				for (const FSharedBuffer& Segment : RawData.GetSegments())
+				{
+					Ar->Serialize(const_cast<void*>(Segment.GetData()), int64(Segment.GetSize()));
+				}
 			}
 		}
 		else
 		{
 			UE_LOG(LogDerivedDataCache, Log,
 				TEXT("Verify: Value%s does not have data with hash %s to save to disk for %s from '%s'."),
-				*IdString, *WriteToString<48>(Value.GetRawHash()), *WriteToString<96>(Key), *Name);
+				*IdString, *WriteToString<48>(RawHash), *WriteToString<96>(Key), *Name);
 		}
 	};
 
 	UE_LOG(LogDerivedDataCache, Error,
 		TEXT("Verify: Value%s has hash %s in the newly generated data and hash %s in the cache for %s from '%s'."),
-		*IdString, *WriteToString<48>(NewValue.GetRawHash()), *WriteToString<48>(OldValue.GetRawHash()),
-		*WriteToString<96>(Key), *Name);
-	LogDataToFile(NewValue, TEXTVIEW(".verify"));
-	LogDataToFile(OldValue, TEXTVIEW(".fromcache"));
+		*IdString, *WriteToString<48>(NewRawHash), *WriteToString<48>(OldRawHash), *WriteToString<96>(Key), *Name);
+	LogDataToFile(NewRawHash, NewRawData, TEXTVIEW(".verify"));
+	LogDataToFile(OldRawHash, OldRawData, TEXTVIEW(".fromcache"));
 }
 
 ILegacyCacheStore* CreateCacheStoreVerify(ILegacyCacheStore* InnerCache, bool bPutOnError)
diff --git a/Engine/Source/Developer/DerivedDataCache/Private/DerivedDataLegacyCacheStore.cpp b/Engine/Source/Developer/DerivedDataCache/Private/DerivedDataLegacyCacheStore.cpp
--- a/Engine/Source/Developer/DerivedDataCache/Private/DerivedDataLegacyCacheStore.cpp
+++ b/Engine/Source/Developer/DerivedDataCache/Private/DerivedDataLegacyCacheStore.cpp
@@ -5,6 +5,7 @@
 #include "Containers/StringConv.h"
 #include "DerivedDataBackendInterface.h"
 #include "Misc/Crc.h"
+#include "Misc/ScopeRWLock.h"
 #include "Misc/SecureHash.h"
 #include "Misc/StringBuilder.h"
 #include "String/BytesToHex.h"
@@ -13,6 +14,74 @@
 namespace UE::DerivedData
 {
 
+void ILegacyCacheStore::LegacyPut(
+	const TConstArrayView<FLegacyCachePutRequest> Requests,
+	IRequestOwner& Owner,
+	FOnLegacyCachePutComplete&& OnComplete)
+{
+	struct FKeyWithUserData
+	{
+		FLegacyCacheKey Key;
+		uint64 UserData;
+	};
+
+	TArray<FKeyWithUserData, TInlineAllocator<1>> LegacyRequests;
+	TArray<FCachePutValueRequest, TInlineAllocator<1>> ValueRequests;
+	LegacyRequests.Reserve(Requests.Num());
+	ValueRequests.Reserve(Requests.Num());
+
+	uint64 RequestIndex = 0;
+	for (const FLegacyCachePutRequest& Request : Requests)
+	{
+		LegacyRequests.Add({Request.Key, Request.UserData});
+		ValueRequests.Add({Request.Name, Request.Key.GetKey(), Request.Value.GetValue(), Request.Policy, RequestIndex++});
+	}
+
+	PutValue(ValueRequests, Owner, [LegacyRequests = MoveTemp(LegacyRequests), OnComplete = MoveTemp(OnComplete)](FCachePutValueResponse&& Response)
+	{
+		const FKeyWithUserData& LegacyRequest = LegacyRequests[int32(Response.UserData)];
+		OnComplete({Response.Name, LegacyRequest.Key, LegacyRequest.UserData, Response.Status});
+	});
+}
+
+void ILegacyCacheStore::LegacyGet(
+	const TConstArrayView<FLegacyCacheGetRequest> Requests,
+	IRequestOwner& Owner,
+	FOnLegacyCacheGetComplete&& OnComplete)
+{
+	struct FKeyWithUserData
+	{
+		FLegacyCacheKey Key;
+		uint64 UserData;
+	};
+
+	TArray<FKeyWithUserData, TInlineAllocator<1>> LegacyRequests;
+	TArray<FCacheGetValueRequest, TInlineAllocator<1>> ValueRequests;
+	LegacyRequests.Reserve(Requests.Num());
+	ValueRequests.Reserve(Requests.Num());
+
+	uint64 RequestIndex = 0;
+	for (const FLegacyCacheGetRequest& Request : Requests)
+	{
+		LegacyRequests.Add({Request.Key, Request.UserData});
+		ValueRequests.Add({Request.Name, Request.Key.GetKey(), Request.Policy, RequestIndex++});
+	}
+
+	GetValue(ValueRequests, Owner, [LegacyRequests = MoveTemp(LegacyRequests), OnComplete = MoveTemp(OnComplete)](FCacheGetValueResponse&& Response)
+	{
+		const FKeyWithUserData& LegacyRequest = LegacyRequests[int32(Response.UserData)];
+		OnComplete({Response.Name, LegacyRequest.Key, FLegacyCacheValue(Response.Value), LegacyRequest.UserData, Response.Status});
+	});
+}
+
+void ILegacyCacheStore::LegacyDelete(
+	const TConstArrayView<FLegacyCacheDeleteRequest> Requests,
+	IRequestOwner& Owner,
+	FOnLegacyCacheDeleteComplete&& OnComplete)
+{
+	CompleteWithStatus(Requests, OnComplete, EStatus::Error);
+}
+
 FLegacyCacheKey::FLegacyCacheKey(const FStringView InFullKey, const int32 MaxKeyLength)
 	: FullKey(InFullKey)
 {
@@ -83,6 +152,62 @@ void FLegacyCacheKey::WriteValueTrailer(FCompositeBuffer& Value) const
 	}
 }
 
+FLegacyCacheValue::FLegacyCacheValue(const FValue& Value)
+{
+	if (Value.HasData() || !Value.GetRawHash().IsZero())
+	{
+		Shared = new Private::FLegacyCacheValueShared(Value);
+	}
+}
+
+FLegacyCacheValue::FLegacyCacheValue(const FCompositeBuffer& RawData)
+{
+	if (RawData)
+	{
+		Shared = new Private::FLegacyCacheValueShared(RawData);
+	}
+}
+
+Private::FLegacyCacheValueShared::FLegacyCacheValueShared(const FValue& InValue)
+	: Value(InValue)
+{
+}
+
+Private::FLegacyCacheValueShared::FLegacyCacheValueShared(const FCompositeBuffer& InRawData)
+	: RawData(InRawData.MakeOwned())
+{
+}
+
+const FValue& Private::FLegacyCacheValueShared::GetValue()
+{
+	FWriteScopeLock WriteLock(Lock);
+	if (!Value.HasData() && !RawData.IsNull())
+	{
+		Value = FValue::Compress(RawData);
+	}
+	return Value;
+}
+
+const FCompositeBuffer& Private::FLegacyCacheValueShared::GetRawData()
+{
+	FWriteScopeLock WriteLock(Lock);
+	if (RawData.IsNull() && Value.HasData())
+	{
+		RawData = Value.GetData().DecompressToComposite();
+	}
+	return RawData;
+}
+
+FIoHash Private::FLegacyCacheValueShared::GetRawHash() const
+{
+	return Value.HasData() ? Value.GetRawHash() : FIoHash::HashBuffer(RawData);
+}
+
+uint64 Private::FLegacyCacheValueShared::GetRawSize() const
+{
+	return Value.HasData() ? Value.GetRawSize() : RawData.GetSize();
+}
+
 FLegacyCachePutResponse FLegacyCachePutRequest::MakeResponse(const EStatus Status) const
 {
 	return {Name, Key, UserData, Status};
diff --git a/Engine/Source/Developer/DerivedDataCache/Private/DerivedDataLegacyCacheStore.h b/Engine/Source/Developer/DerivedDataCache/Private/DerivedDataLegacyCacheStore.h
--- a/Engine/Source/Developer/DerivedDataCache/Private/DerivedDataLegacyCacheStore.h
+++ b/Engine/Source/Developer/DerivedDataCache/Private/DerivedDataLegacyCacheStore.h
@@ -3,7 +3,10 @@
 #pragma once
 
 #include "DerivedDataCacheStore.h"
+#include "HAL/CriticalSection.h"
 #include "Logging/LogMacros.h"
+#include "Templates/RefCounting.h"
+#include <atomic>
 
 class FDerivedDataCacheStatsNode;
 
@@ -15,6 +18,7 @@ namespace UE::DerivedData { struct FLegacyCacheGetRequest; }
 namespace UE::DerivedData { struct FLegacyCacheGetResponse; }
 namespace UE::DerivedData { struct FLegacyCachePutRequest; }
 namespace UE::DerivedData { struct FLegacyCachePutResponse; }
+namespace UE::DerivedData::Private { class FLegacyCacheValueShared; }
 
 DECLARE_LOG_CATEGORY_EXTERN(LogDerivedDataCache, Log, All);
 
@@ -31,17 +35,17 @@ public:
 	virtual void LegacyPut(
 		TConstArrayView<FLegacyCachePutRequest> Requests,
 		IRequestOwner& Owner,
-		FOnLegacyCachePutComplete&& OnComplete) = 0;
+		FOnLegacyCachePutComplete&& OnComplete);
 
 	virtual void LegacyGet(
 		TConstArrayView<FLegacyCacheGetRequest> Requests,
 		IRequestOwner& Owner,
-		FOnLegacyCacheGetComplete&& OnComplete) = 0;
+		FOnLegacyCacheGetComplete&& OnComplete);
 
 	virtual void LegacyDelete(
 		TConstArrayView<FLegacyCacheDeleteRequest> Requests,
 		IRequestOwner& Owner,
-		FOnLegacyCacheDeleteComplete&& OnComplete) = 0;
+		FOnLegacyCacheDeleteComplete&& OnComplete);
 
 	virtual void LegacyStats(FDerivedDataCacheStatsNode& OutNode) = 0;
 
@@ -68,11 +72,65 @@ private:
 	FSharedString ShortKey;
 };
 
+class Private::FLegacyCacheValueShared final
+{
+public:
+	explicit FLegacyCacheValueShared(const FValue& Value);
+	explicit FLegacyCacheValueShared(const FCompositeBuffer& RawData);
+
+	inline bool HasData() const { return Value.HasData() || RawData; }
+	const FValue& GetValue();
+	const FCompositeBuffer& GetRawData();
+	FIoHash GetRawHash() const;
+	uint64 GetRawSize() const;
+
+	inline void AddRef()
+	{
+		ReferenceCount.fetch_add(1, std::memory_order_relaxed);
+	}
+
+	inline void Release()
+	{
+		if (ReferenceCount.fetch_sub(1, std::memory_order_acq_rel) == 1)
+		{
+			delete this;
+		}
+	}
+
+private:
+	FValue Value;
+	FCompositeBuffer RawData;
+	std::atomic<uint32> ReferenceCount{0};
+	FRWLock Lock;
+};
+
+class FLegacyCacheValue
+{
+public:
+	FLegacyCacheValue() = default;
+	explicit FLegacyCacheValue(const FValue& Value);
+	explicit FLegacyCacheValue(const FCompositeBuffer& RawData);
+
+	inline void Reset() { Shared.SafeRelease(); }
+
+	[[nodiscard]] inline explicit operator bool() const { return !IsNull(); }
+	[[nodiscard]] inline bool IsNull() const { return !Shared; }
+
+	[[nodiscard]] bool HasData() const { return Shared ? Shared->HasData() : false; }
+	[[nodiscard]] const FValue& GetValue() const { return Shared ? Shared->GetValue() : FValue::Null; }
+	[[nodiscard]] const FCompositeBuffer& GetRawData() const { return Shared ? Shared->GetRawData() : FCompositeBuffer::Null; }
+	[[nodiscard]] inline FIoHash GetRawHash() const { return Shared ? Shared->GetRawHash() : FIoHash(); }
+	[[nodiscard]] inline uint64 GetRawSize() const { return Shared ? Shared->GetRawSize() : 0; }
+
+private:
+	TRefCountPtr<Private::FLegacyCacheValueShared> Shared;
+};
+
 struct FLegacyCachePutRequest
 {
 	FSharedString Name;
 	FLegacyCacheKey Key;
-	FCompositeBuffer Value;
+	FLegacyCacheValue Value;
 	ECachePolicy Policy = ECachePolicy::Default;
 	uint64 UserData = 0;
 
@@ -101,7 +159,7 @@ struct FLegacyCacheGetResponse
 {
 	FSharedString Name;
 	FLegacyCacheKey Key;
-	FSharedBuffer Value;
+	FLegacyCacheValue Value;
 	uint64 UserData = 0;
 	EStatus Status = EStatus::Error;
 };
