Value API support for Memory DDC backend.

diff --git a/Engine/Source/Developer/DerivedDataCache/Private/MemoryDerivedDataBackend.cpp b/Engine/Source/Developer/DerivedDataCache/Private/MemoryDerivedDataBackend.cpp
--- a/Engine/Source/Developer/DerivedDataCache/Private/MemoryDerivedDataBackend.cpp
+++ b/Engine/Source/Developer/DerivedDataCache/Private/MemoryDerivedDataBackend.cpp
@@ -106,6 +106,16 @@ public:
 		IRequestOwner& Owner,
 		FOnCacheGetComplete&& OnComplete) override;
 
+	virtual void PutValue(
+		TConstArrayView<FCachePutValueRequest> Requests,
+		IRequestOwner& Owner,
+		FOnCachePutValueComplete&& OnComplete) override;
+
+	virtual void GetValue(
+		TConstArrayView<FCacheGetValueRequest> Requests,
+		IRequestOwner& Owner,
+		FOnCacheGetValueComplete&& OnComplete) override;
+
 	virtual void GetChunks(
 		TConstArrayView<FCacheGetChunkRequest> Requests,
 		IRequestOwner& Owner,
@@ -152,6 +162,8 @@ private:
 	TMap<FString, FCacheValue*> CacheItems;
 	/** Set of records in this cache. */
 	TSet<FCacheRecord, FCacheRecordKeyFuncs> CacheRecords;
+	/** Set of values in this cache. */
+	TMap<FCacheKey, FValue> CacheValues;
 	/** Maximum size the cached items can grow up to ( in bytes ) */
 	int64 MaxCacheSize;
 	/** When set to true, this cache is disabled...ignore all requests. */
@@ -722,6 +734,18 @@ void FMemoryDerivedDataBackend::Get(
 	IRequestOwner& Owner,
 	FOnCacheGetComplete&& OnComplete)
 {
+	if (bDisabled)
+	{
+		for (const FCacheGetRequest& Request : Requests)
+		{
+			if (OnComplete)
+			{
+				OnComplete({ Request.Name, FCacheRecordBuilder(Request.Key).Build(), Request.UserData, EStatus::Error });
+			}
+		}
+		return;
+	}
+
 	for (const FCacheGetRequest& Request : Requests)
 	{
 		const FCacheKey& Key = Request.Key;
@@ -735,7 +759,7 @@ void FMemoryDerivedDataBackend::Get(
 			UE_LOG(LogDerivedDataCache, Verbose, TEXT("%s: Simulated miss for get of %s from '%s'"),
 				*GetName(), *WriteToString<96>(Key), *Request.Name);
 		}
-		else if (FWriteScopeLock ScopeLock(SynchronizationObject); const FCacheRecord* CacheRecord = CacheRecords.Find(Key))
+		else if (FReadScopeLock ScopeLock(SynchronizationObject); const FCacheRecord* CacheRecord = CacheRecords.Find(Key))
 		{
 			Status = EStatus::Ok;
 			Record = *CacheRecord;
@@ -773,16 +797,145 @@ void FMemoryDerivedDataBackend::Get(
 	}
 }
 
+void FMemoryDerivedDataBackend::PutValue(
+	const TConstArrayView<FCachePutValueRequest> Requests,
+	IRequestOwner& Owner,
+	FOnCachePutValueComplete&& OnComplete)
+{
+	for (const FCachePutValueRequest& Request : Requests)
+	{
+		const FValue& Value = Request.Value;
+		const FCacheKey& Key = Request.Key;
+		EStatus Status = EStatus::Error;
+		ON_SCOPE_EXIT
+		{
+			if (OnComplete)
+			{
+				OnComplete({Request.Name, Key, Request.UserData, Status});
+			}
+		};
+
+		if (!EnumHasAnyFlags(Request.Policy, ECachePolicy::StoreLocal))
+		{
+			continue;
+		}
+
+		if (ShouldSimulateMiss(Key))
+		{
+			UE_LOG(LogDerivedDataCache, Verbose, TEXT("%s: Simulated miss for put of %s from '%s'"),
+				*GetName(), *WriteToString<96>(Key), *Request.Name);
+			continue;
+		}
+
+		if (!Request.Value.HasData())
+		{
+			continue;
+		}
+
+		COOK_STAT(auto Timer = UsageStats.TimePut());
+		const int64 ValueSize = Value.GetData().GetCompressedSize();
+		const bool bReplaceExisting = !EnumHasAnyFlags(Request.Policy, ECachePolicy::QueryLocal);
+
+		FWriteScopeLock ScopeLock(SynchronizationObject);
+		FValue* const ExistingValue = CacheValues.Find(Key);
+		Status = ExistingValue && !bReplaceExisting ? EStatus::Ok : EStatus::Error;
+		if (bDisabled || Status == EStatus::Ok)
+		{
+			continue;
+		}
+
+		const int64 ExistingValueSize = ExistingValue ? ExistingValue->GetData().GetCompressedSize() : 0;
+		const int64 RequiredSize = ValueSize - ExistingValueSize;
+
+		if (MaxCacheSize > 0 && (CurrentCacheSize + RequiredSize) > MaxCacheSize)
+		{
+			UE_LOG(LogDerivedDataCache, Display, TEXT("Failed to cache data. Maximum cache size reached. CurrentSize %" INT64_FMT " KiB / MaxSize: %" INT64_FMT " KiB"), CurrentCacheSize / 1024, MaxCacheSize / 1024);
+			bMaxSizeExceeded = true;
+			continue;
+		}
+
+		CurrentCacheSize += RequiredSize;
+		if (ExistingValue)
+		{
+			*ExistingValue = Value;
+		}
+		else
+		{
+			CacheValues.Add(Key, Value);
+		}
+		COOK_STAT(Timer.AddHit(ValueSize));
+		Status = EStatus::Ok;
+	}
+}
+
+void FMemoryDerivedDataBackend::GetValue(
+	const TConstArrayView<FCacheGetValueRequest> Requests,
+	IRequestOwner& Owner,
+	FOnCacheGetValueComplete&& OnComplete)
+{
+	if (bDisabled)
+	{
+		CompleteWithStatus(Requests, OnComplete, EStatus::Error);
+		return;
+	}
+
+	for (const FCacheGetValueRequest& Request : Requests)
+	{
+		const FCacheKey& Key = Request.Key;
+		const ECachePolicy Policy = Request.Policy;
+		const bool bExistsOnly = EnumHasAllFlags(Policy, ECachePolicy::SkipData);
+		COOK_STAT(auto Timer = bExistsOnly ? UsageStats.TimeProbablyExists() : UsageStats.TimeGet());
+		bool bProcessHit = false;
+		FValue Value;
+		EStatus Status = EStatus::Error;
+		if (ShouldSimulateMiss(Key))
+		{
+			UE_LOG(LogDerivedDataCache, Verbose, TEXT("%s: Simulated miss for get of %s from '%s'"),
+				*GetName(), *WriteToString<96>(Key), *Request.Name);
+		}
+		else if (FReadScopeLock ScopeLock(SynchronizationObject); const FValue* CacheValue = CacheValues.Find(Key))
+		{
+			Status = EStatus::Ok;
+			Value = *CacheValue;
+			bProcessHit = true;
+		}
+
+		if (bProcessHit)
+		{
+			if (!Value.HasData() && !EnumHasAnyFlags(Policy, ECachePolicy::SkipData))
+			{
+				Status = EStatus::Error;
+			}
+
+			COOK_STAT(Timer.AddHit(Value.GetData().GetCompressedSize()));
+			if (OnComplete)
+			{
+				OnComplete({ Request.Name, Request.Key, EnumHasAnyFlags(Policy, ECachePolicy::SkipData) ? Value.RemoveData() : Value, Request.UserData, Status });
+			}
+		}
+		else
+		{
+			if (OnComplete)
+			{
+				OnComplete({ Request.Name, Request.Key, {}, Request.UserData, EStatus::Error });
+			}
+		}
+	}
+}
+
 void FMemoryDerivedDataBackend::GetChunks(
 	const TConstArrayView<FCacheGetChunkRequest> Requests,
 	IRequestOwner& Owner,
 	FOnCacheGetChunkComplete&& OnComplete)
 {
-	FValueWithId Value;
+	bool bHasValue = false;
+	FValue Value;
+	FValueId ValueId;
 	FCacheKey ValueKey;
 	FCompressedBufferReader Reader;
 	for (const FCacheGetChunkRequest& Request : Requests)
 	{
+		bool bProcessHit = false;
 		const bool bExistsOnly = EnumHasAnyFlags(Request.Policy, ECachePolicy::SkipData);
 		COOK_STAT(auto Timer = bExistsOnly ? UsageStats.TimeProbablyExists() : UsageStats.TimeGet());
 		if (ShouldSimulateMiss(Request.Key))
@@ -790,19 +943,46 @@ void FMemoryDerivedDataBackend::GetChunks(
 			UE_LOG(LogDerivedDataCache, Verbose, TEXT("%s: Simulated miss for get of %s from '%s'"),
 				*GetName(), *WriteToString<96>(Request.Key, '/', Request.Id), *Request.Name);
 		}
-		else if (ValueKey == Request.Key && Value.GetId() == Request.Id)
+		else if (bHasValue && (ValueKey == Request.Key) && (ValueId == Request.Id) && (bExistsOnly || Reader.HasSource()))
 		{
 			// Value matches the request.
+			bProcessHit = true;
 		}
-		else if (FWriteScopeLock ScopeLock(SynchronizationObject); const FCacheRecord* Record = CacheRecords.Find(Request.Key))
+		else
 		{
-			Reader.ResetSource();
-			Value.Reset();
-			Value = Record->GetValue(Request.Id);
-			ValueKey = Request.Key;
-			Reader.SetSource(Value.GetData());
+			FReadScopeLock ScopeLock(SynchronizationObject);
+			if (Request.Id.IsValid())
+			{
+				if (const FCacheRecord* Record = CacheRecords.Find(Request.Key))
+				{
+					const FValueWithId& ValueWithId = Record->GetValue(Request.Id);
+					bHasValue = ValueWithId.IsValid();
+					Reader.ResetSource();
+					Value.Reset();
+					Value = ValueWithId;
+					ValueId = Request.Id;
+					ValueKey = Request.Key;
+					Reader.SetSource(Value.GetData());
+					bProcessHit = true;
+				}
+			}
+			else
+			{
+				if (const FValue* ExistingValue = CacheValues.Find(Request.Key))
+				{
+					bHasValue = true;
+					Reader.ResetSource();
+					Value.Reset();
+					Value = *ExistingValue;
+					ValueId.Reset();
+					ValueKey = Request.Key;
+					Reader.SetSource(Value.GetData());
+					bProcessHit = true;
+				}
+			}
 		}
-		if (Value && Request.RawOffset <= Value.GetRawSize())
+
+		if (bProcessHit && Request.RawOffset <= Value.GetRawSize())
 		{
 			const uint64 RawSize = FMath::Min(Value.GetRawSize() - Request.RawOffset, Request.RawSize);
 			COOK_STAT(Timer.AddHit(RawSize));
