Value API support for PAK DDC backend.

diff --git a/Engine/Source/Developer/DerivedDataCache/Private/PakFileDerivedDataBackend.cpp b/Engine/Source/Developer/DerivedDataCache/Private/PakFileDerivedDataBackend.cpp
--- a/Engine/Source/Developer/DerivedDataCache/Private/PakFileDerivedDataBackend.cpp
+++ b/Engine/Source/Developer/DerivedDataCache/Private/PakFileDerivedDataBackend.cpp
@@ -24,6 +24,7 @@
 #include "Serialization/CompactBinary.h"
 #include "Serialization/CompactBinaryPackage.h"
 #include "Serialization/CompactBinaryValidation.h"
+#include "Serialization/CompactBinaryWriter.h"
 #include "Serialization/MemoryReader.h"
 #include "Serialization/MemoryWriter.h"
 #include "Templates/Greater.h"
@@ -125,14 +126,23 @@ public:
 		IRequestOwner& Owner,
 		FOnCacheGetComplete&& OnComplete) override;
 
+	virtual void PutValue(
+		TConstArrayView<FCachePutValueRequest> Requests,
+		IRequestOwner& Owner,
+		FOnCachePutValueComplete&& OnComplete) override;
+
+	virtual void GetValue(
+		TConstArrayView<FCacheGetValueRequest> Requests,
+		IRequestOwner& Owner,
+		FOnCacheGetValueComplete&& OnComplete) override;
+
 	virtual void GetChunks(
 		TConstArrayView<FCacheGetChunkRequest> Requests,
 		IRequestOwner& Owner,
 		FOnCacheGetChunkComplete&& OnComplete) override;
 
 private:
-	bool PutCacheRecord(FStringView Name, const FCacheRecord& Record, const FCacheRecordPolicy& Policy);
-	bool PutCacheContent(FStringView Name, const FCompressedBuffer& Content);
+	bool PutCacheRecord(FStringView Name, const FCacheRecord& Record, const FCacheRecordPolicy& Policy, uint64& OutWriteSize);
 
 	FOptionalCacheRecord GetCacheRecordOnly(
 		FStringView Name,
@@ -143,12 +153,22 @@ private:
 		const FCacheKey& Key,
 		const FCacheRecordPolicy& Policy,
 		EStatus& OutStatus);
-	void GetCacheContent(
+
+	[[nodiscard]] bool PutCacheValue(FStringView Name, const FCacheKey& Key, const FValue& Value, ECachePolicy Policy, uint64& OutWriteSize);
+
+	[[nodiscard]] bool GetCacheValueOnly(FStringView Name, const FCacheKey& Key, ECachePolicy Policy, FValue& OutValue);
+	[[nodiscard]] bool GetCacheValue(FStringView Name, const FCacheKey& Key, ECachePolicy Policy, FValue& OutValue);
+
+	[[nodiscard]] bool PutCacheContent(FStringView Name, const FCompressedBuffer& Content, uint64& OutWriteSize);
+
+	[[nodiscard]] bool GetCacheContentExists(const FCacheKey& Key, const FIoHash& RawHash);
+	[[nodiscard]] bool GetCacheContent(
 		FStringView Name,
 		const FCacheKey& Key,
+		const FValueId& Id,
+		const FValue& Value,
 		ECachePolicy Policy,
-		FValueWithId& InOutValue,
-		EStatus& InOutStatus);
+		FValue& OutValue);
 
 	bool SaveFile(FStringView Path, FStringView DebugName, TFunctionRef<void (FArchive&)> WriteFunction);
 	FSharedBuffer LoadFile(FStringView Path, FStringView DebugName);
@@ -658,11 +678,12 @@ void FPakFileDerivedDataBackend::Put(
 		const FCacheRecord& Record = Request.Record;
 		TRACE_CPUPROFILER_EVENT_SCOPE(PakFileDDC_Put);
 		COOK_STAT(auto Timer = UsageStats.TimePut());
-		if (PutCacheRecord(Request.Name, Record, Request.Policy))
+		uint64 WriteSize = 0;
+		if (PutCacheRecord(Request.Name, Record, Request.Policy, WriteSize))
 		{
 			UE_LOG(LogDerivedDataCache, Verbose, TEXT("%s: Cache put complete for %s from '%s'"),
 				*CachePath, *WriteToString<96>(Record.GetKey()), *Request.Name);
-			COOK_STAT(Timer.AddHit(Private::GetCacheRecordCompressedSize(Record)));
+			COOK_STAT(if (WriteSize) { Timer.AddHit(WriteSize); });
 			if (OnComplete)
 			{
 				OnComplete({Request.Name, Record.GetKey(), Request.UserData, EStatus::Ok});
@@ -686,7 +707,7 @@ void FPakFileDerivedDataBackend::Get(
 {
 	for (const FCacheGetRequest& Request : Requests)
 	{
-		TRACE_CPUPROFILER_EVENT_SCOPE(PakFileDDC_Get);
+		TRACE_CPUPROFILER_EVENT_SCOPE(PakFileDDC_GetValue);
 		COOK_STAT(auto Timer = UsageStats.TimeGet());
 		EStatus Status = EStatus::Ok;
 		if (FOptionalCacheRecord Record = GetCacheRecord(Request.Name, Request.Key, Request.Policy, Status))
@@ -709,6 +730,66 @@ void FPakFileDerivedDataBackend::Get(
 	}
 }
 
+void FPakFileDerivedDataBackend::PutValue(
+	const TConstArrayView<FCachePutValueRequest> Requests,
+	IRequestOwner& Owner,
+	FOnCachePutValueComplete&& OnComplete)
+{
+	for (const FCachePutValueRequest& Request : Requests)
+	{
+		TRACE_CPUPROFILER_EVENT_SCOPE(PakFileDDC_Put);
+		COOK_STAT(auto Timer = UsageStats.TimePut());
+		uint64 WriteSize = 0;
+		if (PutCacheValue(Request.Name, Request.Key, Request.Value, Request.Policy, WriteSize))
+		{
+			UE_LOG(LogDerivedDataCache, Verbose, TEXT("%s: Cache put complete for %s from '%s'"),
+				*CachePath, *WriteToString<96>(Request.Key), *Request.Name);
+			COOK_STAT(if (WriteSize) { Timer.AddHit(WriteSize); });
+			if (OnComplete)
+			{
+				OnComplete({Request.Name, Request.Key, Request.UserData, EStatus::Ok});
+			}
+		}
+		else
+		{
+			if (OnComplete)
+			{
+				OnComplete({Request.Name, Request.Key, Request.UserData, EStatus::Error});
+			}
+		}
+	}
+}
+
+void FPakFileDerivedDataBackend::GetValue(
+	const TConstArrayView<FCacheGetValueRequest> Requests,
+	IRequestOwner& Owner,
+	FOnCacheGetValueComplete&& OnComplete)
+{
+	for (const FCacheGetValueRequest& Request : Requests)
+	{
+		TRACE_CPUPROFILER_EVENT_SCOPE(PakFileDDC_Get);
+		COOK_STAT(auto Timer = UsageStats.TimeGet());
+		FValue Value;
+		if (GetCacheValue(Request.Name, Request.Key, Request.Policy, Value))
+		{
+			UE_LOG(LogDerivedDataCache, Verbose, TEXT("%s: Cache hit for %s from '%s'"),
+				*CachePath, *WriteToString<96>(Request.Key), *Request.Name);
+			COOK_STAT(Timer.AddHit(Value.GetData().GetCompressedSize()));
+			if (OnComplete)
+			{
+				OnComplete({Request.Name, Request.Key, Value, Request.UserData, EStatus::Ok});
+			}
+		}
+		else
+		{
+			if (OnComplete)
+			{
+				OnComplete({Request.Name, Request.Key, {}, Request.UserData, EStatus::Error});
+			}
+		}
+	}
+}
+
 void FPakFileDerivedDataBackend::GetChunks(
 	const TConstArrayView<FCacheGetChunkRequest> Requests,
 	IRequestOwner& Owner,
@@ -718,7 +799,7 @@ void FPakFileDerivedDataBackend::GetChunks(
 	SortedRequests.StableSort(TChunkLess());
 
 	bool bHasValue = false;
-	FValueWithId Value;
+	FValue Value;
 	FValueId ValueId;
 	FCacheKey ValueKey;
 	FCompressedBufferReader ValueReader;
@@ -746,14 +827,22 @@ void FPakFileDerivedDataBackend::GetChunks(
 				if (Record)
 				{
 					const FValueWithId& ValueWithId = Record.Get().GetValue(Request.Id);
-					bHasValue = ValueWithId.IsValid();
-					Value = ValueWithId;
-					ValueId = Request.Id;
-					ValueKey = Request.Key;
-					EStatus ValueStatus;
-					GetCacheContent(Request.Name, Request.Key, Request.Policy, Value, ValueStatus);
+					if (ValueWithId.IsValid() && GetCacheContent(Request.Name, Request.Key, ValueId, ValueWithId, Request.Policy, Value))
+					{
+						bHasValue = true;
+						Value = ValueWithId;
+						ValueId = Request.Id;
+						ValueKey = Request.Key;
+					}
 				}
 			}
+			else
+			{
+				ValueKey = Request.Key;
+				FValue TempValue;
+				bHasValue = GetCacheValueOnly(Request.Name, Request.Key, Request.Policy, TempValue) &&
+					GetCacheContent(Request.Name, Request.Key, Request.Id, TempValue, Request.Policy, Value);
+			}
 			if (bHasValue)
 			{
 				const uint64 RawOffset = FMath::Min(Value.GetRawSize(), Request.RawOffset);
@@ -781,7 +870,8 @@ void FPakFileDerivedDataBackend::GetChunks(
 bool FPakFileDerivedDataBackend::PutCacheRecord(
 	const FStringView Name,
 	const FCacheRecord& Record,
-	const FCacheRecordPolicy& Policy)
+	const FCacheRecordPolicy& Policy,
+	uint64& OutWriteSize)
 {
 	if (!IsWritable())
 	{
@@ -866,11 +956,16 @@ bool FPakFileDerivedDataBackend::PutCacheRecord(
 	// Save the external content to storage.
 	for (FCompressedBuffer& Content : ExternalContent)
 	{
-		PutCacheContent(Name, Content);
+		uint64 WriteSize = 0;
+		if (!PutCacheContent(Name, Content, OutWriteSize))
+		{
+			return false;
+		}
+		OutWriteSize += WriteSize;
 	}
 
 	// Save the record package to storage.
-	if (!bRecordExists && !SaveFile(Path, Name, [&Package](FArchive& Ar) { Package.Save(Ar); }))
+	if (!bRecordExists && !SaveFile(Path, Name, [&Package, &OutWriteSize](FArchive& Ar) { Package.Save(Ar); OutWriteSize += uint64(Ar.TotalSize()); }))
 	{
 		return false;
 	}
@@ -963,35 +1058,207 @@ FOptionalCacheRecord FPakFileDerivedDataBackend::GetCacheRecord(
 	OutStatus = EStatus::Ok;
 
 	FCacheRecordBuilder RecordBuilder(Key);
-	const ECachePolicy RecordOnlyPolicy = Policy.GetRecordPolicy() & ~FCacheValuePolicy::PolicyMask;
 
-	if (!EnumHasAnyFlags(Policy.GetRecordPolicy(), ECachePolicy::SkipMeta))
+	const ECachePolicy RecordPolicy = Policy.GetRecordPolicy();
+	if (!EnumHasAnyFlags(RecordPolicy, ECachePolicy::SkipMeta))
 	{
 		RecordBuilder.SetMeta(FCbObject(Record.Get().GetMeta()));
 	}
 
 	for (FValueWithId Value : Record.Get().GetValues())
 	{
-		const ECachePolicy ValuePolicy = Policy.GetValuePolicy(Value.GetId());
-		GetCacheContent(Name, Key, ValuePolicy | RecordOnlyPolicy, Value, OutStatus);
-		if (Value.IsNull())
+		const FValueId& Id = Value.GetId();
+		const ECachePolicy ValuePolicy = Policy.GetValuePolicy(Id);
+		FValue Content;
+		if (GetCacheContent(Name, Key, Id, Value, ValuePolicy, Content))
 		{
+			RecordBuilder.AddValue(Id, MoveTemp(Content));
+		}
+		else if (EnumHasAnyFlags(RecordPolicy, ECachePolicy::PartialRecord))
+		{
+			OutStatus = EStatus::Error;
+			RecordBuilder.AddValue(Value);
+		}
+		else
+		{
+			OutStatus = EStatus::Error;
 			return FOptionalCacheRecord();
 		}
-		RecordBuilder.AddValue(MoveTemp(Value));
 	}
 
 	return RecordBuilder.Build();
 }
 
-bool FPakFileDerivedDataBackend::PutCacheContent(const FStringView Name, const FCompressedBuffer& Content)
+bool FPakFileDerivedDataBackend::PutCacheValue(
+	const FStringView Name,
+	const FCacheKey& Key,
+	const FValue& Value,
+	const ECachePolicy Policy,
+	uint64& OutWriteSize)
+{
+	if (!IsWritable())
+	{
+		return false;
+	}
+
+	// Skip the request if storing to the cache is disabled.
+	if (!EnumHasAnyFlags(Policy, ECachePolicy::StoreLocal))
+	{
+		UE_LOG(LogDerivedDataCache, VeryVerbose, TEXT("%s: Skipped put of %s from '%.*s' due to cache policy"),
+			*CachePath, *WriteToString<96>(Key), Name.Len(), Name.GetData());
+		return false;
+	}
+
+	//if (ShouldSimulateMiss(Key))
+	//{
+	//	UE_LOG(LogDerivedDataCache, Verbose, TEXT("%s: Simulated miss for put of %s from '%.*s'"),
+	//		*CachePath, *WriteToString<96>(Key), Name.Len(), Name.GetData());
+	//	return false;
+	//}
+
+	// Check if there is an existing value package.
+	bool bValueExists = false;
+	FCbPackage ExistingPackage;
+	TStringBuilder<256> Path;
+	FPathViews::Append(Path, TEXT("Buckets"), Key);
+	const bool bReplaceExisting = !EnumHasAnyFlags(Policy, ECachePolicy::QueryLocal);
+	if (!bReplaceExisting)
+	{
+		if (FSharedBuffer Buffer = LoadFile(Path, Name))
+		{
+			FCbFieldIterator It = FCbFieldIterator::MakeRange(MoveTemp(Buffer));
+			bValueExists = ExistingPackage.TryLoad(It);
+		}
+	}
+
+	// Save the value to a package and save the data to external content.
+	if (!bValueExists)
+	{
+		FCbWriter Writer;
+		Writer.BeginObject();
+		Writer.AddBinaryAttachment("RawHash", Value.GetRawHash());
+		Writer.AddInteger("RawSize", Value.GetRawSize());
+		Writer.EndObject();
+
+		FCbPackage Package(Writer.Save().AsObject());
+		if (!Value.HasData())
+		{
+			// Verify that the content exists in storage.
+			if (!GetCacheContentExists(Key, Value.GetRawHash()))
+			{
+				UE_LOG(LogDerivedDataCache, Verbose, TEXT("%s: Failed due to missing data for put of %s from '%.*s'"),
+					*CachePath, *WriteToString<96>(Key), Name.Len(), Name.GetData());
+				return false;
+			}
+		}
+		else
+		{
+			// Save the external content to storage.
+			uint64 WriteSize = 0;
+			if (!PutCacheContent(Name, Value.GetData(), WriteSize))
+			{
+				return false;
+			}
+			OutWriteSize += WriteSize;
+		}
+
+		// Save the record package to storage.
+		if (!bValueExists && !SaveFile(Path, Name, [&Package, &OutWriteSize](FArchive& Ar) { Package.Save(Ar); OutWriteSize += uint64(Ar.TotalSize()); }))
+		{
+			return false;
+		}
+	}
+
+	return true;
+}
+
+bool FPakFileDerivedDataBackend::GetCacheValueOnly(
+	const FStringView Name,
+	const FCacheKey& Key,
+	const ECachePolicy Policy,
+	FValue& OutValue)
+{
+	if (bClosed)
+	{
+		UE_LOG(LogDerivedDataCache, VeryVerbose,
+			TEXT("%s: Skipped get of %s from '%.*s' because this cache store is not available"),
+			*CachePath, *WriteToString<96>(Key), Name.Len(), Name.GetData());
+		return false;
+	}
+
+	// Skip the request if querying the cache is disabled.
+	if (!EnumHasAnyFlags(Policy, ECachePolicy::QueryLocal))
+	{
+		UE_LOG(LogDerivedDataCache, VeryVerbose, TEXT("%s: Skipped get of %s from '%.*s' due to cache policy"),
+			*CachePath, *WriteToString<96>(Key), Name.Len(), Name.GetData());
+		return false;
+	}
+
+	//if (ShouldSimulateMiss(Key))
+	//{
+	//	UE_LOG(LogDerivedDataCache, Verbose, TEXT("%s: Simulated miss for get of %s from '%.*s'"),
+	//		*CachePath, *WriteToString<96>(Key), Name.Len(), Name.GetData());
+	//	return false;
+	//}
+	
+	// Request the record from storage.
+	TStringBuilder<256> Path;
+	FPathViews::Append(Path, TEXT("Buckets"), Key);
+	FSharedBuffer Buffer = LoadFile(Path, Name);
+	if (Buffer.IsNull())
+	{
+		UE_LOG(LogDerivedDataCache, Verbose, TEXT("%s: Cache miss with missing value for %s from '%.*s'"),
+			*CachePath, *WriteToString<96>(Key), Name.Len(), Name.GetData());
+		return false;
+	}
+
+	if (ValidateCompactBinary(Buffer, ECbValidateMode::Default | ECbValidateMode::Package) != ECbValidateError::None)
+	{
+		UE_LOG(LogDerivedDataCache, Display, TEXT("%s: Cache miss with invalid package for %s from '%.*s'"),
+			*CachePath, *WriteToString<96>(Key), Name.Len(), Name.GetData());
+		return false;
+	}
+
+	FCbPackage Package;
+	if (FCbFieldIterator It = FCbFieldIterator::MakeRange(Buffer); !Package.TryLoad(It))
+	{
+		UE_LOG(LogDerivedDataCache, Display, TEXT("%s: Cache miss with package load failure for %s from '%.*s'"),
+			*CachePath, *WriteToString<96>(Key), Name.Len(), Name.GetData());
+		return false;
+	}
+
+	const FCbObjectView Object = Package.GetObject();
+	const FIoHash RawHash = Object["RawHash"].AsHash();
+	const uint64 RawSize = Object["RawSize"].AsUInt64(MAX_uint64);
+	if (RawHash.IsZero() || RawSize == MAX_uint64)
+	{
+		UE_LOG(LogDerivedDataCache, Display, TEXT("%s: Cache miss with invalid value for %s from '%.*s'"),
+			*CachePath, *WriteToString<96>(Key), Name.Len(), Name.GetData());
+		return false;
+	}
+
+	OutValue = FValue(RawHash, RawSize);
+
+	return true;
+}
+
+bool FPakFileDerivedDataBackend::GetCacheValue(
+	const FStringView Name,
+	const FCacheKey& Key,
+	const ECachePolicy Policy,
+	FValue& OutValue)
+{
+	return GetCacheValueOnly(Name, Key, Policy, OutValue) && GetCacheContent(Name, Key, {}, OutValue, Policy, OutValue);
+}
+
+bool FPakFileDerivedDataBackend::PutCacheContent(const FStringView Name, const FCompressedBuffer& Content, uint64& OutWriteSize)
 {
 	const FIoHash& RawHash = Content.GetRawHash();
 	TStringBuilder<256> Path;
 	FPathViews::Append(Path, TEXT("Content"), RawHash);
 	if (!FileExists(Path))
 	{
-		if (!SaveFile(Path, Name, [&Content](FArchive& Ar) { Content.Save(Ar); }))
+		if (!SaveFile(Path, Name, [&Content, &OutWriteSize](FArchive& Ar) { Content.Save(Ar); OutWriteSize += uint64(Ar.TotalSize()); }))
 		{
 			return false;
 		}
@@ -999,34 +1266,43 @@ bool FPakFileDerivedDataBackend::PutCacheContent(const FStringView Name, const F
 	return true;
 }
 
-void FPakFileDerivedDataBackend::GetCacheContent(
+bool FPakFileDerivedDataBackend::GetCacheContentExists(const FCacheKey& Key, const FIoHash& RawHash)
+{
+	TStringBuilder<256> Path;
+	FPathViews::Append(Path, TEXT("Buckets"), Key);
+	return FileExists(Path);
+}
+
+bool FPakFileDerivedDataBackend::GetCacheContent(
 	const FStringView Name,
 	const FCacheKey& Key,
+	const FValueId& Id,
+	const FValue& Value,
 	const ECachePolicy Policy,
-	FValueWithId& InOutValue,
-	EStatus& InOutStatus)
+	FValue& OutValue)
 {
-	if (!EnumHasAnyFlags(Policy, ECachePolicy::Query) ||
-		(EnumHasAnyFlags(Policy, ECachePolicy::SkipData) && InOutValue.HasData()))
+	if (!EnumHasAnyFlags(Policy, ECachePolicy::Query))
 	{
-		InOutValue = InOutValue.RemoveData();
-		return;
+		OutValue = Value.RemoveData();
+		return true;
 	}
 
-	if (InOutValue.HasData())
+	if (Value.HasData())
 	{
-		return;
+		OutValue = EnumHasAnyFlags(Policy, ECachePolicy::SkipData) ? Value.RemoveData() : Value;
+		return true;
 	}
 
-	const FIoHash& RawHash = InOutValue.GetRawHash();
+	const FIoHash& RawHash = Value.GetRawHash();
 
 	TStringBuilder<256> Path;
 	FPathViews::Append(Path, TEXT("Content"), RawHash);
-	if (EnumHasAllFlags(Policy, ECachePolicy::SkipData))
+	if (EnumHasAnyFlags(Policy, ECachePolicy::SkipData))
 	{
 		if (FileExists(Path))
 		{
-			return;
+			OutValue = Value;
+			return true;
 		}
 	}
 	else
@@ -1036,31 +1312,22 @@ void FPakFileDerivedDataBackend::GetCacheContent(
 			if (FCompressedBuffer CompressedBuffer = FCompressedBuffer::FromCompressed(MoveTemp(CompressedData));
 				CompressedBuffer && CompressedBuffer.GetRawHash() == RawHash)
 			{
-				InOutValue = FValueWithId(InOutValue.GetId(), MoveTemp(CompressedBuffer));
-				return;
+				OutValue = FValue(MoveTemp(CompressedBuffer));
+				return true;
 			}
 			UE_LOG(LogDerivedDataCache, Display,
 				TEXT("%s: Cache miss with corrupted value %s with hash %s for %s from '%.*s'"),
-				*CachePath, *WriteToString<16>(InOutValue.GetId()), *WriteToString<48>(RawHash),
+				*CachePath, *WriteToString<16>(Id), *WriteToString<48>(RawHash),
 				*WriteToString<96>(Key), Name.Len(), Name.GetData());
-			InOutStatus = EStatus::Error;
-			if (!EnumHasAnyFlags(Policy, ECachePolicy::PartialRecord))
-			{
-				InOutValue = FValueWithId::Null;
-			}
-			return;
+			return false;
 		}
 	}
 
 	UE_LOG(LogDerivedDataCache, Verbose,
 		TEXT("%s: Cache miss with missing value %s with hash %s for %s from '%.*s'"),
-		*CachePath, *WriteToString<16>(InOutValue.GetId()), *WriteToString<48>(RawHash), *WriteToString<96>(Key),
+		*CachePath, *WriteToString<16>(Id), *WriteToString<48>(RawHash), *WriteToString<96>(Key),
 		Name.Len(), Name.GetData());
-	InOutStatus = EStatus::Error;
-	if (!EnumHasAnyFlags(Policy, ECachePolicy::PartialRecord))
-	{
-		InOutValue = FValueWithId::Null;
-	}
+	return false;
 }
 
 class FCrcBuilder
