DDC: Added support for 64-bit sizes in the legacy cache

Includes an optimization to avoid an extra copy of data loaded from the cache.

diff --git a/Engine/Source/Developer/DerivedDataCache/Private/DerivedDataCache.cpp b/Engine/Source/Developer/DerivedDataCache/Private/DerivedDataCache.cpp
--- a/Engine/Source/Developer/DerivedDataCache/Private/DerivedDataCache.cpp
+++ b/Engine/Source/Developer/DerivedDataCache/Private/DerivedDataCache.cpp
@@ -284,7 +284,7 @@ class FDerivedDataCache final
 
 			TRACE_CPUPROFILER_EVENT_SCOPE(DDC_DoWork);
 
-			const int32 NumBeforeDDC = Data.Num();
+			const int64 NumBeforeDDC = Data.Num();
 			bool bGetResult;
 			{
 				TRACE_CPUPROFILER_EVENT_SCOPE(DDC_Get);
@@ -301,14 +301,14 @@ class FDerivedDataCache final
 						[this, &bGetResult](FLegacyCacheGetResponse&& Response)
 						{
 							const uint64 RawSize = Response.Value.GetRawSize();
-							bGetResult = Response.Status == EStatus::Ok && RawSize < MAX_int32;
+							bGetResult = Response.Status == EStatus::Ok && RawSize < MAX_int64;
 							if (bGetResult)
 							{
 								const FCompositeBuffer& RawData = Response.Value.GetRawData();
-								Data.Reset(int32(RawSize));
+								Data.Reset(int64(RawSize));
 								for (const FSharedBuffer& Segment : RawData.GetSegments())
 								{
-									Data.Append(static_cast<const uint8*>(Segment.GetData()), int32(Segment.GetSize()));
+									Data.Append(static_cast<const uint8*>(Segment.GetData()), int64(Segment.GetSize()));
 								}
 							}
 						});
@@ -323,8 +323,8 @@ class FDerivedDataCache final
 				{
 					TArray<uint8> CmpData;
 					DataDeriver->Build(CmpData);
-					const int32 NumInDDC = Data.Num() - NumBeforeDDC;
-					const int32 NumGenerated = CmpData.Num();
+					const int64 NumInDDC = Data.Num() - NumBeforeDDC;
+					const int64 NumGenerated = CmpData.Num();
 					
 					bool bMatchesInSize = NumGenerated == NumInDDC;
 					bool bDifferentMemory = true;
@@ -366,7 +366,9 @@ class FDerivedDataCache final
 					STAT(double ThisTime = 0);
 					{
 						SCOPE_SECONDS_COUNTER(ThisTime);
-						bSuccess = DataDeriver->Build(Data);
+						TArray<uint8> Data32;
+						bSuccess = DataDeriver->Build(Data32);
+						Data = TArray64<uint8>(MoveTemp(Data32));
 						bDataWasBuilt = true;
 					}
 					INC_FLOAT_STAT_BY(STAT_DDC_SyncBuildTime, bSynchronousForStats ? (float)ThisTime : 0.0f);
@@ -428,7 +430,7 @@ class FDerivedDataCache final
 		/** Context from the caller */
 		FSharedString					DebugContext;
 		/** Data to return to caller, later **/
-		TArray<uint8>					Data;
+		TArray64<uint8>					Data;
 	};
 
 public:
@@ -466,7 +468,7 @@ public:
 		PendingTasks.Empty();
 	}
 
-	virtual bool GetSynchronous(FDerivedDataPluginInterface* DataDeriver, TArray<uint8>& OutData, bool* bDataWasBuilt = nullptr) override
+	virtual bool GetSynchronous(FDerivedDataPluginInterface* DataDeriver, TArray<uint8>& OutData, bool* bDataWasBuilt) override
 	{
 		DDC_SCOPE_CYCLE_COUNTER(DDC_GetSynchronous);
 		check(DataDeriver);
@@ -475,7 +477,7 @@ public:
 		FAsyncTask<FBuildAsyncWorker> PendingTask(DataDeriver, *CacheKey, DataDeriver->GetDebugContextString(), true);
 		AddToAsyncCompletionCounter(1);
 		PendingTask.StartSynchronousTask();
-		OutData = PendingTask.GetTask().Data;
+		OutData = TArray<uint8>(MoveTemp(PendingTask.GetTask().Data));
 		if (bDataWasBuilt)
 		{
 			*bDataWasBuilt = PendingTask.GetTask().bDataWasBuilt;
@@ -538,7 +540,8 @@ public:
 		INC_FLOAT_STAT_BY(STAT_DDC_ASyncWaitTime,(float)ThisTime);
 	}
 
-	virtual bool GetAsynchronousResults(uint32 Handle, TArray<uint8>& OutData, bool* bOutDataWasBuilt = nullptr) override
+	template <typename DataType>
+	bool GetAsynchronousResultsByHandle(uint32 Handle, DataType& OutData, bool* bOutDataWasBuilt)
 	{
 		DDC_SCOPE_CYCLE_COUNTER(DDC_GetAsynchronousResults);
 		FAsyncTask<FBuildAsyncWorker>* AsyncTask = NULL;
@@ -560,13 +563,24 @@ public:
 		}
 
 		UE_LOG(LogDerivedDataCache, Verbose, TEXT("GetAsynchronousResults, bDataWasBuilt: %d, Handle %d, SUCCESS"), (int32)bDataWasBuilt, Handle);
-		OutData = MoveTemp(AsyncTask->GetTask().Data);
+		OutData = DataType(MoveTemp(AsyncTask->GetTask().Data));
 		delete AsyncTask;
 		check(OutData.Num());
 		return true;
 	}
 
-	virtual bool GetSynchronous(const TCHAR* CacheKey, TArray<uint8>& OutData, FStringView DebugContext) override
+	virtual bool GetAsynchronousResults(uint32 Handle, TArray<uint8>& OutData, bool* bOutDataWasBuilt) override
+	{
+		return GetAsynchronousResultsByHandle(Handle, OutData, bOutDataWasBuilt);
+	}
+
+	virtual bool GetAsynchronousResults(uint32 Handle, TArray64<uint8>& OutData, bool* bOutDataWasBuilt) override
+	{
+		return GetAsynchronousResultsByHandle(Handle, OutData, bOutDataWasBuilt);
+	}
+
+	template <typename DataType>
+	bool GetSynchronousByKey(const TCHAR* CacheKey, DataType& OutData, FStringView DebugContext)
 	{
 		DDC_SCOPE_CYCLE_COUNTER(DDC_GetSynchronous_Data);
 		UE_LOG(LogDerivedDataCache, VeryVerbose, TEXT("GetSynchronous %s from '%.*s'"), CacheKey, DebugContext.Len(), DebugContext.GetData());
@@ -574,10 +588,20 @@ public:
 		FAsyncTask<FBuildAsyncWorker> PendingTask((FDerivedDataPluginInterface*)NULL, CacheKey, DebugContext, true);
 		AddToAsyncCompletionCounter(1);
 		PendingTask.StartSynchronousTask();
-		OutData = PendingTask.GetTask().Data;
+		OutData = DataType(MoveTemp(PendingTask.GetTask().Data));
 		return PendingTask.GetTask().bSuccess;
 	}
 
+	virtual bool GetSynchronous(const TCHAR* CacheKey, TArray<uint8>& OutData, FStringView DebugContext) override
+	{
+		return GetSynchronousByKey(CacheKey, OutData, DebugContext);
+	}
+
+	virtual bool GetSynchronous(const TCHAR* CacheKey, TArray64<uint8>& OutData, FStringView DebugContext) override
+	{
+		return GetSynchronousByKey(CacheKey, OutData, DebugContext);
+	}
+
 	virtual uint32 GetAsynchronous(const TCHAR* CacheKey, FStringView DebugContext) override
 	{
 		DDC_SCOPE_CYCLE_COUNTER(DDC_GetAsynchronous_Handle);
@@ -594,7 +618,7 @@ public:
 		return Handle;
 	}
 
-	virtual void Put(const TCHAR* CacheKey, TArrayView<const uint8> Data, FStringView DebugContext, bool bPutEvenIfExists = false) override
+	virtual void Put(const TCHAR* CacheKey, TArrayView64<const uint8> Data, FStringView DebugContext, bool bPutEvenIfExists = false) override
 	{
 		DDC_SCOPE_CYCLE_COUNTER(DDC_Put);
 		UE_LOG(LogDerivedDataCache, VeryVerbose, TEXT("Put %s from '%.*s'"), CacheKey, DebugContext.Len(), DebugContext.GetData());
diff --git a/Engine/Source/Developer/DerivedDataCache/Public/DerivedDataCacheInterface.h b/Engine/Source/Developer/DerivedDataCache/Public/DerivedDataCacheInterface.h
--- a/Engine/Source/Developer/DerivedDataCache/Public/DerivedDataCacheInterface.h
+++ b/Engine/Source/Developer/DerivedDataCache/Public/DerivedDataCacheInterface.h
@@ -61,6 +61,7 @@ public:
 	 * @return			true if the data was retrieved from the cache or the deriver built the data successfully. false can only occur if the plugin returns false.
 	**/
 	virtual bool GetAsynchronousResults(uint32 Handle, TArray<uint8>& OutData, bool* bDataWasBuilt = nullptr) = 0;
+	virtual bool GetAsynchronousResults(uint32 Handle, TArray64<uint8>& OutData, bool* bDataWasBuilt = nullptr) = 0;
 
 	//--------------------------
 	// Low Level Static Helpers
@@ -170,6 +171,7 @@ public:
 	 * @return	true if the data was retrieved from the cache.
 	**/
 	virtual bool GetSynchronous(const TCHAR* CacheKey, TArray<uint8>& OutData, FStringView DebugContext) = 0; 
+	virtual bool GetSynchronous(const TCHAR* CacheKey, TArray64<uint8>& OutData, FStringView DebugContext) = 0; 
 
 	/** 
 	 * Starts the async process of checking the cache and if the item is present, retrieving the cached results.
@@ -191,7 +193,7 @@ public:
 	 * @param	Data		Data to put in the cache under this key
 	 * @param	DataContext	A string used to describe the data being generated. Typically the path to the object that it is generated from is sufficient.
 	**/
-	virtual void Put(const TCHAR* CacheKey, TArrayView<const uint8> Data, FStringView DataContext, bool bPutEvenIfExists = false) = 0;
+	virtual void Put(const TCHAR* CacheKey, TArrayView64<const uint8> Data, FStringView DataContext, bool bPutEvenIfExists = false) = 0;
 
 	/**
 	 * Hint that the data associated with the key is transient and may be optionally purged from the cache.
diff --git a/Engine/Source/Runtime/Core/Public/Containers/ArrayView.h b/Engine/Source/Runtime/Core/Public/Containers/ArrayView.h
--- a/Engine/Source/Runtime/Core/Public/Containers/ArrayView.h
+++ b/Engine/Source/Runtime/Core/Public/Containers/ArrayView.h
@@ -707,8 +707,8 @@ struct TIsZeroConstructType<TArrayView<InElementType>>
 	enum { Value = true };
 };
 
-template <typename T>
-struct TIsContiguousContainer<TArrayView<T>>
+template <typename T, typename SizeType>
+struct TIsContiguousContainer<TArrayView<T, SizeType>>
 {
 	enum { Value = true };
 };
