Value API support for S3 DDC backend.

diff --git a/Engine/Source/Developer/DerivedDataCache/Private/S3DerivedDataBackend.cpp b/Engine/Source/Developer/DerivedDataCache/Private/S3DerivedDataBackend.cpp
--- a/Engine/Source/Developer/DerivedDataCache/Private/S3DerivedDataBackend.cpp
+++ b/Engine/Source/Developer/DerivedDataCache/Private/S3DerivedDataBackend.cpp
@@ -16,12 +16,15 @@
 #include "Async/ParallelFor.h"
 #include "DerivedDataBackendCorruptionWrapper.h"
 #include "DerivedDataCacheRecord.h"
+#include "DerivedDataCachePrivate.h"
 #include "DerivedDataCacheUsageStats.h"
+#include "DerivedDataChunk.h"
 #include "DesktopPlatformModule.h"
 #include "Dom/JsonObject.h"
 #include "HAL/PlatformFile.h"
 #include "HAL/PlatformFileManager.h"
 #include "HAL/Runnable.h"
+#include "HashingArchiveProxy.h"
 #include "Memory/SharedBuffer.h"
 #include "Misc/Base64.h"
 #include "Misc/ConfigCacheIni.h"
@@ -32,6 +35,9 @@
 #include "Misc/SecureHash.h"
 #include "ProfilingDebugging/CountersTrace.h"
 #include "ProfilingDebugging/CpuProfilerTrace.h"
+#include "Serialization/CompactBinary.h"
+#include "Serialization/CompactBinaryPackage.h"
+#include "Serialization/CompactBinaryValidation.h"
 #include "Serialization/JsonReader.h"
 #include "Serialization/JsonSerializer.h"
 #include "Serialization/MemoryReader.h"
@@ -62,7 +68,14 @@ TRACE_DECLARE_INT_COUNTER(S3DDC_Get, TEXT("S3DDC Get"));
 TRACE_DECLARE_INT_COUNTER(S3DDC_GetHit, TEXT("S3DDC Get Hit"));
 TRACE_DECLARE_INT_COUNTER(S3DDC_BytesRecieved, TEXT("S3DDC Bytes Recieved"));
 
-FString BuildPathForCacheKey(const TCHAR* CacheKey)
+///////////////////////////////////////////////////////////////////////////////////////////////////////////////////////
+
+static const TCHAR GBucketsDirectoryName[] = TEXT("Buckets");
+static const TCHAR GContentDirectoryName[] = TEXT("Content");
+
+///////////////////////////////////////////////////////////////////////////////////////////////////////////////////////
+
+FString BuildPathForLegacyCache(const TCHAR* CacheKey)
 {
 	FString Key = FString(CacheKey).ToUpper();
 	checkf(Algo::AllOf(Key, [](TCHAR C) { return FChar::IsAlnum(C) || FChar::IsUnderscore(C) || C == TEXT('$'); }),
@@ -73,6 +86,22 @@ FString BuildPathForCacheKey(const TCHAR* CacheKey)
 	return HashPath / Key + TEXT(".udd");
 }
 
+static void BuildPathForCachePackage(const FCacheKey& CacheKey, FStringBuilderBase& Path)
+{
+	const FIoHash::ByteArray& Bytes = CacheKey.Hash.GetBytes();
+	Path.Appendf(TEXT("%s/%hs/%02x/%02x/"), GBucketsDirectoryName, CacheKey.Bucket.ToCString(), Bytes[0], Bytes[1]);
+	UE::String::BytesToHexLower(MakeArrayView(Bytes).RightChop(2), Path);
+	Path << TEXT(".udd"_SV);
+}
+
+static void BuildPathForCacheContent(const FIoHash& RawHash, FStringBuilderBase& Path)
+{
+	const FIoHash::ByteArray& Bytes = RawHash.GetBytes();
+	Path.Appendf(TEXT("%s/%02x/%02x/"), GContentDirectoryName, Bytes[0], Bytes[1]);
+	UE::String::BytesToHexLower(MakeArrayView(Bytes).RightChop(2), Path);
+	Path << TEXT(".udd"_SV);
+}
+
 class FStringAnsi
 {
 public:
@@ -241,6 +270,11 @@ public:
 		IRequestOwner& Owner,
 		FOnCacheGetComplete&& OnComplete) override;
 
+	virtual void GetValue(
+		TConstArrayView<FCacheGetValueRequest> Requests,
+		IRequestOwner& Owner,
+		FOnCacheGetValueComplete&& OnComplete) override;
+
 	virtual void GetChunks(
 		TConstArrayView<FCacheGetChunkRequest> Requests,
 		IRequestOwner& Owner,
@@ -266,6 +300,45 @@ private:
 	FDerivedDataCacheUsageStats UsageStats;
 	bool bEnabled;
 
+	[[nodiscard]] FOptionalCacheRecord GetCacheRecordOnly(
+		FStringView Name,
+		const FCacheKey& Key,
+		const FCacheRecordPolicy& Policy);
+	[[nodiscard]] FOptionalCacheRecord GetCacheRecord(
+		FStringView Name,
+		const FCacheKey& Key,
+		const FCacheRecordPolicy& Policy,
+		EStatus& OutStatus);
+
+	[[nodiscard]] bool GetCacheValueOnly(FStringView Name, const FCacheKey& Key, ECachePolicy Policy, FValue& OutValue);
+	[[nodiscard]] bool GetCacheValue(FStringView Name, const FCacheKey& Key, ECachePolicy Policy, FValue& OutValue);
+
+	[[nodiscard]] bool GetCacheContentExists(const FCacheKey& Key, const FIoHash& RawHash) const;
+	[[nodiscard]] bool GetCacheContent(
+		FStringView Name,
+		const FCacheKey& Key,
+		const FValueId& Id,
+		const FValue& Value,
+		ECachePolicy Policy,
+		FValue& OutValue) const;
+	void GetCacheContent(
+		FStringView Name,
+		const FCacheKey& Key,
+		const FValueId& Id,
+		const FValue& Value,
+		ECachePolicy Policy,
+		FCompressedBufferReader& Reader,
+		TUniquePtr<FArchive>& OutArchive) const;
+
+	void BuildCachePackagePath(const FCacheKey& CacheKey, FStringBuilderBase& Path) const;
+	void BuildCacheContentPath(const FIoHash& RawHash, FStringBuilderBase& Path) const;
+
+	[[nodiscard]] bool LoadFileWithHash(FStringBuilderBase& Path, FStringView DebugName, TFunctionRef<void(FArchive&)> ReadFunction) const;
+	[[nodiscard]] bool LoadFile(FStringBuilderBase& Path, FStringView DebugName, TFunctionRef<void(FArchive&)> ReadFunction) const;
+	[[nodiscard]] TUniquePtr<FArchive> OpenFile(FStringBuilderBase& Path, FStringView DebugName) const;
+
+	[[nodiscard]] bool FileExists(FStringBuilderBase& Path) const;
+
 	bool DownloadManifest(const FRootManifest& RootManifest, FFeedbackContext* Context);
 	void RemoveUnusedBundles();
 	void ReadBundle(FBundle& Bundle);
@@ -1262,7 +1335,7 @@ bool FS3DerivedDataBackend::FindBundleEntry(const TCHAR* CacheKey, const FBundle
 {
 	FSHAHash Hash;
 
-	auto AnsiString = StringCast<ANSICHAR>(*BuildPathForCacheKey(CacheKey).ToUpper());
+	auto AnsiString = StringCast<ANSICHAR>(*BuildPathForLegacyCache(CacheKey).ToUpper());
 	FSHA1::HashBuffer(AnsiString.Get(), AnsiString.Length(), Hash.Hash);
 
 	for (const FBundle& Bundle : Bundles)
@@ -1314,6 +1387,438 @@ bool FS3DerivedDataBackend::ShouldSimulateMiss(const TCHAR* InKey)
 	return false;
 }
 
+FOptionalCacheRecord FS3DerivedDataBackend::GetCacheRecordOnly(
+	const FStringView Name,
+	const FCacheKey& Key,
+	const FCacheRecordPolicy& Policy)
+{
+	if (!IsUsable())
+	{
+		UE_LOG(LogDerivedDataCache, VeryVerbose,
+			TEXT("%s: Skipped get of %s from '%.*s' because this cache store is not available"),
+			*GetName(), *WriteToString<96>(Key), Name.Len(), Name.GetData());
+		return FOptionalCacheRecord();
+	}
+
+	// Skip the request if querying the cache is disabled.
+	if (!EnumHasAnyFlags(Policy.GetRecordPolicy(), ECachePolicy::QueryLocal))
+	{
+		UE_LOG(LogDerivedDataCache, VeryVerbose, TEXT("%s: Skipped get of %s from '%.*s' due to cache policy"),
+			*GetName(), *WriteToString<96>(Key), Name.Len(), Name.GetData());
+		return FOptionalCacheRecord();
+	}
+
+	//if (ShouldSimulateMiss(Key))
+	//{
+	//	UE_LOG(LogDerivedDataCache, Verbose, TEXT("%s: Simulated miss for get of %s from '%.*s'"),
+	//		*GetName(), *WriteToString<96>(Key), Name.Len(), Name.GetData());
+	//	return FOptionalCacheRecord();
+	//}
+
+	TStringBuilder<256> Path;
+	BuildCachePackagePath(Key, Path);
+
+
+	FOptionalCacheRecord Record;
+	{
+		FCbPackage Package;
+		if (!LoadFileWithHash(Path, Name, [&Package](FArchive& Ar) { Package.TryLoad(Ar); }))
+		{
+			UE_LOG(LogDerivedDataCache, Verbose, TEXT("%s: Cache miss with missing package for %s from '%.*s'"),
+				*GetName(), *WriteToString<96>(Key), Name.Len(), Name.GetData());
+			return Record;
+		}
+
+		if (ValidateCompactBinary(Package, ECbValidateMode::Default | ECbValidateMode::Package) != ECbValidateError::None)
+		{
+			UE_LOG(LogDerivedDataCache, Display, TEXT("%s: Cache miss with invalid package for %s from '%.*s'"),
+				*GetName(), *WriteToString<96>(Key), Name.Len(), Name.GetData());
+			return Record;
+		}
+
+		Record = FCacheRecord::Load(Package);
+		if (Record.IsNull())
+		{
+			UE_LOG(LogDerivedDataCache, Display, TEXT("%s: Cache miss with record load failure for %s from '%.*s'"),
+				*GetName(), *WriteToString<96>(Key), Name.Len(), Name.GetData());
+			return Record;
+		}
+	}
+
+	return Record;
+}
+
+FOptionalCacheRecord FS3DerivedDataBackend::GetCacheRecord(
+	const FStringView Name,
+	const FCacheKey& Key,
+	const FCacheRecordPolicy& Policy,
+	EStatus& OutStatus)
+{
+	FOptionalCacheRecord Record = GetCacheRecordOnly(Name, Key, Policy);
+	if (Record.IsNull())
+	{
+		OutStatus = EStatus::Error;
+		return Record;
+	}
+
+	OutStatus = EStatus::Ok;
+
+	FCacheRecordBuilder RecordBuilder(Key);
+
+	const ECachePolicy RecordPolicy = Policy.GetRecordPolicy();
+	if (!EnumHasAnyFlags(RecordPolicy, ECachePolicy::SkipMeta))
+	{
+		RecordBuilder.SetMeta(FCbObject(Record.Get().GetMeta()));
+	}
+
+	for (const FValueWithId& Value : Record.Get().GetValues())
+	{
+		const FValueId& Id = Value.GetId();
+		const ECachePolicy ValuePolicy = Policy.GetValuePolicy(Id);
+		FValue Content;
+		if (GetCacheContent(Name, Key, Id, Value, ValuePolicy, Content))
+		{
+			RecordBuilder.AddValue(Id, MoveTemp(Content));
+		}
+		else if (EnumHasAnyFlags(RecordPolicy, ECachePolicy::PartialRecord))
+		{
+			OutStatus = EStatus::Error;
+			RecordBuilder.AddValue(Value);
+		}
+		else
+		{
+			OutStatus = EStatus::Error;
+			return FOptionalCacheRecord();
+		}
+	}
+
+	return RecordBuilder.Build();
+}
+
+bool FS3DerivedDataBackend::GetCacheValueOnly(
+	const FStringView Name,
+	const FCacheKey& Key,
+	const ECachePolicy Policy,
+	FValue& OutValue)
+{
+	if (!IsUsable())
+	{
+		UE_LOG(LogDerivedDataCache, VeryVerbose,
+			TEXT("%s: Skipped get of %s from '%.*s' because this cache store is not available"),
+			*GetName(), *WriteToString<96>(Key), Name.Len(), Name.GetData());
+		return false;
+	}
+
+	// Skip the request if querying the cache is disabled.
+	if (!EnumHasAnyFlags(Policy, ECachePolicy::QueryLocal))
+	{
+		UE_LOG(LogDerivedDataCache, VeryVerbose, TEXT("%s: Skipped get of %s from '%.*s' due to cache policy"),
+			*GetName(), *WriteToString<96>(Key), Name.Len(), Name.GetData());
+		return false;
+	}
+
+	//if (ShouldSimulateMiss(Key))
+	//{
+	//	UE_LOG(LogDerivedDataCache, Verbose, TEXT("%s: Simulated miss for get of %s from '%.*s'"),
+	//		*CachePath, *WriteToString<96>(Key), Name.Len(), Name.GetData());
+	//	return false;
+	//}
+
+	TStringBuilder<256> Path;
+	BuildCachePackagePath(Key, Path);
+
+	FCbPackage Package;
+	if (!LoadFileWithHash(Path, Name, [&Package](FArchive& Ar) { Package.TryLoad(Ar); }))
+	{
+		UE_LOG(LogDerivedDataCache, Verbose, TEXT("%s: Cache miss with missing package for %s from '%.*s'"),
+			*GetName(), *WriteToString<96>(Key), Name.Len(), Name.GetData());
+		return false;
+	}
+
+	if (ValidateCompactBinary(Package, ECbValidateMode::Default | ECbValidateMode::Package) != ECbValidateError::None)
+	{
+		UE_LOG(LogDerivedDataCache, Display, TEXT("%s: Cache miss with invalid package for %s from '%.*s'"),
+			*GetName(), *WriteToString<96>(Key), Name.Len(), Name.GetData());
+		return false;
+	}
+
+	const FCbObjectView Object = Package.GetObject();
+	const FIoHash RawHash = Object["RawHash"].AsHash();
+	const uint64 RawSize = Object["RawSize"].AsUInt64(MAX_uint64);
+	if (RawHash.IsZero() || RawSize == MAX_uint64)
+	{
+		UE_LOG(LogDerivedDataCache, Display, TEXT("%s: Cache miss with invalid value for %s from '%.*s'"),
+			*GetName(), *WriteToString<96>(Key), Name.Len(), Name.GetData());
+		return false;
+	}
+
+	if (const FCbAttachment* const Attachment = Package.FindAttachment(RawHash))
+	{
+		const FCompressedBuffer& Data = Attachment->AsCompressedBinary();
+		if (Data.GetRawHash() != RawHash || Data.GetRawSize() != RawSize)
+		{
+			UE_LOG(LogDerivedDataCache, Display,
+				TEXT("%s: Cache miss with invalid value attachment for %s from '%.*s'"),
+				*GetName(), *WriteToString<96>(Key), Name.Len(), Name.GetData());
+			return false;
+		}
+		OutValue = FValue(Data);
+	}
+	else
+	{
+		OutValue = FValue(RawHash, RawSize);
+	}
+
+	return true;
+}
+
+bool FS3DerivedDataBackend::GetCacheValue(
+	const FStringView Name,
+	const FCacheKey& Key,
+	const ECachePolicy Policy,
+	FValue& OutValue)
+{
+	return GetCacheValueOnly(Name, Key, Policy, OutValue) && GetCacheContent(Name, Key, {}, OutValue, Policy, OutValue);
+}
+
+bool FS3DerivedDataBackend::GetCacheContentExists(const FCacheKey& Key, const FIoHash& RawHash) const
+{
+	TStringBuilder<256> Path;
+	BuildCacheContentPath(RawHash, Path);
+	return FileExists(Path);
+}
+
+bool FS3DerivedDataBackend::GetCacheContent(
+	const FStringView Name,
+	const FCacheKey& Key,
+	const FValueId& Id,
+	const FValue& Value,
+	const ECachePolicy Policy,
+	FValue& OutValue) const
+{
+	if (!EnumHasAnyFlags(Policy, ECachePolicy::Query))
+	{
+		OutValue = Value.RemoveData();
+		return true;
+	}
+
+	if (Value.HasData())
+	{
+		OutValue = EnumHasAnyFlags(Policy, ECachePolicy::SkipData) ? Value.RemoveData() : Value;
+		return true;
+	}
+
+	const FIoHash& RawHash = Value.GetRawHash();
+
+	TStringBuilder<256> Path;
+	BuildCacheContentPath(RawHash, Path);
+	if (EnumHasAnyFlags(Policy, ECachePolicy::SkipData))
+	{
+		if (FileExists(Path))
+		{
+			OutValue = Value;
+			return true;
+		}
+	}
+	else
+	{
+		FCompressedBuffer CompressedBuffer;
+		if (LoadFileWithHash(Path, Name, [&CompressedBuffer](FArchive& Ar) { CompressedBuffer = FCompressedBuffer::Load(Ar); }))
+		{
+			if (CompressedBuffer.GetRawHash() == RawHash)
+			{
+				OutValue = FValue(MoveTemp(CompressedBuffer));
+				return true;
+			}
+			UE_LOG(LogDerivedDataCache, Display,
+				TEXT("%s: Cache miss with corrupted value %s with hash %s for %s from '%.*s'"),
+				*GetName(), *WriteToString<16>(Id), *WriteToString<48>(RawHash),
+				*WriteToString<96>(Key), Name.Len(), Name.GetData());
+			return false;
+		}
+	}
+
+	UE_LOG(LogDerivedDataCache, Verbose,
+		TEXT("%s: Cache miss with missing value %s with hash %s for %s from '%.*s'"),
+		*GetName(), *WriteToString<16>(Id), *WriteToString<48>(RawHash), *WriteToString<96>(Key),
+		Name.Len(), Name.GetData());
+	return false;
+}
+
+void FS3DerivedDataBackend::GetCacheContent(
+	const FStringView Name,
+	const FCacheKey& Key,
+	const FValueId& Id,
+	const FValue& Value,
+	const ECachePolicy Policy,
+	FCompressedBufferReader& Reader,
+	TUniquePtr<FArchive>& OutArchive) const
+{
+	if (!EnumHasAnyFlags(Policy, ECachePolicy::Query))
+	{
+		return;
+	}
+
+	if (Value.HasData())
+	{
+		if (!EnumHasAnyFlags(Policy, ECachePolicy::SkipData))
+		{
+			Reader.SetSource(Value.GetData());
+		}
+		OutArchive.Reset();
+		return;
+	}
+
+	const FIoHash& RawHash = Value.GetRawHash();
+
+	TStringBuilder<256> Path;
+	BuildCacheContentPath(RawHash, Path);
+	if (EnumHasAllFlags(Policy, ECachePolicy::SkipData))
+	{
+		if (FileExists(Path))
+		{
+			return;
+		}
+	}
+	else
+	{
+		OutArchive = OpenFile(Path, Name);
+		if (OutArchive)
+		{
+			Reader.SetSource(*OutArchive);
+			if (Reader.GetRawHash() == RawHash)
+			{
+				return;
+			}
+			UE_LOG(LogDerivedDataCache, Display,
+				TEXT("%s: Cache miss with corrupted value %s with hash %s for %s from '%.*s'"),
+				*GetName(), *WriteToString<16>(Id), *WriteToString<48>(RawHash),
+				*WriteToString<96>(Key), Name.Len(), Name.GetData());
+			Reader.ResetSource();
+			OutArchive.Reset();
+			return;
+		}
+	}
+
+	UE_LOG(LogDerivedDataCache, Verbose,
+		TEXT("%s: Cache miss with missing value %s with hash %s for %s from '%.*s'"),
+		*GetName(), *WriteToString<16>(Id), *WriteToString<48>(RawHash), *WriteToString<96>(Key),
+		Name.Len(), Name.GetData());
+}
+
+void FS3DerivedDataBackend::BuildCachePackagePath(const FCacheKey& CacheKey, FStringBuilderBase& Path) const
+{
+	BuildPathForCachePackage(CacheKey, Path);
+}
+
+void FS3DerivedDataBackend::BuildCacheContentPath(const FIoHash& RawHash, FStringBuilderBase& Path) const
+{
+	BuildPathForCacheContent(RawHash, Path);
+}
+
+bool FS3DerivedDataBackend::LoadFileWithHash(
+	FStringBuilderBase& Path,
+	const FStringView DebugName,
+	const TFunctionRef<void (FArchive& Ar)> ReadFunction) const
+{
+	return LoadFile(Path, DebugName, [this, &Path, &DebugName, &ReadFunction](FArchive& Ar)
+	{
+		THashingArchiveProxy<FBlake3> HashAr(Ar);
+		ReadFunction(HashAr);
+		const FBlake3Hash Hash = HashAr.GetHash();
+		FBlake3Hash SavedHash;
+		Ar << SavedHash;
+
+		if (Hash != SavedHash && !Ar.IsError())
+		{
+			Ar.SetError();
+			UE_LOG(LogDerivedDataCache, Display,
+				TEXT("%s: File %s from '%.*s' is corrupted and has hash %s when %s is expected."),
+				*GetName(), *Path, DebugName.Len(), DebugName.GetData(),
+				*WriteToString<80>(Hash), *WriteToString<80>(SavedHash));
+		}
+	});
+}
+
+bool FS3DerivedDataBackend::LoadFile(
+	FStringBuilderBase& Path,
+	const FStringView DebugName,
+	const TFunctionRef<void (FArchive& Ar)> ReadFunction) const
+{
+	check(IsUsable());
+	const double StartTime = FPlatformTime::Seconds();
+
+	int64 ReadSize = 0;
+	bool bError = false;
+
+	if (TUniquePtr<FArchive> Ar = OpenFile(Path, DebugName))
+	{
+		int64 ReadStart = Ar->Tell();
+		ReadFunction(*Ar);
+		ReadSize = Ar->Tell() - ReadStart;
+		bError = !Ar->Close();
+
+		if (bError)
+		{
+			UE_LOG(LogDerivedDataCache, Display,
+				TEXT("%s: Failed to load file %s from '%.*s'."),
+				*GetName(), *Path, DebugName.Len(), DebugName.GetData());
+		}
+	}
+
+	const double ReadDuration = FPlatformTime::Seconds() - StartTime;
+	const double ReadSpeed = ReadDuration > 0.001 ? (ReadSize / ReadDuration) / (1024.0 * 1024.0) : 0.0;
+
+	UE_LOG(LogDerivedDataCache, VeryVerbose,
+		TEXT("%s: Loaded %s from '%.*s' (%" INT64_FMT " bytes, %.02f secs, %.2f MiB/s)"),
+		*GetName(), *Path, DebugName.Len(), DebugName.GetData(), ReadSize, ReadDuration, ReadSpeed);
+
+	return !bError && ReadSize > 0;
+}
+
+TUniquePtr<FArchive> FS3DerivedDataBackend::OpenFile(FStringBuilderBase& Path, const FStringView DebugName) const
+{
+	FSHAHash Hash;
+
+	auto AnsiString = StringCast<ANSICHAR>(*FString(Path.ToString()).ToUpper());
+	FSHA1::HashBuffer(AnsiString.Get(), AnsiString.Length(), Hash.Hash);
+
+	for (const FBundle& Bundle : Bundles)
+	{
+		const FBundleEntry* Entry = Bundle.Entries.Find(Hash);
+		if (Entry != nullptr)
+		{
+			TUniquePtr<FArchive> Reader(IFileManager::Get().CreateFileReader(*Bundle.LocalFile));
+			if (Reader.IsValid() && !Reader->IsError())
+			{
+				Reader->Seek(Entry->Offset);
+				return Reader;
+			}
+		}
+	}
+
+	return nullptr;
+}
+
+bool FS3DerivedDataBackend::FileExists(FStringBuilderBase& Path) const
+{
+	FSHAHash Hash;
+
+	auto AnsiString = StringCast<ANSICHAR>(*FString(Path.ToString()).ToUpper());
+	FSHA1::HashBuffer(AnsiString.Get(), AnsiString.Length(), Hash.Hash);
+
+	for (const FBundle& Bundle : Bundles)
+	{
+		const FBundleEntry* Entry = Bundle.Entries.Find(Hash);
+		if (Entry != nullptr)
+		{
+			return IFileManager::Get().FileExists(*Bundle.LocalFile);
+		}
+	}
+
+	return false;
+}
+
 void FS3DerivedDataBackend::Put(
 	const TConstArrayView<FCachePutRequest> Requests,
 	IRequestOwner& Owner,
@@ -1333,11 +1838,61 @@ void FS3DerivedDataBackend::Get(
 	IRequestOwner& Owner,
 	FOnCacheGetComplete&& OnComplete)
 {
-	if (OnComplete)
+	for (const FCacheGetRequest& Request : Requests)
+	{
+		TRACE_CPUPROFILER_EVENT_SCOPE(S3DDC_Get);
+		TRACE_COUNTER_INCREMENT(S3DDC_Get);
+		COOK_STAT(auto Timer = UsageStats.TimeGet());
+		EStatus Status = EStatus::Ok;
+		if (FOptionalCacheRecord Record = GetCacheRecord(Request.Name, Request.Key, Request.Policy, Status))
+		{
+			UE_LOG(LogDerivedDataCache, Verbose, TEXT("%s: Cache hit for %s from '%s'"),
+				*GetName(), *WriteToString<96>(Request.Key), *Request.Name);
+			TRACE_COUNTER_INCREMENT(S3DDC_GetHit);
+			COOK_STAT(Timer.AddHit(Private::GetCacheRecordCompressedSize(Record.Get())));
+			if (OnComplete)
+			{
+				OnComplete({Request.Name, MoveTemp(Record).Get(), Request.UserData, Status});
+			}
+		}
+		else
+		{
+			if (OnComplete)
+			{
+				OnComplete({Request.Name, FCacheRecordBuilder(Request.Key).Build(), Request.UserData, Status});
+			}
+		}
+	}
+}
+
+void FS3DerivedDataBackend::GetValue(
+	const TConstArrayView<FCacheGetValueRequest> Requests,
+	IRequestOwner& Owner,
+	FOnCacheGetValueComplete&& OnComplete)
+{
+	for (const FCacheGetValueRequest& Request : Requests)
 	{
-		for (const FCacheGetRequest& Request : Requests)
+		TRACE_CPUPROFILER_EVENT_SCOPE(S3DDC_Get);
+		TRACE_COUNTER_INCREMENT(S3DDC_Get);
+		COOK_STAT(auto Timer = UsageStats.TimeGet());
+		FValue Value;
+		if (GetCacheValue(Request.Name, Request.Key, Request.Policy, Value))
+		{
+			UE_LOG(LogDerivedDataCache, Verbose, TEXT("%s: Cache hit for %s from '%s'"),
+				*GetName(), *WriteToString<96>(Request.Key), *Request.Name);
+			TRACE_COUNTER_INCREMENT(S3DDC_GetHit);
+			COOK_STAT(Timer.AddHit(Value.GetData().GetCompressedSize()));
+			if (OnComplete)
+			{
+				OnComplete({Request.Name, Request.Key, Value, Request.UserData, EStatus::Ok});
+			}
+		}
+		else
 		{
-			OnComplete({Request.Name, FCacheRecordBuilder(Request.Key).Build(), Request.UserData, EStatus::Error});
+			if (OnComplete)
+			{
+				OnComplete({Request.Name, Request.Key, {}, Request.UserData, EStatus::Error});
+			}
 		}
 	}
 }
@@ -1347,9 +1902,81 @@ void FS3DerivedDataBackend::GetChunks(
 	IRequestOwner& Owner,
 	FOnCacheGetChunkComplete&& OnComplete)
 {
-	if (OnComplete)
+	TArray<FCacheGetChunkRequest, TInlineAllocator<16>> SortedRequests(Requests);
+	SortedRequests.StableSort(TChunkLess());
+
+	bool bHasValue = false;
+	FValue Value;
+	FValueId ValueId;
+	FCacheKey ValueKey;
+	TUniquePtr<FArchive> ValueAr;
+	FCompressedBufferReader ValueReader;
+	EStatus ValueStatus = EStatus::Error;
+	FOptionalCacheRecord Record;
+	for (const FCacheGetChunkRequest& Request : SortedRequests)
 	{
-		for (const FCacheGetChunkRequest& Request : Requests)
+		const bool bExistsOnly = EnumHasAnyFlags(Request.Policy, ECachePolicy::SkipData);
+		TRACE_CPUPROFILER_EVENT_SCOPE(S3DDC_Get);
+		TRACE_COUNTER_INCREMENT(S3DDC_Get);
+		COOK_STAT(auto Timer = bExistsOnly ? UsageStats.TimeProbablyExists() : UsageStats.TimeGet());
+		if (!(bHasValue && ValueKey == Request.Key && ValueId == Request.Id) || ValueReader.HasSource() < !bExistsOnly)
+		{
+			ValueStatus = EStatus::Error;
+			ValueReader.ResetSource();
+			ValueAr.Reset();
+			ValueKey = {};
+			ValueId.Reset();
+			Value.Reset();
+			bHasValue = false;
+			if (Request.Id.IsValid())
+			{
+				if (!(Record && Record.Get().GetKey() == Request.Key))
+				{
+					FCacheRecordPolicyBuilder PolicyBuilder(ECachePolicy::None);
+					PolicyBuilder.AddValuePolicy(Request.Id, Request.Policy);
+					Record.Reset();
+					Record = GetCacheRecordOnly(Request.Name, Request.Key, PolicyBuilder.Build());
+				}
+				if (Record)
+				{
+					const FValueWithId& ValueWithId = Record.Get().GetValue(Request.Id);
+					bHasValue = ValueWithId.IsValid();
+					Value = ValueWithId;
+					ValueId = Request.Id;
+					ValueKey = Request.Key;
+					GetCacheContent(Request.Name, Request.Key, ValueId, Value, Request.Policy, ValueReader, ValueAr);
+				}
+			}
+			else
+			{
+				ValueKey = Request.Key;
+				bHasValue = GetCacheValueOnly(Request.Name, Request.Key, Request.Policy, Value);
+				GetCacheContent(Request.Name, Request.Key, Request.Id, Value, Request.Policy, ValueReader, ValueAr);
+			}
+		}
+		if (bHasValue)
+		{
+			const uint64 RawOffset = FMath::Min(Value.GetRawSize(), Request.RawOffset);
+			const uint64 RawSize = FMath::Min(Value.GetRawSize() - RawOffset, Request.RawSize);
+			UE_LOG(LogDerivedDataCache, Verbose, TEXT("%s: Cache hit for %s from '%s'"),
+				*GetName(), *WriteToString<96>(Request.Key, '/', Request.Id), *Request.Name);
+			TRACE_COUNTER_INCREMENT(S3DDC_GetHit);
+			COOK_STAT(Timer.AddHit(!bExistsOnly ? RawSize : 0));
+			if (OnComplete)
+			{
+				FSharedBuffer Buffer;
+				if (!bExistsOnly)
+				{
+					Buffer = ValueReader.Decompress(RawOffset, RawSize);
+				}
+				const EStatus ChunkStatus = bExistsOnly || Buffer.GetSize() == RawSize ? EStatus::Ok : EStatus::Error;
+				OnComplete({Request.Name, Request.Key, Request.Id, Request.RawOffset,
+					RawSize, Value.GetRawHash(), MoveTemp(Buffer), Request.UserData, ChunkStatus});
+			}
+			continue;
+		}
+
+		if (OnComplete)
 		{
 			OnComplete({Request.Name, Request.Key, Request.Id, Request.RawOffset, 0, {}, {}, Request.UserData, EStatus::Error});
 		}
