Bump CacheStoreTest version as the AWS file share has had some test values put into it that are expected not to exist, meaning that the test always fails when using AWS as the shared cache.

diff --git a/Engine/Source/Developer/DerivedDataCache/Private/Tests/CacheStoreTest.cpp b/Engine/Source/Developer/DerivedDataCache/Private/Tests/CacheStoreTest.cpp
--- a/Engine/Source/Developer/DerivedDataCache/Private/Tests/CacheStoreTest.cpp
+++ b/Engine/Source/Developer/DerivedDataCache/Private/Tests/CacheStoreTest.cpp
@@ -21,7 +21,7 @@ bool FCacheStoreTest::RunTest(const FString& Parameters)
 {
 	using namespace UE::DerivedData;
 	ICache& Cache = GetCache();
-	FStringView TestVersion = TEXTVIEW("F72150A02AE34B57A9EC91D36BA1CE08");
+	FStringView TestVersion = TEXTVIEW("D6B05C93623D46D891D354BB22FCB584");
 
 	FRequestOwner Owner(EPriority::Blocking);
 	FCacheBucket DDCTestBucket(TEXT("DDCTest"));
