DDC: Enabled compression of legacy cache data

- FileSystem, Http, Pak, S3 use the ValueWithLegacyFallback mode by default, which cause them to fall back to searching for uncompressed data if compressed data is not found.
- FileSystem has been fixed to store up to 1 MiB of compressed data inline with the value package rather than separately in content-addressable storage.
- Pak has been optimized to have GetChunks only load the required region of the requested value, rather than the whole value.
- Pak has been changed to stop storing data inline in the record package, since it will end up in the same file anyway when stored separately.
- Pak will upgrade the compressor and compression level when a compressed pak file is requested. Default cache compression uses Oodle Mermaid VeryFast and will upgrade to Oodle Kraken Optimal2.
- Zen does not have compression enabled by default, pending deployment of a new version that stores compressed values to Horde Storage.
- Added a missing request barrier when saving uncompressed data as compressed.

Example reduction in file system cache size when cooking for Windows:
- CitySample dropped from 66.5 GiB to 21.8 GiB.
- Lyra dropped from 2.54 GiB to 672 MiB.
- ShooterGame dropped from 1.21 GiB to 380 MiB.

Example reduction in compressed pak file cache size when cooking for Windows:
- CitySample dropped from 22.3 GiB to 18.5 GiB.
- Lyra dropped from 691 MiB to 543 MiB.
- ShooterGame dropped from 387 MiB to 313 MiB.

diff --git a/Engine/Config/BaseEngine.ini b/Engine/Config/BaseEngine.ini
--- a/Engine/Config/BaseEngine.ini
+++ b/Engine/Config/BaseEngine.ini
@@ -2135,6 +2135,10 @@ Hierarchy=(Type=Hierarchical, Inner=Pak, Inner=EnginePak, Inner=Local, Inner=Sha
 ;	EnvPathOverride: An environment variable that if set will be used instead of Path. E.g., UE-LocalDataCachePath=D:\DDC. ('None' disables the DDC)
 ;	CommandLineOverride: A command line argument used in preference to the Path / EnvPathOverride settings. E.g., -SharedDataCachePath=\\someshare\folder
 ;	EditorOverrideSetting: Editor user setting that overrides the Path / EnvPathOverride line values.
+;	LegacyMode: Where to look for legacy cache keys in storage. One of ValueOnly, ValueWithLegacyFallback, LegacyOnly.
+;		ValueOnly: Look only for compressed data in the 5.0+ directory layout.
+;		LegacyOnly: Look only for uncompressed data in the pre-5.0 directory layout.
+;		ValueWithLegacyFallback: Look in the 5.0+ location but fall back to the pre-5.0 location on miss.
 Local=(Type=FileSystem, PurgeTransient=true, UnusedFileAge=34, PromptIfMissing=true, Path=%ENGINEDIR%DerivedDataCache, EnvPathOverride=UE-LocalDataCachePath, EditorOverrideSetting=LocalDerivedDataCache)
 ; Configure the shared DDC that is accessed after local. It's a filesystem DDC and the parameters are explained above
 Shared=(Type=FileSystem, UnusedFileAge=10, FoldersToClean=10, ConsiderSlowAt=70, Path=?EpicDDC, EnvPathOverride=UE-SharedDataCachePath, EditorOverrideSetting=SharedDerivedDataCache, CommandLineOverride=SharedDataCachePath)
diff --git a/Engine/Source/Developer/DerivedDataCache/Private/DerivedDataBackendCorruptionWrapper.h b/Engine/Source/Developer/DerivedDataCache/Private/DerivedDataBackendCorruptionWrapper.h
--- a/Engine/Source/Developer/DerivedDataCache/Private/DerivedDataBackendCorruptionWrapper.h
+++ b/Engine/Source/Developer/DerivedDataCache/Private/DerivedDataBackendCorruptionWrapper.h
@@ -10,7 +10,7 @@
 #include "Misc/Crc.h"
 #include "Serialization/Archive.h"
 
-namespace UE::DerivedData::CacheStore
+namespace UE::DerivedData
 {
 
 /**
@@ -83,4 +83,4 @@ public:
 	}
 };
 
-} // UE::DerivedData::CacheStore
+} // UE::DerivedData
diff --git a/Engine/Source/Developer/DerivedDataCache/Private/DerivedDataBackendInterface.h b/Engine/Source/Developer/DerivedDataCache/Private/DerivedDataBackendInterface.h
--- a/Engine/Source/Developer/DerivedDataCache/Private/DerivedDataBackendInterface.h
+++ b/Engine/Source/Developer/DerivedDataCache/Private/DerivedDataBackendInterface.h
@@ -210,7 +210,7 @@ public:
 	 */
 	virtual bool ApplyDebugOptions(FBackendDebugOptions& InOptions) = 0;
 
-	virtual EBackendLegacyMode GetLegacyMode() const { return EBackendLegacyMode::LegacyOnly; }
+	virtual EBackendLegacyMode GetLegacyMode() const = 0;
 
 	virtual void LegacyPut(
 		TConstArrayView<FLegacyCachePutRequest> Requests,
diff --git a/Engine/Source/Developer/DerivedDataCache/Private/DerivedDataBackends.cpp b/Engine/Source/Developer/DerivedDataCache/Private/DerivedDataBackends.cpp
--- a/Engine/Source/Developer/DerivedDataCache/Private/DerivedDataBackends.cpp
+++ b/Engine/Source/Developer/DerivedDataCache/Private/DerivedDataBackends.cpp
@@ -23,7 +23,7 @@
 #include "Misc/Paths.h"
 #include "Misc/StringBuilder.h"
 #include "Modules/ModuleManager.h"
-#include "PakFileDerivedDataBackend.h"
+#include "PakFileCacheStore.h"
 #include "ProfilingDebugging/CookStats.h"
 #include "Serialization/CompactBinaryPackage.h"
 #include <atomic>
@@ -43,20 +43,13 @@ static TAutoConsoleVariable<FString> GDerivedDataCacheGraphName(
 
 namespace UE::DerivedData
 {
+
 ILegacyCacheStore* CreateCacheStoreAsync(ILegacyCacheStore* InnerBackend, ECacheStoreFlags InnerFlags, bool bCacheInFlightPuts);
 ILegacyCacheStore* CreateCacheStoreHierarchy(ICacheStoreOwner*& OutOwner);
 ILegacyCacheStore* CreateCacheStoreThrottle(ILegacyCacheStore* InnerCache, uint32 LatencyMS, uint32 MaxBytesPerSecond);
 ILegacyCacheStore* CreateCacheStoreVerify(ILegacyCacheStore* InnerCache, bool bPutOnError);
-} // UE::DerivedData
-
-namespace UE::DerivedData::CacheStore::FileSystem
-{
-FDerivedDataBackendInterface* CreateFileSystemDerivedDataBackend(const TCHAR* CacheDirectory, const TCHAR* Params, const TCHAR* AccessLogFileName);
-} // UE::DerivedData::CacheStore::FileSystem
-
-namespace UE::DerivedData::CacheStore::Http
-{
-FDerivedDataBackendInterface* CreateHttpDerivedDataBackend(
+ILegacyCacheStore* CreateFileSystemCacheStore(const TCHAR* CacheDirectory, const TCHAR* Params, const TCHAR* AccessLogFileName, ECacheStoreFlags& OutFlags);
+ILegacyCacheStore* CreateHttpCacheStore(
 	const TCHAR* NodeName,
 	const TCHAR* ServiceUrl,
 	const TCHAR* Namespace,
@@ -65,31 +58,12 @@ FDerivedDataBackendInterface* CreateHttpDerivedDataBackend(
 	const TCHAR* OAuthClientId,
 	const TCHAR* OAuthData, 
 	const FDerivedDataBackendInterface::ESpeedClass* ForceSpeedClass,
+	EBackendLegacyMode LegacyMode,
 	bool bReadOnly);
-} // UE::DerivedData::CacheStore::Http
-
-namespace UE::DerivedData::CacheStore::Memory
-{
-FFileBackedDerivedDataBackend* CreateMemoryDerivedDataBackend(const TCHAR* Name, int64 MaxCacheSize, bool bCanBeDisabled);
-} // UE::DerivedData::CacheStore::Memory
-
-namespace UE::DerivedData::CacheStore::PakFile
-{
-IPakFileDerivedDataBackend* CreatePakFileDerivedDataBackend(const TCHAR* Filename, bool bWriting, bool bCompressed);
-} // UE::DerivedData::CacheStore::PakFile
-
-namespace UE::DerivedData::CacheStore::S3
-{
-FDerivedDataBackendInterface* CreateS3DerivedDataBackend(const TCHAR* RootManifestPath, const TCHAR* BaseUrl, const TCHAR* Region, const TCHAR* CanaryObjectKey, const TCHAR* CachePath);
-} // UE::DerivedData::CacheStore::S3
-
-namespace UE::DerivedData::CacheStore::ZenCache
-{
-FDerivedDataBackendInterface* CreateZenDerivedDataBackend(const TCHAR* NodeName, const TCHAR* ServiceUrl, const TCHAR* Namespace);
-} // UE::DerivedData::CacheStore::ZenCache
-
-namespace UE::DerivedData
-{
+FFileBackedDerivedDataBackend* CreateMemoryCacheStore(const TCHAR* Name, int64 MaxCacheSize, bool bCanBeDisabled);
+IPakFileCacheStore* CreatePakFileCacheStore(const TCHAR* Filename, bool bWriting, bool bCompressed);
+ILegacyCacheStore* CreateS3CacheStore(const TCHAR* RootManifestPath, const TCHAR* BaseUrl, const TCHAR* Region, const TCHAR* CanaryObjectKey, const TCHAR* CachePath);
+ILegacyCacheStore* CreateZenCacheStore(const TCHAR* NodeName, const TCHAR* ServiceUrl, const TCHAR* Namespace);
 
 /**
  * This class is used to create a singleton that represents the derived data cache hierarchy and all of the wrappers necessary
@@ -383,7 +357,7 @@ public:
 				FGuid Temp = FGuid::NewGuid();
 				ReadPakFilename = PakFilename;
 				WritePakFilename = PakFilename + TEXT(".") + Temp.ToString();
-				WritePakCache = CacheStore::PakFile::CreatePakFileDerivedDataBackend(*WritePakFilename, /*bWriting*/ true, bCompressed);
+				WritePakCache = CreatePakFileCacheStore(*WritePakFilename, /*bWriting*/ true, bCompressed);
 				PakNode = WritePakCache;
 			}
 			else
@@ -391,7 +365,7 @@ public:
 				bool bReadPak = FPlatformFileManager::Get().GetPlatformFile().FileExists( *PakFilename );
 				if( bReadPak )
 				{
-					CacheStore::PakFile::IPakFileDerivedDataBackend* ReadPak = CacheStore::PakFile::CreatePakFileDerivedDataBackend(*PakFilename, /*bWriting*/ false, bCompressed);
+					IPakFileCacheStore* ReadPak = CreatePakFileCacheStore(*PakFilename, /*bWriting*/ false, bCompressed);
 					ReadPakFilename = PakFilename;
 					PakNode = ReadPak;
 					ReadPakCache.Add(ReadPak);
@@ -662,8 +636,6 @@ public:
 		}
 		else
 		{
-			FDerivedDataBackendInterface* InnerFileSystem = nullptr;
-
 			// Try to set up the shared drive, allow user to correct any issues that may exist.
 			bool RetryOnFailure = false;
 			do
@@ -680,21 +652,18 @@ public:
 				FParse::Value(Entry, TEXT("WriteAccessLog="), WriteAccessLog);
 				FParse::Bool(Entry, TEXT("PromptIfMissing="), bPromptIfMissing);
 
+				ILegacyCacheStore* InnerFileSystem = nullptr;
+				ECacheStoreFlags Flags;
 				if (!bShared || IFileManager::Get().DirectoryExists(*Path))
 				{
-					InnerFileSystem = CacheStore::FileSystem::CreateFileSystemDerivedDataBackend(*Path, Entry, *WriteAccessLog);
+					InnerFileSystem = CreateFileSystemCacheStore(*Path, Entry, *WriteAccessLog, Flags);
 				}
 
 				if (InnerFileSystem)
 				{
 					bUsingSharedDDC |= bShared;
-
-					ECacheStoreFlags Flags = ECacheStoreFlags::Query;
-					Flags |= InnerFileSystem->IsWritable() ? ECacheStoreFlags::Store : ECacheStoreFlags::None;
-					Flags |= InnerFileSystem->GetSpeedClass() == EBackendSpeedClass::Local ? ECacheStoreFlags::Local : ECacheStoreFlags::Remote;
-
 					DataCache = MakeTuple(InnerFileSystem, Flags);
-					UE_LOG(LogDerivedDataCache, Log, TEXT("Using %s data cache path %s: %s"), NodeName, *Path, !InnerFileSystem->IsWritable() ? TEXT("ReadOnly") : TEXT("Writable"));
+					UE_LOG(LogDerivedDataCache, Log, TEXT("Using %s data cache path %s: %s"), NodeName, *Path, !EnumHasAnyFlags(Flags, ECacheStoreFlags::Store) ? TEXT("ReadOnly") : TEXT("Writable"));
 					Directories.AddUnique(Path);
 				}
 				else
@@ -780,7 +749,7 @@ public:
 			}
 		}
 
-		if (ILegacyCacheStore* Backend = CacheStore::S3::CreateS3DerivedDataBackend(*ManifestPath, *BaseUrl, *Region, *CanaryObjectKey, *CachePath))
+		if (ILegacyCacheStore* Backend = CreateS3CacheStore(*ManifestPath, *BaseUrl, *Region, *CanaryObjectKey, *CachePath))
 		{
 			return Backend;
 		}
@@ -800,6 +769,7 @@ public:
 		FString& OAuthProvider,
 		FString& OAuthClientId,
 		FString& OAuthSecret,
+		EBackendLegacyMode& LegacyMode,
 		bool& bReadOnly)
 	{
 		FString ServerId;
@@ -809,7 +779,7 @@ public:
 			const TCHAR* ServerSection = TEXT("HordeStorageServers");
 			if (GConfig->GetString(ServerSection, *ServerId, ServerEntry, IniFilename))
 			{
-				ParseHttpCacheParams(NodeName, *ServerEntry, IniFilename, IniSection, Host, Namespace, StructuredNamespace, OAuthProvider, OAuthClientId, OAuthSecret, bReadOnly);
+				ParseHttpCacheParams(NodeName, *ServerEntry, IniFilename, IniSection, Host, Namespace, StructuredNamespace, OAuthProvider, OAuthClientId, OAuthSecret, LegacyMode, bReadOnly);
 			}
 			else
 			{
@@ -845,6 +815,14 @@ public:
 		FParse::Value(Entry, TEXT("OAuthClientId="), OAuthClientId);
 		FParse::Value(Entry, TEXT("OAuthSecret="), OAuthSecret);
 		FParse::Bool(Entry, TEXT("ReadOnly="), bReadOnly);
+
+		if (FString LegacyModeString; FParse::Value(Entry, TEXT("LegacyMode="), LegacyModeString))
+		{
+			if (!TryLexFromString(LegacyMode, LegacyModeString))
+			{
+				UE_LOG(LogDerivedDataCache, Warning, TEXT("%s: Ignoring unrecognized legacy mode '%s'"), NodeName, *LegacyModeString);
+			}
+		}
 	}
 
 	/**
@@ -862,9 +840,10 @@ public:
 		FString OAuthProvider;
 		FString OAuthClientId;
 		FString OAuthSecret;
+		EBackendLegacyMode LegacyMode = EBackendLegacyMode::ValueWithLegacyFallback;
 		bool bReadOnly = false;
 
-		ParseHttpCacheParams(NodeName, Entry, IniFilename, IniSection, Host, Namespace, StructuredNamespace, OAuthProvider, OAuthClientId, OAuthSecret, bReadOnly);
+		ParseHttpCacheParams(NodeName, Entry, IniFilename, IniSection, Host, Namespace, StructuredNamespace, OAuthProvider, OAuthClientId, OAuthSecret, LegacyMode, bReadOnly);
 
 		if (Host.IsEmpty())
 		{
@@ -938,9 +917,9 @@ public:
 			UE_LOG(LogDerivedDataCache, Log, TEXT("Node %s found speed class override ForceSpeedClass=%s"), NodeName, *ForceSpeedClassValue);
 		}
 
-		return MakeTuple(CacheStore::Http::CreateHttpDerivedDataBackend(
+		return MakeTuple(CreateHttpCacheStore(
 			NodeName, *Host, *Namespace, *StructuredNamespace, *OAuthProvider, *OAuthClientId, *OAuthSecret,
-			ForceSpeedClass == FDerivedDataBackendInterface::ESpeedClass::Unknown ? nullptr : &ForceSpeedClass, bReadOnly),
+			ForceSpeedClass == FDerivedDataBackendInterface::ESpeedClass::Unknown ? nullptr : &ForceSpeedClass, LegacyMode, bReadOnly),
 			ECacheStoreFlags::Remote | ECacheStoreFlags::Query | (bReadOnly ? ECacheStoreFlags::None : ECacheStoreFlags::Store));
 	}
 
@@ -959,7 +938,7 @@ public:
 			UE_LOG(LogDerivedDataCache, Warning, TEXT("Node %s does not specify 'Namespace', falling back to '%s'"), NodeName, *Namespace);
 		}
 
-		if (ILegacyCacheStore* Backend = CacheStore::ZenCache::CreateZenDerivedDataBackend(NodeName, *ServiceUrl, *Namespace))
+		if (ILegacyCacheStore* Backend = CreateZenCacheStore(NodeName, *ServiceUrl, *Namespace))
 		{
 			return Backend;
 		}
@@ -998,7 +977,7 @@ public:
 			MaxCacheSize = FMath::Min(MaxCacheSize, MaxSupportedCacheSize);
 
 			UE_LOG( LogDerivedDataCache, Display, TEXT("Max Cache Size: %d MB"), MaxCacheSize);
-			Cache = CacheStore::Memory::CreateMemoryDerivedDataBackend(TEXT("Boot"), MaxCacheSize * 1024 * 1024, /*bCanBeDisabled*/ true);
+			Cache = CreateMemoryCacheStore(TEXT("Boot"), MaxCacheSize * 1024 * 1024, /*bCanBeDisabled*/ true);
 
 			if( Cache && Filename.Len() )
 			{
@@ -1040,7 +1019,7 @@ public:
 		FString Filename;
 
 		FParse::Value( Entry, TEXT("Filename="), Filename );
-		Cache = CacheStore::Memory::CreateMemoryDerivedDataBackend(NodeName, /*MaxCacheSize*/ -1, /*bCanBeDisabled*/ false);
+		Cache = CreateMemoryCacheStore(NodeName, /*MaxCacheSize*/ -1, /*bCanBeDisabled*/ false);
 		if( Cache && Filename.Len() )
 		{
 			if( Cache->LoadCache( *Filename ) )
@@ -1116,8 +1095,8 @@ public:
 
 				for(const FString& MergePakName : MergePakList)
 				{
-					TUniquePtr<CacheStore::PakFile::IPakFileDerivedDataBackend> ReadPak(
-						CacheStore::PakFile::CreatePakFileDerivedDataBackend(*FPaths::Combine(*FPaths::GetPath(WritePakFilename), *MergePakName), /*bWriting*/ false, /*bCompressed*/ false));
+					TUniquePtr<IPakFileCacheStore> ReadPak(
+						CreatePakFileCacheStore(*FPaths::Combine(*FPaths::GetPath(WritePakFilename), *MergePakName), /*bWriting*/ false, /*bCompressed*/ false));
 					WritePakCache->MergeCache(ReadPak.Get());
 				}
 			}
@@ -1142,7 +1121,7 @@ public:
 							UE_LOG(LogDerivedDataCache, Error, TEXT("Could not delete the pak file %s to overwrite it with a new one."), *ReadPakFilename);
 						}
 					}
-					if (!CacheStore::PakFile::IPakFileDerivedDataBackend::SortAndCopy(WritePakFilename, ReadPakFilename))
+					if (!IPakFileCacheStore::SortAndCopy(WritePakFilename, ReadPakFilename))
 					{
 						UE_LOG(LogDerivedDataCache, Error, TEXT("Couldn't sort pak file (%s)"), *WritePakFilename);
 					}
@@ -1206,10 +1185,10 @@ public:
 	{
 		// Assumptions: there's at least one read-only pak backend in the hierarchy
 		// and its parent is a hierarchical backend.
-		CacheStore::PakFile::IPakFileDerivedDataBackend* ReadPak = nullptr;
+		IPakFileCacheStore* ReadPak = nullptr;
 		if (Hierarchy && FPlatformFileManager::Get().GetPlatformFile().FileExists(PakFilename))
 		{
-			ReadPak = CacheStore::PakFile::CreatePakFileDerivedDataBackend(PakFilename, /*bWriting*/ false, /*bCompressed*/ false);
+			ReadPak = CreatePakFileCacheStore(PakFilename, /*bWriting*/ false, /*bCompressed*/ false);
 
 			Hierarchy->Add(ReadPak, ECacheStoreFlags::Local | ECacheStoreFlags::Store | ECacheStoreFlags::StopStore);
 			CreatedNodes.AddUnique(ReadPak);
@@ -1227,7 +1206,7 @@ public:
 	{
 		for (int PakIndex = 0; PakIndex < ReadPakCache.Num(); ++PakIndex)
 		{
-			CacheStore::PakFile::IPakFileDerivedDataBackend* ReadPak = ReadPakCache[PakIndex];
+			IPakFileCacheStore* ReadPak = ReadPakCache[PakIndex];
 			if (ReadPak->GetFilename() == PakFilename)
 			{
 				check(Hierarchy);
@@ -1304,11 +1283,11 @@ private:
 
 	/** Instances of backend interfaces which exist in only one copy */
 	FFileBackedDerivedDataBackend*	BootCache;
-	CacheStore::PakFile::IPakFileDerivedDataBackend* WritePakCache;
+	IPakFileCacheStore* WritePakCache;
 	ILegacyCacheStore*	AsyncNode;
 	ICacheStoreOwner* Hierarchy;
 	/** Support for multiple read only pak files. */
-	TArray<CacheStore::PakFile::IPakFileDerivedDataBackend*> ReadPakCache;
+	TArray<IPakFileCacheStore*> ReadPakCache;
 
 	/** List of directories used by the DDC */
 	TArray<FString> Directories;
@@ -1415,6 +1394,7 @@ void FDerivedDataBackendInterface::LegacyGet(
 	// Query the legacy cache by translating the requests to legacy cache functions.
 
 	FRequestOwner AsyncOwner(FPlatformMath::Min(Owner.GetPriority(), EPriority::Highest));
+	FRequestBarrier Barrier(AsyncOwner);
 	AsyncOwner.KeepAlive();
 
 	TArray<FString, TInlineAllocator<8>> ExistsKeys;
diff --git a/Engine/Source/Developer/DerivedDataCache/Private/DerivedDataCache.cpp b/Engine/Source/Developer/DerivedDataCache/Private/DerivedDataCache.cpp
--- a/Engine/Source/Developer/DerivedDataCache/Private/DerivedDataCache.cpp
+++ b/Engine/Source/Developer/DerivedDataCache/Private/DerivedDataCache.cpp
@@ -160,34 +160,6 @@ bool GVerifyDDC = false;
 namespace UE::DerivedData
 {
 
-void ICacheStore::PutValue(
-	const TConstArrayView<FCachePutValueRequest> Requests,
-	IRequestOwner& Owner,
-	FOnCachePutValueComplete&& OnComplete)
-{
-	if (OnComplete)
-	{
-		for (const FCachePutValueRequest& Request : Requests)
-		{
-			OnComplete({Request.Name, Request.Key, Request.UserData, EStatus::Error});
-		}
-	}
-}
-
-void ICacheStore::GetValue(
-	const TConstArrayView<FCacheGetValueRequest> Requests,
-	IRequestOwner& Owner,
-	FOnCacheGetValueComplete&& OnComplete)
-{
-	if (OnComplete)
-	{
-		for (const FCacheGetValueRequest& Request : Requests)
-		{
-			OnComplete({Request.Name, Request.Key, {}, Request.UserData, EStatus::Error});
-		}
-	}
-}
-
 FCachePutResponse FCachePutRequest::MakeResponse(const EStatus Status) const
 {
 	return {Name, Record.GetKey(), UserData, Status};
diff --git a/Engine/Source/Developer/DerivedDataCache/Private/DerivedDataCacheStore.h b/Engine/Source/Developer/DerivedDataCache/Private/DerivedDataCacheStore.h
--- a/Engine/Source/Developer/DerivedDataCache/Private/DerivedDataCacheStore.h
+++ b/Engine/Source/Developer/DerivedDataCache/Private/DerivedDataCacheStore.h
@@ -65,7 +65,7 @@ public:
 	virtual void PutValue(
 		TConstArrayView<FCachePutValueRequest> Requests,
 		IRequestOwner& Owner,
-		FOnCachePutValueComplete&& OnComplete);
+		FOnCachePutValueComplete&& OnComplete) = 0;
 
 	/**
 	 * Asynchronous request to get values from the cache.
@@ -77,7 +77,7 @@ public:
 	virtual void GetValue(
 		TConstArrayView<FCacheGetValueRequest> Requests,
 		IRequestOwner& Owner,
-		FOnCacheGetValueComplete&& OnComplete);
+		FOnCacheGetValueComplete&& OnComplete) = 0;
 
 	/**
 	 * Asynchronous request to get chunks, which are subsets of values, from records or values.
diff --git a/Engine/Source/Developer/DerivedDataCache/Private/DerivedDataCacheStoreAsync.cpp b/Engine/Source/Developer/DerivedDataCache/Private/DerivedDataCacheStoreAsync.cpp
--- a/Engine/Source/Developer/DerivedDataCache/Private/DerivedDataCacheStoreAsync.cpp
+++ b/Engine/Source/Developer/DerivedDataCache/Private/DerivedDataCacheStoreAsync.cpp
@@ -16,14 +16,11 @@
 #include "Tasks/Task.h"
 #include "Templates/Invoke.h"
 
-namespace UE::DerivedData::CacheStore::Memory
-{
-FFileBackedDerivedDataBackend* CreateMemoryDerivedDataBackend(const TCHAR* Name, int64 MaxCacheSize, bool bCanBeDisabled);
-} // UE::DerivedData::CacheStore::Memory
-
 namespace UE::DerivedData
 {
 
+FFileBackedDerivedDataBackend* CreateMemoryCacheStore(const TCHAR* Name, int64 MaxCacheSize, bool bCanBeDisabled);
+
 /**
  * A cache store that executes non-blocking requests in a dedicated thread pool.
  *
@@ -125,7 +122,7 @@ private:
 FCacheStoreAsync::FCacheStoreAsync(ILegacyCacheStore* InInnerCache, ECacheStoreFlags InInnerFlags, bool bCacheInFlightPuts)
 	: InnerCache(InInnerCache)
 	, InnerFlags(InInnerFlags)
-	, MemoryCache(bCacheInFlightPuts ? CacheStore::Memory::CreateMemoryDerivedDataBackend(TEXT("InflightMemoryCache"), /*MaxCacheSize*/ -1, /*bCanBeDisabled*/ false) : nullptr)
+	, MemoryCache(bCacheInFlightPuts ? CreateMemoryCacheStore(TEXT("InflightMemoryCache"), /*MaxCacheSize*/ -1, /*bCanBeDisabled*/ false) : nullptr)
 {
 	check(InnerCache);
 }
