DDC: Enabled compression of legacy cache data

diff --git a/Engine/Source/Developer/DerivedDataCache/Private/FileSystemCacheStore.cpp b/Engine/Source/Developer/DerivedDataCache/Private/FileSystemCacheStore.cpp
new file mode 100644
--- /dev/null
+++ b/Engine/Source/Developer/DerivedDataCache/Private/FileSystemCacheStore.cpp
@@ -0,0 +1,2650 @@
+// Copyright Epic Games, Inc. All Rights Reserved.
+
+#include "Algo/Accumulate.h"
+#include "Algo/AllOf.h"
+#include "Algo/Compare.h"
+#include "Algo/StableSort.h"
+#include "Algo/Transform.h"
+#include "Async/Async.h"
+#include "Async/TaskGraphInterfaces.h"
+#include "Containers/StaticBitArray.h"
+#include "DerivedDataBackendCorruptionWrapper.h"
+#include "DerivedDataBackendInterface.h"
+#include "DerivedDataCacheInterface.h"
+#include "DerivedDataCacheMaintainer.h"
+#include "DerivedDataCachePrivate.h"
+#include "DerivedDataCacheRecord.h"
+#include "DerivedDataCacheUsageStats.h"
+#include "DerivedDataChunk.h"
+#include "DerivedDataValue.h"
+#include "Experimental/Async/LazyEvent.h"
+#include "Features/IModularFeatures.h"
+#include "HAL/Event.h"
+#include "HAL/FileManager.h"
+#include "HAL/Thread.h"
+#include "Hash/xxhash.h"
+#include "HashingArchiveProxy.h"
+#include "Math/NumericLimits.h"
+#include "Misc/CommandLine.h"
+#include "Misc/ConfigCacheIni.h"
+#include "Misc/CoreMisc.h"
+#include "Misc/FileHelper.h"
+#include "Misc/Guid.h"
+#include "Misc/MessageDialog.h"
+#include "Misc/Paths.h"
+#include "Misc/PathViews.h"
+#include "Misc/ScopeExit.h"
+#include "Misc/ScopeLock.h"
+#include "Misc/StringBuilder.h"
+#include "ProfilingDebugging/CookStats.h"
+#include "ProfilingDebugging/CountersTrace.h"
+#include "ProfilingDebugging/CpuProfilerTrace.h"
+#include "Serialization/CompactBinary.h"
+#include "Serialization/CompactBinaryPackage.h"
+#include "Serialization/CompactBinaryValidation.h"
+#include "Serialization/CompactBinaryWriter.h"
+#include "Templates/Greater.h"
+
+#define MAX_BACKEND_KEY_LENGTH (120)
+#define MAX_BACKEND_NUMBERED_SUBFOLDER_LENGTH (9)
+#if PLATFORM_LINUX	// PATH_MAX on Linux is 4096 (getconf PATH_MAX /, also see limits.h), so this value can be larger (note that it is still arbitrary).
+                    // This should not affect sharing the cache between platforms as the absolute paths will be different anyway.
+	#define MAX_CACHE_DIR_LEN (3119)
+#else
+	#define MAX_CACHE_DIR_LEN (119)
+#endif // PLATFORM_LINUX
+#define MAX_CACHE_EXTENTION_LEN (4)
+
+namespace UE::DerivedData
+{
+
+TRACE_DECLARE_INT_COUNTER(FileSystemDDC_Exist, TEXT("FileSystemDDC Exist"));
+TRACE_DECLARE_INT_COUNTER(FileSystemDDC_ExistHit, TEXT("FileSystemDDC Exist Hit"));
+TRACE_DECLARE_INT_COUNTER(FileSystemDDC_Get, TEXT("FileSystemDDC Get"));
+TRACE_DECLARE_INT_COUNTER(FileSystemDDC_GetHit, TEXT("FileSystemDDC Get Hit"));
+TRACE_DECLARE_INT_COUNTER(FileSystemDDC_Put, TEXT("FileSystemDDC Put"));
+TRACE_DECLARE_INT_COUNTER(FileSystemDDC_PutHit, TEXT("FileSystemDDC Put Hit"));
+TRACE_DECLARE_INT_COUNTER(FileSystemDDC_BytesRead, TEXT("FileSystemDDC Bytes Read"));
+TRACE_DECLARE_INT_COUNTER(FileSystemDDC_BytesWritten, TEXT("FileSystemDDC Bytes Written"));
+
+///////////////////////////////////////////////////////////////////////////////////////////////////////////////////////
+
+static const TCHAR GBucketsDirectoryName[] = TEXT("Buckets");
+static const TCHAR GContentDirectoryName[] = TEXT("Content");
+
+///////////////////////////////////////////////////////////////////////////////////////////////////////////////////////
+
+void BuildPathForLegacyCache(const TCHAR* const CacheKey, FStringBuilderBase& Path)
+{
+	TStringBuilder<256> Key;
+	Key << CacheKey;
+	for (TCHAR& Char : MakeArrayView(Key))
+	{
+		Char = FChar::ToUpper(Char);
+	}
+	checkf(Algo::AllOf(MakeArrayView(Key),
+		[](TCHAR C) { return FChar::IsAlnum(C) || FChar::IsUnderscore(C) || C == TEXT('$'); }),
+		TEXT("Invalid characters in cache key %s"), CacheKey);
+	const uint32 Hash = FCrc::StrCrc_DEPRECATED(Key.Len(), Key.GetData());
+	Path.Appendf(TEXT("%1d/%1d/%1d/%s.udd"), (Hash / 100) % 10, (Hash / 10) % 10, Hash % 10, *Key);
+}
+
+void BuildPathForCachePackage(const FCacheKey& CacheKey, FStringBuilderBase& Path)
+{
+	const FIoHash::ByteArray& Bytes = CacheKey.Hash.GetBytes();
+	Path.Appendf(TEXT("%s/%hs/%02x/%02x/"), GBucketsDirectoryName, CacheKey.Bucket.ToCString(), Bytes[0], Bytes[1]);
+	UE::String::BytesToHexLower(MakeArrayView(Bytes).RightChop(2), Path);
+	Path << TEXT(".udd"_SV);
+}
+
+void BuildPathForCacheContent(const FIoHash& RawHash, FStringBuilderBase& Path)
+{
+	const FIoHash::ByteArray& Bytes = RawHash.GetBytes();
+	Path.Appendf(TEXT("%s/%02x/%02x/"), GContentDirectoryName, Bytes[0], Bytes[1]);
+	UE::String::BytesToHexLower(MakeArrayView(Bytes).RightChop(2), Path);
+	Path << TEXT(".udd"_SV);
+}
+
+///////////////////////////////////////////////////////////////////////////////////////////////////////////////////////
+
+static uint64 RandFromGuid()
+{
+	const FGuid Guid = FGuid::NewGuid();
+	return FXxHash64::HashBuffer(&Guid, sizeof(FGuid)).Hash;
+}
+
+/** A LCG in which the modulus is a power of two where the exponent is the bit width of T. */
+template <typename T, T Modulus = 0>
+class TLinearCongruentialGenerator
+{
+	static_assert(!TIsSigned<T>::Value);
+	static_assert((Modulus & (Modulus - 1)) == 0, "Modulus must be a power of two.");
+
+public:
+	constexpr inline TLinearCongruentialGenerator(T InMultiplier, T InIncrement)
+		: Multiplier(InMultiplier)
+		, Increment(InIncrement)
+	{
+	}
+
+	constexpr inline T GetNext(T& Value)
+	{
+		Value = (Value * Multiplier + Increment) & (Modulus - 1);
+		return Value;
+	}
+
+private:
+	const T Multiplier;
+	const T Increment;
+};
+
+class FRandomStream
+{
+public:
+	inline explicit FRandomStream(const uint32 Seed)
+		: Random(1103515245, 12345) // From ANSI C
+		, Value(Seed)
+	{
+	}
+
+	/** Returns a random value in [Min, Max). */
+	inline uint32 GetRandRange(const uint32 Min, const uint32 Max)
+	{
+		return Min + uint32((uint64(Max - Min) * Random.GetNext(Value)) >> 32);
+	}
+
+private:
+	TLinearCongruentialGenerator<uint32> Random;
+	uint32 Value;
+};
+
+template <uint32 Modulus, uint32 Count = Modulus>
+class TRandomOrder
+{
+	static_assert((Modulus & (Modulus - 1)) == 0 && Modulus > 16, "Modulus must be a power of two greater than 16.");
+	static_assert(Count > 0 && Count <= Modulus, "Count must be in the range (0, Modulus].");
+
+public:
+	inline explicit TRandomOrder(FRandomStream& Stream)
+		: Random(Stream.GetRandRange(0, Modulus / 16) * 8 + 5, 12345)
+		, First(Stream.GetRandRange(0, Count))
+		, Value(First)
+	{
+	}
+
+	inline uint32 GetFirst() const
+	{
+		return First;
+	}
+
+	inline uint32 GetNext()
+	{
+		if constexpr (Count < Modulus)
+		{
+			for (;;)
+			{
+				if (const uint32 Next = Random.GetNext(Value); Next < Count)
+				{
+					return Next;
+				}
+			}
+		}
+		return Random.GetNext(Value);
+	}
+
+private:
+	TLinearCongruentialGenerator<uint32, Modulus> Random;
+	uint32 First;
+	uint32 Value;
+};
+
+///////////////////////////////////////////////////////////////////////////////////////////////////////////////////////
+
+struct FFileSystemCacheStoreMaintainerParams
+{
+	/** Files older than this will be deleted. */
+	FTimespan MaxFileAge = FTimespan::FromDays(15.0);
+	/** Limits the number of files scanned in one second. */
+	uint32 MaxFileScanRate = MAX_uint32;
+	/** Limits the number of directories scanned in each cache bucket or content root. */
+	uint32 MaxDirectoryScanCount = MAX_uint32;
+	/** Minimum duration between the start of consecutive scans. Use MaxValue to scan only once. */
+	FTimespan ScanFrequency = FTimespan::FromHours(1.0);
+	/** Time to wait after initialization before maintenance begins. */
+	FTimespan TimeToWaitAfterInit = FTimespan::FromMinutes(1.0);
+};
+
+class FFileSystemCacheStoreMaintainer final : public ICacheStoreMaintainer
+{
+public:
+	FFileSystemCacheStoreMaintainer(const FFileSystemCacheStoreMaintainerParams& Params, FStringView CachePath);
+	~FFileSystemCacheStoreMaintainer();
+
+	bool IsIdle() const final { return bIdle; }
+	void WaitForIdle() const { IdleEvent.Wait(); }
+	void BoostPriority() final;
+
+private:
+	void Tick();
+	void Loop();
+	void Scan();
+
+	void CreateContentRoot();
+	void CreateBucketRoots();
+	void ScanHashRoot(uint32 RootIndex);
+	TStaticBitArray<256> ScanHashDirectory(FStringBuilderBase& BasePath);
+
+	TStaticBitArray<10> ScanLegacyDirectory(FStringBuilderBase& BasePath);
+	void CreateLegacyRoot();
+	void ScanLegacyRoot();
+
+	void ResetRoots();
+
+	void ProcessFile(const TCHAR* Path, const FFileStatData& Stat);
+
+private:
+	struct FRoot;
+	struct FLegacyRoot;
+
+	FFileSystemCacheStoreMaintainerParams Params;
+	/** Path to the root of the cache store. */
+	FString CachePath;
+	/** True when there is no active maintenance scan. */
+	bool bIdle = false;
+	/** True when maintenance is expected to exit as soon as possible. */
+	bool bExit = false;
+	/** True when maintenance is expected to exit at the end of the scan. */
+	bool bExitAfterScan = false;
+	/** Ignore the file scan rate for one maintenance scan. */
+	bool bIgnoreFileScanRate = false;
+
+	uint32 ProcessCount = 0;
+	uint32 DeleteCount = 0;
+	uint64 DeleteSize = 0;
+
+	double BatchStartTime = 0.0;
+
+	IFileManager& FileManager = IFileManager::Get();
+
+	mutable FLazyEvent IdleEvent;
+	FEventRef WaitEvent;
+	FThread Thread;
+
+	TArray<TUniquePtr<FRoot>> Roots;
+	TUniquePtr<FLegacyRoot> LegacyRoot;
+	FRandomStream Random{uint32(RandFromGuid())};
+
+	static constexpr double MaxScanFrequencyDays = 365.0;
+};
+
+struct FFileSystemCacheStoreMaintainer::FRoot
+{
+	inline FRoot(const FStringView RootPath, FRandomStream& Stream)
+		: Order(Stream)
+	{
+		Path.Append(RootPath);
+	}
+
+	TStringBuilder<256> Path;
+	TRandomOrder<256 * 256> Order;
+	TStaticBitArray<256> ScannedLevel0;
+	TStaticBitArray<256> ExistsLevel0;
+	TStaticBitArray<256> ExistsLevel1[256];
+	uint32 DirectoryScanCount = 0;
+	bool bScannedRoot = false;
+};
+
+struct FFileSystemCacheStoreMaintainer::FLegacyRoot
+{
+	inline explicit FLegacyRoot(FRandomStream& Stream)
+		: Order(Stream)
+	{
+	}
+
+	TRandomOrder<1024, 1000> Order;
+	TStaticBitArray<10> ScannedLevel0;
+	TStaticBitArray<10> ScannedLevel1[10];
+	TStaticBitArray<10> ExistsLevel0;
+	TStaticBitArray<10> ExistsLevel1[10];
+	TStaticBitArray<10> ExistsLevel2[10][10];
+	uint32 DirectoryScanCount = 0;
+};
+
+FFileSystemCacheStoreMaintainer::FFileSystemCacheStoreMaintainer(
+	const FFileSystemCacheStoreMaintainerParams& InParams,
+	const FStringView InCachePath)
+	: Params(InParams)
+	, CachePath(InCachePath)
+	, bExitAfterScan(Params.ScanFrequency.GetTotalDays() > MaxScanFrequencyDays)
+	, IdleEvent(EEventMode::ManualReset)
+	, WaitEvent(EEventMode::AutoReset)
+	, Thread(
+		TEXT("FileSystemCacheStoreMaintainer"),
+		[this] { Loop(); },
+		[this] { Tick(); },
+		/*StackSize*/ 32 * 1024,
+		TPri_BelowNormal)
+{
+	IModularFeatures::Get().RegisterModularFeature(FeatureName, this);
+}
+
+FFileSystemCacheStoreMaintainer::~FFileSystemCacheStoreMaintainer()
+{
+	bExit = true;
+	IModularFeatures::Get().UnregisterModularFeature(FeatureName, this);
+	WaitEvent->Trigger();
+	Thread.Join();
+}
+
+void FFileSystemCacheStoreMaintainer::BoostPriority()
+{
+	bIgnoreFileScanRate = true;
+	WaitEvent->Trigger();
+}
+
+void FFileSystemCacheStoreMaintainer::Tick()
+{
+	// Scan once and exit if the priority has been boosted.
+	if (bIgnoreFileScanRate)
+	{
+		bExitAfterScan = true;
+		Loop();
+	}
+	bIdle = true;
+	IdleEvent.Trigger();
+}
+
+void FFileSystemCacheStoreMaintainer::Loop()
+{
+	WaitEvent->Wait(Params.TimeToWaitAfterInit, /*bIgnoreThreadIdleStats*/ true);
+
+	while (!bExit)
+	{
+		const FDateTime ScanStart = FDateTime::Now();
+		DeleteCount = 0;
+		DeleteSize = 0;
+		IdleEvent.Reset();
+		bIdle = false;
+		Scan();
+		bIdle = true;
+		IdleEvent.Trigger();
+		bIgnoreFileScanRate = false;
+		const FDateTime ScanEnd = FDateTime::Now();
+
+		UE_LOG(LogDerivedDataCache, Log,
+			TEXT("%s: Maintenance finished in %s and deleted %u file(s) with total size %" UINT64_FMT " MiB."),
+			*CachePath, *(ScanEnd - ScanStart).ToString(), DeleteCount, DeleteSize / 1024 / 1024);
+
+		if (bExit || bExitAfterScan)
+		{
+			break;
+		}
+
+		const FDateTime ScanTime = ScanStart + Params.ScanFrequency;
+		UE_CLOG(ScanEnd < ScanTime, LogDerivedDataCache, Verbose,
+			TEXT("%s: Maintenance is paused until the next scan at %s."), *CachePath, *ScanTime.ToString());
+		for (FDateTime Now = ScanEnd; !bExit && Now < ScanTime; Now = FDateTime::Now())
+		{
+			WaitEvent->Wait(ScanTime - Now, /*bIgnoreThreadIdleStats*/ true);
+		}
+	}
+
+	bIdle = true;
+	IdleEvent.Trigger();
+}
+
+void FFileSystemCacheStoreMaintainer::Scan()
+{
+	CreateContentRoot();
+	CreateBucketRoots();
+	CreateLegacyRoot();
+
+	while (!bExit)
+	{
+		const uint32 RootCount = uint32(Roots.Num());
+		const uint32 TotalRootCount = uint32(RootCount + LegacyRoot.IsValid());
+		if (TotalRootCount == 0)
+		{
+			break;
+		}
+		if (const uint32 RootIndex = Random.GetRandRange(0, TotalRootCount); RootIndex < RootCount)
+		{
+			ScanHashRoot(RootIndex);
+		}
+		else
+		{
+			ScanLegacyRoot();
+		}
+	}
+
+	ResetRoots();
+}
+
+void FFileSystemCacheStoreMaintainer::CreateContentRoot()
+{
+	TStringBuilder<256> ContentPath;
+	FPathViews::Append(ContentPath, CachePath, GContentDirectoryName);
+	if (FileManager.DirectoryExists(*ContentPath))
+	{
+		Roots.Add(MakeUnique<FRoot>(ContentPath, Random));
+	}
+}
+
+void FFileSystemCacheStoreMaintainer::CreateBucketRoots()
+{
+	TStringBuilder<256> BucketsPath;
+	FPathViews::Append(BucketsPath, CachePath, GBucketsDirectoryName);
+	FileManager.IterateDirectoryStat(*BucketsPath, [this](const TCHAR* Path, const FFileStatData& Stat) -> bool
+	{
+		if (Stat.bIsDirectory)
+		{
+			Roots.Add(MakeUnique<FRoot>(Path, Random));
+		}
+		return !bExit;
+	});
+}
+
+void FFileSystemCacheStoreMaintainer::ScanHashRoot(const uint32 RootIndex)
+{
+	FRoot& Root = *Roots[int32(RootIndex)];
+	const uint32 DirectoryIndex = Root.Order.GetNext();
+	const uint32 IndexLevel0 = DirectoryIndex / 256;
+	const uint32 IndexLevel1 = DirectoryIndex % 256;
+
+	bool bScanned = false;
+	ON_SCOPE_EXIT
+	{
+		if ((DirectoryIndex == Root.Order.GetFirst()) ||
+			(bScanned && ++Root.DirectoryScanCount >= Params.MaxDirectoryScanCount))
+		{
+			UE_LOG(LogDerivedDataCache, VeryVerbose,
+				TEXT("%s: Maintenance finished scanning %s."), *CachePath, *Root.Path);
+			Roots.RemoveAt(int32(RootIndex));
+		}
+	};
+
+	if (!Root.bScannedRoot)
+	{
+		Root.ExistsLevel0 = ScanHashDirectory(Root.Path);
+		Root.bScannedRoot = true;
+	}
+
+	if (!Root.ExistsLevel0[IndexLevel0])
+	{
+		return;
+	}
+	if (!Root.ScannedLevel0[IndexLevel0])
+	{
+		TStringBuilder<256> Path;
+		Path.Appendf(TEXT("%s/%02x"), *Root.Path, IndexLevel0);
+		Root.ExistsLevel1[IndexLevel0] = ScanHashDirectory(Path);
+		Root.ScannedLevel0[IndexLevel0] = true;
+	}
+
+	if (!Root.ExistsLevel1[IndexLevel0][IndexLevel1])
+	{
+		return;
+	}
+
+	TStringBuilder<256> Path;
+	Path.Appendf(TEXT("%s/%02x/%02x"), *Root.Path, IndexLevel0, IndexLevel1);
+	FileManager.IterateDirectoryStat(*Path, [this](const TCHAR* const Path, const FFileStatData& Stat) -> bool
+	{
+		ProcessFile(Path, Stat);
+		return !bExit;
+	});
+
+	bScanned = true;
+}
+
+TStaticBitArray<256> FFileSystemCacheStoreMaintainer::ScanHashDirectory(FStringBuilderBase& BasePath)
+{
+	TStaticBitArray<256> Exists;
+	FileManager.IterateDirectoryStat(*BasePath, [this, &Exists](const TCHAR* Path, const FFileStatData& Stat) -> bool
+	{
+		FStringView View = FPathViews::GetCleanFilename(Path);
+		if (Stat.bIsDirectory && View.Len() == 2 && Algo::AllOf(View, FChar::IsHexDigit))
+		{
+			uint8 Byte;
+			if (String::HexToBytes(View, &Byte) == 1)
+			{
+				Exists[Byte] = true;
+			}
+		}
+		return !bExit;
+	});
+	return Exists;
+}
+
+TStaticBitArray<10> FFileSystemCacheStoreMaintainer::ScanLegacyDirectory(FStringBuilderBase& BasePath)
+{
+	TStaticBitArray<10> Exists;
+	FileManager.IterateDirectoryStat(*BasePath, [this, &Exists](const TCHAR* Path, const FFileStatData& Stat) -> bool
+	{
+		FStringView View = FPathViews::GetCleanFilename(Path);
+		if (Stat.bIsDirectory && View.Len() == 1 && Algo::AllOf(View, FChar::IsDigit))
+		{
+			Exists[FChar::ConvertCharDigitToInt(View[0])] = true;
+		}
+		return !bExit;
+	});
+	return Exists;
+}
+
+void FFileSystemCacheStoreMaintainer::CreateLegacyRoot()
+{
+	TStringBuilder<256> Path;
+	FPathViews::Append(Path, CachePath);
+	TStaticBitArray<10> Exists = ScanLegacyDirectory(Path);
+	if (Exists.FindFirstSetBit() != INDEX_NONE)
+	{
+		LegacyRoot = MakeUnique<FLegacyRoot>(Random);
+		LegacyRoot->ExistsLevel0 = Exists;
+	}
+}
+
+void FFileSystemCacheStoreMaintainer::ScanLegacyRoot()
+{
+	FLegacyRoot& Root = *LegacyRoot;
+	const uint32 DirectoryIndex = Root.Order.GetNext();
+	const int32 IndexLevel0 = int32(DirectoryIndex / 100) % 10;
+	const int32 IndexLevel1 = int32(DirectoryIndex / 10) % 10;
+	const int32 IndexLevel2 = int32(DirectoryIndex / 1) % 10;
+
+	bool bScanned = false;
+	ON_SCOPE_EXIT
+	{
+		if ((DirectoryIndex == Root.Order.GetFirst()) ||
+			(bScanned && ++Root.DirectoryScanCount >= Params.MaxDirectoryScanCount))
+		{
+			LegacyRoot.Reset();
+		}
+	};
+
+	if (!Root.ExistsLevel0[IndexLevel0])
+	{
+		return;
+	}
+	if (!Root.ScannedLevel0[IndexLevel0])
+	{
+		TStringBuilder<256> Path;
+		FPathViews::Append(Path, CachePath, IndexLevel0);
+		Root.ExistsLevel1[IndexLevel0] = ScanLegacyDirectory(Path);
+		Root.ScannedLevel0[IndexLevel0] = true;
+	}
+
+	if (!Root.ExistsLevel1[IndexLevel0][IndexLevel1])
+	{
+		return;
+	}
+	if (!Root.ScannedLevel1[IndexLevel0][IndexLevel1])
+	{
+		TStringBuilder<256> Path;
+		FPathViews::Append(Path, CachePath, IndexLevel0, IndexLevel1);
+		Root.ExistsLevel2[IndexLevel0][IndexLevel1] = ScanLegacyDirectory(Path);
+		Root.ScannedLevel1[IndexLevel0][IndexLevel1] = true;
+	}
+
+	if (!Root.ExistsLevel2[IndexLevel0][IndexLevel1][IndexLevel2])
+	{
+		return;
+	}
+
+	TStringBuilder<256> Path;
+	FPathViews::Append(Path, CachePath, IndexLevel0, IndexLevel1, IndexLevel2);
+	FileManager.IterateDirectoryStat(*Path, [this](const TCHAR* const Path, const FFileStatData& Stat) -> bool
+	{
+		ProcessFile(Path, Stat);
+		return !bExit;
+	});
+
+	bScanned = true;
+}
+
+void FFileSystemCacheStoreMaintainer::ResetRoots()
+{
+	Roots.Empty();
+	LegacyRoot.Reset();
+}
+
+void FFileSystemCacheStoreMaintainer::ProcessFile(const TCHAR* const Path, const FFileStatData& Stat)
+{
+	if (Stat.bIsDirectory)
+	{
+		return;
+	}
+
+	const FDateTime Now = FDateTime::UtcNow();
+	if (Stat.ModificationTime + Params.MaxFileAge < Now && Stat.AccessTime + Params.MaxFileAge < Now)
+	{
+		++DeleteCount;
+		DeleteSize += Stat.FileSize > 0 ? uint64(Stat.FileSize) : 0;
+		if (FileManager.Delete(Path, /*bRequireExists*/ false, /*bEvenReadOnly*/ false, /*bQuiet*/ true))
+		{
+			UE_LOG(LogDerivedDataCache, VeryVerbose,
+				TEXT("%s: Maintenance deleted file %s that was last modified at %s."),
+				*CachePath, Path, *Stat.ModificationTime.ToIso8601());
+		}
+		else
+		{
+			UE_LOG(LogDerivedDataCache, Verbose,
+				TEXT("%s: Maintenance failed to delete file %s that was last modified at %s."),
+				*CachePath, Path, *Stat.ModificationTime.ToIso8601());
+		}
+	}
+
+	if (!bExit && !bIgnoreFileScanRate && Params.MaxFileScanRate && ++ProcessCount % Params.MaxFileScanRate == 0)
+	{
+		const double BatchEndTime = FPlatformTime::Seconds();
+		if (const double BatchWaitTime = 1.0 - (BatchEndTime - BatchStartTime); BatchWaitTime > 0.0)
+		{
+			WaitEvent->Wait(FTimespan::FromSeconds(BatchWaitTime), /*bIgnoreThreadIdleStats*/ true);
+			BatchStartTime = FPlatformTime::Seconds();
+		}
+		else
+		{
+			BatchStartTime = BatchEndTime;
+		}
+	}
+}
+
+///////////////////////////////////////////////////////////////////////////////////////////////////////////////////////
+
+class FAccessLogWriter
+{
+public:
+	FAccessLogWriter(const TCHAR* FileName, const FString& CachePath);
+
+	void Append(const TCHAR* CacheKey, FStringView Path);
+	void Append(const FIoHash& RawHash, FStringView Path);
+	void Append(const FCacheKey& CacheKey, FStringView Path);
+
+private:
+	void AppendPath(FStringView Path);
+
+	TUniquePtr<FArchive> Archive;
+	FString BasePath;
+	FCriticalSection CriticalSection;
+	TSet<FString> CacheKeys;
+	TSet<FIoHash> ContentKeys;
+	TSet<FCacheKey> RecordKeys;
+};
+
+FAccessLogWriter::FAccessLogWriter(const TCHAR* const FileName, const FString& CachePath)
+	: Archive(IFileManager::Get().CreateFileWriter(FileName, FILEWRITE_AllowRead))
+	, BasePath(CachePath / TEXT(""))
+{
+}
+
+void FAccessLogWriter::Append(const TCHAR* const CacheKey, const FStringView Path)
+{
+	FScopeLock Lock(&CriticalSection);
+
+	bool bIsAlreadyInSet = false;
+	CacheKeys.FindOrAdd(FString(CacheKey), &bIsAlreadyInSet);
+	if (!bIsAlreadyInSet)
+	{
+		AppendPath(Path);
+	}
+}
+
+void FAccessLogWriter::Append(const FIoHash& RawHash, const FStringView Path)
+{
+	FScopeLock Lock(&CriticalSection);
+
+	bool bIsAlreadyInSet = false;
+	ContentKeys.FindOrAdd(RawHash, &bIsAlreadyInSet);
+	if (!bIsAlreadyInSet)
+	{
+		AppendPath(Path);
+	}
+}
+
+void FAccessLogWriter::Append(const FCacheKey& CacheKey, const FStringView Path)
+{
+	FScopeLock Lock(&CriticalSection);
+
+	bool bIsAlreadyInSet = false;
+	RecordKeys.FindOrAdd(CacheKey, &bIsAlreadyInSet);
+	if (!bIsAlreadyInSet)
+	{
+		AppendPath(Path);
+	}
+}
+
+void FAccessLogWriter::AppendPath(const FStringView Path)
+{
+	if (Path.StartsWith(BasePath))
+	{
+		const FTCHARToUTF8 PathUtf8(Path.RightChop(BasePath.Len()));
+		Archive->Serialize(const_cast<ANSICHAR*>(PathUtf8.Get()), PathUtf8.Length());
+		Archive->Serialize(const_cast<ANSICHAR*>(LINE_TERMINATOR_ANSI), sizeof(LINE_TERMINATOR_ANSI) - 1);
+	}
+}
+
+///////////////////////////////////////////////////////////////////////////////////////////////////////////////////////
+
+class FFileSystemCacheStore final : public FDerivedDataBackendInterface
+{
+public:
+	FFileSystemCacheStore(const TCHAR* CachePath, const TCHAR* Params, const TCHAR* AccessLogPath);
+
+	inline bool IsUsable() const { return !bDisabled; }
+
+	bool RunSpeedTest(
+		double InSkipTestsIfSeeksExceedMS,
+		double& OutSeekTimeMS,
+		double& OutReadSpeedMBs,
+		double& OutWriteSpeedMBs) const;
+
+	// FDerivedDataBackendInterface Interface
+
+	FString GetName() const final { return CachePath; }
+	bool IsWritable() const final { return !bReadOnly && !bDisabled; }
+	ESpeedClass GetSpeedClass() const final { return SpeedClass; }
+
+	bool CachedDataProbablyExists(const TCHAR* CacheKey) final;
+	bool GetCachedData(const TCHAR* CacheKey, TArray<uint8>& Data) final;
+	EPutStatus PutCachedData(const TCHAR* CacheKey, TConstArrayView<uint8> Data, bool bPutEvenIfExists) final;
+	void RemoveCachedData(const TCHAR* CacheKey, bool bTransient) final;
+	TSharedRef<FDerivedDataCacheStatsNode> GatherUsageStats() const final;
+	TBitArray<> TryToPrefetch(TConstArrayView<FString> CacheKeys) final;
+	bool WouldCache(const TCHAR* CacheKey, TConstArrayView<uint8> InData) final;
+	bool ApplyDebugOptions(FBackendDebugOptions& InOptions) final;
+
+	EBackendLegacyMode GetLegacyMode() const final { return LegacyMode; }
+
+	// ICacheStore Interface
+
+	void Put(
+		TConstArrayView<FCachePutRequest> Requests,
+		IRequestOwner& Owner,
+		FOnCachePutComplete&& OnComplete) final;
+	void Get(
+		TConstArrayView<FCacheGetRequest> Requests,
+		IRequestOwner& Owner,
+		FOnCacheGetComplete&& OnComplete) final;
+	void PutValue(
+		TConstArrayView<FCachePutValueRequest> Requests,
+		IRequestOwner& Owner,
+		FOnCachePutValueComplete&& OnComplete) final;
+	void GetValue(
+		TConstArrayView<FCacheGetValueRequest> Requests,
+		IRequestOwner& Owner,
+		FOnCacheGetValueComplete&& OnComplete) final;
+	void GetChunks(
+		TConstArrayView<FCacheGetChunkRequest> Requests,
+		IRequestOwner& Owner,
+		FOnCacheGetChunkComplete&& OnComplete) final;
+
+private:
+	[[nodiscard]] bool PutCacheRecord(FStringView Name, const FCacheRecord& Record, const FCacheRecordPolicy& Policy, uint64& OutWriteSize);
+
+	[[nodiscard]] FOptionalCacheRecord GetCacheRecordOnly(
+		FStringView Name,
+		const FCacheKey& Key,
+		const FCacheRecordPolicy& Policy);
+	[[nodiscard]] FOptionalCacheRecord GetCacheRecord(
+		FStringView Name,
+		const FCacheKey& Key,
+		const FCacheRecordPolicy& Policy,
+		EStatus& OutStatus);
+
+	[[nodiscard]] bool PutCacheValue(FStringView Name, const FCacheKey& Key, const FValue& Value, ECachePolicy Policy, uint64& OutWriteSize);
+
+	[[nodiscard]] bool GetCacheValueOnly(FStringView Name, const FCacheKey& Key, ECachePolicy Policy, FValue& OutValue);
+	[[nodiscard]] bool GetCacheValue(FStringView Name, const FCacheKey& Key, ECachePolicy Policy, FValue& OutValue);
+
+	[[nodiscard]] bool PutCacheContent(FStringView Name, const FCompressedBuffer& Content, uint64& OutWriteSize) const;
+
+	[[nodiscard]] bool GetCacheContentExists(const FCacheKey& Key, const FIoHash& RawHash) const;
+	[[nodiscard]] bool GetCacheContent(
+		FStringView Name,
+		const FCacheKey& Key,
+		const FValueId& Id,
+		const FValue& Value,
+		ECachePolicy Policy,
+		FValue& OutValue) const;
+	void GetCacheContent(
+		FStringView Name,
+		const FCacheKey& Key,
+		const FValueId& Id,
+		const FValue& Value,
+		ECachePolicy Policy,
+		FCompressedBufferReader& Reader,
+		TUniquePtr<FArchive>& OutArchive) const;
+
+	void BuildCachePackagePath(const FCacheKey& CacheKey, FStringBuilderBase& Path) const;
+	void BuildCacheContentPath(const FIoHash& RawHash, FStringBuilderBase& Path) const;
+
+	[[nodiscard]] bool SaveFileWithHash(FStringBuilderBase& Path, FStringView DebugName, TFunctionRef<void (FArchive&)> WriteFunction, bool bReplaceExisting = false) const;
+	[[nodiscard]] bool LoadFileWithHash(FStringBuilderBase& Path, FStringView DebugName, TFunctionRef<void (FArchive&)> ReadFunction) const;
+	[[nodiscard]] bool SaveFile(FStringBuilderBase& Path, FStringView DebugName, TFunctionRef<void (FArchive&)> WriteFunction, bool bReplaceExisting = false) const;
+	[[nodiscard]] bool LoadFile(FStringBuilderBase& Path, FStringView DebugName, TFunctionRef<void (FArchive&)> ReadFunction) const;
+	[[nodiscard]] TUniquePtr<FArchive> OpenFile(FStringBuilderBase& Path, FStringView DebugName) const;
+
+	[[nodiscard]] bool FileExists(FStringBuilderBase& Path) const;
+
+	[[nodiscard]] bool ShouldSimulateMiss(const TCHAR* InKey);
+	[[nodiscard]] bool ShouldSimulateMiss(const FCacheKey& Key);
+
+private:
+	void BuildCacheLegacyPath(const TCHAR* const CacheKey, FStringBuilderBase& Path) const
+	{
+		Path << CachePath << TEXT('/');
+		BuildPathForLegacyCache(CacheKey, Path);
+	}
+
+	/** Base path we are storing the cache files in. */
+	FString	CachePath;
+	/** How to access the legacy cache. */
+	EBackendLegacyMode LegacyMode = EBackendLegacyMode::ValueWithLegacyFallback;
+	/** Class of this cache */
+	ESpeedClass SpeedClass;
+	/** If true, we failed to write to this directory and it did not contain anything so we should not be used. */
+	bool		bDisabled;
+	/** If true, do not attempt to write to this cache. */
+	bool		bReadOnly;
+	/** If true, CachedDataProbablyExists will update the file timestamps. */
+	bool		bTouch;
+	/** If true, allow transient data to be removed from the cache. */
+	bool		bPurgeTransient;
+	/** Age of file when it should be deleted from DDC cache. */
+	double		DaysToDeleteUnusedFiles;
+
+	/** Maximum total size of compressed data stored within a record package with multiple attachments. */
+	uint64		MaxRecordSizeKB = 256;
+	/** Maximum total size of compressed data stored within a value package, or a record package with one attachment. */
+	uint64		MaxValueSizeKB = 1024;
+
+	/** Object used for synchronization via a scoped lock. */
+	FCriticalSection SynchronizationObject;
+
+	// DDCNotification metrics
+
+	/** Map of cache keys to miss times for generating timing deltas. */
+	TMap<FString, double> DDCNotificationCacheTimes;
+
+	/** The total estimated build time accumulated from cache miss/put deltas. */
+	double TotalEstimatedBuildTime;
+
+	/** Access log to write to */
+	TUniquePtr<FAccessLogWriter> AccessLogWriter;
+
+	/** Debug Options */
+	FBackendDebugOptions DebugOptions;
+
+	/** Keys we ignored due to miss rate settings */
+	FCriticalSection MissedKeysCS;
+	TSet<FName> DebugMissedKeys;
+	TSet<FCacheKey> DebugMissedCacheKeys;
+
+	FDerivedDataCacheUsageStats UsageStats;
+
+	TUniquePtr<FFileSystemCacheStoreMaintainer> Maintainer;
+};
+
+FFileSystemCacheStore::FFileSystemCacheStore(
+	const TCHAR* const InCachePath,
+	const TCHAR* const InParams,
+	const TCHAR* const InAccessLogPath)
+	: CachePath(InCachePath)
+	, SpeedClass(ESpeedClass::Unknown)
+	, bDisabled(false)
+	, bReadOnly(false)
+	, bTouch(false)
+	, bPurgeTransient(false)
+	, DaysToDeleteUnusedFiles(15.0)
+	, TotalEstimatedBuildTime(0)
+{
+	// If we find a platform that has more stringent limits, this needs to be rethought.
+	checkf(MAX_BACKEND_KEY_LENGTH + MAX_CACHE_DIR_LEN + MAX_BACKEND_NUMBERED_SUBFOLDER_LENGTH + MAX_CACHE_EXTENTION_LEN < FPlatformMisc::GetMaxPathLength(),
+		TEXT("Not enough room left for cache keys in max path."));
+
+	check(CachePath.Len());
+	FPaths::NormalizeFilename(CachePath);
+
+	// Params that override our instance defaults
+	FParse::Bool(InParams, TEXT("ReadOnly="), bReadOnly);
+	FParse::Bool(InParams, TEXT("Touch="), bTouch);
+	FParse::Bool(InParams, TEXT("PurgeTransient="), bPurgeTransient);
+	FParse::Value(InParams, TEXT("UnusedFileAge="), DaysToDeleteUnusedFiles);
+	FParse::Value(InParams, TEXT("MaxRecordSizeKB="), MaxRecordSizeKB);
+	FParse::Value(InParams, TEXT("MaxValueSizeKB="), MaxValueSizeKB);
+
+	if (FString LegacyModeString; FParse::Value(InParams, TEXT("LegacyMode="), LegacyModeString))
+	{
+		if (!TryLexFromString(LegacyMode, LegacyModeString))
+		{
+			UE_LOG(LogDerivedDataCache, Warning, TEXT("%s: Ignoring unrecognized legacy mode '%s'"), *CachePath, *LegacyModeString);
+		}
+	}
+
+	// Flush the cache if requested.
+	bool bFlush = false;
+	if (!bReadOnly && FParse::Bool(InParams, TEXT("Flush="), bFlush) && bFlush)
+	{
+		IFileManager::Get().DeleteDirectory(*(CachePath / TEXT("")), /*bRequireExists*/ false, /*bTree*/ true);
+	}
+
+	// check latency and speed. Read values should always be valid
+	double ReadSpeedMBs = 0.0;
+	double WriteSpeedMBs = 0.0;
+	double SeekTimeMS = 0.0;
+
+	/* Speeds faster than this are considered local*/
+	const float ConsiderFastAtMS = 10;
+	/* Speeds faster than this are ok. Everything else is slow. This value can be overridden in the ini file */
+	float ConsiderSlowAtMS = 50;
+	FParse::Value(InParams, TEXT("ConsiderSlowAt="), ConsiderSlowAtMS);
+
+	// can skip the speed test so everything acts as local (e.g. 4.25 and earlier behavior).
+	bool SkipSpeedTest = !WITH_EDITOR || FParse::Param(FCommandLine::Get(), TEXT("DDCSkipSpeedTest"));
+	if (SkipSpeedTest)
+	{
+		ReadSpeedMBs = 999;
+		WriteSpeedMBs = 999;
+		SeekTimeMS = 0;
+		UE_LOG(LogDerivedDataCache, Log, TEXT("Skipping speed test to %s. Assuming local performance"), *CachePath);
+	}
+
+	if (!SkipSpeedTest && !RunSpeedTest(ConsiderSlowAtMS * 2, SeekTimeMS, ReadSpeedMBs, WriteSpeedMBs))
+	{
+		bDisabled = true;
+		UE_LOG(LogDerivedDataCache, Warning, TEXT("No read or write access to %s"), *CachePath);
+	}
+	else
+	{
+		bool bReadTestPassed = ReadSpeedMBs > 0.0;
+		bool bWriteTestPassed = WriteSpeedMBs > 0.0;
+
+		// if we failed writes mark this as read only
+		bReadOnly = bReadOnly || !bWriteTestPassed;
+
+		// classify and report on these times
+		if (SeekTimeMS < 1)
+		{
+			SpeedClass = ESpeedClass::Local;
+		}
+		else if (SeekTimeMS <= ConsiderFastAtMS)
+		{
+			SpeedClass = ESpeedClass::Fast;
+		}
+		else if (SeekTimeMS >= ConsiderSlowAtMS)
+		{
+			SpeedClass = ESpeedClass::Slow;
+		}
+		else
+		{
+			SpeedClass = ESpeedClass::Ok;
+		}
+
+		UE_LOG(LogDerivedDataCache, Display,
+			TEXT("Performance to %s: Latency=%.02fms. RandomReadSpeed=%.02fMBs, RandomWriteSpeed=%.02fMBs. ")
+			TEXT("Assigned SpeedClass '%s'"),
+			*CachePath, SeekTimeMS, ReadSpeedMBs, WriteSpeedMBs, LexToString(SpeedClass));
+
+		if (SpeedClass <= FDerivedDataBackendInterface::ESpeedClass::Slow && !bReadOnly)
+		{
+			if (GIsBuildMachine)
+			{
+				UE_LOG(LogDerivedDataCache, Display, TEXT("Access to %s appears to be slow. ")
+					TEXT("'Touch' will be disabled and queries/writes will be limited."), *CachePath);
+			}
+			else
+			{
+				UE_LOG(LogDerivedDataCache, Warning, TEXT("Access to %s appears to be slow. ")
+					TEXT("'Touch' will be disabled and queries/writes will be limited."), *CachePath);
+			}
+			bTouch = false;
+			//bReadOnly = true;
+		}
+
+		if (!bReadOnly)
+		{
+			if (FString(FCommandLine::Get()).Contains(TEXT("Run=DerivedDataCache")))
+			{
+				bTouch = true; // we always touch files when running the DDC commandlet
+			}
+
+			// The command line (-ddctouch) enables touch on all filesystem backends if specified.
+			bTouch = bTouch || FParse::Param(FCommandLine::Get(), TEXT("DDCTOUCH"));
+
+			if (bTouch)
+			{
+				UE_LOG(LogDerivedDataCache, Display, TEXT("Files in %s will be touched."), *CachePath);
+			}
+
+			bool bClean = false;
+			bool bDeleteUnused = true;
+			FParse::Bool(InParams, TEXT("Clean="), bClean);
+			FParse::Bool(InParams, TEXT("DeleteUnused="), bDeleteUnused);
+			bDeleteUnused = bDeleteUnused && !FParse::Param(FCommandLine::Get(), TEXT("NODDCCLEANUP"));
+
+			if (bClean || bDeleteUnused)
+			{
+				FFileSystemCacheStoreMaintainerParams MaintainerParams;
+				MaintainerParams.MaxFileAge = FTimespan::FromDays(DaysToDeleteUnusedFiles);
+				if (bDeleteUnused)
+				{
+					if (!FParse::Value(InParams, TEXT("MaxFileChecksPerSec="), MaintainerParams.MaxFileScanRate))
+					{
+						int32 MaxFileScanRate;
+						if (GConfig->GetInt(TEXT("DDCCleanup"), TEXT("MaxFileChecksPerSec"), MaxFileScanRate, GEngineIni))
+						{
+							MaintainerParams.MaxFileScanRate = uint32(MaxFileScanRate);
+						}
+					}
+					FParse::Value(InParams, TEXT("FoldersToClean="), MaintainerParams.MaxDirectoryScanCount);
+				}
+				else
+				{
+					MaintainerParams.ScanFrequency = FTimespan::MaxValue();
+				}
+				double TimeToWaitAfterInit;
+				if (bClean)
+				{
+					MaintainerParams.TimeToWaitAfterInit = FTimespan::Zero();
+				}
+				else if (GConfig->GetDouble(TEXT("DDCCleanup"), TEXT("TimeToWaitAfterInit"), TimeToWaitAfterInit, GEngineIni))
+				{
+					MaintainerParams.TimeToWaitAfterInit = FTimespan::FromSeconds(TimeToWaitAfterInit);
+				}
+
+				Maintainer = MakeUnique<FFileSystemCacheStoreMaintainer>(MaintainerParams, CachePath);
+
+				if (bClean)
+				{
+					Maintainer->BoostPriority();
+					Maintainer->WaitForIdle();
+				}
+			}
+		}
+
+		if (IsUsable() && InAccessLogPath != nullptr && *InAccessLogPath != 0)
+		{
+			AccessLogWriter.Reset(new FAccessLogWriter(InAccessLogPath, CachePath));
+		}
+	}
+}
+
+bool FFileSystemCacheStore::RunSpeedTest(
+	const double InSkipTestsIfSeeksExceedMS,
+	double& OutSeekTimeMS,
+	double& OutReadSpeedMBs,
+	double& OutWriteSpeedMBs) const
+{
+	SCOPED_BOOT_TIMING("RunSpeedTest");
+
+	//  files of increasing size. Most DDC data falls within this range so we don't want to skew by reading
+	// large amounts of data. Ultimately we care most about latency anyway.
+	const int FileSizes[] = { 4, 8, 16, 64, 128, 256 };
+	const int NumTestFolders = 2; //(0-9)
+	const int FileSizeCount = UE_ARRAY_COUNT(FileSizes);
+
+	bool bWriteTestPassed = true;
+	bool bReadTestPassed = true;
+	bool bTestDataExists = true;
+
+	double TotalSeekTime = 0;
+	double TotalReadTime = 0;
+	double TotalWriteTime = 0;
+	int TotalDataRead = 0;
+	int TotalDataWritten = 0;
+
+	const FString AbsoluteCachePath = IFileManager::Get().ConvertToAbsolutePathForExternalAppForRead(*CachePath);
+	if (AbsoluteCachePath.Len() > MAX_CACHE_DIR_LEN)
+	{
+		const FText ErrorMessage = FText::Format(NSLOCTEXT("DerivedDataCache", "PathTooLong", "Cache path {0} is longer than {1} characters... please adjust [DerivedDataBackendGraph] paths to be shorter (this leaves more room for cache keys)."),
+			FText::FromString(AbsoluteCachePath), FText::AsNumber(MAX_CACHE_DIR_LEN));
+		FMessageDialog::Open(EAppMsgType::Ok, ErrorMessage);
+		UE_LOG(LogDerivedDataCache, Fatal, TEXT("%s"), *ErrorMessage.ToString());
+	}
+
+	TArray<FString> Paths;
+	TArray<FString> MissingFiles;
+
+	MissingFiles.Reserve(NumTestFolders * FileSizeCount);
+
+	const FString TestDataPath = FPaths::Combine(CachePath, TEXT("TestData"));
+
+	// create an upfront map of paths to data size in bytes
+	// create the paths we'll use. <path>/0/TestData.dat, <path>/1/TestData.dat etc. If those files don't exist we'll
+	// create them which will likely give an invalid result when measuring them now but not in the future...
+	TMap<FString, int> TestFileEntries;
+	for (int iSize = 0; iSize < FileSizeCount; iSize++)
+	{
+		// make sure we dont stat/read/write to consecuting files in folders
+		for (int iFolder = 0; iFolder < NumTestFolders; iFolder++)
+		{
+			int FileSizeKB = FileSizes[iSize];
+			FString Path = FPaths::Combine(CachePath, TEXT("TestData"), *FString::FromInt(iFolder),
+				*FString::Printf(TEXT("TestData_%dkb.dat"), FileSizeKB));
+			TestFileEntries.Add(Path, FileSizeKB * 1024);
+		}
+	}
+
+	// measure latency by checking for the presence of all these files. We'll also track which don't exist..
+	const double StatStartTime = FPlatformTime::Seconds();
+ 	for (auto& KV : TestFileEntries)
+	{
+		FFileStatData StatData = IFileManager::Get().GetStatData(*KV.Key);
+
+		if (!StatData.bIsValid || StatData.FileSize != KV.Value)
+		{
+			MissingFiles.Add(KV.Key);
+		}
+	}
+
+	// save total stat time
+	TotalSeekTime = (FPlatformTime::Seconds() - StatStartTime);
+
+	// calculate seek time here
+	OutSeekTimeMS = (TotalSeekTime / TestFileEntries.Num()) * 1000;
+
+	UE_LOG(LogDerivedDataCache, Verbose, TEXT("Stat tests to %s took %.02f seconds"), *CachePath, TotalSeekTime);
+
+	// if seek times are very slow do a single read/write test just to confirm access
+	/*if (OutSeekTimeMS >= InSkipTestsIfSeeksExceedMS)
+	{
+		UE_LOG(LogDerivedDataCache, Warning,
+			TEXT("Limiting read/write speed tests due to seek times of %.02f exceeding %.02fms. ")
+			TEXT("Values will be inaccurate."), OutSeekTimeMS, InSkipTestsIfSeeksExceedMS);
+
+		FString Path = TestFileEntries.begin()->Key;
+		int Size = TestFileEntries.begin()->Value;
+
+		TestFileEntries.Reset();
+		TestFileEntries.Add(Path, Size);
+	}*/
+
+	// create any files that were missing
+	if (!bReadOnly)
+	{
+		TArray<uint8> Data;
+		for (auto& File : MissingFiles)
+		{
+			const int DesiredSize = TestFileEntries[File];
+			Data.SetNumUninitialized(DesiredSize);
+
+			if (!FFileHelper::SaveArrayToFile(Data, *File, &IFileManager::Get(), FILEWRITE_Silent))
+			{
+				// Handle the case where something else may have created the path at the same time.
+				// This is less about multiple users and more about things like SCW's / UnrealPak
+				// that can spin up multiple instances at once.
+				if (!IFileManager::Get().FileExists(*File))
+				{
+					uint32 ErrorCode = FPlatformMisc::GetLastError();
+					TCHAR ErrorBuffer[1024];
+					FPlatformMisc::GetSystemErrorMessage(ErrorBuffer, 1024, ErrorCode);
+					UE_LOG(LogDerivedDataCache, Warning,
+						TEXT("Fail to create %s, derived data cache to this directory will be read only. ")
+						TEXT("WriteError: %u (%s)"), *File, ErrorCode, ErrorBuffer);
+					bTestDataExists = false;
+					bWriteTestPassed = false;
+					break;
+				}
+			}
+		}
+	}
+
+	// now read all sizes from random folders
+	{
+		const int ArraySize = UE_ARRAY_COUNT(FileSizes);
+		TArray<uint8> TempData;
+		TempData.Empty(FileSizes[ArraySize - 1] * 1024);
+
+		const double ReadStartTime = FPlatformTime::Seconds();
+
+		for (auto& KV : TestFileEntries)
+		{
+			const int FileSize = KV.Value;
+			const FString& FilePath = KV.Key;
+
+			if (!FFileHelper::LoadFileToArray(TempData, *FilePath, FILEREAD_Silent))
+			{
+				uint32 ErrorCode = FPlatformMisc::GetLastError();
+				TCHAR ErrorBuffer[1024];
+				FPlatformMisc::GetSystemErrorMessage(ErrorBuffer, 1024, ErrorCode);
+				UE_LOG(LogDerivedDataCache, Warning,
+					TEXT("Fail to read from %s, derived data cache will be disabled. ReadError: %u (%s)"),
+					*FilePath, ErrorCode, ErrorBuffer);
+				bReadTestPassed = false;
+				break;
+			}
+
+			TotalDataRead += TempData.Num();
+		}
+
+		TotalReadTime = FPlatformTime::Seconds() - ReadStartTime;
+
+		UE_LOG(LogDerivedDataCache, Verbose, TEXT("Read tests %s on %s and took %.02f seconds"),
+			bReadTestPassed ? TEXT("passed") : TEXT("failed"), *CachePath, TotalReadTime);
+	}
+
+	// do write tests if or read tests passed and our seeks were below the cut-off
+	if (bReadTestPassed && !bReadOnly)
+	{
+		// do write tests but use a unique folder that is cleaned up afterwards
+		FString CustomPath = FPaths::Combine(CachePath, TEXT("TestData"), *FGuid::NewGuid().ToString());
+
+		const int ArraySize = UE_ARRAY_COUNT(FileSizes);
+		TArray<uint8> TempData;
+		TempData.Empty(FileSizes[ArraySize - 1] * 1024);
+
+		const double WriteStartTime = FPlatformTime::Seconds();
+
+		for (auto& KV : TestFileEntries)
+		{
+			const int FileSize = KV.Value;
+			FString FilePath = KV.Key;
+
+			TempData.SetNumUninitialized(FileSize);
+
+			FilePath = FilePath.Replace(*CachePath, *CustomPath);
+
+			if (!FFileHelper::SaveArrayToFile(TempData, *FilePath, &IFileManager::Get(), FILEWRITE_Silent))
+			{
+				uint32 ErrorCode = FPlatformMisc::GetLastError();
+				TCHAR ErrorBuffer[1024];
+				FPlatformMisc::GetSystemErrorMessage(ErrorBuffer, 1024, ErrorCode);
+				UE_LOG(LogDerivedDataCache, Warning,
+					TEXT("Fail to write to %s, derived data cache will be disabled. ReadError: %u (%s)"),
+					*FilePath, ErrorCode, ErrorBuffer);
+				bWriteTestPassed = false;
+				break;
+			}
+
+			TotalDataWritten += TempData.Num();
+		}
+
+		TotalWriteTime = FPlatformTime::Seconds() - WriteStartTime;
+
+		UE_LOG(LogDerivedDataCache, Verbose, TEXT("write tests %s on %s and took %.02f seconds"),
+			bWriteTestPassed ? TEXT("passed") : TEXT("failed"), *CachePath, TotalReadTime)
+
+		// remove the custom path but do it async as this can be slow on remote drives
+		AsyncTask(ENamedThreads::AnyThread, [CustomPath]() {
+			IFileManager::Get().DeleteDirectory(*CustomPath, false, true);
+		});
+
+		// check latency and speed. Read values should always be valid
+		const double ReadSpeedMBs = (bReadTestPassed ? (TotalDataRead / TotalReadTime) : 0) / (1024 * 1024);
+		const double WriteSpeedMBs = (bWriteTestPassed ? (TotalDataWritten / TotalWriteTime) : 0) / (1024 * 1024);
+		const double SeekTimeMS = (TotalSeekTime / TestFileEntries.Num()) * 1000;
+	}
+
+	const double TotalTestTime = FPlatformTime::Seconds() - StatStartTime;
+
+	UE_LOG(LogDerivedDataCache, Log, TEXT("Speed tests for %s took %.02f seconds"), *CachePath, TotalTestTime);
+
+	// check latency and speed. Read values should always be valid
+	OutReadSpeedMBs = (bReadTestPassed ? (TotalDataRead / TotalReadTime) : 0) / (1024 * 1024);
+	OutWriteSpeedMBs = (bWriteTestPassed ? (TotalDataWritten / TotalWriteTime) : 0) / (1024 * 1024);
+
+	return bWriteTestPassed || bReadTestPassed;
+}
+
+bool FFileSystemCacheStore::CachedDataProbablyExists(const TCHAR* const CacheKey)
+{
+	TRACE_CPUPROFILER_EVENT_SCOPE(FileSystemDDC_Exist);
+	TRACE_COUNTER_INCREMENT(FileSystemDDC_Exist);
+	COOK_STAT(auto Timer = UsageStats.TimeProbablyExists());
+
+	if (!IsUsable())
+	{
+		UE_LOG(LogDerivedDataCache, VeryVerbose,
+			TEXT("%s: Skipped exists check of %s because this cache store is not available"), *CachePath, CacheKey);
+		return false;
+	}
+
+	if (ShouldSimulateMiss(CacheKey))
+	{
+		UE_LOG(LogDerivedDataCache, Verbose, TEXT("%s: Simulated miss for exists check of %s"), *CachePath, CacheKey);
+		return false;
+	}
+
+	TStringBuilder<256> LegacyPath;
+	BuildCacheLegacyPath(CacheKey, LegacyPath);
+	const bool bExists = FileExists(LegacyPath);
+	if (bExists)
+	{
+		if (AccessLogWriter)
+		{
+			AccessLogWriter->Append(CacheKey, LegacyPath);
+		}
+
+		TRACE_COUNTER_INCREMENT(FileSystemDDC_ExistHit);
+		COOK_STAT(Timer.AddHit(0));
+	}
+	// If not using a shared cache, record a (probable) miss
+	else if (!GetDerivedDataCacheRef().GetUsingSharedDDC())
+	{
+		// store a cache miss
+		FScopeLock ScopeLock(&SynchronizationObject);
+		if (!DDCNotificationCacheTimes.Contains(CacheKey))
+		{
+			DDCNotificationCacheTimes.Add(CacheKey, FPlatformTime::Seconds());
+		}
+	}
+
+	UE_LOG(LogDerivedDataCache, Verbose,
+		TEXT("%s: CachedDataProbablyExists=%d for %s"), *CachePath, int32(bExists), CacheKey);
+
+	return bExists;
+}
+
+bool FFileSystemCacheStore::GetCachedData(const TCHAR* const CacheKey, TArray<uint8>& Data)
+{
+	TRACE_CPUPROFILER_EVENT_SCOPE(FileSystemDDC_Get);
+	TRACE_COUNTER_INCREMENT(FileSystemDDC_Get);
+	COOK_STAT(auto Timer = UsageStats.TimeGet());
+
+	const double StartTime = FPlatformTime::Seconds();
+
+	if (!IsUsable())
+	{
+		UE_LOG(LogDerivedDataCache, VeryVerbose,
+			TEXT("%s: Skipped get of %s because this cache store is not available"), *CachePath, CacheKey);
+		return false;
+	}
+
+	if (ShouldSimulateMiss(CacheKey))
+	{
+		UE_LOG(LogDerivedDataCache, Verbose, TEXT("%s: Simulated miss for get of %s"), *CachePath, CacheKey);
+		return false;
+	}
+
+	const auto ReadData = [this, CacheKey, &Data](FArchive& Ar)
+	{
+		const int64 Size = Ar.TotalSize();
+		if (Size > MAX_int32)
+		{
+			Ar.SetError();
+			return;
+		}
+
+		Data.Reset(int32(Size));
+		Data.AddUninitialized(int32(Size));
+		Ar.Serialize(Data.GetData(), Size);
+
+		if (!FCorruptionWrapper::ReadTrailer(Data, *CachePath, CacheKey))
+		{
+			Ar.SetError();
+		}
+	};
+
+	TStringBuilder<256> LegacyPath;
+	BuildCacheLegacyPath(CacheKey, LegacyPath);
+	if (LoadFile(LegacyPath, CacheKey, ReadData))
+	{
+		if (AccessLogWriter)
+		{
+			AccessLogWriter->Append(CacheKey, LegacyPath);
+		}
+
+		const double ReadDuration = FPlatformTime::Seconds() - StartTime;
+		const double ReadSpeed = ReadDuration > 0.001 ? (Data.Num() / ReadDuration) / (1024.0 * 1024.0) : 0.0;
+		UE_LOG(LogDerivedDataCache, Verbose,
+			TEXT("%s: Cache hit on %s (%d bytes, %.02f secs, %.2fMB/s)"),
+			*CachePath, CacheKey, Data.Num(), ReadDuration, ReadSpeed);
+		TRACE_COUNTER_INCREMENT(FileSystemDDC_GetHit);
+		TRACE_COUNTER_ADD(FileSystemDDC_BytesRead, int64(Data.Num()));
+		COOK_STAT(Timer.AddHit(Data.Num()));
+		return true;
+	}
+
+	UE_LOG(LogDerivedDataCache, Verbose, TEXT("%s: Cache miss on %s"), *CachePath, CacheKey);
+	Data.Empty();
+
+	// If not using a shared cache, record a miss
+	if (!GetDerivedDataCacheRef().GetUsingSharedDDC())
+	{
+		// store a cache miss
+		FScopeLock ScopeLock(&SynchronizationObject);
+		if (!DDCNotificationCacheTimes.Contains(CacheKey))
+		{
+			DDCNotificationCacheTimes.Add(CacheKey, FPlatformTime::Seconds());
+		}
+	}
+
+	return false;
+}
+
+bool FFileSystemCacheStore::WouldCache(const TCHAR* const CacheKey, const TConstArrayView<uint8> InData)
+{
+	return IsWritable() && !CachedDataProbablyExists(CacheKey);
+}
+
+FFileSystemCacheStore::EPutStatus FFileSystemCacheStore::PutCachedData(
+	const TCHAR* const CacheKey,
+	const TConstArrayView<uint8> Data,
+	const bool bPutEvenIfExists)
+{
+	TRACE_CPUPROFILER_EVENT_SCOPE(FileSystemDDC_Put);
+	TRACE_COUNTER_INCREMENT(FileSystemDDC_Put);
+	COOK_STAT(auto Timer = UsageStats.TimePut());
+
+	checkf(Data.Num(), TEXT("%s: Failed to put %s due to empty data"), *CachePath, CacheKey);
+
+	if (!IsUsable())
+	{
+		UE_LOG(LogDerivedDataCache, VeryVerbose,
+			TEXT("%s: Skipped put of %s because this cache store is not available"), *CachePath, CacheKey);
+		return EPutStatus::NotCached;
+	}
+
+	if (!IsWritable())
+	{
+		UE_LOG(LogDerivedDataCache, VeryVerbose,
+			TEXT("%s: Skipped put of %s because this cache store is read-only"), *CachePath, CacheKey);
+		return EPutStatus::NotCached;
+	}
+
+	if (ShouldSimulateMiss(CacheKey))
+	{
+		UE_LOG(LogDerivedDataCache, Verbose, TEXT("%s: Simulated miss for put of %s"), *CachePath, CacheKey);
+		return EPutStatus::Skipped;
+	}
+
+	TStringBuilder<256> LegacyPath;
+	BuildCacheLegacyPath(CacheKey, LegacyPath);
+
+	if (AccessLogWriter)
+	{
+		AccessLogWriter->Append(CacheKey, LegacyPath);
+	}
+
+	EPutStatus Status = EPutStatus::NotCached;
+
+	if (bPutEvenIfExists || !FileExists(LegacyPath))
+	{
+		const auto WriteData = [this, CacheKey, &Data](FArchive& Ar)
+		{
+			Ar.Serialize(const_cast<uint8*>(Data.GetData()), Data.Num());
+			FCorruptionWrapper::WriteTrailer(Ar, Data, *CachePath, CacheKey);
+		};
+
+		if (SaveFile(LegacyPath, CacheKey, WriteData))
+		{
+			UE_LOG(LogDerivedDataCache, Verbose,
+				TEXT("%s: Successful cache put of %s to %s"), *CachePath, CacheKey, *LegacyPath);
+			TRACE_COUNTER_INCREMENT(FileSystemDDC_PutHit);
+			TRACE_COUNTER_ADD(FileSystemDDC_BytesWritten, int64(Data.Num()));
+			COOK_STAT(Timer.AddHit(Data.Num()));
+			Status = EPutStatus::Cached;
+		}
+	}
+	else
+	{
+		UE_LOG(LogDerivedDataCache, Verbose, TEXT("%s: Skipping put of %s to existing file"), *CachePath, CacheKey);
+		Status = EPutStatus::Cached;
+	}
+
+	// If not using a shared cache, update estimated build time
+	if (!GetDerivedDataCacheRef().GetUsingSharedDDC())
+	{
+		FScopeLock ScopeLock(&SynchronizationObject);
+
+		if (DDCNotificationCacheTimes.Contains(CacheKey))
+		{
+			// There isn't any way to get exact build times in the DDC code as custom asset processing and async are factors.
+			// So, estimate the asset build time based on the delta between the cache miss and the put
+			TotalEstimatedBuildTime += (FPlatformTime::Seconds() - DDCNotificationCacheTimes[CacheKey]);
+			DDCNotificationCacheTimes.Remove(CacheKey);
+
+			// If more than 20 seconds has been spent building assets, send out a notification
+			if (TotalEstimatedBuildTime > 20.0f)
+			{
+				// Send out a DDC put notification if we have any subscribers
+				FDerivedDataCacheInterface::FOnDDCNotification& DDCNotificationEvent =
+					GetDerivedDataCacheRef().GetDDCNotificationEvent();
+
+				if (DDCNotificationEvent.IsBound())
+				{
+					TotalEstimatedBuildTime = 0.0f;
+
+					DECLARE_CYCLE_STAT(TEXT("FSimpleDelegateGraphTask.PutCachedData"),
+					STAT_FSimpleDelegateGraphTask_DDCNotification, STATGROUP_TaskGraphTasks);
+
+					FSimpleDelegateGraphTask::CreateAndDispatchWhenReady(
+						FSimpleDelegateGraphTask::FDelegate::CreateLambda([DDCNotificationEvent]() {
+						DDCNotificationEvent.Broadcast(FDerivedDataCacheInterface::SharedDDCPerformanceNotification);
+					}),
+						GET_STATID(STAT_FSimpleDelegateGraphTask_DDCNotification),
+						nullptr,
+						ENamedThreads::GameThread);
+				}
+			}
+		}
+	}
+
+	return Status;
+}
+
+void FFileSystemCacheStore::RemoveCachedData(const TCHAR* const CacheKey, const bool bTransient)
+{
+	check(IsUsable());
+	if (IsWritable() && (!bTransient || bPurgeTransient))
+	{
+		TStringBuilder<256> LegacyPath;
+		BuildCacheLegacyPath(CacheKey, LegacyPath);
+		if (bTransient)
+		{
+			UE_LOG(LogDerivedDataCache, Verbose,
+				TEXT("%s: Deleting cached data for %s from %s"), *CachePath, CacheKey, *LegacyPath);
+		}
+		IFileManager::Get().Delete(*LegacyPath, /*bRequireExists*/ false, /*bEvenReadOnly*/ false, /*bQuiet*/ true);
+	}
+}
+
+TSharedRef<FDerivedDataCacheStatsNode> FFileSystemCacheStore::GatherUsageStats() const
+{
+	TSharedRef<FDerivedDataCacheStatsNode> Usage =
+		MakeShared<FDerivedDataCacheStatsNode>(TEXT("File System"), *CachePath, SpeedClass == ESpeedClass::Local);
+	Usage->Stats.Add(TEXT(""), UsageStats);
+	return Usage;
+}
+
+TBitArray<> FFileSystemCacheStore::TryToPrefetch(const TConstArrayView<FString> CacheKeys)
+{
+	return CachedDataProbablyExistsBatch(CacheKeys);
+}
+
+bool FFileSystemCacheStore::ApplyDebugOptions(FBackendDebugOptions& InOptions)
+{
+	DebugOptions = InOptions;
+	return true;
+}
+
+void FFileSystemCacheStore::Put(
+	const TConstArrayView<FCachePutRequest> Requests,
+	IRequestOwner& Owner,
+	FOnCachePutComplete&& OnComplete)
+{
+	for (const FCachePutRequest& Request : Requests)
+	{
+		const FCacheRecord& Record = Request.Record;
+		TRACE_CPUPROFILER_EVENT_SCOPE(FileSystemDDC_Put);
+		TRACE_COUNTER_INCREMENT(FileSystemDDC_Put);
+		COOK_STAT(auto Timer = UsageStats.TimePut());
+		uint64 WriteSize = 0;
+		if (PutCacheRecord(Request.Name, Record, Request.Policy, WriteSize))
+		{
+			UE_LOG(LogDerivedDataCache, Verbose, TEXT("%s: Cache put complete for %s from '%s'"),
+				*CachePath, *WriteToString<96>(Record.GetKey()), *Request.Name);
+			TRACE_COUNTER_INCREMENT(FileSystemDDC_PutHit);
+			TRACE_COUNTER_ADD(FileSystemDDC_BytesWritten, WriteSize);
+			COOK_STAT(if (WriteSize) { Timer.AddHit(WriteSize); });
+			OnComplete(Request.MakeResponse(EStatus::Ok));
+		}
+		else
+		{
+			OnComplete(Request.MakeResponse(EStatus::Error));
+		}
+	}
+}
+
+void FFileSystemCacheStore::Get(
+	const TConstArrayView<FCacheGetRequest> Requests,
+	IRequestOwner& Owner,
+	FOnCacheGetComplete&& OnComplete)
+{
+	for (const FCacheGetRequest& Request : Requests)
+	{
+		TRACE_CPUPROFILER_EVENT_SCOPE(FileSystemDDC_Get);
+		TRACE_COUNTER_INCREMENT(FileSystemDDC_Get);
+		COOK_STAT(auto Timer = UsageStats.TimeGet());
+		EStatus Status = EStatus::Ok;
+		if (FOptionalCacheRecord Record = GetCacheRecord(Request.Name, Request.Key, Request.Policy, Status))
+		{
+			UE_LOG(LogDerivedDataCache, Verbose, TEXT("%s: Cache hit for %s from '%s'"),
+				*CachePath, *WriteToString<96>(Request.Key), *Request.Name);
+			TRACE_COUNTER_INCREMENT(FileSystemDDC_GetHit);
+			TRACE_COUNTER_ADD(FileSystemDDC_BytesRead, Private::GetCacheRecordCompressedSize(Record.Get()));
+			COOK_STAT(Timer.AddHit(Private::GetCacheRecordCompressedSize(Record.Get())));
+			OnComplete({Request.Name, MoveTemp(Record).Get(), Request.UserData, Status});
+		}
+		else
+		{
+			OnComplete(Request.MakeResponse(Status));
+		}
+	}
+}
+
+void FFileSystemCacheStore::PutValue(
+	const TConstArrayView<FCachePutValueRequest> Requests,
+	IRequestOwner& Owner,
+	FOnCachePutValueComplete&& OnComplete)
+{
+	for (const FCachePutValueRequest& Request : Requests)
+	{
+		TRACE_CPUPROFILER_EVENT_SCOPE(FileSystemDDC_Put);
+		TRACE_COUNTER_INCREMENT(FileSystemDDC_Put);
+		COOK_STAT(auto Timer = UsageStats.TimePut());
+		uint64 WriteSize = 0;
+		if (PutCacheValue(Request.Name, Request.Key, Request.Value, Request.Policy, WriteSize))
+		{
+			UE_LOG(LogDerivedDataCache, Verbose, TEXT("%s: Cache put complete for %s from '%s'"),
+				*CachePath, *WriteToString<96>(Request.Key), *Request.Name);
+			TRACE_COUNTER_INCREMENT(FileSystemDDC_PutHit);
+			TRACE_COUNTER_ADD(FileSystemDDC_BytesWritten, WriteSize);
+			COOK_STAT(if (WriteSize) { Timer.AddHit(WriteSize); });
+			OnComplete(Request.MakeResponse(EStatus::Ok));
+		}
+		else
+		{
+			OnComplete(Request.MakeResponse(EStatus::Ok));
+		}
+	}
+}
+
+void FFileSystemCacheStore::GetValue(
+	const TConstArrayView<FCacheGetValueRequest> Requests,
+	IRequestOwner& Owner,
+	FOnCacheGetValueComplete&& OnComplete)
+{
+	for (const FCacheGetValueRequest& Request : Requests)
+	{
+		TRACE_CPUPROFILER_EVENT_SCOPE(FileSystemDDC_Get);
+		TRACE_COUNTER_INCREMENT(FileSystemDDC_Get);
+		COOK_STAT(auto Timer = UsageStats.TimeGet());
+		FValue Value;
+		if (GetCacheValue(Request.Name, Request.Key, Request.Policy, Value))
+		{
+			UE_LOG(LogDerivedDataCache, Verbose, TEXT("%s: Cache hit for %s from '%s'"),
+				*CachePath, *WriteToString<96>(Request.Key), *Request.Name);
+			TRACE_COUNTER_INCREMENT(FileSystemDDC_GetHit);
+			TRACE_COUNTER_ADD(FileSystemDDC_BytesRead, Value.GetData().GetCompressedSize());
+			COOK_STAT(Timer.AddHit(Value.GetData().GetCompressedSize()));
+			OnComplete({Request.Name, Request.Key, Value, Request.UserData, EStatus::Ok});
+		}
+		else
+		{
+			OnComplete(Request.MakeResponse(EStatus::Error));
+		}
+	}
+}
+
+void FFileSystemCacheStore::GetChunks(
+	const TConstArrayView<FCacheGetChunkRequest> Requests,
+	IRequestOwner& Owner,
+	FOnCacheGetChunkComplete&& OnComplete)
+{
+	TArray<FCacheGetChunkRequest, TInlineAllocator<16>> SortedRequests(Requests);
+	SortedRequests.StableSort(TChunkLess());
+
+	bool bHasValue = false;
+	FValue Value;
+	FValueId ValueId;
+	FCacheKey ValueKey;
+	TUniquePtr<FArchive> ValueAr;
+	FCompressedBufferReader ValueReader;
+	FOptionalCacheRecord Record;
+	for (const FCacheGetChunkRequest& Request : SortedRequests)
+	{
+		const bool bExistsOnly = EnumHasAnyFlags(Request.Policy, ECachePolicy::SkipData);
+		TRACE_CPUPROFILER_EVENT_SCOPE(FileSystemDDC_Get);
+		TRACE_COUNTER_INCREMENT(FileSystemDDC_Get);
+		COOK_STAT(auto Timer = bExistsOnly ? UsageStats.TimeProbablyExists() : UsageStats.TimeGet());
+		if (!(bHasValue && ValueKey == Request.Key && ValueId == Request.Id) || ValueReader.HasSource() < !bExistsOnly)
+		{
+			ValueReader.ResetSource();
+			ValueAr.Reset();
+			ValueKey = {};
+			ValueId.Reset();
+			Value.Reset();
+			bHasValue = false;
+			if (Request.Id.IsValid())
+			{
+				if (!(Record && Record.Get().GetKey() == Request.Key))
+				{
+					FCacheRecordPolicyBuilder PolicyBuilder(ECachePolicy::None);
+					PolicyBuilder.AddValuePolicy(Request.Id, Request.Policy);
+					Record.Reset();
+					Record = GetCacheRecordOnly(Request.Name, Request.Key, PolicyBuilder.Build());
+				}
+				if (Record)
+				{
+					if (const FValueWithId& ValueWithId = Record.Get().GetValue(Request.Id))
+					{
+						bHasValue = true;
+						Value = ValueWithId;
+						ValueId = Request.Id;
+						ValueKey = Request.Key;
+						GetCacheContent(Request.Name, Request.Key, ValueId, Value, Request.Policy, ValueReader, ValueAr);
+					}
+				}
+			}
+			else
+			{
+				ValueKey = Request.Key;
+				bHasValue = GetCacheValueOnly(Request.Name, Request.Key, Request.Policy, Value);
+				if (bHasValue)
+				{
+					GetCacheContent(Request.Name, Request.Key, Request.Id, Value, Request.Policy, ValueReader, ValueAr);
+				}
+			}
+		}
+		if (bHasValue)
+		{
+			const uint64 RawOffset = FMath::Min(Value.GetRawSize(), Request.RawOffset);
+			const uint64 RawSize = FMath::Min(Value.GetRawSize() - RawOffset, Request.RawSize);
+			UE_LOG(LogDerivedDataCache, Verbose, TEXT("%s: Cache hit for %s from '%s'"),
+				*CachePath, *WriteToString<96>(Request.Key, '/', Request.Id), *Request.Name);
+			TRACE_COUNTER_INCREMENT(FileSystemDDC_GetHit);
+			TRACE_COUNTER_ADD(FileSystemDDC_BytesRead, !bExistsOnly ? RawSize : 0);
+			COOK_STAT(Timer.AddHit(!bExistsOnly ? RawSize : 0));
+			FSharedBuffer Buffer;
+			if (!bExistsOnly)
+			{
+				Buffer = ValueReader.Decompress(RawOffset, RawSize);
+			}
+			const EStatus ChunkStatus = bExistsOnly || Buffer.GetSize() == RawSize ? EStatus::Ok : EStatus::Error;
+			OnComplete({Request.Name, Request.Key, Request.Id, Request.RawOffset,
+				RawSize, Value.GetRawHash(), MoveTemp(Buffer), Request.UserData, ChunkStatus});
+			continue;
+		}
+
+		OnComplete(Request.MakeResponse(EStatus::Error));
+	}
+}
+
+bool FFileSystemCacheStore::PutCacheRecord(
+	const FStringView Name,
+	const FCacheRecord& Record,
+	const FCacheRecordPolicy& Policy,
+	uint64& OutWriteSize)
+{
+	if (!IsWritable())
+	{
+		UE_LOG(LogDerivedDataCache, VeryVerbose,
+			TEXT("%s: Skipped put of %s from '%.*s' because this cache store is read-only"),
+			*CachePath, *WriteToString<96>(Record.GetKey()), Name.Len(), Name.GetData());
+		return false;
+	}
+
+	const FCacheKey& Key = Record.GetKey();
+	const ECachePolicy RecordPolicy = Policy.GetRecordPolicy();
+
+	// Skip the request if storing to the cache is disabled.
+	const ECachePolicy StoreFlag = SpeedClass == ESpeedClass::Local ? ECachePolicy::StoreLocal : ECachePolicy::StoreRemote;
+	if (!EnumHasAnyFlags(RecordPolicy, StoreFlag))
+	{
+		UE_LOG(LogDerivedDataCache, VeryVerbose, TEXT("%s: Skipped put of %s from '%.*s' due to cache policy"),
+			*CachePath, *WriteToString<96>(Key), Name.Len(), Name.GetData());
+		return false;
+	}
+
+	if (ShouldSimulateMiss(Key))
+	{
+		UE_LOG(LogDerivedDataCache, Verbose, TEXT("%s: Simulated miss for put of %s from '%.*s'"),
+			*CachePath, *WriteToString<96>(Key), Name.Len(), Name.GetData());
+		return false;
+	}
+
+	TStringBuilder<256> Path;
+	BuildCachePackagePath(Key, Path);
+
+	// Check if there is an existing record package.
+	FCbPackage ExistingPackage;
+	const ECachePolicy QueryFlag = SpeedClass == ESpeedClass::Local ? ECachePolicy::QueryLocal : ECachePolicy::QueryRemote;
+	bool bReplaceExisting = !EnumHasAnyFlags(RecordPolicy, QueryFlag);
+	bool bSaveRecord = bReplaceExisting;
+	if (const bool bLoadRecord = !bReplaceExisting || !Algo::AllOf(Record.GetValues(), &FValue::HasData))
+	{
+		// Load the existing package to take its inline attachments into account.
+		bSaveRecord |= !LoadFileWithHash(Path, Name, [&ExistingPackage](FArchive& Ar) { ExistingPackage.TryLoad(Ar); });
+		if (!bSaveRecord)
+		{
+			// Load the existing record to detect determinism issues.
+			const FOptionalCacheRecord ExistingRecord = FCacheRecord::Load(ExistingPackage);
+			bSaveRecord |= !ExistingRecord;
+			const auto MakeValueTuple = [](const FValueWithId& Value) -> TTuple<FValueId, FIoHash>
+			{
+				return MakeTuple(Value.GetId(), Value.GetRawHash());
+			};
+			if (ExistingRecord && !Algo::CompareBy(ExistingRecord.Get().GetValues(), Record.GetValues(), MakeValueTuple))
+			{
+				// Values differ between the existing record and the new record.
+				// Overwrite the existing record if the cache is missing content for any of its values.
+				UE_LOG(LogDerivedDataCache, Display, TEXT("%s: Cache put found non-deterministic record for %s from '%.*s'"),
+					*CachePath, *WriteToString<96>(Key), Name.Len(), Name.GetData());
+				const auto HasValueContent = [this, Name, &Key](const FValueWithId& Value) -> bool
+				{
+					if (!Value.HasData() && !GetCacheContentExists(Key, Value.GetRawHash()))
+					{
+						UE_LOG(LogDerivedDataCache, Log,
+							TEXT("%s: Cache put of non-deterministic record will overwrite existing record due to ")
+							TEXT("missing value %s with hash %s for %s from '%.*s'"),
+							*CachePath, *WriteToString<16>(Value.GetId()), *WriteToString<48>(Value.GetRawHash()),
+							*WriteToString<96>(Key), Name.Len(), Name.GetData());
+						return false;
+					}
+					return true;
+				};
+				bSaveRecord |= !Algo::AllOf(ExistingRecord.Get().GetValues(), HasValueContent);
+				bReplaceExisting |= bSaveRecord;
+			}
+		}
+	}
+
+	// Serialize the record to a package and remove attachments that will be stored externally.
+	FCbPackage Package = Record.Save();
+	TArray<FCompressedBuffer, TInlineAllocator<8>> ExternalContent;
+	if (ExistingPackage && !bSaveRecord)
+	{
+		// Mirror the existing internal/external attachment storage.
+		TArray<FCompressedBuffer, TInlineAllocator<8>> AllContent;
+		Algo::Transform(Package.GetAttachments(), AllContent, &FCbAttachment::AsCompressedBinary);
+		for (FCompressedBuffer& Content : AllContent)
+		{
+			const FIoHash RawHash = Content.GetRawHash();
+			if (!ExistingPackage.FindAttachment(RawHash))
+			{
+				Package.RemoveAttachment(RawHash);
+				ExternalContent.Add(MoveTemp(Content));
+			}
+		}
+	}
+	else
+	{
+		// Attempt to copy missing attachments from the existing package.
+		if (ExistingPackage)
+		{
+			for (const FValue& Value : Record.GetValues())
+			{
+				if (!Value.HasData())
+				{
+					if (const FCbAttachment* Attachment = ExistingPackage.FindAttachment(Value.GetRawHash()))
+					{
+						Package.AddAttachment(*Attachment);
+					}
+				}
+			}
+		}
+
+		// Remove the largest attachments from the package until it fits within the size limits.
+		TArray<FCompressedBuffer, TInlineAllocator<8>> AllContent;
+		Algo::Transform(Package.GetAttachments(), AllContent, &FCbAttachment::AsCompressedBinary);
+		uint64 TotalSize = Algo::TransformAccumulate(AllContent, &FCompressedBuffer::GetCompressedSize, uint64(0));
+		const uint64 MaxSize = (Record.GetValues().Num() == 1 ? MaxValueSizeKB : MaxRecordSizeKB) * 1024;
+		if (TotalSize > MaxSize)
+		{
+			Algo::StableSortBy(AllContent, &FCompressedBuffer::GetCompressedSize, TGreater<>());
+			for (FCompressedBuffer& Content : AllContent)
+			{
+				const uint64 CompressedSize = Content.GetCompressedSize();
+				Package.RemoveAttachment(Content.GetRawHash());
+				ExternalContent.Add(MoveTemp(Content));
+				TotalSize -= CompressedSize;
+				if (TotalSize <= MaxSize)
+				{
+					break;
+				}
+			}
+		}
+	}
+
+	// Save the external content to storage.
+	for (FCompressedBuffer& Content : ExternalContent)
+	{
+		uint64 WriteSize = 0;
+		if (!PutCacheContent(Name, Content, OutWriteSize))
+		{
+			return false;
+		}
+		OutWriteSize += WriteSize;
+	}
+
+	// Save the record package to storage.
+	const auto WriteRecord = [&](FArchive& Ar) { Package.Save(Ar); OutWriteSize += uint64(Ar.TotalSize()); };
+	if (bSaveRecord && !SaveFileWithHash(Path, Name, WriteRecord, bReplaceExisting))
+	{
+		return false;
+	}
+
+	if (AccessLogWriter)
+	{
+		AccessLogWriter->Append(Key, Path);
+	}
+
+	return true;
+}
+
+FOptionalCacheRecord FFileSystemCacheStore::GetCacheRecordOnly(
+	const FStringView Name,
+	const FCacheKey& Key,
+	const FCacheRecordPolicy& Policy)
+{
+	if (!IsUsable())
+	{
+		UE_LOG(LogDerivedDataCache, VeryVerbose,
+			TEXT("%s: Skipped get of %s from '%.*s' because this cache store is not available"),
+			*CachePath, *WriteToString<96>(Key), Name.Len(), Name.GetData());
+		return FOptionalCacheRecord();
+	}
+
+	// Skip the request if querying the cache is disabled.
+	const ECachePolicy QueryFlag = SpeedClass == ESpeedClass::Local ? ECachePolicy::QueryLocal : ECachePolicy::QueryRemote;
+	if (!EnumHasAnyFlags(Policy.GetRecordPolicy(), QueryFlag))
+	{
+		UE_LOG(LogDerivedDataCache, VeryVerbose, TEXT("%s: Skipped get of %s from '%.*s' due to cache policy"),
+			*CachePath, *WriteToString<96>(Key), Name.Len(), Name.GetData());
+		return FOptionalCacheRecord();
+	}
+
+	if (ShouldSimulateMiss(Key))
+	{
+		UE_LOG(LogDerivedDataCache, Verbose, TEXT("%s: Simulated miss for get of %s from '%.*s'"),
+			*CachePath, *WriteToString<96>(Key), Name.Len(), Name.GetData());
+		return FOptionalCacheRecord();
+	}
+
+	TStringBuilder<256> Path;
+	BuildCachePackagePath(Key, Path);
+
+	bool bDeletePackage = true;
+	ON_SCOPE_EXIT
+	{
+		if (bDeletePackage && !bReadOnly)
+		{
+			IFileManager::Get().Delete(*Path, /*bRequireExists*/ false, /*bEvenReadOnly*/ false, /*bQuiet*/ true);
+		}
+	};
+
+	FOptionalCacheRecord Record;
+	{
+		FCbPackage Package;
+		if (!LoadFileWithHash(Path, Name, [&Package](FArchive& Ar) { Package.TryLoad(Ar); }))
+		{
+			UE_LOG(LogDerivedDataCache, Verbose, TEXT("%s: Cache miss with missing package for %s from '%.*s'"),
+				*CachePath, *WriteToString<96>(Key), Name.Len(), Name.GetData());
+			return Record;
+		}
+
+		if (ValidateCompactBinary(Package, ECbValidateMode::Default | ECbValidateMode::Package) != ECbValidateError::None)
+		{
+			UE_LOG(LogDerivedDataCache, Display, TEXT("%s: Cache miss with invalid package for %s from '%.*s'"),
+				*CachePath, *WriteToString<96>(Key), Name.Len(), Name.GetData());
+			return Record;
+		}
+
+		Record = FCacheRecord::Load(Package);
+		if (Record.IsNull())
+		{
+			UE_LOG(LogDerivedDataCache, Display, TEXT("%s: Cache miss with record load failure for %s from '%.*s'"),
+				*CachePath, *WriteToString<96>(Key), Name.Len(), Name.GetData());
+			return Record;
+		}
+	}
+
+	if (AccessLogWriter)
+	{
+		AccessLogWriter->Append(Key, Path);
+	}
+
+	bDeletePackage = false;
+	return Record;
+}
+
+FOptionalCacheRecord FFileSystemCacheStore::GetCacheRecord(
+	const FStringView Name,
+	const FCacheKey& Key,
+	const FCacheRecordPolicy& Policy,
+	EStatus& OutStatus)
+{
+	FOptionalCacheRecord Record = GetCacheRecordOnly(Name, Key, Policy);
+	if (Record.IsNull())
+	{
+		OutStatus = EStatus::Error;
+		return Record;
+	}
+
+	OutStatus = EStatus::Ok;
+
+	FCacheRecordBuilder RecordBuilder(Key);
+
+	const ECachePolicy RecordPolicy = Policy.GetRecordPolicy();
+	if (!EnumHasAnyFlags(RecordPolicy, ECachePolicy::SkipMeta))
+	{
+		RecordBuilder.SetMeta(FCbObject(Record.Get().GetMeta()));
+	}
+
+	for (const FValueWithId& Value : Record.Get().GetValues())
+	{
+		const FValueId& Id = Value.GetId();
+		const ECachePolicy ValuePolicy = Policy.GetValuePolicy(Id);
+		FValue Content;
+		if (GetCacheContent(Name, Key, Id, Value, ValuePolicy, Content))
+		{
+			RecordBuilder.AddValue(Id, MoveTemp(Content));
+		}
+		else if (EnumHasAnyFlags(RecordPolicy, ECachePolicy::PartialRecord))
+		{
+			OutStatus = EStatus::Error;
+			RecordBuilder.AddValue(Value);
+		}
+		else
+		{
+			OutStatus = EStatus::Error;
+			return FOptionalCacheRecord();
+		}
+	}
+
+	return RecordBuilder.Build();
+}
+
+bool FFileSystemCacheStore::PutCacheValue(
+	const FStringView Name,
+	const FCacheKey& Key,
+	const FValue& Value,
+	const ECachePolicy Policy,
+	uint64& OutWriteSize)
+{
+	if (!IsWritable())
+	{
+		UE_LOG(LogDerivedDataCache, VeryVerbose,
+			TEXT("%s: Skipped put of %s from '%.*s' because this cache store is read-only"),
+			*CachePath, *WriteToString<96>(Key), Name.Len(), Name.GetData());
+		return false;
+	}
+
+	// Skip the request if storing to the cache is disabled.
+	const ECachePolicy StoreFlag = SpeedClass == ESpeedClass::Local ? ECachePolicy::StoreLocal : ECachePolicy::StoreRemote;
+	if (!EnumHasAnyFlags(Policy, StoreFlag))
+	{
+		UE_LOG(LogDerivedDataCache, VeryVerbose, TEXT("%s: Skipped put of %s from '%.*s' due to cache policy"),
+			*CachePath, *WriteToString<96>(Key), Name.Len(), Name.GetData());
+		return false;
+	}
+
+	if (ShouldSimulateMiss(Key))
+	{
+		UE_LOG(LogDerivedDataCache, Verbose, TEXT("%s: Simulated miss for put of %s from '%.*s'"),
+			*CachePath, *WriteToString<96>(Key), Name.Len(), Name.GetData());
+		return false;
+	}
+
+	// Check if there is an existing value package.
+	bool bValueExists = false;
+	FCbPackage ExistingPackage;
+	TStringBuilder<256> Path;
+	BuildCachePackagePath(Key, Path);
+	const ECachePolicy QueryFlag = SpeedClass == ESpeedClass::Local ? ECachePolicy::QueryLocal : ECachePolicy::QueryRemote;
+	const bool bReplaceExisting = !EnumHasAnyFlags(Policy, QueryFlag);
+	if (!bReplaceExisting)
+	{
+		bValueExists = LoadFileWithHash(Path, Name, [&ExistingPackage](FArchive& Ar) { ExistingPackage.TryLoad(Ar); })
+			&& (ExistingPackage.FindAttachment(Value.GetRawHash()) || GetCacheContentExists(Key, Value.GetRawHash()));
+	}
+
+	// Save the value to a package and save the data to external content depending on its size.
+	if (!bValueExists)
+	{
+		FCbWriter Writer;
+		Writer.BeginObject();
+		Writer.AddBinaryAttachment("RawHash", Value.GetRawHash());
+		Writer.AddInteger("RawSize", Value.GetRawSize());
+		Writer.EndObject();
+
+		FCbPackage Package(Writer.Save().AsObject());
+		if (!Value.HasData())
+		{
+			// Verify that the content exists in storage.
+			if (!GetCacheContentExists(Key, Value.GetRawHash()))
+			{
+				UE_LOG(LogDerivedDataCache, Verbose, TEXT("%s: Failed due to missing data for put of %s from '%.*s'"),
+					*CachePath, *WriteToString<96>(Key), Name.Len(), Name.GetData());
+				return false;
+			}
+		}
+		else if (Value.GetData().GetCompressedSize() <= MaxValueSizeKB * 1024)
+		{
+			// Store the content in the package.
+			Package.AddAttachment(FCbAttachment(Value.GetData()));
+		}
+		else
+		{
+			// Save the external content to storage.
+			uint64 WriteSize = 0;
+			if (!PutCacheContent(Name, Value.GetData(), WriteSize))
+			{
+				return false;
+			}
+			OutWriteSize += WriteSize;
+		}
+
+		// Save the value package to storage.
+		const auto WritePackage = [&](FArchive& Ar) { Package.Save(Ar); OutWriteSize += uint64(Ar.TotalSize()); };
+		if (!SaveFileWithHash(Path, Name, WritePackage, bReplaceExisting))
+		{
+			return false;
+		}
+	}
+
+	if (AccessLogWriter)
+	{
+		AccessLogWriter->Append(Key, Path);
+	}
+
+	return true;
+}
+
+bool FFileSystemCacheStore::GetCacheValueOnly(
+	const FStringView Name,
+	const FCacheKey& Key,
+	const ECachePolicy Policy,
+	FValue& OutValue)
+{
+	if (!IsUsable())
+	{
+		UE_LOG(LogDerivedDataCache, VeryVerbose,
+			TEXT("%s: Skipped get of %s from '%.*s' because this cache store is not available"),
+			*CachePath, *WriteToString<96>(Key), Name.Len(), Name.GetData());
+		return false;
+	}
+
+	// Skip the request if querying the cache is disabled.
+	const ECachePolicy QueryFlag = SpeedClass == ESpeedClass::Local ? ECachePolicy::QueryLocal : ECachePolicy::QueryRemote;
+	if (!EnumHasAnyFlags(Policy, QueryFlag))
+	{
+		UE_LOG(LogDerivedDataCache, VeryVerbose, TEXT("%s: Skipped get of %s from '%.*s' due to cache policy"),
+			*CachePath, *WriteToString<96>(Key), Name.Len(), Name.GetData());
+		return false;
+	}
+
+	if (ShouldSimulateMiss(Key))
+	{
+		UE_LOG(LogDerivedDataCache, Verbose, TEXT("%s: Simulated miss for get of %s from '%.*s'"),
+			*CachePath, *WriteToString<96>(Key), Name.Len(), Name.GetData());
+		return false;
+	}
+
+	TStringBuilder<256> Path;
+	BuildCachePackagePath(Key, Path);
+
+	bool bDeletePackage = true;
+	ON_SCOPE_EXIT
+	{
+		if (bDeletePackage && !bReadOnly)
+		{
+			IFileManager::Get().Delete(*Path, /*bRequireExists*/ false, /*bEvenReadOnly*/ false, /*bQuiet*/ true);
+		}
+	};
+
+	FCbPackage Package;
+	if (!LoadFileWithHash(Path, Name, [&Package](FArchive& Ar) { Package.TryLoad(Ar); }))
+	{
+		UE_LOG(LogDerivedDataCache, Verbose, TEXT("%s: Cache miss with missing package for %s from '%.*s'"),
+			*CachePath, *WriteToString<96>(Key), Name.Len(), Name.GetData());
+		return false;
+	}
+
+	if (ValidateCompactBinary(Package, ECbValidateMode::Default | ECbValidateMode::Package) != ECbValidateError::None)
+	{
+		UE_LOG(LogDerivedDataCache, Display, TEXT("%s: Cache miss with invalid package for %s from '%.*s'"),
+			*CachePath, *WriteToString<96>(Key), Name.Len(), Name.GetData());
+		return false;
+	}
+
+	const FCbObjectView Object = Package.GetObject();
+	const FIoHash RawHash = Object["RawHash"].AsHash();
+	const uint64 RawSize = Object["RawSize"].AsUInt64(MAX_uint64);
+	if (RawHash.IsZero() || RawSize == MAX_uint64)
+	{
+		UE_LOG(LogDerivedDataCache, Display, TEXT("%s: Cache miss with invalid value for %s from '%.*s'"),
+			*CachePath, *WriteToString<96>(Key), Name.Len(), Name.GetData());
+		return false;
+	}
+
+	if (const FCbAttachment* const Attachment = Package.FindAttachment(RawHash))
+	{
+		const FCompressedBuffer& Data = Attachment->AsCompressedBinary();
+		if (Data.GetRawHash() != RawHash || Data.GetRawSize() != RawSize)
+		{
+			UE_LOG(LogDerivedDataCache, Display,
+				TEXT("%s: Cache miss with invalid value attachment for %s from '%.*s'"),
+				*CachePath, *WriteToString<96>(Key), Name.Len(), Name.GetData());
+			return false;
+		}
+		OutValue = FValue(Data);
+	}
+	else
+	{
+		OutValue = FValue(RawHash, RawSize);
+	}
+
+	if (AccessLogWriter)
+	{
+		AccessLogWriter->Append(Key, Path);
+	}
+
+	bDeletePackage = false;
+	return true;
+}
+
+bool FFileSystemCacheStore::GetCacheValue(
+	const FStringView Name,
+	const FCacheKey& Key,
+	const ECachePolicy Policy,
+	FValue& OutValue)
+{
+	return GetCacheValueOnly(Name, Key, Policy, OutValue) && GetCacheContent(Name, Key, {}, OutValue, Policy, OutValue);
+}
+
+bool FFileSystemCacheStore::PutCacheContent(
+	const FStringView Name,
+	const FCompressedBuffer& Content,
+	uint64& OutWriteSize) const
+{
+	const FIoHash& RawHash = Content.GetRawHash();
+	TStringBuilder<256> Path;
+	BuildCacheContentPath(RawHash, Path);
+	const auto WriteContent = [&](FArchive& Ar) { Content.Save(Ar); OutWriteSize += uint64(Ar.TotalSize()); };
+	if (!FileExists(Path) && !SaveFileWithHash(Path, Name, WriteContent))
+	{
+		return false;
+	}
+	if (AccessLogWriter)
+	{
+		AccessLogWriter->Append(RawHash, Path);
+	}
+	return true;
+}
+
+bool FFileSystemCacheStore::GetCacheContentExists(const FCacheKey& Key, const FIoHash& RawHash) const
+{
+	TStringBuilder<256> Path;
+	BuildCacheContentPath(RawHash, Path);
+	return FileExists(Path);
+}
+
+bool FFileSystemCacheStore::GetCacheContent(
+	const FStringView Name,
+	const FCacheKey& Key,
+	const FValueId& Id,
+	const FValue& Value,
+	const ECachePolicy Policy,
+	FValue& OutValue) const
+{
+	if (!EnumHasAnyFlags(Policy, ECachePolicy::Query))
+	{
+		OutValue = Value.RemoveData();
+		return true;
+	}
+
+	if (Value.HasData())
+	{
+		OutValue = EnumHasAnyFlags(Policy, ECachePolicy::SkipData) ? Value.RemoveData() : Value;
+		return true;
+	}
+
+	const FIoHash& RawHash = Value.GetRawHash();
+
+	TStringBuilder<256> Path;
+	BuildCacheContentPath(RawHash, Path);
+	if (EnumHasAnyFlags(Policy, ECachePolicy::SkipData))
+	{
+		if (FileExists(Path))
+		{
+			if (AccessLogWriter)
+			{
+				AccessLogWriter->Append(RawHash, Path);
+			}
+			OutValue = Value;
+			return true;
+		}
+	}
+	else
+	{
+		FCompressedBuffer CompressedBuffer;
+		if (LoadFileWithHash(Path, Name, [&CompressedBuffer](FArchive& Ar) { CompressedBuffer = FCompressedBuffer::Load(Ar); }))
+		{
+			if (CompressedBuffer.GetRawHash() == RawHash)
+			{
+				if (AccessLogWriter)
+				{
+					AccessLogWriter->Append(RawHash, Path);
+				}
+				OutValue = FValue(MoveTemp(CompressedBuffer));
+				return true;
+			}
+			UE_LOG(LogDerivedDataCache, Display,
+				TEXT("%s: Cache miss with corrupted value %s with hash %s for %s from '%.*s'"),
+				*CachePath, *WriteToString<16>(Id), *WriteToString<48>(RawHash),
+				*WriteToString<96>(Key), Name.Len(), Name.GetData());
+			return false;
+		}
+	}
+
+	UE_LOG(LogDerivedDataCache, Verbose,
+		TEXT("%s: Cache miss with missing value %s with hash %s for %s from '%.*s'"),
+		*CachePath, *WriteToString<16>(Id), *WriteToString<48>(RawHash), *WriteToString<96>(Key),
+		Name.Len(), Name.GetData());
+	return false;
+}
+
+void FFileSystemCacheStore::GetCacheContent(
+	const FStringView Name,
+	const FCacheKey& Key,
+	const FValueId& Id,
+	const FValue& Value,
+	const ECachePolicy Policy,
+	FCompressedBufferReader& Reader,
+	TUniquePtr<FArchive>& OutArchive) const
+{
+	if (!EnumHasAnyFlags(Policy, ECachePolicy::Query))
+	{
+		return;
+	}
+
+	if (Value.HasData())
+	{
+		if (!EnumHasAnyFlags(Policy, ECachePolicy::SkipData))
+		{
+			Reader.SetSource(Value.GetData());
+		}
+		OutArchive.Reset();
+		return;
+	}
+
+	const FIoHash& RawHash = Value.GetRawHash();
+
+	TStringBuilder<256> Path;
+	BuildCacheContentPath(RawHash, Path);
+	if (EnumHasAllFlags(Policy, ECachePolicy::SkipData))
+	{
+		if (FileExists(Path))
+		{
+			if (AccessLogWriter)
+			{
+				AccessLogWriter->Append(RawHash, Path);
+			}
+			return;
+		}
+	}
+	else
+	{
+		OutArchive = OpenFile(Path, Name);
+		if (OutArchive)
+		{
+			Reader.SetSource(*OutArchive);
+			if (Reader.GetRawHash() == RawHash)
+			{
+				if (AccessLogWriter)
+				{
+					AccessLogWriter->Append(RawHash, Path);
+				}
+				return;
+			}
+			UE_LOG(LogDerivedDataCache, Display,
+				TEXT("%s: Cache miss with corrupted value %s with hash %s for %s from '%.*s'"),
+				*CachePath, *WriteToString<16>(Id), *WriteToString<48>(RawHash),
+				*WriteToString<96>(Key), Name.Len(), Name.GetData());
+			Reader.ResetSource();
+			OutArchive.Reset();
+			return;
+		}
+	}
+
+	UE_LOG(LogDerivedDataCache, Verbose,
+		TEXT("%s: Cache miss with missing value %s with hash %s for %s from '%.*s'"),
+		*CachePath, *WriteToString<16>(Id), *WriteToString<48>(RawHash), *WriteToString<96>(Key),
+		Name.Len(), Name.GetData());
+}
+
+void FFileSystemCacheStore::BuildCachePackagePath(const FCacheKey& CacheKey, FStringBuilderBase& Path) const
+{
+	Path << CachePath << TEXT('/');
+	BuildPathForCachePackage(CacheKey, Path);
+}
+
+void FFileSystemCacheStore::BuildCacheContentPath(const FIoHash& RawHash, FStringBuilderBase& Path) const
+{
+	Path << CachePath << TEXT('/');
+	BuildPathForCacheContent(RawHash, Path);
+}
+
+bool FFileSystemCacheStore::SaveFileWithHash(
+	FStringBuilderBase& Path,
+	const FStringView DebugName,
+	const TFunctionRef<void (FArchive&)> WriteFunction,
+	const bool bReplaceExisting) const
+{
+	return SaveFile(Path, DebugName, [&WriteFunction](FArchive& Ar)
+	{
+		THashingArchiveProxy<FBlake3> HashAr(Ar);
+		WriteFunction(HashAr);
+		FBlake3Hash Hash = HashAr.GetHash();
+		Ar << Hash;
+	}, bReplaceExisting);
+}
+
+bool FFileSystemCacheStore::LoadFileWithHash(
+	FStringBuilderBase& Path,
+	const FStringView DebugName,
+	const TFunctionRef<void (FArchive& Ar)> ReadFunction) const
+{
+	return LoadFile(Path, DebugName, [this, &Path, &DebugName, &ReadFunction](FArchive& Ar)
+	{
+		THashingArchiveProxy<FBlake3> HashAr(Ar);
+		ReadFunction(HashAr);
+		const FBlake3Hash Hash = HashAr.GetHash();
+		FBlake3Hash SavedHash;
+		Ar << SavedHash;
+
+		if (Hash != SavedHash && !Ar.IsError())
+		{
+			Ar.SetError();
+			UE_LOG(LogDerivedDataCache, Display,
+				TEXT("%s: File %s from '%.*s' is corrupted and has hash %s when %s is expected."),
+				*CachePath, *Path, DebugName.Len(), DebugName.GetData(),
+				*WriteToString<80>(Hash), *WriteToString<80>(SavedHash));
+		}
+	});
+}
+
+bool FFileSystemCacheStore::SaveFile(
+	FStringBuilderBase& Path,
+	const FStringView DebugName,
+	const TFunctionRef<void (FArchive&)> WriteFunction,
+	const bool bReplaceExisting) const
+{
+	const double StartTime = FPlatformTime::Seconds();
+
+	TStringBuilder<256> TempPath;
+	TempPath << FPathViews::GetPath(Path) << TEXT("/Temp.") << FGuid::NewGuid();
+
+	ON_SCOPE_EXIT
+	{
+		IFileManager::Get().Delete(*TempPath, /*bRequireExists*/ false, /*bEvenReadOnly*/ false, /*bQuiet*/ true);
+	};
+
+	int64 WriteSize = 0;
+	bool bError = false;
+
+	if (TUniquePtr<FArchive> Ar{IFileManager::Get().CreateFileWriter(*TempPath, FILEWRITE_Silent)})
+	{
+		WriteFunction(*Ar);
+		WriteSize = Ar->Tell();
+		bError = !Ar->Close();
+	}
+
+	if (bError || WriteSize == 0)
+	{
+		UE_LOG(LogDerivedDataCache, Warning,
+			TEXT("%s: Failed to write to temp file %s when saving %s from '%.*s'. Error 0x%08x"),
+			*CachePath, *TempPath, *Path, DebugName.Len(), DebugName.GetData(), FPlatformMisc::GetLastError());
+		return false;
+	}
+
+	if (IFileManager::Get().FileSize(*TempPath) != WriteSize)
+	{
+		UE_LOG(LogDerivedDataCache, Warning,
+			TEXT("%s: Failed to write to temp file %s when saving %s from '%.*s'. ")
+			TEXT("File is %" INT64_FMT " bytes when %" INT64_FMT " bytes are expected."),
+			*CachePath, *TempPath, *Path, DebugName.Len(), DebugName.GetData(),
+			IFileManager::Get().FileSize(*TempPath), WriteSize);
+		return false;
+	}
+
+	if (!IFileManager::Get().Move(*Path, *TempPath,
+		bReplaceExisting, /*bEvenIfReadOnly*/ false, /*bAttributes*/ false, /*bDoNotRetryOrError*/ true))
+	{
+		UE_LOG(LogDerivedDataCache, Log,
+			TEXT("%s: Move collision when writing file %s from '%.*s'."),
+			*CachePath, *Path, DebugName.Len(), DebugName.GetData());
+	}
+
+	const double WriteDuration = FPlatformTime::Seconds() - StartTime;
+	const double WriteSpeed = WriteDuration > 0.001 ? (WriteSize / WriteDuration) / (1024.0 * 1024.0) : 0.0;
+	UE_LOG(LogDerivedDataCache, VeryVerbose,
+		TEXT("%s: Saved %s from '%.*s' (%" INT64_FMT " bytes, %.02f secs, %.2f MiB/s)"),
+		*CachePath, *Path, DebugName.Len(), DebugName.GetData(), WriteSize, WriteDuration, WriteSpeed);
+
+	return true;
+}
+
+bool FFileSystemCacheStore::LoadFile(
+	FStringBuilderBase& Path,
+	const FStringView DebugName,
+	const TFunctionRef<void (FArchive& Ar)> ReadFunction) const
+{
+	check(IsUsable());
+	const double StartTime = FPlatformTime::Seconds();
+
+	int64 ReadSize = 0;
+	bool bError = false;
+
+	if (TUniquePtr<FArchive> Ar = OpenFile(Path, DebugName))
+	{
+		ReadFunction(*Ar);
+		ReadSize = Ar->Tell();
+		bError = !Ar->Close();
+
+		if (bError)
+		{
+			UE_LOG(LogDerivedDataCache, Display,
+				TEXT("%s: Failed to load file %s from '%.*s'."),
+				*CachePath, *Path, DebugName.Len(), DebugName.GetData());
+		}
+	}
+
+	const double ReadDuration = FPlatformTime::Seconds() - StartTime;
+	const double ReadSpeed = ReadDuration > 0.001 ? (ReadSize / ReadDuration) / (1024.0 * 1024.0) : 0.0;
+
+	if (!GIsBuildMachine && ReadDuration > 5.0)
+	{
+		// Slower than 0.5 MiB/s?
+		UE_CLOG(ReadSpeed < 0.5, LogDerivedDataCache, Warning,
+			TEXT("%s: Loading %s from '%.*s' is very slow (%.2f MiB/s); consider disabling this cache backend"),
+			*CachePath, *Path, DebugName.Len(), DebugName.GetData(), ReadSpeed);
+	}
+
+	UE_LOG(LogDerivedDataCache, VeryVerbose,
+		TEXT("%s: Loaded %s from '%.*s' (%" INT64_FMT " bytes, %.02f secs, %.2f MiB/s)"),
+		*CachePath, *Path, DebugName.Len(), DebugName.GetData(), ReadSize, ReadDuration, ReadSpeed);
+
+	if (bError && !bReadOnly)
+	{
+		IFileManager::Get().Delete(*Path, /*bRequireExists*/ false, /*bEvenReadOnly*/ false, /*bQuiet*/ true);
+	}
+
+	return !bError && ReadSize > 0;
+}
+
+TUniquePtr<FArchive> FFileSystemCacheStore::OpenFile(FStringBuilderBase& Path, const FStringView DebugName) const
+{
+	// Check for existence before reading because it may update the modification time and avoid the
+	// file being deleted by a cache cleanup thread or process.
+	if (!FileExists(Path))
+	{
+		return nullptr;
+	}
+
+	return TUniquePtr<FArchive>(IFileManager::Get().CreateFileReader(*Path, FILEREAD_Silent));
+}
+
+bool FFileSystemCacheStore::FileExists(FStringBuilderBase& Path) const
+{
+	const FDateTime TimeStamp = IFileManager::Get().GetTimeStamp(*Path);
+	if (TimeStamp == FDateTime::MinValue())
+	{
+		return false;
+	}
+	if (bTouch || (!bReadOnly && (FDateTime::UtcNow() - TimeStamp).GetTotalDays() > (DaysToDeleteUnusedFiles / 4)))
+	{
+		IFileManager::Get().SetTimeStamp(*Path, FDateTime::UtcNow());
+	}
+	return true;
+}
+
+bool FFileSystemCacheStore::ShouldSimulateMiss(const TCHAR* const InKey)
+{
+	if (DebugOptions.RandomMissRate == 0 && DebugOptions.SimulateMissTypes.IsEmpty())
+	{
+		return false;
+	}
+
+	const FName Key(InKey);
+	const uint32 Hash = GetTypeHash(Key);
+
+	if (FScopeLock Lock(&MissedKeysCS); DebugMissedKeys.ContainsByHash(Hash, Key))
+	{
+		return true;
+	}
+
+	if (DebugOptions.ShouldSimulateMiss(InKey))
+	{
+		FScopeLock Lock(&MissedKeysCS);
+		DebugMissedKeys.AddByHash(Hash, Key);
+		return true;
+	}
+
+	return false;
+}
+
+bool FFileSystemCacheStore::ShouldSimulateMiss(const FCacheKey& Key)
+{
+	if (DebugOptions.RandomMissRate == 0 && DebugOptions.SimulateMissTypes.IsEmpty())
+	{
+		return false;
+	}
+
+	const uint32 Hash = GetTypeHash(Key);
+
+	if (FScopeLock Lock(&MissedKeysCS); DebugMissedCacheKeys.ContainsByHash(Hash, Key))
+	{
+		return true;
+	}
+
+	if (DebugOptions.ShouldSimulateMiss(Key))
+	{
+		FScopeLock Lock(&MissedKeysCS);
+		DebugMissedCacheKeys.AddByHash(Hash, Key);
+		return true;
+	}
+
+	return false;
+}
+
+ILegacyCacheStore* CreateFileSystemCacheStore(
+	const TCHAR* CachePath,
+	const TCHAR* Params,
+	const TCHAR* AccessLogPath,
+	ECacheStoreFlags& OutFlags)
+{
+	TUniquePtr<FFileSystemCacheStore> Store(new FFileSystemCacheStore(CachePath, Params, AccessLogPath));
+	if (Store->IsUsable())
+	{
+		ECacheStoreFlags Flags = ECacheStoreFlags::Query;
+		Flags |= Store->IsWritable() ? ECacheStoreFlags::Store : ECacheStoreFlags::None;
+		Flags |= Store->GetSpeedClass() == EBackendSpeedClass::Local ? ECacheStoreFlags::Local : ECacheStoreFlags::Remote;
+		OutFlags = Flags;
+		return Store.Release();
+	}
+	return nullptr;
+}
+
+} // UE::DerivedData
