DDC: Enabled compression of legacy cache data

diff --git a/Engine/Source/Developer/DerivedDataCache/Private/MemoryCacheStore.cpp b/Engine/Source/Developer/DerivedDataCache/Private/MemoryCacheStore.cpp
new file mode 100644
--- /dev/null
+++ b/Engine/Source/Developer/DerivedDataCache/Private/MemoryCacheStore.cpp
@@ -0,0 +1,989 @@
+// Copyright Epic Games, Inc. All Rights Reserved.
+
+#include "Algo/AllOf.h"
+#include "DerivedDataCachePrivate.h"
+#include "DerivedDataCacheRecord.h"
+#include "DerivedDataCacheUsageStats.h"
+#include "DerivedDataValue.h"
+#include "FileBackedDerivedDataBackend.h"
+#include "HAL/FileManager.h"
+#include "Misc/ScopeExit.h"
+#include "Misc/ScopeLock.h"
+#include "Misc/ScopeRWLock.h"
+#include "ProfilingDebugging/CookStats.h"
+#include "Serialization/CompactBinary.h"
+#include "Templates/UniquePtr.h"
+
+namespace UE::DerivedData
+{
+
+/**
+ * A simple thread safe, memory based backend. This is used for Async puts and the boot cache.
+ */
+class FMemoryCacheStore final : public FFileBackedDerivedDataBackend
+{
+public:
+	explicit FMemoryCacheStore(const TCHAR* InName, int64 InMaxCacheSize = -1, bool bCanBeDisabled = false);
+	~FMemoryCacheStore() final;
+
+	/** Return a name for this interface */
+	FString GetName() const final { return Name; }
+
+	/** return true if this cache is writable **/
+	bool IsWritable() const final;
+
+	/** Returns a class of speed for this interface **/
+	ESpeedClass GetSpeedClass()  const final;
+
+	/**
+	 * Synchronous test for the existence of a cache item
+	 *
+	 * @param	CacheKey	Alphanumeric+underscore key of this cache item
+	 * @return				true if the data probably will be found, this can't be guaranteed because of concurrency in the backends, corruption, etc
+	 */
+	bool CachedDataProbablyExists(const TCHAR* CacheKey) final;
+
+	/**
+	 * Synchronous retrieve of a cache item
+	 *
+	 * @param	CacheKey	Alphanumeric+underscore key of this cache item
+	 * @param	OutData		Buffer to receive the results, if any were found
+	 * @return				true if any data was found, and in this case OutData is non-empty
+	 */
+	bool GetCachedData(const TCHAR* CacheKey, TArray<uint8>& OutData) final;
+
+	/**
+	 * Asynchronous, fire-and-forget placement of a cache item
+	 *
+	 * @param	CacheKey	Alphanumeric+underscore key of this cache item
+	 * @param	InData		Buffer containing the data to cache, can be destroyed after the call returns, immediately
+	 * @param	bPutEvenIfExists	If true, then do not attempt skip the put even if CachedDataProbablyExists returns true
+	 */
+	EPutStatus PutCachedData(const TCHAR* CacheKey, TArrayView<const uint8> InData, bool bPutEvenIfExists) final;
+
+	void RemoveCachedData(const TCHAR* CacheKey, bool bTransient) final;
+
+	/**
+	 * Save the cache to disk
+	 * @param	Filename	Filename to save
+	 * @return	true if file was saved successfully
+	 */
+	bool SaveCache(const TCHAR* Filename);
+
+	/**
+	 * Load the cache from disk
+	 * @param	Filename	Filename to load
+	 * @return	true if file was loaded successfully
+	 */
+	bool LoadCache(const TCHAR* Filename);
+
+	/**
+	 * Disable cache and ignore all subsequent requests
+	 */
+	void Disable() final;
+
+	TSharedRef<FDerivedDataCacheStatsNode> GatherUsageStats() const final;
+
+	TBitArray<> TryToPrefetch(TConstArrayView<FString> CacheKeys) final;
+
+	/**
+	 *  Determines if we would cache the provided data
+	 */
+	bool WouldCache(const TCHAR* CacheKey, TArrayView<const uint8> InData) final;
+
+	/**
+	 * Apply debug options
+	 */
+	bool ApplyDebugOptions(FBackendDebugOptions& InOptions) final;
+
+	EBackendLegacyMode GetLegacyMode() const final { return EBackendLegacyMode::ValueWithLegacyFallback; }
+
+	void Put(
+		TConstArrayView<FCachePutRequest> Requests,
+		IRequestOwner& Owner,
+		FOnCachePutComplete&& OnComplete) final;
+	void Get(
+		TConstArrayView<FCacheGetRequest> Requests,
+		IRequestOwner& Owner,
+		FOnCacheGetComplete&& OnComplete) final;
+	void PutValue(
+		TConstArrayView<FCachePutValueRequest> Requests,
+		IRequestOwner& Owner,
+		FOnCachePutValueComplete&& OnComplete) final;
+	void GetValue(
+		TConstArrayView<FCacheGetValueRequest> Requests,
+		IRequestOwner& Owner,
+		FOnCacheGetValueComplete&& OnComplete) final;
+	void GetChunks(
+		TConstArrayView<FCacheGetChunkRequest> Requests,
+		IRequestOwner& Owner,
+		FOnCacheGetChunkComplete&& OnComplete) final;
+
+private:
+	/** Name of the cache file loaded (if any). */
+	FString CacheFilename;
+	FDerivedDataCacheUsageStats UsageStats;
+
+	struct FCacheValue
+	{
+		int32 Age;
+		int32 Size;
+		FRWLock DataLock;
+		TArray<uint8> Data;
+		FCacheValue(int32 InSize, int32 InAge = 0)
+			: Age(InAge)
+			, Size(InSize)
+		{
+		}
+
+		~FCacheValue()
+		{
+			// Ensure we don't begin destruction of this value while our own PutCachedData function is holding this lock
+			FWriteScopeLock WriteLock(DataLock);
+		}
+	};
+
+	FORCEINLINE int32 CalcSerializedCacheValueSize(const FString& Key, const FCacheValue& Val)
+	{
+		return (Key.Len() + 1) * sizeof(TCHAR) + sizeof(FCacheValue::Age) + Val.Size;
+	}
+
+	FORCEINLINE int32 CalcSerializedCacheValueSize(const FString& Key, const TArrayView<const uint8>& Data)
+	{
+		return (Key.Len() + 1) * sizeof(TCHAR) + sizeof(FCacheValue::Age) + Data.Num();
+	}
+
+	/** Name of this cache (used for debugging) */
+	FString Name;
+
+	/** Set of files that are being written to disk asynchronously. */
+	TMap<FString, FCacheValue*> CacheItems;
+	/** Set of records in this cache. */
+	TSet<FCacheRecord, FCacheRecordKeyFuncs> CacheRecords;
+	/** Set of values in this cache. */
+	TMap<FCacheKey, FValue> CacheValues;
+	/** Maximum size the cached items can grow up to ( in bytes ) */
+	int64 MaxCacheSize;
+	/** When set to true, this cache is disabled...ignore all requests. */
+	std::atomic<bool> bDisabled;
+	/** Object used for synchronization via a scoped lock						*/
+	mutable FRWLock SynchronizationObject;
+	/** Current estimated cache size in bytes */
+	int64 CurrentCacheSize;
+	/** Indicates that the cache max size has been exceeded. This is used to avoid
+		warning spam after the size has reached the limit. */
+	bool bMaxSizeExceeded;
+
+	/** When a memory cache can be disabled, it won't return true for CachedDataProbablyExists calls.
+	  * This is to avoid having the Boot DDC tells it has some resources that will suddenly disappear after the boot.
+	  * Get() requests will still get fulfilled and other cache level will be properly back-filled 
+	  * offering the speed benefit of the boot cache while maintaining coherency at all cache levels.
+	  * 
+	  * The problem is that most asset types (audio/staticmesh/texture) will always verify if their different LODS/Chunks can be found in the cache using CachedDataProbablyExists.
+	  * If any of the LOD/MIP can't be found, a build of the asset is triggered, otherwise they skip asset compilation altogether.
+	  * However, we should not skip the compilation based on the CachedDataProbablyExists result of the boot cache because it is a lie and will disappear at some point.
+	  * When the boot cache disappears and the streamer tries to fetch a LOD that it has been told was cached, it will fail and will then have no choice but to rebuild the asset synchronously.
+	  * This obviously causes heavy game-thread stutters.
+
+	  * However, if the bootcache returns false during CachedDataProbablyExists. The async compilation will be triggered and data will be put in the both the boot.ddc and the local cache.
+	  * This way, no more heavy game-thread stutters during streaming...
+
+	  * This can be reproed when you clear the local cache but do not clear the boot.ddc file, but even if it's a corner case, I stumbled upon it enough times that I though it was worth to fix so the caches are coherent.
+	  */
+	bool bCanBeDisabled = false;
+	bool bShuttingDown  = false;
+
+	enum 
+	{
+		/** Magic number to use in header */
+		MemCache_Magic = 0x0cac0ddc,
+		/** Magic number to use in header (new, > 2GB size compatible) */
+		MemCache_Magic64 = 0x0cac1ddc,
+		/** Oldest cache items to keep */
+		MaxAge = 3,
+		/** Size of data that is stored in the cachefile apart from the cache entries (64 bit size). */
+		SerializationSpecificDataSize = sizeof(uint32)	// Magic
+									  + sizeof(int64)	// Size
+									  + sizeof(uint32), // CRC
+	};
+
+protected:
+
+	/* Debug helpers */
+	bool ShouldSimulateMiss(const TCHAR* InKey);
+	bool ShouldSimulateMiss(const FCacheKey& InKey);
+
+	/** Debug Options */
+	FBackendDebugOptions DebugOptions;
+
+	/** Keys we ignored due to miss rate settings */
+	FCriticalSection MissedKeysCS;
+	TSet<FName> DebugMissedKeys;
+	TSet<FCacheKey> DebugMissedCacheKeys;
+};
+
+FMemoryCacheStore::FMemoryCacheStore(const TCHAR* InName, int64 InMaxCacheSize, bool bInCanBeDisabled)
+	: Name(InName)
+	, MaxCacheSize(InMaxCacheSize)
+	, bDisabled( false )
+	, CurrentCacheSize( SerializationSpecificDataSize )
+	, bMaxSizeExceeded(false)
+	, bCanBeDisabled(bInCanBeDisabled)
+{
+}
+
+FMemoryCacheStore::~FMemoryCacheStore()
+{
+	bShuttingDown = true;
+	Disable();
+}
+
+bool FMemoryCacheStore::IsWritable() const
+{
+	return !bDisabled;
+}
+
+FDerivedDataBackendInterface::ESpeedClass FMemoryCacheStore::GetSpeedClass() const
+{
+	return ESpeedClass::Local;
+}
+
+bool FMemoryCacheStore::CachedDataProbablyExists(const TCHAR* CacheKey)
+{
+	// See comments on the declaration of bCanBeDisabled variable.
+	if (bCanBeDisabled)
+	{
+		return false;
+	}
+
+	COOK_STAT(auto Timer = UsageStats.TimeProbablyExists());
+
+	if (ShouldSimulateMiss(CacheKey))
+	{
+		return false;
+	}
+
+	if (bDisabled)
+	{
+		return false;
+	}
+
+	FReadScopeLock ScopeLock(SynchronizationObject);
+	bool Result = CacheItems.Contains(FString(CacheKey));
+	if (Result)
+	{
+		COOK_STAT(Timer.AddHit(0));
+	}
+	return Result;
+}
+
+bool FMemoryCacheStore::GetCachedData(const TCHAR* CacheKey, TArray<uint8>& OutData)
+{
+	COOK_STAT(auto Timer = UsageStats.TimeGet());
+
+	if (ShouldSimulateMiss(CacheKey))
+	{
+		return false;
+	}
+	
+	if (!bDisabled)
+	{
+		FReadScopeLock ScopeLock(SynchronizationObject);
+
+		FCacheValue* Item = CacheItems.FindRef(FString(CacheKey));
+		if (Item)
+		{
+			TRACE_CPUPROFILER_EVENT_SCOPE(FMemoryCacheStore::GetCachedData);
+			FReadScopeLock ItemLock(Item->DataLock);
+
+			OutData = Item->Data;
+			Item->Age = 0;
+			check(OutData.Num());
+			COOK_STAT(Timer.AddHit(OutData.Num()));
+			return true;
+		}
+	}
+	OutData.Empty();
+	return false;
+}
+
+TBitArray<> FMemoryCacheStore::TryToPrefetch(TConstArrayView<FString> CacheKeys)
+{
+	return CachedDataProbablyExistsBatch(CacheKeys);
+}
+
+bool FMemoryCacheStore::WouldCache(const TCHAR* CacheKey, TArrayView<const uint8> InData)
+{
+	if (bDisabled || bMaxSizeExceeded)
+	{
+		return false;
+	}
+
+	return true;
+}
+
+FDerivedDataBackendInterface::EPutStatus FMemoryCacheStore::PutCachedData(const TCHAR* CacheKey, TArrayView<const uint8> InData, bool bPutEvenIfExists)
+{
+	TRACE_CPUPROFILER_EVENT_SCOPE(FMemoryCacheStore::PutCachedData);
+	COOK_STAT(auto Timer = UsageStats.TimePut());
+
+	FString Key;
+	{
+		FReadScopeLock ReadScopeLock(SynchronizationObject);
+
+		if (ShouldSimulateMiss(CacheKey))
+		{
+			return EPutStatus::Skipped;
+		}
+
+		// Should never hit this as higher level code should be checking..
+		if (!WouldCache(CacheKey, InData))
+		{
+			//UE_LOG(LogDerivedDataCache, Warning, TEXT("WouldCache was not called prior to attempted Put!"));
+			return EPutStatus::NotCached;
+		}
+
+		Key = CacheKey;
+		FCacheValue* Item = CacheItems.FindRef(Key);
+		if (Item)
+		{
+			//check(Item->Data == InData); // any second attempt to push data should be identical data
+			return EPutStatus::Cached;
+		}
+	}
+	
+	// Compute the size of the cache before copying the data to avoid alloc/memcpy of something we might throw away
+	// if we bust the cache size.
+	int32 CacheValueSize = CalcSerializedCacheValueSize(Key, InData);
+
+	FCacheValue* Val = nullptr;
+
+	{
+		FWriteScopeLock WriteScopeLock(SynchronizationObject);
+
+		// check if we haven't exceeded the MaxCacheSize
+		if (MaxCacheSize > 0 && (CurrentCacheSize + CacheValueSize) > MaxCacheSize)
+		{
+			UE_LOG(LogDerivedDataCache, Display, TEXT("Failed to cache data. Maximum cache size reached. CurrentSize %d kb / MaxSize: %d kb"), CurrentCacheSize / 1024, MaxCacheSize / 1024);
+			bMaxSizeExceeded = true;
+			return EPutStatus::NotCached;
+		}
+
+		if (CacheItems.Contains(Key))
+		{
+			// Another thread already beat us, just return.
+			return EPutStatus::Cached;
+		}
+
+		COOK_STAT(Timer.AddHit(InData.Num()));
+		Val = new FCacheValue(InData.Num());
+		// Make sure no other thread can access the data until we finish copying it
+		Val->DataLock.WriteLock();
+		CacheItems.Add(Key, Val);
+
+		CurrentCacheSize += CacheValueSize;
+	}
+
+	// Data is copied outside of the lock so that we only lock
+	// for this specific key instead of always locking all keys by default.
+	Val->Data = InData;
+	// It is now safe for other thread to access this data, unlock
+	Val->DataLock.WriteUnlock();
+
+	return EPutStatus::Cached;
+}
+
+void FMemoryCacheStore::RemoveCachedData(const TCHAR* CacheKey, bool bTransient)
+{
+	if (bDisabled || bTransient)
+	{
+		return;
+	}
+
+	TRACE_CPUPROFILER_EVENT_SCOPE(FMemoryCacheStore::RemoveCachedData);
+	FString Key(CacheKey);
+	FCacheValue* Item = nullptr;
+	{
+		FWriteScopeLock WriteScopeLock(SynchronizationObject);
+		
+		if (CacheItems.RemoveAndCopyValue(Key, Item) && Item)
+		{
+			CurrentCacheSize -= CalcSerializedCacheValueSize(Key, *Item);
+			bMaxSizeExceeded = false;
+		}
+	}
+
+	if (Item)
+	{
+		// Just make sure any other thread has finished writing or reading the data before deleting.
+		Item->DataLock.WriteLock();
+
+		// Avoid deleting the item with a locked FRWLock, unlocking is safe because
+		// the item is now unreachable from other threads
+		Item->DataLock.WriteUnlock();
+		delete Item;
+	}
+}
+
+bool FMemoryCacheStore::SaveCache(const TCHAR* Filename)
+{
+	TRACE_CPUPROFILER_EVENT_SCOPE(FMemoryCacheStore::SaveCache);
+
+	double StartTime = FPlatformTime::Seconds();
+	TUniquePtr<FArchive> SaverArchive(IFileManager::Get().CreateFileWriter(Filename, FILEWRITE_EvenIfReadOnly));
+	if (!SaverArchive)
+	{
+		UE_LOG(LogDerivedDataCache, Error, TEXT("Could not save memory cache %s."), Filename);
+		return false;
+	}
+
+	FArchive& Saver = *SaverArchive;
+	uint32 Magic = MemCache_Magic64;
+	Saver << Magic;
+	const int64 DataStartOffset = Saver.Tell();
+	{
+		FWriteScopeLock ScopeLock(SynchronizationObject);
+		check(!bDisabled);
+		for (TMap<FString, FCacheValue*>::TIterator It(CacheItems); It; ++It )
+		{
+			Saver << It.Key();
+			Saver << It.Value()->Age;
+			FReadScopeLock ReadScopeLock(It.Value()->DataLock);
+			Saver << It.Value()->Data;
+		}
+	}
+	const int64 DataSize = Saver.Tell(); // Everything except the footer
+	int64 Size = DataSize;
+	uint32 Crc = MemCache_Magic64; // Crc takes more time than I want to spend  FCrc::MemCrc_DEPRECATED(&Buffer[0], Size);
+	Saver << Size;
+	Saver << Crc;
+
+	check(SerializationSpecificDataSize + DataSize <= MaxCacheSize || MaxCacheSize <= 0);
+
+	UE_LOG(LogDerivedDataCache, Log, TEXT("Saved boot cache %4.2fs %lldMB %s."), float(FPlatformTime::Seconds() - StartTime), DataSize / (1024 * 1024), Filename);
+	return true;
+}
+
+bool FMemoryCacheStore::LoadCache(const TCHAR* Filename)
+{
+	TRACE_CPUPROFILER_EVENT_SCOPE(FMemoryCacheStore::LoadCache);
+
+	double StartTime = FPlatformTime::Seconds();
+	const int64 FileSize = IFileManager::Get().FileSize(Filename);
+	if (FileSize < 0)
+	{
+		UE_LOG(LogDerivedDataCache, Warning, TEXT("Could not find memory cache %s."), Filename);
+		return false;
+	}
+	// We test 3 * uint32 which is the old format (< SerializationSpecificDataSize). We'll test
+	// against SerializationSpecificDataSize later when we read the magic number from the cache.
+	if (FileSize < sizeof(uint32) * 3)
+	{
+		UE_LOG(LogDerivedDataCache, Error, TEXT("Memory cache was corrputed (short) %s."), Filename);
+		return false;
+	}
+	if (FileSize > MaxCacheSize*2 && MaxCacheSize > 0)
+	{
+		UE_LOG(LogDerivedDataCache, Error, TEXT("Refusing to load DDC cache %s. Size exceeds doubled MaxCacheSize."), Filename);
+		return false;
+	}
+
+	TUniquePtr<FArchive> LoaderArchive(IFileManager::Get().CreateFileReader(Filename));
+	if (!LoaderArchive)
+	{
+		UE_LOG(LogDerivedDataCache, Warning, TEXT("Could not read memory cache %s."), Filename);
+		return false;
+	}
+
+	FArchive& Loader = *LoaderArchive;
+	uint32 Magic = 0;
+	Loader << Magic;
+	if (Magic != MemCache_Magic && Magic != MemCache_Magic64)
+	{
+		UE_LOG(LogDerivedDataCache, Error, TEXT("Memory cache was corrputed (magic) %s."), Filename);
+		return false;
+	}
+	// Check the file size again, this time against the correct minimum size.
+	if (Magic == MemCache_Magic64 && FileSize < SerializationSpecificDataSize)
+	{
+		UE_LOG(LogDerivedDataCache, Error, TEXT("Memory cache was corrputed (short) %s."), Filename);
+		return false;
+	}
+	// Calculate expected DataSize based on the magic number (footer size difference)
+	const int64 DataSize = FileSize - (Magic == MemCache_Magic64 ? (SerializationSpecificDataSize - sizeof(uint32)) : (sizeof(uint32) * 2));		
+	Loader.Seek(DataSize);
+	int64 Size = 0;
+	uint32 Crc = 0;
+	if (Magic == MemCache_Magic64)
+	{
+		Loader << Size;
+	}
+	else
+	{
+		uint32 Size32 = 0;
+		Loader << Size32;
+		Size = (int64)Size32;
+	}
+	Loader << Crc;
+	if (Size != DataSize)
+	{
+		UE_LOG(LogDerivedDataCache, Error, TEXT("Memory cache was corrputed (size) %s."), Filename);
+		return false;
+	}
+	if ((Crc != MemCache_Magic && Crc != MemCache_Magic64) || Crc != Magic)
+	{
+		UE_LOG(LogDerivedDataCache, Warning, TEXT("Memory cache was corrputed (crc) %s."), Filename);
+		return false;
+	}
+	// Seek to data start offset (skip magic number)
+	Loader.Seek(sizeof(uint32));
+	{
+		TArray<uint8> Working;
+		FWriteScopeLock ScopeLock(SynchronizationObject);
+		check(!bDisabled);
+		while (Loader.Tell() < DataSize)
+		{
+			FString Key;
+			int32 Age;
+			Loader << Key;
+			Loader << Age;
+			Age++;
+			Loader << Working;
+			if (Age < MaxAge)
+			{
+				CacheItems.Add(Key, new FCacheValue(Working.Num(), Age))->Data = MoveTemp(Working);
+			}
+			Working.Reset();
+		}
+		// these are just a double check on ending correctly
+		if (Magic == MemCache_Magic64)
+		{
+			Loader << Size;
+		}
+		else
+		{
+			uint32 Size32 = 0;
+			Loader << Size32;
+			Size = (int64)Size32;
+		}
+		Loader << Crc;
+
+		CurrentCacheSize = FileSize;
+		CacheFilename = Filename;
+	}
+
+	UE_LOG(LogDerivedDataCache, Log, TEXT("Loaded boot cache %4.2fs %lldMB %s."), float(FPlatformTime::Seconds() - StartTime), DataSize / (1024 * 1024), Filename);
+	return true;
+}
+
+void FMemoryCacheStore::Disable()
+{
+	check(bCanBeDisabled || bShuttingDown);
+	FWriteScopeLock ScopeLock(SynchronizationObject);
+	bDisabled = true;
+	for (TMap<FString,FCacheValue*>::TIterator It(CacheItems); It; ++It )
+	{
+		delete It.Value();
+	}
+	CacheItems.Empty();
+
+	CurrentCacheSize = SerializationSpecificDataSize;
+}
+
+TSharedRef<FDerivedDataCacheStatsNode> FMemoryCacheStore::GatherUsageStats() const
+{
+	TSharedRef<FDerivedDataCacheStatsNode> Usage =
+		MakeShared<FDerivedDataCacheStatsNode>(CacheFilename.IsEmpty() ? TEXT("Memory") : TEXT("Boot"), CacheFilename, /*bIsLocal*/ true);
+	Usage->Stats.Add(TEXT(""), UsageStats);
+	return Usage;
+}
+
+bool FMemoryCacheStore::ApplyDebugOptions(FBackendDebugOptions& InOptions)
+{
+	DebugOptions = InOptions;
+	return true;
+}
+
+bool FMemoryCacheStore::ShouldSimulateMiss(const TCHAR* InKey)
+{
+	if (DebugOptions.RandomMissRate == 0 && DebugOptions.SimulateMissTypes.IsEmpty())
+	{
+		return false;
+	}
+
+	const FName Key(InKey);
+	const uint32 Hash = GetTypeHash(Key);
+
+	if (FScopeLock Lock(&MissedKeysCS); DebugMissedKeys.ContainsByHash(Hash, Key))
+	{
+		return true;
+	}
+
+	if (DebugOptions.ShouldSimulateMiss(InKey))
+	{
+		FScopeLock Lock(&MissedKeysCS);
+		UE_LOG(LogDerivedDataCache, Verbose, TEXT("Simulating miss in %s for %s"), *GetName(), InKey);
+		DebugMissedKeys.AddByHash(Hash, Key);
+		return true;
+	}
+
+	return false;
+}
+
+bool FMemoryCacheStore::ShouldSimulateMiss(const FCacheKey& Key)
+{
+	if (DebugOptions.RandomMissRate == 0 && DebugOptions.SimulateMissTypes.IsEmpty())
+	{
+		return false;
+	}
+
+	const uint32 Hash = GetTypeHash(Key);
+
+	if (FScopeLock Lock(&MissedKeysCS); DebugMissedCacheKeys.ContainsByHash(Hash, Key))
+	{
+		return true;
+	}
+
+	if (DebugOptions.ShouldSimulateMiss(Key))
+	{
+		FScopeLock Lock(&MissedKeysCS);
+		DebugMissedCacheKeys.AddByHash(Hash, Key);
+		return true;
+	}
+
+	return false;
+}
+
+void FMemoryCacheStore::Put(
+	const TConstArrayView<FCachePutRequest> Requests,
+	IRequestOwner& Owner,
+	FOnCachePutComplete&& OnComplete)
+{
+	for (const FCachePutRequest& Request : Requests)
+	{
+		const FCacheRecord& Record = Request.Record;
+		const FCacheKey& Key = Record.GetKey();
+		EStatus Status = EStatus::Error;
+		ON_SCOPE_EXIT
+		{
+			OnComplete({Request.Name, Key, Request.UserData, Status});
+		};
+
+		if (!EnumHasAnyFlags(Request.Policy.GetRecordPolicy(), ECachePolicy::StoreLocal))
+		{
+			continue;
+		}
+
+		if (ShouldSimulateMiss(Key))
+		{
+			UE_LOG(LogDerivedDataCache, Verbose, TEXT("%s: Simulated miss for put of %s from '%s'"),
+				*GetName(), *WriteToString<96>(Key), *Request.Name);
+			continue;
+		}
+
+		const TConstArrayView<FValueWithId> Values = Record.GetValues();
+
+		if (!Algo::AllOf(Values, &FValue::HasData))
+		{
+			continue;
+		}
+
+		COOK_STAT(auto Timer = UsageStats.TimePut());
+		const int64 RecordSize = Private::GetCacheRecordCompressedSize(Record);
+		const bool bReplaceExisting = !EnumHasAnyFlags(Request.Policy.GetRecordPolicy(), ECachePolicy::QueryLocal);
+
+		FWriteScopeLock ScopeLock(SynchronizationObject);
+		FCacheRecord* const ExistingRecord = CacheRecords.Find(Key);
+		Status = ExistingRecord && !bReplaceExisting ? EStatus::Ok : EStatus::Error;
+		if (bDisabled || Status == EStatus::Ok)
+		{
+			continue;
+		}
+
+		const int64 ExistingRecordSize = ExistingRecord ? Private::GetCacheRecordCompressedSize(*ExistingRecord) : 0;
+		const int64 RequiredSize = RecordSize - ExistingRecordSize;
+
+		if (MaxCacheSize > 0 && (CurrentCacheSize + RequiredSize) > MaxCacheSize)
+		{
+			UE_LOG(LogDerivedDataCache, Display, TEXT("Failed to cache data. Maximum cache size reached. CurrentSize %" INT64_FMT " KiB / MaxSize: %" INT64_FMT " KiB"), CurrentCacheSize / 1024, MaxCacheSize / 1024);
+			bMaxSizeExceeded = true;
+			continue;
+		}
+
+		CurrentCacheSize += RequiredSize;
+		if (ExistingRecord)
+		{
+			*ExistingRecord = Record;
+		}
+		else
+		{
+			CacheRecords.Add(Record);
+		}
+		COOK_STAT(Timer.AddHit(RecordSize));
+		Status = EStatus::Ok;
+	}
+}
+
+void FMemoryCacheStore::Get(
+	const TConstArrayView<FCacheGetRequest> Requests,
+	IRequestOwner& Owner,
+	FOnCacheGetComplete&& OnComplete)
+{
+	if (bDisabled)
+	{
+		return CompleteWithStatus(Requests, OnComplete, EStatus::Error);
+	}
+
+	for (const FCacheGetRequest& Request : Requests)
+	{
+		const FCacheKey& Key = Request.Key;
+		const FCacheRecordPolicy& Policy = Request.Policy;
+		const bool bExistsOnly = EnumHasAllFlags(Policy.GetRecordPolicy(), ECachePolicy::SkipData);
+		COOK_STAT(auto Timer = bExistsOnly ? UsageStats.TimeProbablyExists() : UsageStats.TimeGet());
+		FOptionalCacheRecord Record;
+		EStatus Status = EStatus::Error;
+		if (ShouldSimulateMiss(Key))
+		{
+			UE_LOG(LogDerivedDataCache, Verbose, TEXT("%s: Simulated miss for get of %s from '%s'"),
+				*GetName(), *WriteToString<96>(Key), *Request.Name);
+		}
+		else if (FReadScopeLock ScopeLock(SynchronizationObject); const FCacheRecord* CacheRecord = CacheRecords.Find(Key))
+		{
+			Status = EStatus::Ok;
+			Record = *CacheRecord;
+		}
+		if (Record)
+		{
+			for (const FValueWithId& Value : Record.Get().GetValues())
+			{
+				if (!Value.HasData() && !EnumHasAnyFlags(Policy.GetValuePolicy(Value.GetId()), ECachePolicy::SkipData))
+				{
+					Status = EStatus::Error;
+					if (!EnumHasAllFlags(Policy.GetRecordPolicy(), ECachePolicy::PartialRecord))
+					{
+						Record.Reset();
+						break;
+					}
+				}
+			}
+		}
+		if (Record)
+		{
+			COOK_STAT(Timer.AddHit(Private::GetCacheRecordCompressedSize(Record.Get())));
+			OnComplete({Request.Name, MoveTemp(Record).Get(), Request.UserData, Status});
+		}
+		else
+		{
+			OnComplete(Request.MakeResponse(EStatus::Error));
+		}
+	}
+}
+
+void FMemoryCacheStore::PutValue(
+	const TConstArrayView<FCachePutValueRequest> Requests,
+	IRequestOwner& Owner,
+	FOnCachePutValueComplete&& OnComplete)
+{
+	for (const FCachePutValueRequest& Request : Requests)
+	{
+		EStatus Status = EStatus::Error;
+		ON_SCOPE_EXIT
+		{
+			OnComplete(Request.MakeResponse(Status));
+		};
+
+		if (!EnumHasAnyFlags(Request.Policy, ECachePolicy::StoreLocal))
+		{
+			continue;
+		}
+
+		const FCacheKey& Key = Request.Key;
+		const FValue& Value = Request.Value;
+
+		if (ShouldSimulateMiss(Key))
+		{
+			UE_LOG(LogDerivedDataCache, Verbose, TEXT("%s: Simulated miss for put of %s from '%s'"),
+				*GetName(), *WriteToString<96>(Key), *Request.Name);
+			continue;
+		}
+
+		if (!Value.HasData())
+		{
+			continue;
+		}
+
+		COOK_STAT(auto Timer = UsageStats.TimePut());
+		const int64 ValueSize = Value.GetData().GetCompressedSize();
+		const bool bReplaceExisting = !EnumHasAnyFlags(Request.Policy, ECachePolicy::QueryLocal);
+
+		FWriteScopeLock ScopeLock(SynchronizationObject);
+		FValue* const ExistingValue = CacheValues.Find(Key);
+		Status = ExistingValue && !bReplaceExisting ? EStatus::Ok : EStatus::Error;
+		if (bDisabled || Status == EStatus::Ok)
+		{
+			continue;
+		}
+
+		const int64 ExistingValueSize = ExistingValue ? ExistingValue->GetData().GetCompressedSize() : 0;
+		const int64 RequiredSize = ValueSize - ExistingValueSize;
+
+		if (MaxCacheSize > 0 && (CurrentCacheSize + RequiredSize) > MaxCacheSize)
+		{
+			UE_LOG(LogDerivedDataCache, Display, TEXT("Failed to cache data. Maximum cache size reached. CurrentSize %" INT64_FMT " KiB / MaxSize: %" INT64_FMT " KiB"), CurrentCacheSize / 1024, MaxCacheSize / 1024);
+			bMaxSizeExceeded = true;
+			continue;
+		}
+
+		CurrentCacheSize += RequiredSize;
+		if (ExistingValue)
+		{
+			*ExistingValue = Value;
+		}
+		else
+		{
+			CacheValues.Add(Key, Value);
+		}
+		COOK_STAT(Timer.AddHit(ValueSize));
+		Status = EStatus::Ok;
+	}
+}
+
+void FMemoryCacheStore::GetValue(
+	const TConstArrayView<FCacheGetValueRequest> Requests,
+	IRequestOwner& Owner,
+	FOnCacheGetValueComplete&& OnComplete)
+{
+	if (bDisabled)
+	{
+		return CompleteWithStatus(Requests, OnComplete, EStatus::Error);
+	}
+
+	for (const FCacheGetValueRequest& Request : Requests)
+	{
+		const FCacheKey& Key = Request.Key;
+		const ECachePolicy Policy = Request.Policy;
+		const bool bExistsOnly = EnumHasAllFlags(Policy, ECachePolicy::SkipData);
+		COOK_STAT(auto Timer = bExistsOnly ? UsageStats.TimeProbablyExists() : UsageStats.TimeGet());
+		bool bProcessHit = false;
+		FValue Value;
+		EStatus Status = EStatus::Error;
+		if (ShouldSimulateMiss(Key))
+		{
+			UE_LOG(LogDerivedDataCache, Verbose, TEXT("%s: Simulated miss for get of %s from '%s'"),
+				*GetName(), *WriteToString<96>(Key), *Request.Name);
+		}
+		else if (FReadScopeLock ScopeLock(SynchronizationObject); const FValue* CacheValue = CacheValues.Find(Key))
+		{
+			Status = EStatus::Ok;
+			Value = *CacheValue;
+			bProcessHit = true;
+		}
+
+		if (bProcessHit)
+		{
+			if (!Value.HasData() && !EnumHasAnyFlags(Policy, ECachePolicy::SkipData))
+			{
+				Status = EStatus::Error;
+			}
+
+			COOK_STAT(Timer.AddHit(Value.GetData().GetCompressedSize()));
+			Value = EnumHasAnyFlags(Policy, ECachePolicy::SkipData) ? Value.RemoveData() : Value;
+			OnComplete({Request.Name, Request.Key, MoveTemp(Value), Request.UserData, Status});
+		}
+		else
+		{
+			OnComplete(Request.MakeResponse(EStatus::Error));
+		}
+	}
+}
+
+void FMemoryCacheStore::GetChunks(
+	const TConstArrayView<FCacheGetChunkRequest> Requests,
+	IRequestOwner& Owner,
+	FOnCacheGetChunkComplete&& OnComplete)
+{
+	if (bDisabled)
+	{
+		return CompleteWithStatus(Requests, OnComplete, EStatus::Error);
+	}
+
+	bool bHasValue = false;
+	FValue Value;
+	FValueId ValueId;
+	FCacheKey ValueKey;
+	FCompressedBufferReader Reader;
+	for (const FCacheGetChunkRequest& Request : Requests)
+	{
+		bool bProcessHit = false;
+		const bool bExistsOnly = EnumHasAnyFlags(Request.Policy, ECachePolicy::SkipData);
+		COOK_STAT(auto Timer = bExistsOnly ? UsageStats.TimeProbablyExists() : UsageStats.TimeGet());
+		if (ShouldSimulateMiss(Request.Key))
+		{
+			UE_LOG(LogDerivedDataCache, Verbose, TEXT("%s: Simulated miss for get of %s from '%s'"),
+				*GetName(), *WriteToString<96>(Request.Key, '/', Request.Id), *Request.Name);
+		}
+		else if (bHasValue && (ValueKey == Request.Key) && (ValueId == Request.Id) && (bExistsOnly || Reader.HasSource()))
+		{
+			// Value matches the request.
+			bProcessHit = true;
+		}
+		else
+		{
+			FReadScopeLock ScopeLock(SynchronizationObject);
+			if (Request.Id.IsValid())
+			{
+				if (const FCacheRecord* Record = CacheRecords.Find(Request.Key))
+				{
+					const FValueWithId& ValueWithId = Record->GetValue(Request.Id);
+					bHasValue = ValueWithId.IsValid();
+					Reader.ResetSource();
+					Value.Reset();
+					Value = ValueWithId;
+					ValueId = Request.Id;
+					ValueKey = Request.Key;
+					Reader.SetSource(Value.GetData());
+					bProcessHit = true;
+				}
+			}
+			else
+			{
+				if (const FValue* ExistingValue = CacheValues.Find(Request.Key))
+				{
+					bHasValue = true;
+					Reader.ResetSource();
+					Value.Reset();
+					Value = *ExistingValue;
+					ValueId.Reset();
+					ValueKey = Request.Key;
+					Reader.SetSource(Value.GetData());
+					bProcessHit = true;
+				}
+			}
+		}
+
+		if (bProcessHit && Request.RawOffset <= Value.GetRawSize())
+		{
+			const uint64 RawSize = FMath::Min(Value.GetRawSize() - Request.RawOffset, Request.RawSize);
+			COOK_STAT(Timer.AddHit(RawSize));
+			FSharedBuffer Buffer;
+			if (Value.HasData() && !bExistsOnly)
+			{
+				Buffer = Reader.Decompress(Request.RawOffset, RawSize);
+			}
+			const EStatus Status = bExistsOnly || Buffer ? EStatus::Ok : EStatus::Error;
+			OnComplete({Request.Name, Request.Key, Request.Id, Request.RawOffset,
+				RawSize, Value.GetRawHash(), MoveTemp(Buffer), Request.UserData, Status});
+		}
+		else
+		{
+			OnComplete(Request.MakeResponse(EStatus::Error));
+		}
+	}
+}
+
+FFileBackedDerivedDataBackend* CreateMemoryCacheStore(const TCHAR* Name, int64 MaxCacheSize, bool bCanBeDisabled)
+{
+	return new FMemoryCacheStore(Name, MaxCacheSize, bCanBeDisabled);
+}
+
+} // UE::DerivedData
diff --git a/Engine/Source/Developer/DerivedDataCache/Private/PakFileCacheStore.cpp b/Engine/Source/Developer/DerivedDataCache/Private/PakFileCacheStore.cpp
new file mode 100644
--- /dev/null
+++ b/Engine/Source/Developer/DerivedDataCache/Private/PakFileCacheStore.cpp
@@ -0,0 +1,1673 @@
+// Copyright Epic Games, Inc. All Rights Reserved.
+
+#include "PakFileCacheStore.h"
+
+#include "Algo/Accumulate.h"
+#include "Algo/StableSort.h"
+#include "Algo/Transform.h"
+#include "Compression/OodleDataCompression.h"
+#include "DerivedDataCachePrivate.h"
+#include "DerivedDataCacheRecord.h"
+#include "DerivedDataCacheUsageStats.h"
+#include "DerivedDataChunk.h"
+#include "DerivedDataRequest.h"
+#include "DerivedDataRequestOwner.h"
+#include "DerivedDataValue.h"
+#include "HAL/CriticalSection.h"
+#include "HAL/FileManager.h"
+#include "HAL/PlatformFile.h"
+#include "HAL/PlatformFileManager.h"
+#include "HashingArchiveProxy.h"
+#include "Misc/Compression.h"
+#include "Misc/FileHelper.h"
+#include "Misc/Paths.h"
+#include "Misc/PathViews.h"
+#include "Misc/ScopeRWLock.h"
+#include "ProfilingDebugging/CookStats.h"
+#include "Serialization/CompactBinary.h"
+#include "Serialization/CompactBinaryPackage.h"
+#include "Serialization/CompactBinaryValidation.h"
+#include "Serialization/CompactBinaryWriter.h"
+#include "Serialization/MemoryReader.h"
+#include "Serialization/MemoryWriter.h"
+#include "Tasks/Task.h"
+#include "Templates/Greater.h"
+#include "Templates/UniquePtr.h"
+
+namespace UE::DerivedData
+{
+
+/**
+ * A simple thread safe, pak file based backend.
+ */
+class FPakFileCacheStore : public IPakFileCacheStore
+{
+public:
+	FPakFileCacheStore(const TCHAR* InFilename, bool bInWriting);
+	~FPakFileCacheStore();
+
+	void Close() final;
+
+	/** Return a name for this interface */
+	FString GetName() const final;
+
+	/** return true if this cache is writable **/
+	bool IsWritable() const final;
+
+	/** Returns a class of speed for this interface **/
+	ESpeedClass GetSpeedClass() const final;
+
+	bool BackfillLowerCacheLevels() const final;
+
+	/**
+	 * Synchronous test for the existence of a cache item
+	 *
+	 * @param	CacheKey	Alphanumeric+underscore key of this cache item
+	 * @return				true if the data probably will be found, this can't be guaranteed because of concurrency in the backends, corruption, etc
+	 */
+	bool CachedDataProbablyExists(const TCHAR* CacheKey) final;
+
+	/**
+	 * Synchronous retrieve of a cache item
+	 *
+	 * @param	CacheKey	Alphanumeric+underscore key of this cache item
+	 * @param	OutData		Buffer to receive the results, if any were found
+	 * @return				true if any data was found, and in this case OutData is non-empty
+	 */
+	bool GetCachedData(const TCHAR* CacheKey, TArray<uint8>& OutData) override;
+
+	/**
+	 * Asynchronous, fire-and-forget placement of a cache item
+	 *
+	 * @param	CacheKey	Alphanumeric+underscore key of this cache item
+	 * @param	InData		Buffer containing the data to cache, can be destroyed after the call returns, immediately
+	 * @param	bPutEvenIfExists	If true, then do not attempt skip the put even if CachedDataProbablyExists returns true
+	 */
+	EPutStatus PutCachedData(const TCHAR* CacheKey, TArrayView<const uint8> InData, bool bPutEvenIfExists) override;
+
+	void RemoveCachedData(const TCHAR* CacheKey, bool bTransient) final;
+
+	/**
+	 * Save the cache to disk
+	 * @return	true if file was saved successfully
+	 */
+	bool SaveCache() final;
+
+	/**
+	 * Load the cache to disk	 * @param	Filename	Filename to load
+	 * @return	true if file was loaded successfully
+	 */
+	bool LoadCache(const TCHAR* InFilename) final;
+
+	/**
+	 * Merges another cache file into this one.
+	 * @return true on success
+	 */
+	void MergeCache(IPakFileCacheStore* OtherPak) final;
+	
+	const FString& GetFilename() const final
+	{
+		return CachePath;
+	}
+
+	TSharedRef<FDerivedDataCacheStatsNode> GatherUsageStats() const final;
+
+	TBitArray<> TryToPrefetch(TConstArrayView<FString> CacheKeys) final
+	{
+		return CachedDataProbablyExistsBatch(CacheKeys);
+	}
+
+	bool WouldCache(const TCHAR* CacheKey, TArrayView<const uint8> InData) final { return true; }
+
+	bool ApplyDebugOptions(FBackendDebugOptions& InOptions) final { return false; }
+
+	EBackendLegacyMode GetLegacyMode() const final { return EBackendLegacyMode::ValueWithLegacyFallback; }
+
+	void Put(
+		TConstArrayView<FCachePutRequest> Requests,
+		IRequestOwner& Owner,
+		FOnCachePutComplete&& OnComplete) override;
+
+	void Get(
+		TConstArrayView<FCacheGetRequest> Requests,
+		IRequestOwner& Owner,
+		FOnCacheGetComplete&& OnComplete) final;
+
+	void PutValue(
+		TConstArrayView<FCachePutValueRequest> Requests,
+		IRequestOwner& Owner,
+		FOnCachePutValueComplete&& OnComplete) override;
+
+	void GetValue(
+		TConstArrayView<FCacheGetValueRequest> Requests,
+		IRequestOwner& Owner,
+		FOnCacheGetValueComplete&& OnComplete) final;
+
+	void GetChunks(
+		TConstArrayView<FCacheGetChunkRequest> Requests,
+		IRequestOwner& Owner,
+		FOnCacheGetChunkComplete&& OnComplete) final;
+
+private:
+	[[nodiscard]] bool PutCacheRecord(FStringView Name, const FCacheRecord& Record, const FCacheRecordPolicy& Policy, uint64& OutWriteSize);
+
+	[[nodiscard]] FOptionalCacheRecord GetCacheRecordOnly(
+		FStringView Name,
+		const FCacheKey& Key,
+		const FCacheRecordPolicy& Policy);
+	[[nodiscard]] FOptionalCacheRecord GetCacheRecord(
+		FStringView Name,
+		const FCacheKey& Key,
+		const FCacheRecordPolicy& Policy,
+		EStatus& OutStatus);
+
+	[[nodiscard]] bool PutCacheValue(FStringView Name, const FCacheKey& Key, const FValue& Value, ECachePolicy Policy, uint64& OutWriteSize);
+
+	[[nodiscard]] bool GetCacheValueOnly(FStringView Name, const FCacheKey& Key, ECachePolicy Policy, FValue& OutValue);
+	[[nodiscard]] bool GetCacheValue(FStringView Name, const FCacheKey& Key, ECachePolicy Policy, FValue& OutValue);
+
+	[[nodiscard]] bool PutCacheContent(FStringView Name, const FCompressedBuffer& Content, uint64& OutWriteSize);
+
+	[[nodiscard]] bool GetCacheContentExists(const FCacheKey& Key, const FIoHash& RawHash);
+	[[nodiscard]] bool GetCacheContent(
+		FStringView Name,
+		const FCacheKey& Key,
+		const FValueId& Id,
+		const FValue& Value,
+		ECachePolicy Policy,
+		FValue& OutValue);
+	void GetCacheContent(
+		const FStringView Name,
+		const FCacheKey& Key,
+		const FValueId& Id,
+		const FValue& Value,
+		const ECachePolicy Policy,
+		FCompressedBufferReader& Reader,
+		TUniquePtr<FArchive>& OutArchive);
+
+	[[nodiscard]] bool SaveFile(FStringView Path, FStringView DebugName, TFunctionRef<void (FArchive&)> WriteFunction);
+	[[nodiscard]] FSharedBuffer LoadFile(FStringView Path, FStringView DebugName);
+	[[nodiscard]] TUniquePtr<FArchive> OpenFile(FStringBuilderBase& Path, const FStringView DebugName);
+	[[nodiscard]] bool FileExists(FStringView Path);
+
+private:
+	FDerivedDataCacheUsageStats UsageStats;
+
+	struct FCacheValue
+	{
+		int64 Offset;
+		int64 Size;
+		uint32 Crc;
+		FCacheValue(int64 InOffset, int64 InSize, uint32 InCrc)
+			: Offset(InOffset)
+			, Size(InSize)
+			, Crc(InCrc)
+		{
+		}
+	};
+
+	/** When set to true, we are a pak writer (we don't do reads). */
+	bool bWriting;
+	/** When set to true, we are a pak writer and we saved, so we shouldn't be used anymore. Also, a read cache that failed to open. */
+	bool bClosed;
+	/** Object used for synchronization via scoped read or write locks. */
+	FRWLock SynchronizationObject;
+	/** Set of files that are being written to disk asynchronously. */
+	TMap<FString, FCacheValue> CacheItems;
+	/** File handle of pak. */
+	TUniquePtr<IFileHandle> FileHandle;
+	/** File name of pak. */
+	FString CachePath;
+
+	/** Maximum total size of compressed data stored within a record package with multiple attachments. */
+	uint64 MaxRecordSizeKB = 256;
+	/** Maximum total size of compressed data stored within a value package, or a record package with one attachment. */
+	uint64 MaxValueSizeKB = 1024;
+
+	enum 
+	{
+		/** Magic number to use in header */
+		PakCache_Magic = 0x0c7c0ddc,
+	};
+
+	friend class IPakFileCacheStore;
+};
+
+FPakFileCacheStore::FPakFileCacheStore(const TCHAR* const InCachePath, const bool bInWriting)
+	: bWriting(bInWriting)
+	, bClosed(false)
+	, CachePath(InCachePath)
+{
+	IPlatformFile& PlatformFile = FPlatformFileManager::Get().GetPlatformFile();
+	if (bWriting)
+	{
+		PlatformFile.CreateDirectoryTree(*FPaths::GetPath(CachePath));
+		FileHandle.Reset(PlatformFile.OpenWrite(*CachePath, /*bAppend*/ false, /*bAllowRead*/ true));
+		if (!FileHandle)
+		{
+			UE_LOG(LogDerivedDataCache, Fatal, TEXT("%s: Failed to open pak cache for writing."), *CachePath);
+			bClosed = true;
+		}
+		else
+		{
+			UE_LOG(LogDerivedDataCache, Display, TEXT("%s: Opened pak cache for writing."), *CachePath);
+		}
+	}
+	else
+	{
+		FileHandle.Reset(PlatformFile.OpenRead(*CachePath));
+		if (!FileHandle)
+		{
+			UE_LOG(LogDerivedDataCache, Warning, TEXT("%s: Failed to open pak cache for reading."), *CachePath);
+		}
+		else if (!LoadCache(*CachePath))
+		{
+			FileHandle.Reset();
+			CacheItems.Empty();
+			bClosed = true;
+		}
+		else
+		{
+			UE_LOG(LogDerivedDataCache, Display, TEXT("%s: Opened pak cache for reading. (%" INT64_FMT " MiB)"),
+				*CachePath, FileHandle->Size() / 1024 / 1024);
+		}
+	}
+}
+
+FPakFileCacheStore::~FPakFileCacheStore()
+{
+	Close();
+}
+
+FString FPakFileCacheStore::GetName() const
+{
+	return CachePath;
+}
+
+void FPakFileCacheStore::Close()
+{
+	FDerivedDataBackend::Get().WaitForQuiescence();
+	if (!bClosed)
+	{
+		if (bWriting)
+		{
+			SaveCache();
+		}
+		FWriteScopeLock ScopeLock(SynchronizationObject);
+		FileHandle.Reset();
+		CacheItems.Empty();
+		bClosed = true;
+	}
+}
+
+bool FPakFileCacheStore::IsWritable() const
+{
+	return bWriting && !bClosed;
+}
+
+FDerivedDataBackendInterface::ESpeedClass FPakFileCacheStore::GetSpeedClass() const
+{
+	return ESpeedClass::Local;
+}
+
+bool FPakFileCacheStore::BackfillLowerCacheLevels() const
+{
+	return false;
+}
+
+bool FPakFileCacheStore::CachedDataProbablyExists(const TCHAR* CacheKey)
+{
+	COOK_STAT(auto Timer = UsageStats.TimeProbablyExists());
+	FReadScopeLock ScopeLock(SynchronizationObject);
+	bool Result = CacheItems.Contains(FString(CacheKey));
+	if (Result)
+	{
+		COOK_STAT(Timer.AddHit(0));
+	}
+	return Result;
+}
+
+bool FPakFileCacheStore::GetCachedData(const TCHAR* CacheKey, TArray<uint8>& OutData)
+{
+	COOK_STAT(auto Timer = UsageStats.TimeGet());
+	if (bClosed)
+	{
+		return false;
+	}
+	FWriteScopeLock ScopeLock(SynchronizationObject);
+	if (FCacheValue* Item = CacheItems.Find(FString(CacheKey)))
+	{
+		check(FileHandle);
+		ON_SCOPE_EXIT
+		{
+			if (bWriting)
+			{
+				FileHandle->SeekFromEnd();
+			}
+		};
+		if (Item->Size >= MAX_int32)
+		{
+			UE_LOG(LogDerivedDataCache, Warning, TEXT("%s: Pak file, %s exceeds 2 GiB limit."), *CachePath, CacheKey);
+		}
+		else if (!FileHandle->Seek(Item->Offset))
+		{
+			UE_LOG(LogDerivedDataCache, Warning, TEXT("%s: Pak file, bad seek."), *CachePath);
+		}
+		else
+		{
+			check(Item->Size);
+			check(!OutData.Num());
+			OutData.AddUninitialized(Item->Size);
+			if (!FileHandle->Read(OutData.GetData(), int64(Item->Size)))
+			{
+				UE_LOG(LogDerivedDataCache, Warning, TEXT("%s: Pak file, bad read."), *CachePath);
+			}
+			else if (uint32 TestCrc = FCrc::MemCrc_DEPRECATED(OutData.GetData(), Item->Size); TestCrc != Item->Crc)
+			{
+				UE_LOG(LogDerivedDataCache, Warning, TEXT("%s: Pak file, bad crc."), *CachePath);
+			}
+			else
+			{
+				UE_LOG(LogDerivedDataCache, Verbose, TEXT("%s: Cache hit on %s"), *CachePath, CacheKey);
+				check(OutData.Num());
+				COOK_STAT(Timer.AddHit(OutData.Num()));
+				return true;
+			}
+		}
+	}
+	else
+	{
+		UE_LOG(LogDerivedDataCache, Verbose, TEXT("%s: Cache miss on %s"), *CachePath, CacheKey);
+	}
+	OutData.Empty();
+	return false;
+}
+
+FDerivedDataBackendInterface::EPutStatus FPakFileCacheStore::PutCachedData(const TCHAR* CacheKey, TArrayView<const uint8> InData, bool bPutEvenIfExists)
+{
+	COOK_STAT(auto Timer = UsageStats.TimePut());
+	if (!IsWritable())
+	{
+		return EPutStatus::NotCached;
+	}
+	{
+		FWriteScopeLock ScopeLock(SynchronizationObject);
+		FString Key(CacheKey);
+		TOptional<uint32> Crc;
+		check(InData.Num());
+		check(Key.Len());
+		check(FileHandle);
+
+		if (bPutEvenIfExists)
+		{
+			if (FCacheValue* Item = CacheItems.Find(FString(CacheKey)))
+			{
+				// If there was an existing entry for this key, if it had the same contents, do nothing as the desired value is already stored.
+				// If the contents differ, replace it if the size hasn't changed, but if the size has changed, 
+				// remove the existing entry from the index but leave they actual data payload in place as it is too
+				// costly to go back and attempt to rewrite all offsets and shift all bytes that follow it in the file.
+				if (Item->Size == InData.Num())
+				{
+					COOK_STAT(Timer.AddHit(InData.Num()));
+					Crc = FCrc::MemCrc_DEPRECATED(InData.GetData(), InData.Num());
+					if (Crc.GetValue() != Item->Crc)
+					{
+						int64 Offset = FileHandle->Tell();
+						FileHandle->Seek(Item->Offset);
+						FileHandle->Write(InData.GetData(), InData.Num());
+						Item->Crc = Crc.GetValue();
+						FileHandle->Seek(Offset);
+					}
+					return EPutStatus::Cached;
+				}
+
+				UE_LOG(LogDerivedDataCache, Warning,
+					TEXT("%s: Repeated put of %s with different sized contents. Multiple contents will be in the file, ")
+					TEXT("but only the last will be in the index. This has wasted %" INT64_FMT " bytes in the file."),
+					*CachePath, CacheKey, Item->Size);
+				CacheItems.Remove(Key);
+			}
+		}
+
+		int64 Offset = FileHandle->Tell();
+		if (Offset < 0)
+		{
+			CacheItems.Empty();
+			FileHandle.Reset();
+			UE_LOG(LogDerivedDataCache, Fatal, TEXT("%s: Could not write pak file... out of disk space?"), *CachePath);
+			return EPutStatus::NotCached;
+		}
+		else
+		{
+			COOK_STAT(Timer.AddHit(InData.Num()));
+			if (!Crc.IsSet())
+			{
+				Crc = FCrc::MemCrc_DEPRECATED(InData.GetData(), InData.Num());
+			}
+			FileHandle->Write(InData.GetData(), InData.Num());
+			UE_LOG(LogDerivedDataCache, Verbose, TEXT("%s: Put %s"), *CachePath, CacheKey);
+			CacheItems.Add(Key, FCacheValue(Offset, InData.Num(), Crc.GetValue()));
+			return EPutStatus::Cached;
+		}
+	}
+}
+
+void FPakFileCacheStore::RemoveCachedData(const TCHAR* CacheKey, bool bTransient)
+{
+	if (bClosed || bTransient)
+	{
+		return;
+	}
+	// strangish. We can delete from a pak, but it only deletes the index 
+	// if this is a read cache, it will read it next time
+	// if this is a write cache, we wasted space
+	FWriteScopeLock ScopeLock(SynchronizationObject);
+	FString Key(CacheKey);
+	CacheItems.Remove(Key);
+}
+
+bool FPakFileCacheStore::SaveCache()
+{
+	FWriteScopeLock ScopeLock(SynchronizationObject);
+	check(FileHandle);
+	int64 IndexOffset = FileHandle->Tell();
+	check(IndexOffset >= 0);
+	uint32 NumItems = uint32(CacheItems.Num());
+	check(IndexOffset > 0 || !NumItems);
+	TArray<uint8> IndexBuffer;
+	{
+		FMemoryWriter Saver(IndexBuffer);
+		uint32 NumProcessed = 0;
+		for (TMap<FString, FCacheValue>::TIterator It(CacheItems); It; ++It )
+		{
+			FCacheValue& Value = It.Value();
+			check(It.Key().Len());
+			check(Value.Size);
+			check(Value.Offset >= 0 && Value.Offset < IndexOffset);
+			Saver << It.Key();
+			Saver << Value.Offset;
+			Saver << Value.Size;
+			Saver << Value.Crc;
+			NumProcessed++;
+		}
+		check(NumProcessed == NumItems);
+	}
+	uint32 IndexCrc = FCrc::MemCrc_DEPRECATED(IndexBuffer.GetData(), IndexBuffer.Num());
+	uint32 SizeIndex = uint32(IndexBuffer.Num());
+
+	uint32 Magic = PakCache_Magic;
+	TArray<uint8> Buffer;
+	FMemoryWriter Saver(Buffer);
+	Saver << Magic;
+	Saver << IndexCrc;
+	Saver << NumItems;
+	Saver << SizeIndex;
+	Saver.Serialize(IndexBuffer.GetData(), IndexBuffer.Num());
+	Saver << Magic;
+	Saver << IndexOffset;
+	FileHandle->Write(Buffer.GetData(), Buffer.Num());
+	CacheItems.Empty();
+	FileHandle.Reset();
+	bClosed = true;
+	return true;
+}
+
+bool FPakFileCacheStore::LoadCache(const TCHAR* InFilename)
+{
+	check(FileHandle);
+	int64 FileSize = FileHandle->Size();
+	check(FileSize >= 0);
+	if (FileSize < sizeof(int64) + sizeof(uint32) * 5)
+	{
+		UE_LOG(LogDerivedDataCache, Error, TEXT("%s: Pak cache was corrupted (short)."), InFilename);
+		return false;
+	}
+	int64 IndexOffset = -1;
+	int64 Trailer = -1;
+	{
+		TArray<uint8> Buffer;
+		const int64 SeekPos = FileSize - int64(sizeof(int64) + sizeof(uint32));
+		FileHandle->Seek(SeekPos);
+		Trailer = FileHandle->Tell();
+		if (Trailer != SeekPos)
+		{
+			UE_LOG(LogDerivedDataCache, Error, TEXT("%s: Pak cache was corrupted (bad seek)."), InFilename);
+			return false;
+		}
+		check(Trailer >= 0 && Trailer < FileSize);
+		Buffer.AddUninitialized(sizeof(int64) + sizeof(uint32));
+		FileHandle->Read(Buffer.GetData(), int64(sizeof(int64)+sizeof(uint32)));
+		FMemoryReader Loader(Buffer);
+		uint32 Magic = 0;
+		Loader << Magic;
+		Loader << IndexOffset;
+		if (Magic != PakCache_Magic || IndexOffset < 0 || IndexOffset + int64(sizeof(uint32) * 4) > Trailer)
+		{
+			UE_LOG(LogDerivedDataCache, Error, TEXT("%s: Pak cache was corrupted (bad footer)."), InFilename);
+			return false;
+		}
+	}
+	uint32 IndexCrc = 0;
+	uint32 NumIndex = 0;
+	uint32 SizeIndex = 0;
+	{
+		TArray<uint8> Buffer;
+		FileHandle->Seek(IndexOffset);
+		if (FileHandle->Tell() != IndexOffset)
+		{
+			UE_LOG(LogDerivedDataCache, Error, TEXT("%s: Pak cache was corrupted (bad seek index)."), InFilename);
+			return false;
+		}
+		Buffer.AddUninitialized(sizeof(uint32) * 4);
+		FileHandle->Read(Buffer.GetData(), sizeof(uint32) * 4);
+		FMemoryReader Loader(Buffer);
+		uint32 Magic = 0;
+		Loader << Magic;
+		Loader << IndexCrc;
+		Loader << NumIndex;
+		Loader << SizeIndex;
+		if (Magic != PakCache_Magic || (SizeIndex != 0 && NumIndex == 0) || (SizeIndex == 0 && NumIndex != 0)) 
+		{
+			UE_LOG(LogDerivedDataCache, Error, TEXT("%s: Pak cache was corrupted (bad index header)."), InFilename);
+			return false;
+		}
+		if (IndexOffset + sizeof(uint32) * 4 + SizeIndex != Trailer) 
+		{
+			UE_LOG(LogDerivedDataCache, Error, TEXT("%s: Pak cache was corrupted (bad index size)."), InFilename);
+			return false;
+		}
+	}
+	{
+		TArray<uint8> Buffer;
+		Buffer.AddUninitialized(SizeIndex);
+		FileHandle->Read(Buffer.GetData(), SizeIndex);
+		FMemoryReader Loader(Buffer);
+		while (Loader.Tell() < SizeIndex)
+		{
+			FString Key;
+			int64 Offset;
+			int64 Size;
+			uint32 Crc;
+			Loader << Key;
+			Loader << Offset;
+			Loader << Size;
+			Loader << Crc;
+			if (!Key.Len() || Offset < 0 || Offset >= IndexOffset || !Size)
+			{
+				UE_LOG(LogDerivedDataCache, Error, TEXT("%s: Pak cache was corrupted (bad index entry)."), InFilename);
+				return false;
+			}
+			CacheItems.Add(Key, FCacheValue(Offset, Size, Crc));
+		}
+		if (CacheItems.Num() != NumIndex)
+		{
+			UE_LOG(LogDerivedDataCache, Error, TEXT("%s: Pak cache was corrupted (bad index count)."), InFilename);
+			return false;
+		}
+	}
+	return true;
+}
+
+void FPakFileCacheStore::MergeCache(IPakFileCacheStore* OtherPakInterface)
+{
+	FPakFileCacheStore* OtherPak = static_cast<FPakFileCacheStore*>(OtherPakInterface);
+
+	// Get all the existing keys
+	TArray<FString> KeyNames;
+	OtherPak->CacheItems.GenerateKeyArray(KeyNames);
+
+	// Find all the keys to copy
+	TArray<FString> CopyKeyNames;
+	for(const FString& KeyName : KeyNames)
+	{
+		if(!CachedDataProbablyExists(*KeyName))
+		{
+			CopyKeyNames.Add(KeyName);
+		}
+	}
+	UE_LOG(LogDerivedDataCache, Display, TEXT("Merging %d entries (%d skipped)."), CopyKeyNames.Num(), KeyNames.Num() - CopyKeyNames.Num());
+
+	// Copy them all to the new cache. Don't use the overloaded get/put methods (which may compress/decompress); copy the raw data directly.
+	TArray<uint8> Buffer;
+	for(const FString& CopyKeyName : CopyKeyNames)
+	{
+		Buffer.Reset();
+		if(OtherPak->FPakFileCacheStore::GetCachedData(*CopyKeyName, Buffer))
+		{
+			FPakFileCacheStore::PutCachedData(*CopyKeyName, Buffer, false);
+		}
+	}
+}
+
+bool IPakFileCacheStore::SortAndCopy(const FString &InputFilename, const FString &OutputFilename)
+{
+	// Open the input and output files
+	FPakFileCacheStore InputPak(*InputFilename, false);
+	if (InputPak.bClosed) return false;
+
+	FPakFileCacheStore OutputPak(*OutputFilename, true);
+	if (OutputPak.bClosed) return false;
+
+	// Sort the key names
+	TArray<FString> KeyNames;
+	InputPak.CacheItems.GenerateKeyArray(KeyNames);
+	KeyNames.Sort();
+
+	// Copy all the DDC to the new cache
+	TArray<uint8> Buffer;
+	TArray<uint32> KeySizes;
+	for (int KeyIndex = 0; KeyIndex < KeyNames.Num(); KeyIndex++)
+	{
+		Buffer.Reset();
+		InputPak.GetCachedData(*KeyNames[KeyIndex], Buffer);
+		OutputPak.PutCachedData(*KeyNames[KeyIndex], Buffer, false);
+		KeySizes.Add(Buffer.Num());
+	}
+
+	// Write out a TOC listing for debugging
+	FStringOutputDevice Output;
+	Output.Logf(TEXT("Asset,Size") LINE_TERMINATOR);
+	for(int KeyIndex = 0; KeyIndex < KeyNames.Num(); KeyIndex++)
+	{
+		Output.Logf(TEXT("%s,%d") LINE_TERMINATOR, *KeyNames[KeyIndex], KeySizes[KeyIndex]);
+	}
+	FFileHelper::SaveStringToFile(Output, *FPaths::Combine(*FPaths::GetPath(OutputFilename), *(FPaths::GetBaseFilename(OutputFilename) + TEXT(".csv"))));
+	return true;
+}
+
+TSharedRef<FDerivedDataCacheStatsNode> FPakFileCacheStore::GatherUsageStats() const
+{
+	TSharedRef<FDerivedDataCacheStatsNode> Usage =
+		MakeShared<FDerivedDataCacheStatsNode>(TEXT("PakFile"), CachePath, /*bIsLocal*/ true);
+	Usage->Stats.Add(TEXT(""), UsageStats);
+	return Usage;
+}
+
+void FPakFileCacheStore::Put(
+	const TConstArrayView<FCachePutRequest> Requests,
+	IRequestOwner& Owner,
+	FOnCachePutComplete&& OnComplete)
+{
+	for (const FCachePutRequest& Request : Requests)
+	{
+		const FCacheRecord& Record = Request.Record;
+		TRACE_CPUPROFILER_EVENT_SCOPE(PakFileDDC_Put);
+		COOK_STAT(auto Timer = UsageStats.TimePut());
+		uint64 WriteSize = 0;
+		if (PutCacheRecord(Request.Name, Record, Request.Policy, WriteSize))
+		{
+			UE_LOG(LogDerivedDataCache, Verbose, TEXT("%s: Cache put complete for %s from '%s'"),
+				*CachePath, *WriteToString<96>(Record.GetKey()), *Request.Name);
+			COOK_STAT(if (WriteSize) { Timer.AddHit(WriteSize); });
+			OnComplete(Request.MakeResponse(EStatus::Ok));
+		}
+		else
+		{
+			OnComplete(Request.MakeResponse(EStatus::Error));
+		}
+	}
+}
+
+void FPakFileCacheStore::Get(
+	const TConstArrayView<FCacheGetRequest> Requests,
+	IRequestOwner& Owner,
+	FOnCacheGetComplete&& OnComplete)
+{
+	for (const FCacheGetRequest& Request : Requests)
+	{
+		TRACE_CPUPROFILER_EVENT_SCOPE(PakFileDDC_GetValue);
+		COOK_STAT(auto Timer = UsageStats.TimeGet());
+		EStatus Status = EStatus::Ok;
+		if (FOptionalCacheRecord Record = GetCacheRecord(Request.Name, Request.Key, Request.Policy, Status))
+		{
+			UE_LOG(LogDerivedDataCache, Verbose, TEXT("%s: Cache hit for %s from '%s'"),
+				*CachePath, *WriteToString<96>(Request.Key), *Request.Name);
+			COOK_STAT(Timer.AddHit(Private::GetCacheRecordCompressedSize(Record.Get())));
+			OnComplete({Request.Name, MoveTemp(Record).Get(), Request.UserData, Status});
+		}
+		else
+		{
+			OnComplete(Request.MakeResponse(Status));
+		}
+	}
+}
+
+void FPakFileCacheStore::PutValue(
+	const TConstArrayView<FCachePutValueRequest> Requests,
+	IRequestOwner& Owner,
+	FOnCachePutValueComplete&& OnComplete)
+{
+	for (const FCachePutValueRequest& Request : Requests)
+	{
+		TRACE_CPUPROFILER_EVENT_SCOPE(PakFileDDC_Put);
+		COOK_STAT(auto Timer = UsageStats.TimePut());
+		uint64 WriteSize = 0;
+		if (PutCacheValue(Request.Name, Request.Key, Request.Value, Request.Policy, WriteSize))
+		{
+			UE_LOG(LogDerivedDataCache, Verbose, TEXT("%s: Cache put complete for %s from '%s'"),
+				*CachePath, *WriteToString<96>(Request.Key), *Request.Name);
+			COOK_STAT(if (WriteSize) { Timer.AddHit(WriteSize); });
+			OnComplete(Request.MakeResponse(EStatus::Ok));
+		}
+		else
+		{
+			OnComplete(Request.MakeResponse(EStatus::Error));
+		}
+	}
+}
+
+void FPakFileCacheStore::GetValue(
+	const TConstArrayView<FCacheGetValueRequest> Requests,
+	IRequestOwner& Owner,
+	FOnCacheGetValueComplete&& OnComplete)
+{
+	for (const FCacheGetValueRequest& Request : Requests)
+	{
+		TRACE_CPUPROFILER_EVENT_SCOPE(PakFileDDC_Get);
+		COOK_STAT(auto Timer = UsageStats.TimeGet());
+		FValue Value;
+		if (GetCacheValue(Request.Name, Request.Key, Request.Policy, Value))
+		{
+			UE_LOG(LogDerivedDataCache, Verbose, TEXT("%s: Cache hit for %s from '%s'"),
+				*CachePath, *WriteToString<96>(Request.Key), *Request.Name);
+			COOK_STAT(Timer.AddHit(Value.GetData().GetCompressedSize()));
+			OnComplete({Request.Name, Request.Key, Value, Request.UserData, EStatus::Ok});
+		}
+		else
+		{
+			OnComplete(Request.MakeResponse(EStatus::Error));
+		}
+	}
+}
+
+void FPakFileCacheStore::GetChunks(
+	const TConstArrayView<FCacheGetChunkRequest> Requests,
+	IRequestOwner& Owner,
+	FOnCacheGetChunkComplete&& OnComplete)
+{
+	TArray<FCacheGetChunkRequest, TInlineAllocator<16>> SortedRequests(Requests);
+	SortedRequests.StableSort(TChunkLess());
+
+	bool bHasValue = false;
+	FValue Value;
+	FValueId ValueId;
+	FCacheKey ValueKey;
+	TUniquePtr<FArchive> ValueAr;
+	FCompressedBufferReader ValueReader;
+	FOptionalCacheRecord Record;
+	for (const FCacheGetChunkRequest& Request : SortedRequests)
+	{
+		const bool bExistsOnly = EnumHasAnyFlags(Request.Policy, ECachePolicy::SkipData);
+		TRACE_CPUPROFILER_EVENT_SCOPE(PakFileDDC_Get);
+		COOK_STAT(auto Timer = bExistsOnly ? UsageStats.TimeProbablyExists() : UsageStats.TimeGet());
+		if (!(bHasValue && ValueKey == Request.Key && ValueId == Request.Id) || ValueReader.HasSource() < !bExistsOnly)
+		{
+			ValueReader.ResetSource();
+			ValueAr.Reset();
+			ValueKey = {};
+			ValueId.Reset();
+			Value.Reset();
+			bHasValue = false;
+			if (Request.Id.IsValid())
+			{
+				if (!(Record && Record.Get().GetKey() == Request.Key))
+				{
+					FCacheRecordPolicyBuilder PolicyBuilder(ECachePolicy::None);
+					PolicyBuilder.AddValuePolicy(Request.Id, Request.Policy);
+					Record.Reset();
+					Record = GetCacheRecordOnly(Request.Name, Request.Key, PolicyBuilder.Build());
+				}
+				if (Record)
+				{
+					if (const FValueWithId& ValueWithId = Record.Get().GetValue(Request.Id))
+					{
+						bHasValue = true;
+						Value = ValueWithId;
+						ValueId = Request.Id;
+						ValueKey = Request.Key;
+						GetCacheContent(Request.Name, Request.Key, ValueId, Value, Request.Policy, ValueReader, ValueAr);
+					}
+				}
+			}
+			else
+			{
+				ValueKey = Request.Key;
+				bHasValue = GetCacheValueOnly(Request.Name, Request.Key, Request.Policy, Value);
+				if (bHasValue)
+				{
+					GetCacheContent(Request.Name, Request.Key, Request.Id, Value, Request.Policy, ValueReader, ValueAr);
+				}
+			}
+		}
+		if (bHasValue)
+		{
+			const uint64 RawOffset = FMath::Min(Value.GetRawSize(), Request.RawOffset);
+			const uint64 RawSize = FMath::Min(Value.GetRawSize() - RawOffset, Request.RawSize);
+			UE_LOG(LogDerivedDataCache, Verbose, TEXT("%s: Cache hit for %s from '%s'"),
+				*CachePath, *WriteToString<96>(Request.Key, '/', Request.Id), *Request.Name);
+			COOK_STAT(Timer.AddHit(!bExistsOnly ? RawSize : 0));
+			FSharedBuffer Buffer;
+			if (!bExistsOnly)
+			{
+				Buffer = ValueReader.Decompress(RawOffset, RawSize);
+			}
+			const EStatus ChunkStatus = bExistsOnly || Buffer.GetSize() == RawSize ? EStatus::Ok : EStatus::Error;
+			OnComplete({Request.Name, Request.Key, Request.Id, Request.RawOffset,
+				RawSize, Value.GetRawHash(), MoveTemp(Buffer), Request.UserData, ChunkStatus});
+			continue;
+		}
+
+		OnComplete(Request.MakeResponse(EStatus::Error));
+	}
+}
+
+bool FPakFileCacheStore::PutCacheRecord(
+	const FStringView Name,
+	const FCacheRecord& Record,
+	const FCacheRecordPolicy& Policy,
+	uint64& OutWriteSize)
+{
+	if (!IsWritable())
+	{
+		return false;
+	}
+
+	const FCacheKey& Key = Record.GetKey();
+	const ECachePolicy RecordPolicy = Policy.GetRecordPolicy();
+
+	// Skip the request if storing to the cache is disabled.
+	if (!EnumHasAnyFlags(Policy.GetRecordPolicy(), ECachePolicy::StoreLocal))
+	{
+		UE_LOG(LogDerivedDataCache, VeryVerbose, TEXT("%s: Skipped put of %s from '%.*s' due to cache policy"),
+			*CachePath, *WriteToString<96>(Key), Name.Len(), Name.GetData());
+		return false;
+	}
+
+	//if (ShouldSimulateMiss(Key))
+	//{
+	//	UE_LOG(LogDerivedDataCache, Verbose, TEXT("%s: Simulated miss for put of %s from '%.*s'"),
+	//		*CachePath, *WriteToString<96>(Key), Name.Len(), Name.GetData());
+	//	return false;
+	//}
+
+	TStringBuilder<256> Path;
+	FPathViews::Append(Path, TEXT("Buckets"), Key);
+
+	// Check if there is an existing record package.
+	bool bReplaceExisting = !EnumHasAnyFlags(RecordPolicy, ECachePolicy::QueryLocal);
+	bool bSaveRecord = bReplaceExisting;
+	if (!bReplaceExisting)
+	{
+		bSaveRecord |= !FileExists(Path);
+	}
+
+	// Serialize the record to a package and remove attachments that will be stored externally.
+	FCbPackage Package = Record.Save();
+	TArray<FCompressedBuffer, TInlineAllocator<8>> ExternalContent;
+	Algo::Transform(Package.GetAttachments(), ExternalContent, &FCbAttachment::AsCompressedBinary);
+	Package = FCbPackage(Package.GetObject());
+
+	// Save the external content to storage.
+	for (FCompressedBuffer& Content : ExternalContent)
+	{
+		uint64 WriteSize = 0;
+		if (!PutCacheContent(Name, Content, WriteSize))
+		{
+			return false;
+		}
+		OutWriteSize += WriteSize;
+	}
+
+	// Save the record package to storage.
+	const auto WriteRecord = [&](FArchive& Ar) { Package.Save(Ar); OutWriteSize += uint64(Ar.TotalSize()); };
+	if (bSaveRecord && !SaveFile(Path, Name, WriteRecord))
+	{
+		return false;
+	}
+
+	return true;
+}
+
+FOptionalCacheRecord FPakFileCacheStore::GetCacheRecordOnly(
+	const FStringView Name,
+	const FCacheKey& Key,
+	const FCacheRecordPolicy& Policy)
+{
+	if (bClosed)
+	{
+		UE_LOG(LogDerivedDataCache, VeryVerbose,
+			TEXT("%s: Skipped get of %s from '%.*s' because this cache store is not available"),
+			*CachePath, *WriteToString<96>(Key), Name.Len(), Name.GetData());
+		return FOptionalCacheRecord();
+	}
+
+	// Skip the request if querying the cache is disabled.
+	if (!EnumHasAnyFlags(Policy.GetRecordPolicy(), ECachePolicy::QueryLocal))
+	{
+		UE_LOG(LogDerivedDataCache, VeryVerbose, TEXT("%s: Skipped get of %s from '%.*s' due to cache policy"),
+			*CachePath, *WriteToString<96>(Key), Name.Len(), Name.GetData());
+		return FOptionalCacheRecord();
+	}
+
+	//if (ShouldSimulateMiss(Key))
+	//{
+	//	UE_LOG(LogDerivedDataCache, Verbose, TEXT("%s: Simulated miss for get of %s from '%.*s'"),
+	//		*CachePath, *WriteToString<96>(Key), Name.Len(), Name.GetData());
+	//	return FOptionalCacheRecord();
+	//}
+
+	TStringBuilder<256> Path;
+	FPathViews::Append(Path, TEXT("Buckets"), Key);
+
+	// Request the record from storage.
+	FSharedBuffer Buffer = LoadFile(Path, Name);
+	if (Buffer.IsNull())
+	{
+		UE_LOG(LogDerivedDataCache, Verbose, TEXT("%s: Cache miss with missing record for %s from '%.*s'"),
+			*CachePath, *WriteToString<96>(Key), Name.Len(), Name.GetData());
+		return FOptionalCacheRecord();
+	}
+
+	// Validate that the record can be read as a compact binary package without crashing.
+	if (ValidateCompactBinaryPackage(Buffer, ECbValidateMode::Default | ECbValidateMode::Package) != ECbValidateError::None)
+	{
+		UE_LOG(LogDerivedDataCache, Display, TEXT("%s: Cache miss with invalid package for %s from '%.*s'"),
+			*CachePath, *WriteToString<96>(Key), Name.Len(), Name.GetData());
+		return FOptionalCacheRecord();
+	}
+
+	// Load the record from the package.
+	FOptionalCacheRecord Record;
+	{
+		FCbPackage Package;
+		if (FCbFieldIterator It = FCbFieldIterator::MakeRange(Buffer); !Package.TryLoad(It))
+		{
+			UE_LOG(LogDerivedDataCache, Display, TEXT("%s: Cache miss with package load failure for %s from '%.*s'"),
+				*CachePath, *WriteToString<96>(Key), Name.Len(), Name.GetData());
+			return FOptionalCacheRecord();
+		}
+		Record = FCacheRecord::Load(Package);
+		if (Record.IsNull())
+		{
+			UE_LOG(LogDerivedDataCache, Display, TEXT("%s: Cache miss with record load failure for %s from '%.*s'"),
+				*CachePath, *WriteToString<96>(Key), Name.Len(), Name.GetData());
+			return FOptionalCacheRecord();
+		}
+	}
+
+	return Record.Get();
+}
+
+FOptionalCacheRecord FPakFileCacheStore::GetCacheRecord(
+	const FStringView Name,
+	const FCacheKey& Key,
+	const FCacheRecordPolicy& Policy,
+	EStatus& OutStatus)
+{
+	FOptionalCacheRecord Record = GetCacheRecordOnly(Name, Key, Policy);
+	if (Record.IsNull())
+	{
+		OutStatus = EStatus::Error;
+		return Record;
+	}
+
+	OutStatus = EStatus::Ok;
+
+	FCacheRecordBuilder RecordBuilder(Key);
+
+	const ECachePolicy RecordPolicy = Policy.GetRecordPolicy();
+	if (!EnumHasAnyFlags(RecordPolicy, ECachePolicy::SkipMeta))
+	{
+		RecordBuilder.SetMeta(FCbObject(Record.Get().GetMeta()));
+	}
+
+	for (const FValueWithId& Value : Record.Get().GetValues())
+	{
+		const FValueId& Id = Value.GetId();
+		const ECachePolicy ValuePolicy = Policy.GetValuePolicy(Id);
+		FValue Content;
+		if (GetCacheContent(Name, Key, Id, Value, ValuePolicy, Content))
+		{
+			RecordBuilder.AddValue(Id, MoveTemp(Content));
+		}
+		else if (EnumHasAnyFlags(RecordPolicy, ECachePolicy::PartialRecord))
+		{
+			OutStatus = EStatus::Error;
+			RecordBuilder.AddValue(Value);
+		}
+		else
+		{
+			OutStatus = EStatus::Error;
+			return FOptionalCacheRecord();
+		}
+	}
+
+	return RecordBuilder.Build();
+}
+
+bool FPakFileCacheStore::PutCacheValue(
+	const FStringView Name,
+	const FCacheKey& Key,
+	const FValue& Value,
+	const ECachePolicy Policy,
+	uint64& OutWriteSize)
+{
+	if (!IsWritable())
+	{
+		return false;
+	}
+
+	// Skip the request if storing to the cache is disabled.
+	if (!EnumHasAnyFlags(Policy, ECachePolicy::StoreLocal))
+	{
+		UE_LOG(LogDerivedDataCache, VeryVerbose, TEXT("%s: Skipped put of %s from '%.*s' due to cache policy"),
+			*CachePath, *WriteToString<96>(Key), Name.Len(), Name.GetData());
+		return false;
+	}
+
+	//if (ShouldSimulateMiss(Key))
+	//{
+	//	UE_LOG(LogDerivedDataCache, Verbose, TEXT("%s: Simulated miss for put of %s from '%.*s'"),
+	//		*CachePath, *WriteToString<96>(Key), Name.Len(), Name.GetData());
+	//	return false;
+	//}
+
+	// Check if there is an existing value package.
+	bool bValueExists = false;
+	TStringBuilder<256> Path;
+	FPathViews::Append(Path, TEXT("Buckets"), Key);
+	const bool bReplaceExisting = !EnumHasAnyFlags(Policy, ECachePolicy::QueryLocal);
+	if (!bReplaceExisting)
+	{
+		bValueExists = FileExists(Path);
+	}
+
+	// Save the value to a package and save the data to external content.
+	if (!bValueExists)
+	{
+		FCbWriter Writer;
+		Writer.BeginObject();
+		Writer.AddBinaryAttachment("RawHash", Value.GetRawHash());
+		Writer.AddInteger("RawSize", Value.GetRawSize());
+		Writer.EndObject();
+
+		FCbPackage Package(Writer.Save().AsObject());
+		if (!Value.HasData())
+		{
+			// Verify that the content exists in storage.
+			if (!GetCacheContentExists(Key, Value.GetRawHash()))
+			{
+				UE_LOG(LogDerivedDataCache, Verbose, TEXT("%s: Failed due to missing data for put of %s from '%.*s'"),
+					*CachePath, *WriteToString<96>(Key), Name.Len(), Name.GetData());
+				return false;
+			}
+		}
+		else
+		{
+			// Save the external content to storage.
+			uint64 WriteSize = 0;
+			if (!PutCacheContent(Name, Value.GetData(), WriteSize))
+			{
+				return false;
+			}
+			OutWriteSize += WriteSize;
+		}
+
+		// Save the value package to storage.
+		const auto WritePackage = [&](FArchive& Ar) { Package.Save(Ar); OutWriteSize += uint64(Ar.TotalSize()); };
+		if (!SaveFile(Path, Name, WritePackage))
+		{
+			return false;
+		}
+	}
+
+	return true;
+}
+
+bool FPakFileCacheStore::GetCacheValueOnly(
+	const FStringView Name,
+	const FCacheKey& Key,
+	const ECachePolicy Policy,
+	FValue& OutValue)
+{
+	if (bClosed)
+	{
+		UE_LOG(LogDerivedDataCache, VeryVerbose,
+			TEXT("%s: Skipped get of %s from '%.*s' because this cache store is not available"),
+			*CachePath, *WriteToString<96>(Key), Name.Len(), Name.GetData());
+		return false;
+	}
+
+	// Skip the request if querying the cache is disabled.
+	if (!EnumHasAnyFlags(Policy, ECachePolicy::QueryLocal))
+	{
+		UE_LOG(LogDerivedDataCache, VeryVerbose, TEXT("%s: Skipped get of %s from '%.*s' due to cache policy"),
+			*CachePath, *WriteToString<96>(Key), Name.Len(), Name.GetData());
+		return false;
+	}
+
+	//if (ShouldSimulateMiss(Key))
+	//{
+	//	UE_LOG(LogDerivedDataCache, Verbose, TEXT("%s: Simulated miss for get of %s from '%.*s'"),
+	//		*CachePath, *WriteToString<96>(Key), Name.Len(), Name.GetData());
+	//	return false;
+	//}
+
+	TStringBuilder<256> Path;
+	FPathViews::Append(Path, TEXT("Buckets"), Key);
+
+	// Request the value package from storage.
+	FSharedBuffer Buffer = LoadFile(Path, Name);
+	if (Buffer.IsNull())
+	{
+		UE_LOG(LogDerivedDataCache, Verbose, TEXT("%s: Cache miss with missing value for %s from '%.*s'"),
+			*CachePath, *WriteToString<96>(Key), Name.Len(), Name.GetData());
+		return false;
+	}
+
+	if (ValidateCompactBinary(Buffer, ECbValidateMode::Default | ECbValidateMode::Package) != ECbValidateError::None)
+	{
+		UE_LOG(LogDerivedDataCache, Display, TEXT("%s: Cache miss with invalid package for %s from '%.*s'"),
+			*CachePath, *WriteToString<96>(Key), Name.Len(), Name.GetData());
+		return false;
+	}
+
+	FCbPackage Package;
+	if (FCbFieldIterator It = FCbFieldIterator::MakeRange(Buffer); !Package.TryLoad(It))
+	{
+		UE_LOG(LogDerivedDataCache, Display, TEXT("%s: Cache miss with package load failure for %s from '%.*s'"),
+			*CachePath, *WriteToString<96>(Key), Name.Len(), Name.GetData());
+		return false;
+	}
+
+	const FCbObjectView Object = Package.GetObject();
+	const FIoHash RawHash = Object["RawHash"].AsHash();
+	const uint64 RawSize = Object["RawSize"].AsUInt64(MAX_uint64);
+	if (RawHash.IsZero() || RawSize == MAX_uint64)
+	{
+		UE_LOG(LogDerivedDataCache, Display, TEXT("%s: Cache miss with invalid value for %s from '%.*s'"),
+			*CachePath, *WriteToString<96>(Key), Name.Len(), Name.GetData());
+		return false;
+	}
+
+	OutValue = FValue(RawHash, RawSize);
+
+	return true;
+}
+
+bool FPakFileCacheStore::GetCacheValue(
+	const FStringView Name,
+	const FCacheKey& Key,
+	const ECachePolicy Policy,
+	FValue& OutValue)
+{
+	return GetCacheValueOnly(Name, Key, Policy, OutValue) && GetCacheContent(Name, Key, {}, OutValue, Policy, OutValue);
+}
+
+bool FPakFileCacheStore::PutCacheContent(const FStringView Name, const FCompressedBuffer& Content, uint64& OutWriteSize)
+{
+	const FIoHash& RawHash = Content.GetRawHash();
+	TStringBuilder<256> Path;
+	FPathViews::Append(Path, TEXT("Content"), RawHash);
+	if (!FileExists(Path))
+	{
+		if (!SaveFile(Path, Name, [&Content, &OutWriteSize](FArchive& Ar) { Content.Save(Ar); OutWriteSize += uint64(Ar.TotalSize()); }))
+		{
+			return false;
+		}
+	}
+	return true;
+}
+
+bool FPakFileCacheStore::GetCacheContentExists(const FCacheKey& Key, const FIoHash& RawHash)
+{
+	TStringBuilder<256> Path;
+	FPathViews::Append(Path, TEXT("Buckets"), Key);
+	return FileExists(Path);
+}
+
+bool FPakFileCacheStore::GetCacheContent(
+	const FStringView Name,
+	const FCacheKey& Key,
+	const FValueId& Id,
+	const FValue& Value,
+	const ECachePolicy Policy,
+	FValue& OutValue)
+{
+	if (!EnumHasAnyFlags(Policy, ECachePolicy::Query))
+	{
+		OutValue = Value.RemoveData();
+		return true;
+	}
+
+	if (Value.HasData())
+	{
+		OutValue = EnumHasAnyFlags(Policy, ECachePolicy::SkipData) ? Value.RemoveData() : Value;
+		return true;
+	}
+
+	const FIoHash& RawHash = Value.GetRawHash();
+
+	TStringBuilder<256> Path;
+	FPathViews::Append(Path, TEXT("Content"), RawHash);
+	if (EnumHasAnyFlags(Policy, ECachePolicy::SkipData))
+	{
+		if (FileExists(Path))
+		{
+			OutValue = Value;
+			return true;
+		}
+	}
+	else
+	{
+		if (FSharedBuffer CompressedData = LoadFile(Path, Name))
+		{
+			if (FCompressedBuffer CompressedBuffer = FCompressedBuffer::FromCompressed(MoveTemp(CompressedData));
+				CompressedBuffer && CompressedBuffer.GetRawHash() == RawHash)
+			{
+				OutValue = FValue(MoveTemp(CompressedBuffer));
+				return true;
+			}
+			UE_LOG(LogDerivedDataCache, Display,
+				TEXT("%s: Cache miss with corrupted value %s with hash %s for %s from '%.*s'"),
+				*CachePath, *WriteToString<16>(Id), *WriteToString<48>(RawHash),
+				*WriteToString<96>(Key), Name.Len(), Name.GetData());
+			return false;
+		}
+	}
+
+	UE_LOG(LogDerivedDataCache, Verbose,
+		TEXT("%s: Cache miss with missing value %s with hash %s for %s from '%.*s'"),
+		*CachePath, *WriteToString<16>(Id), *WriteToString<48>(RawHash), *WriteToString<96>(Key),
+		Name.Len(), Name.GetData());
+	return false;
+}
+
+void FPakFileCacheStore::GetCacheContent(
+	const FStringView Name,
+	const FCacheKey& Key,
+	const FValueId& Id,
+	const FValue& Value,
+	const ECachePolicy Policy,
+	FCompressedBufferReader& Reader,
+	TUniquePtr<FArchive>& OutArchive)
+{
+	if (!EnumHasAnyFlags(Policy, ECachePolicy::Query))
+	{
+		return;
+	}
+
+	if (Value.HasData())
+	{
+		if (!EnumHasAnyFlags(Policy, ECachePolicy::SkipData))
+		{
+			Reader.SetSource(Value.GetData());
+		}
+		OutArchive.Reset();
+		return;
+	}
+
+	const FIoHash& RawHash = Value.GetRawHash();
+
+	TStringBuilder<256> Path;
+	FPathViews::Append(Path, TEXT("Content"), RawHash);
+	if (EnumHasAllFlags(Policy, ECachePolicy::SkipData))
+	{
+		if (FileExists(Path))
+		{
+			return;
+		}
+	}
+	else
+	{
+		OutArchive = OpenFile(Path, Name);
+		if (OutArchive)
+		{
+			Reader.SetSource(*OutArchive);
+			if (Reader.GetRawHash() == RawHash)
+			{
+				return;
+			}
+			UE_LOG(LogDerivedDataCache, Display,
+				TEXT("%s: Cache miss with corrupted value %s with hash %s for %s from '%.*s'"),
+				*CachePath, *WriteToString<16>(Id), *WriteToString<48>(RawHash),
+				*WriteToString<96>(Key), Name.Len(), Name.GetData());
+			Reader.ResetSource();
+			OutArchive.Reset();
+			return;
+		}
+	}
+
+	UE_LOG(LogDerivedDataCache, Verbose,
+		TEXT("%s: Cache miss with missing value %s with hash %s for %s from '%.*s'"),
+		*CachePath, *WriteToString<16>(Id), *WriteToString<48>(RawHash), *WriteToString<96>(Key),
+		Name.Len(), Name.GetData());
+}
+
+class FCrcBuilder
+{
+public:
+	inline void Update(const void* Data, uint64 Size)
+	{
+		while (Size > 0)
+		{
+			const int32 CrcSize = int32(FMath::Min<uint64>(Size, MAX_int32));
+			Crc = FCrc::MemCrc_DEPRECATED(Data, CrcSize, Crc);
+			Size -= CrcSize;
+		}
+	}
+
+	inline uint32 Finalize()
+	{
+		return Crc;
+	}
+
+private:
+	uint32 Crc = 0;
+};
+
+class FPakWriterArchive final : public FArchive
+{
+public:
+	inline FPakWriterArchive(IFileHandle& InHandle, FStringView InPath)
+		: Handle(InHandle)
+		, Path(InPath)
+	{
+		SetIsSaving(true);
+		SetIsPersistent(true);
+	}
+
+	inline FString GetArchiveName() const final { return FString(Path); }
+	inline int64 TotalSize() final { return Handle.Size(); }
+	inline int64 Tell() final { unimplemented(); return 0; }
+	inline void Seek(int64 InPos) final { unimplemented(); }
+	inline void Flush() final { unimplemented(); }
+	inline bool Close() final { unimplemented(); return false; }
+
+	inline void Serialize(void* V, int64 Length) final
+	{
+		if (!Handle.Write(static_cast<uint8*>(V), Length))
+		{
+			SetError();
+		}
+	}
+
+private:
+	IFileHandle& Handle;
+	FStringView Path;
+};
+
+class FPakReaderArchive final : public FArchive
+{
+public:
+	inline FPakReaderArchive(IFileHandle& InHandle, FStringView InPath)
+		: Handle(InHandle)
+		, Path(InPath)
+	{
+		SetIsLoading(true);
+		SetIsPersistent(true);
+	}
+
+	inline FString GetArchiveName() const final { return FString(Path); }
+	inline int64 TotalSize() final { return Handle.Size(); }
+	inline int64 Tell() final { unimplemented(); return 0; }
+	inline void Seek(int64 InPos) final { unimplemented(); }
+	inline void Flush() final { unimplemented(); }
+	inline bool Close() final { unimplemented(); return false; }
+
+	inline void Serialize(void* V, int64 Length) final
+	{
+		if (!Handle.Read(static_cast<uint8*>(V), Length))
+		{
+			SetError();
+		}
+	}
+
+private:
+	IFileHandle& Handle;
+	FStringView Path;
+};
+
+bool FPakFileCacheStore::SaveFile(
+	const FStringView Path,
+	const FStringView DebugName,
+	TFunctionRef<void (FArchive&)> WriteFunction)
+{
+	FWriteScopeLock ScopeLock(SynchronizationObject);
+	check(FileHandle);
+	if (const int64 Offset = FileHandle->Tell(); Offset >= 0)
+	{
+		FPakWriterArchive Ar(*FileHandle, CachePath);
+		THashingArchiveProxy<FCrcBuilder> HashAr(Ar);
+		WriteFunction(HashAr);
+		if (const int64 EndOffset = FileHandle->Tell(); EndOffset >= Offset && !Ar.IsError())
+		{
+			FCacheValue& Item = CacheItems.Emplace(Path, FCacheValue(Offset, EndOffset - Offset, HashAr.GetHash()));
+			UE_LOG(LogDerivedDataCache, Log,
+				TEXT("%s: File %.*s from '%.*s' written with offset %" INT64_FMT ", size %" INT64_FMT", CRC 0x%08x."),
+				*CachePath, Path.Len(), Path.GetData(), DebugName.Len(), DebugName.GetData(), Item.Offset, Item.Size, Item.Crc);
+			return true;
+		}
+	}
+	return false;
+}
+
+FSharedBuffer FPakFileCacheStore::LoadFile(const FStringView Path, const FStringView DebugName)
+{
+	FWriteScopeLock ScopeLock(SynchronizationObject);
+	if (const FCacheValue* Item = CacheItems.FindByHash(GetTypeHash(Path), Path))
+	{
+		check(FileHandle);
+		ON_SCOPE_EXIT
+		{
+			if (bWriting)
+			{
+				FileHandle->SeekFromEnd();
+			}
+		};
+		check(Item->Size);
+		if (FileHandle->Seek(Item->Offset))
+		{
+			FPakReaderArchive Ar(*FileHandle, CachePath);
+			THashingArchiveProxy<FCrcBuilder> HashAr(Ar);
+			FUniqueBuffer MutableBuffer = FUniqueBuffer::Alloc(uint64(Item->Size));
+			HashAr.Serialize(MutableBuffer.GetData(), Item->Size);
+			if (Ar.IsError())
+			{
+				UE_LOG(LogDerivedDataCache, Display,
+					TEXT("%s: File %.*s from '%.*s' failed to read %" INT64_FMT " bytes."),
+					*CachePath, Path.Len(), Path.GetData(), DebugName.Len(), DebugName.GetData(), Item->Size);
+			}
+			else if (const uint32 TestCrc = HashAr.GetHash(); TestCrc != Item->Crc)
+			{
+				UE_LOG(LogDerivedDataCache, Display,
+					TEXT("%s: File %.*s from '%.*s' is corrupted and has CRC 0x%08x when 0x%08x is expected."),
+					*CachePath, Path.Len(), Path.GetData(), DebugName.Len(), DebugName.GetData(), TestCrc, Item->Crc);
+			}
+			else
+			{
+				return MutableBuffer.MoveToShared();
+			}
+		}
+	}
+	return FSharedBuffer();
+}
+
+TUniquePtr<FArchive> FPakFileCacheStore::OpenFile(FStringBuilderBase& Path, const FStringView DebugName)
+{
+	FReadScopeLock ScopeLock(SynchronizationObject);
+	const FStringView PathView(Path);
+	if (const FCacheValue* Item = CacheItems.FindByHash(GetTypeHash(PathView), PathView))
+	{
+		check(Item->Size);
+		if (TUniquePtr<FArchive> Ar{IFileManager::Get().CreateFileReader(*CachePath, FILEREAD_Silent | FILEREAD_AllowWrite)})
+		{
+			Ar->Seek(Item->Offset);
+			return Ar;
+		}
+	}
+	return nullptr;
+}
+
+bool FPakFileCacheStore::FileExists(const FStringView Path)
+{
+	FReadScopeLock ScopeLock(SynchronizationObject);
+	const uint32 PathHash = GetTypeHash(Path);
+	return CacheItems.ContainsByHash(PathHash, Path);
+}
+
+static void ScheduleAsyncRequest(IRequestOwner& Owner, const TCHAR* DebugName, TUniqueFunction<void (IRequestOwner& Owner)>&& Function)
+{
+	class FAsyncRequest final : public FRequestBase
+	{
+	public:
+		Tasks::FTask Task;
+		TUniqueFunction<void (IRequestOwner& Owner)> Function;
+
+		void SetPriority(EPriority Priority) final {}
+		void Cancel() final { Task.Wait(); }
+		void Wait() final { Task.Wait(); }
+	};
+
+	FAsyncRequest* Request = new FAsyncRequest;
+	Request->Function = MoveTemp(Function);
+	Tasks::FTaskEvent TaskEvent(TEXT("ScheduleAsyncRequest"));
+	Request->Task = Tasks::Launch(DebugName,
+		[Request, &Owner] { Owner.End(Request, [Request, &Owner] { Request->Function(Owner); }); },
+		TaskEvent, Tasks::ETaskPriority::BackgroundNormal);
+	Owner.Begin(Request);
+	TaskEvent.Trigger();
+}
+
+class FCompressedPakFileCacheStore final : public FPakFileCacheStore
+{
+public:
+	FCompressedPakFileCacheStore(const TCHAR* InFilename, bool bInWriting);
+
+	EPutStatus PutCachedData(const TCHAR* CacheKey, TArrayView<const uint8> InData, bool bPutEvenIfExists) final;
+	bool GetCachedData(const TCHAR* CacheKey, TArray<uint8>& OutData) final;
+
+	void Put(
+		TConstArrayView<FCachePutRequest> Requests,
+		IRequestOwner& Owner,
+		FOnCachePutComplete&& OnComplete) final;
+
+	void PutValue(
+		TConstArrayView<FCachePutValueRequest> Requests,
+		IRequestOwner& Owner,
+		FOnCachePutValueComplete&& OnComplete) final;
+
+private:
+	static const EName CompressionFormat = NAME_Zlib;
+	static const ECompressionFlags CompressionFlags = COMPRESS_BiasMemory;
+	static const ECompressedBufferCompressor RequiredCompressor = ECompressedBufferCompressor::Kraken;
+	static const ECompressedBufferCompressionLevel MinRequiredCompressionLevel = ECompressedBufferCompressionLevel::Optimal2;
+
+	static FValue Compress(const FValue& Value);
+};
+
+FCompressedPakFileCacheStore::FCompressedPakFileCacheStore(const TCHAR* InFilename, bool bInWriting)
+	: FPakFileCacheStore(InFilename, bInWriting)
+{
+}
+
+FDerivedDataBackendInterface::EPutStatus FCompressedPakFileCacheStore::PutCachedData(const TCHAR* CacheKey, TArrayView<const uint8> InData, bool bPutEvenIfExists)
+{
+	int32 UncompressedSize = InData.Num();
+	int32 CompressedSize = FCompression::CompressMemoryBound(CompressionFormat, UncompressedSize, CompressionFlags);
+
+	TArray<uint8> CompressedData;
+	CompressedData.AddUninitialized(CompressedSize + sizeof(UncompressedSize));
+
+	FMemory::Memcpy(&CompressedData[0], &UncompressedSize, sizeof(UncompressedSize));
+	verify(FCompression::CompressMemory(CompressionFormat, CompressedData.GetData() + sizeof(UncompressedSize), CompressedSize, InData.GetData(), InData.Num(), CompressionFlags));
+	CompressedData.SetNum(CompressedSize + sizeof(UncompressedSize), false);
+
+	return FPakFileCacheStore::PutCachedData(CacheKey, CompressedData, bPutEvenIfExists);
+}
+
+bool FCompressedPakFileCacheStore::GetCachedData(const TCHAR* CacheKey, TArray<uint8>& OutData)
+{
+	TArray<uint8> CompressedData;
+	if(!FPakFileCacheStore::GetCachedData(CacheKey, CompressedData))
+	{
+		return false;
+	}
+
+	int32 UncompressedSize;
+	FMemory::Memcpy(&UncompressedSize, &CompressedData[0], sizeof(UncompressedSize));
+	OutData.SetNum(UncompressedSize);
+	verify(FCompression::UncompressMemory(CompressionFormat, OutData.GetData(), UncompressedSize, CompressedData.GetData() + sizeof(UncompressedSize), CompressedData.Num() - sizeof(UncompressedSize), CompressionFlags));
+
+	return true;
+}
+
+void FCompressedPakFileCacheStore::Put(
+	const TConstArrayView<FCachePutRequest> Requests,
+	IRequestOwner& Owner,
+	FOnCachePutComplete&& OnComplete)
+{
+	ScheduleAsyncRequest(Owner, TEXT("PakFileDDC_Put"),
+		[this, Requests = TArray<FCachePutRequest, TInlineAllocator<1>>(Requests), OnComplete = MoveTemp(OnComplete)](IRequestOwner& Owner) mutable
+		{
+			for (FCachePutRequest& Request : Requests)
+			{
+				FCacheRecordBuilder Builder(Request.Record.GetKey());
+				Builder.SetMeta(CopyTemp(Request.Record.GetMeta()));
+				for (const FValueWithId& Value : Request.Record.GetValues())
+				{
+					Builder.AddValue(Value.GetId(), Compress(Value));
+				}
+				Request.Record = Builder.Build();
+			}
+			Private::ExecuteInCacheThreadPool(Owner,
+				[this, Requests = MoveTemp(Requests), OnComplete = MoveTemp(OnComplete)](IRequestOwner& Owner, bool bCancel) mutable
+				{
+					FPakFileCacheStore::Put(Requests, Owner, MoveTemp(OnComplete));
+				});
+		});
+}
+
+void FCompressedPakFileCacheStore::PutValue(
+	const TConstArrayView<FCachePutValueRequest> Requests,
+	IRequestOwner& Owner,
+	FOnCachePutValueComplete&& OnComplete)
+{
+	ScheduleAsyncRequest(Owner, TEXT("PakFileDDC_PutValue"),
+		[this, Requests = TArray<FCachePutValueRequest, TInlineAllocator<1>>(Requests), OnComplete = MoveTemp(OnComplete)](IRequestOwner& Owner) mutable
+		{
+			for (FCachePutValueRequest& Request : Requests)
+			{
+				Request.Value = Compress(Request.Value);
+			}
+			Private::ExecuteInCacheThreadPool(Owner,
+				[this, Requests = MoveTemp(Requests), OnComplete = MoveTemp(OnComplete)](IRequestOwner& Owner, bool bCancel) mutable
+				{
+					FPakFileCacheStore::PutValue(Requests, Owner, MoveTemp(OnComplete));
+				});
+		});
+}
+
+FValue FCompressedPakFileCacheStore::Compress(const FValue& Value)
+{
+	uint64 BlockSize = 0;
+	ECompressedBufferCompressor Compressor;
+	ECompressedBufferCompressionLevel CompressionLevel;
+	if (!Value.HasData() ||
+		(Value.GetData().TryGetCompressParameters(Compressor, CompressionLevel, BlockSize) &&
+			Compressor == RequiredCompressor &&
+			CompressionLevel >= MinRequiredCompressionLevel))
+	{
+		return Value;
+	}
+	TRACE_CPUPROFILER_EVENT_SCOPE(PakFileDDC_Compress);
+	const FCompositeBuffer Data = Value.GetData().DecompressToComposite();
+	return FValue(FCompressedBuffer::Compress(Data, RequiredCompressor, MinRequiredCompressionLevel, BlockSize));
+}
+
+IPakFileCacheStore* CreatePakFileCacheStore(const TCHAR* Filename, bool bWriting, bool bCompressed)
+{
+	return bCompressed ? new FCompressedPakFileCacheStore(Filename, bWriting) : new FPakFileCacheStore(Filename, bWriting);
+}
+
+} // UE::DerivedData
diff --git a/Engine/Source/Developer/DerivedDataCache/Private/PakFileCacheStore.h b/Engine/Source/Developer/DerivedDataCache/Private/PakFileCacheStore.h
new file mode 100644
--- /dev/null
+++ b/Engine/Source/Developer/DerivedDataCache/Private/PakFileCacheStore.h
@@ -0,0 +1,24 @@
+// Copyright Epic Games, Inc. All Rights Reserved.
+
+#pragma once
+
+#include "CoreTypes.h"
+#include "Containers/StringFwd.h"
+#include "DerivedDataBackendInterface.h"
+
+namespace UE::DerivedData
+{
+
+class IPakFileCacheStore : public FDerivedDataBackendInterface
+{
+public:
+	virtual void Close() = 0;
+	virtual bool SaveCache() = 0;
+	virtual bool LoadCache(const TCHAR* InFilename) = 0;
+	virtual void MergeCache(IPakFileCacheStore* OtherPak) = 0;
+	virtual const FString& GetFilename() const = 0;
+
+	static bool SortAndCopy(const FString& InputFilename, const FString& OutputFilename);
+};
+
+} // UE::DerivedData
