DDC: Enabled compression of legacy cache data

diff --git a/Engine/Source/Developer/DerivedDataCache/Private/S3CacheStore.cpp b/Engine/Source/Developer/DerivedDataCache/Private/S3CacheStore.cpp
new file mode 100644
--- /dev/null
+++ b/Engine/Source/Developer/DerivedDataCache/Private/S3CacheStore.cpp
@@ -0,0 +1,1979 @@
+// Copyright Epic Games, Inc. All Rights Reserved.
+
+#include "DerivedDataBackendInterface.h"
+
+#if WITH_S3_DDC_BACKEND
+
+#if PLATFORM_WINDOWS || PLATFORM_HOLOLENS
+	#include "Windows/WindowsHWrapper.h"
+	#include "Windows/AllowWindowsPlatformTypes.h"
+#endif
+#if PLATFORM_WINDOWS || PLATFORM_HOLOLENS
+	#include "Windows/HideWindowsPlatformTypes.h"
+#endif
+
+#include "Algo/AllOf.h"
+#include "Async/ParallelFor.h"
+#include "DerivedDataBackendCorruptionWrapper.h"
+#include "DerivedDataCacheRecord.h"
+#include "DerivedDataCachePrivate.h"
+#include "DerivedDataCacheUsageStats.h"
+#include "DerivedDataChunk.h"
+#include "DesktopPlatformModule.h"
+#include "Dom/JsonObject.h"
+#include "HAL/PlatformFile.h"
+#include "HAL/PlatformFileManager.h"
+#include "HAL/Runnable.h"
+#include "HashingArchiveProxy.h"
+#include "Memory/SharedBuffer.h"
+#include "Misc/Base64.h"
+#include "Misc/ConfigCacheIni.h"
+#include "Misc/FeedbackContext.h"
+#include "Misc/FileHelper.h"
+#include "Misc/OutputDeviceRedirector.h"
+#include "Misc/Paths.h"
+#include "Misc/SecureHash.h"
+#include "ProfilingDebugging/CountersTrace.h"
+#include "ProfilingDebugging/CpuProfilerTrace.h"
+#include "Serialization/CompactBinary.h"
+#include "Serialization/CompactBinaryPackage.h"
+#include "Serialization/CompactBinaryValidation.h"
+#include "Serialization/JsonReader.h"
+#include "Serialization/JsonSerializer.h"
+#include "Serialization/MemoryReader.h"
+
+#include "curl/curl.h"
+
+#if WITH_SSL
+#include "Ssl.h"
+#include <openssl/hmac.h>
+#include <openssl/sha.h>
+#include <openssl/ssl.h>
+#endif
+
+#define S3DDC_BACKEND_WAIT_INTERVAL 0.01f
+#define S3DDC_HTTP_REQUEST_TIMEOUT_SECONDS 30L
+#define S3DDC_HTTP_REQUEST_TIMOUT_ENABLED 1
+#define S3DDC_REQUEST_POOL_SIZE 16
+#define S3DDC_MAX_FAILED_LOGIN_ATTEMPTS 16
+#define S3DDC_MAX_ATTEMPTS 4
+#define S3DDC_MAX_BUFFER_RESERVE 104857600u
+
+namespace UE::DerivedData
+{
+
+TRACE_DECLARE_INT_COUNTER(S3DDC_Exist, TEXT("S3DDC Exist"));
+TRACE_DECLARE_INT_COUNTER(S3DDC_ExistHit, TEXT("S3DDC Exist Hit"));
+TRACE_DECLARE_INT_COUNTER(S3DDC_Get, TEXT("S3DDC Get"));
+TRACE_DECLARE_INT_COUNTER(S3DDC_GetHit, TEXT("S3DDC Get Hit"));
+TRACE_DECLARE_INT_COUNTER(S3DDC_BytesRecieved, TEXT("S3DDC Bytes Recieved"));
+
+///////////////////////////////////////////////////////////////////////////////////////////////////////////////////////
+
+static const TCHAR GBucketsDirectoryName[] = TEXT("Buckets");
+static const TCHAR GContentDirectoryName[] = TEXT("Content");
+
+///////////////////////////////////////////////////////////////////////////////////////////////////////////////////////
+
+void BuildPathForLegacyCache(const TCHAR* const CacheKey, FStringBuilderBase& Path);
+void BuildPathForCachePackage(const FCacheKey& CacheKey, FStringBuilderBase& Path);
+void BuildPathForCacheContent(const FIoHash& RawHash, FStringBuilderBase& Path);
+
+static FString BuildPathForLegacyCache(const TCHAR* CacheKey)
+{
+	TStringBuilder<256> Path;
+	BuildPathForLegacyCache(CacheKey, Path);
+	return FString(Path);
+}
+
+class FStringAnsi
+{
+public:
+	FStringAnsi()
+	{
+		Inner.Add(0);
+	}
+
+	FStringAnsi(const ANSICHAR* Text)
+	{
+		Inner.Append(Text, FCStringAnsi::Strlen(Text) + 1);
+	}
+
+	void Append(ANSICHAR Character)
+	{
+		Inner[Inner.Num() - 1] = Character;
+		Inner.Add(0);
+	}
+
+	void Append(const FStringAnsi& Other)
+	{
+		Inner.RemoveAt(Inner.Num() - 1);
+		Inner.Append(Other.Inner);
+	}
+
+	void Append(const ANSICHAR* Text)
+	{
+		Inner.RemoveAt(Inner.Num() - 1);
+		Inner.Append(Text, FCStringAnsi::Strlen(Text) + 1);
+	}
+
+	void Append(const ANSICHAR* Start, const ANSICHAR* End)
+	{
+		Inner.RemoveAt(Inner.Num() - 1);
+		Inner.Append(Start, UE_PTRDIFF_TO_INT32(End - Start));
+		Inner.Add(0);
+	}
+
+	static FStringAnsi Printf(const ANSICHAR* Format, ...)
+	{
+		ANSICHAR Buffer[1024];
+		GET_VARARGS_ANSI(Buffer, UE_ARRAY_COUNT(Buffer), UE_ARRAY_COUNT(Buffer) - 1, Format, Format);
+		return Buffer;
+	}
+
+	FString ToWideString() const
+	{
+		return ANSI_TO_TCHAR(Inner.GetData());
+	}
+
+	const ANSICHAR* operator*() const
+	{
+		return Inner.GetData();
+	}
+
+	int32 Len() const
+	{
+		return Inner.Num() - 1;
+	}
+
+private:
+	TArray<ANSICHAR> Inner;
+};
+
+struct FSHA256
+{
+	uint8 Digest[32];
+
+	FStringAnsi ToString() const
+	{
+		ANSICHAR Buffer[65];
+		for (int Idx = 0; Idx < 32; Idx++)
+		{
+			FCStringAnsi::Sprintf(Buffer + (Idx * 2), "%02x", Digest[Idx]);
+		}
+		return Buffer;
+	}
+};
+
+FSHA256 Sha256(const uint8* Input, size_t InputLen)
+{
+	FSHA256 Output;
+	SHA256(Input, InputLen, Output.Digest);
+	return Output;
+}
+
+FSHA256 HmacSha256(const uint8* Input, size_t InputLen, const uint8* Key, size_t KeyLen)
+{
+	FSHA256 Output;
+	unsigned int OutputLen = 0;
+	HMAC(EVP_sha256(), Key, KeyLen, (const unsigned char*)Input, InputLen, Output.Digest, &OutputLen);
+	return Output;
+}
+
+FSHA256 HmacSha256(const FStringAnsi& Input, const uint8* Key, size_t KeyLen)
+{
+	return HmacSha256((const uint8*)*Input, (size_t)Input.Len(), Key, KeyLen);
+}
+
+FSHA256 HmacSha256(const char* Input, const uint8* Key, size_t KeyLen)
+{
+	return HmacSha256((const uint8*)Input, (size_t)FCStringAnsi::Strlen(Input), Key, KeyLen);
+}
+
+bool IsSuccessfulHttpResponse(long ResponseCode)
+{
+	return (ResponseCode >= 200 && ResponseCode <= 299);
+}
+
+struct IRequestCallback
+{
+	virtual ~IRequestCallback() { }
+	virtual bool Update(int NumBytes, int TotalBytes) = 0;
+};
+
+/**
+ * Backend for a read-only AWS S3 based caching service.
+ **/
+class FS3CacheStore final : public FDerivedDataBackendInterface
+{
+public:
+	/**
+	 * Creates the backend, checks health status and attempts to acquire an access token.
+	 *
+	 * @param  InRootManifestPath   Local path to the JSON manifest in the workspace containing a list of files to download
+	 * @param  InBaseUrl            Base URL for the bucket, with trailing slash (eg. https://foo.s3.us-east-1.amazonaws.com/)
+	 * @param  InRegion	            Name of the AWS region (eg. us-east-1)
+	 * @param  InCanaryObjectKey    Key for a canary object used to test whether this backend is usable
+	 * @param  InCachePath          Path to cache the DDC files
+	 */
+	FS3CacheStore(const TCHAR* InRootManifestPath, const TCHAR* InBaseUrl, const TCHAR* InRegion, const TCHAR* InCanaryObjectKey, const TCHAR* InCachePath);
+	~FS3CacheStore() final;
+
+	/**
+	 * Checks is backend is usable (reachable and accessible).
+	 * @return true if usable
+	 */
+	bool IsUsable() const;
+
+	/* S3 Cache cannot be written to*/
+	bool IsWritable() const final { return false; }
+
+	/* S3 Cache does not try to write back to lower caches (e.g. Shared DDC) */
+	bool BackfillLowerCacheLevels() const final { return false; }
+
+	bool CachedDataProbablyExists(const TCHAR* CacheKey) final;
+	bool GetCachedData(const TCHAR* CacheKey, TArray<uint8>& OutData) final;
+	EPutStatus PutCachedData(const TCHAR* CacheKey, TArrayView<const uint8> InData, bool bPutEvenIfExists) final;
+	void RemoveCachedData(const TCHAR* CacheKey, bool bTransient) final;
+	TSharedRef<FDerivedDataCacheStatsNode> GatherUsageStats() const final;
+
+	FString GetName() const final;
+	ESpeedClass GetSpeedClass() const final;
+	TBitArray<> TryToPrefetch(TConstArrayView<FString> CacheKeys) final;
+	bool WouldCache(const TCHAR* CacheKey, TArrayView<const uint8> InData) final;
+
+	bool ApplyDebugOptions(FBackendDebugOptions& InOptions) final;
+
+	EBackendLegacyMode GetLegacyMode() const final { return EBackendLegacyMode::ValueWithLegacyFallback; }
+
+	void Put(
+		TConstArrayView<FCachePutRequest> Requests,
+		IRequestOwner& Owner,
+		FOnCachePutComplete&& OnComplete) final;
+	void Get(
+		TConstArrayView<FCacheGetRequest> Requests,
+		IRequestOwner& Owner,
+		FOnCacheGetComplete&& OnComplete) final;
+	void PutValue(
+		TConstArrayView<FCachePutValueRequest> Requests,
+		IRequestOwner& Owner,
+		FOnCachePutValueComplete&& OnComplete) final;
+	void GetValue(
+		TConstArrayView<FCacheGetValueRequest> Requests,
+		IRequestOwner& Owner,
+		FOnCacheGetValueComplete&& OnComplete) final;
+	void GetChunks(
+		TConstArrayView<FCacheGetChunkRequest> Requests,
+		IRequestOwner& Owner,
+		FOnCacheGetChunkComplete&& OnComplete) final;
+
+private:
+	struct FBundle;
+	struct FBundleEntry;
+	struct FBundleDownload;
+
+	struct FRootManifest;
+
+	class FHttpRequest;
+	class FRequestPool;
+
+	FString RootManifestPath;
+	FString BaseUrl;
+	FString Region;
+	FString CanaryObjectKey;
+	FString CacheDir;
+	TArray<FBundle> Bundles;
+	TUniquePtr<FRequestPool> RequestPool;
+	FDerivedDataCacheUsageStats UsageStats;
+	bool bEnabled;
+
+	[[nodiscard]] FOptionalCacheRecord GetCacheRecordOnly(
+		FStringView Name,
+		const FCacheKey& Key,
+		const FCacheRecordPolicy& Policy);
+	[[nodiscard]] FOptionalCacheRecord GetCacheRecord(
+		FStringView Name,
+		const FCacheKey& Key,
+		const FCacheRecordPolicy& Policy,
+		EStatus& OutStatus);
+
+	[[nodiscard]] bool GetCacheValueOnly(FStringView Name, const FCacheKey& Key, ECachePolicy Policy, FValue& OutValue);
+	[[nodiscard]] bool GetCacheValue(FStringView Name, const FCacheKey& Key, ECachePolicy Policy, FValue& OutValue);
+
+	[[nodiscard]] bool GetCacheContentExists(const FCacheKey& Key, const FIoHash& RawHash) const;
+	[[nodiscard]] bool GetCacheContent(
+		FStringView Name,
+		const FCacheKey& Key,
+		const FValueId& Id,
+		const FValue& Value,
+		ECachePolicy Policy,
+		FValue& OutValue) const;
+	void GetCacheContent(
+		FStringView Name,
+		const FCacheKey& Key,
+		const FValueId& Id,
+		const FValue& Value,
+		ECachePolicy Policy,
+		FCompressedBufferReader& Reader,
+		TUniquePtr<FArchive>& OutArchive) const;
+
+	void BuildCachePackagePath(const FCacheKey& CacheKey, FStringBuilderBase& Path) const;
+	void BuildCacheContentPath(const FIoHash& RawHash, FStringBuilderBase& Path) const;
+
+	[[nodiscard]] bool LoadFileWithHash(FStringBuilderBase& Path, FStringView DebugName, TFunctionRef<void(FArchive&)> ReadFunction) const;
+	[[nodiscard]] bool LoadFile(FStringBuilderBase& Path, FStringView DebugName, TFunctionRef<void(FArchive&)> ReadFunction) const;
+	[[nodiscard]] TUniquePtr<FArchive> OpenFile(FStringBuilderBase& Path, FStringView DebugName) const;
+
+	[[nodiscard]] bool FileExists(FStringBuilderBase& Path) const;
+
+	bool DownloadManifest(const FRootManifest& RootManifest, FFeedbackContext* Context);
+	void RemoveUnusedBundles();
+	void ReadBundle(FBundle& Bundle);
+	bool FindBundleEntry(const TCHAR* CacheKey, const FBundle*& OutBundle, const FBundleEntry*& OutBundleEntry) const;
+
+	/* Debug helpers */
+	bool DidSimulateMiss(const TCHAR* InKey);
+	bool ShouldSimulateMiss(const TCHAR* InKey);
+
+	/** Debug Options */
+	FBackendDebugOptions DebugOptions;
+
+	/** Keys we ignored due to miss rate settings */
+	FCriticalSection MissedKeysCS;
+	TSet<FName> DebugMissedKeys;
+};
+
+/**
+ * Minimal HTTP request type wrapping CURL without the need for managers. This request
+ * is written to allow reuse of request objects, in order to allow connections to be reused.
+ *
+ * CURL has a global library initialization (curl_global_init). We rely on this happening in 
+ * the Online/HTTP library which is a dependency on this module.
+ */
+class FS3CacheStore::FHttpRequest
+{
+public:
+	FHttpRequest(const ANSICHAR* InRegion, const ANSICHAR* InAccessKey, const ANSICHAR* InSecretKey)
+		: Region(InRegion)
+		, AccessKey(InAccessKey)
+		, SecretKey(InSecretKey)
+	{
+		Curl = curl_easy_init();
+	}
+
+	~FHttpRequest()
+	{
+		curl_easy_cleanup(Curl);
+	}
+
+	/**
+	 * Performs the request, blocking until finished.
+	 * @param Url HTTP URL to fetch
+	 * @param Callback Object used to convey state to/from the operation
+	 * @param Buffer Optional buffer to directly receive the result of the request. 
+	 * If unset the response body will be stored in the request.
+	 */
+	long PerformBlocking(const ANSICHAR* Url, IRequestCallback* Callback, TArray<uint8>& OutResponseBody, FOutputDevice* Log)
+	{
+		TRACE_CPUPROFILER_EVENT_SCOPE(S3DDC_CurlPerform);
+
+		// Find the host from the URL
+		const ANSICHAR* ProtocolEnd = FCStringAnsi::Strchr(Url, ':');
+		check(ProtocolEnd != nullptr && *(ProtocolEnd + 1) == '/' && *(ProtocolEnd + 2) == '/');
+
+		const ANSICHAR* UrlHost = ProtocolEnd + 3;
+		const ANSICHAR* UrlHostEnd = FCStringAnsi::Strchr(UrlHost, '/');
+		check(UrlHostEnd != nullptr);
+
+		FStringAnsi Host;
+		Host.Append(UrlHost, UrlHostEnd);
+
+		// Get the header strings
+		FDateTime Timestamp = FDateTime::UtcNow();// FDateTime(2015, 9, 15, 12, 45, 0);
+		FStringAnsi TimeString = FStringAnsi::Printf("%04d%02d%02dT%02d%02d%02dZ", Timestamp.GetYear(), Timestamp.GetMonth(), Timestamp.GetDay(), Timestamp.GetHour(), Timestamp.GetMinute(), Timestamp.GetSecond());
+
+		// Payload string
+		FStringAnsi EmptyPayloadSha256 = Sha256(nullptr, 0).ToString();
+
+		// Create the headers
+		curl_slist* CurlHeaders = nullptr;
+		CurlHeaders = curl_slist_append(CurlHeaders, *FStringAnsi::Printf("Host: %s", *Host));
+		CurlHeaders = curl_slist_append(CurlHeaders, *FStringAnsi::Printf("x-amz-content-sha256: %s", *EmptyPayloadSha256));
+		CurlHeaders = curl_slist_append(CurlHeaders, *FStringAnsi::Printf("x-amz-date: %s", *TimeString));
+		CurlHeaders = curl_slist_append(CurlHeaders, *GetAuthorizationHeader("GET", UrlHostEnd, "", CurlHeaders, *TimeString, *EmptyPayloadSha256));
+
+		// Create the callback data
+		FStringAnsi Domain;
+		Domain.Append(Url, UrlHostEnd);
+
+		FCallbackData CallbackData(Domain, OutResponseBody);
+
+		// Setup the request
+		curl_easy_reset(Curl);
+		curl_easy_setopt(Curl, CURLOPT_FOLLOWLOCATION, 1L);
+		curl_easy_setopt(Curl, CURLOPT_NOSIGNAL, 1L);
+		curl_easy_setopt(Curl, CURLOPT_HTTPGET, 1L);
+		curl_easy_setopt(Curl, CURLOPT_URL, Url);
+		curl_easy_setopt(Curl, CURLOPT_ACCEPT_ENCODING, "gzip");
+#if S3DDC_HTTP_REQUEST_TIMOUT_ENABLED
+		curl_easy_setopt(Curl, CURLOPT_CONNECTTIMEOUT, S3DDC_HTTP_REQUEST_TIMEOUT_SECONDS);
+#endif
+
+		// Headers
+		curl_easy_setopt(Curl, CURLOPT_HTTPHEADER, CurlHeaders);
+		curl_easy_setopt(Curl, CURLOPT_HEADERDATA, CurlHeaders);
+
+		// Progress
+		curl_easy_setopt(Curl, CURLOPT_NOPROGRESS, 0);
+		curl_easy_setopt(Curl, CURLOPT_XFERINFODATA, Callback);
+		curl_easy_setopt(Curl, CURLOPT_XFERINFOFUNCTION, &FHttpRequest::StaticStatusFn);
+
+		// Response
+		curl_easy_setopt(Curl, CURLOPT_HEADERDATA, &CallbackData);
+		curl_easy_setopt(Curl, CURLOPT_HEADERFUNCTION, &FHttpRequest::StaticWriteHeaderFn);
+		curl_easy_setopt(Curl, CURLOPT_WRITEDATA, &CallbackData);
+		curl_easy_setopt(Curl, CURLOPT_WRITEFUNCTION, StaticWriteBodyFn);
+
+		// SSL options
+		curl_easy_setopt(Curl, CURLOPT_USE_SSL, CURLUSESSL_ALL);
+		curl_easy_setopt(Curl, CURLOPT_SSL_VERIFYPEER, 1);
+		curl_easy_setopt(Curl, CURLOPT_SSL_VERIFYHOST, 1);
+		curl_easy_setopt(Curl, CURLOPT_SSLCERTTYPE, "PEM");
+
+		// SSL certification verification
+		curl_easy_setopt(Curl, CURLOPT_CAINFO, nullptr);
+		curl_easy_setopt(Curl, CURLOPT_SSL_CTX_FUNCTION, *sslctx_function);
+		curl_easy_setopt(Curl, CURLOPT_SSL_CTX_DATA, &CallbackData);
+
+		// Send the request
+		CURLcode CurlResult = curl_easy_perform(Curl);
+
+		// Free the headers object
+		curl_slist_free_all(CurlHeaders);
+		curl_easy_setopt(Curl, CURLOPT_HEADERDATA, nullptr);
+
+		// Get the response code
+		long ResponseCode = 0;
+		if (CurlResult == CURLE_OK)
+		{
+			CurlResult = curl_easy_getinfo(Curl, CURLINFO_RESPONSE_CODE, &ResponseCode);
+		}
+		if (CurlResult != CURLE_OK)
+		{
+			if (CurlResult != CURLE_ABORTED_BY_CALLBACK)
+			{
+				Log->Logf(ELogVerbosity::Error, TEXT("Error while connecting to %s: %d (%s)"), ANSI_TO_TCHAR(Url), CurlResult, ANSI_TO_TCHAR(curl_easy_strerror(CurlResult)));
+			}
+			return 500;
+		}
+
+		// Print any diagnostic output
+		if (!(ResponseCode >= 200 && ResponseCode <= 299))
+		{
+			if (FAnsiStringView(Url).StartsWith("file://"))
+			{
+				ResponseCode = 200;
+			}
+			else
+			{
+				Log->Logf(ELogVerbosity::Error, TEXT("Download failed for %s (response %d):\n%s\n%s"), ANSI_TO_TCHAR(Url), ResponseCode, *CallbackData.ResponseHeader.ToWideString(), ANSI_TO_TCHAR((const ANSICHAR*)OutResponseBody.GetData()));
+			}
+		}
+		return ResponseCode;
+	}
+
+private:
+	struct FCallbackData
+	{
+		const FStringAnsi& Domain;
+		FStringAnsi ResponseHeader;
+		TArray<uint8>& ResponseBody;
+
+		FCallbackData(const FStringAnsi& InDomain, TArray<uint8>& InResponseBody)
+			: Domain(InDomain)
+			, ResponseBody(InResponseBody)
+		{
+		}
+	};
+
+	CURL* Curl;
+	FStringAnsi Region;
+	FStringAnsi AccessKey;
+	FStringAnsi SecretKey;
+
+	FStringAnsi GetAuthorizationHeader(const ANSICHAR* Verb, const ANSICHAR* RelativeUrl, const ANSICHAR* QueryString, const curl_slist* Headers, const ANSICHAR* Timestamp, const ANSICHAR* Digest)
+	{
+		// Create the canonical list of headers
+		FStringAnsi CanonicalHeaders;
+		for (const curl_slist* Header = Headers; Header != nullptr; Header = Header->next)
+		{
+			const ANSICHAR* Colon = FCStringAnsi::Strchr(Header->data, ':');
+			if (Colon != nullptr)
+			{
+				for (const ANSICHAR* Char = Header->data; Char != Colon; Char++)
+				{
+					CanonicalHeaders.Append(FCharAnsi::ToLower(*Char));
+				}
+				CanonicalHeaders.Append(':');
+
+				const ANSICHAR* Value = Colon + 1;
+				while (*Value == ' ')
+				{
+					Value++;
+				}
+				for (; *Value != 0; Value++)
+				{
+					CanonicalHeaders.Append(*Value);
+				}
+				CanonicalHeaders.Append('\n');
+			}
+		}
+
+		// Create the list of signed headers
+		FStringAnsi SignedHeaders;
+		for (const curl_slist* Header = Headers; Header != nullptr; Header = Header->next)
+		{
+			const ANSICHAR* Colon = FCStringAnsi::Strchr(Header->data, ':');
+			if (Colon != nullptr)
+			{
+				if (SignedHeaders.Len() > 0)
+				{
+					SignedHeaders.Append(';');
+				}
+				for (const ANSICHAR* Char = Header->data; Char != Colon; Char++)
+				{
+					SignedHeaders.Append(FCharAnsi::ToLower(*Char));
+				}
+			}
+		}
+
+		// Build the canonical request string
+		FStringAnsi CanonicalRequest;
+		CanonicalRequest.Append(Verb);
+		CanonicalRequest.Append('\n');
+		CanonicalRequest.Append(RelativeUrl);
+		CanonicalRequest.Append('\n');
+		CanonicalRequest.Append(QueryString);
+		CanonicalRequest.Append('\n');
+		CanonicalRequest.Append(CanonicalHeaders);
+		CanonicalRequest.Append('\n');
+		CanonicalRequest.Append(SignedHeaders);
+		CanonicalRequest.Append('\n');
+		CanonicalRequest.Append(Digest);
+
+		// Get the date
+		FStringAnsi DateString;
+		for (int32 Idx = 0; Timestamp[Idx] != 0 && Timestamp[Idx] != 'T'; Idx++)
+		{
+			DateString.Append(Timestamp[Idx]);
+		}
+
+		// Generate the signature key
+		FStringAnsi Key = FStringAnsi::Printf("AWS4%s", *SecretKey);
+
+		FSHA256 DateHash = HmacSha256(DateString, (const uint8*)*Key, Key.Len());
+		FSHA256 RegionHash = HmacSha256(Region, DateHash.Digest, sizeof(DateHash.Digest));
+		FSHA256 ServiceHash = HmacSha256("s3", RegionHash.Digest, sizeof(RegionHash.Digest));
+		FSHA256 SigningKeyHash = HmacSha256("aws4_request", ServiceHash.Digest, sizeof(ServiceHash.Digest));
+
+		// Calculate the signature
+		FStringAnsi DateRequest = FStringAnsi::Printf("%s/%s/s3/aws4_request", *DateString, *Region);
+		FStringAnsi CanonicalRequestSha256 = Sha256((const uint8*)*CanonicalRequest, CanonicalRequest.Len()).ToString();
+		FStringAnsi StringToSign = FStringAnsi::Printf("AWS4-HMAC-SHA256\n%s\n%s\n%s", Timestamp, *DateRequest, *CanonicalRequestSha256);
+		FStringAnsi Signature = HmacSha256(*StringToSign, SigningKeyHash.Digest, sizeof(SigningKeyHash.Digest)).ToString();
+
+		// Format the final header
+		return FStringAnsi::Printf("Authorization: AWS4-HMAC-SHA256 Credential=%s/%s, SignedHeaders=%s, Signature=%s", *AccessKey, *DateRequest, *SignedHeaders, *Signature);
+	}
+
+	/**
+	 * Returns the response buffer as a string. Note that is the request is performed
+	 * with an external buffer as target buffer this string will be empty.
+	 */
+	static FString GetResponseAsString(const TArray<uint8>& Buffer)
+	{
+		FUTF8ToTCHAR TCHARData(reinterpret_cast<const ANSICHAR*>(Buffer.GetData()), Buffer.Num());
+		return FString(TCHARData.Length(), TCHARData.Get());
+	}
+
+	static int StaticStatusFn(void* Ptr, curl_off_t TotalDownloadSize, curl_off_t CurrentDownloadSize, curl_off_t TotalUploadSize, curl_off_t CurrentUploadSize)
+	{
+		IRequestCallback* Callback = (IRequestCallback*)Ptr;
+		if (Callback != nullptr)
+		{
+			return Callback->Update(IntCastChecked<int>(CurrentDownloadSize), IntCastChecked<int>(TotalDownloadSize)) ? 0 : 1;
+		}
+		return 0;
+	}
+
+	static size_t StaticWriteHeaderFn(void* Ptr, size_t SizeInBlocks, size_t BlockSizeInBytes, void* UserData)
+	{
+		const size_t WriteSize = SizeInBlocks * BlockSizeInBytes;
+		if (WriteSize > 0)
+		{
+			FCallbackData* CallbackData = static_cast<FCallbackData*>(UserData);
+			CallbackData->ResponseHeader.Append((const ANSICHAR*)Ptr, (const ANSICHAR*)Ptr + WriteSize);
+			return WriteSize;
+		}
+		return 0;
+	}
+
+	static size_t StaticWriteBodyFn(void* Ptr, size_t SizeInBlocks, size_t BlockSizeInBytes, void* UserData)
+	{
+		const size_t WriteSize = SizeInBlocks * BlockSizeInBytes;
+		if (WriteSize > 0)
+		{
+			FCallbackData* CallbackData = static_cast<FCallbackData*>(UserData);
+
+			// If this is the first part of the body being received, try to reserve 
+			// memory if content length is defined in the header.
+			if (CallbackData->ResponseBody.Num() == 0)
+			{
+				static const ANSICHAR Prefix[] = "Content-Length: ";
+				static const size_t PrefixLen = UE_ARRAY_COUNT(Prefix) - 1;
+
+				for(const ANSICHAR* Header = *CallbackData->ResponseHeader;;Header++)
+				{
+					// Check this header
+					if (FCStringAnsi::Strnicmp(Header, Prefix, PrefixLen) == 0)
+					{
+						size_t ContentLength = (size_t)atol(Header + PrefixLen);
+						if (ContentLength > 0u && ContentLength < S3DDC_MAX_BUFFER_RESERVE)
+						{
+							CallbackData->ResponseBody.Reserve(ContentLength);
+						}
+						break;
+					}
+
+					// Move to the next string
+					Header = FCStringAnsi::Strchr(Header, '\n');
+					if (Header == nullptr)
+					{
+						break;
+					}
+				}
+			}
+
+			// Write to the target buffer
+			CallbackData->ResponseBody.Append((const uint8*)Ptr, WriteSize);
+			return WriteSize;
+		}
+		return 0;
+	}
+
+	static int SslCertVerify(int PreverifyOk, X509_STORE_CTX* Context)
+	{
+		if (PreverifyOk == 1)
+		{
+			SSL* Handle = static_cast<SSL*>(X509_STORE_CTX_get_ex_data(Context, SSL_get_ex_data_X509_STORE_CTX_idx()));
+			check(Handle);
+
+			SSL_CTX* SslContext = SSL_get_SSL_CTX(Handle);
+			check(SslContext);
+
+			FCallbackData* CallbackData = reinterpret_cast<FCallbackData*>(SSL_CTX_get_app_data(SslContext));
+			check(CallbackData);
+
+			if (!FSslModule::Get().GetCertificateManager().VerifySslCertificates(Context, *CallbackData->Domain))
+			{
+				PreverifyOk = 0;
+			}
+		}
+
+		return PreverifyOk;
+	}
+
+	static CURLcode sslctx_function(CURL* curl, void* sslctx, void* parm)
+	{
+		SSL_CTX* Context = static_cast<SSL_CTX*>(sslctx);
+		const ISslCertificateManager& CertificateManager = FSslModule::Get().GetCertificateManager();
+
+		CertificateManager.AddCertificatesToSslContext(Context);
+		SSL_CTX_set_verify(Context, SSL_CTX_get_verify_mode(Context), SslCertVerify);
+		SSL_CTX_set_app_data(Context, parm);
+
+		/* all set to go */
+		return CURLE_OK;
+	}
+};
+
+//----------------------------------------------------------------------------------------------------------
+// FS3CacheStore::FRequestPool
+//----------------------------------------------------------------------------------------------------------
+
+class FS3CacheStore::FRequestPool
+{
+public:
+	FRequestPool(const ANSICHAR* Region, const ANSICHAR* AccessKey, const ANSICHAR* SecretKey)
+	{
+		Pool.SetNum(S3DDC_REQUEST_POOL_SIZE);
+		for (uint8 i = 0; i < Pool.Num(); ++i)
+		{
+			Pool[i].Usage = 0u;
+			Pool[i].Request = new FHttpRequest(Region, AccessKey, SecretKey);
+		}
+	}
+
+	~FRequestPool()
+	{
+		for (uint8 i = 0; i < Pool.Num(); ++i)
+		{
+			// No requests should be in use by now.
+			check(Pool[i].Usage.Load(EMemoryOrder::Relaxed) == 0u);
+			delete Pool[i].Request;
+		}
+	}
+
+	long Download(const TCHAR* Url, IRequestCallback* Callback, TArray<uint8>& OutData, FOutputDevice* Log)
+	{
+		FHttpRequest* Request = WaitForFreeRequest();
+		long ResponseCode = Request->PerformBlocking(TCHAR_TO_ANSI(Url), Callback, OutData, Log);
+		ReleaseRequestToPool(Request);
+		return ResponseCode;
+	}
+
+private:
+	struct FEntry
+	{
+		TAtomic<uint8> Usage;
+		FHttpRequest* Request;
+	};
+
+	TArray<FEntry> Pool;
+
+	FHttpRequest* WaitForFreeRequest()
+	{
+		TRACE_CPUPROFILER_EVENT_SCOPE(S3DDC_WaitForConnPool);
+		while (true)
+		{
+			for (uint8 i = 0; i < Pool.Num(); ++i)
+			{
+				if (!Pool[i].Usage.Load(EMemoryOrder::Relaxed))
+				{
+					uint8 Expected = 0u;
+					if (Pool[i].Usage.CompareExchange(Expected, 1u))
+					{
+						return Pool[i].Request;
+					}
+				}
+			}
+			FPlatformProcess::Sleep(S3DDC_BACKEND_WAIT_INTERVAL);
+		}
+	}
+
+	void ReleaseRequestToPool(FHttpRequest* Request)
+	{
+		for (uint8 i = 0; i < Pool.Num(); ++i)
+		{
+			if (Pool[i].Request == Request)
+			{
+				uint8 Expected = 1u;
+				Pool[i].Usage.CompareExchange(Expected, 0u);
+				return;
+			}
+		}
+		check(false);
+	}
+};
+
+//----------------------------------------------------------------------------------------------------------
+// FS3CacheStore::FBundleEntry
+//----------------------------------------------------------------------------------------------------------
+
+struct FS3CacheStore::FBundleEntry
+{
+	int64 Offset;
+	int32 Length;
+};
+
+//----------------------------------------------------------------------------------------------------------
+// FS3CacheStore::FBundle
+//----------------------------------------------------------------------------------------------------------
+
+struct FS3CacheStore::FBundle
+{
+	FString Name;
+	FString ObjectKey;
+	FString LocalFile;
+	int32 CompressedLength;
+	int32 UncompressedLength;
+	TMap<FSHAHash, FBundleEntry> Entries;
+};
+
+//----------------------------------------------------------------------------------------------------------
+// FS3CacheStore::FBundleDownloadInfo
+//----------------------------------------------------------------------------------------------------------
+
+struct FS3CacheStore::FBundleDownload final : IRequestCallback
+{
+	FCriticalSection& CriticalSection;
+	FBundle& Bundle;
+	FString BundleUrl;
+	FRequestPool& RequestPool;
+	FFeedbackContext* Context;
+	FGraphEventRef Event;
+	int DownloadedBytes;
+
+	FBundleDownload(FCriticalSection& InCriticalSection, FBundle& InBundle, FString InBundleUrl, FRequestPool& InRequestPool, FFeedbackContext* InContext)
+		: CriticalSection(InCriticalSection)
+		, Bundle(InBundle)
+		, BundleUrl(InBundleUrl)
+		, RequestPool(InRequestPool)
+		, Context(InContext)
+		, DownloadedBytes(0)
+	{
+	}
+
+	void Execute()
+	{
+		if (Context->ReceivedUserCancel())
+		{
+			return;
+		}
+
+		Context->Logf(TEXT("Downloading %s (%d bytes)"), *BundleUrl, Bundle.CompressedLength);
+
+		TArray<uint8> CompressedData;
+		CompressedData.Reserve(Bundle.CompressedLength);
+
+		long ResponseCode = RequestPool.Download(*BundleUrl, this, CompressedData, Context);
+		if(!IsSuccessfulHttpResponse(ResponseCode))
+		{
+			if (!Context->ReceivedUserCancel())
+			{
+				Context->Logf(ELogVerbosity::Warning, TEXT("Unable to download bundle %s (%d)"), *BundleUrl, ResponseCode);
+			}
+			return;
+		}
+
+		Context->Logf(TEXT("Decompressing %s (%d bytes)"), *BundleUrl, Bundle.UncompressedLength);
+
+		TArray<uint8> UncompressedData;
+		UncompressedData.SetNum(Bundle.UncompressedLength);
+
+		if (!FCompression::UncompressMemory(NAME_Gzip, UncompressedData.GetData(), Bundle.UncompressedLength, CompressedData.GetData(), CompressedData.Num()))
+		{
+			Context->Logf(ELogVerbosity::Warning, TEXT("Unable to decompress bundle %s"), *BundleUrl);
+			return;
+		}
+
+		FString TempFile = Bundle.LocalFile + TEXT(".incoming");
+		if (!FFileHelper::SaveArrayToFile(UncompressedData, *TempFile))
+		{
+			Context->Logf(ELogVerbosity::Warning, TEXT("Unable to save bundle to %s"), *TempFile);
+			return;
+		}
+
+		IFileManager::Get().Move(*Bundle.LocalFile, *TempFile);
+		Context->Logf(TEXT("Finished downloading %s to %s"), *BundleUrl, *Bundle.LocalFile);
+	}
+
+	bool Update(int NumBytes, int MaxBytes) final
+	{
+		FScopeLock Lock(&CriticalSection);
+		DownloadedBytes = NumBytes;
+		return !Context->ReceivedUserCancel();
+	}
+};
+
+//----------------------------------------------------------------------------------------------------------
+// FS3CacheStore::FRootManifest
+//----------------------------------------------------------------------------------------------------------
+
+struct FS3CacheStore::FRootManifest
+{
+	FString AccessKey;
+	FString SecretKey;
+	TArray<FString> Keys;
+
+	bool Load(const FString& InRootManifestPath)
+	{
+		// Read the root manifest from disk
+		FString RootManifestText;
+		if (!FFileHelper::LoadFileToString(RootManifestText, *InRootManifestPath))
+		{
+			UE_LOG(LogDerivedDataCache, Warning, TEXT("Unable to read manifest from %s"), *InRootManifestPath);
+			return false;
+		}
+
+		// Deserialize a JSON object from the string
+		TSharedPtr<FJsonObject> RootManifestObject;
+		if (!FJsonSerializer::Deserialize(TJsonReaderFactory<>::Create(RootManifestText), RootManifestObject) || !RootManifestObject.IsValid())
+		{
+			UE_LOG(LogDerivedDataCache, Warning, TEXT("Unable to parse manifest from %s"), *InRootManifestPath);
+			return false;
+		}
+
+		// Read the access and secret keys
+		if (!RootManifestObject->TryGetStringField(TEXT("AccessKey"), AccessKey))
+		{
+			UE_LOG(LogDerivedDataCache, Warning, TEXT("Root manifest %s does not specify AccessKey"), *InRootManifestPath);
+			return false;
+		}
+		if (!RootManifestObject->TryGetStringField(TEXT("SecretKey"), SecretKey))
+		{
+			UE_LOG(LogDerivedDataCache, Warning, TEXT("Root manifest %s does not specify SecretKey"), *InRootManifestPath);
+			return false;
+		}
+
+		// Parse out the list of manifests
+		const TArray<TSharedPtr<FJsonValue>>* RootEntriesArray;
+		if (!RootManifestObject->TryGetArrayField(TEXT("Entries"), RootEntriesArray))
+		{
+			UE_LOG(LogDerivedDataCache, Warning, TEXT("Root manifest from %s is missing entries array"), *InRootManifestPath);
+			return false;
+		}
+		for (const TSharedPtr<FJsonValue>& Value : *RootEntriesArray)
+		{
+			const TSharedPtr<FJsonObject>& LastRootManifestEntry = (*RootEntriesArray)[RootEntriesArray->Num() - 1]->AsObject();
+			Keys.Add(LastRootManifestEntry->GetStringField("Key"));
+		}
+
+		return true;
+	}
+};
+
+//----------------------------------------------------------------------------------------------------------
+// FS3CacheStore
+//----------------------------------------------------------------------------------------------------------
+
+FS3CacheStore::FS3CacheStore(const TCHAR* InRootManifestPath, const TCHAR* InBaseUrl, const TCHAR* InRegion, const TCHAR* InCanaryObjectKey, const TCHAR* InCachePath)
+	: RootManifestPath(InRootManifestPath)
+	, BaseUrl(InBaseUrl)
+	, Region(InRegion)
+	, CanaryObjectKey(InCanaryObjectKey)
+	, CacheDir(InCachePath)
+	, bEnabled(false)
+{
+	FRootManifest RootManifest;
+	if (RootManifest.Load(InRootManifestPath))
+	{
+		RequestPool.Reset(new FRequestPool(TCHAR_TO_ANSI(InRegion), TCHAR_TO_ANSI(*RootManifest.AccessKey), TCHAR_TO_ANSI(*RootManifest.SecretKey)));
+
+		FString LocalManifestPath;
+		if (RootManifest.Keys.Last().StartsWith(TEXT("file://")))
+		{
+			LocalManifestPath = FPaths::GetPath(RootManifest.Keys.Last());
+		}
+
+		// Test whether we can reach the canary URL
+		bool bCanaryValid = true;
+		if (LocalManifestPath.IsEmpty())
+		{
+			if (GIsBuildMachine)
+			{
+				UE_LOG(LogDerivedDataCache, Log, TEXT("S3DerivedDataBackend: Disabling on build machine"));
+				bCanaryValid = false;
+			}
+			else if (CanaryObjectKey.Len() > 0)
+			{
+				TArray<uint8> Data;
+
+				FStringOutputDevice DummyOutputDevice;
+				if (!IsSuccessfulHttpResponse(RequestPool->Download(*(BaseUrl / CanaryObjectKey), nullptr, Data, &DummyOutputDevice)))
+				{
+					UE_LOG(LogDerivedDataCache, Log, TEXT("S3DerivedDataBackend: Unable to download canary file. Disabling."));
+					bCanaryValid = false;
+				}
+			}
+		}
+
+		// Allow the user to override it from the editor
+		bool bSetting;
+		if (GConfig->GetBool(TEXT("/Script/UnrealEd.EditorSettings"), TEXT("bEnableS3DDC"), bSetting, GEditorSettingsIni) && !bSetting)
+		{
+			UE_LOG(LogDerivedDataCache, Log, TEXT("S3DerivedDataBackend: Disabling due to config setting"));
+			bCanaryValid = false;
+		}
+
+		// Try to read the bundles
+		if (bCanaryValid)
+		{
+			UE_LOG(LogDerivedDataCache, Log, TEXT("Using %s S3 backend at %s"), *Region, *BaseUrl);
+
+			FFeedbackContext* Context = FDesktopPlatformModule::Get()->GetNativeFeedbackContext();
+			Context->BeginSlowTask(NSLOCTEXT("S3DerivedDataBackend", "DownloadingDDCBundles", "Downloading DDC bundles..."), true, true);
+
+			if (DownloadManifest(RootManifest, Context))
+			{
+				// Get the path for each bundle that needs downloading
+				for (FBundle& Bundle : Bundles)
+				{
+					Bundle.LocalFile = CacheDir / Bundle.Name;
+				}
+
+				// Remove any bundles that are no longer required
+				RemoveUnusedBundles();
+
+				// Create a critical section used for updating download state
+				FCriticalSection CriticalSection;
+
+				// Create all the download tasks
+				TArray<TSharedPtr<FBundleDownload>> Downloads;
+				for (FBundle& Bundle : Bundles)
+				{
+					if (!FPaths::FileExists(Bundle.LocalFile))
+					{
+						FString BundleUrl = LocalManifestPath.IsEmpty() ? BaseUrl + Bundle.ObjectKey : LocalManifestPath / Bundle.Name + TEXT(".gz");
+						TSharedPtr<FBundleDownload> Download(new FBundleDownload(CriticalSection, Bundle, BundleUrl, *RequestPool.Get(), Context));
+						Download->Event = FFunctionGraphTask::CreateAndDispatchWhenReady([Download]() { Download->Execute(); }, TStatId());
+						Downloads.Add(MoveTemp(Download));
+					}
+				}
+
+				// Loop until the downloads have all finished
+				for (bool bComplete = false; !bComplete; )
+				{
+					FPlatformProcess::Sleep(0.1f);
+
+					int64 NumBytes = 0;
+					int64 MaxBytes = 0;
+					bComplete = true;
+
+					CriticalSection.Lock();
+					for (TSharedPtr<FBundleDownload>& Download : Downloads)
+					{
+						NumBytes += Download->DownloadedBytes;
+						MaxBytes += Download->Bundle.CompressedLength;
+						bComplete &= Download->Event->IsComplete();
+					}
+					CriticalSection.Unlock();
+
+					int NumMB = (int)((NumBytes + (1024 * 1024) - 1) / (1024 * 1024));
+					int MaxMB = (int)((MaxBytes + (1024 * 1024) - 1) / (1024 * 1024));
+					if (MaxBytes > 0)
+					{
+						FText StatusText = FText::Format(NSLOCTEXT("S3DerivedDataBackend", "DownloadingDDCBundlesPct", "Downloading DDC bundles... ({0}Mb/{1}Mb)"), NumMB, MaxMB);
+						Context->StatusUpdate((int)((NumBytes * 1000) / MaxBytes), 1000, StatusText);
+					}
+				}
+
+				// Mount all the bundles
+				ParallelFor(Bundles.Num(), [this](int32 Index) { ReadBundle(Bundles[Index]); });
+				bEnabled = true;
+			}
+
+			Context->EndSlowTask();
+		}
+	}
+}
+
+FS3CacheStore::~FS3CacheStore()
+{
+}
+
+bool FS3CacheStore::IsUsable() const
+{
+	return bEnabled;
+}
+
+bool FS3CacheStore::CachedDataProbablyExists(const TCHAR* CacheKey)
+{
+	TRACE_CPUPROFILER_EVENT_SCOPE(S3DDC_Exist);
+	TRACE_COUNTER_ADD(S3DDC_Exist, int64(1));
+	COOK_STAT(auto Timer = UsageStats.TimeProbablyExists());
+
+	const FBundle* Bundle;
+	const FBundleEntry* BundleEntry;
+
+	if (ShouldSimulateMiss(CacheKey))
+	{
+		return false;
+	}
+
+	if (!FindBundleEntry(CacheKey, Bundle, BundleEntry))
+	{
+		UE_LOG(LogDerivedDataCache, Verbose, TEXT("S3DerivedDataBackend: Cache miss on %s (probably)"), CacheKey);
+		return false;
+	}
+
+	TRACE_COUNTER_ADD(S3DDC_ExistHit, int64(1));
+	COOK_STAT(Timer.AddHit(BundleEntry->Length));
+	return true;
+}
+
+bool FS3CacheStore::GetCachedData(const TCHAR* CacheKey, TArray<uint8>& OutData)
+{
+	TRACE_CPUPROFILER_EVENT_SCOPE(S3DDC_Get);
+	TRACE_COUNTER_ADD(S3DDC_Get, int64(1));
+	COOK_STAT(auto Timer = UsageStats.TimeGet());
+
+	if (ShouldSimulateMiss(CacheKey))
+	{
+		return false;
+	}
+
+	const FBundle* Bundle;
+	const FBundleEntry* BundleEntry;
+	if (FindBundleEntry(CacheKey, Bundle, BundleEntry))
+	{
+		TUniquePtr<FArchive> Reader(IFileManager::Get().CreateFileReader(*Bundle->LocalFile));
+		if (Reader.IsValid() && !Reader->IsError())
+		{
+			TRACE_COUNTER_ADD(S3DDC_GetHit, int64(1));
+			COOK_STAT(Timer.AddHit(BundleEntry->Length));
+			UE_LOG(LogDerivedDataCache, Verbose, TEXT("S3DerivedDataBackend: Cache hit on %s"), CacheKey);
+			OutData.SetNum(BundleEntry->Length);
+			Reader->Seek(BundleEntry->Offset);
+			Reader->Serialize(OutData.GetData(), BundleEntry->Length);
+			if (FCorruptionWrapper::ReadTrailer(OutData, *BaseUrl, CacheKey))
+			{
+				return true;
+			}
+		}
+	}
+
+	UE_LOG(LogDerivedDataCache, Verbose, TEXT("S3DerivedDataBackend: Cache miss on %s"), CacheKey);
+	return false;
+}
+
+FDerivedDataBackendInterface::EPutStatus FS3CacheStore::PutCachedData(const TCHAR* CacheKey, TArrayView<const uint8> InData, bool bPutEvenIfExists)
+{
+	// Not implemented
+	return EPutStatus::NotCached;
+}
+
+void FS3CacheStore::RemoveCachedData(const TCHAR* CacheKey, bool bTransient)
+{
+	// Not implemented
+}
+
+TSharedRef<FDerivedDataCacheStatsNode> FS3CacheStore::GatherUsageStats() const
+{
+	TSharedRef<FDerivedDataCacheStatsNode> Usage =
+		MakeShared<FDerivedDataCacheStatsNode>(TEXT("S3"), BaseUrl, /*bIsLocal*/ false);
+	Usage->Stats.Add(TEXT(""), UsageStats);
+	return Usage;
+}
+
+FString FS3CacheStore::GetName() const
+{
+	return BaseUrl;
+}
+
+FDerivedDataBackendInterface::ESpeedClass FS3CacheStore::GetSpeedClass() const
+{
+	return ESpeedClass::Local;
+}
+
+TBitArray<> FS3CacheStore::TryToPrefetch(TConstArrayView<FString> CacheKeys)
+{
+	return CachedDataProbablyExistsBatch(CacheKeys);
+}
+
+bool FS3CacheStore::WouldCache(const TCHAR* CacheKey, TArrayView<const uint8> InData)
+{
+	return false;
+}
+
+bool FS3CacheStore::DownloadManifest(const FRootManifest& RootManifest, FFeedbackContext* Context)
+{
+	// Read the root manifest from disk
+	if (RootManifest.Keys.Num() == 0)
+	{
+		Context->Logf(ELogVerbosity::Warning, TEXT("Root manifest has empty entries array"));
+		return false;
+	}
+
+	// Get the object key for the last entry
+	FString BundleManifestKey = RootManifest.Keys.Last();
+
+	// Download the bundle manifest
+	FString BundleManifestUrl = BundleManifestKey.StartsWith(TEXT("file://")) ? BundleManifestKey : BaseUrl + BundleManifestKey;
+	TArray<uint8> BundleManifestData;
+	long ResponseCode = RequestPool->Download(*BundleManifestUrl, nullptr, BundleManifestData, Context);
+	if (!IsSuccessfulHttpResponse(ResponseCode))
+	{
+		Context->Logf(ELogVerbosity::Warning, TEXT("Unable to download bundle manifest from %s (%d)"), *BundleManifestKey, (int)ResponseCode);
+		return false;
+	}
+
+	// Convert it to text
+	BundleManifestData.Add(0);
+	FString BundleManifestText = ANSI_TO_TCHAR((const ANSICHAR*)BundleManifestData.GetData());
+
+	// Deserialize a JSON object from the string
+	TSharedPtr<FJsonObject> BundleManifestObject;
+	if (!FJsonSerializer::Deserialize(TJsonReaderFactory<>::Create(BundleManifestText), BundleManifestObject) || !BundleManifestObject.IsValid())
+	{
+		Context->Logf(ELogVerbosity::Warning, TEXT("Unable to parse manifest from %s"), *BundleManifestKey);
+		return false;
+	}
+
+	// Parse out the list of bundles
+	const TArray<TSharedPtr<FJsonValue>>* BundlesArray;
+	if (!BundleManifestObject->TryGetArrayField(TEXT("Entries"), BundlesArray))
+	{
+		Context->Logf(ELogVerbosity::Warning, TEXT("Manifest from %s is missing bundles array"), *BundleManifestKey);
+		return false;
+	}
+
+	// Parse out each bundle
+	for (const TSharedPtr<FJsonValue>& BundleValue : *BundlesArray)
+	{
+		const FJsonObject& BundleObject = *BundleValue->AsObject();
+
+		FBundle Bundle;
+		if (!BundleObject.TryGetStringField(TEXT("Name"), Bundle.Name))
+		{
+			Context->Logf(ELogVerbosity::Warning, TEXT("Manifest from %s is missing a bundle name"), *BundleManifestKey);
+			return false;
+		}
+		if (!BundleObject.TryGetStringField(TEXT("ObjectKey"), Bundle.ObjectKey))
+		{
+			Context->Logf(ELogVerbosity::Warning, TEXT("Manifest from %s is missing an bundle object key"), *BundleManifestKey);
+			return false;
+		}
+		if (!BundleObject.TryGetNumberField(TEXT("CompressedLength"), Bundle.CompressedLength))
+		{
+			Context->Logf(ELogVerbosity::Warning, TEXT("Manifest from %s is missing the compressed length"), *BundleManifestKey);
+			return false;
+		}
+		if (!BundleObject.TryGetNumberField(TEXT("UncompressedLength"), Bundle.UncompressedLength))
+		{
+			Context->Logf(ELogVerbosity::Warning, TEXT("Manifest from %s is missing the uncompressed length"), *BundleManifestKey);
+			return false;
+		}
+
+		Bundles.Add(MoveTemp(Bundle));
+	}
+
+	return true;
+}
+
+void FS3CacheStore::RemoveUnusedBundles()
+{
+	IFileManager& FileManager = IFileManager::Get();
+
+	// Find all the files we want to keep
+	TSet<FString> KeepFiles;
+	for (const FBundle& Bundle : Bundles)
+	{
+		KeepFiles.Add(Bundle.Name);
+	}
+
+	// Find all the files on disk
+	TArray<FString> Files;
+	FileManager.FindFiles(Files, *CacheDir);
+
+	// Remove anything left over
+	for (const FString& File : Files)
+	{
+		if (!KeepFiles.Contains(File))
+		{
+			FileManager.Delete(*(CacheDir / File));
+		}
+	}
+}
+
+void FS3CacheStore::ReadBundle(FBundle& Bundle)
+{
+	IFileManager& FileManager = IFileManager::Get();
+
+	// Open the file for reading. If this fails, assume it's because the download was aborted.
+	TUniquePtr<FArchive> Reader(FileManager.CreateFileReader(*Bundle.LocalFile));
+	if (!Reader.IsValid() || Reader->IsError())
+	{
+		UE_LOG(LogDerivedDataCache, Warning, TEXT("Unable to open bundle %s for reading. Ignoring."), *Bundle.LocalFile);
+		return;
+	}
+	
+	struct FFileHeader
+	{
+		uint32 Signature;
+		int32 NumRecords;
+	};
+
+	FFileHeader Header;
+	Reader->Serialize(&Header, sizeof(Header));
+
+	const uint32 BundleSignature = (uint32)'D' | ((uint32)'D' << 8) | ((uint32)'B' << 16);
+	const uint32 BundleSignatureV1 = BundleSignature | (1U << 24);
+	if (Header.Signature != BundleSignatureV1)
+	{
+		UE_LOG(LogDerivedDataCache, Warning, TEXT("Unable to read bundle with signature %08x"), Header.Signature);
+		return;
+	}
+
+	struct FFileRecord
+	{
+		FSHAHash Hash;
+		uint32 Length;
+	};
+
+	TArray<FFileRecord> Records;
+	Records.SetNum(Header.NumRecords);
+	Reader->Serialize(Records.GetData(), Header.NumRecords * sizeof(FFileRecord));
+
+	Bundle.Entries.Reserve(Records.Num());
+
+	int64 Offset = Reader->Tell();
+	for (const FFileRecord& Record : Records)
+	{
+		FBundleEntry& Entry = Bundle.Entries.Add(Record.Hash);
+		Entry.Offset = Offset;
+		Entry.Length = Record.Length;
+		Offset += Record.Length;
+		check(Offset <= Bundle.UncompressedLength);
+	}
+}
+
+bool FS3CacheStore::FindBundleEntry(const TCHAR* CacheKey, const FBundle*& OutBundle, const FBundleEntry*& OutBundleEntry) const
+{
+	FSHAHash Hash;
+
+	auto AnsiString = StringCast<ANSICHAR>(*BuildPathForLegacyCache(CacheKey).ToUpper());
+	FSHA1::HashBuffer(AnsiString.Get(), AnsiString.Length(), Hash.Hash);
+
+	for (const FBundle& Bundle : Bundles)
+	{
+		const FBundleEntry* Entry = Bundle.Entries.Find(Hash);
+		if (Entry != nullptr)
+		{
+			OutBundle = &Bundle;
+			OutBundleEntry = Entry;
+			return true;
+		}
+	}
+
+	return false;
+}
+
+bool FS3CacheStore::ApplyDebugOptions(FBackendDebugOptions& InOptions)
+{
+	DebugOptions = InOptions;
+	return true;
+}
+
+bool FS3CacheStore::DidSimulateMiss(const TCHAR* InKey)
+{
+	if (DebugOptions.RandomMissRate == 0 && DebugOptions.SimulateMissTypes.IsEmpty())
+	{
+		return false;
+	}
+	FScopeLock Lock(&MissedKeysCS);
+	return DebugMissedKeys.Contains(FName(InKey));
+}
+
+bool FS3CacheStore::ShouldSimulateMiss(const TCHAR* InKey)
+{
+	// once missed, always missed
+	if (DidSimulateMiss(InKey))
+	{
+		return true;
+	}
+
+	if (DebugOptions.ShouldSimulateMiss(InKey))
+	{
+		FScopeLock Lock(&MissedKeysCS);
+		UE_LOG(LogDerivedDataCache, Verbose, TEXT("Simulating miss in %s for %s"), *GetName(), InKey);
+		DebugMissedKeys.Add(FName(InKey));
+		return true;
+	}
+
+	return false;
+}
+
+FOptionalCacheRecord FS3CacheStore::GetCacheRecordOnly(
+	const FStringView Name,
+	const FCacheKey& Key,
+	const FCacheRecordPolicy& Policy)
+{
+	if (!IsUsable())
+	{
+		UE_LOG(LogDerivedDataCache, VeryVerbose,
+			TEXT("%s: Skipped get of %s from '%.*s' because this cache store is not available"),
+			*GetName(), *WriteToString<96>(Key), Name.Len(), Name.GetData());
+		return FOptionalCacheRecord();
+	}
+
+	// Skip the request if querying the cache is disabled.
+	if (!EnumHasAnyFlags(Policy.GetRecordPolicy(), ECachePolicy::QueryLocal))
+	{
+		UE_LOG(LogDerivedDataCache, VeryVerbose, TEXT("%s: Skipped get of %s from '%.*s' due to cache policy"),
+			*GetName(), *WriteToString<96>(Key), Name.Len(), Name.GetData());
+		return FOptionalCacheRecord();
+	}
+
+	//if (ShouldSimulateMiss(Key))
+	//{
+	//	UE_LOG(LogDerivedDataCache, Verbose, TEXT("%s: Simulated miss for get of %s from '%.*s'"),
+	//		*GetName(), *WriteToString<96>(Key), Name.Len(), Name.GetData());
+	//	return FOptionalCacheRecord();
+	//}
+
+	TStringBuilder<256> Path;
+	BuildCachePackagePath(Key, Path);
+
+
+	FOptionalCacheRecord Record;
+	{
+		FCbPackage Package;
+		if (!LoadFileWithHash(Path, Name, [&Package](FArchive& Ar) { Package.TryLoad(Ar); }))
+		{
+			UE_LOG(LogDerivedDataCache, Verbose, TEXT("%s: Cache miss with missing package for %s from '%.*s'"),
+				*GetName(), *WriteToString<96>(Key), Name.Len(), Name.GetData());
+			return Record;
+		}
+
+		if (ValidateCompactBinary(Package, ECbValidateMode::Default | ECbValidateMode::Package) != ECbValidateError::None)
+		{
+			UE_LOG(LogDerivedDataCache, Display, TEXT("%s: Cache miss with invalid package for %s from '%.*s'"),
+				*GetName(), *WriteToString<96>(Key), Name.Len(), Name.GetData());
+			return Record;
+		}
+
+		Record = FCacheRecord::Load(Package);
+		if (Record.IsNull())
+		{
+			UE_LOG(LogDerivedDataCache, Display, TEXT("%s: Cache miss with record load failure for %s from '%.*s'"),
+				*GetName(), *WriteToString<96>(Key), Name.Len(), Name.GetData());
+			return Record;
+		}
+	}
+
+	return Record;
+}
+
+FOptionalCacheRecord FS3CacheStore::GetCacheRecord(
+	const FStringView Name,
+	const FCacheKey& Key,
+	const FCacheRecordPolicy& Policy,
+	EStatus& OutStatus)
+{
+	FOptionalCacheRecord Record = GetCacheRecordOnly(Name, Key, Policy);
+	if (Record.IsNull())
+	{
+		OutStatus = EStatus::Error;
+		return Record;
+	}
+
+	OutStatus = EStatus::Ok;
+
+	FCacheRecordBuilder RecordBuilder(Key);
+
+	const ECachePolicy RecordPolicy = Policy.GetRecordPolicy();
+	if (!EnumHasAnyFlags(RecordPolicy, ECachePolicy::SkipMeta))
+	{
+		RecordBuilder.SetMeta(FCbObject(Record.Get().GetMeta()));
+	}
+
+	for (const FValueWithId& Value : Record.Get().GetValues())
+	{
+		const FValueId& Id = Value.GetId();
+		const ECachePolicy ValuePolicy = Policy.GetValuePolicy(Id);
+		FValue Content;
+		if (GetCacheContent(Name, Key, Id, Value, ValuePolicy, Content))
+		{
+			RecordBuilder.AddValue(Id, MoveTemp(Content));
+		}
+		else if (EnumHasAnyFlags(RecordPolicy, ECachePolicy::PartialRecord))
+		{
+			OutStatus = EStatus::Error;
+			RecordBuilder.AddValue(Value);
+		}
+		else
+		{
+			OutStatus = EStatus::Error;
+			return FOptionalCacheRecord();
+		}
+	}
+
+	return RecordBuilder.Build();
+}
+
+bool FS3CacheStore::GetCacheValueOnly(
+	const FStringView Name,
+	const FCacheKey& Key,
+	const ECachePolicy Policy,
+	FValue& OutValue)
+{
+	if (!IsUsable())
+	{
+		UE_LOG(LogDerivedDataCache, VeryVerbose,
+			TEXT("%s: Skipped get of %s from '%.*s' because this cache store is not available"),
+			*GetName(), *WriteToString<96>(Key), Name.Len(), Name.GetData());
+		return false;
+	}
+
+	// Skip the request if querying the cache is disabled.
+	if (!EnumHasAnyFlags(Policy, ECachePolicy::QueryLocal))
+	{
+		UE_LOG(LogDerivedDataCache, VeryVerbose, TEXT("%s: Skipped get of %s from '%.*s' due to cache policy"),
+			*GetName(), *WriteToString<96>(Key), Name.Len(), Name.GetData());
+		return false;
+	}
+
+	//if (ShouldSimulateMiss(Key))
+	//{
+	//	UE_LOG(LogDerivedDataCache, Verbose, TEXT("%s: Simulated miss for get of %s from '%.*s'"),
+	//		*CachePath, *WriteToString<96>(Key), Name.Len(), Name.GetData());
+	//	return false;
+	//}
+
+	TStringBuilder<256> Path;
+	BuildCachePackagePath(Key, Path);
+
+	FCbPackage Package;
+	if (!LoadFileWithHash(Path, Name, [&Package](FArchive& Ar) { Package.TryLoad(Ar); }))
+	{
+		UE_LOG(LogDerivedDataCache, Verbose, TEXT("%s: Cache miss with missing package for %s from '%.*s'"),
+			*GetName(), *WriteToString<96>(Key), Name.Len(), Name.GetData());
+		return false;
+	}
+
+	if (ValidateCompactBinary(Package, ECbValidateMode::Default | ECbValidateMode::Package) != ECbValidateError::None)
+	{
+		UE_LOG(LogDerivedDataCache, Display, TEXT("%s: Cache miss with invalid package for %s from '%.*s'"),
+			*GetName(), *WriteToString<96>(Key), Name.Len(), Name.GetData());
+		return false;
+	}
+
+	const FCbObjectView Object = Package.GetObject();
+	const FIoHash RawHash = Object["RawHash"].AsHash();
+	const uint64 RawSize = Object["RawSize"].AsUInt64(MAX_uint64);
+	if (RawHash.IsZero() || RawSize == MAX_uint64)
+	{
+		UE_LOG(LogDerivedDataCache, Display, TEXT("%s: Cache miss with invalid value for %s from '%.*s'"),
+			*GetName(), *WriteToString<96>(Key), Name.Len(), Name.GetData());
+		return false;
+	}
+
+	if (const FCbAttachment* const Attachment = Package.FindAttachment(RawHash))
+	{
+		const FCompressedBuffer& Data = Attachment->AsCompressedBinary();
+		if (Data.GetRawHash() != RawHash || Data.GetRawSize() != RawSize)
+		{
+			UE_LOG(LogDerivedDataCache, Display,
+				TEXT("%s: Cache miss with invalid value attachment for %s from '%.*s'"),
+				*GetName(), *WriteToString<96>(Key), Name.Len(), Name.GetData());
+			return false;
+		}
+		OutValue = FValue(Data);
+	}
+	else
+	{
+		OutValue = FValue(RawHash, RawSize);
+	}
+
+	return true;
+}
+
+bool FS3CacheStore::GetCacheValue(
+	const FStringView Name,
+	const FCacheKey& Key,
+	const ECachePolicy Policy,
+	FValue& OutValue)
+{
+	return GetCacheValueOnly(Name, Key, Policy, OutValue) && GetCacheContent(Name, Key, {}, OutValue, Policy, OutValue);
+}
+
+bool FS3CacheStore::GetCacheContentExists(const FCacheKey& Key, const FIoHash& RawHash) const
+{
+	TStringBuilder<256> Path;
+	BuildCacheContentPath(RawHash, Path);
+	return FileExists(Path);
+}
+
+bool FS3CacheStore::GetCacheContent(
+	const FStringView Name,
+	const FCacheKey& Key,
+	const FValueId& Id,
+	const FValue& Value,
+	const ECachePolicy Policy,
+	FValue& OutValue) const
+{
+	if (!EnumHasAnyFlags(Policy, ECachePolicy::Query))
+	{
+		OutValue = Value.RemoveData();
+		return true;
+	}
+
+	if (Value.HasData())
+	{
+		OutValue = EnumHasAnyFlags(Policy, ECachePolicy::SkipData) ? Value.RemoveData() : Value;
+		return true;
+	}
+
+	const FIoHash& RawHash = Value.GetRawHash();
+
+	TStringBuilder<256> Path;
+	BuildCacheContentPath(RawHash, Path);
+	if (EnumHasAnyFlags(Policy, ECachePolicy::SkipData))
+	{
+		if (FileExists(Path))
+		{
+			OutValue = Value;
+			return true;
+		}
+	}
+	else
+	{
+		FCompressedBuffer CompressedBuffer;
+		if (LoadFileWithHash(Path, Name, [&CompressedBuffer](FArchive& Ar) { CompressedBuffer = FCompressedBuffer::Load(Ar); }))
+		{
+			if (CompressedBuffer.GetRawHash() == RawHash)
+			{
+				OutValue = FValue(MoveTemp(CompressedBuffer));
+				return true;
+			}
+			UE_LOG(LogDerivedDataCache, Display,
+				TEXT("%s: Cache miss with corrupted value %s with hash %s for %s from '%.*s'"),
+				*GetName(), *WriteToString<16>(Id), *WriteToString<48>(RawHash),
+				*WriteToString<96>(Key), Name.Len(), Name.GetData());
+			return false;
+		}
+	}
+
+	UE_LOG(LogDerivedDataCache, Verbose,
+		TEXT("%s: Cache miss with missing value %s with hash %s for %s from '%.*s'"),
+		*GetName(), *WriteToString<16>(Id), *WriteToString<48>(RawHash), *WriteToString<96>(Key),
+		Name.Len(), Name.GetData());
+	return false;
+}
+
+void FS3CacheStore::GetCacheContent(
+	const FStringView Name,
+	const FCacheKey& Key,
+	const FValueId& Id,
+	const FValue& Value,
+	const ECachePolicy Policy,
+	FCompressedBufferReader& Reader,
+	TUniquePtr<FArchive>& OutArchive) const
+{
+	if (!EnumHasAnyFlags(Policy, ECachePolicy::Query))
+	{
+		return;
+	}
+
+	if (Value.HasData())
+	{
+		if (!EnumHasAnyFlags(Policy, ECachePolicy::SkipData))
+		{
+			Reader.SetSource(Value.GetData());
+		}
+		OutArchive.Reset();
+		return;
+	}
+
+	const FIoHash& RawHash = Value.GetRawHash();
+
+	TStringBuilder<256> Path;
+	BuildCacheContentPath(RawHash, Path);
+	if (EnumHasAllFlags(Policy, ECachePolicy::SkipData))
+	{
+		if (FileExists(Path))
+		{
+			return;
+		}
+	}
+	else
+	{
+		OutArchive = OpenFile(Path, Name);
+		if (OutArchive)
+		{
+			Reader.SetSource(*OutArchive);
+			if (Reader.GetRawHash() == RawHash)
+			{
+				return;
+			}
+			UE_LOG(LogDerivedDataCache, Display,
+				TEXT("%s: Cache miss with corrupted value %s with hash %s for %s from '%.*s'"),
+				*GetName(), *WriteToString<16>(Id), *WriteToString<48>(RawHash),
+				*WriteToString<96>(Key), Name.Len(), Name.GetData());
+			Reader.ResetSource();
+			OutArchive.Reset();
+			return;
+		}
+	}
+
+	UE_LOG(LogDerivedDataCache, Verbose,
+		TEXT("%s: Cache miss with missing value %s with hash %s for %s from '%.*s'"),
+		*GetName(), *WriteToString<16>(Id), *WriteToString<48>(RawHash), *WriteToString<96>(Key),
+		Name.Len(), Name.GetData());
+}
+
+void FS3CacheStore::BuildCachePackagePath(const FCacheKey& CacheKey, FStringBuilderBase& Path) const
+{
+	BuildPathForCachePackage(CacheKey, Path);
+}
+
+void FS3CacheStore::BuildCacheContentPath(const FIoHash& RawHash, FStringBuilderBase& Path) const
+{
+	BuildPathForCacheContent(RawHash, Path);
+}
+
+bool FS3CacheStore::LoadFileWithHash(
+	FStringBuilderBase& Path,
+	const FStringView DebugName,
+	const TFunctionRef<void (FArchive& Ar)> ReadFunction) const
+{
+	return LoadFile(Path, DebugName, [this, &Path, &DebugName, &ReadFunction](FArchive& Ar)
+	{
+		THashingArchiveProxy<FBlake3> HashAr(Ar);
+		ReadFunction(HashAr);
+		const FBlake3Hash Hash = HashAr.GetHash();
+		FBlake3Hash SavedHash;
+		Ar << SavedHash;
+
+		if (Hash != SavedHash && !Ar.IsError())
+		{
+			Ar.SetError();
+			UE_LOG(LogDerivedDataCache, Display,
+				TEXT("%s: File %s from '%.*s' is corrupted and has hash %s when %s is expected."),
+				*GetName(), *Path, DebugName.Len(), DebugName.GetData(),
+				*WriteToString<80>(Hash), *WriteToString<80>(SavedHash));
+		}
+	});
+}
+
+bool FS3CacheStore::LoadFile(
+	FStringBuilderBase& Path,
+	const FStringView DebugName,
+	const TFunctionRef<void (FArchive& Ar)> ReadFunction) const
+{
+	check(IsUsable());
+	const double StartTime = FPlatformTime::Seconds();
+
+	int64 ReadSize = 0;
+	bool bError = false;
+
+	if (TUniquePtr<FArchive> Ar = OpenFile(Path, DebugName))
+	{
+		int64 ReadStart = Ar->Tell();
+		ReadFunction(*Ar);
+		ReadSize = Ar->Tell() - ReadStart;
+		bError = !Ar->Close();
+
+		if (bError)
+		{
+			UE_LOG(LogDerivedDataCache, Display,
+				TEXT("%s: Failed to load file %s from '%.*s'."),
+				*GetName(), *Path, DebugName.Len(), DebugName.GetData());
+		}
+	}
+
+	const double ReadDuration = FPlatformTime::Seconds() - StartTime;
+	const double ReadSpeed = ReadDuration > 0.001 ? (ReadSize / ReadDuration) / (1024.0 * 1024.0) : 0.0;
+
+	UE_LOG(LogDerivedDataCache, VeryVerbose,
+		TEXT("%s: Loaded %s from '%.*s' (%" INT64_FMT " bytes, %.02f secs, %.2f MiB/s)"),
+		*GetName(), *Path, DebugName.Len(), DebugName.GetData(), ReadSize, ReadDuration, ReadSpeed);
+
+	return !bError && ReadSize > 0;
+}
+
+TUniquePtr<FArchive> FS3CacheStore::OpenFile(FStringBuilderBase& Path, const FStringView DebugName) const
+{
+	FSHAHash Hash;
+
+	auto AnsiString = StringCast<ANSICHAR>(*FString(Path.ToString()).ToUpper());
+	FSHA1::HashBuffer(AnsiString.Get(), AnsiString.Length(), Hash.Hash);
+
+	for (const FBundle& Bundle : Bundles)
+	{
+		const FBundleEntry* Entry = Bundle.Entries.Find(Hash);
+		if (Entry != nullptr)
+		{
+			TUniquePtr<FArchive> Reader(IFileManager::Get().CreateFileReader(*Bundle.LocalFile));
+			if (Reader.IsValid() && !Reader->IsError())
+			{
+				Reader->Seek(Entry->Offset);
+				return Reader;
+			}
+		}
+	}
+
+	return nullptr;
+}
+
+bool FS3CacheStore::FileExists(FStringBuilderBase& Path) const
+{
+	FSHAHash Hash;
+
+	auto AnsiString = StringCast<ANSICHAR>(*FString(Path.ToString()).ToUpper());
+	FSHA1::HashBuffer(AnsiString.Get(), AnsiString.Length(), Hash.Hash);
+
+	for (const FBundle& Bundle : Bundles)
+	{
+		const FBundleEntry* Entry = Bundle.Entries.Find(Hash);
+		if (Entry != nullptr)
+		{
+			return IFileManager::Get().FileExists(*Bundle.LocalFile);
+		}
+	}
+
+	return false;
+}
+
+void FS3CacheStore::Put(
+	const TConstArrayView<FCachePutRequest> Requests,
+	IRequestOwner& Owner,
+	FOnCachePutComplete&& OnComplete)
+{
+	CompleteWithStatus(Requests, OnComplete, EStatus::Error);
+}
+
+void FS3CacheStore::Get(
+	const TConstArrayView<FCacheGetRequest> Requests,
+	IRequestOwner& Owner,
+	FOnCacheGetComplete&& OnComplete)
+{
+	for (const FCacheGetRequest& Request : Requests)
+	{
+		TRACE_CPUPROFILER_EVENT_SCOPE(S3DDC_Get);
+		TRACE_COUNTER_INCREMENT(S3DDC_Get);
+		COOK_STAT(auto Timer = UsageStats.TimeGet());
+		EStatus Status = EStatus::Ok;
+		if (FOptionalCacheRecord Record = GetCacheRecord(Request.Name, Request.Key, Request.Policy, Status))
+		{
+			UE_LOG(LogDerivedDataCache, Verbose, TEXT("%s: Cache hit for %s from '%s'"),
+				*GetName(), *WriteToString<96>(Request.Key), *Request.Name);
+			TRACE_COUNTER_INCREMENT(S3DDC_GetHit);
+			COOK_STAT(Timer.AddHit(Private::GetCacheRecordCompressedSize(Record.Get())));
+			OnComplete({Request.Name, MoveTemp(Record).Get(), Request.UserData, Status});
+		}
+		else
+		{
+			OnComplete(Request.MakeResponse(Status));
+		}
+	}
+}
+
+void FS3CacheStore::PutValue(
+	const TConstArrayView<FCachePutValueRequest> Requests,
+	IRequestOwner& Owner,
+	FOnCachePutValueComplete&& OnComplete)
+{
+	CompleteWithStatus(Requests, OnComplete, EStatus::Error);
+}
+
+void FS3CacheStore::GetValue(
+	const TConstArrayView<FCacheGetValueRequest> Requests,
+	IRequestOwner& Owner,
+	FOnCacheGetValueComplete&& OnComplete)
+{
+	for (const FCacheGetValueRequest& Request : Requests)
+	{
+		TRACE_CPUPROFILER_EVENT_SCOPE(S3DDC_Get);
+		TRACE_COUNTER_INCREMENT(S3DDC_Get);
+		COOK_STAT(auto Timer = UsageStats.TimeGet());
+		FValue Value;
+		if (GetCacheValue(Request.Name, Request.Key, Request.Policy, Value))
+		{
+			UE_LOG(LogDerivedDataCache, Verbose, TEXT("%s: Cache hit for %s from '%s'"),
+				*GetName(), *WriteToString<96>(Request.Key), *Request.Name);
+			TRACE_COUNTER_INCREMENT(S3DDC_GetHit);
+			COOK_STAT(Timer.AddHit(Value.GetData().GetCompressedSize()));
+			OnComplete({Request.Name, Request.Key, Value, Request.UserData, EStatus::Ok});
+		}
+		else
+		{
+			OnComplete(Request.MakeResponse(EStatus::Error));
+		}
+	}
+}
+
+void FS3CacheStore::GetChunks(
+	const TConstArrayView<FCacheGetChunkRequest> Requests,
+	IRequestOwner& Owner,
+	FOnCacheGetChunkComplete&& OnComplete)
+{
+	TArray<FCacheGetChunkRequest, TInlineAllocator<16>> SortedRequests(Requests);
+	SortedRequests.StableSort(TChunkLess());
+
+	bool bHasValue = false;
+	FValue Value;
+	FValueId ValueId;
+	FCacheKey ValueKey;
+	TUniquePtr<FArchive> ValueAr;
+	FCompressedBufferReader ValueReader;
+	FOptionalCacheRecord Record;
+	for (const FCacheGetChunkRequest& Request : SortedRequests)
+	{
+		const bool bExistsOnly = EnumHasAnyFlags(Request.Policy, ECachePolicy::SkipData);
+		TRACE_CPUPROFILER_EVENT_SCOPE(S3DDC_Get);
+		TRACE_COUNTER_INCREMENT(S3DDC_Get);
+		COOK_STAT(auto Timer = bExistsOnly ? UsageStats.TimeProbablyExists() : UsageStats.TimeGet());
+		if (!(bHasValue && ValueKey == Request.Key && ValueId == Request.Id) || ValueReader.HasSource() < !bExistsOnly)
+		{
+			ValueReader.ResetSource();
+			ValueAr.Reset();
+			ValueKey = {};
+			ValueId.Reset();
+			Value.Reset();
+			bHasValue = false;
+			if (Request.Id.IsValid())
+			{
+				if (!(Record && Record.Get().GetKey() == Request.Key))
+				{
+					FCacheRecordPolicyBuilder PolicyBuilder(ECachePolicy::None);
+					PolicyBuilder.AddValuePolicy(Request.Id, Request.Policy);
+					Record.Reset();
+					Record = GetCacheRecordOnly(Request.Name, Request.Key, PolicyBuilder.Build());
+				}
+				if (Record)
+				{
+					if (const FValueWithId& ValueWithId = Record.Get().GetValue(Request.Id))
+					{
+						bHasValue = true;
+						Value = ValueWithId;
+						ValueId = Request.Id;
+						ValueKey = Request.Key;
+						GetCacheContent(Request.Name, Request.Key, ValueId, Value, Request.Policy, ValueReader, ValueAr);
+					}
+				}
+			}
+			else
+			{
+				ValueKey = Request.Key;
+				bHasValue = GetCacheValueOnly(Request.Name, Request.Key, Request.Policy, Value);
+				if (bHasValue)
+				{
+					GetCacheContent(Request.Name, Request.Key, Request.Id, Value, Request.Policy, ValueReader, ValueAr);
+				}
+			}
+		}
+		if (bHasValue)
+		{
+			const uint64 RawOffset = FMath::Min(Value.GetRawSize(), Request.RawOffset);
+			const uint64 RawSize = FMath::Min(Value.GetRawSize() - RawOffset, Request.RawSize);
+			UE_LOG(LogDerivedDataCache, Verbose, TEXT("%s: Cache hit for %s from '%s'"),
+				*GetName(), *WriteToString<96>(Request.Key, '/', Request.Id), *Request.Name);
+			TRACE_COUNTER_INCREMENT(S3DDC_GetHit);
+			COOK_STAT(Timer.AddHit(!bExistsOnly ? RawSize : 0));
+			FSharedBuffer Buffer;
+			if (!bExistsOnly)
+			{
+				Buffer = ValueReader.Decompress(RawOffset, RawSize);
+			}
+			const EStatus ChunkStatus = bExistsOnly || Buffer.GetSize() == RawSize ? EStatus::Ok : EStatus::Error;
+			OnComplete({Request.Name, Request.Key, Request.Id, Request.RawOffset,
+				RawSize, Value.GetRawHash(), MoveTemp(Buffer), Request.UserData, ChunkStatus});
+			continue;
+		}
+
+		OnComplete(Request.MakeResponse(EStatus::Error));
+	}
+}
+
+ILegacyCacheStore* CreateS3CacheStore(const TCHAR* RootManifestPath, const TCHAR* BaseUrl, const TCHAR* Region, const TCHAR* CanaryObjectKey, const TCHAR* CachePath)
+{
+	return new FS3CacheStore(RootManifestPath, BaseUrl, Region, CanaryObjectKey, CachePath);
+}
+
+} // UE::DerivedData
+
+#else
+
+namespace UE::DerivedData
+{
+
+ILegacyCacheStore* CreateS3CacheStore(const TCHAR* RootManifestPath, const TCHAR* BaseUrl, const TCHAR* Region, const TCHAR* CanaryObjectKey, const TCHAR* CachePath)
+{
+	return nullptr;
+}
+
+} // UE::DerivedData
+
+#endif // WITH_S3_DDC_BACKEND
diff --git a/Engine/Source/Developer/DerivedDataCache/Private/Tests/HttpCacheStoreTest.cpp b/Engine/Source/Developer/DerivedDataCache/Private/Tests/HttpCacheStoreTest.cpp
new file mode 100644
--- /dev/null
+++ b/Engine/Source/Developer/DerivedDataCache/Private/Tests/HttpCacheStoreTest.cpp
@@ -0,0 +1,799 @@
+// Copyright Epic Games, Inc. All Rights Reserved.
+
+#include "Async/ParallelFor.h"
+#include "DerivedDataBackendInterface.h"
+#include "DerivedDataCacheKey.h"
+#include "DerivedDataCacheRecord.h"
+#include "DerivedDataRequestOwner.h"
+#include "DerivedDataValue.h"
+#include "DerivedDataValueId.h"
+#include "Memory/CompositeBuffer.h"
+#include "Misc/AutomationTest.h"
+#include "Misc/Paths.h"
+#include "Misc/SecureHash.h"
+#include "Serialization/CompactBinaryWriter.h"
+#include "ZenServerInterface.h"
+
+// Test is targeted at HttpDerivedDataBackend but with some backend test interface it could be generalized
+// to function against all backends.
+
+#if WITH_DEV_AUTOMATION_TESTS && WITH_HTTP_DDC_BACKEND
+
+DEFINE_LOG_CATEGORY_STATIC(LogHttpDerivedDataBackendTests, Log, All);
+
+#define TEST_NAME_ROOT TEXT("System.DerivedDataCache.HttpDerivedDataBackend")
+
+#define IMPLEMENT_HTTPDERIVEDDATA_AUTOMATION_TEST( TClass, PrettyName, TFlags ) \
+	IMPLEMENT_CUSTOM_COMPLEX_AUTOMATION_TEST(TClass, FHttpCacheStoreTestBase, TEST_NAME_ROOT PrettyName, TFlags) \
+	void TClass::GetTests(TArray<FString>& OutBeautifiedNames, TArray <FString>& OutTestCommands) const \
+	{ \
+		if (CheckPrequisites()) \
+		{ \
+			OutBeautifiedNames.Add(TEST_NAME_ROOT PrettyName); \
+			OutTestCommands.Add(FString()); \
+		} \
+	}
+
+namespace UE::DerivedData
+{
+
+FDerivedDataBackendInterface* GetAnyHttpCacheStore(
+	FString& OutDomain,
+	FString& OutOAuthProvider,
+	FString& OutOAuthClientId,
+	FString& OutOAuthSecret,
+	FString& OutNamespace,
+	FString& OutStructuredNamespace);
+
+ILegacyCacheStore* CreateZenCacheStore(const TCHAR* NodeName, const TCHAR* ServiceUrl, const TCHAR* Namespace);
+
+class FHttpCacheStoreTestBase : public FAutomationTestBase
+{
+public:
+	FHttpCacheStoreTestBase(const FString& InName, const bool bInComplexTask)
+	: FAutomationTestBase(InName, bInComplexTask)
+	{
+	}
+
+	bool CheckPrequisites() const
+	{
+		if (FDerivedDataBackendInterface* Backend = GetTestBackend())
+		{
+			return true;
+		}
+		return false;
+	}
+
+protected:
+
+	void ConcurrentTestWithStats(TFunctionRef<void()> TestFunction, int32 ThreadCount, double Duration)
+	{
+		std::atomic<uint64> Requests{ 0 };
+		std::atomic<uint64> MaxLatency{ 0 };
+		std::atomic<uint64> TotalMS{ 0 };
+		std::atomic<uint64> TotalRequests{ 0 };
+
+		FEvent* StartEvent = FPlatformProcess::GetSynchEventFromPool(true);
+		FEvent* LastEvent = FPlatformProcess::GetSynchEventFromPool(true);
+		std::atomic<double> StopTime{ 0.0 };
+		std::atomic<uint64> ActiveCount{ 0 };
+
+		for (int32 ThreadIndex = 0; ThreadIndex < ThreadCount; ++ThreadIndex)
+		{
+			ActiveCount++;
+			Async(
+				ThreadIndex < FTaskGraphInterface::Get().GetNumWorkerThreads() ? EAsyncExecution::TaskGraph : EAsyncExecution::Thread,
+				[&]()
+				{
+					// No false start, wait until everyone is ready before starting the test
+					StartEvent->Wait();
+
+					while (FPlatformTime::Seconds() < StopTime.load(std::memory_order_relaxed))
+					{
+						uint64 Before = FPlatformTime::Cycles64();
+						TestFunction();
+						uint64 Delta = FPlatformTime::Cycles64() - Before;
+						Requests++;
+						TotalMS += FPlatformTime::ToMilliseconds64(Delta);
+						TotalRequests++;
+
+						// Compare exchange loop until we either succeed to set the maximum value
+						// or we bail out because we don't have the maximum value anymore.
+						while (true)
+						{
+							uint64 Snapshot = MaxLatency.load();
+							if (Delta > Snapshot)
+							{
+								// Only do the exchange if the value has not changed since we confirmed
+								// we had a bigger one.
+								if (MaxLatency.compare_exchange_strong(Snapshot, Delta))
+								{
+									// Exchange succeeded
+									break;
+								}
+							}
+							else
+							{
+								// We don't have the maximum
+								break;
+							}
+						}
+					}
+					
+					if (--ActiveCount == 0)
+					{
+						LastEvent->Trigger();
+					}
+				}
+			);
+		}
+
+		StopTime = FPlatformTime::Seconds() + Duration;
+
+		// GO!
+		StartEvent->Trigger();
+
+		while (FPlatformTime::Seconds() < StopTime)
+		{
+			FPlatformProcess::Sleep(1.0f);
+
+			if (TotalRequests)
+			{
+				UE_LOG(LogHttpDerivedDataBackendTests, Display, TEXT("RPS: %llu, AvgLatency: %.02f ms, MaxLatency: %.02f s"), Requests.exchange(0), double(TotalMS) / TotalRequests, FPlatformTime::ToSeconds(MaxLatency));
+			}
+			else
+			{
+				UE_LOG(LogHttpDerivedDataBackendTests, Display, TEXT("RPS: %llu, AvgLatency: N/A, MaxLatency: %.02f s"), Requests.exchange(0), FPlatformTime::ToSeconds(MaxLatency));
+			}
+		}
+
+		LastEvent->Wait();
+
+		FPlatformProcess::ReturnSynchEventToPool(StartEvent);
+		FPlatformProcess::ReturnSynchEventToPool(LastEvent);
+	}
+
+	FDerivedDataBackendInterface* GetTestBackend() const
+	{
+		static FDerivedDataBackendInterface* CachedBackend = GetAnyHttpCacheStore(
+			TestDomain, TestOAuthProvider, TestOAuthClientId, TestOAuthSecret, TestNamespace, TestStructuredNamespace);
+		return CachedBackend;
+	}
+
+	bool GetRecords(TConstArrayView<FCacheRecord> Records, FCacheRecordPolicy Policy, TArray<FCacheRecord>& OutRecords)
+	{
+		using namespace UE::DerivedData;
+		ICacheStore* TestBackend = GetTestBackend();
+
+		TArray<FCacheGetRequest> Requests;
+		Requests.Reserve(Records.Num());
+
+		for (int32 RecordIndex = 0; RecordIndex < Records.Num(); ++RecordIndex)
+		{
+			const FCacheRecord& Record = Records[RecordIndex];
+			Requests.Add({ {TEXT("FHttpCacheStoreTestBase")}, Record.GetKey(), Policy, static_cast<uint64>(RecordIndex) });
+		}
+
+		struct FGetOutput
+		{
+			FCacheRecord Record;
+			EStatus Status = EStatus::Error;
+		};
+
+		TArray<TOptional<FGetOutput>> GetOutputs;
+		GetOutputs.SetNum(Records.Num());
+		FRequestOwner RequestOwner(EPriority::Blocking);
+		TestBackend->Get(Requests, RequestOwner, [&GetOutputs](FCacheGetResponse&& Response)
+			{
+				FCacheRecordBuilder RecordBuilder(Response.Record.GetKey());
+
+				if (Response.Record.GetMeta())
+				{
+					RecordBuilder.SetMeta(FCbObject::Clone(Response.Record.GetMeta()));
+				}
+
+				for (const FValueWithId& Value : Response.Record.GetValues())
+				{
+					if (Value)
+					{
+						RecordBuilder.AddValue(Value);
+					}
+				}
+				GetOutputs[Response.UserData].Emplace(FGetOutput{ RecordBuilder.Build(), Response.Status });
+			});
+		RequestOwner.Wait();
+
+		for (int32 RecordIndex = 0; RecordIndex < Records.Num(); ++RecordIndex)
+		{
+			FGetOutput& ReceivedOutput = GetOutputs[RecordIndex].GetValue();
+			if (ReceivedOutput.Status != EStatus::Ok)
+			{
+				return false;
+			}
+			OutRecords.Add(MoveTemp(ReceivedOutput.Record));
+		}
+
+		return true;
+	}
+
+	bool GetValues(TConstArrayView<FValue> Values, ECachePolicy Policy, TArray<FValue>& OutValues, const char* BucketName = nullptr)
+	{
+		using namespace UE::DerivedData;
+		ICacheStore* TestBackend = GetTestBackend();
+		FCacheBucket TestCacheBucket(BucketName ? BucketName : "AutoTestDummy");
+
+		TArray<FCacheGetValueRequest> Requests;
+		Requests.Reserve(Values.Num());
+
+		for (int32 ValueIndex = 0; ValueIndex < Values.Num(); ++ValueIndex)
+		{
+			const FValue& Value = Values[ValueIndex];
+			FCacheKey Key;
+			Key.Bucket = TestCacheBucket;
+			Key.Hash = Value.GetRawHash();
+			Requests.Add({ {TEXT("FHttpCacheStoreTestBase")}, Key, Policy, static_cast<uint64>(ValueIndex) });
+		}
+
+		struct FGetValueOutput
+		{
+			FValue Value;
+			EStatus Status = EStatus::Error;
+		};
+
+		TArray<TOptional<FGetValueOutput>> GetValueOutputs;
+		GetValueOutputs.SetNum(Values.Num());
+		FRequestOwner RequestOwner(EPriority::Blocking);
+		TestBackend->GetValue(Requests, RequestOwner, [&GetValueOutputs](FCacheGetValueResponse&& Response)
+			{
+				GetValueOutputs[Response.UserData].Emplace(FGetValueOutput{ Response.Value, Response.Status });
+			});
+		RequestOwner.Wait();
+
+		for (int32 ValueIndex = 0; ValueIndex < Values.Num(); ++ValueIndex)
+		{
+			FGetValueOutput& ReceivedOutput = GetValueOutputs[ValueIndex].GetValue();
+			if (ReceivedOutput.Status != EStatus::Ok)
+			{
+				return false;
+			}
+			OutValues.Add(MoveTemp(ReceivedOutput.Value));
+		}
+
+		return true;
+	}
+
+	bool GetRecordChunks(TConstArrayView<FCacheRecord> Records, FCacheRecordPolicy Policy, uint64 Offset, uint64 Size, TArray<FSharedBuffer>& OutChunks)
+	{
+		using namespace UE::DerivedData;
+		ICacheStore* TestBackend = GetTestBackend();
+
+		TArray<FCacheGetChunkRequest> Requests;
+
+		int32 OverallIndex = 0;
+		for (int32 RecordIndex = 0; RecordIndex < Records.Num(); ++RecordIndex)
+		{
+			const FCacheRecord& Record = Records[RecordIndex];
+			TConstArrayView<FValueWithId> Values = Record.GetValues();
+			for (int32 ValueIndex = 0; ValueIndex < Values.Num(); ++ValueIndex)
+			{
+				const FValueWithId& Value = Values[ValueIndex];
+				Requests.Add({ {TEXT("FHttpCacheStoreTestBase")}, Record.GetKey(),  Value.GetId(), Offset, Size, Value.GetRawHash(), Policy.GetValuePolicy(Value.GetId()), static_cast<uint64>(OverallIndex) });
+				++OverallIndex;
+			}
+		}
+
+		struct FGetChunksOutput
+		{
+			FSharedBuffer Chunk;
+			EStatus Status = EStatus::Error;
+		};
+
+		TArray<TOptional<FGetChunksOutput>> GetOutputs;
+		GetOutputs.SetNum(Requests.Num());
+		FRequestOwner RequestOwner(EPriority::Blocking);
+		TestBackend->GetChunks(Requests, RequestOwner, [&GetOutputs](FCacheGetChunkResponse&& Response)
+			{
+				GetOutputs[Response.UserData].Emplace(FGetChunksOutput { Response.RawData, Response.Status });
+			});
+		RequestOwner.Wait();
+
+		for (int32 RequestIndex = 0; RequestIndex < Requests.Num(); ++RequestIndex)
+		{
+			FGetChunksOutput& ReceivedOutput = GetOutputs[RequestIndex].GetValue();
+			if (ReceivedOutput.Status != EStatus::Ok)
+			{
+				return false;
+			}
+			OutChunks.Add(MoveTemp(ReceivedOutput.Chunk));
+		}
+
+		return true;
+
+	}
+	
+	void ValidateRecords(const TCHAR* Name, TConstArrayView<FCacheRecord> RecordsToTest, TConstArrayView<FCacheRecord> ReferenceRecords, FCacheRecordPolicy Policy)
+	{
+		using namespace UE::DerivedData;
+
+		if (!TestEqual(FString::Printf(TEXT("%s::Record quantity"), Name), RecordsToTest.Num(), ReferenceRecords.Num()))
+		{
+			return;
+		}
+
+		for (int32 RecordIndex = 0; RecordIndex < RecordsToTest.Num(); ++RecordIndex)
+		{
+			const FCacheRecord& ExpectedRecord = ReferenceRecords[RecordIndex];
+			const FCacheRecord& RecordToTest = RecordsToTest[RecordIndex];
+
+			if (EnumHasAnyFlags(Policy.GetRecordPolicy(), ECachePolicy::SkipMeta))
+			{
+				TestTrue(FString::Printf(TEXT("%s::Get meta null"), Name), !RecordToTest.GetMeta());
+			}
+			else
+			{
+				TestTrue(FString::Printf(TEXT("%s::Get meta equality"), Name), ExpectedRecord.GetMeta().Equals(RecordToTest.GetMeta()));
+			}
+
+			TestEqual(FString::Printf(TEXT("%s::Get value quantity"), Name), ExpectedRecord.GetValues().Num(), RecordToTest.GetValues().Num());
+
+			const TConstArrayView<FValueWithId> ExpectedValues = ExpectedRecord.GetValues();
+			const TConstArrayView<FValueWithId> ReceivedValues = RecordToTest.GetValues();
+			for (int32 ValueIndex = 0; ValueIndex < ExpectedValues.Num(); ++ValueIndex)
+			{
+				if (EnumHasAnyFlags(Policy.GetRecordPolicy(), ECachePolicy::SkipData))
+				{
+					TestTrue(FString::Printf(TEXT("%s::Get value[%d] !HasData"), Name, ValueIndex), !ReceivedValues[ValueIndex].HasData());
+				}
+				else
+				{
+					TestTrue(FString::Printf(TEXT("%s::Get value[%d] HasData"), Name, ValueIndex), ReceivedValues[ValueIndex].HasData());
+					TestTrue(FString::Printf(TEXT("%s::Get value[%d] equality"), Name, ValueIndex), ExpectedValues[ValueIndex] == ReceivedValues[ValueIndex]);
+					TestTrue(FString::Printf(TEXT("%s::Get value[%d] data equality"), Name, ValueIndex), FIoHash::HashBuffer(ReceivedValues[ValueIndex].GetData().GetCompressed()) == FIoHash::HashBuffer(ExpectedValues[ValueIndex].GetData().GetCompressed()));
+				}
+			}
+		}
+	}
+
+	void ValidateValues(const TCHAR* Name, TConstArrayView<FValue> ValuesToTest, TConstArrayView<FValue> ReferenceValues, ECachePolicy Policy)
+	{
+		using namespace UE::DerivedData;
+
+		if (!TestEqual(FString::Printf(TEXT("%s::Value quantity"), Name), ValuesToTest.Num(), ReferenceValues.Num()))
+		{
+			return;
+		}
+
+		for (int32 ValueIndex = 0; ValueIndex < ValuesToTest.Num(); ++ValueIndex)
+		{
+			const FValue& ExpectedValue = ReferenceValues[ValueIndex];
+			const FValue& ValueToTest = ValuesToTest[ValueIndex];
+
+			if (EnumHasAnyFlags(Policy, ECachePolicy::SkipData))
+			{
+				TestTrue(FString::Printf(TEXT("%s::Get value[%d] !HasData"), Name, ValueIndex), !ValueToTest.HasData());
+			}
+			else
+			{
+				TestTrue(FString::Printf(TEXT("%s::Get value[%d] HasData"), Name, ValueIndex), ValueToTest.HasData());
+				TestTrue(FString::Printf(TEXT("%s::Get value[%d] equality"), Name, ValueIndex), ExpectedValue == ValueToTest);
+				TestTrue(FString::Printf(TEXT("%s::Get value[%d] data equality"), Name, ValueIndex), FIoHash::HashBuffer(ValueToTest.GetData().GetCompressed()) == FIoHash::HashBuffer(ExpectedValue.GetData().GetCompressed()));
+			}
+		}
+	}
+
+	void ValidateRecordChunks(const TCHAR* Name, TConstArrayView<FSharedBuffer> RecordChunksToTest, TConstArrayView<FCacheRecord> ReferenceRecords, FCacheRecordPolicy Policy, uint64 Offset, uint64 Size)
+	{
+		using namespace UE::DerivedData;
+
+		int32 TotalChunks = 0;
+		for (int32 RecordIndex = 0; RecordIndex < ReferenceRecords.Num(); ++RecordIndex)
+		{
+			const FCacheRecord& Record = ReferenceRecords[RecordIndex];
+			TConstArrayView<FValueWithId> Values = Record.GetValues();
+			for (int32 ValueIndex = 0; ValueIndex < Values.Num(); ++ValueIndex)
+			{
+				++TotalChunks;
+			}
+		}
+
+		if (!TestEqual(FString::Printf(TEXT("%s::Chunk quantity"), Name), RecordChunksToTest.Num(), TotalChunks))
+		{
+			return;
+		}
+
+		int32 ChunkIndex = 0;
+		for (int32 RecordIndex = 0; RecordIndex < ReferenceRecords.Num(); ++RecordIndex)
+		{
+			const FCacheRecord& ExpectedRecord = ReferenceRecords[RecordIndex];
+
+			const TConstArrayView<FValueWithId> ExpectedValues = ExpectedRecord.GetValues();
+			for (int32 ValueIndex = 0; ValueIndex < ExpectedValues.Num(); ++ValueIndex)
+			{
+				const FSharedBuffer& ChunkToTest = RecordChunksToTest[ChunkIndex];
+
+				if (EnumHasAnyFlags(Policy.GetRecordPolicy(), ECachePolicy::SkipData))
+				{
+					TestTrue(FString::Printf(TEXT("%s::Get chunk[%d] IsNull"), Name, ChunkIndex), ChunkToTest.IsNull());
+				}
+				else
+				{
+					FSharedBuffer ReferenceBuffer = ExpectedValues[ValueIndex].GetData().Decompress();
+					FMemoryView ReferenceView = ReferenceBuffer.GetView().Mid(Offset, Size);
+					TestTrue(FString::Printf(TEXT("%s::Get chunk[%d] data equality"), Name, ChunkIndex), ReferenceView.EqualBytes(ChunkToTest.GetView()));
+				}
+				++ChunkIndex;
+			}
+		}
+	}
+
+	TArray<FCacheRecord> GetAndValidateRecords(const TCHAR* Name, TConstArrayView<FCacheRecord> Records, FCacheRecordPolicy Policy)
+	{
+		using namespace UE::DerivedData;
+		TArray<FCacheRecord> ReceivedRecords;
+		bool bGetSuccessful = GetRecords(Records, Policy, ReceivedRecords);
+		TestTrue(FString::Printf(TEXT("%s::Get status"), Name), bGetSuccessful);
+
+		if (!bGetSuccessful)
+		{
+			return TArray<FCacheRecord>();
+		}
+
+		ValidateRecords(Name, ReceivedRecords, Records, Policy);
+		return ReceivedRecords;
+	}
+
+	TArray<FValue> GetAndValidateValues(const TCHAR* Name, TConstArrayView<FValue> Values, ECachePolicy Policy)
+	{
+		using namespace UE::DerivedData;
+		TArray<FValue> ReceivedValues;
+		bool bGetSuccessful = GetValues(Values, Policy, ReceivedValues);
+		TestTrue(FString::Printf(TEXT("%s::Get status"), Name), bGetSuccessful);
+
+		if (!bGetSuccessful)
+		{
+			return TArray<FValue>();
+		}
+
+		ValidateValues(Name, ReceivedValues, Values, Policy);
+		return ReceivedValues;
+	}
+
+	TArray<FSharedBuffer> GetAndValidateRecordChunks(const TCHAR* Name, TConstArrayView<FCacheRecord> Records, FCacheRecordPolicy Policy, uint64 Offset, uint64 Size)
+	{
+		using namespace UE::DerivedData;
+		TArray<FSharedBuffer> ReceivedChunks;
+		bool bGetSuccessful = GetRecordChunks(Records, Policy, Offset, Size, ReceivedChunks);
+		TestTrue(FString::Printf(TEXT("%s::GetChunks status"), Name), bGetSuccessful);
+
+		if (!bGetSuccessful)
+		{
+			return TArray<FSharedBuffer>();
+		}
+
+		ValidateRecordChunks(Name, ReceivedChunks, Records, Policy, Offset, Size);
+		return ReceivedChunks;
+	}
+
+	TArray<FCacheRecord> GetAndValidateRecordsAndChunks(const TCHAR* Name, TConstArrayView<FCacheRecord> Records, FCacheRecordPolicy Policy)
+	{
+		GetAndValidateRecordChunks(Name, Records, Policy, 5, 5);
+		return GetAndValidateRecords(Name, Records, Policy);
+	}
+
+protected:
+	static inline FString TestDomain;
+	static inline FString TestOAuthProvider;
+	static inline FString TestOAuthClientId;
+	static inline FString TestOAuthSecret;
+	static inline FString TestNamespace;
+	static inline FString TestStructuredNamespace;
+};
+
+// Helper function to create a number of dummy cache keys for testing
+TArray<FString> CreateTestCacheKeys(FDerivedDataBackendInterface* InTestBackend, uint32 InNumKeys)
+{
+	TArray<FString> Keys;
+	TArray<uint8> KeyContents;
+	KeyContents.Add(42);
+
+	FSHA1 HashState;
+	HashState.Update(KeyContents.GetData(), KeyContents.Num());
+	HashState.Final();
+	uint8 Hash[FSHA1::DigestSize];
+	HashState.GetHash(Hash);
+	const FString HashString = BytesToHex(Hash, FSHA1::DigestSize);
+
+	for (uint32 KeyIndex = 0; KeyIndex < InNumKeys; ++KeyIndex)
+	{
+		FString NewKey = FString::Printf(TEXT("__AutoTest_Dummy_%u__%s"), KeyIndex, *HashString);
+		Keys.Add(NewKey);
+		InTestBackend->PutCachedData(*NewKey, KeyContents, false);
+	}
+	return Keys;
+}
+
+TArray<FCacheRecord> CreateTestCacheRecords(ICacheStore* InTestBackend, uint32 InNumKeys, uint32 InNumValues, FCbObject MetaContents = FCbObject(), const char* BucketName = nullptr)
+{
+	using namespace UE::DerivedData;
+	FCacheBucket TestCacheBucket(BucketName ? BucketName : "AutoTestDummy");
+
+	TArray<FCacheRecord> CacheRecords;
+	TArray<FCachePutRequest> PutRequests;
+	PutRequests.Reserve(InNumKeys);
+
+	for (uint32 KeyIndex = 0; KeyIndex < InNumKeys; ++KeyIndex)
+	{
+		FIoHashBuilder HashBuilder;
+
+		TArray<FSharedBuffer> Values;
+		for (uint32 ValueIndex = 0; ValueIndex < InNumValues; ++ValueIndex)
+		{
+			TArray<uint8> ValueContents;
+			// Add N zeroed bytes where N corresponds to the value index times 10.
+			const int32 NumBytes = (ValueIndex+1)*10;
+			ValueContents.AddUninitialized(NumBytes);
+			for (int32 ContentIndex = 0; ContentIndex < NumBytes; ++ContentIndex)
+			{
+				ValueContents[ContentIndex] = (uint8)(KeyIndex + ContentIndex);
+			}
+			Values.Emplace(MakeSharedBufferFromArray(MoveTemp(ValueContents)));
+			HashBuilder.Update(Values.Last().GetView());
+		}
+
+		FCacheKey Key;
+		Key.Bucket = TestCacheBucket;
+		Key.Hash = HashBuilder.Finalize();
+
+		FCacheRecordBuilder RecordBuilder(Key);
+
+		for (const FSharedBuffer& ValueBuffer : Values)
+		{
+			FIoHash ValueHash(FIoHash::HashBuffer(ValueBuffer));
+			RecordBuilder.AddValue(FValueId::FromHash(ValueHash), ValueBuffer);
+		}
+
+		if (MetaContents)
+		{
+			RecordBuilder.SetMeta(MoveTemp(MetaContents));
+		}
+
+		PutRequests.Add({ {TEXT("AutoTest")}, RecordBuilder.Build(), ECachePolicy::Default, KeyIndex });
+	}
+
+	FRequestOwner Owner(EPriority::Blocking);
+	InTestBackend->Put(PutRequests, Owner, [&CacheRecords, &PutRequests] (FCachePutResponse&& Response)
+	{
+		check(Response.Status == EStatus::Ok);
+		CacheRecords.Add(PutRequests[Response.UserData].Record);
+	});
+	Owner.Wait();
+
+	return CacheRecords;
+}
+
+TArray<FValue> CreateTestCacheValues(ICacheStore* InTestBackend, uint32 InNumValues, const char* BucketName = nullptr)
+{
+	using namespace UE::DerivedData;
+	FCacheBucket TestCacheBucket(BucketName ? BucketName : "AutoTestDummy");
+
+	TArray<FCachePutValueRequest> PutValueRequests;
+	PutValueRequests.Reserve(InNumValues);
+
+	TArray<FSharedBuffer> ValueBuffers;
+	for (uint32 ValueIndex = 0; ValueIndex < InNumValues; ++ValueIndex)
+	{
+		TArray<uint8> ValueContents;
+		// Add N zeroed bytes where N corresponds to the value index times 10.
+		const int32 NumBytes = (ValueIndex+1)*10;
+		ValueContents.AddUninitialized(NumBytes);
+		for (int32 ContentIndex = 0; ContentIndex < NumBytes; ++ContentIndex)
+		{
+			ValueContents[ContentIndex] = (uint8)(ValueIndex + ContentIndex + 42); // offset of 42 to keep the contents distinct from record test data
+		}
+		ValueBuffers.Emplace(MakeSharedBufferFromArray(MoveTemp(ValueContents)));
+	}
+
+	uint64 KeyIndex = 0;
+	for (const FSharedBuffer& ValueBuffer : ValueBuffers)
+	{
+		FIoHash ValueHash(FIoHash::HashBuffer(ValueBuffer));
+
+		FCacheKey Key;
+		Key.Bucket = TestCacheBucket;
+		Key.Hash = ValueHash;
+
+		PutValueRequests.Add({ {TEXT("AutoTest")}, Key, FValue::Compress(ValueBuffer), ECachePolicy::Default, KeyIndex++ });
+	}
+
+	TArray<FValue> Values;
+	FRequestOwner Owner(EPriority::Blocking);
+	InTestBackend->PutValue(PutValueRequests, Owner, [&Values, &PutValueRequests](FCachePutValueResponse&& Response)
+		{
+			check(Response.Status == EStatus::Ok);
+			Values.Add(PutValueRequests[Response.UserData].Value);
+		});
+	Owner.Wait();
+
+	return Values;
+}
+
+IMPLEMENT_HTTPDERIVEDDATA_AUTOMATION_TEST(FConcurrentCachedDataProbablyExistsBatch, TEXT(".FConcurrentCachedDataProbablyExistsBatch"), EAutomationTestFlags::EditorContext | EAutomationTestFlags::ProductFilter)
+bool FConcurrentCachedDataProbablyExistsBatch::RunTest(const FString& Parameters)
+{
+	FDerivedDataBackendInterface* TestBackend = GetTestBackend();
+
+	const int32 ThreadCount = 64;
+	const double Duration = 10;
+	const uint32 KeysInBatch = 4;
+
+	TArray<FString> Keys = CreateTestCacheKeys(TestBackend, KeysInBatch);
+
+	std::atomic<uint32> MismatchedResults = 0;
+
+	ConcurrentTestWithStats(
+		[&]()
+		{
+			TConstArrayView<FString> BatchView = MakeArrayView(Keys.GetData(), KeysInBatch);
+			TBitArray<> Result = TestBackend->CachedDataProbablyExistsBatch(BatchView);
+			if (Result.CountSetBits() != BatchView.Num())
+			{
+				MismatchedResults.fetch_add(BatchView.Num() - Result.CountSetBits(), std::memory_order_relaxed);
+			}
+		},
+		ThreadCount,
+		Duration
+	);
+
+	TestEqual(TEXT("Concurrent calls to CachedDataProbablyExistsBatch for a batch of keys that were put are not reliably found"), MismatchedResults, 0);
+
+	return true;
+}
+
+// This test validate that batch requests wont mismatch head and get request for the same keys in the same batch
+IMPLEMENT_HTTPDERIVEDDATA_AUTOMATION_TEST(FConcurrentExistsAndGetForSameKeyBatch, TEXT(".FConcurrentExistsAndGetForSameKeyBatch"), EAutomationTestFlags::EditorContext | EAutomationTestFlags::ProductFilter)
+bool FConcurrentExistsAndGetForSameKeyBatch::RunTest(const FString& Parameters)
+{
+	FDerivedDataBackendInterface* TestBackend = GetTestBackend();
+
+	const int32 ParallelTasks = 32;
+	const uint32 Iterations = 20;
+	const uint32 KeysInBatch = 4;
+
+	TArray<FString> Keys = CreateTestCacheKeys(TestBackend, KeysInBatch);
+	// Add some non valid keys by just using guids
+	for (int32 Index = 0; Index < KeysInBatch; ++Index)
+	{
+		Keys.Add(FGuid::NewGuid().ToString());
+	}
+	
+	ParallelFor(ParallelTasks,
+		[&](int32 TaskIndex)
+		{
+			for (uint32 Iteration = 0; Iteration < Iterations; ++Iteration)
+			{
+				for (int32 KeyIndex = 0; KeyIndex < Keys.Num(); ++KeyIndex)
+				{
+					if ((Iteration % 2) ^ (KeyIndex % 2))
+					{
+						TestBackend->CachedDataProbablyExists(*Keys[KeyIndex]);
+					}
+					else
+					{
+						TArray<uint8> OutData;
+						TestBackend->GetCachedData(*Keys[KeyIndex], OutData);
+					}
+				}
+			}
+		}
+	);
+	return true;
+}
+
+// Tests basic functionality for structured cache operations
+IMPLEMENT_HTTPDERIVEDDATA_AUTOMATION_TEST(CacheStore, TEXT(".CacheStore"), EAutomationTestFlags::EditorContext | EAutomationTestFlags::ProductFilter)
+bool CacheStore::RunTest(const FString& Parameters)
+{
+	using namespace UE::DerivedData;
+	FDerivedDataBackendInterface* TestBackend = GetTestBackend();
+
+#if UE_WITH_ZEN
+	using namespace UE::Zen;
+	FServiceSettings ZenTestServiceSettings;
+	FServiceAutoLaunchSettings& ZenTestAutoLaunchSettings = ZenTestServiceSettings.SettingsVariant.Get<FServiceAutoLaunchSettings>();
+	ZenTestAutoLaunchSettings.DataPath = FPaths::ConvertRelativePathToFull(FPaths::Combine(FPaths::EngineSavedDir(), "ZenUnitTest"));
+	ZenTestAutoLaunchSettings.ExtraArgs = FString::Printf(TEXT("--http asio --upstream-jupiter-url \"%s\" --upstream-jupiter-oauth-url \"%s\" --upstream-jupiter-oauth-clientid \"%s\" --upstream-jupiter-oauth-clientsecret \"%s\" --upstream-jupiter-namespace-ddc \"%s\" --upstream-jupiter-namespace \"%s\""),
+		*TestDomain,
+		*TestOAuthProvider,
+		*TestOAuthClientId,
+		*TestOAuthSecret,
+		*TestNamespace,
+		*TestStructuredNamespace
+	);
+	ZenTestAutoLaunchSettings.DesiredPort = 13337; // Avoid the normal default port
+	ZenTestAutoLaunchSettings.bShowConsole = true;
+	ZenTestAutoLaunchSettings.bLimitProcessLifetime = true;
+
+	FScopeZenService ScopeZenService(MoveTemp(ZenTestServiceSettings));
+	TUniquePtr<ILegacyCacheStore> ZenIntermediaryBackend(CreateZenCacheStore(TEXT("Test"), ScopeZenService.GetInstance().GetURL(), *TestNamespace));
+	auto WaitForZenPushToUpstream = [](ILegacyCacheStore* ZenBackend, TConstArrayView<FCacheRecord> Records)
+	{
+		// TODO: Expecting a legitimate means to wait for zen to finish pushing records to its upstream in the future
+		FPlatformProcess::Sleep(1.0f);
+	};
+#endif // UE_WITH_ZEN
+
+	const uint32 RecordsInBatch = 3;
+	const uint32 ValuesInBatch = RecordsInBatch;
+
+	{
+		TArray<FCacheRecord> PutRecords = CreateTestCacheRecords(TestBackend, RecordsInBatch, 1);
+		TArray<FCacheRecord> RecievedRecords = GetAndValidateRecordsAndChunks(TEXT("SimpleValue"), PutRecords, ECachePolicy::Default);
+		TArray<FCacheRecord> RecievedRecordsSkipMeta = GetAndValidateRecordsAndChunks(TEXT("SimpleValueSkipMeta"), PutRecords, ECachePolicy::Default | ECachePolicy::SkipMeta);
+		TArray<FCacheRecord> RecievedRecordsSkipData = GetAndValidateRecordsAndChunks(TEXT("SimpleValueSkipData"), PutRecords, ECachePolicy::Default | ECachePolicy::SkipData);
+
+#if UE_WITH_ZEN
+		if (ZenIntermediaryBackend)
+		{
+			TArray<FCacheRecord> PutRecordsZen = CreateTestCacheRecords(ZenIntermediaryBackend.Get(), RecordsInBatch, 1, FCbObject(), "AutoTestDummyZen");
+			WaitForZenPushToUpstream(ZenIntermediaryBackend.Get(), PutRecordsZen);
+			ValidateRecords(TEXT("SimpleValueZenAndDirect"), GetAndValidateRecords(TEXT("SimpleValueZen"), PutRecordsZen, ECachePolicy::Default), RecievedRecords, ECachePolicy::Default);
+			ValidateRecords(TEXT("SimpleValueSkipMetaZenAndDirect"), GetAndValidateRecords(TEXT("SimpleValueSkipMetaZen"), PutRecordsZen, ECachePolicy::Default | ECachePolicy::SkipMeta), RecievedRecordsSkipMeta, ECachePolicy::Default | ECachePolicy::SkipMeta);
+			ValidateRecords(TEXT("SimpleValueSkipDataZenAndDirect"), GetAndValidateRecords(TEXT("SimpleValueSkipDataZen"), PutRecordsZen, ECachePolicy::Default | ECachePolicy::SkipData), RecievedRecordsSkipData, ECachePolicy::Default | ECachePolicy::SkipData);
+		}
+#endif // UE_WITH_ZEN
+	}
+
+	{
+		TCbWriter<64> MetaWriter;
+		MetaWriter.BeginObject();
+		MetaWriter.AddInteger("MetaKey"_ASV, 42);
+		MetaWriter.EndObject();
+		FCbObject MetaObject = MetaWriter.Save().AsObject();
+
+		TArray<FCacheRecord> PutRecords = CreateTestCacheRecords(TestBackend, RecordsInBatch, 1, MetaObject);
+		TArray<FCacheRecord> RecievedRecords = GetAndValidateRecordsAndChunks(TEXT("SimpleValueWithMeta"), PutRecords, ECachePolicy::Default);
+		TArray<FCacheRecord> RecievedRecordsSkipMeta = GetAndValidateRecordsAndChunks(TEXT("SimpleValueWithMetaSkipMeta"), PutRecords, ECachePolicy::Default | ECachePolicy::SkipMeta);
+		TArray<FCacheRecord> RecievedRecordsSkipData = GetAndValidateRecordsAndChunks(TEXT("SimpleValueWithMetaSkipData"), PutRecords, ECachePolicy::Default | ECachePolicy::SkipData);
+
+#if UE_WITH_ZEN
+		if (ZenIntermediaryBackend)
+		{
+			TArray<FCacheRecord> PutRecordsZen = CreateTestCacheRecords(ZenIntermediaryBackend.Get(), RecordsInBatch, 1, MetaObject, "AutoTestDummyZen");
+			WaitForZenPushToUpstream(ZenIntermediaryBackend.Get(), PutRecordsZen);
+			ValidateRecords(TEXT("SimpleValueWithMetaZenAndDirect"), GetAndValidateRecords(TEXT("SimpleValueWithMetaZen"), PutRecordsZen, ECachePolicy::Default), RecievedRecords, ECachePolicy::Default);
+			ValidateRecords(TEXT("SimpleValueWithMetaSkipMetaZenAndDirect"), GetAndValidateRecords(TEXT("SimpleValueWithMetaSkipMetaZen"), PutRecordsZen, ECachePolicy::Default | ECachePolicy::SkipMeta), RecievedRecordsSkipMeta, ECachePolicy::Default | ECachePolicy::SkipMeta);
+			ValidateRecords(TEXT("SimpleValueWithMetaSkipDataZenAndDirect"), GetAndValidateRecords(TEXT("SimpleValueWithMetaSkipDataZen"), PutRecordsZen, ECachePolicy::Default | ECachePolicy::SkipData), RecievedRecordsSkipData, ECachePolicy::Default | ECachePolicy::SkipData);
+		}
+#endif // UE_WITH_ZEN
+	}
+
+	{
+		TArray<FCacheRecord> PutRecords = CreateTestCacheRecords(TestBackend, RecordsInBatch, 5);
+		TArray<FCacheRecord> RecievedRecords = GetAndValidateRecordsAndChunks(TEXT("MultiValue"), PutRecords, ECachePolicy::Default);
+		TArray<FCacheRecord> RecievedRecordsSkipMeta = GetAndValidateRecordsAndChunks(TEXT("MultiValueSkipMeta"), PutRecords, ECachePolicy::Default | ECachePolicy::SkipMeta);
+		TArray<FCacheRecord> RecievedRecordsSkipData = GetAndValidateRecordsAndChunks(TEXT("MultiValueSkipData"), PutRecords, ECachePolicy::Default | ECachePolicy::SkipData);
+
+#if UE_WITH_ZEN
+		if (ZenIntermediaryBackend)
+		{
+			TArray<FCacheRecord> PutRecordsZen = CreateTestCacheRecords(ZenIntermediaryBackend.Get(), RecordsInBatch, 5, FCbObject(), "AutoTestDummyZen");
+			WaitForZenPushToUpstream(ZenIntermediaryBackend.Get(), PutRecordsZen);
+			ValidateRecords(TEXT("MultiValueZenAndDirect"), GetAndValidateRecords(TEXT("MultiValueZen"), PutRecordsZen, ECachePolicy::Default), RecievedRecords, ECachePolicy::Default);
+			ValidateRecords(TEXT("MultiValueSkipMetaZenAndDirect"), GetAndValidateRecords(TEXT("MultiValueSkipMetaZen"), PutRecordsZen, ECachePolicy::Default | ECachePolicy::SkipMeta), RecievedRecordsSkipMeta, ECachePolicy::Default | ECachePolicy::SkipMeta);
+			ValidateRecords(TEXT("MultiValueSkipDataZenAndDirect"), GetAndValidateRecords(TEXT("MultiValueSkipDataZen"), PutRecordsZen, ECachePolicy::Default | ECachePolicy::SkipData), RecievedRecordsSkipData, ECachePolicy::Default | ECachePolicy::SkipData);
+		}
+#endif // UE_WITH_ZEN
+	}
+
+	{
+		TArray<FValue> PutValues = CreateTestCacheValues(TestBackend, ValuesInBatch);
+		GetAndValidateValues(TEXT("SimpleValue"), PutValues, ECachePolicy::Default);
+		GetAndValidateValues(TEXT("SimpleValueSkipData"), PutValues, ECachePolicy::Default | ECachePolicy::SkipData);
+	}
+
+	return true;
+}
+
+} // UE::DerivedData
+
+#endif // #if WITH_DEV_AUTOMATION_TESTS && WITH_HTTP_DDC_BACKEND
