DDC: Enabled compression of legacy cache data

diff --git a/Engine/Source/Developer/DerivedDataCache/Private/ZenCacheStore.cpp b/Engine/Source/Developer/DerivedDataCache/Private/ZenCacheStore.cpp
new file mode 100644
--- /dev/null
+++ b/Engine/Source/Developer/DerivedDataCache/Private/ZenCacheStore.cpp
@@ -0,0 +1,1322 @@
+// Copyright Epic Games, Inc. All Rights Reserved.
+
+#include "DerivedDataBackendInterface.h"
+#include "ZenServerInterface.h"
+
+#if UE_WITH_ZEN
+
+#include "BatchView.h"
+#include "Containers/Set.h"
+#include "Containers/StringFwd.h"
+#include "Containers/StringFwd.h"
+#include "DerivedDataCachePrivate.h"
+#include "DerivedDataCacheRecord.h"
+#include "DerivedDataCacheUsageStats.h"
+#include "DerivedDataChunk.h"
+#include "HAL/CriticalSection.h"
+#include "Misc/ConfigCacheIni.h"
+#include "Misc/Optional.h"
+#include "Misc/ScopeLock.h"
+#include "ProfilingDebugging/CountersTrace.h"
+#include "ProfilingDebugging/CpuProfilerTrace.h"
+#include "Serialization/BufferArchive.h"
+#include "Serialization/CompactBinary.h"
+#include "Serialization/CompactBinaryPackage.h"
+#include "Serialization/CompactBinarySerialization.h"
+#include "Serialization/CompactBinaryValidation.h"
+#include "Serialization/CompactBinaryWriter.h"
+#include "ZenBackendUtils.h"
+#include "ZenSerialization.h"
+#include "ZenServerHttp.h"
+#include "ZenStatistics.h"
+
+TRACE_DECLARE_INT_COUNTER(ZenDDC_Exist,			TEXT("ZenDDC Exist"));
+TRACE_DECLARE_INT_COUNTER(ZenDDC_ExistHit,		TEXT("ZenDDC Exist Hit"));
+TRACE_DECLARE_INT_COUNTER(ZenDDC_Get,			TEXT("ZenDDC Get"));
+TRACE_DECLARE_INT_COUNTER(ZenDDC_GetHit,		TEXT("ZenDDC Get Hit"));
+TRACE_DECLARE_INT_COUNTER(ZenDDC_Put,			TEXT("ZenDDC Put"));
+TRACE_DECLARE_INT_COUNTER(ZenDDC_PutHit,		TEXT("ZenDDC Put Hit"));
+TRACE_DECLARE_INT_COUNTER(ZenDDC_BytesReceived, TEXT("ZenDDC Bytes Received"));
+TRACE_DECLARE_INT_COUNTER(ZenDDC_BytesSent,		TEXT("ZenDDC Bytes Sent"));
+TRACE_DECLARE_INT_COUNTER(ZenDDC_CacheRecordRequestCount, TEXT("ZenDDC CacheRecord Request Count"));
+TRACE_DECLARE_INT_COUNTER(ZenDDC_ChunkRequestCount, TEXT("ZenDDC Chunk Request Count"));
+
+namespace UE::DerivedData
+{
+
+template<typename T>
+void ForEachBatch(const int32 BatchSize, const int32 TotalCount, T&& Fn)
+{
+	check(BatchSize > 0);
+
+	if (TotalCount > 0)
+	{
+		const int32 BatchCount = FMath::DivideAndRoundUp(TotalCount, BatchSize);
+		const int32 Last = TotalCount - 1;
+
+		for (int32 BatchIndex = 0; BatchIndex < BatchCount; BatchIndex++)
+		{
+			const int32 BatchFirstIndex	= BatchIndex * BatchSize;
+			const int32 BatchLastIndex	= FMath::Min(BatchFirstIndex + BatchSize - 1, Last);
+
+			Fn(BatchFirstIndex, BatchLastIndex);
+		}
+	}
+}
+
+//----------------------------------------------------------------------------------------------------------
+// FZenCacheStore
+//----------------------------------------------------------------------------------------------------------
+
+/**
+ * Backend for a HTTP based caching service (Zen)
+ */
+class FZenCacheStore final : public FDerivedDataBackendInterface
+{
+public:
+	
+	/**
+	 * Creates the backend, checks health status and attempts to acquire an access token.
+	 *
+	 * @param ServiceUrl	Base url to the service including schema.
+	 * @param Namespace		Namespace to use.
+	 */
+	FZenCacheStore(const TCHAR* ServiceUrl, const TCHAR* Namespace);
+
+	~FZenCacheStore();
+
+	/**
+	 * Checks is backend is usable (reachable and accessible).
+	 * @return true if usable
+	 */
+	bool IsUsable() const { return bIsUsable; }
+
+	virtual bool IsWritable() const override;
+	virtual ESpeedClass GetSpeedClass() const override;
+	virtual TSharedRef<FDerivedDataCacheStatsNode> GatherUsageStats() const override;
+
+	/**
+	 * Synchronous attempt to make sure the cached data will be available as optimally as possible.
+	 *
+	 * @param	CacheKeys	Alphanumeric+underscore keys of the cache items
+	 * @return				true if the data will probably be found in a fast backend on a future request.
+	 */
+	virtual TBitArray<> TryToPrefetch(TConstArrayView<FString> CacheKeys) override;
+
+	virtual bool CachedDataProbablyExists(const TCHAR* CacheKey) override;
+	virtual bool GetCachedData(const TCHAR* CacheKey, TArray<uint8>& OutData) override;
+	virtual EPutStatus PutCachedData(const TCHAR* CacheKey, TArrayView<const uint8> InData, bool bPutEvenIfExists) override;
+
+	virtual void RemoveCachedData(const TCHAR* CacheKey, bool bTransient) override;
+	
+	virtual FString GetName() const override;
+	virtual bool WouldCache(const TCHAR* CacheKey, TArrayView<const uint8> InData) override;
+	virtual bool ApplyDebugOptions(FBackendDebugOptions& InOptions) override;
+
+	virtual EBackendLegacyMode GetLegacyMode() const override { return EBackendLegacyMode::LegacyOnly; }
+
+	// ICacheStore
+
+	virtual void Put(
+		TConstArrayView<FCachePutRequest> Requests,
+		IRequestOwner& Owner,
+		FOnCachePutComplete&& OnComplete = FOnCachePutComplete()) override;
+
+	virtual void Get(
+		TConstArrayView<FCacheGetRequest> Requests,
+		IRequestOwner& Owner,
+		FOnCacheGetComplete&& OnComplete) override;
+
+	virtual void PutValue(
+		TConstArrayView<FCachePutValueRequest> Requests,
+		IRequestOwner& Owner,
+		FOnCachePutValueComplete&& OnComplete = FOnCachePutValueComplete()) override;
+
+	virtual void GetValue(
+		TConstArrayView<FCacheGetValueRequest> Requests,
+		IRequestOwner& Owner,
+		FOnCacheGetValueComplete&& OnComplete = FOnCacheGetValueComplete()) override;
+
+	virtual void GetChunks(
+		TConstArrayView<FCacheGetChunkRequest> Requests,
+		IRequestOwner& Owner,
+		FOnCacheGetChunkComplete&& OnComplete) override;
+
+private:
+	enum class EGetResult
+	{
+		Success,
+		NotFound,
+		Corrupted
+	};
+	EGetResult GetZenData(FStringView Uri, TArray64<uint8>* OutData, Zen::EContentType ContentType) const;
+
+	// TODO: need ability to specify content type
+	FDerivedDataBackendInterface::EPutStatus PutZenData(const TCHAR* Uri, const FCompositeBuffer& InData, Zen::EContentType ContentType);
+	EGetResult GetZenData(const FCacheKey& Key, ECachePolicy CachePolicy, FCbPackage& OutPackage) const;
+
+	bool IsServiceReady();
+	static FString MakeLegacyZenKey(const TCHAR* CacheKey);
+	static void AppendZenUri(const FCacheKey& CacheKey, FStringBuilderBase& Out);
+	static void AppendZenUri(const FCacheKey& CacheKey, const FValueId& Id, FStringBuilderBase& Out);
+	static void AppendPolicyQueryString(ECachePolicy Policy, FStringBuilderBase& Out);
+
+	static bool ShouldRetryOnError(int64 ResponseCode);
+
+	/* Debug helpers */
+	bool ShouldSimulateMiss(const TCHAR* InKey);
+	bool ShouldSimulateMiss(const FCacheKey& InKey);
+
+private:
+	FString Namespace;
+	UE::Zen::FScopeZenService ZenService;
+	mutable FDerivedDataCacheUsageStats UsageStats;
+	TUniquePtr<UE::Zen::FZenHttpRequestPool> RequestPool;
+	bool bIsUsable = false;
+	bool bIsRemote = false;
+	uint32 FailedLoginAttempts = 0;
+	uint32 MaxAttempts = 4;
+	int32 BatchPutMaxBytes = 1024*1024;
+	int32 CacheRecordBatchSize = 8;
+	int32 CacheChunksBatchSize = 8;
+
+	/** Debug Options */
+	FBackendDebugOptions DebugOptions;
+
+	/** Keys we ignored due to miss rate settings */
+	FCriticalSection MissedKeysCS;
+	TSet<FName> DebugMissedKeys;
+	TSet<FCacheKey> DebugMissedCacheKeys;
+};
+
+FZenCacheStore::FZenCacheStore(
+	const TCHAR* InServiceUrl,
+	const TCHAR* InNamespace)
+	: Namespace(InNamespace)
+	, ZenService(InServiceUrl)
+{
+	if (IsServiceReady())
+	{
+		RequestPool = MakeUnique<Zen::FZenHttpRequestPool>(ZenService.GetInstance().GetURL(), 32);
+		bIsUsable = true;
+		bIsRemote = false;
+
+		Zen::FZenStats ZenStats;
+
+		if (ZenService.GetInstance().GetStats(ZenStats) == true)
+		{
+			 bIsRemote = ZenStats.UpstreamStats.EndPointStats.IsEmpty() == false;
+		}
+	}
+
+	GConfig->GetInt(TEXT("Zen"), TEXT("BatchPutMaxBytes"), BatchPutMaxBytes, GEngineIni);
+	GConfig->GetInt(TEXT("Zen"), TEXT("CacheRecordBatchSize"), CacheRecordBatchSize, GEngineIni);
+	GConfig->GetInt(TEXT("Zen"), TEXT("CacheChunksBatchSize"), CacheChunksBatchSize, GEngineIni);
+}
+
+FZenCacheStore::~FZenCacheStore()
+{
+}
+
+FString FZenCacheStore::GetName() const
+{
+	return ZenService.GetInstance().GetURL();
+}
+
+bool FZenCacheStore::IsServiceReady()
+{
+	return ZenService.GetInstance().IsServiceReady();
+}
+
+bool FZenCacheStore::ShouldRetryOnError(int64 ResponseCode)
+{
+	// Access token might have expired, request a new token and try again.
+	if (ResponseCode == 401)
+	{
+		return true;
+	}
+
+	// Too many requests, make a new attempt
+	if (ResponseCode == 429)
+	{
+		return true;
+	}
+
+	return false;
+}
+
+bool FZenCacheStore::CachedDataProbablyExists(const TCHAR* CacheKey)
+{
+	TRACE_CPUPROFILER_EVENT_SCOPE(ZenDDC::Exist);
+	TRACE_COUNTER_ADD(ZenDDC_Exist, int64(1));
+	COOK_STAT(auto Timer = UsageStats.TimeProbablyExists());
+
+	if (ShouldSimulateMiss(CacheKey))
+	{
+		return false;
+	}
+	
+	FString Uri = MakeLegacyZenKey(CacheKey);
+	int64 ResponseCode = 0; 
+	uint32 Attempts = 0;
+
+	// Retry request until we get an accepted response or exhaust allowed number of attempts.
+	while (ResponseCode == 0 && ++Attempts < MaxAttempts)
+	{
+		Zen::FZenScopedRequestPtr Request(RequestPool.Get());
+		Zen::FZenHttpRequest::Result Result = Request->PerformBlockingHead(*Uri, Zen::EContentType::Binary);
+		ResponseCode = Request->GetResponseCode();
+
+		if (Zen::IsSuccessCode(ResponseCode) || ResponseCode == 404)
+		{
+			const bool bIsHit = (Result == Zen::FZenHttpRequest::Result::Success && Zen::IsSuccessCode(ResponseCode));
+
+			if (bIsHit)
+			{
+				TRACE_COUNTER_ADD(ZenDDC_ExistHit, int64(1));
+				COOK_STAT(Timer.AddHit(0));
+			}
+			return bIsHit;
+		}
+
+		if (!ShouldRetryOnError(ResponseCode))
+		{
+			return false;
+		}
+
+		ResponseCode = 0;
+	}
+
+	return false;
+}
+
+bool FZenCacheStore::GetCachedData(const TCHAR* CacheKey, TArray<uint8>& OutData)
+{
+	TRACE_CPUPROFILER_EVENT_SCOPE(ZenDDC::GetCachedData);
+	TRACE_COUNTER_ADD(ZenDDC_Get, int64(1));
+	COOK_STAT(auto Timer = UsageStats.TimeGet());
+
+	if (ShouldSimulateMiss(CacheKey))
+	{
+		return false;
+	}
+
+	double StartTime = FPlatformTime::Seconds();
+
+	TArray64<uint8> ArrayBuffer;
+	EGetResult Result = GetZenData(MakeLegacyZenKey(CacheKey), &ArrayBuffer, Zen::EContentType::Binary);
+	check(ArrayBuffer.Num() <= UINT32_MAX);
+	OutData = TArray<uint8>(MoveTemp(ArrayBuffer));
+	if (Result != EGetResult::Success)
+	{
+		switch (Result)
+		{
+		default:
+			UE_LOG(LogDerivedDataCache, Verbose, TEXT("%s: Cache miss on %s"), *GetName(), CacheKey);
+			break;
+		case EGetResult::Corrupted:
+			UE_LOG(LogDerivedDataCache, Warning,
+				TEXT("Checksum from server on %s did not match recieved data. Discarding cached result."), CacheKey);
+			break;
+		}
+		return false;
+	}
+
+	TRACE_COUNTER_ADD(ZenDDC_GetHit, int64(1));
+	TRACE_COUNTER_ADD(ZenDDC_BytesReceived, int64(OutData.Num()));
+	COOK_STAT(Timer.AddHit(OutData.Num()));
+	double ReadDuration = FPlatformTime::Seconds() - StartTime;
+	double ReadSpeed = (OutData.Num() / ReadDuration) / (1024.0 * 1024.0);
+	UE_LOG(LogDerivedDataCache, Verbose, TEXT("%s: Cache hit on %s (%d bytes, %.02f secs, %.2fMB/s)"),
+		*GetName(), CacheKey, OutData.Num(), ReadDuration, ReadSpeed);
+	return true;
+}
+
+FZenCacheStore::EGetResult
+FZenCacheStore::GetZenData(FStringView Uri, TArray64<uint8>* OutData, Zen::EContentType ContentType) const
+{
+	// Retry request until we get an accepted response or exhaust allowed number of attempts.
+	EGetResult GetResult = EGetResult::NotFound;
+	for (uint32 Attempts = 0; Attempts < MaxAttempts; ++Attempts)
+	{
+		Zen::FZenScopedRequestPtr Request(RequestPool.Get());
+		if (Request.IsValid())
+		{
+			Zen::FZenHttpRequest::Result Result;
+			if (OutData)
+			{
+				Result = Request->PerformBlockingDownload(Uri, OutData, ContentType);
+			}
+			else
+			{
+				Result = Request->PerformBlockingHead(Uri, ContentType);
+			}
+			int64 ResponseCode = Request->GetResponseCode();
+
+			// Request was successful, make sure we got all the expected data.
+			if (Zen::IsSuccessCode(ResponseCode))
+			{
+				return EGetResult::Success;
+			}
+
+			if (!ShouldRetryOnError(ResponseCode))
+			{
+				break;
+			}
+		}
+	}
+
+	if (OutData)
+	{
+		OutData->Reset();
+	}
+
+	return GetResult;
+}
+
+FZenCacheStore::EGetResult
+FZenCacheStore::GetZenData(const FCacheKey& CacheKey, ECachePolicy CachePolicy, FCbPackage& OutPackage) const
+{
+	TStringBuilder<256> QueryUri;
+	AppendZenUri(CacheKey, QueryUri);
+	AppendPolicyQueryString(CachePolicy, QueryUri);
+
+	// Retry request until we get an accepted response or exhaust allowed number of attempts.
+	EGetResult GetResult = EGetResult::NotFound;
+	for (uint32 Attempts = 0; Attempts < MaxAttempts; ++Attempts)
+	{
+		Zen::FZenScopedRequestPtr Request(RequestPool.Get());
+		if (Request.IsValid())
+		{
+			Zen::FZenHttpRequest::Result Result = Request->PerformBlockingDownload(QueryUri.ToString(), OutPackage);
+			int64 ResponseCode = Request->GetResponseCode();
+			bool bPackageValid = Request->GetResponseFormatValid();
+
+			// Request was successful, make sure we got all the expected data.
+			if (Zen::IsSuccessCode(ResponseCode))
+			{
+				if (bPackageValid)
+				{
+					GetResult = EGetResult::Success;
+				}
+				else
+				{
+					GetResult = EGetResult::Corrupted;
+				}
+				break;
+			}
+
+			if (!ShouldRetryOnError(ResponseCode))
+			{
+				break;
+			}
+		}
+	}
+
+	if (GetResult != EGetResult::Success)
+	{
+		OutPackage.Reset();
+	}
+	return GetResult;
+}
+
+FDerivedDataBackendInterface::EPutStatus
+FZenCacheStore::PutCachedData(const TCHAR* CacheKey, TArrayView<const uint8> InData, bool bPutEvenIfExists)
+{
+	TRACE_CPUPROFILER_EVENT_SCOPE(ZenDDC::PutCachedData);
+
+	if (ShouldSimulateMiss(CacheKey))
+	{
+		return EPutStatus::NotCached;
+	}
+
+	FSharedBuffer DataBuffer = FSharedBuffer::MakeView(InData.GetData(), InData.Num());
+	return PutZenData(*MakeLegacyZenKey(CacheKey), FCompositeBuffer(DataBuffer), Zen::EContentType::Binary);
+}
+
+FDerivedDataBackendInterface::EPutStatus
+FZenCacheStore::PutZenData(const TCHAR* Uri, const FCompositeBuffer& InData, Zen::EContentType ContentType)
+{
+	TRACE_CPUPROFILER_EVENT_SCOPE(ZenDDC_Put);
+	COOK_STAT(auto Timer = UsageStats.TimePut());
+
+	int64 ResponseCode = 0; 
+	uint32 Attempts = 0;
+
+	// Retry request until we get an accepted response or exhaust allowed number of attempts.
+	while (ResponseCode == 0 && ++Attempts < MaxAttempts)
+	{
+		Zen::FZenScopedRequestPtr Request(RequestPool.Get());
+		if (Request.IsValid())
+		{
+			Zen::FZenHttpRequest::Result Result = Request->PerformBlockingPut(Uri, InData, ContentType);
+			ResponseCode = Request->GetResponseCode();
+
+			if (Zen::IsSuccessCode(ResponseCode))
+			{
+				TRACE_COUNTER_ADD(ZenDDC_BytesSent, int64(Request->GetBytesSent()));
+				COOK_STAT(Timer.AddHit(Request->GetBytesSent()));
+				return EPutStatus::Cached;
+			}
+
+			if (!ShouldRetryOnError(ResponseCode))
+			{
+				return EPutStatus::NotCached;
+			}
+
+			ResponseCode = 0;
+		}
+	}
+
+	return EPutStatus::NotCached;
+}
+
+FString FZenCacheStore::MakeLegacyZenKey(const TCHAR* CacheKey)
+{
+	FIoHash KeyHash = FIoHash::HashBuffer(CacheKey, FCString::Strlen(CacheKey) * sizeof(TCHAR));
+	return FString::Printf(TEXT("/z$/legacy/%s"), *LexToString(KeyHash));
+}
+
+void FZenCacheStore::AppendZenUri(const FCacheKey& CacheKey, FStringBuilderBase& Out)
+{
+	Out << TEXT("/z$/") << CacheKey.Bucket << TEXT('/') << CacheKey.Hash;
+}
+
+void FZenCacheStore::AppendZenUri(const FCacheKey& CacheKey, const FValueId& Id, FStringBuilderBase& Out)
+{
+	AppendZenUri(CacheKey, Out);
+	Out << TEXT('/') << Id;
+}
+
+void FZenCacheStore::AppendPolicyQueryString(ECachePolicy Policy, FStringBuilderBase& Uri)
+{
+	if (Policy != ECachePolicy::Default)
+	{
+		Uri << TEXT("?Policy=") << Policy;
+	}
+}
+
+void FZenCacheStore::RemoveCachedData(const TCHAR* CacheKey, bool bTransient)
+{
+	TRACE_CPUPROFILER_EVENT_SCOPE(ZenDDC_Remove);
+	FString Uri = MakeLegacyZenKey(CacheKey);
+
+	int64 ResponseCode = 0; 
+	uint32 Attempts = 0;
+
+	// Retry request until we get an accepted response or exhaust allowed number of attempts.
+	while (ResponseCode == 0 && ++Attempts < MaxAttempts)
+	{
+		Zen::FZenScopedRequestPtr Request(RequestPool.Get());
+		if (Request)
+		{
+			Zen::FZenHttpRequest::Result Result = Request->PerformBlockingDelete(*Uri);
+			ResponseCode = Request->GetResponseCode();
+
+			if (Zen::IsSuccessCode(ResponseCode))
+			{
+				return;
+			}
+
+			if (!ShouldRetryOnError(ResponseCode))
+			{
+				return;
+			}
+
+			ResponseCode = 0;
+		}
+	}
+}
+
+bool FZenCacheStore::IsWritable() const 
+{
+	return true;
+}
+
+FDerivedDataBackendInterface::ESpeedClass 
+FZenCacheStore::GetSpeedClass() const
+{
+	return ESpeedClass::Fast;
+}
+
+TSharedRef<FDerivedDataCacheStatsNode> FZenCacheStore::GatherUsageStats() const
+{
+	Zen::FZenStats ZenStats;
+
+	FDerivedDataCacheUsageStats LocalStats;
+	FDerivedDataCacheUsageStats RemoteStats;
+
+#if ENABLE_COOK_STATS
+	using EHitOrMiss = FCookStats::CallStats::EHitOrMiss;
+	using ECacheStatType = FCookStats::CallStats::EStatType;
+
+	Zen::GetDefaultServiceInstance().GetStats(ZenStats);
+
+	const int64 RemotePutSize = int64(ZenStats.UpstreamStats.TotalUploadedMB * 1024 * 1024);
+	const int64 RemoteGetSize = int64(ZenStats.UpstreamStats.TotalDownloadedMB * 1024 * 1024);
+	const int64 LocalGetSize = FMath::Max<int64>(0, UsageStats.GetStats.GetAccumulatedValueAnyThread(EHitOrMiss::Hit, ECacheStatType::Bytes) - RemoteGetSize);
+
+	LocalStats.PutStats = UsageStats.PutStats;
+	LocalStats.ExistsStats = UsageStats.ExistsStats;
+	LocalStats.PrefetchStats = UsageStats.PrefetchStats;
+
+	LocalStats.GetStats.Accumulate(EHitOrMiss::Hit, ECacheStatType::Counter, ZenStats.CacheStats.Hits - ZenStats.CacheStats.UpstreamHits, /*bIsInGameThread*/ false);
+	LocalStats.GetStats.Accumulate(EHitOrMiss::Miss, ECacheStatType::Counter, ZenStats.CacheStats.Misses + ZenStats.CacheStats.UpstreamHits, /*bIsInGameThread*/ false);
+	RemoteStats.GetStats.Accumulate(EHitOrMiss::Hit, ECacheStatType::Counter, ZenStats.CacheStats.UpstreamHits, /*bIsInGameThread*/ false);
+	RemoteStats.GetStats.Accumulate(EHitOrMiss::Miss, ECacheStatType::Counter, ZenStats.CacheStats.Misses, /*bIsInGameThread*/ false);
+
+	LocalStats.GetStats.Accumulate(EHitOrMiss::Hit, ECacheStatType::Bytes, LocalGetSize, /*bIsInGameThread*/ false);
+	RemoteStats.GetStats.Accumulate(EHitOrMiss::Hit, ECacheStatType::Bytes, RemoteGetSize, /*bIsInGameThread*/ false);
+	RemoteStats.PutStats.Accumulate(EHitOrMiss::Hit, ECacheStatType::Bytes, RemotePutSize, /*bIsInGameThread*/ false);
+#endif
+
+	TSharedRef<FDerivedDataCacheStatsNode> LocalNode =
+		MakeShared<FDerivedDataCacheStatsNode>(TEXT("Zen"), ZenService.GetInstance().GetURL(), /*bIsLocal*/ true);
+	LocalNode->Stats.Add(TEXT(""), LocalStats);
+
+	if (ZenStats.UpstreamStats.EndPointStats.IsEmpty())
+	{
+		return LocalNode;
+	}
+
+	TSharedRef<FDerivedDataCacheStatsNode> RemoteNode =
+		MakeShared<FDerivedDataCacheStatsNode>(ZenStats.UpstreamStats.EndPointStats[0].Name, ZenStats.UpstreamStats.EndPointStats[0].Url, /*bIsLocal*/ false);
+	RemoteNode->Stats.Add(TEXT(""), RemoteStats);
+
+	TSharedRef<FDerivedDataCacheStatsNode> GroupNode =
+		MakeShared<FDerivedDataCacheStatsNode>(TEXT("Zen Group"), TEXT(""), /*bIsLocal*/ true);
+	GroupNode->Children.Add(LocalNode);
+	GroupNode->Children.Add(RemoteNode);
+	return GroupNode;
+}
+
+TBitArray<> FZenCacheStore::TryToPrefetch(TConstArrayView<FString> CacheKeys)
+{
+	return CachedDataProbablyExistsBatch(CacheKeys);
+}
+
+bool FZenCacheStore::WouldCache(const TCHAR* CacheKey, TArrayView<const uint8> InData)
+{
+	return true;
+}
+
+bool FZenCacheStore::ApplyDebugOptions(FBackendDebugOptions& InOptions)
+{
+	DebugOptions = InOptions;
+	return true;
+}
+
+bool FZenCacheStore::ShouldSimulateMiss(const TCHAR* InKey)
+{
+	if (DebugOptions.RandomMissRate == 0 && DebugOptions.SimulateMissTypes.IsEmpty())
+	{
+		return false;
+	}
+
+	const FName Key(InKey);
+	const uint32 Hash = GetTypeHash(Key);
+
+	if (FScopeLock Lock(&MissedKeysCS); DebugMissedKeys.ContainsByHash(Hash, Key))
+	{
+		return true;
+	}
+
+	if (DebugOptions.ShouldSimulateMiss(InKey))
+	{
+		FScopeLock Lock(&MissedKeysCS);
+		UE_LOG(LogDerivedDataCache, Verbose, TEXT("Simulating miss in %s for %s"), *GetName(), InKey);
+		DebugMissedKeys.AddByHash(Hash, Key);
+		return true;
+	}
+
+	return false;
+}
+
+bool FZenCacheStore::ShouldSimulateMiss(const FCacheKey& Key)
+{
+	if (DebugOptions.RandomMissRate == 0 && DebugOptions.SimulateMissTypes.IsEmpty())
+	{
+		return false;
+	}
+
+	const uint32 Hash = GetTypeHash(Key);
+
+	if (FScopeLock Lock(&MissedKeysCS); DebugMissedCacheKeys.ContainsByHash(Hash, Key))
+	{
+		return true;
+	}
+
+	if (DebugOptions.ShouldSimulateMiss(Key))
+	{
+		FScopeLock Lock(&MissedKeysCS);
+		DebugMissedCacheKeys.AddByHash(Hash, Key);
+		return true;
+	}
+
+	return false;
+}
+
+void FZenCacheStore::Put(
+	const TConstArrayView<FCachePutRequest> Requests,
+	IRequestOwner& Owner,
+	FOnCachePutComplete&& OnComplete)
+{
+	TRACE_CPUPROFILER_EVENT_SCOPE(ZenDDC::PutCachedRecord);
+	COOK_STAT(auto Timer = UsageStats.TimePut());
+
+	uint64 BatchSize = 0;
+	TBatchView<const FCachePutRequest> Batches(Requests,
+		[&BatchSize, this](const FCachePutRequest& NextRequest) -> EBatchView
+		{
+			const FCacheRecord& Record = NextRequest.Record;
+			uint64 RecordSize = sizeof(FCacheKey) + Record.GetMeta().GetSize();
+			for (const FValueWithId& Value : Record.GetValues())
+			{
+				RecordSize += Value.GetData().GetCompressedSize();
+			}
+			BatchSize += RecordSize;
+			if (BatchSize > BatchPutMaxBytes)
+			{
+				BatchSize = RecordSize;
+				return EBatchView::NewBatch;
+			}
+			return EBatchView::Continue;
+		});
+
+	auto OnHit = [this, &OnComplete COOK_STAT(, &Timer)](const FCachePutRequest& Request)
+	{
+		const FCacheKey& Key = Request.Record.GetKey();
+		UE_LOG(LogDerivedDataCache, Verbose, TEXT("%s: Cache Put complete for %s from '%s'"),
+			*GetName(), *WriteToString<96>(Key), *Request.Name);
+		COOK_STAT(Timer.AddHit(Private::GetCacheRecordCompressedSize(Request.Record)));
+		OnComplete(Request.MakeResponse(EStatus::Ok));
+	};
+	auto OnMiss = [this, &OnComplete COOK_STAT(, &Timer)](const FCachePutRequest& Request)
+	{
+		const FCacheKey& Key = Request.Record.GetKey();
+		UE_LOG(LogDerivedDataCache, Verbose, TEXT("%s: Cache Put miss for '%s' from '%s'"),
+			*GetName(), *WriteToString<96>(Key), *Request.Name);
+		COOK_STAT(Timer.AddMiss());
+		OnComplete(Request.MakeResponse(EStatus::Error));
+	};
+
+	for (TArrayView<const FCachePutRequest> Batch : Batches)
+	{
+		FCbPackage BatchPackage;
+		FCbWriter BatchWriter;
+		BatchWriter.BeginObject();
+		{
+			BatchWriter << "Method"_ASV << "PutCacheRecords";
+			BatchWriter.BeginObject("Params"_ASV);
+			{
+				ECachePolicy BatchDefaultPolicy = Batch[0].Policy.GetRecordPolicy();
+				BatchWriter << "DefaultPolicy"_ASV << *WriteToString<128>(BatchDefaultPolicy);
+
+				BatchWriter.BeginArray("Requests"_ASV);
+				for (const FCachePutRequest& Request : Batch)
+				{
+					const FCacheRecord& Record = Request.Record;
+
+					BatchWriter.BeginObject();
+					{
+						BatchWriter.SetName("Record"_ASV);
+						Record.Save(BatchPackage, BatchWriter);
+						if (!Request.Policy.IsUniform() || Request.Policy.GetRecordPolicy() != BatchDefaultPolicy)
+						{
+							BatchWriter.SetName("Policy"_ASV);
+							Request.Policy.Save(BatchWriter);
+						}
+					}
+					BatchWriter.EndObject();
+				}
+				BatchWriter.EndArray();
+			}
+			BatchWriter.EndObject();
+		}
+		BatchWriter.EndObject();
+		BatchPackage.SetObject(BatchWriter.Save().AsObject());
+
+		FCbPackage BatchResponse;
+		Zen::FZenHttpRequest::Result HttpResult = Zen::FZenHttpRequest::Result::Failed;
+
+		{
+			Zen::FZenScopedRequestPtr Request(RequestPool.Get());
+			HttpResult = Request->PerformRpc(TEXT("/z$/$rpc"_SV), BatchPackage, BatchResponse);
+		}
+
+		int32 RequestIndex = 0;
+		if (HttpResult == Zen::FZenHttpRequest::Result::Success)
+		{
+			const FCbObject& ResponseObj = BatchResponse.GetObject();
+			for (FCbField ResponseField : ResponseObj["Result"_ASV])
+			{
+				if (RequestIndex >= Batch.Num())
+				{
+					++RequestIndex;
+					continue;
+				}
+				const FCachePutRequest& Request = Batch[RequestIndex++];
+
+				const FCacheKey& Key = Request.Record.GetKey();
+				bool bPutSucceeded = ResponseField.AsBool();
+				if (ShouldSimulateMiss(Key))
+				{
+					UE_LOG(LogDerivedDataCache, Verbose, TEXT("%s: Simulated miss for put of %s from '%s'"),
+						*GetName(), *WriteToString<96>(Key), *Request.Name);
+					bPutSucceeded = false;
+				}
+				bPutSucceeded ? OnHit(Request) : OnMiss(Request);
+			}
+			if (RequestIndex != Batch.Num())
+			{
+				UE_LOG(LogDerivedDataCache, Warning,
+					TEXT("Invalid response received from PutCacheRecords rpc: %d results expected, received %d."), Batch.Num(), RequestIndex);
+			}
+		}
+		for (const FCachePutRequest& Request : Batch.RightChop(RequestIndex))
+		{
+			OnMiss(Request);
+		}
+	}
+}
+
+void FZenCacheStore::Get(
+	const TConstArrayView<FCacheGetRequest> Requests,
+	IRequestOwner& Owner,
+	FOnCacheGetComplete&& OnComplete)
+{
+	TRACE_CPUPROFILER_EVENT_SCOPE(ZenDDC::GetCacheRecord);
+	TRACE_COUNTER_ADD(ZenDDC_CacheRecordRequestCount, int64(Requests.Num()));
+	TRACE_COUNTER_ADD(ZenDDC_Get, int64(Requests.Num()));
+	COOK_STAT(auto Timer = UsageStats.TimeGet());
+
+	auto OnHit = [this, &OnComplete COOK_STAT(, &Timer)](const FCacheGetRequest& Request, FCacheRecord&& Record)
+	{
+		UE_LOG(LogDerivedDataCache, Verbose, TEXT("%s: Cache hit for '%s' from '%s'"),
+			*GetName(), *WriteToString<96>(Request.Key), *Request.Name);
+
+		TRACE_COUNTER_ADD(ZenDDC_GetHit, int64(1));
+		int64 ReceivedSize = Private::GetCacheRecordCompressedSize(Record);
+		TRACE_COUNTER_ADD(ZenDDC_BytesReceived, ReceivedSize);
+		COOK_STAT(Timer.AddHit(ReceivedSize));
+
+		OnComplete({Request.Name, MoveTemp(Record), Request.UserData, EStatus::Ok});
+	};
+	auto OnMiss = [this, &OnComplete](const FCacheGetRequest& Request)
+	{
+		UE_LOG(LogDerivedDataCache, Verbose, TEXT("%s: Cache miss for '%s' from '%s'"),
+			*GetName(), *WriteToString<96>(Request.Key), *Request.Name);
+
+		OnComplete(Request.MakeResponse(EStatus::Error));
+	};
+
+	ForEachBatch(CacheRecordBatchSize, Requests.Num(),
+		[this, &Requests, &Owner, &OnComplete, &OnHit, &OnMiss](int32 BatchFirst, int32 BatchLast)
+	{
+		TArrayView<const FCacheGetRequest> Batch(Requests.GetData() + BatchFirst, BatchLast - BatchFirst + 1);
+
+		FCbWriter BatchRequest;
+		BatchRequest.BeginObject();
+		{
+			BatchRequest << "Method"_ASV << "GetCacheRecords"_ASV;
+			BatchRequest.BeginObject("Params"_ASV);
+			{
+				ECachePolicy BatchDefaultPolicy = Batch[0].Policy.GetRecordPolicy();
+				BatchRequest << "DefaultPolicy"_ASV << *WriteToString<128>(BatchDefaultPolicy);
+
+				BatchRequest.BeginArray("Requests"_ASV);
+				for (const FCacheGetRequest& Request : Batch)
+				{
+					BatchRequest.BeginObject();
+					{
+						BatchRequest << "Key"_ASV << Request.Key;
+						if (!Request.Policy.IsUniform() || Request.Policy.GetRecordPolicy() != BatchDefaultPolicy)
+						{
+							BatchRequest.SetName("Policy"_ASV);
+							Request.Policy.Save(BatchRequest);
+						}
+					}
+					BatchRequest.EndObject();
+				}
+				BatchRequest.EndArray();
+			}
+			BatchRequest.EndObject();
+		}
+		BatchRequest.EndObject();
+
+		FCbPackage BatchResponse;
+		Zen::FZenHttpRequest::Result HttpResult = Zen::FZenHttpRequest::Result::Failed;
+
+		{
+			LLM_SCOPE_BYTAG(UntaggedDDCResult);
+			Zen::FZenScopedRequestPtr Request(RequestPool.Get());
+			HttpResult = Request->PerformRpc(TEXT("/z$/$rpc"_SV), BatchRequest.Save().AsObject(), BatchResponse);
+		}
+
+		int32 RequestIndex = 0;
+		if (HttpResult == Zen::FZenHttpRequest::Result::Success)
+		{
+			const FCbObject& ResponseObj = BatchResponse.GetObject();
+			
+			for (FCbField RecordField : ResponseObj["Result"_ASV])
+			{
+				if (RequestIndex >= Batch.Num())
+				{
+					++RequestIndex;
+					continue;
+				}
+				const FCacheGetRequest& Request = Batch[RequestIndex++];
+
+				const FCacheKey& Key = Request.Key;
+				FOptionalCacheRecord Record;
+
+				if (ShouldSimulateMiss(Key))
+				{
+					UE_LOG(LogDerivedDataCache, Verbose, TEXT("%s: Simulated miss for put of '%s' from '%s'"),
+						*GetName(), *WriteToString<96>(Key), *Request.Name);
+				}
+				else if (!RecordField.IsNull())
+				{
+					Record = FCacheRecord::Load(BatchResponse, RecordField.AsObject());
+				}
+				Record ? OnHit(Request, MoveTemp(Record).Get()) : OnMiss(Request);
+			}
+			if (RequestIndex != Batch.Num())
+			{
+				UE_LOG(LogDerivedDataCache, Warning,
+					TEXT("Invalid response received from GetCacheRecords rpc: %d results expected, received %d."), Batch.Num(), RequestIndex);
+			}
+		}
+		for (const FCacheGetRequest& Request : Batch.RightChop(RequestIndex))
+		{
+			OnMiss(Request);
+		}
+	});
+	
+	TRACE_COUNTER_SUBTRACT(ZenDDC_CacheRecordRequestCount, int64(Requests.Num()));
+}
+
+void FZenCacheStore::PutValue(
+	TConstArrayView<FCachePutValueRequest> Requests,
+	IRequestOwner& Owner,
+	FOnCachePutValueComplete&& OnComplete)
+{
+	TRACE_CPUPROFILER_EVENT_SCOPE(ZenDDC::PutValue);
+	COOK_STAT(auto Timer = UsageStats.TimePut());
+
+	uint64 BatchSize = 0;
+	TBatchView<const FCachePutValueRequest> Batches(Requests,
+		[&BatchSize, this](const FCachePutValueRequest& NextValue) -> EBatchView
+		{
+			uint64 ValueSize = sizeof(FCacheKey) + NextValue.Value.GetData().GetCompressedSize();
+			BatchSize += ValueSize;
+			if (BatchSize > BatchPutMaxBytes)
+			{
+				BatchSize = ValueSize;
+				return EBatchView::NewBatch;
+			}
+			return EBatchView::Continue;
+		});
+
+	auto OnHit = [this, &OnComplete COOK_STAT(, &Timer)](const FCachePutValueRequest& Request)
+	{
+		UE_LOG(LogDerivedDataCache, Verbose, TEXT("%s: Cache PutValue complete for %s from '%s'"),
+			*GetName(), *WriteToString<96>(Request.Key), *Request.Name);
+		COOK_STAT(Timer.AddHit(Request.Value.GetData().GetCompressedSize()));
+		OnComplete(Request.MakeResponse(EStatus::Ok));
+	};
+	auto OnMiss = [this, &OnComplete COOK_STAT(, &Timer)](const FCachePutValueRequest& Request)
+	{
+		UE_LOG(LogDerivedDataCache, Verbose, TEXT("%s: Cache PutValue miss for '%s' from '%s'"),
+			*GetName(), *WriteToString<96>(Request.Key), *Request.Name);
+		COOK_STAT(Timer.AddMiss());
+		OnComplete(Request.MakeResponse(EStatus::Error));
+	};
+
+	for (TArrayView<const FCachePutValueRequest> Batch : Batches)
+	{
+		FCbPackage BatchPackage;
+		FCbWriter BatchWriter;
+		BatchWriter.BeginObject();
+		{
+			BatchWriter << "Method"_ASV << "PutCacheValues"_ASV;
+			BatchWriter.BeginObject("Params"_ASV);
+			{
+				ECachePolicy BatchDefaultPolicy = Batch[0].Policy;
+				BatchWriter << "DefaultPolicy"_ASV << *WriteToString<128>(BatchDefaultPolicy);
+
+				BatchWriter.BeginArray("Requests");
+				for (const FCachePutValueRequest& Request : Batch)
+				{
+					BatchWriter.BeginObject();
+					{
+						BatchWriter << "Key"_ASV << Request.Key;
+						const FValue& Value = Request.Value;
+						BatchWriter.AddBinaryAttachment("RawHash", Value.GetRawHash());
+						if (Value.HasData())
+						{
+							BatchPackage.AddAttachment(FCbAttachment(Value.GetData()));
+						}
+						if (Request.Policy != BatchDefaultPolicy)
+						{
+							BatchWriter << "Policy"_ASV << WriteToString<128>(Request.Policy);
+						}
+					}
+					BatchWriter.EndObject();
+				}
+				BatchWriter.EndArray();
+			}
+			BatchWriter.EndObject();
+		}
+		BatchWriter.EndObject();
+		BatchPackage.SetObject(BatchWriter.Save().AsObject());
+
+		FCbPackage BatchResponse;
+		Zen::FZenHttpRequest::Result HttpResult = Zen::FZenHttpRequest::Result::Failed;
+		{
+			Zen::FZenScopedRequestPtr Request(RequestPool.Get());
+			HttpResult = Request->PerformRpc(TEXT("/z$/$rpc"_SV), BatchPackage, BatchResponse);
+		}
+
+		int32 RequestIndex = 0;
+		if (HttpResult == Zen::FZenHttpRequest::Result::Success)
+		{
+			const FCbObject& ResponseObj = BatchResponse.GetObject();
+			for (FCbField ResponseField : ResponseObj["Result"_ASV])
+			{
+				if (RequestIndex >= Batch.Num())
+				{
+					++RequestIndex;
+					continue;
+				}
+				const FCachePutValueRequest& Request = Batch[RequestIndex++];
+
+				bool bPutSucceeded = ResponseField.AsBool();
+				if (ShouldSimulateMiss(Request.Key))
+				{
+					UE_LOG(LogDerivedDataCache, Verbose, TEXT("%s: Simulated miss for PutValue of %s from '%s'"),
+						*GetName(), *WriteToString<96>(Request.Key), *Request.Name);
+					bPutSucceeded = false;
+				}
+				bPutSucceeded ? OnHit(Request) : OnMiss(Request);
+			}
+			if (RequestIndex != Batch.Num())
+			{
+				UE_LOG(LogDerivedDataCache, Warning,
+					TEXT("Invalid response received from PutCacheValues rpc: %d results expected, received %d."), Batch.Num(), RequestIndex);
+			}
+		}
+		for (const FCachePutValueRequest& Request : Batch.RightChop(RequestIndex))
+		{
+			OnMiss(Request);
+		}
+	}
+}
+
+void FZenCacheStore::GetValue(
+	TConstArrayView<FCacheGetValueRequest> Requests,
+	IRequestOwner& Owner,
+	FOnCacheGetValueComplete&& OnComplete)
+{
+	TRACE_CPUPROFILER_EVENT_SCOPE(ZenDDC::GetValue);
+	TRACE_COUNTER_ADD(ZenDDC_CacheRecordRequestCount, int64(Requests.Num()));
+	COOK_STAT(auto Timer = UsageStats.TimeGet());
+	TRACE_COUNTER_ADD(ZenDDC_Get, (int64)Requests.Num());
+
+	auto OnHit = [this, &OnComplete COOK_STAT(, &Timer)](const FCacheGetValueRequest& Request, FValue&& Value)
+	{
+		UE_LOG(LogDerivedDataCache, Verbose, TEXT("%s: Cache hit for '%s' from '%s'"),
+			*GetName(), *WriteToString<96>(Request.Key), *Request.Name);
+		TRACE_COUNTER_ADD(ZenDDC_GetHit, int64(1));
+		int64 ReceivedSize = Value.GetData().GetCompressedSize();
+		TRACE_COUNTER_ADD(ZenDDC_BytesReceived, ReceivedSize);
+		COOK_STAT(Timer.AddHit(ReceivedSize));
+
+		OnComplete({Request.Name, Request.Key, MoveTemp(Value), Request.UserData, EStatus::Ok});
+	};
+	auto OnMiss = [this, &OnComplete](const FCacheGetValueRequest& Request)
+	{
+		UE_LOG(LogDerivedDataCache, Verbose, TEXT("%s: Cache miss for '%s' from '%s'"),
+			*GetName(), *WriteToString<96>(Request.Key), *Request.Name);
+
+		OnComplete(Request.MakeResponse(EStatus::Error));
+	};
+
+	ForEachBatch(CacheRecordBatchSize, Requests.Num(),
+		[this, &Requests, &Owner, &OnComplete, &OnHit, &OnMiss](int32 BatchFirst, int32 BatchLast)
+		{
+			TArrayView<const FCacheGetValueRequest> Batch(Requests.GetData() + BatchFirst, BatchLast - BatchFirst + 1);
+
+			FCbWriter BatchRequest;
+			BatchRequest.BeginObject();
+			{
+				BatchRequest << "Method"_ASV << "GetCacheValues"_ASV;
+				BatchRequest.BeginObject("Params"_ASV);
+				{
+					ECachePolicy BatchDefaultPolicy = Batch[0].Policy;
+					BatchRequest << "DefaultPolicy"_ASV << *WriteToString<128>(BatchDefaultPolicy);
+
+					BatchRequest.BeginArray("Requests");
+					for (const FCacheGetValueRequest& Request : Batch)
+					{
+						BatchRequest.BeginObject();
+						{
+							BatchRequest << "Key"_ASV << Request.Key;
+							if (Request.Policy != BatchDefaultPolicy)
+							{
+								BatchRequest << "Policy"_ASV << WriteToString<128>(Request.Policy);
+							}
+						}
+						BatchRequest.EndObject();
+					}
+					BatchRequest.EndArray();
+				}
+				BatchRequest.EndObject();
+			}
+			BatchRequest.EndObject();
+
+			FCbPackage BatchResponse;
+			Zen::FZenHttpRequest::Result HttpResult = Zen::FZenHttpRequest::Result::Failed;
+
+			{
+				LLM_SCOPE_BYTAG(UntaggedDDCResult);
+				Zen::FZenScopedRequestPtr Request(RequestPool.Get());
+				HttpResult = Request->PerformRpc(TEXT("/z$/$rpc"_SV), BatchRequest.Save().AsObject(), BatchResponse);
+			}
+
+			int32 RequestIndex = 0;
+			if (HttpResult == Zen::FZenHttpRequest::Result::Success)
+			{
+				const FCbObject& ResponseObj = BatchResponse.GetObject();
+
+				for (FCbFieldView ResultField : ResponseObj["Result"_ASV])
+				{
+					if (RequestIndex >= Batch.Num())
+					{
+						++RequestIndex;
+						continue;
+					}
+					const FCacheGetValueRequest& Request = Batch[RequestIndex++];
+
+					FCbObjectView ResultObj = ResultField.AsObjectView();
+					TOptional<FValue> Value;
+					if (ShouldSimulateMiss(Request.Key))
+					{
+						UE_LOG(LogDerivedDataCache, Verbose, TEXT("%s: Simulated miss for GetValue of '%s' from '%s'"),
+							*GetName(), *WriteToString<96>(Request.Key), *Request.Name);
+					}
+					else
+					{
+						FCbFieldView RawHashField = ResultObj["RawHash"];
+						FIoHash RawHash = RawHashField.AsHash();
+						if (const FCbAttachment* Attachment = BatchResponse.FindAttachment(RawHash))
+						{
+							Value.Emplace(Attachment->AsCompressedBinary());
+						}
+						else
+						{
+							FCbFieldView RawSizeField = ResultObj["RawSize"];
+							uint64 RawSize = RawSizeField.AsUInt64();
+							if (!RawSizeField.HasError() && !RawHashField.HasError())
+							{
+								Value.Emplace(RawHash, RawSize);
+							}
+						}
+					}
+					(bool)Value ? OnHit(Request, MoveTemp(*Value)) : OnMiss(Request);
+				}
+				if (RequestIndex != Batch.Num())
+				{
+					UE_LOG(LogDerivedDataCache, Warning,
+						TEXT("Invalid response received from GetCacheValues rpc: %d results expected, received %d."), Batch.Num(), RequestIndex);
+				}
+			}
+			for (const FCacheGetValueRequest& Request : Batch.RightChop(RequestIndex))
+			{
+				OnMiss(Request);
+			}
+		});
+	TRACE_COUNTER_SUBTRACT(ZenDDC_CacheRecordRequestCount, int64(Requests.Num()));
+}
+
+
+void FZenCacheStore::GetChunks(
+	TConstArrayView<FCacheGetChunkRequest> Requests,
+	IRequestOwner& Owner,
+	FOnCacheGetChunkComplete&& OnComplete)
+{
+	TRACE_CPUPROFILER_EVENT_SCOPE(ZenDDC::GetChunks);
+	TRACE_COUNTER_ADD(ZenDDC_ChunkRequestCount, int64(Requests.Num()));
+	TRACE_COUNTER_ADD(ZenDDC_Get, int64(Requests.Num()));
+	COOK_STAT(auto Timer = UsageStats.TimeGet());
+
+	TArray<FCacheGetChunkRequest, TInlineAllocator<16>> SortedRequests(Requests);
+	SortedRequests.StableSort(TChunkLess());
+	auto OnHit = [this, &OnComplete COOK_STAT(, &Timer)](const FCacheGetChunkRequest& Request, FIoHash&& RawHash, uint64 RawSize, FSharedBuffer&& RequestedBytes)
+	{
+		UE_LOG(LogDerivedDataCache, Verbose, TEXT("%s: CacheChunk hit for '%s' from '%s'"),
+			*GetName(), *WriteToString<96>(Request.Key, '/', Request.Id), *Request.Name);
+		TRACE_COUNTER_ADD(ZenDDC_GetHit, int64(1));
+		COOK_STAT(Timer.AddHit(RawSize));
+		OnComplete({Request.Name, Request.Key, Request.Id, Request.RawOffset,
+			RawSize, MoveTemp(RawHash), MoveTemp(RequestedBytes), Request.UserData, EStatus::Ok});
+	};
+
+	auto OnMiss = [this, &OnComplete COOK_STAT(, &Timer)](const FCacheGetChunkRequest& Request)
+	{
+		UE_LOG(LogDerivedDataCache, Display, TEXT("%s: CacheChunk miss with missing value '%s' for '%s' from '%s'"),
+			*GetName(), *WriteToString<16>(Request.Id), *WriteToString<96>(Request.Key), *Request.Name);
+		OnComplete(Request.MakeResponse(EStatus::Error));
+	};
+
+	ForEachBatch(CacheChunksBatchSize, SortedRequests.Num(),
+		[this, &SortedRequests, &Owner, &OnHit, &OnMiss](int32 BatchFirst, int32 BatchLast)
+	{
+		TArrayView<const FCacheGetChunkRequest> Batch(SortedRequests.GetData() + BatchFirst, BatchLast - BatchFirst + 1);
+
+		FCbWriter BatchRequest;
+		BatchRequest.BeginObject();
+		{
+			BatchRequest << "Method"_ASV << "GetCacheChunks";
+			BatchRequest.AddInteger("MethodVersion"_ASV, 1);
+			BatchRequest.BeginObject("Params"_ASV);
+			{
+				ECachePolicy DefaultPolicy = Batch[0].Policy;
+				BatchRequest << "DefaultPolicy"_ASV << WriteToString<128>(DefaultPolicy);
+				BatchRequest.BeginArray("ChunkRequests"_ASV);
+				for (const FCacheGetChunkRequest& Request : Batch)
+				{
+					BatchRequest.BeginObject();
+					
+					BatchRequest << "Key"_ASV << Request.Key;
+
+					if (Request.Id.IsValid())
+					{
+						BatchRequest.AddObjectId("ValueId"_ASV, Request.Id);
+					}
+					if (Request.RawOffset != 0)
+					{
+						BatchRequest << "RawOffset"_ASV << Request.RawOffset;
+					}
+					if (Request.RawSize != MAX_uint64)
+					{
+						BatchRequest << "RawSize"_ASV << Request.RawSize;
+					}
+					if (!Request.RawHash.IsZero())
+					{
+						BatchRequest << "ChunkId"_ASV << Request.RawHash;
+					}
+					if (Request.Policy != DefaultPolicy)
+					{
+						BatchRequest << "Policy"_ASV << WriteToString<128>(Request.Policy);
+					}
+
+					BatchRequest.EndObject();
+				}
+				BatchRequest.EndArray();
+			}
+			BatchRequest.EndObject();
+		}
+		BatchRequest.EndObject();
+
+		FCbPackage BatchResponse;
+		Zen::FZenHttpRequest::Result HttpResult = Zen::FZenHttpRequest::Result::Failed;
+
+		{
+			LLM_SCOPE_BYTAG(UntaggedDDCResult);
+			Zen::FZenScopedRequestPtr Request(RequestPool.Get());
+			HttpResult = Request->PerformRpc(TEXT("/z$/$rpc"_SV), BatchRequest.Save().AsObject(), BatchResponse);
+		}
+
+		int32 RequestIndex = 0;
+		if (HttpResult == Zen::FZenHttpRequest::Result::Success)
+		{
+			const FCbObject& ResponseObj = BatchResponse.GetObject();
+
+			for (FCbFieldView ResultView : ResponseObj["Result"_ASV])
+			{
+				if (RequestIndex >= Batch.Num())
+				{
+					++RequestIndex;
+					continue;
+				}
+				const FCacheGetChunkRequest& Request = Batch[RequestIndex++];
+
+				FIoHash RawHash;
+				bool Succeeded = false;
+				uint64 RawSize = 0;
+				FCbObjectView ResultObject = ResultView.AsObjectView();
+				FSharedBuffer RequestedBytes;
+				if (ShouldSimulateMiss(Request.Key))
+				{
+					UE_LOG(LogDerivedDataCache, Verbose, TEXT("%s: Simulated miss for get of '%s' from '%s'"),
+						*GetName(), *WriteToString<96>(Request.Key, '/', Request.Id), *Request.Name);
+				}
+				else
+				{
+					FCbFieldView HashView = ResultObject["RawHash"_ASV];
+					RawHash = HashView.AsHash();
+					if (!HashView.HasError())
+					{
+						if (const FCbAttachment* Attachment = BatchResponse.FindAttachment(HashView.AsHash()))
+						{
+							FCompressedBuffer CompressedBuffer = Attachment->AsCompressedBinary();
+							if (CompressedBuffer)
+							{
+								TRACE_COUNTER_ADD(ZenDDC_BytesReceived, CompressedBuffer.GetCompressedSize());
+								RequestedBytes = FCompressedBufferReader(CompressedBuffer).Decompress(Request.RawOffset, Request.RawSize);
+								RawSize = RequestedBytes.GetSize();
+								Succeeded = true;
+							}
+						}
+						else
+						{
+							FCbFieldView RawSizeField = ResultObject["RawSize"_ASV];
+							uint64 TotalSize = RawSizeField.AsUInt64();
+							Succeeded = !RawSizeField.HasError();
+							if (Succeeded)
+							{
+								RawSize = FMath::Min(Request.RawSize, TotalSize - FMath::Min(Request.RawOffset, TotalSize));
+							}
+						}
+					}
+				}
+				Succeeded ? OnHit(Request, MoveTemp(RawHash), RawSize, MoveTemp(RequestedBytes)) : OnMiss(Request);
+			}
+		}
+		for (const FCacheGetChunkRequest& Request : Batch.RightChop(RequestIndex))
+		{
+			OnMiss(Request);
+		}
+	});
+
+	TRACE_COUNTER_SUBTRACT(ZenDDC_ChunkRequestCount, int64(Requests.Num()));
+}
+
+ILegacyCacheStore* CreateZenCacheStore(const TCHAR* NodeName, const TCHAR* ServiceUrl, const TCHAR* Namespace)
+{
+	FZenCacheStore* Backend = new FZenCacheStore(ServiceUrl, Namespace);
+	if (Backend->IsUsable())
+	{
+		return Backend;
+	}
+	UE_LOG(LogDerivedDataCache, Warning, TEXT("%s could not contact the service (%s), will not use it."), NodeName, *Backend->GetName());
+	delete Backend;
+	return nullptr;
+}
+
+} // namespace UE::DerivedData
+
+#else
+
+namespace UE::DerivedData
+{
+
+ILegacyCacheStore* CreateZenCacheStore(const TCHAR* NodeName, const TCHAR* ServiceUrl, const TCHAR* Namespace)
+{
+	return nullptr;
+}
+
+} // UE::DerivedData
+
+#endif // UE_WITH_ZEN
\ No newline at end of file
