DDC: Fixed an error caught by StressTestUnity

#rnx

diff --git a/Engine/Source/Developer/DerivedDataCache/Private/S3CacheStore.cpp b/Engine/Source/Developer/DerivedDataCache/Private/S3CacheStore.cpp
--- a/Engine/Source/Developer/DerivedDataCache/Private/S3CacheStore.cpp
+++ b/Engine/Source/Developer/DerivedDataCache/Private/S3CacheStore.cpp
@@ -70,11 +70,6 @@ TRACE_DECLARE_INT_COUNTER(S3DDC_BytesRecieved, TEXT("S3DDC Bytes Recieved"));
 
 ///////////////////////////////////////////////////////////////////////////////////////////////////////////////////////
 
-static const TCHAR GBucketsDirectoryName[] = TEXT("Buckets");
-static const TCHAR GContentDirectoryName[] = TEXT("Content");
-
-///////////////////////////////////////////////////////////////////////////////////////////////////////////////////////
-
 void BuildPathForLegacyCache(const TCHAR* const CacheKey, FStringBuilderBase& Path);
 void BuildPathForCachePackage(const FCacheKey& CacheKey, FStringBuilderBase& Path);
 void BuildPathForCacheContent(const FIoHash& RawHash, FStringBuilderBase& Path);
