Unshelved from pending changelist '18992789' in UE5/Release-5.0:

Change Zen to  ValueWithLegacyFallback mode now that it handles Value PUTs to Jupiter correctly.  Add Jupiter auto test that ensures a PUT via UE->Zen->UpstreamZen->Jupiter produces something that can be GET from UE<-Jupiter.  Confirmed test passes.

diff --git a/Engine/Source/Developer/DerivedDataCache/Private/Tests/HttpCacheStoreTest.cpp b/Engine/Source/Developer/DerivedDataCache/Private/Tests/HttpCacheStoreTest.cpp
--- a/Engine/Source/Developer/DerivedDataCache/Private/Tests/HttpCacheStoreTest.cpp
+++ b/Engine/Source/Developer/DerivedDataCache/Private/Tests/HttpCacheStoreTest.cpp
@@ -588,7 +588,7 @@ TArray<FValue> CreateTestCacheValues(ICacheStore* InTestBackend, uint32 InNumVal
 		ValueContents.AddUninitialized(NumBytes);
 		for (int32 ContentIndex = 0; ContentIndex < NumBytes; ++ContentIndex)
 		{
-			ValueContents[ContentIndex] = (uint8)(ValueIndex + ContentIndex + 42); // offset of 42 to keep the contents distinct from record test data
+			ValueContents[ContentIndex] = (uint8)(ValueIndex + ContentIndex + 52); // offset of 52 to keep the contents distinct from record test data
 		}
 		ValueBuffers.Emplace(MakeSharedBufferFromArray(MoveTemp(ValueContents)));
 	}
@@ -698,10 +698,10 @@ bool CacheStore::RunTest(const FString& Parameters)
 
 #if UE_WITH_ZEN
 	using namespace UE::Zen;
-	FServiceSettings ZenTestServiceSettings;
-	FServiceAutoLaunchSettings& ZenTestAutoLaunchSettings = ZenTestServiceSettings.SettingsVariant.Get<FServiceAutoLaunchSettings>();
-	ZenTestAutoLaunchSettings.DataPath = FPaths::ConvertRelativePathToFull(FPaths::Combine(FPaths::EngineSavedDir(), "ZenUnitTest"));
-	ZenTestAutoLaunchSettings.ExtraArgs = FString::Printf(TEXT("--http asio --upstream-jupiter-url \"%s\" --upstream-jupiter-oauth-url \"%s\" --upstream-jupiter-oauth-clientid \"%s\" --upstream-jupiter-oauth-clientsecret \"%s\" --upstream-jupiter-namespace-ddc \"%s\" --upstream-jupiter-namespace \"%s\""),
+	FServiceSettings ZenUpstreamTestServiceSettings;
+	FServiceAutoLaunchSettings& ZenUpstreamTestAutoLaunchSettings = ZenUpstreamTestServiceSettings.SettingsVariant.Get<FServiceAutoLaunchSettings>();
+	ZenUpstreamTestAutoLaunchSettings.DataPath = FPaths::ConvertRelativePathToFull(FPaths::Combine(FPaths::EngineSavedDir(), "ZenUpstreamUnitTest"));
+	ZenUpstreamTestAutoLaunchSettings.ExtraArgs = FString::Printf(TEXT("--http asio --upstream-jupiter-url \"%s\" --upstream-jupiter-oauth-url \"%s\" --upstream-jupiter-oauth-clientid \"%s\" --upstream-jupiter-oauth-clientsecret \"%s\" --upstream-jupiter-namespace-ddc \"%s\" --upstream-jupiter-namespace \"%s\""),
 		*TestDomain,
 		*TestOAuthProvider,
 		*TestOAuthClientId,
@@ -709,6 +709,17 @@ bool CacheStore::RunTest(const FString& Parameters)
 		*TestNamespace,
 		*TestStructuredNamespace
 	);
+	ZenUpstreamTestAutoLaunchSettings.DesiredPort = 23337; // Avoid the normal default port
+	ZenUpstreamTestAutoLaunchSettings.bShowConsole = true;
+	ZenUpstreamTestAutoLaunchSettings.bLimitProcessLifetime = true;
+	FScopeZenService ScopeZenUpstreamService(MoveTemp(ZenUpstreamTestServiceSettings));
+
+	FServiceSettings ZenTestServiceSettings;
+	FServiceAutoLaunchSettings& ZenTestAutoLaunchSettings = ZenTestServiceSettings.SettingsVariant.Get<FServiceAutoLaunchSettings>();
+	ZenTestAutoLaunchSettings.DataPath = FPaths::ConvertRelativePathToFull(FPaths::Combine(FPaths::EngineSavedDir(), "ZenUnitTest"));
+	ZenTestAutoLaunchSettings.ExtraArgs = FString::Printf(TEXT("--http asio --upstream-zen-url \"http://localhost:%d\""),
+		ScopeZenUpstreamService.GetInstance().GetPort()
+	);
 	ZenTestAutoLaunchSettings.DesiredPort = 13337; // Avoid the normal default port
 	ZenTestAutoLaunchSettings.bShowConsole = true;
 	ZenTestAutoLaunchSettings.bLimitProcessLifetime = true;
@@ -720,6 +731,11 @@ bool CacheStore::RunTest(const FString& Parameters)
 		// TODO: Expecting a legitimate means to wait for zen to finish pushing records to its upstream in the future
 		FPlatformProcess::Sleep(1.0f);
 	};
+	auto WaitForZenPushValuesToUpstream = [](ILegacyCacheStore* ZenBackend, TConstArrayView<FValue> Values)
+	{
+		// TODO: Expecting a legitimate means to wait for zen to finish pushing records to its upstream in the future
+		FPlatformProcess::Sleep(1.0f);
+	};
 #endif // UE_WITH_ZEN
 
 	const uint32 RecordsInBatch = 3;
@@ -787,8 +803,18 @@ bool CacheStore::RunTest(const FString& Parameters)
 
 	{
 		TArray<FValue> PutValues = CreateTestCacheValues(TestBackend, ValuesInBatch);
-		GetAndValidateValues(TEXT("SimpleValue"), PutValues, ECachePolicy::Default);
-		GetAndValidateValues(TEXT("SimpleValueSkipData"), PutValues, ECachePolicy::Default | ECachePolicy::SkipData);
+		TArray<FValue> ReceivedValues = GetAndValidateValues(TEXT("SimpleValue"), PutValues, ECachePolicy::Default);
+		TArray<FValue> ReceivedValuesSkipData = GetAndValidateValues(TEXT("SimpleValueSkipData"), PutValues, ECachePolicy::Default | ECachePolicy::SkipData);
+
+#if UE_WITH_ZEN
+		if (ZenIntermediaryBackend)
+		{
+			TArray<FValue> PutValuesZen = CreateTestCacheValues(ZenIntermediaryBackend.Get(), ValuesInBatch);
+			WaitForZenPushValuesToUpstream(ZenIntermediaryBackend.Get(), PutValuesZen);
+			ValidateValues(TEXT("SimpleValueZenAndDirect"), GetAndValidateValues(TEXT("SimpleValueZen"), PutValuesZen, ECachePolicy::Default), ReceivedValues, ECachePolicy::Default);
+			ValidateValues(TEXT("SimpleValueSkipDataZenAndDirect"), GetAndValidateValues(TEXT("SimpleValueSkipDataZen"), PutValuesZen, ECachePolicy::Default | ECachePolicy::SkipData), ReceivedValuesSkipData, ECachePolicy::Default | ECachePolicy::SkipData);
+		}
+#endif // UE_WITH_ZEN
 	}
 
 	return true;
diff --git a/Engine/Source/Developer/DerivedDataCache/Private/ZenCacheStore.cpp b/Engine/Source/Developer/DerivedDataCache/Private/ZenCacheStore.cpp
--- a/Engine/Source/Developer/DerivedDataCache/Private/ZenCacheStore.cpp
+++ b/Engine/Source/Developer/DerivedDataCache/Private/ZenCacheStore.cpp
@@ -113,7 +113,7 @@ public:
 	virtual bool WouldCache(const TCHAR* CacheKey, TArrayView<const uint8> InData) override;
 	virtual bool ApplyDebugOptions(FBackendDebugOptions& InOptions) override;
 
-	virtual EBackendLegacyMode GetLegacyMode() const override { return EBackendLegacyMode::LegacyOnly; }
+	virtual EBackendLegacyMode GetLegacyMode() const override { return EBackendLegacyMode::ValueWithLegacyFallback; }
 
 	// ICacheStore
 
