Fix audio crash when using -DDC-All-MissTypes on streamed audio by only returning DDC miss until a put is done.
  - Infinite misses break the editor in too many ways as we often rely on the local cache to be working correctly as off-memory storage.

#rnx

diff --git a/Engine/Source/Developer/DerivedDataCache/Private/FileSystemCacheStore.cpp b/Engine/Source/Developer/DerivedDataCache/Private/FileSystemCacheStore.cpp
--- a/Engine/Source/Developer/DerivedDataCache/Private/FileSystemCacheStore.cpp
+++ b/Engine/Source/Developer/DerivedDataCache/Private/FileSystemCacheStore.cpp
@@ -878,6 +878,8 @@ private:
 	FCriticalSection MissedKeysCS;
 	TSet<FName> DebugMissedKeys;
 	TSet<FCacheKey> DebugMissedCacheKeys;
+	TSet<FName> DebugPutKeys;
+	TSet<FCacheKey> DebugPutCacheKeys;
 
 	FDerivedDataCacheUsageStats UsageStats;
 
@@ -1443,8 +1445,11 @@ FFileSystemCacheStore::EPutStatus FFileSystemCacheStore::PutCachedData(
 
 	if (ShouldSimulateMiss(CacheKey))
 	{
-		UE_LOG(LogDerivedDataCache, Verbose, TEXT("%s: Simulated miss for put of %s"), *CachePath, CacheKey);
-		return EPutStatus::Skipped;
+		const FName Key(CacheKey);
+		const uint32 Hash = GetTypeHash(Key);
+
+		FScopeLock Lock(&MissedKeysCS);
+		DebugPutKeys.AddByHash(Hash, Key);
 	}
 
 	TStringBuilder<256> LegacyPath;
@@ -1778,9 +1783,10 @@ bool FFileSystemCacheStore::PutCacheRecord(
 
 	if (ShouldSimulateMiss(Key))
 	{
-		UE_LOG(LogDerivedDataCache, Verbose, TEXT("%s: Simulated miss for put of %s from '%.*s'"),
-			*CachePath, *WriteToString<96>(Key), Name.Len(), Name.GetData());
-		return false;
+		const uint32 Hash = GetTypeHash(Key);
+
+		FScopeLock Lock(&MissedKeysCS);
+		DebugPutCacheKeys.AddByHash(Hash, Key);
 	}
 
 	TStringBuilder<256> Path;
@@ -2061,9 +2067,10 @@ bool FFileSystemCacheStore::PutCacheValue(
 
 	if (ShouldSimulateMiss(Key))
 	{
-		UE_LOG(LogDerivedDataCache, Verbose, TEXT("%s: Simulated miss for put of %s from '%.*s'"),
-			*CachePath, *WriteToString<96>(Key), Name.Len(), Name.GetData());
-		return false;
+		const uint32 Hash = GetTypeHash(Key);
+
+		FScopeLock Lock(&MissedKeysCS);
+		DebugPutCacheKeys.AddByHash(Hash, Key);
 	}
 
 	// Check if there is an existing value package.
@@ -2590,9 +2597,17 @@ bool FFileSystemCacheStore::ShouldSimulateMiss(const TCHAR* const InKey)
 	const FName Key(InKey);
 	const uint32 Hash = GetTypeHash(Key);
 
-	if (FScopeLock Lock(&MissedKeysCS); DebugMissedKeys.ContainsByHash(Hash, Key))
 	{
-		return true;
+		FScopeLock Lock(&MissedKeysCS);
+		if (DebugPutKeys.ContainsByHash(Hash, Key))
+		{
+			return false;
+		}
+
+		if (DebugMissedKeys.ContainsByHash(Hash, Key))
+		{
+			return true;
+		}
 	}
 
 	if (DebugOptions.ShouldSimulateMiss(InKey))
@@ -2614,9 +2629,17 @@ bool FFileSystemCacheStore::ShouldSimulateMiss(const FCacheKey& Key)
 
 	const uint32 Hash = GetTypeHash(Key);
 
-	if (FScopeLock Lock(&MissedKeysCS); DebugMissedCacheKeys.ContainsByHash(Hash, Key))
 	{
-		return true;
+		FScopeLock Lock(&MissedKeysCS);
+		if (DebugPutCacheKeys.ContainsByHash(Hash, Key))
+		{
+			return false;
+		}
+
+		if (DebugMissedCacheKeys.ContainsByHash(Hash, Key))
+		{
+			return true;
+		}
 	}
 
 	if (DebugOptions.ShouldSimulateMiss(Key))
