Fix for Unity name collision on MakeArrayView.

diff --git a/Engine/Source/Developer/DerivedDataCache/Private/HttpDerivedDataCacheStore.cpp b/Engine/Source/Developer/DerivedDataCache/Private/HttpDerivedDataCacheStore.cpp
--- a/Engine/Source/Developer/DerivedDataCache/Private/HttpDerivedDataCacheStore.cpp
+++ b/Engine/Source/Developer/DerivedDataCache/Private/HttpDerivedDataCacheStore.cpp
@@ -1950,7 +1950,7 @@ bool ShouldAbortForShutdown()
 	return !GIsBuildMachine && FDerivedDataBackend::Get().IsShuttingDown();
 }
 
-TConstArrayView<uint8> MakeArrayView(FSharedBuffer Buffer)
+TConstArrayView<uint8> MakeConstArrayView(FSharedBuffer Buffer)
 {
 	return TConstArrayView<uint8>(reinterpret_cast<const uint8*>(Buffer.GetData()), Buffer.GetSize());
 }
@@ -2510,7 +2510,7 @@ uint64 FHttpCacheStore::PutRef(const FCbPackage& Package, const FCacheKey& Key,
 		else
 		{
 			Request->SetHeader(TEXT("X-Jupiter-IoHash"), *WriteToString<48>(Package.GetObjectHash()));
-			Request->PerformBlockingUpload<FHttpRequest::PutCompactBinary>(*RefsUri, MakeArrayView(Package.GetObject().GetBuffer().ToShared()));
+			Request->PerformBlockingUpload<FHttpRequest::PutCompactBinary>(*RefsUri, MakeConstArrayView(Package.GetObject().GetBuffer().ToShared()));
 		}
 		ResponseCode = Request->GetResponseCode();
 
@@ -2619,17 +2619,17 @@ bool FHttpCacheStore::PutCacheRecord(
 			if (Attachment->IsCompressedBinary())
 			{
 				TempBuffer = Attachment->AsCompressedBinary().GetCompressed().ToShared();
-				BlobArrayView = MakeArrayView(TempBuffer);
+				BlobArrayView = MakeConstArrayView(TempBuffer);
 			}
 			else if (Attachment->IsBinary())
 			{
 				TempBuffer = FCompressedBuffer::Compress(Attachment->AsCompositeBinary()).GetCompressed().ToShared();
-				BlobArrayView = MakeArrayView(TempBuffer);
+				BlobArrayView = MakeConstArrayView(TempBuffer);
 			}
 			else
 			{
 				TempBuffer = FCompressedBuffer::Compress(Attachment->AsObject().GetBuffer()).GetCompressed().ToShared();
-				BlobArrayView = MakeArrayView(TempBuffer);
+				BlobArrayView = MakeConstArrayView(TempBuffer);
 			}
 
 			int64 ResponseCode = 0;
@@ -2809,7 +2809,7 @@ bool FHttpCacheStore::PutCacheValue(
 		for (uint32 Attempts = 0; (Attempts < UE_HTTPDDC_MAX_ATTEMPTS) && !ShouldAbortForShutdown() && (Attempts == 0 || ShouldRetryOnError(ResponseCode)); ++Attempts)
 		{
 			FScopedRequestPtr Request(PutRequestPools[IsInGameThread()].Get());
-			Request->PerformBlockingUpload<FHttpRequest::PutCompressedBlob>(*CompressedBlobsUri, MakeArrayView(TempBuffer));
+			Request->PerformBlockingUpload<FHttpRequest::PutCompressedBlob>(*CompressedBlobsUri, MakeConstArrayView(TempBuffer));
 
 			ResponseCode = Request->GetResponseCode();
 
