Improve error handling and messaging when Horde Storage server has errors during the multi-step PUT of structured data (records or values).

diff --git a/Engine/Source/Developer/DerivedDataCache/Private/HttpDerivedDataCacheStore.cpp b/Engine/Source/Developer/DerivedDataCache/Private/HttpDerivedDataCacheStore.cpp
--- a/Engine/Source/Developer/DerivedDataCache/Private/HttpDerivedDataCacheStore.cpp
+++ b/Engine/Source/Developer/DerivedDataCache/Private/HttpDerivedDataCacheStore.cpp
@@ -2538,10 +2538,7 @@ uint64 FHttpCacheStore::PutRef(const FCbPackage& Package, const FCacheKey& Key,
 				OutNeededBlobHashes.Empty();
 			}
 
-			if (OutNeededBlobHashes.IsEmpty())
-			{
-				bOutPutCompletedSuccessfully = true;
-			}
+			bOutPutCompletedSuccessfully = true;
 			break;
 		}
 		else
@@ -2604,7 +2601,16 @@ bool FHttpCacheStore::PutCacheRecord(
 	size_t PutRefBytesSent = PutRef(Package, Record.GetKey(), Bucket, false, NeededBlobHashes, bPutCompletedSuccessfully);
 	OutWriteSize += PutRefBytesSent;
 
+	if (!bPutCompletedSuccessfully)
+	{
+		UE_LOG(LogDerivedDataCache, Warning, TEXT("%s: Failed to put reference object for put of %s from '%.*s'"),
+			*GetName(), *WriteToString<96>(Key), Name.Len(), Name.GetData());
+		return false;
+	}
+
 	// TODO: blob uploading and finalization should be replaced with a single batch compressed blob upload endpoint in the future.
+	TStringBuilder<128> ExpectedHashes;
+	bool bExpectedHashesSerialized = false;
 
 	// Needed blob upload (if any missing)
 	for (const FIoHash& NeededBlobHash : NeededBlobHashes)
@@ -2647,6 +2653,25 @@ bool FHttpCacheStore::PutCacheRecord(
 				}
 			}
 		}
+		else
+		{
+			if (!bExpectedHashesSerialized)
+			{
+				bool bFirstHash = true;
+				for (const FCbAttachment& PackageAttachment : Package.GetAttachments())
+				{
+					if (!bFirstHash)
+					{
+						ExpectedHashes << TEXT(", ");
+					}
+					ExpectedHashes << PackageAttachment.GetHash();
+					bFirstHash = false;
+				}
+				bExpectedHashesSerialized = true;
+			}
+			UE_LOG(LogDerivedDataCache, Warning, TEXT("%s: Server reported needed hash '%s' that is outside the set of expected hashes (%s) for put of %s from '%.*s'"),
+				*GetName(), *WriteToString<96>(NeededBlobHash), ExpectedHashes.ToString(), *WriteToString<96>(Key), Name.Len(), Name.GetData());
+		}
 	}
 
 	// Finalization (if any blobs needed)
@@ -2656,7 +2681,7 @@ bool FHttpCacheStore::PutCacheRecord(
 		OutWriteSize += FinalizeRefBytesSent;
 	}
 
-	return bPutCompletedSuccessfully;
+	return bPutCompletedSuccessfully && NeededBlobHashes.IsEmpty();
 }
 
 FOptionalCacheRecord FHttpCacheStore::GetCacheRecordOnly(
@@ -2795,11 +2820,42 @@ bool FHttpCacheStore::PutCacheValue(
 	size_t PutRefBytesSent = PutRef(Package, Key, Bucket, false, NeededBlobHashes, bPutCompletedSuccessfully);
 	OutWriteSize += PutRefBytesSent;
 
+	if (!bPutCompletedSuccessfully)
+	{
+		UE_LOG(LogDerivedDataCache, Warning, TEXT("%s: Failed to put reference object for put of %s from '%.*s'"),
+			*GetName(), *WriteToString<96>(Key), Name.Len(), Name.GetData());
+		return false;
+	}
+
 	// Needed blob upload (if any missing)
 	if (!NeededBlobHashes.IsEmpty())
 	{
-		check(NeededBlobHashes.Num() == 1);
-		check(NeededBlobHashes[0] == Value.GetRawHash());
+		if (NeededBlobHashes.Num() != 1)
+		{
+			TStringBuilder<128> NeededHashString;
+			bool bFirstHash = true;
+			for (const FIoHash& NeededBlobHash : NeededBlobHashes)
+			{
+				if (!bFirstHash)
+				{
+					NeededHashString << TEXT(", ");
+				}
+				NeededHashString << NeededBlobHash;
+				bFirstHash = false;
+			}
+
+			UE_LOG(LogDerivedDataCache, Warning, TEXT("%s: Server reported unexpected needed hash quantity '%d' (%s) for put of %s from '%.*s'"),
+				*GetName(), NeededBlobHashes.Num(), NeededHashString.ToString(), *WriteToString<96>(Key), Name.Len(), Name.GetData());
+			return false;
+		}
+
+		if (NeededBlobHashes[0] != Value.GetRawHash())
+		{
+			UE_LOG(LogDerivedDataCache, Warning, TEXT("%s: Server reported needed hash '%s' that is outside the set of expected hashes (%s) for put of %s from '%.*s'"),
+				*GetName(), *WriteToString<96>(NeededBlobHashes[0]), *WriteToString<96>(Value.GetRawHash()), *WriteToString<96>(Key), Name.Len(), Name.GetData());
+			return false;
+		}
+
 		TStringBuilder<256> CompressedBlobsUri;
 		CompressedBlobsUri << "api/v1/compressed-blobs/" << StructuredNamespace << "/" << Value.GetRawHash();
 
@@ -2824,7 +2880,7 @@ bool FHttpCacheStore::PutCacheValue(
 	}
 
 
-	return true;
+	return bPutCompletedSuccessfully && NeededBlobHashes.IsEmpty();
 }
 
 bool FHttpCacheStore::GetCacheValueOnly(
