HttpCacheStoreTest: Add test to confirm that downloading from Jupiter -> Upstream Zen -> Zen -> Client works.
#rnx

diff --git a/Engine/Source/Developer/DerivedDataCache/Private/Tests/HttpCacheStoreTest.cpp b/Engine/Source/Developer/DerivedDataCache/Private/Tests/HttpCacheStoreTest.cpp
--- a/Engine/Source/Developer/DerivedDataCache/Private/Tests/HttpCacheStoreTest.cpp
+++ b/Engine/Source/Developer/DerivedDataCache/Private/Tests/HttpCacheStoreTest.cpp
@@ -7,6 +7,7 @@
 #include "DerivedDataRequestOwner.h"
 #include "DerivedDataValue.h"
 #include "DerivedDataValueId.h"
+#include "HAL/FileManager.h"
 #include "Memory/CompositeBuffer.h"
 #include "Misc/AutomationTest.h"
 #include "Misc/Paths.h"
@@ -153,7 +154,7 @@ protected:
 		FPlatformProcess::ReturnSynchEventToPool(LastEvent);
 	}
 
-	FDerivedDataBackendInterface* GetTestBackend() const
+	static FDerivedDataBackendInterface* GetTestBackend()
 	{
 		static FDerivedDataBackendInterface* CachedBackend = GetAnyHttpCacheStore(
 			TestDomain, TestOAuthProvider, TestOAuthClientId, TestOAuthSecret, TestNamespace, TestStructuredNamespace);
@@ -216,10 +217,10 @@ protected:
 		return true;
 	}
 
-	bool GetValues(TConstArrayView<FValue> Values, ECachePolicy Policy, TArray<FValue>& OutValues, const char* BucketName = nullptr)
+	bool GetValues(TConstArrayView<FValue> Values, ECachePolicy Policy, TArray<FValue>& OutValues, const char* BucketName = nullptr, ICacheStore* CacheStore = GetTestBackend())
 	{
 		using namespace UE::DerivedData;
-		ICacheStore* TestBackend = GetTestBackend();
+		ICacheStore* TestBackend = CacheStore;
 		FCacheBucket TestCacheBucket(BucketName ? BucketName : "AutoTestDummy");
 
 		TArray<FCacheGetValueRequest> Requests;
@@ -442,11 +443,11 @@ protected:
 		return ReceivedRecords;
 	}
 
-	TArray<FValue> GetAndValidateValues(const TCHAR* Name, TConstArrayView<FValue> Values, ECachePolicy Policy)
+	TArray<FValue> GetAndValidateValues(const TCHAR* Name, TConstArrayView<FValue> Values, ECachePolicy Policy, ICacheStore* CacheStore = GetTestBackend())
 	{
 		using namespace UE::DerivedData;
 		TArray<FValue> ReceivedValues;
-		bool bGetSuccessful = GetValues(Values, Policy, ReceivedValues);
+		bool bGetSuccessful = GetValues(Values, Policy, ReceivedValues, nullptr, CacheStore);
 		TestTrue(FString::Printf(TEXT("%s::Get status"), Name), bGetSuccessful);
 
 		if (!bGetSuccessful)
@@ -714,6 +715,23 @@ bool CacheStore::RunTest(const FString& Parameters)
 	ZenUpstreamTestAutoLaunchSettings.bLimitProcessLifetime = true;
 	FScopeZenService ScopeZenUpstreamService(MoveTemp(ZenUpstreamTestServiceSettings));
 
+	IFileManager::Get().DeleteDirectory(*FPaths::Combine(FPaths::EngineSavedDir(), "ZenUpstreamSiblingUnitTest"), false, true);
+	FServiceSettings ZenUpstreamSiblingTestServiceSettings;
+	FServiceAutoLaunchSettings& ZenUpstreamSiblingTestAutoLaunchSettings = ZenUpstreamSiblingTestServiceSettings.SettingsVariant.Get<FServiceAutoLaunchSettings>();
+	ZenUpstreamSiblingTestAutoLaunchSettings.DataPath = FPaths::ConvertRelativePathToFull(FPaths::Combine(FPaths::EngineSavedDir(), "ZenUpstreamSiblingUnitTest"));
+	ZenUpstreamSiblingTestAutoLaunchSettings.ExtraArgs = FString::Printf(TEXT("--http asio --upstream-jupiter-url \"%s\" --upstream-jupiter-oauth-url \"%s\" --upstream-jupiter-oauth-clientid \"%s\" --upstream-jupiter-oauth-clientsecret \"%s\" --upstream-jupiter-namespace-ddc \"%s\" --upstream-jupiter-namespace \"%s\""),
+		*TestDomain,
+		*TestOAuthProvider,
+		*TestOAuthClientId,
+		*TestOAuthSecret,
+		*TestNamespace,
+		*TestStructuredNamespace
+	);
+	ZenUpstreamSiblingTestAutoLaunchSettings.DesiredPort = 23338; // Avoid the normal default port
+	ZenUpstreamSiblingTestAutoLaunchSettings.bShowConsole = true;
+	ZenUpstreamSiblingTestAutoLaunchSettings.bLimitProcessLifetime = true;
+	FScopeZenService ScopeZenUpstreamSiblingService(MoveTemp(ZenUpstreamSiblingTestServiceSettings));
+
 	FServiceSettings ZenTestServiceSettings;
 	FServiceAutoLaunchSettings& ZenTestAutoLaunchSettings = ZenTestServiceSettings.SettingsVariant.Get<FServiceAutoLaunchSettings>();
 	ZenTestAutoLaunchSettings.DataPath = FPaths::ConvertRelativePathToFull(FPaths::Combine(FPaths::EngineSavedDir(), "ZenUnitTest"));
@@ -724,6 +742,20 @@ bool CacheStore::RunTest(const FString& Parameters)
 	ZenTestAutoLaunchSettings.bShowConsole = true;
 	ZenTestAutoLaunchSettings.bLimitProcessLifetime = true;
 
+	IFileManager::Get().DeleteDirectory(*FPaths::Combine(FPaths::EngineSavedDir(), "ZenUnitTestSibling"), false, true);
+	FServiceSettings ZenTestServiceSiblingSettings;
+	FServiceAutoLaunchSettings& ZenTestSiblingAutoLaunchSettings = ZenTestServiceSiblingSettings.SettingsVariant.Get<FServiceAutoLaunchSettings>();
+	ZenTestSiblingAutoLaunchSettings.DataPath = FPaths::ConvertRelativePathToFull(FPaths::Combine(FPaths::EngineSavedDir(), "ZenUnitTestSibling"));
+	ZenTestSiblingAutoLaunchSettings.ExtraArgs = FString::Printf(TEXT("--http asio --upstream-zen-url \"http://localhost:%d\""),
+		ScopeZenUpstreamSiblingService.GetInstance().GetPort()
+	);
+	ZenTestSiblingAutoLaunchSettings.DesiredPort = 13338; // Avoid the normal default port
+	ZenTestSiblingAutoLaunchSettings.bShowConsole = true;
+	ZenTestSiblingAutoLaunchSettings.bLimitProcessLifetime = true;
+
+	FScopeZenService ScopeZenSiblingService(MoveTemp(ZenTestServiceSiblingSettings));
+	TUniquePtr<ILegacyCacheStore> ZenIntermediarySiblingBackend(CreateZenCacheStore(TEXT("TestSibling"), ScopeZenSiblingService.GetInstance().GetURL(), *TestNamespace));
+
 	FScopeZenService ScopeZenService(MoveTemp(ZenTestServiceSettings));
 	TUniquePtr<ILegacyCacheStore> ZenIntermediaryBackend(CreateZenCacheStore(TEXT("Test"), ScopeZenService.GetInstance().GetURL(), *TestNamespace));
 	auto WaitForZenPushToUpstream = [](ILegacyCacheStore* ZenBackend, TConstArrayView<FCacheRecord> Records)
@@ -814,6 +846,11 @@ bool CacheStore::RunTest(const FString& Parameters)
 			ValidateValues(TEXT("SimpleValueZenAndDirect"), GetAndValidateValues(TEXT("SimpleValueZen"), PutValuesZen, ECachePolicy::Default), ReceivedValues, ECachePolicy::Default);
 			ValidateValues(TEXT("SimpleValueSkipDataZenAndDirect"), GetAndValidateValues(TEXT("SimpleValueSkipDataZen"), PutValuesZen, ECachePolicy::Default | ECachePolicy::SkipData), ReceivedValuesSkipData, ECachePolicy::Default | ECachePolicy::SkipData);
 		}
+		if (ZenIntermediarySiblingBackend)
+		{
+			GetAndValidateValues(TEXT("SimpleValueZen"), PutValues, ECachePolicy::Default, ZenIntermediarySiblingBackend.Get());
+			GetAndValidateValues(TEXT("SimpleValueSkipDataZen"), PutValues, ECachePolicy::Default | ECachePolicy::SkipData, ZenIntermediarySiblingBackend.Get());
+		}
 #endif // UE_WITH_ZEN
 	}
 
