DDC: Fixed MissTypes configuration broken by the new legacy key format

#rnx

diff --git a/Engine/Source/Developer/DerivedDataCache/Private/DerivedDataBackends.cpp b/Engine/Source/Developer/DerivedDataCache/Private/DerivedDataBackends.cpp
--- a/Engine/Source/Developer/DerivedDataCache/Private/DerivedDataBackends.cpp
+++ b/Engine/Source/Developer/DerivedDataCache/Private/DerivedDataBackends.cpp
@@ -1566,6 +1566,10 @@ bool FDerivedDataBackendInterface::FBackendDebugOptions::ShouldSimulateMiss(cons
 	{
 		TStringBuilder<256> Type;
 		Type << InCacheKey.Bucket;
+		if (Type.ToView().StartsWith(TEXTVIEW("Legacy")))
+		{
+			Type.RemoveAt(0, TEXTVIEW("Legacy").Len());
+		}
 		if (SimulateMissTypes.Contains(Type.ToView()))
 		{
 			return true;
