DDC: Re-added the in-flight cache for async puts

- Legacy puts are now executing asynchronously.
- Memory cache now merges partial cache records.
- Memory cache is used to store data temporarily while it is being written by an async put.
- Expanded the deprecation of persisted boot/memory caches, which no longer offer the performance benefit that they used to.
- Fixed the pak file cache to skip data with a compressed size of over 2 GiB.

#rnx

diff --git a/Engine/Source/Developer/DerivedDataCache/Private/DerivedDataBackends.cpp b/Engine/Source/Developer/DerivedDataCache/Private/DerivedDataBackends.cpp
--- a/Engine/Source/Developer/DerivedDataCache/Private/DerivedDataBackends.cpp
+++ b/Engine/Source/Developer/DerivedDataCache/Private/DerivedDataBackends.cpp
@@ -7,7 +7,6 @@
 #include "DerivedDataCacheStore.h"
 #include "DerivedDataCacheUsageStats.h"
 #include "DerivedDataRequestOwner.h"
-#include "FileBackedDerivedDataBackend.h"
 #include "HAL/FileManager.h"
 #include "HAL/IConsoleManager.h"
 #include "HAL/PlatformFileManager.h"
@@ -15,6 +14,7 @@
 #include "Internationalization/FastDecimalFormat.h"
 #include "Math/BasicMathExpressionEvaluator.h"
 #include "Math/UnitConversion.h"
+#include "MemoryCacheStore.h"
 #include "Misc/App.h"
 #include "Misc/CommandLine.h"
 #include "Misc/ConfigCacheIni.h"
@@ -44,8 +44,8 @@ static TAutoConsoleVariable<FString> GDerivedDataCacheGraphName(
 namespace UE::DerivedData
 {
 
-ILegacyCacheStore* CreateCacheStoreAsync(ILegacyCacheStore* InnerBackend, ECacheStoreFlags InnerFlags, bool bCacheInFlightPuts);
-ILegacyCacheStore* CreateCacheStoreHierarchy(ICacheStoreOwner*& OutOwner);
+ILegacyCacheStore* CreateCacheStoreAsync(ILegacyCacheStore* InnerBackend, ECacheStoreFlags InnerFlags, IMemoryCacheStore* MemoryCache);
+ILegacyCacheStore* CreateCacheStoreHierarchy(ICacheStoreOwner*& OutOwner, IMemoryCacheStore* MemoryCache);
 ILegacyCacheStore* CreateCacheStoreThrottle(ILegacyCacheStore* InnerCache, uint32 LatencyMS, uint32 MaxBytesPerSecond);
 ILegacyCacheStore* CreateCacheStoreVerify(ILegacyCacheStore* InnerCache, bool bPutOnError);
 ILegacyCacheStore* CreateFileSystemCacheStore(const TCHAR* CacheDirectory, const TCHAR* Params, const TCHAR* AccessLogFileName, ECacheStoreFlags& OutFlags);
@@ -60,7 +60,7 @@ ILegacyCacheStore* CreateHttpCacheStore(
 	const FDerivedDataBackendInterface::ESpeedClass* ForceSpeedClass,
 	EBackendLegacyMode LegacyMode,
 	bool bReadOnly);
-FFileBackedDerivedDataBackend* CreateMemoryCacheStore(const TCHAR* Name, int64 MaxCacheSize, bool bCanBeDisabled);
+IMemoryCacheStore* CreateMemoryCacheStore(const TCHAR* Name, int64 MaxCacheSize, bool bCanBeDisabled);
 IPakFileCacheStore* CreatePakFileCacheStore(const TCHAR* Filename, bool bWriting, bool bCompressed);
 ILegacyCacheStore* CreateS3CacheStore(const TCHAR* RootManifestPath, const TCHAR* BaseUrl, const TCHAR* Region, const TCHAR* CanaryObjectKey, const TCHAR* CachePath);
 ILegacyCacheStore* CreateZenCacheStore(const TCHAR* NodeName, const TCHAR* ServiceUrl, const TCHAR* Namespace);
@@ -80,6 +80,7 @@ public:
 	 */
 	FDerivedDataBackendGraph()
 		: RootCache(nullptr)
+		, MemoryCache(nullptr)
 		, BootCache(nullptr)
 		, WritePakCache(nullptr)
 		, AsyncNode(nullptr)
@@ -155,7 +156,7 @@ public:
 		// Hierarchy must exist in the graph.
 		if (!Hierarchy)
 		{
-			ILegacyCacheStore* HierarchyStore = CreateCacheStoreHierarchy(Hierarchy);
+			ILegacyCacheStore* HierarchyStore = CreateCacheStoreHierarchy(Hierarchy, GetMemoryCache());
 			if (RootNode.Key)
 			{
 				Hierarchy->Add(RootNode.Key, RootNode.Value);
@@ -167,7 +168,7 @@ public:
 		// Async must exist in the graph.
 		if (!AsyncNode)
 		{
-			AsyncNode = CreateCacheStoreAsync(RootNode.Key, RootNode.Value, /*bCacheInFlightPuts*/ true);
+			AsyncNode = CreateCacheStoreAsync(RootNode.Key, RootNode.Value, GetMemoryCache());
 			CreatedNodes.AddUnique(AsyncNode);
 			RootNode.Key = AsyncNode;
 		}
@@ -227,7 +228,7 @@ public:
 					UE_LOG(LogDerivedDataCache, Display, TEXT("Boot nodes are deprecated. Please remove the Boot node from the cache graph."));
 					if (BootCache == nullptr)
 					{
-						BootCache = ParseBootCache(*NodeName, *Entry, BootCacheFilename);
+						BootCache = ParseBootCache(*NodeName, *Entry);
 						ParsedNode = MakeTuple(BootCache, ECacheStoreFlags::Local | ECacheStoreFlags::Query | ECacheStoreFlags::Store);
 					}
 					else
@@ -436,13 +437,13 @@ public:
 		{
 			if (!Hierarchy)
 			{
-				ILegacyCacheStore* HierarchyStore = CreateCacheStoreHierarchy(Hierarchy);
+				ILegacyCacheStore* HierarchyStore = CreateCacheStoreHierarchy(Hierarchy, GetMemoryCache());
 				Hierarchy->Add(InnerNode.Key, InnerNode.Value);
 				CreatedNodes.AddUnique(HierarchyStore);
 				InnerNode.Key = HierarchyStore;
 			}
 
-			ILegacyCacheStore* AsyncStore = CreateCacheStoreAsync(InnerNode.Key, InnerNode.Value, /*bCacheInFlightPuts*/ true);
+			ILegacyCacheStore* AsyncStore = CreateCacheStoreAsync(InnerNode.Key, InnerNode.Value, GetMemoryCache());
 			InnerNode = MakeTuple(AsyncStore, InnerNode.Value);
 		}
 		else
@@ -539,7 +540,7 @@ public:
 			return MakeTuple(nullptr, ECacheStoreFlags::None);
 		}
 
-		ILegacyCacheStore* HierarchyStore = CreateCacheStoreHierarchy(Hierarchy);
+		ILegacyCacheStore* HierarchyStore = CreateCacheStoreHierarchy(Hierarchy, GetMemoryCache());
 		ECacheStoreFlags Flags = ECacheStoreFlags::None;
 		for (const FParsedNode& Node : InnerNodes)
 		{
@@ -955,9 +956,9 @@ public:
 	 * @param OutFilename filename specified for the cache
 	 * @return Boot data cache backend interface instance or nullptr if unsuccessful
 	 */
-	FFileBackedDerivedDataBackend* ParseBootCache( const TCHAR* NodeName, const TCHAR* Entry, FString& OutFilename )
+	IMemoryCacheStore* ParseBootCache(const TCHAR* NodeName, const TCHAR* Entry)
 	{
-		FFileBackedDerivedDataBackend* Cache = nullptr;
+		IMemoryCacheStore* Cache = nullptr;
 
 		// Only allow boot cache with the editor. We don't want other tools and utilities (eg. SCW) writing to the same file.
 #if WITH_EDITOR
@@ -965,44 +966,15 @@ public:
 		int64 MaxCacheSize = -1; // in MB
 		const int64 MaxSupportedCacheSize = 2048; // 2GB
 
-		FParse::Value( Entry, TEXT("MaxCacheSize="), MaxCacheSize);
-		FParse::Value( Entry, TEXT("Filename="), Filename );
-		if ( !Filename.Len() )
-		{
-			UE_LOG( LogDerivedDataCache, Warning, TEXT("FDerivedDataBackendGraph:  %s filename not found in *engine.ini, will not use %s cache."), NodeName, NodeName );
-		}
-		else
-		{
-			// make sure MaxCacheSize does not exceed 2GB
-			MaxCacheSize = FMath::Min(MaxCacheSize, MaxSupportedCacheSize);
-
-			UE_LOG( LogDerivedDataCache, Display, TEXT("Max Cache Size: %d MB"), MaxCacheSize);
-			Cache = CreateMemoryCacheStore(TEXT("Boot"), MaxCacheSize * 1024 * 1024, /*bCanBeDisabled*/ true);
+		FParse::Value(Entry, TEXT("MaxCacheSize="), MaxCacheSize);
 
-			if( Cache && Filename.Len() )
-			{
-				OutFilename = Filename;
-
-				if (MaxCacheSize > 0 && IFileManager::Get().FileSize(*Filename) >= (MaxCacheSize * 1024 * 1024))
-				{
-					UE_LOG( LogDerivedDataCache, Warning, TEXT("FDerivedDataBackendGraph:  %s filename exceeds max size."), NodeName );
-				}
+		// make sure MaxCacheSize does not exceed 2GB
+		MaxCacheSize = FMath::Min(MaxCacheSize, MaxSupportedCacheSize);
 
-				if (IFileManager::Get().FileSize(*Filename) < 0)
-				{
-					UE_LOG( LogDerivedDataCache, Display, TEXT("Starting with empty %s cache"), NodeName );
-				}
-				else if( Cache->LoadCache( *Filename ) )
-				{
-					UE_LOG( LogDerivedDataCache, Display, TEXT("Loaded %s cache: %s"), NodeName, *Filename );
-				}
-				else
-				{
-					UE_LOG( LogDerivedDataCache, Warning, TEXT("Could not load %s cache: %s"), NodeName, *Filename );
-				}
-			}
-		}
+		UE_LOG(LogDerivedDataCache, Display, TEXT("%s: Max Cache Size: %d MB"), NodeName, MaxCacheSize);
+		Cache = CreateMemoryCacheStore(TEXT("Boot"), MaxCacheSize * 1024 * 1024, /*bCanBeDisabled*/ true);
 #endif
+
 		return Cache;
 	}
 
@@ -1013,28 +985,28 @@ public:
 	 * @param Entry Node definition.
 	 * @return Memory data cache backend interface instance or nullptr if unsuccessful
 	 */
-	FFileBackedDerivedDataBackend* ParseMemoryCache( const TCHAR* NodeName, const TCHAR* Entry )
+	IMemoryCacheStore* ParseMemoryCache( const TCHAR* NodeName, const TCHAR* Entry )
 	{
-		FFileBackedDerivedDataBackend* Cache = nullptr;
 		FString Filename;
-
-		FParse::Value( Entry, TEXT("Filename="), Filename );
-		Cache = CreateMemoryCacheStore(NodeName, /*MaxCacheSize*/ -1, /*bCanBeDisabled*/ false);
-		if( Cache && Filename.Len() )
+		FParse::Value(Entry, TEXT("Filename="), Filename);
+		IMemoryCacheStore* Cache = CreateMemoryCacheStore(NodeName, /*MaxCacheSize*/ -1, /*bCanBeDisabled*/ false);
+		if (Cache && Filename.Len())
 		{
-			if( Cache->LoadCache( *Filename ) )
-			{
-				UE_LOG( LogDerivedDataCache, Display, TEXT("Loaded %s cache: %s"), NodeName, *Filename );
-			}
-			else
-			{
-				UE_LOG( LogDerivedDataCache, Warning, TEXT("Could not load %s cache: %s"), NodeName, *Filename );
-			}
+			UE_LOG(LogDerivedDataCache, Display, TEXT("Memory nodes that load from a file are deprecated. Please remove the filename from the cache configuration."));
 		}
-
 		return Cache;
 	}
 
+	IMemoryCacheStore* GetMemoryCache()
+	{
+		if (!MemoryCache)
+		{
+			MemoryCache = CreateMemoryCacheStore(TEXT("Memory"), 0, /*bCanBeDisabled*/ false);
+			CreatedNodes.Add(MemoryCache);
+		}
+		return MemoryCache;
+	}
+
 	virtual ~FDerivedDataBackendGraph()
 	{
 		RootCache = nullptr;
@@ -1057,10 +1029,6 @@ public:
 		check(RootCache);
 		if (BootCache)
 		{
-			if( !FParse::Param( FCommandLine::Get(), TEXT("DDCNOSAVEBOOT") ) && !FParse::Param( FCommandLine::Get(), TEXT("Multiprocess") ) )
-			{
-				BootCache->SaveCache(*BootCacheFilename);
-			}
 			BootCache->Disable();
 		}
 	}
@@ -1271,7 +1239,6 @@ private:
 
 	FThreadSafeCounter								AsyncCompletionCounter;
 	FString											GraphName;
-	FString											BootCacheFilename;
 	FString											ReadPakFilename;
 	FString											WritePakFilename;
 
@@ -1282,7 +1249,8 @@ private:
 	TArray< ILegacyCacheStore* > CreatedNodes;
 
 	/** Instances of backend interfaces which exist in only one copy */
-	FFileBackedDerivedDataBackend*	BootCache;
+	IMemoryCacheStore* MemoryCache;
+	IMemoryCacheStore*	BootCache;
 	IPakFileCacheStore* WritePakCache;
 	ILegacyCacheStore*	AsyncNode;
 	ICacheStoreOwner* Hierarchy;
diff --git a/Engine/Source/Developer/DerivedDataCache/Private/DerivedDataCache.cpp b/Engine/Source/Developer/DerivedDataCache/Private/DerivedDataCache.cpp
--- a/Engine/Source/Developer/DerivedDataCache/Private/DerivedDataCache.cpp
+++ b/Engine/Source/Developer/DerivedDataCache/Private/DerivedDataCache.cpp
@@ -361,9 +361,9 @@ class FDerivedDataCache final
 						LegacyRequest.Name = DebugContext;
 						LegacyRequest.Key = FLegacyCacheKey(CacheKey, FDerivedDataBackend::Get().GetMaxKeyLength());
 						LegacyRequest.Value = FLegacyCacheValue(FCompositeBuffer(FSharedBuffer::Clone(MakeMemoryView(Data))));
-						FRequestOwner BlockingOwner(EPriority::Blocking);
-						FDerivedDataBackend::Get().GetRoot().LegacyPut({LegacyRequest}, BlockingOwner, [](auto&&){});
-						BlockingOwner.Wait();
+						FRequestOwner AsyncOwner(EPriority::Normal);
+						FDerivedDataBackend::Get().GetRoot().LegacyPut({LegacyRequest}, AsyncOwner, [](auto&&){});
+						AsyncOwner.KeepAlive();
 					}
 					INC_FLOAT_STAT_BY(STAT_DDC_PutTime, bSynchronousForStats ? (float)ThisTime : 0.0f);
 				}
@@ -602,9 +602,9 @@ public:
 			LegacyRequest.Name = DebugContext;
 			LegacyRequest.Key = FLegacyCacheKey(CacheKey, FDerivedDataBackend::Get().GetMaxKeyLength());
 			LegacyRequest.Value = FLegacyCacheValue(FCompositeBuffer(FSharedBuffer::Clone(MakeMemoryView(Data))));
-			FRequestOwner BlockingOwner(EPriority::Blocking);
-			FDerivedDataBackend::Get().GetRoot().LegacyPut({LegacyRequest}, BlockingOwner, [](auto&&){});
-			BlockingOwner.Wait();
+			FRequestOwner AsyncOwner(EPriority::Normal);
+			FDerivedDataBackend::Get().GetRoot().LegacyPut({LegacyRequest}, AsyncOwner, [](auto&&){});
+			AsyncOwner.KeepAlive();
 		}
 		INC_FLOAT_STAT_BY(STAT_DDC_PutTime,(float)ThisTime);
 		INC_DWORD_STAT(STAT_DDC_NumPuts);
diff --git a/Engine/Source/Developer/DerivedDataCache/Private/DerivedDataCacheStoreAsync.cpp b/Engine/Source/Developer/DerivedDataCache/Private/DerivedDataCacheStoreAsync.cpp
--- a/Engine/Source/Developer/DerivedDataCache/Private/DerivedDataCacheStoreAsync.cpp
+++ b/Engine/Source/Developer/DerivedDataCache/Private/DerivedDataCacheStoreAsync.cpp
@@ -1,6 +1,7 @@
 // Copyright Epic Games, Inc. All Rights Reserved.
 
 #include "Async/AsyncWork.h"
+#include "DerivedDataBackendInterface.h"
 #include "DerivedDataCachePrivate.h"
 #include "DerivedDataCacheStore.h"
 #include "DerivedDataCacheUsageStats.h"
@@ -9,8 +10,8 @@
 #include "DerivedDataRequestOwner.h"
 #include "DerivedDataValueId.h"
 #include "Experimental/Async/LazyEvent.h"
-#include "FileBackedDerivedDataBackend.h"
 #include "Memory/SharedBuffer.h"
+#include "MemoryCacheStore.h"
 #include "ProfilingDebugging/CookStats.h"
 #include "Stats/Stats.h"
 #include "Tasks/Task.h"
@@ -19,8 +20,6 @@
 namespace UE::DerivedData
 {
 
-FFileBackedDerivedDataBackend* CreateMemoryCacheStore(const TCHAR* Name, int64 MaxCacheSize, bool bCanBeDisabled);
-
 /**
  * A cache store that executes non-blocking requests in a dedicated thread pool.
  *
@@ -29,14 +28,27 @@ FFileBackedDerivedDataBackend* CreateMemoryCacheStore(const TCHAR* Name, int64 M
 class FCacheStoreAsync : public ILegacyCacheStore
 {
 public:
-	FCacheStoreAsync(ILegacyCacheStore* InInnerCache, ECacheStoreFlags InInnerFlags, bool bCacheInFlightPuts);
+	FCacheStoreAsync(ILegacyCacheStore* InnerCache, ECacheStoreFlags InnerFlags, IMemoryCacheStore* MemoryCache);
 
 	virtual void Put(
 		TConstArrayView<FCachePutRequest> Requests,
 		IRequestOwner& Owner,
 		FOnCachePutComplete&& OnComplete) override
 	{
-		Execute(COOK_STAT(&FDerivedDataCacheUsageStats::TimePut,) Requests, Owner, MoveTemp(OnComplete), &ICacheStore::Put);
+		if (MemoryCache)
+		{
+			MemoryCache->Put(Requests, Owner, [](auto&&){});
+			Execute(COOK_STAT(&FDerivedDataCacheUsageStats::TimePut,) Requests, Owner,
+				[this, OnComplete = MoveTemp(OnComplete)](FCachePutResponse&& Response)
+				{
+					MemoryCache->Delete(Response.Key);
+					OnComplete(MoveTemp(Response));
+				}, &ICacheStore::Put);
+		}
+		else
+		{
+			Execute(COOK_STAT(&FDerivedDataCacheUsageStats::TimePut,) Requests, Owner, MoveTemp(OnComplete), &ICacheStore::Put);
+		}
 	}
 
 	virtual void Get(
@@ -52,7 +64,20 @@ public:
 		IRequestOwner& Owner,
 		FOnCachePutValueComplete&& OnComplete) override
 	{
-		Execute(COOK_STAT(&FDerivedDataCacheUsageStats::TimePut,) Requests, Owner, MoveTemp(OnComplete), &ICacheStore::PutValue);
+		if (MemoryCache)
+		{
+			MemoryCache->PutValue(Requests, Owner, [](auto&&){});
+			Execute(COOK_STAT(&FDerivedDataCacheUsageStats::TimePut,) Requests, Owner,
+				[this, OnComplete = MoveTemp(OnComplete)](FCachePutValueResponse&& Response)
+				{
+					MemoryCache->DeleteValue(Response.Key);
+					OnComplete(MoveTemp(Response));
+				}, &ICacheStore::PutValue);
+		}
+		else
+		{
+			Execute(COOK_STAT(&FDerivedDataCacheUsageStats::TimePut,) Requests, Owner, MoveTemp(OnComplete), &ICacheStore::PutValue);
+		}
 	}
 
 	virtual void GetValue(
@@ -76,7 +101,20 @@ public:
 		IRequestOwner& Owner,
 		FOnLegacyCachePutComplete&& OnComplete) override
 	{
-		Execute(COOK_STAT(&FDerivedDataCacheUsageStats::TimePut,) Requests, Owner, MoveTemp(OnComplete), &ILegacyCacheStore::LegacyPut);
+		if (MemoryCache)
+		{
+			MemoryCache->LegacyPut(Requests, Owner, [](auto&&){});
+			Execute(COOK_STAT(&FDerivedDataCacheUsageStats::TimePut,) Requests, Owner,
+				[this, OnComplete = MoveTemp(OnComplete)](FLegacyCachePutResponse&& Response)
+				{
+					MemoryCache->LegacyDelete(Response.Key);
+					OnComplete(MoveTemp(Response));
+				}, &ILegacyCacheStore::LegacyPut);
+		}
+		else
+		{
+			Execute(COOK_STAT(&FDerivedDataCacheUsageStats::TimePut,) Requests, Owner, MoveTemp(OnComplete), &ILegacyCacheStore::LegacyPut);
+		}
 	}
 
 	virtual void LegacyGet(
@@ -114,22 +152,22 @@ private:
 		OnExecuteType&& OnExecute);
 
 	ILegacyCacheStore* InnerCache;
-	ECacheStoreFlags InnerFlags;
-	TUniquePtr<ILegacyCacheStore> MemoryCache;
+	IMemoryCacheStore* MemoryCache;
 	FDerivedDataCacheUsageStats UsageStats;
+	ECacheStoreFlags InnerFlags;
 };
 
-FCacheStoreAsync::FCacheStoreAsync(ILegacyCacheStore* InInnerCache, ECacheStoreFlags InInnerFlags, bool bCacheInFlightPuts)
+FCacheStoreAsync::FCacheStoreAsync(ILegacyCacheStore* InInnerCache, ECacheStoreFlags InInnerFlags, IMemoryCacheStore* InMemoryCache)
 	: InnerCache(InInnerCache)
+	, MemoryCache(InMemoryCache)
 	, InnerFlags(InInnerFlags)
-	, MemoryCache(bCacheInFlightPuts ? CreateMemoryCacheStore(TEXT("InflightMemoryCache"), /*MaxCacheSize*/ -1, /*bCanBeDisabled*/ false) : nullptr)
 {
 	check(InnerCache);
 }
 
 void FCacheStoreAsync::LegacyStats(FDerivedDataCacheStatsNode& OutNode)
 {
-	OutNode = {TEXT("AsyncWrapper"), TEXT(""), EnumHasAnyFlags(InnerFlags, ECacheStoreFlags::Local)};
+	OutNode = {TEXT("Async"), TEXT(""), EnumHasAnyFlags(InnerFlags, ECacheStoreFlags::Local)};
 	OutNode.Stats.Add(TEXT(""), UsageStats);
 
 	InnerCache->LegacyStats(OutNode.Children.Add_GetRef(MakeShared<FDerivedDataCacheStatsNode>()).Get());
@@ -302,9 +340,9 @@ void FCacheStoreAsync::Execute(
 		});
 }
 
-ILegacyCacheStore* CreateCacheStoreAsync(ILegacyCacheStore* InnerCache, ECacheStoreFlags InnerFlags, bool bCacheInFlightPuts)
+ILegacyCacheStore* CreateCacheStoreAsync(ILegacyCacheStore* InnerCache, ECacheStoreFlags InnerFlags, IMemoryCacheStore* MemoryCache)
 {
-	return new FCacheStoreAsync(InnerCache, InnerFlags, bCacheInFlightPuts);
+	return new FCacheStoreAsync(InnerCache, InnerFlags, MemoryCache);
 }
 
 } // UE::DerivedData
diff --git a/Engine/Source/Developer/DerivedDataCache/Private/DerivedDataCacheStoreHierarchy.cpp b/Engine/Source/Developer/DerivedDataCache/Private/DerivedDataCacheStoreHierarchy.cpp
--- a/Engine/Source/Developer/DerivedDataCache/Private/DerivedDataCacheStoreHierarchy.cpp
+++ b/Engine/Source/Developer/DerivedDataCache/Private/DerivedDataCacheStoreHierarchy.cpp
@@ -11,6 +11,7 @@
 #include "DerivedDataRequest.h"
 #include "DerivedDataRequestOwner.h"
 #include "HAL/CriticalSection.h"
+#include "MemoryCacheStore.h"
 #include "Misc/EnumClassFlags.h"
 #include "Misc/ScopeRWLock.h"
 #include "Serialization/CompactBinary.h"
@@ -22,13 +23,14 @@
 namespace UE::DerivedData
 {
 
-ILegacyCacheStore* CreateCacheStoreAsync(ILegacyCacheStore* InnerCache, ECacheStoreFlags InnerFlags, bool bCacheInFlightPuts);
+ILegacyCacheStore* CreateCacheStoreAsync(ILegacyCacheStore* InnerCache, ECacheStoreFlags InnerFlags, IMemoryCacheStore* MemoryCache);
 
 ///////////////////////////////////////////////////////////////////////////////////////////////////////////////////////
 
 class FCacheStoreHierarchy final : public ILegacyCacheStore, public ICacheStoreOwner
 {
 public:
+	explicit FCacheStoreHierarchy(IMemoryCacheStore* MemoryCache);
 	~FCacheStoreHierarchy() final = default;
 
 	void Add(ILegacyCacheStore* CacheStore, ECacheStoreFlags Flags) final;
@@ -126,6 +128,7 @@ private:
 	mutable FRWLock NodesLock;
 	ECacheStoreNodeFlags CombinedNodeFlags{};
 	TArray<FCacheStoreNode, TInlineAllocator<8>> Nodes;
+	IMemoryCacheStore* MemoryCache;
 };
 
 ///////////////////////////////////////////////////////////////////////////////////////////////////////////////////////
@@ -213,12 +216,21 @@ ENUM_CLASS_FLAGS(FCacheStoreHierarchy::ECacheStoreNodeFlags);
 
 ///////////////////////////////////////////////////////////////////////////////////////////////////////////////////////
 
+FCacheStoreHierarchy::FCacheStoreHierarchy(IMemoryCacheStore* InMemoryCache)
+	: MemoryCache(InMemoryCache)
+{
+	if (MemoryCache)
+	{
+		Add(MemoryCache, ECacheStoreFlags::Local | ECacheStoreFlags::Query);
+	}
+}
+
 void FCacheStoreHierarchy::Add(ILegacyCacheStore* CacheStore, ECacheStoreFlags Flags)
 {
 	FWriteScopeLock Lock(NodesLock);
 	checkf(!Algo::FindBy(Nodes, CacheStore, &FCacheStoreNode::Cache),
 		TEXT("Attempting to add a cache store that was previously registered to the hierarchy."));
-	TUniquePtr<ILegacyCacheStore> AsyncCacheStore(CreateCacheStoreAsync(CacheStore, Flags, /*bCacheInFlightPuts*/ false));
+	TUniquePtr<ILegacyCacheStore> AsyncCacheStore(CreateCacheStoreAsync(CacheStore, Flags, MemoryCache));
 	Nodes.Add({CacheStore, Flags, {}, MoveTemp(AsyncCacheStore)});
 	UpdateNodeFlags();
 }
@@ -1260,9 +1272,9 @@ bool FCacheStoreHierarchy::LegacyDebugOptions(FBackendDebugOptions& Options)
 
 ///////////////////////////////////////////////////////////////////////////////////////////////////////////////////////
 
-ILegacyCacheStore* CreateCacheStoreHierarchy(ICacheStoreOwner*& OutOwner)
+ILegacyCacheStore* CreateCacheStoreHierarchy(ICacheStoreOwner*& OutOwner, IMemoryCacheStore* MemoryCache)
 {
-	FCacheStoreHierarchy* Hierarchy = new FCacheStoreHierarchy;
+	FCacheStoreHierarchy* Hierarchy = new FCacheStoreHierarchy(MemoryCache);
 	OutOwner = Hierarchy;
 	return Hierarchy;
 }
diff --git a/Engine/Source/Developer/DerivedDataCache/Private/FileBackedDerivedDataBackend.h b/Engine/Source/Developer/DerivedDataCache/Private/FileBackedDerivedDataBackend.h
deleted file mode 100644
diff --git a/Engine/Source/Developer/DerivedDataCache/Private/MemoryCacheStore.cpp b/Engine/Source/Developer/DerivedDataCache/Private/MemoryCacheStore.cpp
--- a/Engine/Source/Developer/DerivedDataCache/Private/MemoryCacheStore.cpp
+++ b/Engine/Source/Developer/DerivedDataCache/Private/MemoryCacheStore.cpp
@@ -1,11 +1,16 @@
 // Copyright Epic Games, Inc. All Rights Reserved.
 
+#include "MemoryCacheStore.h"
+
+#include "Algo/Accumulate.h"
 #include "Algo/AllOf.h"
+#include "Algo/BinarySearch.h"
+#include "Algo/NoneOf.h"
+#include "DerivedDataBackendInterface.h"
 #include "DerivedDataCachePrivate.h"
 #include "DerivedDataCacheRecord.h"
 #include "DerivedDataCacheUsageStats.h"
 #include "DerivedDataValue.h"
-#include "FileBackedDerivedDataBackend.h"
 #include "HAL/FileManager.h"
 #include "Misc/ScopeExit.h"
 #include "Misc/ScopeLock.h"
@@ -20,83 +25,13 @@ namespace UE::DerivedData
 /**
  * A simple thread safe, memory based backend. This is used for Async puts and the boot cache.
  */
-class FMemoryCacheStore final : public FFileBackedDerivedDataBackend
+class FMemoryCacheStore final : public IMemoryCacheStore
 {
 public:
 	explicit FMemoryCacheStore(const TCHAR* InName, int64 InMaxCacheSize = -1, bool bCanBeDisabled = false);
 	~FMemoryCacheStore() final;
 
-	/** Return a name for this interface */
-	FString GetName() const final { return Name; }
-
-	/** return true if this cache is writable **/
-	bool IsWritable() const final;
-
-	/** Returns a class of speed for this interface **/
-	ESpeedClass GetSpeedClass()  const final;
-
-	/**
-	 * Synchronous test for the existence of a cache item
-	 *
-	 * @param	CacheKey	Alphanumeric+underscore key of this cache item
-	 * @return				true if the data probably will be found, this can't be guaranteed because of concurrency in the backends, corruption, etc
-	 */
-	bool CachedDataProbablyExists(const TCHAR* CacheKey) final;
-
-	/**
-	 * Synchronous retrieve of a cache item
-	 *
-	 * @param	CacheKey	Alphanumeric+underscore key of this cache item
-	 * @param	OutData		Buffer to receive the results, if any were found
-	 * @return				true if any data was found, and in this case OutData is non-empty
-	 */
-	bool GetCachedData(const TCHAR* CacheKey, TArray<uint8>& OutData) final;
-
-	/**
-	 * Asynchronous, fire-and-forget placement of a cache item
-	 *
-	 * @param	CacheKey	Alphanumeric+underscore key of this cache item
-	 * @param	InData		Buffer containing the data to cache, can be destroyed after the call returns, immediately
-	 * @param	bPutEvenIfExists	If true, then do not attempt skip the put even if CachedDataProbablyExists returns true
-	 */
-	EPutStatus PutCachedData(const TCHAR* CacheKey, TArrayView<const uint8> InData, bool bPutEvenIfExists) final;
-
-	void RemoveCachedData(const TCHAR* CacheKey, bool bTransient) final;
-
-	/**
-	 * Save the cache to disk
-	 * @param	Filename	Filename to save
-	 * @return	true if file was saved successfully
-	 */
-	bool SaveCache(const TCHAR* Filename);
-
-	/**
-	 * Load the cache from disk
-	 * @param	Filename	Filename to load
-	 * @return	true if file was loaded successfully
-	 */
-	bool LoadCache(const TCHAR* Filename);
-
-	/**
-	 * Disable cache and ignore all subsequent requests
-	 */
-	void Disable() final;
-
-	TSharedRef<FDerivedDataCacheStatsNode> GatherUsageStats() const final;
-
-	TBitArray<> TryToPrefetch(TConstArrayView<FString> CacheKeys) final;
-
-	/**
-	 *  Determines if we would cache the provided data
-	 */
-	bool WouldCache(const TCHAR* CacheKey, TArrayView<const uint8> InData) final;
-
-	/**
-	 * Apply debug options
-	 */
-	bool ApplyDebugOptions(FBackendDebugOptions& InOptions) final;
-
-	EBackendLegacyMode GetLegacyMode() const final { return EBackendLegacyMode::ValueWithLegacyFallback; }
+	// ICacheStore Interface
 
 	void Put(
 		TConstArrayView<FCachePutRequest> Requests,
@@ -119,98 +54,84 @@ public:
 		IRequestOwner& Owner,
 		FOnCacheGetChunkComplete&& OnComplete) final;
 
-private:
-	/** Name of the cache file loaded (if any). */
-	FString CacheFilename;
-	FDerivedDataCacheUsageStats UsageStats;
+	// ILegacyCacheStore Interface
 
-	struct FCacheValue
-	{
-		int32 Age;
-		int32 Size;
-		FRWLock DataLock;
-		TArray<uint8> Data;
-		FCacheValue(int32 InSize, int32 InAge = 0)
-			: Age(InAge)
-			, Size(InSize)
-		{
-		}
+	void LegacyPut(
+		TConstArrayView<FLegacyCachePutRequest> Requests,
+		IRequestOwner& Owner,
+		FOnLegacyCachePutComplete&& OnComplete) final;
+	void LegacyGet(
+		TConstArrayView<FLegacyCacheGetRequest> Requests,
+		IRequestOwner& Owner,
+		FOnLegacyCacheGetComplete&& OnComplete) final;
+	void LegacyDelete(
+		TConstArrayView<FLegacyCacheDeleteRequest> Requests,
+		IRequestOwner& Owner,
+		FOnLegacyCacheDeleteComplete&& OnComplete) final;
 
-		~FCacheValue()
-		{
-			// Ensure we don't begin destruction of this value while our own PutCachedData function is holding this lock
-			FWriteScopeLock WriteLock(DataLock);
-		}
-	};
+	void LegacyStats(FDerivedDataCacheStatsNode& OutNode) final;
+	bool LegacyDebugOptions(FBackendDebugOptions& Options) final;
 
-	FORCEINLINE int32 CalcSerializedCacheValueSize(const FString& Key, const FCacheValue& Val)
-	{
-		return (Key.Len() + 1) * sizeof(TCHAR) + sizeof(FCacheValue::Age) + Val.Size;
-	}
+	// IMemoryCacheStore Interface
 
-	FORCEINLINE int32 CalcSerializedCacheValueSize(const FString& Key, const TArrayView<const uint8>& Data)
+	void Delete(const FCacheKey& Key) final;
+	void DeleteValue(const FCacheKey& Key) final;
+	void LegacyDelete(const FLegacyCacheKey& Key) final;
+
+	void Disable() final;
+
+private:
+	FDerivedDataCacheUsageStats UsageStats;
+
+	struct FCacheRecordComponents
 	{
-		return (Key.Len() + 1) * sizeof(TCHAR) + sizeof(FCacheValue::Age) + Data.Num();
-	}
+		FCbObject Meta;
+		TArray<FValueWithId> Values;
+	};
 
 	/** Name of this cache (used for debugging) */
 	FString Name;
 
-	/** Set of files that are being written to disk asynchronously. */
-	TMap<FString, FCacheValue*> CacheItems;
+	/** Set of legacy values in this cache. */
+	TMap<FCacheKey, FLegacyCacheValue> LegacyCacheValues;
 	/** Set of records in this cache. */
-	TSet<FCacheRecord, FCacheRecordKeyFuncs> CacheRecords;
+	TMap<FCacheKey, FCacheRecordComponents> CacheRecords;
 	/** Set of values in this cache. */
 	TMap<FCacheKey, FValue> CacheValues;
-	/** Maximum size the cached items can grow up to ( in bytes ) */
-	int64 MaxCacheSize;
+	/** Maximum size the cached items can grow up to (in bytes) */
+	uint64 MaxCacheSize;
 	/** When set to true, this cache is disabled...ignore all requests. */
 	std::atomic<bool> bDisabled;
 	/** Object used for synchronization via a scoped lock						*/
 	mutable FRWLock SynchronizationObject;
 	/** Current estimated cache size in bytes */
-	int64 CurrentCacheSize;
+	uint64 CurrentCacheSize;
 	/** Indicates that the cache max size has been exceeded. This is used to avoid
 		warning spam after the size has reached the limit. */
 	bool bMaxSizeExceeded;
 
-	/** When a memory cache can be disabled, it won't return true for CachedDataProbablyExists calls.
-	  * This is to avoid having the Boot DDC tells it has some resources that will suddenly disappear after the boot.
-	  * Get() requests will still get fulfilled and other cache level will be properly back-filled 
-	  * offering the speed benefit of the boot cache while maintaining coherency at all cache levels.
-	  * 
-	  * The problem is that most asset types (audio/staticmesh/texture) will always verify if their different LODS/Chunks can be found in the cache using CachedDataProbablyExists.
-	  * If any of the LOD/MIP can't be found, a build of the asset is triggered, otherwise they skip asset compilation altogether.
-	  * However, we should not skip the compilation based on the CachedDataProbablyExists result of the boot cache because it is a lie and will disappear at some point.
-	  * When the boot cache disappears and the streamer tries to fetch a LOD that it has been told was cached, it will fail and will then have no choice but to rebuild the asset synchronously.
-	  * This obviously causes heavy game-thread stutters.
-
-	  * However, if the bootcache returns false during CachedDataProbablyExists. The async compilation will be triggered and data will be put in the both the boot.ddc and the local cache.
-	  * This way, no more heavy game-thread stutters during streaming...
-
-	  * This can be reproed when you clear the local cache but do not clear the boot.ddc file, but even if it's a corner case, I stumbled upon it enough times that I though it was worth to fix so the caches are coherent.
-	  */
+	/**
+	 * When a memory cache can be disabled, it won't return true for CachedDataProbablyExists calls.
+	 * This is to avoid having the Boot DDC tells it has some resources that will suddenly disappear after the boot.
+	 * Get() requests will still get fulfilled and other cache level will be properly back-filled 
+	 * offering the speed benefit of the boot cache while maintaining coherency at all cache levels.
+	 * 
+	 * The problem is that most asset types (audio/staticmesh/texture) will always verify if their different LODS/Chunks can be found in the cache using CachedDataProbablyExists.
+	 * If any of the LOD/MIP can't be found, a build of the asset is triggered, otherwise they skip asset compilation altogether.
+	 * However, we should not skip the compilation based on the CachedDataProbablyExists result of the boot cache because it is a lie and will disappear at some point.
+	 * When the boot cache disappears and the streamer tries to fetch a LOD that it has been told was cached, it will fail and will then have no choice but to rebuild the asset synchronously.
+	 * This obviously causes heavy game-thread stutters.
+
+	 * However, if the bootcache returns false during CachedDataProbablyExists. The async compilation will be triggered and data will be put in the both the boot.ddc and the local cache.
+	 * This way, no more heavy game-thread stutters during streaming...
+
+	 * This can be reproed when you clear the local cache but do not clear the boot.ddc file, but even if it's a corner case, I stumbled upon it enough times that I though it was worth to fix so the caches are coherent.
+	 */
 	bool bCanBeDisabled = false;
 	bool bShuttingDown  = false;
 
-	enum 
-	{
-		/** Magic number to use in header */
-		MemCache_Magic = 0x0cac0ddc,
-		/** Magic number to use in header (new, > 2GB size compatible) */
-		MemCache_Magic64 = 0x0cac1ddc,
-		/** Oldest cache items to keep */
-		MaxAge = 3,
-		/** Size of data that is stored in the cachefile apart from the cache entries (64 bit size). */
-		SerializationSpecificDataSize = sizeof(uint32)	// Magic
-									  + sizeof(int64)	// Size
-									  + sizeof(uint32), // CRC
-	};
-
 protected:
 
-	/* Debug helpers */
-	bool ShouldSimulateMiss(const TCHAR* InKey);
 	bool ShouldSimulateMiss(const FCacheKey& InKey);
 
 	/** Debug Options */
@@ -224,9 +145,9 @@ protected:
 
 FMemoryCacheStore::FMemoryCacheStore(const TCHAR* InName, int64 InMaxCacheSize, bool bInCanBeDisabled)
 	: Name(InName)
-	, MaxCacheSize(InMaxCacheSize)
-	, bDisabled( false )
-	, CurrentCacheSize( SerializationSpecificDataSize )
+	, MaxCacheSize(InMaxCacheSize < 0 ? 0 : uint64(InMaxCacheSize))
+	, bDisabled(false)
+	, CurrentCacheSize(0)
 	, bMaxSizeExceeded(false)
 	, bCanBeDisabled(bInCanBeDisabled)
 {
@@ -238,401 +159,32 @@ FMemoryCacheStore::~FMemoryCacheStore()
 	Disable();
 }
 
-bool FMemoryCacheStore::IsWritable() const
-{
-	return !bDisabled;
-}
-
-FDerivedDataBackendInterface::ESpeedClass FMemoryCacheStore::GetSpeedClass() const
-{
-	return ESpeedClass::Local;
-}
-
-bool FMemoryCacheStore::CachedDataProbablyExists(const TCHAR* CacheKey)
-{
-	// See comments on the declaration of bCanBeDisabled variable.
-	if (bCanBeDisabled)
-	{
-		return false;
-	}
-
-	COOK_STAT(auto Timer = UsageStats.TimeProbablyExists());
-
-	if (ShouldSimulateMiss(CacheKey))
-	{
-		return false;
-	}
-
-	if (bDisabled)
-	{
-		return false;
-	}
-
-	FReadScopeLock ScopeLock(SynchronizationObject);
-	bool Result = CacheItems.Contains(FString(CacheKey));
-	if (Result)
-	{
-		COOK_STAT(Timer.AddHit(0));
-	}
-	return Result;
-}
-
-bool FMemoryCacheStore::GetCachedData(const TCHAR* CacheKey, TArray<uint8>& OutData)
-{
-	COOK_STAT(auto Timer = UsageStats.TimeGet());
-
-	if (ShouldSimulateMiss(CacheKey))
-	{
-		return false;
-	}
-	
-	if (!bDisabled)
-	{
-		FReadScopeLock ScopeLock(SynchronizationObject);
-
-		FCacheValue* Item = CacheItems.FindRef(FString(CacheKey));
-		if (Item)
-		{
-			TRACE_CPUPROFILER_EVENT_SCOPE(FMemoryCacheStore::GetCachedData);
-			FReadScopeLock ItemLock(Item->DataLock);
-
-			OutData = Item->Data;
-			Item->Age = 0;
-			check(OutData.Num());
-			COOK_STAT(Timer.AddHit(OutData.Num()));
-			return true;
-		}
-	}
-	OutData.Empty();
-	return false;
-}
-
-TBitArray<> FMemoryCacheStore::TryToPrefetch(TConstArrayView<FString> CacheKeys)
-{
-	return CachedDataProbablyExistsBatch(CacheKeys);
-}
-
-bool FMemoryCacheStore::WouldCache(const TCHAR* CacheKey, TArrayView<const uint8> InData)
-{
-	if (bDisabled || bMaxSizeExceeded)
-	{
-		return false;
-	}
-
-	return true;
-}
-
-FDerivedDataBackendInterface::EPutStatus FMemoryCacheStore::PutCachedData(const TCHAR* CacheKey, TArrayView<const uint8> InData, bool bPutEvenIfExists)
-{
-	TRACE_CPUPROFILER_EVENT_SCOPE(FMemoryCacheStore::PutCachedData);
-	COOK_STAT(auto Timer = UsageStats.TimePut());
-
-	FString Key;
-	{
-		FReadScopeLock ReadScopeLock(SynchronizationObject);
-
-		if (ShouldSimulateMiss(CacheKey))
-		{
-			return EPutStatus::Skipped;
-		}
-
-		// Should never hit this as higher level code should be checking..
-		if (!WouldCache(CacheKey, InData))
-		{
-			//UE_LOG(LogDerivedDataCache, Warning, TEXT("WouldCache was not called prior to attempted Put!"));
-			return EPutStatus::NotCached;
-		}
-
-		Key = CacheKey;
-		FCacheValue* Item = CacheItems.FindRef(Key);
-		if (Item)
-		{
-			//check(Item->Data == InData); // any second attempt to push data should be identical data
-			return EPutStatus::Cached;
-		}
-	}
-	
-	// Compute the size of the cache before copying the data to avoid alloc/memcpy of something we might throw away
-	// if we bust the cache size.
-	int32 CacheValueSize = CalcSerializedCacheValueSize(Key, InData);
-
-	FCacheValue* Val = nullptr;
-
-	{
-		FWriteScopeLock WriteScopeLock(SynchronizationObject);
-
-		// check if we haven't exceeded the MaxCacheSize
-		if (MaxCacheSize > 0 && (CurrentCacheSize + CacheValueSize) > MaxCacheSize)
-		{
-			UE_LOG(LogDerivedDataCache, Display, TEXT("Failed to cache data. Maximum cache size reached. CurrentSize %d kb / MaxSize: %d kb"), CurrentCacheSize / 1024, MaxCacheSize / 1024);
-			bMaxSizeExceeded = true;
-			return EPutStatus::NotCached;
-		}
-
-		if (CacheItems.Contains(Key))
-		{
-			// Another thread already beat us, just return.
-			return EPutStatus::Cached;
-		}
-
-		COOK_STAT(Timer.AddHit(InData.Num()));
-		Val = new FCacheValue(InData.Num());
-		// Make sure no other thread can access the data until we finish copying it
-		Val->DataLock.WriteLock();
-		CacheItems.Add(Key, Val);
-
-		CurrentCacheSize += CacheValueSize;
-	}
-
-	// Data is copied outside of the lock so that we only lock
-	// for this specific key instead of always locking all keys by default.
-	Val->Data = InData;
-	// It is now safe for other thread to access this data, unlock
-	Val->DataLock.WriteUnlock();
-
-	return EPutStatus::Cached;
-}
-
-void FMemoryCacheStore::RemoveCachedData(const TCHAR* CacheKey, bool bTransient)
-{
-	if (bDisabled || bTransient)
-	{
-		return;
-	}
-
-	TRACE_CPUPROFILER_EVENT_SCOPE(FMemoryCacheStore::RemoveCachedData);
-	FString Key(CacheKey);
-	FCacheValue* Item = nullptr;
-	{
-		FWriteScopeLock WriteScopeLock(SynchronizationObject);
-		
-		if (CacheItems.RemoveAndCopyValue(Key, Item) && Item)
-		{
-			CurrentCacheSize -= CalcSerializedCacheValueSize(Key, *Item);
-			bMaxSizeExceeded = false;
-		}
-	}
-
-	if (Item)
-	{
-		// Just make sure any other thread has finished writing or reading the data before deleting.
-		Item->DataLock.WriteLock();
-
-		// Avoid deleting the item with a locked FRWLock, unlocking is safe because
-		// the item is now unreachable from other threads
-		Item->DataLock.WriteUnlock();
-		delete Item;
-	}
-}
-
-bool FMemoryCacheStore::SaveCache(const TCHAR* Filename)
-{
-	TRACE_CPUPROFILER_EVENT_SCOPE(FMemoryCacheStore::SaveCache);
-
-	double StartTime = FPlatformTime::Seconds();
-	TUniquePtr<FArchive> SaverArchive(IFileManager::Get().CreateFileWriter(Filename, FILEWRITE_EvenIfReadOnly));
-	if (!SaverArchive)
-	{
-		UE_LOG(LogDerivedDataCache, Error, TEXT("Could not save memory cache %s."), Filename);
-		return false;
-	}
-
-	FArchive& Saver = *SaverArchive;
-	uint32 Magic = MemCache_Magic64;
-	Saver << Magic;
-	const int64 DataStartOffset = Saver.Tell();
-	{
-		FWriteScopeLock ScopeLock(SynchronizationObject);
-		check(!bDisabled);
-		for (TMap<FString, FCacheValue*>::TIterator It(CacheItems); It; ++It )
-		{
-			Saver << It.Key();
-			Saver << It.Value()->Age;
-			FReadScopeLock ReadScopeLock(It.Value()->DataLock);
-			Saver << It.Value()->Data;
-		}
-	}
-	const int64 DataSize = Saver.Tell(); // Everything except the footer
-	int64 Size = DataSize;
-	uint32 Crc = MemCache_Magic64; // Crc takes more time than I want to spend  FCrc::MemCrc_DEPRECATED(&Buffer[0], Size);
-	Saver << Size;
-	Saver << Crc;
-
-	check(SerializationSpecificDataSize + DataSize <= MaxCacheSize || MaxCacheSize <= 0);
-
-	UE_LOG(LogDerivedDataCache, Log, TEXT("Saved boot cache %4.2fs %lldMB %s."), float(FPlatformTime::Seconds() - StartTime), DataSize / (1024 * 1024), Filename);
-	return true;
-}
-
-bool FMemoryCacheStore::LoadCache(const TCHAR* Filename)
-{
-	TRACE_CPUPROFILER_EVENT_SCOPE(FMemoryCacheStore::LoadCache);
-
-	double StartTime = FPlatformTime::Seconds();
-	const int64 FileSize = IFileManager::Get().FileSize(Filename);
-	if (FileSize < 0)
-	{
-		UE_LOG(LogDerivedDataCache, Warning, TEXT("Could not find memory cache %s."), Filename);
-		return false;
-	}
-	// We test 3 * uint32 which is the old format (< SerializationSpecificDataSize). We'll test
-	// against SerializationSpecificDataSize later when we read the magic number from the cache.
-	if (FileSize < sizeof(uint32) * 3)
-	{
-		UE_LOG(LogDerivedDataCache, Error, TEXT("Memory cache was corrputed (short) %s."), Filename);
-		return false;
-	}
-	if (FileSize > MaxCacheSize*2 && MaxCacheSize > 0)
-	{
-		UE_LOG(LogDerivedDataCache, Error, TEXT("Refusing to load DDC cache %s. Size exceeds doubled MaxCacheSize."), Filename);
-		return false;
-	}
-
-	TUniquePtr<FArchive> LoaderArchive(IFileManager::Get().CreateFileReader(Filename));
-	if (!LoaderArchive)
-	{
-		UE_LOG(LogDerivedDataCache, Warning, TEXT("Could not read memory cache %s."), Filename);
-		return false;
-	}
-
-	FArchive& Loader = *LoaderArchive;
-	uint32 Magic = 0;
-	Loader << Magic;
-	if (Magic != MemCache_Magic && Magic != MemCache_Magic64)
-	{
-		UE_LOG(LogDerivedDataCache, Error, TEXT("Memory cache was corrputed (magic) %s."), Filename);
-		return false;
-	}
-	// Check the file size again, this time against the correct minimum size.
-	if (Magic == MemCache_Magic64 && FileSize < SerializationSpecificDataSize)
-	{
-		UE_LOG(LogDerivedDataCache, Error, TEXT("Memory cache was corrputed (short) %s."), Filename);
-		return false;
-	}
-	// Calculate expected DataSize based on the magic number (footer size difference)
-	const int64 DataSize = FileSize - (Magic == MemCache_Magic64 ? (SerializationSpecificDataSize - sizeof(uint32)) : (sizeof(uint32) * 2));		
-	Loader.Seek(DataSize);
-	int64 Size = 0;
-	uint32 Crc = 0;
-	if (Magic == MemCache_Magic64)
-	{
-		Loader << Size;
-	}
-	else
-	{
-		uint32 Size32 = 0;
-		Loader << Size32;
-		Size = (int64)Size32;
-	}
-	Loader << Crc;
-	if (Size != DataSize)
-	{
-		UE_LOG(LogDerivedDataCache, Error, TEXT("Memory cache was corrputed (size) %s."), Filename);
-		return false;
-	}
-	if ((Crc != MemCache_Magic && Crc != MemCache_Magic64) || Crc != Magic)
-	{
-		UE_LOG(LogDerivedDataCache, Warning, TEXT("Memory cache was corrputed (crc) %s."), Filename);
-		return false;
-	}
-	// Seek to data start offset (skip magic number)
-	Loader.Seek(sizeof(uint32));
-	{
-		TArray<uint8> Working;
-		FWriteScopeLock ScopeLock(SynchronizationObject);
-		check(!bDisabled);
-		while (Loader.Tell() < DataSize)
-		{
-			FString Key;
-			int32 Age;
-			Loader << Key;
-			Loader << Age;
-			Age++;
-			Loader << Working;
-			if (Age < MaxAge)
-			{
-				CacheItems.Add(Key, new FCacheValue(Working.Num(), Age))->Data = MoveTemp(Working);
-			}
-			Working.Reset();
-		}
-		// these are just a double check on ending correctly
-		if (Magic == MemCache_Magic64)
-		{
-			Loader << Size;
-		}
-		else
-		{
-			uint32 Size32 = 0;
-			Loader << Size32;
-			Size = (int64)Size32;
-		}
-		Loader << Crc;
-
-		CurrentCacheSize = FileSize;
-		CacheFilename = Filename;
-	}
-
-	UE_LOG(LogDerivedDataCache, Log, TEXT("Loaded boot cache %4.2fs %lldMB %s."), float(FPlatformTime::Seconds() - StartTime), DataSize / (1024 * 1024), Filename);
-	return true;
-}
-
 void FMemoryCacheStore::Disable()
 {
 	check(bCanBeDisabled || bShuttingDown);
 	FWriteScopeLock ScopeLock(SynchronizationObject);
 	bDisabled = true;
-	for (TMap<FString,FCacheValue*>::TIterator It(CacheItems); It; ++It )
-	{
-		delete It.Value();
-	}
-	CacheItems.Empty();
-
-	CurrentCacheSize = SerializationSpecificDataSize;
+	CacheRecords.Empty();
+	CacheValues.Empty();
+	LegacyCacheValues.Empty();
+	CurrentCacheSize = 0;
 }
 
-TSharedRef<FDerivedDataCacheStatsNode> FMemoryCacheStore::GatherUsageStats() const
+void FMemoryCacheStore::LegacyStats(FDerivedDataCacheStatsNode& OutNode)
 {
-	TSharedRef<FDerivedDataCacheStatsNode> Usage =
-		MakeShared<FDerivedDataCacheStatsNode>(CacheFilename.IsEmpty() ? TEXT("Memory") : TEXT("Boot"), CacheFilename, /*bIsLocal*/ true);
-	Usage->Stats.Add(TEXT(""), UsageStats);
-	return Usage;
+	OutNode = {!bCanBeDisabled ? TEXT("Memory") : TEXT("Boot"), TEXT(""), /*bIsLocal*/ true};
+	OutNode.Stats.Add(TEXT(""), UsageStats);
 }
 
-bool FMemoryCacheStore::ApplyDebugOptions(FBackendDebugOptions& InOptions)
+bool FMemoryCacheStore::LegacyDebugOptions(FBackendDebugOptions& InOptions)
 {
 	DebugOptions = InOptions;
 	return true;
 }
 
-bool FMemoryCacheStore::ShouldSimulateMiss(const TCHAR* InKey)
-{
-	if (DebugOptions.RandomMissRate == 0 && DebugOptions.SimulateMissTypes.IsEmpty())
-	{
-		return false;
-	}
-
-	const FName Key(InKey);
-	const uint32 Hash = GetTypeHash(Key);
-
-	if (FScopeLock Lock(&MissedKeysCS); DebugMissedKeys.ContainsByHash(Hash, Key))
-	{
-		return true;
-	}
-
-	if (DebugOptions.ShouldSimulateMiss(InKey))
-	{
-		FScopeLock Lock(&MissedKeysCS);
-		UE_LOG(LogDerivedDataCache, Verbose, TEXT("Simulating miss in %s for %s"), *GetName(), InKey);
-		DebugMissedKeys.AddByHash(Hash, Key);
-		return true;
-	}
-
-	return false;
-}
-
 bool FMemoryCacheStore::ShouldSimulateMiss(const FCacheKey& Key)
 {
-	if (DebugOptions.RandomMissRate == 0 && DebugOptions.SimulateMissTypes.IsEmpty())
+	if (!bCanBeDisabled || (DebugOptions.RandomMissRate == 0 && DebugOptions.SimulateMissTypes.IsEmpty()))
 	{
 		return false;
 	}
@@ -659,17 +211,24 @@ void FMemoryCacheStore::Put(
 	IRequestOwner& Owner,
 	FOnCachePutComplete&& OnComplete)
 {
+	if (bDisabled)
+	{
+		return CompleteWithStatus(Requests, OnComplete, EStatus::Error);
+	}
+
 	for (const FCachePutRequest& Request : Requests)
 	{
-		const FCacheRecord& Record = Request.Record;
-		const FCacheKey& Key = Record.GetKey();
 		EStatus Status = EStatus::Error;
 		ON_SCOPE_EXIT
 		{
-			OnComplete({Request.Name, Key, Request.UserData, Status});
+			OnComplete(Request.MakeResponse(Status));
 		};
 
-		if (!EnumHasAnyFlags(Request.Policy.GetRecordPolicy(), ECachePolicy::StoreLocal))
+		const FCacheRecord& Record = Request.Record;
+		const FCacheKey& Key = Record.GetKey();
+		const TConstArrayView<FValueWithId> Values = Record.GetValues();
+
+		if (Algo::NoneOf(Values, &FValue::HasData))
 		{
 			continue;
 		}
@@ -677,50 +236,128 @@ void FMemoryCacheStore::Put(
 		if (ShouldSimulateMiss(Key))
 		{
 			UE_LOG(LogDerivedDataCache, Verbose, TEXT("%s: Simulated miss for put of %s from '%s'"),
-				*GetName(), *WriteToString<96>(Key), *Request.Name);
+				*Name, *WriteToString<96>(Key), *Request.Name);
 			continue;
 		}
 
-		const TConstArrayView<FValueWithId> Values = Record.GetValues();
+		COOK_STAT(auto Timer = UsageStats.TimePut());
+		const bool bReplaceExisting = !EnumHasAnyFlags(Request.Policy.GetRecordPolicy(), ECachePolicy::QueryLocal);
+
+		FWriteScopeLock ScopeLock(SynchronizationObject);
+		FCacheRecordComponents& Components = CacheRecords.FindOrAdd(Key);
+		const bool bHasExisting = !Components.Values.IsEmpty();
+		Status = EStatus::Ok;
 
-		if (!Algo::AllOf(Values, &FValue::HasData))
+		if (bHasExisting && !bReplaceExisting && Algo::AllOf(Components.Values, &FValue::HasData))
 		{
 			continue;
 		}
 
-		COOK_STAT(auto Timer = UsageStats.TimePut());
-		const int64 RecordSize = Private::GetCacheRecordCompressedSize(Record);
-		const bool bReplaceExisting = !EnumHasAnyFlags(Request.Policy.GetRecordPolicy(), ECachePolicy::QueryLocal);
+		int64 RequiredSize = 0;
+		if (bHasExisting && !bReplaceExisting)
+		{
+			for (const FValueWithId& Value : Components.Values)
+			{
+				if (!Value.HasData())
+				{
+					if (const FValueWithId& NewValue = Record.GetValue(Value.GetId());
+						NewValue &&
+						NewValue.HasData() &&
+						NewValue.GetRawHash() == Value.GetRawHash())
+					{
+						RequiredSize += NewValue.GetData().GetCompressedSize();
+					}
+					else if (!EnumHasAnyFlags(Request.Policy.GetRecordPolicy(), ECachePolicy::PartialRecord) &&
+						EnumHasAnyFlags(Request.Policy.GetValuePolicy(Value.GetId()), ECachePolicy::StoreLocal))
+					{
+						Status = EStatus::Error;
+					}
+				}
+			}
+		}
+		else
+		{
+			RequiredSize = Record.GetMeta().GetSize() - Components.Meta.GetSize();
+			RequiredSize -= Algo::TransformAccumulate(Components.Values, [](const FValue& Value) { return Value.GetData().GetCompressedSize(); }, uint64(0));
+			for (const FValueWithId& NewValue : Record.GetValues())
+			{
+				if (NewValue.HasData())
+				{
+					RequiredSize += NewValue.GetData().GetCompressedSize();
+				}
+				else if (const int32 ValueIndex = Algo::BinarySearchBy(Components.Values, NewValue.GetId(), &FValueWithId::GetId);
+					ValueIndex != INDEX_NONE &&
+					Components.Values[ValueIndex].HasData() &&
+					Components.Values[ValueIndex].GetRawHash() == NewValue.GetRawHash())
+				{
+					RequiredSize += Components.Values[ValueIndex].GetData().GetCompressedSize();
+				}
+				else if (!EnumHasAnyFlags(Request.Policy.GetRecordPolicy(), ECachePolicy::PartialRecord) &&
+					EnumHasAnyFlags(Request.Policy.GetValuePolicy(NewValue.GetId()), ECachePolicy::StoreLocal))
+				{
+					Status = EStatus::Error;
+				}
+			}
+		}
 
-		FWriteScopeLock ScopeLock(SynchronizationObject);
-		FCacheRecord* const ExistingRecord = CacheRecords.Find(Key);
-		Status = ExistingRecord && !bReplaceExisting ? EStatus::Ok : EStatus::Error;
-		if (bDisabled || Status == EStatus::Ok)
+		if (Status == EStatus::Error)
 		{
 			continue;
 		}
 
-		const int64 ExistingRecordSize = ExistingRecord ? Private::GetCacheRecordCompressedSize(*ExistingRecord) : 0;
-		const int64 RequiredSize = RecordSize - ExistingRecordSize;
-
 		if (MaxCacheSize > 0 && (CurrentCacheSize + RequiredSize) > MaxCacheSize)
 		{
-			UE_LOG(LogDerivedDataCache, Display, TEXT("Failed to cache data. Maximum cache size reached. CurrentSize %" INT64_FMT " KiB / MaxSize: %" INT64_FMT " KiB"), CurrentCacheSize / 1024, MaxCacheSize / 1024);
+			UE_CLOG(!bMaxSizeExceeded, LogDerivedDataCache, Display,
+				TEXT("Failed to cache data. Maximum cache size reached. ")
+				TEXT("CurrentSize %" UINT64_FMT " KiB / MaxSize: %" UINT64_FMT " KiB"),
+				CurrentCacheSize / 1024, MaxCacheSize / 1024);
 			bMaxSizeExceeded = true;
+			Status = EStatus::Ok;
 			continue;
 		}
 
+		Status = EStatus::Ok;
 		CurrentCacheSize += RequiredSize;
-		if (ExistingRecord)
+
+		if (!bHasExisting)
+		{
+			Components.Meta = Record.GetMeta();
+			Components.Values = Record.GetValues();
+		}
+		else if (!bReplaceExisting)
 		{
-			*ExistingRecord = Record;
+			for (FValueWithId& Value : Components.Values)
+			{
+				if (!Value.HasData())
+				{
+					if (const FValueWithId& NewValue = Record.GetValue(Value.GetId()); NewValue && NewValue.GetRawHash() == Value.GetRawHash())
+					{
+						Value = NewValue;
+					}
+				}
+			}
 		}
 		else
 		{
-			CacheRecords.Add(Record);
+			FCacheRecordComponents ExistingComponents = MoveTemp(Components);
+			Components.Meta = Record.GetMeta();
+			Components.Values = Record.GetValues();
+			for (FValueWithId& Value : Components.Values)
+			{
+				if (!Value.HasData())
+				{
+					if (const int32 ExistingValueIndex = Algo::BinarySearchBy(ExistingComponents.Values, Value.GetId(), &FValueWithId::GetId);
+						ExistingValueIndex != INDEX_NONE &&
+						ExistingComponents.Values[ExistingValueIndex].HasData() &&
+						ExistingComponents.Values[ExistingValueIndex].GetRawHash() == Value.GetRawHash())
+					{
+						Value = ExistingComponents.Values[ExistingValueIndex];
+					}
+				}
+			}
 		}
-		COOK_STAT(Timer.AddHit(RecordSize));
-		Status = EStatus::Ok;
+
+		COOK_STAT(Timer.AddHit(RequiredSize));
 	}
 }
 
@@ -738,44 +375,68 @@ void FMemoryCacheStore::Get(
 	{
 		const FCacheKey& Key = Request.Key;
 		const FCacheRecordPolicy& Policy = Request.Policy;
-		const bool bExistsOnly = EnumHasAllFlags(Policy.GetRecordPolicy(), ECachePolicy::SkipData);
-		COOK_STAT(auto Timer = bExistsOnly ? UsageStats.TimeProbablyExists() : UsageStats.TimeGet());
-		FOptionalCacheRecord Record;
+		COOK_STAT(auto Timer = EnumHasAnyFlags(Policy.GetRecordPolicy(), ECachePolicy::SkipData) ? UsageStats.TimeProbablyExists() : UsageStats.TimeGet());
+		FCacheRecordComponents Components;
 		EStatus Status = EStatus::Error;
+
 		if (ShouldSimulateMiss(Key))
 		{
 			UE_LOG(LogDerivedDataCache, Verbose, TEXT("%s: Simulated miss for get of %s from '%s'"),
-				*GetName(), *WriteToString<96>(Key), *Request.Name);
+				*Name, *WriteToString<96>(Key), *Request.Name);
 		}
-		else if (FReadScopeLock ScopeLock(SynchronizationObject); const FCacheRecord* CacheRecord = CacheRecords.Find(Key))
+		else if (FReadScopeLock ScopeLock(SynchronizationObject); const FCacheRecordComponents* FoundComponents = CacheRecords.Find(Key))
 		{
-			Status = EStatus::Ok;
-			Record = *CacheRecord;
+			if (!FoundComponents->Values.IsEmpty())
+			{
+				Status = EStatus::Ok;
+				Components = *FoundComponents;
+			}
 		}
-		if (Record)
+
+		FCacheRecordBuilder Builder(Request.Key);
+
+		if (Status == EStatus::Ok)
 		{
-			for (const FValueWithId& Value : Record.Get().GetValues())
+			const ECachePolicy RecordPolicy = Policy.GetRecordPolicy();
+			for (const FValueWithId& Value : Components.Values)
 			{
-				if (!Value.HasData() && !EnumHasAnyFlags(Policy.GetValuePolicy(Value.GetId()), ECachePolicy::SkipData))
+				const ECachePolicy ValuePolicy = Policy.GetValuePolicy(Value.GetId());
+				if (!EnumHasAnyFlags(ValuePolicy, ECachePolicy::QueryLocal))
+				{
+					Builder.AddValue(Value.RemoveData());
+					continue;
+				}
+				if (!Value.HasData())
 				{
 					Status = EStatus::Error;
-					if (!EnumHasAllFlags(Policy.GetRecordPolicy(), ECachePolicy::PartialRecord))
+					if (!EnumHasAnyFlags(RecordPolicy, ECachePolicy::PartialRecord))
 					{
-						Record.Reset();
+						Builder = FCacheRecordBuilder(Request.Key);
 						break;
 					}
 				}
+				const bool bExistsOnly = EnumHasAnyFlags(ValuePolicy, ECachePolicy::SkipData);
+				Builder.AddValue(bExistsOnly ? Value.RemoveData() : Value);
 			}
 		}
-		if (Record)
-		{
-			COOK_STAT(Timer.AddHit(Private::GetCacheRecordCompressedSize(Record.Get())));
-			OnComplete({Request.Name, MoveTemp(Record).Get(), Request.UserData, Status});
-		}
-		else
-		{
-			OnComplete(Request.MakeResponse(EStatus::Error));
-		}
+
+		FCacheRecord Record = Builder.Build();
+		COOK_STAT(Timer.AddHitOrMiss(
+			Status == EStatus::Ok ? FCookStats::CallStats::EHitOrMiss::Hit : FCookStats::CallStats::EHitOrMiss::Miss,
+			Private::GetCacheRecordCompressedSize(Record)));
+		OnComplete({Request.Name, MoveTemp(Record), Request.UserData, Status});
+	}
+}
+
+void FMemoryCacheStore::Delete(const FCacheKey& Key)
+{
+	FWriteScopeLock ScopeLock(SynchronizationObject);
+	FCacheRecordComponents Components;
+	if (CacheRecords.RemoveAndCopyValue(Key, Components))
+	{
+		CurrentCacheSize -= Components.Meta.GetSize() +
+			Algo::TransformAccumulate(Components.Values, [](const FValue& Value) { return Value.GetData().GetCompressedSize(); }, uint64(0));
+		bMaxSizeExceeded = false;
 	}
 }
 
@@ -784,6 +445,11 @@ void FMemoryCacheStore::PutValue(
 	IRequestOwner& Owner,
 	FOnCachePutValueComplete&& OnComplete)
 {
+	if (bDisabled)
+	{
+		return CompleteWithStatus(Requests, OnComplete, EStatus::Error);
+	}
+
 	for (const FCachePutValueRequest& Request : Requests)
 	{
 		EStatus Status = EStatus::Error;
@@ -792,23 +458,18 @@ void FMemoryCacheStore::PutValue(
 			OnComplete(Request.MakeResponse(Status));
 		};
 
-		if (!EnumHasAnyFlags(Request.Policy, ECachePolicy::StoreLocal))
-		{
-			continue;
-		}
-
 		const FCacheKey& Key = Request.Key;
 		const FValue& Value = Request.Value;
 
-		if (ShouldSimulateMiss(Key))
+		if (!Value.HasData())
 		{
-			UE_LOG(LogDerivedDataCache, Verbose, TEXT("%s: Simulated miss for put of %s from '%s'"),
-				*GetName(), *WriteToString<96>(Key), *Request.Name);
 			continue;
 		}
 
-		if (!Value.HasData())
+		if (ShouldSimulateMiss(Key))
 		{
+			UE_LOG(LogDerivedDataCache, Verbose, TEXT("%s: Simulated miss for put of %s from '%s'"),
+				*Name, *WriteToString<96>(Key), *Request.Name);
 			continue;
 		}
 
@@ -818,9 +479,10 @@ void FMemoryCacheStore::PutValue(
 
 		FWriteScopeLock ScopeLock(SynchronizationObject);
 		FValue* const ExistingValue = CacheValues.Find(Key);
-		Status = ExistingValue && !bReplaceExisting ? EStatus::Ok : EStatus::Error;
-		if (bDisabled || Status == EStatus::Ok)
+
+		if (ExistingValue && !bReplaceExisting)
 		{
+			Status = EStatus::Ok;
 			continue;
 		}
 
@@ -829,7 +491,10 @@ void FMemoryCacheStore::PutValue(
 
 		if (MaxCacheSize > 0 && (CurrentCacheSize + RequiredSize) > MaxCacheSize)
 		{
-			UE_LOG(LogDerivedDataCache, Display, TEXT("Failed to cache data. Maximum cache size reached. CurrentSize %" INT64_FMT " KiB / MaxSize: %" INT64_FMT " KiB"), CurrentCacheSize / 1024, MaxCacheSize / 1024);
+			UE_CLOG(!bMaxSizeExceeded, LogDerivedDataCache, Display,
+				TEXT("Failed to cache data. Maximum cache size reached. ")
+				TEXT("CurrentSize %" UINT64_FMT " KiB / MaxSize: %" UINT64_FMT " KiB"),
+				CurrentCacheSize / 1024, MaxCacheSize / 1024);
 			bMaxSizeExceeded = true;
 			continue;
 		}
@@ -843,6 +508,7 @@ void FMemoryCacheStore::PutValue(
 		{
 			CacheValues.Add(Key, Value);
 		}
+
 		COOK_STAT(Timer.AddHit(ValueSize));
 		Status = EStatus::Ok;
 	}
@@ -861,39 +527,38 @@ void FMemoryCacheStore::GetValue(
 	for (const FCacheGetValueRequest& Request : Requests)
 	{
 		const FCacheKey& Key = Request.Key;
-		const ECachePolicy Policy = Request.Policy;
-		const bool bExistsOnly = EnumHasAllFlags(Policy, ECachePolicy::SkipData);
+		const bool bExistsOnly = EnumHasAllFlags(Request.Policy, ECachePolicy::SkipData);
 		COOK_STAT(auto Timer = bExistsOnly ? UsageStats.TimeProbablyExists() : UsageStats.TimeGet());
-		bool bProcessHit = false;
+
 		FValue Value;
 		EStatus Status = EStatus::Error;
 		if (ShouldSimulateMiss(Key))
 		{
 			UE_LOG(LogDerivedDataCache, Verbose, TEXT("%s: Simulated miss for get of %s from '%s'"),
-				*GetName(), *WriteToString<96>(Key), *Request.Name);
+				*Name, *WriteToString<96>(Key), *Request.Name);
+		}
+		else if (bExistsOnly && bCanBeDisabled)
+		{
 		}
 		else if (FReadScopeLock ScopeLock(SynchronizationObject); const FValue* CacheValue = CacheValues.Find(Key))
 		{
 			Status = EStatus::Ok;
-			Value = *CacheValue;
-			bProcessHit = true;
+			Value = !bExistsOnly ? *CacheValue : CacheValue->RemoveData();
+			COOK_STAT(Timer.AddHit(Value.GetData().GetCompressedSize()));
 		}
 
-		if (bProcessHit)
-		{
-			if (!Value.HasData() && !EnumHasAnyFlags(Policy, ECachePolicy::SkipData))
-			{
-				Status = EStatus::Error;
-			}
+		OnComplete({Request.Name, Request.Key, MoveTemp(Value), Request.UserData, Status});
+	}
+}
 
-			COOK_STAT(Timer.AddHit(Value.GetData().GetCompressedSize()));
-			Value = EnumHasAnyFlags(Policy, ECachePolicy::SkipData) ? Value.RemoveData() : Value;
-			OnComplete({Request.Name, Request.Key, MoveTemp(Value), Request.UserData, Status});
-		}
-		else
-		{
-			OnComplete(Request.MakeResponse(EStatus::Error));
-		}
+void FMemoryCacheStore::DeleteValue(const FCacheKey& Key)
+{
+	FWriteScopeLock ScopeLock(SynchronizationObject);
+	FValue Value;
+	if (CacheValues.RemoveAndCopyValue(Key, Value))
+	{
+		CurrentCacheSize -= Value.GetData().GetCompressedSize();
+		bMaxSizeExceeded = false;
 	}
 }
 
@@ -920,7 +585,7 @@ void FMemoryCacheStore::GetChunks(
 		if (ShouldSimulateMiss(Request.Key))
 		{
 			UE_LOG(LogDerivedDataCache, Verbose, TEXT("%s: Simulated miss for get of %s from '%s'"),
-				*GetName(), *WriteToString<96>(Request.Key, '/', Request.Id), *Request.Name);
+				*Name, *WriteToString<96>(Request.Key, '/', Request.Id), *Request.Name);
 		}
 		else if (bHasValue && (ValueKey == Request.Key) && (ValueId == Request.Id) && (bExistsOnly || Reader.HasSource()))
 		{
@@ -932,17 +597,20 @@ void FMemoryCacheStore::GetChunks(
 			FReadScopeLock ScopeLock(SynchronizationObject);
 			if (Request.Id.IsValid())
 			{
-				if (const FCacheRecord* Record = CacheRecords.Find(Request.Key))
+				if (const FCacheRecordComponents* Components = CacheRecords.Find(Request.Key))
 				{
-					const FValueWithId& ValueWithId = Record->GetValue(Request.Id);
-					bHasValue = ValueWithId.IsValid();
-					Reader.ResetSource();
-					Value.Reset();
-					Value = ValueWithId;
-					ValueId = Request.Id;
-					ValueKey = Request.Key;
-					Reader.SetSource(Value.GetData());
-					bProcessHit = true;
+					if (const int32 ValueIndex = Algo::BinarySearchBy(Components->Values, Request.Id, &FValueWithId::GetId); ValueIndex != INDEX_NONE)
+					{
+						const FValueWithId& ValueWithId = Components->Values[ValueIndex];
+						bHasValue = ValueWithId.IsValid();
+						Reader.ResetSource();
+						Value.Reset();
+						Value = ValueWithId;
+						ValueId = Request.Id;
+						ValueKey = Request.Key;
+						Reader.SetSource(Value.GetData());
+						bProcessHit = true;
+					}
 				}
 			}
 			else
@@ -981,7 +649,149 @@ void FMemoryCacheStore::GetChunks(
 	}
 }
 
-FFileBackedDerivedDataBackend* CreateMemoryCacheStore(const TCHAR* Name, int64 MaxCacheSize, bool bCanBeDisabled)
+void FMemoryCacheStore::LegacyPut(
+	const TConstArrayView<FLegacyCachePutRequest> Requests,
+	IRequestOwner& Owner,
+	FOnLegacyCachePutComplete&& OnComplete)
+{
+	if (bDisabled)
+	{
+		return CompleteWithStatus(Requests, OnComplete, EStatus::Error);
+	}
+
+	for (const FLegacyCachePutRequest& Request : Requests)
+	{
+		EStatus Status = EStatus::Error;
+		ON_SCOPE_EXIT
+		{
+			OnComplete(Request.MakeResponse(Status));
+		};
+
+		const FCacheKey& Key = Request.Key.GetKey();
+		const FLegacyCacheValue& Value = Request.Value;
+
+		if (!Value.HasData())
+		{
+			continue;
+		}
+
+		if (ShouldSimulateMiss(Key))
+		{
+			UE_LOG(LogDerivedDataCache, Verbose, TEXT("%s: Simulated miss for put of %s from '%s'"),
+				*Name, *WriteToString<96>(Key), *Request.Name);
+			continue;
+		}
+
+		COOK_STAT(auto Timer = UsageStats.TimePut());
+		const int64 ValueSize = Value.GetRawSize();
+		const bool bReplaceExisting = !EnumHasAnyFlags(Request.Policy, ECachePolicy::QueryLocal);
+
+		FWriteScopeLock ScopeLock(SynchronizationObject);
+		FLegacyCacheValue* const ExistingValue = LegacyCacheValues.Find(Key);
+
+		if (ExistingValue && !bReplaceExisting)
+		{
+			Status = EStatus::Ok;
+			continue;
+		}
+
+		const int64 ExistingValueSize = ExistingValue ? ExistingValue->GetRawSize() : 0;
+		const int64 RequiredSize = ValueSize - ExistingValueSize;
+
+		if (MaxCacheSize > 0 && (CurrentCacheSize + RequiredSize) > MaxCacheSize)
+		{
+			UE_CLOG(!bMaxSizeExceeded, LogDerivedDataCache, Display,
+				TEXT("Failed to cache data. Maximum cache size reached. ")
+				TEXT("CurrentSize %" UINT64_FMT " KiB / MaxSize: %" UINT64_FMT " KiB"),
+				CurrentCacheSize / 1024, MaxCacheSize / 1024);
+			bMaxSizeExceeded = true;
+			continue;
+		}
+
+		CurrentCacheSize += RequiredSize;
+		if (ExistingValue)
+		{
+			*ExistingValue = Value;
+		}
+		else
+		{
+			LegacyCacheValues.Add(Key, Value);
+		}
+
+		COOK_STAT(Timer.AddHit(ValueSize));
+		Status = EStatus::Ok;
+	}
+}
+
+void FMemoryCacheStore::LegacyGet(
+	const TConstArrayView<FLegacyCacheGetRequest> Requests,
+	IRequestOwner& Owner,
+	FOnLegacyCacheGetComplete&& OnComplete)
+{
+	if (bDisabled)
+	{
+		return CompleteWithStatus(Requests, OnComplete, EStatus::Error);
+	}
+
+	for (const FLegacyCacheGetRequest& Request : Requests)
+	{
+		const FCacheKey& Key = Request.Key.GetKey();
+		const bool bExistsOnly = EnumHasAllFlags(Request.Policy, ECachePolicy::SkipData);
+		COOK_STAT(auto Timer = bExistsOnly ? UsageStats.TimeProbablyExists() : UsageStats.TimeGet());
+
+		FLegacyCacheValue Value;
+		EStatus Status = EStatus::Error;
+		if (ShouldSimulateMiss(Key))
+		{
+			UE_LOG(LogDerivedDataCache, Verbose, TEXT("%s: Simulated miss for get of %s from '%s'"),
+				*Name, *WriteToString<96>(Key), *Request.Name);
+		}
+		else if (bExistsOnly && bCanBeDisabled)
+		{
+		}
+		else if (FReadScopeLock ScopeLock(SynchronizationObject); const FLegacyCacheValue* CacheValue = LegacyCacheValues.Find(Key))
+		{
+			Status = EStatus::Ok;
+			if (!bExistsOnly)
+			{
+				Value = *CacheValue;
+			}
+			COOK_STAT(Timer.AddHit(Value.GetRawSize()));
+		}
+
+		OnComplete({Request.Name, Request.Key, MoveTemp(Value), Request.UserData, Status});
+	}
+}
+
+void FMemoryCacheStore::LegacyDelete(const FLegacyCacheKey& Key)
+{
+	FWriteScopeLock ScopeLock(SynchronizationObject);
+	FLegacyCacheValue Value;
+	if (LegacyCacheValues.RemoveAndCopyValue(Key.GetKey(), Value))
+	{
+		CurrentCacheSize -= Value.GetRawSize();
+		bMaxSizeExceeded = false;
+	}
+}
+
+void FMemoryCacheStore::LegacyDelete(
+	const TConstArrayView<FLegacyCacheDeleteRequest> Requests,
+	IRequestOwner& Owner,
+	FOnLegacyCacheDeleteComplete&& OnComplete)
+{
+	if (bDisabled)
+	{
+		return CompleteWithStatus(Requests, OnComplete, EStatus::Error);
+	}
+
+	for (const FLegacyCacheDeleteRequest& Request : Requests)
+	{
+		LegacyDelete(Request.Key);
+	}
+	CompleteWithStatus(Requests, OnComplete, EStatus::Ok);
+}
+
+IMemoryCacheStore* CreateMemoryCacheStore(const TCHAR* Name, int64 MaxCacheSize, bool bCanBeDisabled)
 {
 	return new FMemoryCacheStore(Name, MaxCacheSize, bCanBeDisabled);
 }
diff --git a/Engine/Source/Developer/DerivedDataCache/Private/MemoryCacheStore.h b/Engine/Source/Developer/DerivedDataCache/Private/MemoryCacheStore.h
new file mode 100644
--- /dev/null
+++ b/Engine/Source/Developer/DerivedDataCache/Private/MemoryCacheStore.h
@@ -0,0 +1,22 @@
+// Copyright Epic Games, Inc. All Rights Reserved.
+
+#pragma once
+
+#include "DerivedDataLegacyCacheStore.h"
+
+namespace UE::DerivedData
+{
+
+class IMemoryCacheStore : public ILegacyCacheStore
+{
+public:
+	using ILegacyCacheStore::LegacyDelete;
+
+	virtual void Delete(const FCacheKey& Key) = 0;
+	virtual void DeleteValue(const FCacheKey& Key) = 0;
+	virtual void LegacyDelete(const FLegacyCacheKey& Key) = 0;
+
+	virtual void Disable() = 0;
+};
+
+} // UE::DerivedData
diff --git a/Engine/Source/Developer/DerivedDataCache/Private/PakFileCacheStore.cpp b/Engine/Source/Developer/DerivedDataCache/Private/PakFileCacheStore.cpp
--- a/Engine/Source/Developer/DerivedDataCache/Private/PakFileCacheStore.cpp
+++ b/Engine/Source/Developer/DerivedDataCache/Private/PakFileCacheStore.cpp
@@ -659,8 +659,11 @@ bool IPakFileCacheStore::SortAndCopy(const FString &InputFilename, const FString
 	for (int KeyIndex = 0; KeyIndex < KeyNames.Num(); KeyIndex++)
 	{
 		Buffer.Reset();
-		InputPak.GetCachedData(*KeyNames[KeyIndex], Buffer);
-		OutputPak.PutCachedData(*KeyNames[KeyIndex], Buffer, false);
+		// Data over 2 GiB is not copied.
+		if (InputPak.GetCachedData(*KeyNames[KeyIndex], Buffer))
+		{
+			OutputPak.PutCachedData(*KeyNames[KeyIndex], Buffer, false);
+		}
 		KeySizes.Add(Buffer.Num());
 	}
 
@@ -1443,7 +1446,7 @@ bool FPakFileCacheStore::SaveFile(
 		if (const int64 EndOffset = FileHandle->Tell(); EndOffset >= Offset && !Ar.IsError())
 		{
 			FCacheValue& Item = CacheItems.Emplace(Path, FCacheValue(Offset, EndOffset - Offset, HashAr.GetHash()));
-			UE_LOG(LogDerivedDataCache, Log,
+			UE_LOG(LogDerivedDataCache, VeryVerbose,
 				TEXT("%s: File %.*s from '%.*s' written with offset %" INT64_FMT ", size %" INT64_FMT", CRC 0x%08x."),
 				*CachePath, Path.Len(), Path.GetData(), DebugName.Len(), DebugName.GetData(), Item.Offset, Item.Size, Item.Crc);
 			return true;
