DDC: Destroy the cache store graph when shutting down the cache

diff --git a/Engine/Source/Developer/DerivedDataCache/Private/DerivedDataBackendInterface.h b/Engine/Source/Developer/DerivedDataCache/Private/DerivedDataBackendInterface.h
--- a/Engine/Source/Developer/DerivedDataCache/Private/DerivedDataBackendInterface.h
+++ b/Engine/Source/Developer/DerivedDataCache/Private/DerivedDataBackendInterface.h
@@ -235,6 +235,8 @@ public:
 class FDerivedDataBackend
 {
 public:
+	static FDerivedDataBackend* Create();
+
 	/**
 	 * Singleton to retrieve the GLOBAL backend
 	 *
@@ -242,6 +244,8 @@ public:
 	 */
 	static FDerivedDataBackend& Get();
 
+	virtual ~FDerivedDataBackend() = default;
+
 	/**
 	 * Singleton to retrieve the root cache
 	 * @return Reference to the global cache root
diff --git a/Engine/Source/Developer/DerivedDataCache/Private/DerivedDataBackends.cpp b/Engine/Source/Developer/DerivedDataCache/Private/DerivedDataBackends.cpp
--- a/Engine/Source/Developer/DerivedDataCache/Private/DerivedDataBackends.cpp
+++ b/Engine/Source/Developer/DerivedDataCache/Private/DerivedDataBackends.cpp
@@ -69,7 +69,7 @@ ILegacyCacheStore* CreateZenCacheStore(const TCHAR* NodeName, const TCHAR* Servi
  * This class is used to create a singleton that represents the derived data cache hierarchy and all of the wrappers necessary
  * ideally this would be data driven and the backends would be plugins...
  */
-class FDerivedDataBackendGraph : public FDerivedDataBackend
+class FDerivedDataBackendGraph final : public FDerivedDataBackend
 {
 public:
 	using FParsedNode = TPair<ILegacyCacheStore*, ECacheStoreFlags>;
@@ -96,6 +96,9 @@ public:
 			*LOCTEXT("CommandText_DDCUnmountPak", "Unmounts read-only pak file").ToString(),
 			FConsoleCommandWithArgsDelegate::CreateRaw(this, &FDerivedDataBackendGraph::UnmountPakCommandHandler))
 	{
+		check(!StaticGraph);
+		StaticGraph = this;
+
 		check(IsInGameThread()); // we pretty much need this to be initialized from the main thread...it uses GConfig, etc
 		check(GConfig && GConfig->IsReadyForUse());
 		RootCache = nullptr;
@@ -1009,8 +1012,10 @@ public:
 
 	virtual ~FDerivedDataBackendGraph()
 	{
+		check(StaticGraph == this);
 		RootCache = nullptr;
 		DestroyCreatedBackends();
+		StaticGraph = nullptr;
 	}
 
 	ILegacyCacheStore& GetRoot() override
@@ -1145,8 +1150,8 @@ public:
 
 	static FORCEINLINE FDerivedDataBackendGraph& Get()
 	{
-		static FDerivedDataBackendGraph SingletonInstance;
-		return SingletonInstance;
+		check(StaticGraph);
+		return *StaticGraph;
 	}
 
 	virtual FDerivedDataBackendInterface* MountPakFile(const TCHAR* PakFilename) override
@@ -1237,6 +1242,8 @@ private:
 		MountPakFile(*Args[0]);
 	}
 
+	static inline FDerivedDataBackendGraph*			StaticGraph;
+
 	FThreadSafeCounter								AsyncCompletionCounter;
 	FString											GraphName;
 	FString											ReadPakFilename;
@@ -1458,6 +1465,11 @@ bool FDerivedDataBackendInterface::LegacyDebugOptions(FBackendDebugOptions& Opti
 	return ApplyDebugOptions(Options);
 }
 
+FDerivedDataBackend* FDerivedDataBackend::Create()
+{
+	return new FDerivedDataBackendGraph();
+}
+
 FDerivedDataBackend& FDerivedDataBackend::Get()
 {
 	return FDerivedDataBackendGraph::Get();
diff --git a/Engine/Source/Developer/DerivedDataCache/Private/DerivedDataCache.cpp b/Engine/Source/Developer/DerivedDataCache/Private/DerivedDataCache.cpp
--- a/Engine/Source/Developer/DerivedDataCache/Private/DerivedDataCache.cpp
+++ b/Engine/Source/Developer/DerivedDataCache/Private/DerivedDataCache.cpp
@@ -223,10 +223,11 @@ class FDerivedDataCache final
 		 * @param	InDataDeriver	plugin to produce cache key and in the event of a miss, return the data.
 		 * @param	InCacheKey		Complete cache key for this data.
 		**/
-		FBuildAsyncWorker(FDerivedDataPluginInterface* InDataDeriver, const TCHAR* InCacheKey, FStringView InDebugContext, bool bInSynchronousForStats)
+		FBuildAsyncWorker(FDerivedDataBackend* InBackend, FDerivedDataPluginInterface* InDataDeriver, const TCHAR* InCacheKey, FStringView InDebugContext, bool bInSynchronousForStats)
 		: bSuccess(false)
 		, bSynchronousForStats(bInSynchronousForStats)
 		, bDataWasBuilt(false)
+		, Backend(InBackend)
 		, DataDeriver(InDataDeriver)
 		, CacheKey(InCacheKey)
 		, DebugContext(InDebugContext)
@@ -267,9 +268,9 @@ class FDerivedDataCache final
 					SCOPE_SECONDS_COUNTER(ThisTime);
 					FLegacyCacheGetRequest LegacyRequest;
 					LegacyRequest.Name = DebugContext;
-					LegacyRequest.Key = FLegacyCacheKey(CacheKey, FDerivedDataBackend::Get().GetMaxKeyLength());
+					LegacyRequest.Key = FLegacyCacheKey(CacheKey, Backend->GetMaxKeyLength());
 					FRequestOwner BlockingOwner(EPriority::Blocking);
-					FDerivedDataBackend::Get().GetRoot().LegacyGet({LegacyRequest}, BlockingOwner,
+					Backend->GetRoot().LegacyGet({LegacyRequest}, BlockingOwner,
 						[this, &bGetResult](FLegacyCacheGetResponse&& Response)
 						{
 							const uint64 RawSize = Response.Value.GetRawSize();
@@ -359,10 +360,10 @@ class FDerivedDataCache final
 						SCOPE_SECONDS_COUNTER(ThisTime);
 						FLegacyCachePutRequest LegacyRequest;
 						LegacyRequest.Name = DebugContext;
-						LegacyRequest.Key = FLegacyCacheKey(CacheKey, FDerivedDataBackend::Get().GetMaxKeyLength());
+						LegacyRequest.Key = FLegacyCacheKey(CacheKey, Backend->GetMaxKeyLength());
 						LegacyRequest.Value = FLegacyCacheValue(FCompositeBuffer(FSharedBuffer::Clone(MakeMemoryView(Data))));
 						FRequestOwner AsyncOwner(EPriority::Normal);
-						FDerivedDataBackend::Get().GetRoot().LegacyPut({LegacyRequest}, AsyncOwner, [](auto&&){});
+						Backend->GetRoot().LegacyPut({LegacyRequest}, AsyncOwner, [](auto&&){});
 						AsyncOwner.KeepAlive();
 					}
 					INC_FLOAT_STAT_BY(STAT_DDC_PutTime, bSynchronousForStats ? (float)ThisTime : 0.0f);
@@ -372,7 +373,7 @@ class FDerivedDataCache final
 			{
 				Data.Empty();
 			}
-			FDerivedDataBackend::Get().AddToAsyncCompletionCounter(-1);
+			Backend->AddToAsyncCompletionCounter(-1);
 
 			// Record that the task is finished and check that it was running and not finished or destroyed previously.
 			{
@@ -395,6 +396,8 @@ class FDerivedDataCache final
 		bool							bSynchronousForStats;
 		/** true if we had to build the data */
 		bool							bDataWasBuilt;
+		/** Backend graph to execute against. */
+		FDerivedDataBackend*			Backend;
 		/** Data dervier we are operating on **/
 		FDerivedDataPluginInterface*	DataDeriver;
 		/** Cache key associated with this build **/
@@ -418,7 +421,7 @@ public:
 			verify(GCacheThreadPool->Create(ThreadCount, 96 * 1024, TPri_AboveNormal, TEXT("DDC IO ThreadPool")));
 		}
 
-		FDerivedDataBackend::Get(); // we need to make sure this starts before we allow us to start
+		Backend = FDerivedDataBackend::Create();
 
 		CacheStoreMaintainers = IModularFeatures::Get().GetModularFeatureImplementations<ICacheStoreMaintainer>(FeatureName);
 
@@ -438,6 +441,7 @@ public:
 			delete It.Value();
 		}
 		PendingTasks.Empty();
+		delete Backend;
 	}
 
 	virtual bool GetSynchronous(FDerivedDataPluginInterface* DataDeriver, TArray<uint8>& OutData, bool* bDataWasBuilt) override
@@ -446,7 +450,7 @@ public:
 		check(DataDeriver);
 		FString CacheKey = FDerivedDataCache::BuildCacheKey(DataDeriver);
 		UE_LOG(LogDerivedDataCache, VeryVerbose, TEXT("GetSynchronous %s from '%s'"), *CacheKey, *DataDeriver->GetDebugContextString());
-		FAsyncTask<FBuildAsyncWorker> PendingTask(DataDeriver, *CacheKey, DataDeriver->GetDebugContextString(), true);
+		FAsyncTask<FBuildAsyncWorker> PendingTask(Backend, DataDeriver, *CacheKey, DataDeriver->GetDebugContextString(), true);
 		AddToAsyncCompletionCounter(1);
 		PendingTask.StartSynchronousTask();
 		OutData = TArray<uint8>(MoveTemp(PendingTask.GetTask().Data));
@@ -465,7 +469,7 @@ public:
 		FString CacheKey = FDerivedDataCache::BuildCacheKey(DataDeriver);
 		UE_LOG(LogDerivedDataCache, VeryVerbose, TEXT("GetAsynchronous %s from '%s', Handle %d"), *CacheKey, *DataDeriver->GetDebugContextString(), Handle);
 		const bool bSync = !DataDeriver->IsBuildThreadsafe();
-		FAsyncTask<FBuildAsyncWorker>* AsyncTask = new FAsyncTask<FBuildAsyncWorker>(DataDeriver, *CacheKey, DataDeriver->GetDebugContextString(), bSync);
+		FAsyncTask<FBuildAsyncWorker>* AsyncTask = new FAsyncTask<FBuildAsyncWorker>(Backend, DataDeriver, *CacheKey, DataDeriver->GetDebugContextString(), bSync);
 		check(!PendingTasks.Contains(Handle));
 		PendingTasks.Add(Handle,AsyncTask);
 		AddToAsyncCompletionCounter(1);
@@ -557,7 +561,7 @@ public:
 		DDC_SCOPE_CYCLE_COUNTER(DDC_GetSynchronous_Data);
 		UE_LOG(LogDerivedDataCache, VeryVerbose, TEXT("GetSynchronous %s from '%.*s'"), CacheKey, DebugContext.Len(), DebugContext.GetData());
 		ValidateCacheKey(CacheKey);
-		FAsyncTask<FBuildAsyncWorker> PendingTask((FDerivedDataPluginInterface*)NULL, CacheKey, DebugContext, true);
+		FAsyncTask<FBuildAsyncWorker> PendingTask(Backend, nullptr, CacheKey, DebugContext, true);
 		AddToAsyncCompletionCounter(1);
 		PendingTask.StartSynchronousTask();
 		OutData = DataType(MoveTemp(PendingTask.GetTask().Data));
@@ -581,7 +585,7 @@ public:
 		const uint32 Handle = NextHandle();
 		UE_LOG(LogDerivedDataCache, VeryVerbose, TEXT("GetAsynchronous %s from '%.*s', Handle %d"), CacheKey, DebugContext.Len(), DebugContext.GetData(), Handle);
 		ValidateCacheKey(CacheKey);
-		FAsyncTask<FBuildAsyncWorker>* AsyncTask = new FAsyncTask<FBuildAsyncWorker>((FDerivedDataPluginInterface*)NULL, CacheKey, DebugContext, false);
+		FAsyncTask<FBuildAsyncWorker>* AsyncTask = new FAsyncTask<FBuildAsyncWorker>(Backend, nullptr, CacheKey, DebugContext, false);
 		check(!PendingTasks.Contains(Handle));
 		PendingTasks.Add(Handle, AsyncTask);
 		AddToAsyncCompletionCounter(1);
@@ -600,10 +604,10 @@ public:
 			SCOPE_SECONDS_COUNTER(ThisTime);
 			FLegacyCachePutRequest LegacyRequest;
 			LegacyRequest.Name = DebugContext;
-			LegacyRequest.Key = FLegacyCacheKey(CacheKey, FDerivedDataBackend::Get().GetMaxKeyLength());
+			LegacyRequest.Key = FLegacyCacheKey(CacheKey, Backend->GetMaxKeyLength());
 			LegacyRequest.Value = FLegacyCacheValue(FCompositeBuffer(FSharedBuffer::Clone(MakeMemoryView(Data))));
 			FRequestOwner AsyncOwner(EPriority::Normal);
-			FDerivedDataBackend::Get().GetRoot().LegacyPut({LegacyRequest}, AsyncOwner, [](auto&&){});
+			Backend->GetRoot().LegacyPut({LegacyRequest}, AsyncOwner, [](auto&&){});
 			AsyncOwner.KeepAlive();
 		}
 		INC_FLOAT_STAT_BY(STAT_DDC_PutTime,(float)ThisTime);
@@ -614,11 +618,11 @@ public:
 	{
 		ValidateCacheKey(CacheKey);
 		FLegacyCacheDeleteRequest LegacyRequest;
-		LegacyRequest.Key = FLegacyCacheKey(CacheKey, FDerivedDataBackend::Get().GetMaxKeyLength());
+		LegacyRequest.Key = FLegacyCacheKey(CacheKey, Backend->GetMaxKeyLength());
 		LegacyRequest.Name = LegacyRequest.Key.GetFullKey();
 		LegacyRequest.bTransient = true;
 		FRequestOwner BlockingOwner(EPriority::Blocking);
-		FDerivedDataBackend::Get().GetRoot().LegacyDelete({LegacyRequest}, BlockingOwner, [](auto&&){});
+		Backend->GetRoot().LegacyDelete({LegacyRequest}, BlockingOwner, [](auto&&){});
 		BlockingOwner.Wait();
 	}
 
@@ -632,11 +636,11 @@ public:
 		{
 			SCOPE_SECONDS_COUNTER(ThisTime);
 			FLegacyCacheGetRequest LegacyRequest;
-			LegacyRequest.Key = FLegacyCacheKey(CacheKey, FDerivedDataBackend::Get().GetMaxKeyLength());
+			LegacyRequest.Key = FLegacyCacheKey(CacheKey, Backend->GetMaxKeyLength());
 			LegacyRequest.Name = LegacyRequest.Key.GetFullKey();
 			LegacyRequest.Policy = ECachePolicy::Query | ECachePolicy::SkipData;
 			FRequestOwner BlockingOwner(EPriority::Blocking);
-			FDerivedDataBackend::Get().GetRoot().LegacyGet({LegacyRequest}, BlockingOwner,
+			Backend->GetRoot().LegacyGet({LegacyRequest}, BlockingOwner,
 				[&bResult](FLegacyCacheGetResponse&& Response) { bResult = Response.Status == EStatus::Ok; });
 			BlockingOwner.Wait();
 		}
@@ -659,14 +663,14 @@ public:
 				for (const FString& CacheKey : CacheKeys)
 				{
 					FLegacyCacheGetRequest& LegacyRequest = LegacyRequests.AddDefaulted_GetRef();
-					LegacyRequest.Key = FLegacyCacheKey(CacheKey, FDerivedDataBackend::Get().GetMaxKeyLength());
+					LegacyRequest.Key = FLegacyCacheKey(CacheKey, Backend->GetMaxKeyLength());
 					LegacyRequest.Name = LegacyRequest.Key.GetFullKey();
 					LegacyRequest.Policy = ECachePolicy::Query | ECachePolicy::SkipData;
 					LegacyRequest.UserData = uint64(Index);
 					++Index;
 				}
 				FRequestOwner BlockingOwner(EPriority::Blocking);
-				FDerivedDataBackend::Get().GetRoot().LegacyGet(LegacyRequests, BlockingOwner,
+				Backend->GetRoot().LegacyGet(LegacyRequests, BlockingOwner,
 					[&Result](FLegacyCacheGetResponse&& Response)
 					{
 						Result[int32(Response.UserData)] = Response.Status == EStatus::Ok;
@@ -697,14 +701,14 @@ public:
 			{
 				FLegacyCacheGetRequest& LegacyRequest = LegacyRequests.AddDefaulted_GetRef();
 				LegacyRequest.Name = Name;
-				LegacyRequest.Key = FLegacyCacheKey(CacheKey, FDerivedDataBackend::Get().GetMaxKeyLength());
+				LegacyRequest.Key = FLegacyCacheKey(CacheKey, Backend->GetMaxKeyLength());
 				LegacyRequest.Policy = ECachePolicy::Default | ECachePolicy::SkipData;
 				LegacyRequest.UserData = uint64(Index);
 				++Index;
 			}
 			bool bOk = true;
 			FRequestOwner BlockingOwner(EPriority::Blocking);
-			FDerivedDataBackend::Get().GetRoot().LegacyGet(LegacyRequests, BlockingOwner,
+			Backend->GetRoot().LegacyGet(LegacyRequests, BlockingOwner,
 				[&bOk](FLegacyCacheGetResponse&& Response)
 				{
 					bOk &= Response.Status == EStatus::Ok;
@@ -718,44 +722,44 @@ public:
 	void NotifyBootComplete() override
 	{
 		DDC_SCOPE_CYCLE_COUNTER(DDC_NotifyBootComplete);
-		FDerivedDataBackend::Get().NotifyBootComplete();
+		Backend->NotifyBootComplete();
 	}
 
 	void AddToAsyncCompletionCounter(int32 Addend) override
 	{
-		FDerivedDataBackend::Get().AddToAsyncCompletionCounter(Addend);
+		Backend->AddToAsyncCompletionCounter(Addend);
 	}
 
 	bool AnyAsyncRequestsRemaining() const override
 	{
-		return FDerivedDataBackend::Get().AnyAsyncRequestsRemaining();
+		return Backend->AnyAsyncRequestsRemaining();
 	}
 
 	void WaitForQuiescence(bool bShutdown) override
 	{
 		DDC_SCOPE_CYCLE_COUNTER(DDC_WaitForQuiescence);
-		FDerivedDataBackend::Get().WaitForQuiescence(bShutdown);
+		Backend->WaitForQuiescence(bShutdown);
 	}
 
 	/** Get whether a Shared Data Cache is in use */
 	virtual bool GetUsingSharedDDC() const override
 	{		
-		return FDerivedDataBackend::Get().GetUsingSharedDDC();
+		return Backend->GetUsingSharedDDC();
 	}
 
 	virtual const TCHAR* GetGraphName() const override
 	{
-		return FDerivedDataBackend::Get().GetGraphName();
+		return Backend->GetGraphName();
 	}
 
 	virtual const TCHAR* GetDefaultGraphName() const override
 	{
-		return FDerivedDataBackend::Get().GetDefaultGraphName();
+		return Backend->GetDefaultGraphName();
 	}
 
 	void GetDirectories(TArray<FString>& OutResults) override
 	{
-		FDerivedDataBackend::Get().GetDirectories(OutResults);
+		Backend->GetDirectories(OutResults);
 	}
 
 	PRAGMA_DISABLE_DEPRECATION_WARNINGS
@@ -786,7 +790,7 @@ public:
 	PRAGMA_DISABLE_DEPRECATION_WARNINGS
 	virtual TSharedRef<FDerivedDataCacheStatsNode> GatherUsageStats() const override
 	{
-		return FDerivedDataBackend::Get().GatherUsageStats();
+		return Backend->GatherUsageStats();
 	}
 	PRAGMA_ENABLE_DEPRECATION_WARNINGS
 
@@ -832,6 +836,7 @@ private:
 			TEXT("Invalid characters in cache key %s. Use SanitizeCacheKey or BuildCacheKey to create valid keys."), CacheKey);
 	}
 
+	FDerivedDataBackend*		Backend;
 	/** Counter used to produce unique handles **/
 	FThreadSafeCounter			CurrentHandle;
 	/** Object used for synchronization via a scoped lock **/
@@ -850,7 +855,7 @@ public:
 		IRequestOwner& Owner,
 		FOnCachePutComplete&& OnComplete) final
 	{
-		return FDerivedDataBackend::Get().GetRoot().Put(Requests, Owner, OnComplete ? MoveTemp(OnComplete) : [](auto&&){});
+		return Backend->GetRoot().Put(Requests, Owner, OnComplete ? MoveTemp(OnComplete) : [](auto&&){});
 	}
 
 	void Get(
@@ -858,7 +863,7 @@ public:
 		IRequestOwner& Owner,
 		FOnCacheGetComplete&& OnComplete) final
 	{
-		return FDerivedDataBackend::Get().GetRoot().Get(Requests, Owner, OnComplete ? MoveTemp(OnComplete) : [](auto&&){});
+		return Backend->GetRoot().Get(Requests, Owner, OnComplete ? MoveTemp(OnComplete) : [](auto&&){});
 	}
 
 	void PutValue(
@@ -866,7 +871,7 @@ public:
 		IRequestOwner& Owner,
 		FOnCachePutValueComplete&& OnComplete) final
 	{
-		return FDerivedDataBackend::Get().GetRoot().PutValue(Requests, Owner, OnComplete ? MoveTemp(OnComplete) : [](auto&&){});
+		return Backend->GetRoot().PutValue(Requests, Owner, OnComplete ? MoveTemp(OnComplete) : [](auto&&){});
 	}
 
 	void GetValue(
@@ -874,7 +879,7 @@ public:
 		IRequestOwner& Owner,
 		FOnCacheGetValueComplete&& OnComplete) final
 	{
-		return FDerivedDataBackend::Get().GetRoot().GetValue(Requests, Owner, OnComplete ? MoveTemp(OnComplete) : [](auto&&){});
+		return Backend->GetRoot().GetValue(Requests, Owner, OnComplete ? MoveTemp(OnComplete) : [](auto&&){});
 	}
 
 	void GetChunks(
@@ -882,7 +887,7 @@ public:
 		IRequestOwner& Owner,
 		FOnCacheGetChunkComplete&& OnComplete) final
 	{
-		return FDerivedDataBackend::Get().GetRoot().GetChunks(Requests, Owner, OnComplete ? MoveTemp(OnComplete) : [](auto&&){});
+		return Backend->GetRoot().GetChunks(Requests, Owner, OnComplete ? MoveTemp(OnComplete) : [](auto&&){});
 	}
 
 	ICacheStoreMaintainer& GetMaintainer() final
