DDC: Fixed -DDC-All-MissRate to work with callers that expect consistent hits on the same key

Fixes the accumulation of stats in the async wrapper as well.

#rnx

diff --git a/Engine/Source/Developer/DerivedDataCache/Private/DerivedDataBackendInterface.h b/Engine/Source/Developer/DerivedDataCache/Private/DerivedDataBackendInterface.h
--- a/Engine/Source/Developer/DerivedDataCache/Private/DerivedDataBackendInterface.h
+++ b/Engine/Source/Developer/DerivedDataCache/Private/DerivedDataBackendInterface.h
@@ -10,11 +10,14 @@
 #include "DerivedDataLegacyCacheStore.h"
 #include "HAL/LowLevelMemTracker.h"
 #include "Stats/Stats.h"
+#include "Templates/DontCopy.h"
+#include "Templates/PimplPtr.h"
 
 class FDerivedDataCacheUsageStats;
 class FDerivedDataCacheStatsNode;
 
 namespace UE::DerivedData { struct FCacheKey; }
+namespace UE::DerivedData::Private { struct FBackendDebugMissState; }
 
 DECLARE_DWORD_ACCUMULATOR_STAT_EXTERN(TEXT("Num Gets"),STAT_DDC_NumGets,STATGROUP_DDC, );
 DECLARE_DWORD_ACCUMULATOR_STAT_EXTERN(TEXT("Num Puts"),STAT_DDC_NumPuts,STATGROUP_DDC, );
@@ -69,22 +72,27 @@ struct FBackendDebugOptions
 	/** Types of DDC entries that should always be a miss */
 	TArray<FString>		SimulateMissTypes;
 
-	FBackendDebugOptions()
-		: RandomMissRate(0)
-		, SpeedClass(EBackendSpeedClass::Unknown)
-	{
-	}
+	/** State for simulated misses. */
+	TDontCopy<TPimplPtr<Private::FBackendDebugMissState>> SimulateMissState;
+
+	FBackendDebugOptions();
 
 	/** Fill in the provided structure based on the name of the node (e.g. 'shared') and the provided token stream */
 	static bool ParseFromTokens(FBackendDebugOptions& OutOptions, const TCHAR* InNodeName, const TCHAR* InTokens);
 
 	/**
-		* Returns true if, according to the properties of this struct, the provided key should be treated as a miss.
-		* Implementing that miss and accounting for any behavior impact (e.g. skipping a subsequent put) is left to
-		* each backend.
-		*/
-	bool ShouldSimulateMiss(const TCHAR* InCacheKey);
-	bool ShouldSimulateMiss(const UE::DerivedData::FCacheKey& InCacheKey);
+	 * Returns true if, according to the properties of this struct, the provided key should be treated as a miss.
+	 * Implementing that miss and accounting for any behavior impact (e.g. skipping a subsequent put) is left to
+	 * each backend.
+	 */
+	bool ShouldSimulateMiss(const TCHAR* CacheKey) { return ShouldSimulateGetMiss(CacheKey); }
+	bool ShouldSimulateMiss(const FCacheKey& CacheKey) { return ShouldSimulateGetMiss(CacheKey); }
+
+	bool ShouldSimulatePutMiss(const FCacheKey& Key);
+	bool ShouldSimulateGetMiss(const FCacheKey& Key);
+
+	bool ShouldSimulatePutMiss(const TCHAR* LegacyKey);
+	bool ShouldSimulateGetMiss(const TCHAR* LegacyKey);
 };
 
 /**
diff --git a/Engine/Source/Developer/DerivedDataCache/Private/DerivedDataBackends.cpp b/Engine/Source/Developer/DerivedDataCache/Private/DerivedDataBackends.cpp
--- a/Engine/Source/Developer/DerivedDataCache/Private/DerivedDataBackends.cpp
+++ b/Engine/Source/Developer/DerivedDataCache/Private/DerivedDataBackends.cpp
@@ -7,6 +7,7 @@
 #include "DerivedDataCacheStore.h"
 #include "DerivedDataCacheUsageStats.h"
 #include "DerivedDataRequestOwner.h"
+#include "HAL/CriticalSection.h"
 #include "HAL/FileManager.h"
 #include "HAL/IConsoleManager.h"
 #include "HAL/PlatformFileManager.h"
@@ -21,6 +22,7 @@
 #include "Misc/EngineBuildSettings.h"
 #include "Misc/Guid.h"
 #include "Misc/Paths.h"
+#include "Misc/ScopeLock.h"
 #include "Misc/StringBuilder.h"
 #include "Modules/ModuleManager.h"
 #include "PakFileCacheStore.h"
@@ -1475,10 +1477,33 @@ FDerivedDataBackend& FDerivedDataBackend::Get()
 	return FDerivedDataBackendGraph::Get();
 }
 
+///////////////////////////////////////////////////////////////////////////////////////////////////////////////////////
+
+enum class EBackendDebugKeyState
+{
+	None,
+	HitGet,
+	MissGet,
+};
+
+struct Private::FBackendDebugMissState
+{
+	FCriticalSection Lock;
+	TMap<FCacheKey, EBackendDebugKeyState> Keys;
+	TMap<FName, EBackendDebugKeyState> LegacyKeys;
+};
+
+FBackendDebugOptions::FBackendDebugOptions()
+	: RandomMissRate(0)
+	, SpeedClass(EBackendSpeedClass::Unknown)
+{
+	SimulateMissState.Get() = MakePimpl<Private::FBackendDebugMissState>();
+}
+
 /**
  * Parse debug options for the provided node name. Returns true if any options were specified
  */
-bool FDerivedDataBackendInterface::FBackendDebugOptions::ParseFromTokens(FDerivedDataBackendInterface::FBackendDebugOptions& OutOptions, const TCHAR* InNodeName, const TCHAR* InInputTokens)
+bool FBackendDebugOptions::ParseFromTokens(FDerivedDataBackendInterface::FBackendDebugOptions& OutOptions, const TCHAR* InNodeName, const TCHAR* InInputTokens)
 {
 	// check if the input stream has any ddc options for this node
 	FString PrefixKey = FString(TEXT("-ddc-")) + InNodeName;
@@ -1523,40 +1548,96 @@ bool FDerivedDataBackendInterface::FBackendDebugOptions::ParseFromTokens(FDerive
 	return true;
 }
 
-/* Convenience function for backends that check if the key should be missed or not */
-bool FDerivedDataBackendInterface::FBackendDebugOptions::ShouldSimulateMiss(const TCHAR* InCacheKey)
+bool FBackendDebugOptions::ShouldSimulatePutMiss(const TCHAR* LegacyKey)
 {
-	if (!SimulateMissTypes.IsEmpty())
+	if (RandomMissRate == 0 && SimulateMissTypes.IsEmpty())
 	{
-		FStringView TypeStr(InCacheKey);
-		TypeStr.LeftInline(TypeStr.Find(TEXT("_")));
-		if (SimulateMissTypes.Contains(TypeStr))
-		{
-			return true;
-		}
+		return false;
 	}
 
-	return RandomMissRate > 0 && FMath::RandHelper(100) < RandomMissRate;
+	const FName Key(LegacyKey);
+
+	Private::FBackendDebugMissState& State = *SimulateMissState.Get();
+	const uint32 KeyHash = GetTypeHash(Key);
+
+	FScopeLock Lock(&State.Lock);
+	State.LegacyKeys.AddByHash(KeyHash, Key, EBackendDebugKeyState::HitGet);
+	return false;
 }
 
-/* Convenience function for backends that check if the key should be missed or not */
-bool FDerivedDataBackendInterface::FBackendDebugOptions::ShouldSimulateMiss(const UE::DerivedData::FCacheKey& InCacheKey)
+bool FBackendDebugOptions::ShouldSimulateGetMiss(const TCHAR* LegacyKey)
 {
-	if (!SimulateMissTypes.IsEmpty())
+	if (RandomMissRate == 0 && SimulateMissTypes.IsEmpty())
 	{
-		TStringBuilder<256> Type;
-		Type << InCacheKey.Bucket;
-		if (Type.ToView().StartsWith(TEXTVIEW("Legacy")))
-		{
-			Type.RemoveAt(0, TEXTVIEW("Legacy").Len());
-		}
-		if (SimulateMissTypes.Contains(Type.ToView()))
+		return false;
+	}
+
+	const FStringView KeyView(LegacyKey);
+	const FName Key(KeyView);
+
+	bool bMiss = (RandomMissRate >= 100);
+	if (!bMiss && !SimulateMissTypes.IsEmpty())
+	{
+		const FStringView Bucket = KeyView.Left(KeyView.Find(TEXT("_")));
+		bMiss = SimulateMissTypes.Contains(Bucket);
+	}
+	if (!bMiss && RandomMissRate > 0)
+	{
+		bMiss = FMath::RandHelper(100) < RandomMissRate;
+	}
+
+	Private::FBackendDebugMissState& State = *SimulateMissState.Get();
+	const uint32 KeyHash = GetTypeHash(Key);
+
+	FScopeLock Lock(&State.Lock);
+	const EBackendDebugKeyState KeyState = bMiss ? EBackendDebugKeyState::MissGet : EBackendDebugKeyState::HitGet;
+	return State.LegacyKeys.FindOrAddByHash(KeyHash, Key, KeyState) == EBackendDebugKeyState::MissGet;
+}
+
+bool FBackendDebugOptions::ShouldSimulatePutMiss(const FCacheKey& Key)
+{
+	if (RandomMissRate == 0 && SimulateMissTypes.IsEmpty())
+	{
+		return false;
+	}
+
+	Private::FBackendDebugMissState& State = *SimulateMissState.Get();
+	const uint32 KeyHash = GetTypeHash(Key);
+
+	FScopeLock Lock(&State.Lock);
+	State.Keys.AddByHash(KeyHash, Key, EBackendDebugKeyState::HitGet);
+	return false;
+}
+
+bool FBackendDebugOptions::ShouldSimulateGetMiss(const FCacheKey& Key)
+{
+	if (RandomMissRate == 0 && SimulateMissTypes.IsEmpty())
+	{
+		return false;
+	}
+
+	bool bMiss = (RandomMissRate >= 100);
+	if (!bMiss && !SimulateMissTypes.IsEmpty())
+	{
+		TStringBuilder<256> Bucket;
+		Bucket << Key.Bucket;
+		if (Bucket.ToView().StartsWith(TEXTVIEW("Legacy")))
 		{
-			return true;
+			Bucket.RemoveAt(0, TEXTVIEW("Legacy").Len());
 		}
+		bMiss = SimulateMissTypes.Contains(Bucket.ToView());
 	}
+	if (!bMiss && RandomMissRate > 0)
+	{
+		bMiss = FMath::RandHelper(100) < RandomMissRate;
+	}
+
+	Private::FBackendDebugMissState& State = *SimulateMissState.Get();
+	const uint32 KeyHash = GetTypeHash(Key);
 
-	return RandomMissRate > 0 && FMath::RandHelper(100) < RandomMissRate;
+	FScopeLock Lock(&State.Lock);
+	const EBackendDebugKeyState KeyState = bMiss ? EBackendDebugKeyState::MissGet : EBackendDebugKeyState::HitGet;
+	return State.Keys.FindOrAddByHash(KeyHash, Key, KeyState) == EBackendDebugKeyState::MissGet;
 }
 
 } // UE::DerivedData
diff --git a/Engine/Source/Developer/DerivedDataCache/Private/DerivedDataCacheStoreAsync.cpp b/Engine/Source/Developer/DerivedDataCache/Private/DerivedDataCacheStoreAsync.cpp
--- a/Engine/Source/Developer/DerivedDataCache/Private/DerivedDataCacheStoreAsync.cpp
+++ b/Engine/Source/Developer/DerivedDataCache/Private/DerivedDataCacheStoreAsync.cpp
@@ -306,9 +306,9 @@ void FCacheStoreAsync::Execute(
 	{
 		Invoke(OnExecute, *InnerCache, Requests, Owner, [this, COOK_STAT(OnAddStats,) OnComplete = MoveTemp(OnComplete)](auto&& Response)
 		{
-			if (Response.Status == EStatus::Ok)
+			if (COOK_STAT(auto Timer = Invoke(OnAddStats, UsageStats)); Response.Status == EStatus::Ok)
 			{
-				COOK_STAT(Invoke(OnAddStats, UsageStats).AddHit(0));
+				COOK_STAT(Timer.AddHit(0));
 			}
 			OnComplete(MoveTemp(Response));
 			FDerivedDataBackend::Get().AddToAsyncCompletionCounter(-1);
diff --git a/Engine/Source/Developer/DerivedDataCache/Private/FileSystemCacheStore.cpp b/Engine/Source/Developer/DerivedDataCache/Private/FileSystemCacheStore.cpp
--- a/Engine/Source/Developer/DerivedDataCache/Private/FileSystemCacheStore.cpp
+++ b/Engine/Source/Developer/DerivedDataCache/Private/FileSystemCacheStore.cpp
@@ -825,9 +825,6 @@ private:
 
 	[[nodiscard]] bool FileExists(FStringBuilderBase& Path) const;
 
-	[[nodiscard]] bool ShouldSimulateMiss(const TCHAR* InKey);
-	[[nodiscard]] bool ShouldSimulateMiss(const FCacheKey& Key);
-
 private:
 	void BuildCacheLegacyPath(const TCHAR* const CacheKey, FStringBuilderBase& Path) const
 	{
@@ -874,13 +871,6 @@ private:
 	/** Debug Options */
 	FBackendDebugOptions DebugOptions;
 
-	/** Keys we ignored due to miss rate settings */
-	FCriticalSection MissedKeysCS;
-	TSet<FName> DebugMissedKeys;
-	TSet<FCacheKey> DebugMissedCacheKeys;
-	TSet<FName> DebugPutKeys;
-	TSet<FCacheKey> DebugPutCacheKeys;
-
 	FDerivedDataCacheUsageStats UsageStats;
 
 	TUniquePtr<FFileSystemCacheStoreMaintainer> Maintainer;
@@ -1300,7 +1290,7 @@ bool FFileSystemCacheStore::CachedDataProbablyExists(const TCHAR* const CacheKey
 		return false;
 	}
 
-	if (ShouldSimulateMiss(CacheKey))
+	if (DebugOptions.ShouldSimulateGetMiss(CacheKey))
 	{
 		UE_LOG(LogDerivedDataCache, Verbose, TEXT("%s: Simulated miss for exists check of %s"), *CachePath, CacheKey);
 		return false;
@@ -1351,7 +1341,7 @@ bool FFileSystemCacheStore::GetCachedData(const TCHAR* const CacheKey, TArray<ui
 		return false;
 	}
 
-	if (ShouldSimulateMiss(CacheKey))
+	if (DebugOptions.ShouldSimulateGetMiss(CacheKey))
 	{
 		UE_LOG(LogDerivedDataCache, Verbose, TEXT("%s: Simulated miss for get of %s"), *CachePath, CacheKey);
 		return false;
@@ -1443,13 +1433,10 @@ FFileSystemCacheStore::EPutStatus FFileSystemCacheStore::PutCachedData(
 		return EPutStatus::NotCached;
 	}
 
-	if (ShouldSimulateMiss(CacheKey))
+	if (DebugOptions.ShouldSimulatePutMiss(CacheKey))
 	{
-		const FName Key(CacheKey);
-		const uint32 Hash = GetTypeHash(Key);
-
-		FScopeLock Lock(&MissedKeysCS);
-		DebugPutKeys.AddByHash(Hash, Key);
+		UE_LOG(LogDerivedDataCache, Verbose, TEXT("%s: Simulated miss for put of %s"), *CachePath, CacheKey);
+		return EPutStatus::NotCached;
 	}
 
 	TStringBuilder<256> LegacyPath;
@@ -1781,12 +1768,11 @@ bool FFileSystemCacheStore::PutCacheRecord(
 		return false;
 	}
 
-	if (ShouldSimulateMiss(Key))
+	if (DebugOptions.ShouldSimulatePutMiss(Key))
 	{
-		const uint32 Hash = GetTypeHash(Key);
-
-		FScopeLock Lock(&MissedKeysCS);
-		DebugPutCacheKeys.AddByHash(Hash, Key);
+		UE_LOG(LogDerivedDataCache, Verbose, TEXT("%s: Simulated miss for put of %s from '%.*s'"),
+			*CachePath, *WriteToString<96>(Key), Name.Len(), Name.GetData());
+		return false;
 	}
 
 	TStringBuilder<256> Path;
@@ -1940,7 +1926,7 @@ FOptionalCacheRecord FFileSystemCacheStore::GetCacheRecordOnly(
 		return FOptionalCacheRecord();
 	}
 
-	if (ShouldSimulateMiss(Key))
+	if (DebugOptions.ShouldSimulateGetMiss(Key))
 	{
 		UE_LOG(LogDerivedDataCache, Verbose, TEXT("%s: Simulated miss for get of %s from '%.*s'"),
 			*CachePath, *WriteToString<96>(Key), Name.Len(), Name.GetData());
@@ -2065,12 +2051,11 @@ bool FFileSystemCacheStore::PutCacheValue(
 		return false;
 	}
 
-	if (ShouldSimulateMiss(Key))
+	if (DebugOptions.ShouldSimulatePutMiss(Key))
 	{
-		const uint32 Hash = GetTypeHash(Key);
-
-		FScopeLock Lock(&MissedKeysCS);
-		DebugPutCacheKeys.AddByHash(Hash, Key);
+		UE_LOG(LogDerivedDataCache, Verbose, TEXT("%s: Simulated miss for put of %s from '%.*s'"),
+			*CachePath, *WriteToString<96>(Key), Name.Len(), Name.GetData());
+		return false;
 	}
 
 	// Check if there is an existing value package.
@@ -2161,7 +2146,7 @@ bool FFileSystemCacheStore::GetCacheValueOnly(
 		return false;
 	}
 
-	if (ShouldSimulateMiss(Key))
+	if (DebugOptions.ShouldSimulateGetMiss(Key))
 	{
 		UE_LOG(LogDerivedDataCache, Verbose, TEXT("%s: Simulated miss for get of %s from '%.*s'"),
 			*CachePath, *WriteToString<96>(Key), Name.Len(), Name.GetData());
@@ -2587,71 +2572,6 @@ bool FFileSystemCacheStore::FileExists(FStringBuilderBase& Path) const
 	return true;
 }
 
-bool FFileSystemCacheStore::ShouldSimulateMiss(const TCHAR* const InKey)
-{
-	if (DebugOptions.RandomMissRate == 0 && DebugOptions.SimulateMissTypes.IsEmpty())
-	{
-		return false;
-	}
-
-	const FName Key(InKey);
-	const uint32 Hash = GetTypeHash(Key);
-
-	{
-		FScopeLock Lock(&MissedKeysCS);
-		if (DebugPutKeys.ContainsByHash(Hash, Key))
-		{
-			return false;
-		}
-
-		if (DebugMissedKeys.ContainsByHash(Hash, Key))
-		{
-			return true;
-		}
-	}
-
-	if (DebugOptions.ShouldSimulateMiss(InKey))
-	{
-		FScopeLock Lock(&MissedKeysCS);
-		DebugMissedKeys.AddByHash(Hash, Key);
-		return true;
-	}
-
-	return false;
-}
-
-bool FFileSystemCacheStore::ShouldSimulateMiss(const FCacheKey& Key)
-{
-	if (DebugOptions.RandomMissRate == 0 && DebugOptions.SimulateMissTypes.IsEmpty())
-	{
-		return false;
-	}
-
-	const uint32 Hash = GetTypeHash(Key);
-
-	{
-		FScopeLock Lock(&MissedKeysCS);
-		if (DebugPutCacheKeys.ContainsByHash(Hash, Key))
-		{
-			return false;
-		}
-
-		if (DebugMissedCacheKeys.ContainsByHash(Hash, Key))
-		{
-			return true;
-		}
-	}
-
-	if (DebugOptions.ShouldSimulateMiss(Key))
-	{
-		FScopeLock Lock(&MissedKeysCS);
-		DebugMissedCacheKeys.AddByHash(Hash, Key);
-		return true;
-	}
-
-	return false;
-}
-
 ILegacyCacheStore* CreateFileSystemCacheStore(
 	const TCHAR* CachePath,
 	const TCHAR* Params,
