Avoid wasted time when pulling data from Jupiter by:
1) Increasing connection pool size
2) Allowing DNS and SSL/TLS data to be shared by connections
3) Switching to ValueOnly mode for Jupiter backends by default
4) Using "inline get" feature when querying Value data

Tested:
- .run editor -- -execCmds="Automation RunTests System.DerivedDataCache.HttpDerivedDataBackend.CacheStore;Quit" -NullRHI -stdout -unattended -ddc=DerivedDataBackendGraph_Jupiter
-Launch CitySample editor with empty local cache and open Small_City_LVL
-Launch CitySample editor with hot local cache and open Small_City_LVL

Can be improved further with a batched existence API for refs (instead of blobs).  The current one does not work due to a Jupiter server bug, so this is an improvement that will have to come afterwards.

Note: There is a known issue involving Jupiter where the server will sometimes return uncompressed blobs.  This manifests as cache misses due to invalid/rejected results returned from Jupiter.  The client handles the miss and builds the data and stores it in its local cache but does not upload it to Jupiter (because Jupiter says it already has the blob).  Subsequent runs will not suffer from the issue as the data will be present in the local cache.  There are discussions in progress for how to address this issue on server or client side.  No immediate mitigation is needed for it.

diff --git a/Engine/Source/Developer/DerivedDataCache/Private/DerivedDataBackends.cpp b/Engine/Source/Developer/DerivedDataCache/Private/DerivedDataBackends.cpp
--- a/Engine/Source/Developer/DerivedDataCache/Private/DerivedDataBackends.cpp
+++ b/Engine/Source/Developer/DerivedDataCache/Private/DerivedDataBackends.cpp
@@ -846,7 +846,7 @@ public:
 		FString OAuthProvider;
 		FString OAuthClientId;
 		FString OAuthSecret;
-		EBackendLegacyMode LegacyMode = EBackendLegacyMode::ValueWithLegacyFallback;
+		EBackendLegacyMode LegacyMode = EBackendLegacyMode::ValueOnly;
 		bool bReadOnly = false;
 
 		ParseHttpCacheParams(NodeName, Entry, IniFilename, IniSection, Host, Namespace, StructuredNamespace, OAuthProvider, OAuthClientId, OAuthSecret, LegacyMode, bReadOnly);
diff --git a/Engine/Source/Developer/DerivedDataCache/Private/HttpCacheStore.cpp b/Engine/Source/Developer/DerivedDataCache/Private/HttpCacheStore.cpp
new file mode 100644
--- /dev/null
+++ b/Engine/Source/Developer/DerivedDataCache/Private/HttpCacheStore.cpp
@@ -0,0 +1,4008 @@
+// Copyright Epic Games, Inc. All Rights Reserved.
+
+#include "DerivedDataBackendInterface.h"
+
+#if WITH_HTTP_DDC_BACKEND
+
+#if PLATFORM_WINDOWS || PLATFORM_HOLOLENS
+#include "Windows/WindowsHWrapper.h"
+#include "Windows/AllowWindowsPlatformTypes.h"
+#endif
+#include "curl/curl.h"
+#if PLATFORM_WINDOWS || PLATFORM_HOLOLENS
+#include "Windows/HideWindowsPlatformTypes.h"
+#endif
+#include "Algo/Accumulate.h"
+#include "Algo/Find.h"
+#include "Algo/Transform.h"
+#include "Compression/CompressedBuffer.h"
+#include "Containers/StaticArray.h"
+#include "Containers/StringView.h"
+#include "Containers/Ticker.h"
+#include "DerivedDataCacheKey.h"
+#include "DerivedDataCachePrivate.h"
+#include "DerivedDataCacheRecord.h"
+#include "DerivedDataCacheUsageStats.h"
+#include "DerivedDataChunk.h"
+#include "DerivedDataValue.h"
+#include "Dom/JsonObject.h"
+#include "Experimental/Containers/FAAArrayQueue.h"
+#include "GenericPlatform/GenericPlatformFile.h"
+#include "HAL/PlatformFileManager.h"
+#include "IO/IoHash.h"
+#include "Memory/SharedBuffer.h"
+#include "Misc/FileHelper.h"
+#include "Misc/ScopeLock.h"
+#include "Misc/SecureHash.h"
+#include "Misc/StringBuilder.h"
+#include "Policies/CondensedJsonPrintPolicy.h"
+#include "ProfilingDebugging/CountersTrace.h"
+#include "ProfilingDebugging/CpuProfilerTrace.h"
+#include "Serialization/BufferArchive.h"
+#include "Serialization/CompactBinary.h"
+#include "Serialization/CompactBinaryPackage.h"
+#include "Serialization/CompactBinaryValidation.h"
+#include "Serialization/CompactBinaryWriter.h"
+#include "Serialization/JsonReader.h"
+#include "Serialization/JsonSerializer.h"
+#include "Serialization/JsonWriter.h"
+
+#if WITH_SSL
+#include "Ssl.h"
+#include <openssl/ssl.h>
+#endif
+
+// Enables data request helpers that internally
+// batch requests to reduce the number of concurrent
+// connections.
+#ifndef WITH_DATAREQUEST_HELPER
+	#define WITH_DATAREQUEST_HELPER 1
+#endif
+
+#define UE_HTTPDDC_BACKEND_WAIT_INTERVAL 0.01f
+#define UE_HTTPDDC_BACKEND_WAIT_INTERVAL_MS ((uint32)(UE_HTTPDDC_BACKEND_WAIT_INTERVAL*1000))
+#define UE_HTTPDDC_HTTP_REQUEST_TIMEOUT_SECONDS 30L
+#define UE_HTTPDDC_HTTP_REQUEST_TIMEOUT_ENABLED 1
+#define UE_HTTPDDC_HTTP_DEBUG 0
+#define UE_HTTPDDC_GET_REQUEST_POOL_SIZE 48
+#define UE_HTTPDDC_PUT_REQUEST_POOL_SIZE 16
+#define UE_HTTPDDC_MAX_FAILED_LOGIN_ATTEMPTS 16
+#define UE_HTTPDDC_MAX_ATTEMPTS 4
+#define UE_HTTPDDC_MAX_BUFFER_RESERVE 104857600u
+#define UE_HTTPDDC_BATCH_SIZE 12
+#define UE_HTTPDDC_BATCH_NUM 64
+#define UE_HTTPDDC_BATCH_GET_WEIGHT 4
+#define UE_HTTPDDC_BATCH_HEAD_WEIGHT 1
+#define UE_HTTPDDC_BATCH_WEIGHT_HINT 12
+
+namespace UE::DerivedData
+{
+
+TRACE_DECLARE_INT_COUNTER(HttpDDC_Exist, TEXT("HttpDDC Exist"));
+TRACE_DECLARE_INT_COUNTER(HttpDDC_ExistHit, TEXT("HttpDDC Exist Hit"));
+TRACE_DECLARE_INT_COUNTER(HttpDDC_Get, TEXT("HttpDDC Get"));
+TRACE_DECLARE_INT_COUNTER(HttpDDC_GetHit, TEXT("HttpDDC Get Hit"));
+TRACE_DECLARE_INT_COUNTER(HttpDDC_Put, TEXT("HttpDDC Put"));
+TRACE_DECLARE_INT_COUNTER(HttpDDC_PutHit, TEXT("HttpDDC Put Hit"));
+TRACE_DECLARE_INT_COUNTER(HttpDDC_BytesReceived, TEXT("HttpDDC Bytes Received"));
+TRACE_DECLARE_INT_COUNTER(HttpDDC_BytesSent, TEXT("HttpDDC Bytes Sent"));
+
+static CURLcode sslctx_function(CURL * curl, void * sslctx, void * parm);
+
+typedef TSharedPtr<class IHttpRequest> FHttpRequestPtr;
+typedef TSharedPtr<class IHttpResponse, ESPMode::ThreadSafe> FHttpResponsePtr;
+
+/**
+ * Encapsulation for access token shared by all requests.
+ */
+struct FHttpAccessToken
+{
+public:
+	FHttpAccessToken() = default;
+	FString GetHeader();
+	void SetHeader(const TCHAR*);
+	uint32 GetSerial() const;
+private:
+	FRWLock		Lock;
+	FString		Token;
+	uint32		Serial;
+};
+
+struct FHttpSharedData
+{
+	FHttpSharedData()
+	{
+		FMemory::Memset(WriteLocked, 0, sizeof(WriteLocked));
+		CurlShare = curl_share_init();
+		curl_share_setopt(CurlShare, CURLSHOPT_USERDATA, this);
+		curl_share_setopt(CurlShare, CURLSHOPT_LOCKFUNC, LockFn);
+		curl_share_setopt(CurlShare, CURLSHOPT_UNLOCKFUNC, UnlockFn);
+		curl_share_setopt(CurlShare, CURLSHOPT_SHARE, CURL_LOCK_DATA_DNS);
+		curl_share_setopt(CurlShare, CURLSHOPT_SHARE, CURL_LOCK_DATA_SSL_SESSION);
+	}
+
+	~FHttpSharedData()
+	{
+		curl_share_cleanup(CurlShare);
+	}
+
+	CURLSH* CurlShare;
+
+private:
+	FRWLock Locks[CURL_LOCK_DATA_LAST];
+	bool WriteLocked[CURL_LOCK_DATA_LAST];
+
+	static void LockFn(CURL* Handle, curl_lock_data Data, curl_lock_access Access, void* User)
+	{
+		FHttpSharedData* SharedData = (FHttpSharedData*)User;
+		if (Access == CURL_LOCK_ACCESS_SHARED)
+		{
+			SharedData->Locks[Data].ReadLock();
+		}
+		else
+		{
+			SharedData->Locks[Data].WriteLock();
+			SharedData->WriteLocked[Data] = true;
+		}
+	}
+	static void UnlockFn(CURL* Handle, curl_lock_data Data, void* User)
+	{
+		FHttpSharedData* SharedData = (FHttpSharedData*)User;
+		if (!SharedData->WriteLocked[Data])
+		{
+			SharedData->Locks[Data].ReadUnlock();
+		}
+		else
+		{
+			SharedData->WriteLocked[Data] = false;
+			SharedData->Locks[Data].WriteUnlock();
+		}
+	}
+};
+
+/**
+ * Minimal HTTP request type wrapping CURL without the need for managers. This request
+ * is written to allow reuse of request objects, in order to allow connections to be reused.
+ *
+ * CURL has a global library initialization (curl_global_init). We rely on this happening in 
+ * the Online/HTTP library which is a dependency on this module.
+ */
+class FHttpRequest
+{
+public:
+	/**
+	 * Supported request verbs
+	 */
+	enum RequestVerb
+	{
+		Get,
+		Put,
+		PutCompactBinary,
+		PutCompressedBlob,
+		Post,
+		PostJson,
+		Delete,
+		Head
+	};
+
+	/**
+	 * Convenience result type interpreted from HTTP response code.
+	 */
+	enum Result
+	{
+		Success,
+		Failed,
+		FailedTimeout
+	};
+
+	FHttpRequest(const TCHAR* InDomain, const TCHAR* InEffectiveDomain, FHttpAccessToken* InAuthorizationToken, bool bInLogErrors)
+		: bLogErrors(bInLogErrors)
+		, Domain(InDomain)
+		, EffectiveDomain(InEffectiveDomain)
+		, AuthorizationToken(InAuthorizationToken)
+	{
+		Curl = curl_easy_init();
+		Reset();
+	}
+
+	~FHttpRequest()
+	{
+		curl_easy_cleanup(Curl);
+	}
+
+	/**
+	 * Resets all options on the request except those that should always be set.
+	 */
+	void Reset()
+	{
+		Headers.Reset();
+		ResponseHeader.Reset();
+		ResponseBuffer.Reset();
+		ResponseCode = 0;
+		ReadDataView = TArrayView<const uint8>();
+		WriteDataBufferPtr = nullptr;
+		WriteHeaderBufferPtr = nullptr;
+		BytesSent = 0;
+		BytesReceived = 0;
+		CurlResult = CURL_LAST;
+
+		static FHttpSharedData SharedData;
+
+		curl_easy_reset(Curl);
+
+		// Options that are always set for all connections.
+#if UE_HTTPDDC_HTTP_REQUEST_TIMEOUT_ENABLED
+		curl_easy_setopt(Curl, CURLOPT_CONNECTTIMEOUT, UE_HTTPDDC_HTTP_REQUEST_TIMEOUT_SECONDS);
+#endif
+		curl_easy_setopt(Curl, CURLOPT_FOLLOWLOCATION, 1L);
+		curl_easy_setopt(Curl, CURLOPT_NOSIGNAL, 1L);
+		curl_easy_setopt(Curl, CURLOPT_DNS_CACHE_TIMEOUT, 300L); // Don't re-resolve every minute
+		curl_easy_setopt(Curl, CURLOPT_SHARE, SharedData.CurlShare);
+		// SSL options
+		curl_easy_setopt(Curl, CURLOPT_USE_SSL, CURLUSESSL_ALL);
+		curl_easy_setopt(Curl, CURLOPT_SSL_VERIFYPEER, 1);
+		curl_easy_setopt(Curl, CURLOPT_SSL_VERIFYHOST, 1);
+		curl_easy_setopt(Curl, CURLOPT_SSLCERTTYPE, "PEM");
+		// Response functions
+		curl_easy_setopt(Curl, CURLOPT_HEADERDATA, this);
+		curl_easy_setopt(Curl, CURLOPT_HEADERFUNCTION, &FHttpRequest::StaticWriteHeaderFn);
+		curl_easy_setopt(Curl, CURLOPT_WRITEDATA, this);
+		curl_easy_setopt(Curl, CURLOPT_WRITEFUNCTION, StaticWriteBodyFn);
+		// SSL certification verification
+		curl_easy_setopt(Curl, CURLOPT_CAINFO, nullptr);
+		curl_easy_setopt(Curl, CURLOPT_SSL_CTX_FUNCTION, *sslctx_function);
+		curl_easy_setopt(Curl, CURLOPT_SSL_CTX_DATA, this);
+		// Allow compressed data
+		curl_easy_setopt(Curl, CURLOPT_ACCEPT_ENCODING, "gzip");
+		// Rewind method, handle special error case where request need to rewind data stream
+		curl_easy_setopt(Curl, CURLOPT_SEEKFUNCTION, StaticSeekFn);
+		curl_easy_setopt(Curl, CURLOPT_SEEKDATA, this);
+		// Set minimum speed behavior to allow operations to abort if the transfer speed is poor for the given duration (1kbps over a 30 second span)
+		curl_easy_setopt(Curl, CURLOPT_LOW_SPEED_TIME, 30L);
+		curl_easy_setopt(Curl, CURLOPT_LOW_SPEED_LIMIT, 1024L);
+		// Debug hooks
+#if UE_HTTPDDC_HTTP_DEBUG
+		curl_easy_setopt(Curl, CURLOPT_DEBUGDATA, this);
+		curl_easy_setopt(Curl, CURLOPT_DEBUGFUNCTION, StaticDebugCallback);
+		curl_easy_setopt(Curl, CURLOPT_VERBOSE, 1L);
+#endif
+	}
+
+	/** Gets the domain name for this request */
+	const FString& GetName() const
+	{
+		return Domain;
+	}
+
+	/** Gets the domain name for this request */
+	const FString& GetDomain() const
+	{
+		return Domain;
+	}
+
+	/** Gets the effective domain name for this request */
+	const FString& GetEffectiveDomain() const
+	{
+		return EffectiveDomain;
+	}
+
+	/** Returns the HTTP response code.*/
+	const int64 GetResponseCode() const
+	{
+		return ResponseCode;
+	}
+
+	/** Returns the number of bytes received this request (headers withstanding). */
+	const size_t GetBytesReceived() const
+	{
+		return BytesReceived;
+	}
+
+	/** Returns the number of bytes sent during this request (headers withstanding). */
+	const size_t  GetBytesSent() const
+	{
+		return BytesSent;
+	}
+
+	/**
+	 * Upload buffer using the request, using either "Put" or "Post" verbs.
+	 * @param Uri Url to use.
+	 * @param Buffer Data to upload
+	 * @return Result of the request
+	 */
+	template<RequestVerb V>
+	Result PerformBlockingUpload(const TCHAR* Uri, TArrayView<const uint8> Buffer, TConstArrayView<long> ExpectedErrorCodes = {})
+	{
+		static_assert(V == Put || V == PutCompactBinary || V == PutCompressedBlob || V == Post || V == PostJson, "Upload should use either Put or Post verbs.");
+		
+		uint32 ContentLength = 0u;
+
+		if constexpr (V == Put || V == PutCompactBinary || V == PutCompressedBlob)
+		{
+			curl_easy_setopt(Curl, CURLOPT_UPLOAD, 1L);
+			curl_easy_setopt(Curl, CURLOPT_INFILESIZE, Buffer.Num());
+			curl_easy_setopt(Curl, CURLOPT_READDATA, this);
+			curl_easy_setopt(Curl, CURLOPT_READFUNCTION, StaticReadFn);
+			if constexpr (V == PutCompactBinary)
+			{
+				Headers.Add(FString(TEXT("Content-Type: application/x-ue-cb")));
+			}
+			else if constexpr (V == PutCompressedBlob)
+			{
+				Headers.Add(FString(TEXT("Content-Type: application/x-ue-comp")));
+			}
+			else
+			{
+				Headers.Add(FString(TEXT("Content-Type: application/octet-stream")));
+			}
+			ContentLength = Buffer.Num();
+			ReadDataView = Buffer;
+		}
+		else if (V == Post || V == PostJson)
+		{
+			curl_easy_setopt(Curl, CURLOPT_POST, 1L);
+			curl_easy_setopt(Curl, CURLOPT_INFILESIZE, Buffer.Num());
+			curl_easy_setopt(Curl, CURLOPT_READDATA, this);
+			curl_easy_setopt(Curl, CURLOPT_READFUNCTION, StaticReadFn);
+			Headers.Add(V == Post ? FString(TEXT("Content-Type: application/x-www-form-urlencoded")) : FString(TEXT("Content-Type: application/json")));
+			ContentLength = Buffer.Num();
+			ReadDataView = Buffer;
+		}
+
+		return PerformBlocking(Uri, V, ContentLength, ExpectedErrorCodes);
+	}
+
+	/**
+	 * Download an url into a buffer using the request.
+	 * @param Uri Url to use.
+	 * @param Buffer Optional buffer where data should be downloaded to. If empty downloaded data will
+	 * be stored in an internal buffer and accessed GetResponse* methods.
+	 * @return Result of the request
+	 */
+	Result PerformBlockingDownload(const TCHAR* Uri, TArray<uint8>* Buffer, TConstArrayView<long> ExpectedErrorCodes = {400})
+	{
+		curl_easy_setopt(Curl, CURLOPT_HTTPGET, 1L);
+		WriteDataBufferPtr = Buffer;
+
+		return PerformBlocking(Uri, Get, 0u, ExpectedErrorCodes);
+	}
+
+	/**
+	 * Query an url using the request. Queries can use either "Head" or "Delete" verbs.
+	 * @param Uri Url to use.
+	 * @return Result of the request
+	 */
+	template<RequestVerb V>
+	Result PerformBlockingQuery(const TCHAR* Uri, TConstArrayView<long> ExpectedErrorCodes = {400})
+	{
+		static_assert(V == Head || V == Delete, "Queries should use either Head or Delete verbs.");
+
+		if (V == Delete)
+		{
+			curl_easy_setopt(Curl, CURLOPT_CUSTOMREQUEST, "DELETE");
+		}
+		else if (V == Head)
+		{
+			curl_easy_setopt(Curl, CURLOPT_NOBODY, 1L);
+		}
+
+		return PerformBlocking(Uri, V, 0u, ExpectedErrorCodes);
+	}
+
+	/**
+	 * Set a header to send with the request.
+	 */
+	void SetHeader(const TCHAR* Header, const TCHAR* Value)
+	{
+		check(CurlResult == CURL_LAST); // Cannot set header after request is sent
+		Headers.Add(FString::Printf(TEXT("%s: %s"), Header, Value));
+	}
+
+	/**
+	 * Attempts to find the header from the response. Returns false if header is not present.
+	 */
+	bool GetHeader(const ANSICHAR* Header, FString& OutValue) const
+	{
+		check(CurlResult != CURL_LAST);  // Cannot query headers before request is sent
+
+		const ANSICHAR* HeadersBuffer = (const ANSICHAR*) ResponseHeader.GetData();
+		size_t HeaderLen = strlen(Header);
+
+		// Find the header key in the (ANSI) response buffer. If not found we can exist immediately
+		if (const ANSICHAR* Found = strstr(HeadersBuffer, Header))
+		{
+			const ANSICHAR* Linebreak = strchr(Found, '\r');
+			const ANSICHAR* ValueStart = Found + HeaderLen + 2; //colon and space
+			const size_t ValueSize = Linebreak - ValueStart;
+			FUTF8ToTCHAR TCHARData(ValueStart, ValueSize);
+			OutValue = FString(TCHARData.Length(), TCHARData.Get());
+			return true;
+		}
+		return false;
+	}
+
+	/**
+	 * Returns the response buffer. Note that is the request is performed
+	 * with an external buffer as target buffer this string will be empty.
+	 */
+	const TArray<uint8>& GetResponseBuffer() const
+	{
+		return ResponseBuffer;
+	}
+
+	/**
+	 * Returns the response buffer as a string. Note that is the request is performed
+	 * with an external buffer as target buffer this string will be empty.
+	 */
+	FString GetResponseAsString() const
+	{
+		return GetAnsiBufferAsString(ResponseBuffer);
+	}
+
+	/**
+	 * Returns the response header as a string.
+	 */
+	FString GetResponseHeaderAsString()
+	{
+		return GetAnsiBufferAsString(ResponseHeader);
+	}
+
+	/**
+	 * Tries to parse the response buffer as a JsonObject. Return empty pointer if 
+	 * parse error occurs.
+	 */
+	TSharedPtr<FJsonObject> GetResponseAsJsonObject() const
+	{
+		FString Response = GetAnsiBufferAsString(ResponseBuffer);
+
+		TSharedPtr<FJsonObject> JsonObject;
+		TSharedRef<TJsonReader<> > JsonReader = TJsonReaderFactory<>::Create(Response);
+		if (!FJsonSerializer::Deserialize(JsonReader, JsonObject) || !JsonObject.IsValid())
+		{
+			return TSharedPtr<FJsonObject>(nullptr);
+		}
+
+		return JsonObject;
+	}
+
+	/**
+	 * Tries to parse the response buffer as a JsonArray. Return empty array if
+	 * parse error occurs.
+	 */
+	TArray<TSharedPtr<FJsonValue>> GetResponseAsJsonArray() const
+	{
+		FString Response = GetAnsiBufferAsString(ResponseBuffer);
+
+		TArray<TSharedPtr<FJsonValue>> JsonArray;
+		TSharedRef<TJsonReader<>> JsonReader = TJsonReaderFactory<>::Create(Response);
+		FJsonSerializer::Deserialize(JsonReader, JsonArray);
+		return JsonArray;
+	}
+
+	/** Will return true if the response code is considered a success */
+	static bool IsSuccessResponse(long ResponseCode)
+	{
+		// We consider anything in the 1XX or 2XX range a success
+		return ResponseCode >= 100 && ResponseCode < 300;
+	}
+
+private:
+
+	CURL*					Curl;
+	CURLcode				CurlResult;
+	long					ResponseCode;
+	size_t					BytesSent;
+	size_t					BytesReceived;
+	bool					bLogErrors;
+
+	TArrayView<const uint8>	ReadDataView;
+	TArray<uint8>*			WriteDataBufferPtr;
+	TArray<uint8>*			WriteHeaderBufferPtr;
+	
+	TArray<uint8>			ResponseHeader;
+	TArray<uint8>			ResponseBuffer;
+	TArray<FString>			Headers;
+	FString					Domain;
+	FString					EffectiveDomain;
+	FHttpAccessToken*		AuthorizationToken;
+
+	/**
+	 * Performs the request, blocking until finished.
+	 * @param Uri Address on the domain to query
+	 * @param Verb HTTP verb to use
+	 * @param Buffer Optional buffer to directly receive the result of the request.
+	 * If unset the response body will be stored in the request.
+	 */
+	Result PerformBlocking(const TCHAR* Uri, RequestVerb Verb, uint32 ContentLength, TConstArrayView<long> ExpectedErrorCodes)
+	{
+		static const char* CommonHeaders[] = {
+			"User-Agent: Unreal Engine",
+			nullptr
+		};
+
+		TRACE_CPUPROFILER_EVENT_SCOPE(HttpDDC_CurlPerform);
+
+		// Setup request options
+		FString Url = FString::Printf(TEXT("%s/%s"), *EffectiveDomain, Uri);
+		curl_easy_setopt(Curl, CURLOPT_URL, TCHAR_TO_ANSI(*Url));
+
+		// Setup response header buffer. If caller has not setup a response data buffer, use internal.
+		WriteHeaderBufferPtr = &ResponseHeader;
+		if (WriteDataBufferPtr == nullptr)
+		{
+			WriteDataBufferPtr = &ResponseBuffer;
+		}
+
+		// Content-Length should always be set
+		Headers.Add(FString::Printf(TEXT("Content-Length: %d"), ContentLength));
+
+		// And auth token if it's set
+		if (AuthorizationToken)
+		{
+			Headers.Add(AuthorizationToken->GetHeader());
+		}
+
+		// Build headers list
+		curl_slist* CurlHeaders = nullptr;
+		// Add common headers
+		for (uint8 i = 0; CommonHeaders[i] != nullptr; ++i)
+		{
+			CurlHeaders = curl_slist_append(CurlHeaders, CommonHeaders[i]);
+		}
+		// Setup added headers
+		for (const FString& Header : Headers)
+		{
+			CurlHeaders = curl_slist_append(CurlHeaders, TCHAR_TO_ANSI(*Header));
+		}
+		curl_easy_setopt(Curl, CURLOPT_HTTPHEADER, CurlHeaders);
+
+		// Shots fired!
+		CurlResult = curl_easy_perform(Curl);
+
+		// Get response code
+		bool bRedirected = false;
+		if (CURLE_OK == curl_easy_getinfo(Curl, CURLINFO_RESPONSE_CODE, &ResponseCode))
+		{
+			bRedirected = (ResponseCode >= 300 && ResponseCode < 400);
+		}
+
+		LogResult(CurlResult, Uri, Verb, ExpectedErrorCodes);
+
+		// Clean up
+		curl_slist_free_all(CurlHeaders);
+
+		return CurlResult == CURLE_OK ? Success : Failed;
+	}
+
+	void LogResult(CURLcode Result, const TCHAR* Uri, RequestVerb Verb, TConstArrayView<long> ExpectedErrorCodes) const
+	{
+		if (Result == CURLE_OK)
+		{
+			bool bSuccess = false;
+			const TCHAR* VerbStr = nullptr;
+			FString AdditionalInfo;
+
+			switch (Verb)
+			{
+			case Head:
+				bSuccess = (ExpectedErrorCodes.Contains(ResponseCode) || IsSuccessResponse(ResponseCode));
+				VerbStr = TEXT("querying");
+				break;
+			case Get:
+				bSuccess = (ExpectedErrorCodes.Contains(ResponseCode) || IsSuccessResponse(ResponseCode));
+				VerbStr = TEXT("fetching");
+				AdditionalInfo = FString::Printf(TEXT("Received: %d bytes."), BytesReceived);
+				break;
+			case Put:
+			case PutCompactBinary:
+			case PutCompressedBlob:
+				bSuccess = (ExpectedErrorCodes.Contains(ResponseCode) || IsSuccessResponse(ResponseCode));
+				VerbStr = TEXT("updating");
+				AdditionalInfo = FString::Printf(TEXT("Sent: %d bytes."), BytesSent);
+				break;
+			case Post:
+			case PostJson:
+				bSuccess = (ExpectedErrorCodes.Contains(ResponseCode) || IsSuccessResponse(ResponseCode));
+				VerbStr = TEXT("posting");
+				break;
+			case Delete:
+				bSuccess = (ExpectedErrorCodes.Contains(ResponseCode) || IsSuccessResponse(ResponseCode));
+				VerbStr = TEXT("deleting");
+				break;
+			}
+
+			if (bSuccess)
+			{
+				UE_LOG(
+					LogDerivedDataCache, 
+					Verbose, 
+					TEXT("%s: Finished %s HTTP cache entry (response %d) from %s. %s"), 
+					*GetName(),
+					VerbStr,
+					ResponseCode, 
+					Uri,
+					*AdditionalInfo
+				);
+			}
+			else if(bLogErrors)
+			{
+				// Print the response body if we got one, otherwise print header.
+				FString Response = GetAnsiBufferAsString(ResponseBuffer.Num() > 0 ? ResponseBuffer : ResponseHeader);
+				Response.ReplaceCharInline('\n', ' ');
+				Response.ReplaceCharInline('\r', ' ');
+				// Dont log access denied as error, since tokens can expire mid session
+				if (ResponseCode == 401)
+				{
+					UE_LOG(
+						LogDerivedDataCache,
+						Verbose,
+						TEXT("%s: Failed %s HTTP cache entry (response %d) from %s. Response: %s"),
+						*GetName(),
+						VerbStr,
+						ResponseCode,
+						Uri,
+						*Response
+					);
+				}
+				else
+				{ 
+					UE_LOG(
+						LogDerivedDataCache,
+						Display,
+						TEXT("%s: Failed %s HTTP cache entry (response %d) from %s. Response: %s"),
+						*GetName(),
+						VerbStr,
+						ResponseCode,
+						Uri,
+						*Response
+					);
+				}
+
+
+			}
+		}
+		else if(bLogErrors)
+		{
+			UE_LOG(
+				LogDerivedDataCache, 
+				Display, 
+				TEXT("%s: Error while connecting to %s: %s"), 
+				*GetName(),
+				*EffectiveDomain,
+				ANSI_TO_TCHAR(curl_easy_strerror(Result))
+			);
+		}
+	}
+
+	FString GetAnsiBufferAsString(const TArray<uint8>& Buffer) const
+	{
+		// Content is NOT null-terminated; we need to specify lengths here
+		FUTF8ToTCHAR TCHARData(reinterpret_cast<const ANSICHAR*>(Buffer.GetData()), Buffer.Num());
+		return FString(TCHARData.Length(), TCHARData.Get());
+	}
+
+	static size_t StaticDebugCallback(CURL * Handle, curl_infotype DebugInfoType, char * DebugInfo, size_t DebugInfoSize, void* UserData)
+	{
+		FHttpRequest* Request = static_cast<FHttpRequest*>(UserData);
+
+		switch (DebugInfoType)
+		{
+		case CURLINFO_TEXT:
+		{
+			// Truncate at 1023 characters. This is just an arbitrary number based on a buffer size seen in
+			// the libcurl code.
+			DebugInfoSize = FMath::Min(DebugInfoSize, (size_t)1023);
+
+			// Calculate the actual length of the string due to incorrect use of snprintf() in lib/vtls/openssl.c.
+			char* FoundNulPtr = (char*)memchr(DebugInfo, 0, DebugInfoSize);
+			int CalculatedSize = FoundNulPtr != nullptr ? FoundNulPtr - DebugInfo : DebugInfoSize;
+
+			auto ConvertedString = StringCast<TCHAR>(static_cast<const ANSICHAR*>(DebugInfo), CalculatedSize);
+			FString DebugText(ConvertedString.Length(), ConvertedString.Get());
+			DebugText.ReplaceInline(TEXT("\n"), TEXT(""), ESearchCase::CaseSensitive);
+			DebugText.ReplaceInline(TEXT("\r"), TEXT(""), ESearchCase::CaseSensitive);
+			UE_LOG(LogDerivedDataCache, VeryVerbose, TEXT("%s: %p: '%s'"), *Request->GetName(), Request, *DebugText);
+		}
+		break;
+
+		case CURLINFO_HEADER_IN:
+			UE_LOG(LogDerivedDataCache, VeryVerbose, TEXT("%s: %p: Received header (%d bytes)"), *Request->GetName(), Request, DebugInfoSize);
+			break;
+
+		case CURLINFO_DATA_IN:
+			UE_LOG(LogDerivedDataCache, VeryVerbose, TEXT("%s: %p: Received data (%d bytes)"), *Request->GetName(), Request, DebugInfoSize);
+			break;
+
+		case CURLINFO_DATA_OUT:
+			UE_LOG(LogDerivedDataCache, VeryVerbose, TEXT("%s: %p: Sent data (%d bytes)"), *Request->GetName(), Request, DebugInfoSize);
+			break;
+
+		case CURLINFO_SSL_DATA_IN:
+			UE_LOG(LogDerivedDataCache, VeryVerbose, TEXT("%s: %p: Received SSL data (%d bytes)"), *Request->GetName(), Request, DebugInfoSize);
+			break;
+
+		case CURLINFO_SSL_DATA_OUT:
+			UE_LOG(LogDerivedDataCache, VeryVerbose, TEXT("%s: %p: Sent SSL data (%d bytes)"), *Request->GetName(), Request, DebugInfoSize);
+			break;
+		}
+
+		return 0;
+	}
+
+	static size_t StaticReadFn(void* Ptr, size_t SizeInBlocks, size_t BlockSizeInBytes, void* UserData)
+	{
+		FHttpRequest* Request = static_cast<FHttpRequest*>(UserData);
+		TArrayView<const uint8>& ReadDataView = Request->ReadDataView;
+
+		const size_t Offset = Request->BytesSent;
+		const size_t ReadSize = FMath::Min((size_t)ReadDataView.Num() - Offset, SizeInBlocks * BlockSizeInBytes);
+		check(ReadDataView.Num() >= Offset + ReadSize);
+
+		FMemory::Memcpy(Ptr, ReadDataView.GetData() + Offset, ReadSize);
+		Request->BytesSent += ReadSize;
+		return ReadSize;
+		
+		return 0;
+	}
+
+	static size_t StaticWriteHeaderFn(void* Ptr, size_t SizeInBlocks, size_t BlockSizeInBytes, void* UserData)
+	{
+		FHttpRequest* Request = static_cast<FHttpRequest*>(UserData);
+		const size_t WriteSize = SizeInBlocks * BlockSizeInBytes;
+		TArray<uint8>* WriteHeaderBufferPtr = Request->WriteHeaderBufferPtr;
+		if (WriteHeaderBufferPtr && WriteSize > 0)
+		{
+			const size_t CurrentBufferLength = WriteHeaderBufferPtr->Num();
+			if (CurrentBufferLength > 0)
+			{
+				// Remove the previous zero termination
+				(*WriteHeaderBufferPtr)[CurrentBufferLength-1] = ' ';
+			}
+
+			// Write the header
+			WriteHeaderBufferPtr->Append((const uint8*)Ptr, WriteSize + 1);
+			(*WriteHeaderBufferPtr)[WriteHeaderBufferPtr->Num()-1] = 0; // Zero terminate string
+			return WriteSize;
+		}
+		return 0;
+	}
+
+	static size_t StaticWriteBodyFn(void* Ptr, size_t SizeInBlocks, size_t BlockSizeInBytes, void* UserData)
+	{
+		FHttpRequest* Request = static_cast<FHttpRequest*>(UserData);
+		const size_t WriteSize = SizeInBlocks * BlockSizeInBytes;
+		TArray<uint8>* WriteDataBufferPtr = Request->WriteDataBufferPtr;
+
+		if (WriteDataBufferPtr && WriteSize > 0)
+		{
+			// If this is the first part of the body being received, try to reserve 
+			// memory if content length is defined in the header.
+			if (Request->BytesReceived == 0 && Request->WriteHeaderBufferPtr)
+			{
+				static const ANSICHAR* ContentLengthHeaderStr = "Content-Length: ";
+				const ANSICHAR* Header = (const ANSICHAR*)Request->WriteHeaderBufferPtr->GetData();
+
+				if (const ANSICHAR* ContentLengthHeader = FCStringAnsi::Strstr(Header, ContentLengthHeaderStr))
+				{
+					size_t ContentLength = (size_t)FCStringAnsi::Atoi64(ContentLengthHeader + strlen(ContentLengthHeaderStr));
+					if (ContentLength > 0u && ContentLength < UE_HTTPDDC_MAX_BUFFER_RESERVE)
+					{
+						WriteDataBufferPtr->Reserve(ContentLength);
+					}
+				}
+			}
+
+			// Write to the target buffer
+			WriteDataBufferPtr->Append((const uint8*)Ptr, WriteSize);
+			Request->BytesReceived += WriteSize;
+			return WriteSize;
+		}
+		
+		return 0;
+	}
+
+	static size_t StaticSeekFn(void* UserData, curl_off_t Offset, int Origin)
+	{
+		FHttpRequest* Request = static_cast<FHttpRequest*>(UserData);
+		size_t NewPosition = 0;
+
+		switch (Origin)
+		{
+		case SEEK_SET: NewPosition = Offset; break;
+		case SEEK_CUR: NewPosition = Request->BytesSent + Offset; break;
+		case SEEK_END: NewPosition = Request->ReadDataView.Num() + Offset; break;
+		}
+
+		// Make sure we don't seek outside of the buffer
+		if (NewPosition < 0 || NewPosition >= Request->ReadDataView.Num())
+		{
+			return CURL_SEEKFUNC_FAIL;
+		}
+
+		// Update the used offset
+		Request->BytesSent = NewPosition;
+		return CURL_SEEKFUNC_OK;
+	}
+
+};
+
+
+//----------------------------------------------------------------------------------------------------------
+// Forward declarations
+//----------------------------------------------------------------------------------------------------------
+bool VerifyPayload(const FSHAHash& Hash, const TCHAR* Namespace, const TCHAR* Bucket, const TCHAR* CacheKey, const TArray<uint8>& Payload);
+bool VerifyPayload(const FIoHash& Hash, const TCHAR* Namespace, const TCHAR* Bucket, const TCHAR* CacheKey, const TArray<uint8>& Payload);
+bool VerifyRequest(const class FHttpRequest* Request, const TCHAR* Namespace, const TCHAR* Bucket, const TCHAR* CacheKey, const TArray<uint8>& Payload);
+bool HashPayload(class FHttpRequest* Request, const TArrayView<const uint8> Payload);
+bool ShouldAbortForShutdown();
+
+//----------------------------------------------------------------------------------------------------------
+// Request pool
+//----------------------------------------------------------------------------------------------------------
+
+
+/**
+ * Pool that manages a fixed set of requests. Users are required to release requests that have been 
+ * acquired. Usable with \ref FScopedRequestPtr which handles this automatically.
+ */
+struct FRequestPool
+{
+	FRequestPool(const TCHAR* InServiceUrl, const TCHAR* InEffectiveServiceUrl, FHttpAccessToken* InAuthorizationToken, uint32 PoolSize)
+	{
+		Pool.AddUninitialized(PoolSize);
+		for (uint8 i = 0; i < Pool.Num(); ++i)
+		{
+			Pool[i].Usage = 0u;
+			Pool[i].Request = new FHttpRequest(InServiceUrl, InEffectiveServiceUrl, InAuthorizationToken, true);
+		}
+		
+	}
+
+	~FRequestPool()
+	{
+		for (uint8 i = 0; i < Pool.Num(); ++i)
+		{
+			// No requests should be in use by now.
+			check(Pool[i].Usage.load(std::memory_order_acquire) == 0u);
+			delete Pool[i].Request;
+		}
+	}
+
+	/**
+	 * Attempts to get a request is free. Once a request has been returned it is
+	 * "owned by the caller and need to release it to the pool when work has been completed.
+	 * @return Usable request instance if one is available, otherwise null.
+	 */
+	FHttpRequest* GetFreeRequest()
+	{
+		for (uint8 i = 0; i < Pool.Num(); ++i)
+		{
+			if (!Pool[i].Usage.load(std::memory_order_relaxed))
+			{
+				uint8 Expected = 0u;
+				if (Pool[i].Usage.compare_exchange_strong(Expected, 1u))
+				{
+					Pool[i].Request->Reset();
+					return Pool[i].Request;
+				}
+			}
+		}
+		return nullptr;
+	}
+
+	class FWaiter : public FThreadSafeRefCountedObject
+	{
+	public:
+		std::atomic<FHttpRequest*> Request{ nullptr };
+
+		FWaiter(FRequestPool* InPool)
+			: Event(FPlatformProcess::GetSynchEventFromPool(true))
+			, Pool(InPool)
+		{
+		}
+		
+		bool Wait(uint32 TimeMS)
+		{
+			return Event->Wait(TimeMS);
+		}
+
+		void Trigger()
+		{
+			Event->Trigger();
+		}
+	private:
+		~FWaiter()
+		{
+			FPlatformProcess::ReturnSynchEventToPool(Event);
+
+			if (Request)
+			{
+				Pool->ReleaseRequestToPool(Request.exchange(nullptr));
+			}
+		}
+
+		FEvent* Event;
+		FRequestPool* Pool;
+	};
+
+	/**
+	 * Block until a request is free. Once a request has been returned it is 
+	 * "owned by the caller and need to release it to the pool when work has been completed.
+	 * @return Usable request instance.
+	 */
+	FHttpRequest* WaitForFreeRequest()
+	{
+		TRACE_CPUPROFILER_EVENT_SCOPE(HttpDDC_WaitForConnPool);
+
+		FHttpRequest* Request = GetFreeRequest();
+		if (Request == nullptr)
+		{
+			// Make it a fair by allowing each thread to register itself in a FIFO
+			// so that the first thread to start waiting is the first one to get a request.
+			FWaiter* Waiter = new FWaiter(this);
+			Waiter->AddRef(); // One ref for the thread that will dequeue
+			Waiter->AddRef(); // One ref for us
+
+			Waiters.enqueue(Waiter);
+
+			while (!Waiter->Wait(UE_HTTPDDC_BACKEND_WAIT_INTERVAL_MS))
+			{
+				// While waiting, allow us to check if a race occurred and a request has been freed
+				// between the time we checked for free requests and the time we queued ourself as a Waiter.
+				if ((Request = GetFreeRequest()) != nullptr)
+				{
+					// We abandon the FWaiter, it will be freed by the next dequeue
+					// and if it has a request, it will be queued back to the pool.
+					Waiter->Release();
+					return Request;
+				}
+			}
+
+			Request = Waiter->Request.exchange(nullptr);
+			Request->Reset();
+			Waiter->Release();
+		}
+		check(Request);
+		return Request;
+	}
+
+	/**
+	 * Release request to the pool.
+	 * @param Request Request that should be freed. Note that any buffer owened by the request can now be reset.
+	 */
+	void ReleaseRequestToPool(FHttpRequest* Request)
+	{
+		for (uint8 i = 0; i < Pool.Num(); ++i)
+		{
+			if (Pool[i].Request == Request)
+			{
+				// If only 1 user is remaining, we can give it to a waiter
+				// instead of releasing it back to the pool.
+				if (Pool[i].Usage == 1u)
+				{
+					if (FWaiter* Waiter = Waiters.dequeue())
+					{
+						Waiter->Request = Request;
+						Waiter->Trigger();
+						Waiter->Release();
+						return;
+					}
+				}
+				
+				Pool[i].Usage--;
+				return;
+			}
+		}
+		check(false);
+	}
+
+	/**
+	 * While holding a request, make it shared across many users.
+	 */
+	void MakeRequestShared(FHttpRequest* Request, uint8 Users)
+	{
+		check(Users != 0);
+		for (uint8 i = 0; i < Pool.Num(); ++i)
+		{
+			if (Pool[i].Request == Request)
+			{
+				Pool[i].Usage = Users;
+				return;
+			}
+		}
+		check(false);
+	}
+
+private:
+
+	struct FEntry
+	{
+		std::atomic<uint8> Usage;
+		FHttpRequest* Request;
+	};
+
+	TArray<FEntry> Pool;
+	FAAArrayQueue<FWaiter> Waiters;
+
+	FRequestPool() = delete;
+};
+
+//----------------------------------------------------------------------------------------------------------
+// FScopedRequestPtr
+//----------------------------------------------------------------------------------------------------------
+
+/**
+ * Utility class to manage requesting and releasing requests from the \ref FRequestPool.
+ */
+struct FScopedRequestPtr
+{
+public:
+	FScopedRequestPtr(FRequestPool* InPool)
+		: Request(InPool->WaitForFreeRequest())
+		, Pool(InPool)
+	{}
+
+	~FScopedRequestPtr()
+	{
+		Pool->ReleaseRequestToPool(Request);
+	}
+
+	bool IsValid() const 
+	{
+		return Request != nullptr;
+	}
+
+	FHttpRequest* Get() const
+	{
+		check(IsValid());
+		return Request;
+	}
+
+	FHttpRequest* operator->()
+	{
+		check(IsValid());
+		return Request;
+	}
+
+private:
+	FHttpRequest* Request;
+	FRequestPool* Pool;
+};
+
+
+#if WITH_DATAREQUEST_HELPER
+
+//----------------------------------------------------------------------------------------------------------
+// FDataRequestHelper
+//----------------------------------------------------------------------------------------------------------
+/**
+ * Helper class for requesting data. Will batch requests once the number of concurrent requests reach a threshold.
+ */
+struct FDataRequestHelper
+{
+	FDataRequestHelper(FRequestPool* InPool, const TCHAR* InNamespace, const TCHAR* InBucket, const TCHAR* InCacheKey, TArray<uint8>* OutData)
+		: Request(nullptr)
+		, Pool(InPool)
+		, bVerified(false, 1)
+	{
+		Request = Pool->GetFreeRequest();
+		if (Request && OutData != nullptr)
+		{
+			// We are below the threshold, make the connection immediately. OutData is set so this is a get.
+			FString Uri = FString::Printf(TEXT("api/v1/c/ddc/%s/%s/%s.raw"), InNamespace, InBucket, InCacheKey);
+			const FHttpRequest::Result Result = Request->PerformBlockingDownload(*Uri, OutData);
+			if (FHttpRequest::IsSuccessResponse(Request->GetResponseCode()))
+			{
+				if (VerifyRequest(Request, InNamespace, InBucket, InCacheKey, *OutData))
+				{
+					TRACE_COUNTER_ADD(HttpDDC_GetHit, int64(1));
+					TRACE_COUNTER_ADD(HttpDDC_BytesReceived, int64(Request->GetBytesReceived()));
+					bVerified[0] = true;
+				}
+			}
+		}
+		else if (Request)
+		{
+			// We are below the threshold, make the connection immediately. OutData is missing so this is a head.
+			FString Uri = FString::Printf(TEXT("api/v1/c/ddc/%s/%s/%s"), InNamespace, InBucket, InCacheKey);
+			const FHttpRequest::Result Result = Request->PerformBlockingQuery<FHttpRequest::Head>(*Uri);
+			if (FHttpRequest::IsSuccessResponse(Request->GetResponseCode()))
+			{
+				TRACE_COUNTER_ADD(HttpDDC_ExistHit, int64(1));
+				bVerified[0] = true;
+			}
+		}
+		else
+		{
+			// We have exceeded the threshold for concurrent connections, start or add this request
+			// to a batched request.
+			if (IsQueueCandidate(1, OutData && !OutData->IsEmpty()))
+			{
+				Request = QueueBatchRequest(
+					InPool,
+					InNamespace,
+					InBucket,
+					TConstArrayView<const TCHAR*>({InCacheKey}),
+					OutData ? TConstArrayView<TArray<uint8>*>({OutData}) : TConstArrayView<TArray<uint8>*>(),
+					bVerified
+				);
+			}
+
+			if (!Request)
+			{
+				Request = Pool->WaitForFreeRequest();
+
+				FQueuedBatchEntry Entry{
+					InNamespace,
+					InBucket,
+					TConstArrayView<const TCHAR*>({InCacheKey}),
+					OutData ? TConstArrayView<TArray<uint8>*>({OutData}) : TConstArrayView<TArray<uint8>*>(),
+					OutData && !OutData->IsEmpty() ? FHttpRequest::RequestVerb::Get : FHttpRequest::RequestVerb::Head,
+					&bVerified
+				};
+
+				PerformBatchQuery(Request, TArrayView<FQueuedBatchEntry>(&Entry, 1));
+			}
+		}
+	}
+
+	// Constructor specifically for batched head queries
+	FDataRequestHelper(FRequestPool* InPool, const TCHAR* InNamespace, const TCHAR* InBucket, TConstArrayView<FString> InCacheKeys)
+		: Request(nullptr)
+		, Pool(InPool)
+		, bVerified(false, InCacheKeys.Num())
+	{
+		// Transform the FString array to char pointers
+		TArray<const TCHAR*> CacheKeys;
+		Algo::Transform(InCacheKeys, CacheKeys, [](const FString& Key) { return *Key; });
+		
+		Request = Pool->GetFreeRequest();
+
+		if (Request || !IsQueueCandidate(InCacheKeys.Num(), false))
+		{
+			// If the request is too big for existing batches, wait for a free connection and create our own.
+			if (!Request)
+			{
+				Request = Pool->WaitForFreeRequest();
+			}
+
+			FQueuedBatchEntry Entry{
+				InNamespace, 
+				InBucket,
+				CacheKeys,
+				TConstArrayView<TArray<uint8>*>(),
+				FHttpRequest::RequestVerb::Head,
+				&bVerified
+			};
+
+			PerformBatchQuery(Request, TArrayView<FQueuedBatchEntry>(&Entry, 1));
+		}
+		else
+		{
+			Request = QueueBatchRequest(
+				InPool, 
+				InNamespace, 
+				InBucket, 
+				CacheKeys, 
+				TConstArrayView<TArray<uint8>*>(), 
+				bVerified
+			);
+
+			if (!Request)
+			{
+				Request = Pool->WaitForFreeRequest();
+
+				FQueuedBatchEntry Entry{
+					InNamespace,
+					InBucket,
+					CacheKeys,
+					TConstArrayView<TArray<uint8>*>(),
+					FHttpRequest::RequestVerb::Head,
+					&bVerified
+				};
+
+				PerformBatchQuery(Request, TArrayView<FQueuedBatchEntry>(&Entry, 1));
+			}
+		}
+	}
+
+	~FDataRequestHelper()
+	{
+		if (Request)
+		{
+			Pool->ReleaseRequestToPool(Request);
+		}
+	}
+
+	static void StaticInitialize()
+	{
+		static bool bInitialized = false;
+		check(!bInitialized);
+		for (FBatch& Batch : Batches)
+		{
+			Batch.Reserved = 0;
+			Batch.Ready = 0;
+			Batch.Complete = TUniquePtr<FEvent, FBatch::FEventDeleter>(FPlatformProcess::GetSynchEventFromPool(true));
+		}
+		bInitialized = true;
+	}
+
+	static void StaticShutdown()
+	{
+		for (FBatch& Batch : Batches)
+		{
+			Batch.Complete.Reset();
+		}
+	}
+
+	bool IsSuccess() const
+	{
+		return bVerified[0];
+	}
+
+	const TBitArray<>& IsBatchSuccess() const
+	{
+		return bVerified;
+	}
+
+	int64 GetResponseCode() const
+	{
+		return Request ? Request->GetResponseCode() : 0;
+	}
+
+private:
+
+	struct FQueuedBatchEntry
+	{
+		const TCHAR* Namespace;
+		const TCHAR* Bucket;
+		TConstArrayView<const TCHAR*> CacheKeys;
+		TConstArrayView<TArray<uint8>*> OutDatas;
+		FHttpRequest::RequestVerb Verb;
+		TBitArray<>* bSuccess;
+	};
+
+	struct FBatch
+	{
+		struct FEventDeleter
+		{
+			void operator()(FEvent* Event)
+			{
+				FPlatformProcess::ReturnSynchEventToPool(Event);
+			}
+		};
+
+		FQueuedBatchEntry Entries[UE_HTTPDDC_BATCH_SIZE];
+		std::atomic<uint32> Reserved;
+		std::atomic<uint32> Ready;
+		std::atomic<uint32> WeightHint;
+		FHttpRequest* Request;
+		TUniquePtr<FEvent, FEventDeleter> Complete;
+	};
+
+	FHttpRequest* Request;
+	FRequestPool* Pool;
+	TBitArray<> bVerified;
+	static std::atomic<uint32> FirstAvailableBatch;
+	static TStaticArray<FBatch, UE_HTTPDDC_BATCH_NUM> Batches;
+
+	static uint32 ComputeWeight(int32 NumKeys, bool bHasDatas)
+	{
+		return NumKeys * (bHasDatas ? UE_HTTPDDC_BATCH_GET_WEIGHT : UE_HTTPDDC_BATCH_HEAD_WEIGHT);
+	}
+
+	static bool IsQueueCandidate(int32 NumKeys, bool bHasDatas)
+	{
+		if (NumKeys > UE_HTTPDDC_BATCH_SIZE)
+		{
+			return false;
+		}
+		const uint32 Weight = ComputeWeight(NumKeys, bHasDatas);
+		if (Weight > UE_HTTPDDC_BATCH_WEIGHT_HINT)
+		{
+			return false;
+		}
+		return true;
+	}
+
+	/**
+	 * Queues up a request to be batched. Blocks until the query is made.
+	 */
+	static FHttpRequest* QueueBatchRequest(FRequestPool* InPool, 
+		const TCHAR* InNamespace, 
+		const TCHAR* InBucket, 
+		TConstArrayView<const TCHAR*> InCacheKeys,
+		TConstArrayView<TArray<uint8>*> OutDatas, 
+		TBitArray<>& bOutVerified)
+	{
+		TRACE_CPUPROFILER_EVENT_SCOPE(HttpDDC_BatchQuery);
+		check(InCacheKeys.Num() == OutDatas.Num() || OutDatas.Num() == 0);
+		const uint32 RequestNum = InCacheKeys.Num();
+		const uint32 RequestWeight = ComputeWeight(InCacheKeys.Num(), !OutDatas.IsEmpty());
+
+		for (int32 i = 0; i < Batches.Num(); i++)
+		{
+			uint32 Index = (FirstAvailableBatch.load(std::memory_order_relaxed) + i) % Batches.Num();
+			FBatch& Batch = Batches[Index];
+
+			//Assign different weights to head vs. get queries
+			if (Batch.WeightHint.load(std::memory_order_acquire) + RequestWeight > UE_HTTPDDC_BATCH_WEIGHT_HINT)
+			{
+				continue;
+			}
+
+			// Attempt to reserve a spot in the batch
+			const uint32 Reserve = Batch.Reserved.fetch_add(1, std::memory_order_acquire);
+			if (Reserve >= UE_HTTPDDC_BATCH_SIZE)
+			{
+				// We didn't manage to snag a valid reserve index try next batch
+				continue;
+			}
+
+			// Add our weight to the batch. Note we are treating it as a hint, so don't syncronize.
+			const uint32 ActualWeight = Batch.WeightHint.fetch_add(RequestWeight, std::memory_order_release);
+
+			TAnsiStringBuilder<64> BatchString;
+			BatchString << "HttpDDC_Batch" << Index;
+			TRACE_CPUPROFILER_EVENT_SCOPE_TEXT(*BatchString);
+
+			if (Reserve == (UE_HTTPDDC_BATCH_SIZE - 1))
+			{
+				FirstAvailableBatch++;
+			}
+
+			Batch.Entries[Reserve] = FQueuedBatchEntry{
+				InNamespace,
+				InBucket,
+				InCacheKeys,
+				OutDatas,
+				OutDatas.Num() ? FHttpRequest::RequestVerb::Get : FHttpRequest::RequestVerb::Head,
+				&bOutVerified
+			};
+
+			// Signal we are ready for batch to be submitted
+			Batch.Ready.fetch_add(1u, std::memory_order_release);
+
+			FHttpRequest* Request = nullptr;
+
+			// The first to reserve a slot is the "driver" of the batch
+			if (Reserve == 0)
+			{
+				Batch.Request = InPool->WaitForFreeRequest();
+
+				// Make sure no new requests are added
+				const uint32 Reserved = FMath::Min((uint32)UE_HTTPDDC_BATCH_SIZE, Batch.Reserved.fetch_add(UE_HTTPDDC_BATCH_SIZE, std::memory_order_acquire));
+
+				// Give other threads time to copy their data to batch
+				while (Batch.Ready.load(std::memory_order_acquire) < Reserved) {}
+
+				// Increment request ref count to reflect all waiting threads
+				InPool->MakeRequestShared(Batch.Request, Reserved);
+
+				// Do the actual query and write response to respective target arrays
+				PerformBatchQuery(Batch.Request, TArrayView<FQueuedBatchEntry>(Batch.Entries, Batch.Ready));
+
+				// Signal to waiting threads the batch is complete
+				Batch.Complete->Trigger();
+
+				// Store away the request and wait until other threads have too
+				Request = Batch.Request;
+				while (Batch.Ready.load(std::memory_order_acquire) > 1) {}
+
+				//Reset batch for next use
+				Batch.Complete->Reset();
+				Batch.WeightHint.store(0, std::memory_order_release);
+				Batch.Ready.store(0, std::memory_order_release);
+				Batch.Reserved.store(0, std::memory_order_release);
+			}
+			else
+			{
+				// Wait until "driver" has done query
+				{
+					TRACE_CPUPROFILER_EVENT_SCOPE(WaitForMasterOfBatch);
+					Batch.Complete->Wait(~0);
+				}
+
+				// Store away request and signal we are done
+				Request = Batch.Request;
+				Batch.Ready.fetch_sub(1u, std::memory_order_release);
+			}
+
+			return Request;
+		}
+
+		return nullptr;
+	}
+
+
+	/**
+	 * Creates request uri and headers and submits the request
+	 */
+	static void PerformBatchQuery(FHttpRequest* Request, TArrayView<FQueuedBatchEntry> Entries)
+	{
+		TRACE_CPUPROFILER_EVENT_SCOPE(HttpDDC_BatchGet);
+		const TCHAR* Uri(TEXT("api/v1/c/ddc-rpc/batchget"));
+		int64 ResponseCode = 0; uint32 Attempts = 0;
+
+		//Prepare request object
+		TArray<TSharedPtr<FJsonValue>> Operations;
+		for (const FQueuedBatchEntry& Entry : Entries)
+		{
+			for (int32 KeyIdx = 0; KeyIdx < Entry.CacheKeys.Num(); KeyIdx++)
+			{
+				TSharedPtr<FJsonObject> Object = MakeShared<FJsonObject>();
+				Object->SetField(TEXT("bucket"), MakeShared<FJsonValueString>(Entry.Bucket));
+				Object->SetField(TEXT("key"), MakeShared<FJsonValueString>(Entry.CacheKeys[KeyIdx]));
+				if (Entry.Verb == FHttpRequest::RequestVerb::Head)
+				{
+					Object->SetField(TEXT("verb"), MakeShared<FJsonValueString>(TEXT("HEAD")));
+				}
+				Operations.Add(MakeShared<FJsonValueObject>(Object));
+			}
+		}
+		TSharedPtr<FJsonObject> RequestObject = MakeShared<FJsonObject>();
+		RequestObject->SetField(TEXT("namespace"), MakeShared<FJsonValueString>(Entries[0].Namespace));
+		RequestObject->SetField(TEXT("operations"), MakeShared<FJsonValueArray>(Operations));
+
+		//Serialize to a buffer
+		FBufferArchive RequestData;
+		if (FJsonSerializer::Serialize(RequestObject.ToSharedRef(), TJsonWriterFactory<ANSICHAR, TCondensedJsonPrintPolicy<ANSICHAR>>::Create(&RequestData)))
+		{
+			Request->PerformBlockingUpload<FHttpRequest::PostJson>(Uri, MakeArrayView(RequestData));
+			ResponseCode = Request->GetResponseCode();
+
+			if (ResponseCode == 200)
+			{
+				const TArray<uint8>& ResponseBuffer = Request->GetResponseBuffer();
+				const uint8* Response = ResponseBuffer.GetData();
+				const int32 ResponseSize = ResponseBuffer.Num();
+
+				// Parse the response and move the data to the target requests.
+				if (ParseBatchedResponse(Response, ResponseSize, Entries))
+				{
+					UE_LOG(LogDerivedDataCache, VeryVerbose, TEXT("%s: Batch query with %d operations completed."), *Request->GetName(), Entries.Num());
+					return;
+				}
+			}
+		}
+		
+		// If we get here the request failed.
+		UE_LOG(LogDerivedDataCache, Display, TEXT("%s: Batch query failed. Query: %s"), *Request->GetName(), ANSI_TO_TCHAR((ANSICHAR*)RequestData.GetData()));
+
+		// Set all batch operations to failures
+		for (FQueuedBatchEntry Entry : Entries)
+		{
+			Entry.bSuccess->SetRange(0, Entry.CacheKeys.Num(), false);
+		}
+	}
+
+	// Above result value
+	enum class OpResult : uint8
+	{
+		Ok = 0,			// Op finished succesfully
+		Error = 1,		// Error during op
+		NotFound = 2,	// Key was not found
+		Exists = 3		// Used to indicate head op success
+	};
+
+	// Searches for potentially multiple key requests that are satisfied the given cache key result
+	// Search strategy is exhaustive forward search from the last found entry.  If the results come in ordered the same as the requests,
+	//  and there are no duplicates, the search will be somewhat efficient (still has to do exhaustive searching looking for duplicates).
+	//  If the results are unordered or there are duplicates, search will become more inefficient.
+	struct FRequestSearchHelper
+	{
+		FRequestSearchHelper(TArrayView<FQueuedBatchEntry> InRequests, const FUTF8ToTCHAR& InCacheKey, int32 InEntryIdx, int32 InKeyIdx, OpResult InRequestResult)
+			: Requests(InRequests)
+			, CacheKey(InCacheKey)
+			, StartEntryIdx(InEntryIdx)
+			, StartKeyIdx(InKeyIdx)
+			, RequestResult(InRequestResult)
+		{}
+
+		bool FindNext(int32& EntryIdx, int32& KeyIdx)
+		{
+			int32 CurrentEntryIdx = EntryIdx;
+			int32 CurrentKeyIdx = KeyIdx;
+			do
+			{
+				// Do not match a get request with a head response code (i.e. Exists) 
+				// or a head request with a get response code (i.e. Ok)
+				// if the response code is an error or not found they can be matched to both head or get request it doesn't matter
+				const FQueuedBatchEntry& CurrentRequest = Requests[CurrentEntryIdx];
+				bool bRequestTypeMatch = !((CurrentRequest.Verb == FHttpRequest::Get) && (RequestResult == OpResult::Exists))
+					&& !((CurrentRequest.Verb == FHttpRequest::Head) && (RequestResult == OpResult::Ok));
+				if (bRequestTypeMatch && FCString::Stricmp(CurrentRequest.CacheKeys[CurrentKeyIdx], CacheKey.Get()) == 0)
+				{
+					EntryIdx = CurrentEntryIdx;
+					KeyIdx = CurrentKeyIdx;
+					return true;
+				}
+			} while (AdvanceIndices(CurrentEntryIdx, CurrentKeyIdx));
+
+			return false;
+		}
+
+		bool AdvanceIndices(int32& EntryIdx, int32& KeyIdx)
+		{
+			if (++KeyIdx >= Requests[EntryIdx].CacheKeys.Num())
+			{
+				EntryIdx = (EntryIdx + 1) % Requests.Num();
+				KeyIdx = 0;
+			}
+
+			return !((EntryIdx == StartEntryIdx) && (KeyIdx == StartKeyIdx));
+		}
+
+		TArrayView<FQueuedBatchEntry> Requests;
+		const FUTF8ToTCHAR& CacheKey;
+		int32 StartEntryIdx;
+		int32 StartKeyIdx;
+		OpResult RequestResult;
+	};
+
+	/**
+	 * Parses a batched response stream, moves the data to target requests and marks them with result.
+	 * @param Response Pointer to Response buffer
+	 * @param ResponseSize Size of response buffer
+	 * @param Requests Requests that will be filled with data.
+	 * @return True if response was successfully parsed, false otherwise.
+	 */
+	static bool ParseBatchedResponse(const uint8* ResponseStart, const int32 ResponseSize, TArrayView<FQueuedBatchEntry> Requests)
+	{
+		// The expected data stream is structured accordingly
+		// {"JPTR"} {PayloadCount:uint32} {{"JPEE"} {Name:cstr} {Result:uint8} {Hash:IoHash} {Size:uint64} {Payload...}} ...
+
+		const TCHAR ResponseErrorMessage[] = TEXT("Malformed response from server.");
+		const ANSICHAR* ProtocolMagic = "JPTR";
+		const ANSICHAR* PayloadMagic = "JPEE";
+		const uint32 MagicSize = 4;
+		const uint8* Response = ResponseStart;
+		const uint8* ResponseEnd = Response + ResponseSize;
+
+		// Check that the stream starts with the protocol magic
+		if (FMemory::Memcmp(ProtocolMagic, Response, MagicSize) != 0)
+		{
+			UE_LOG(LogDerivedDataCache, Display, ResponseErrorMessage);
+			return false;
+		}
+		Response += MagicSize;
+
+		// Number of payloads recieved
+		uint32 PayloadCount = *(uint32*)Response;
+		Response += sizeof(uint32);
+
+		uint32 PayloadIdx = 0; 	// Current processed result
+		int32 EntryIdx = 0; 	// Current Entry index
+		int32 KeyIdx = 0; 		// Current Key index for current Entry
+
+		while (Response < ResponseEnd && FMemory::Memcmp(PayloadMagic, Response, MagicSize) == 0)
+		{
+			PayloadIdx++;
+			Response += MagicSize;
+
+			const ANSICHAR* PayloadNameA = (const ANSICHAR*)Response;
+			Response += FCStringAnsi::Strlen(PayloadNameA) + 1; //String and zero termination
+			const ANSICHAR* CacheKeyA = FCStringAnsi::Strrchr(PayloadNameA, '.') + 1; // "namespace.bucket.cachekey"
+
+			// Result of the operation is used to match to the appropriate request (i.e. get or head)
+			OpResult PayloadResult = static_cast<OpResult>(*Response);
+			Response += sizeof(uint8);
+
+			const uint8* ResponseRewindMark = Response;
+
+			// Find the payload among the requests.  Payloads may be returned in any order and if the same cache key was part of two requests,
+			// a single payload may satisfy multiple cache keys in multiple requests.
+			FUTF8ToTCHAR CacheKey(CacheKeyA);
+			FRequestSearchHelper RequestSearch(Requests, CacheKey, EntryIdx, KeyIdx, PayloadResult);
+			bool bFoundAny = false;
+
+			while (RequestSearch.FindNext(EntryIdx, KeyIdx))
+			{
+				Response = ResponseRewindMark;
+				bFoundAny = true;
+
+				FQueuedBatchEntry& RequestOp = Requests[EntryIdx];
+				TBitArray<>& bSuccess = *RequestOp.bSuccess;
+
+				switch (PayloadResult)
+				{
+				case OpResult::Ok:
+					{
+						// Payload hash of the following payload data
+						FIoHash PayloadHash = *(FIoHash*)Response;
+						Response += sizeof(FIoHash);
+
+						// Size of the following payload data
+						const uint64 PayloadSize = *(uint64*)Response;
+						Response += sizeof(uint64);
+
+						if (PayloadSize > 0)
+						{
+							if (Response + PayloadSize > ResponseEnd)
+							{
+								UE_LOG(LogDerivedDataCache, Display, ResponseErrorMessage);
+								return false;
+							}
+
+							if (bSuccess[KeyIdx])
+							{
+								Response += PayloadSize;
+							}
+							else
+							{
+								TArray<uint8>* OutData = RequestOp.OutDatas[KeyIdx];
+
+								OutData->Append(Response, PayloadSize);
+								Response += PayloadSize;
+								// Verify the received and parsed payload
+								if (VerifyPayload(PayloadHash, RequestOp.Namespace, RequestOp.Bucket, RequestOp.CacheKeys[KeyIdx], *OutData))
+								{
+									TRACE_COUNTER_ADD(HttpDDC_GetHit, int64(1));
+									TRACE_COUNTER_ADD(HttpDDC_BytesReceived, int64(PayloadSize));
+									
+									bSuccess[KeyIdx] = true;
+								}
+								else
+								{
+									OutData->Empty();
+									bSuccess[KeyIdx] = false;
+								}
+							}
+						}
+						else
+						{
+							bSuccess[KeyIdx] = false;
+						}
+					}
+					break;
+
+				case OpResult::Exists:
+					{
+						TRACE_COUNTER_ADD(HttpDDC_ExistHit, int64(1));
+						bSuccess[KeyIdx] = true;
+					}
+					break;
+
+				default:
+				case OpResult::Error:
+					UE_LOG(LogDerivedDataCache, Display, TEXT("Server error while getting %s"), CacheKey.Get());
+					// intentional falltrough
+
+				case OpResult::NotFound:
+					bSuccess[KeyIdx] = false;
+					break;
+
+				}
+
+				if (!RequestSearch.AdvanceIndices(EntryIdx, KeyIdx))
+				{
+					break;
+				}
+			}
+
+			if (!bFoundAny)
+			{
+				UE_LOG(LogDerivedDataCache, Error, ResponseErrorMessage);
+				return false;
+			}
+		}
+
+		// Have we parsed all the payloads from the message?
+		if (PayloadIdx != PayloadCount)
+		{
+			UE_LOG(LogDerivedDataCache, Display, TEXT("%s: Found %d payloads but %d was reported."), ResponseErrorMessage, PayloadIdx, PayloadCount);
+		}
+
+		return true;
+	}
+};
+
+TStaticArray<FDataRequestHelper::FBatch, UE_HTTPDDC_BATCH_NUM> FDataRequestHelper::Batches;
+std::atomic<uint32> FDataRequestHelper::FirstAvailableBatch;
+
+//----------------------------------------------------------------------------------------------------------
+// FDataUploadHelper
+//----------------------------------------------------------------------------------------------------------
+struct FDataUploadHelper
+{
+	FDataUploadHelper(FRequestPool* InPool, 
+		const TCHAR* InNamespace, 
+		const TCHAR* InBucket, 
+		const TCHAR* InCacheKey, 
+		const TArrayView<const uint8>& InData,
+		FDerivedDataCacheUsageStats& InUsageStats)
+		: ResponseCode(0)
+		, bSuccess(false)
+		, bQueued(false)
+	{
+		FHttpRequest* Request = InPool->GetFreeRequest();
+		if (Request)
+		{
+			ResponseCode = PerformPut(Request, InNamespace, InBucket, InCacheKey, InData, InUsageStats);
+			bSuccess = FHttpRequest::IsSuccessResponse(Request->GetResponseCode());
+
+			ProcessQueuedPutsAndReleaseRequest(InPool, Request, InUsageStats);
+		}
+		else
+		{
+			FQueuedEntry* Entry = new FQueuedEntry(InNamespace, InBucket, InCacheKey, InData);
+			QueuedPuts.Push(Entry);
+			bSuccess = true;
+			bQueued = true;
+			
+			// A request may have been released while the entry was being queued.
+			Request = InPool->GetFreeRequest();
+			if (Request)
+			{
+				ProcessQueuedPutsAndReleaseRequest(InPool, Request, InUsageStats);
+			}
+		}
+	}
+
+	bool IsSuccess() const
+	{
+		return bSuccess;
+	}
+
+	int64 GetResponseCode() const
+	{
+		return ResponseCode;
+	}
+
+	bool IsQueued() const
+	{
+		return bQueued;
+	}
+
+private:
+
+	struct FQueuedEntry
+	{
+		FString Namespace;
+		FString Bucket;
+		FString CacheKey;
+		TArray<uint8> Data;
+
+		FQueuedEntry(const TCHAR* InNamespace, const TCHAR* InBucket, const TCHAR* InCacheKey, const TArrayView<const uint8> InData)
+			: Namespace(InNamespace)
+			, Bucket(InBucket)
+			, CacheKey(InCacheKey)
+			, Data(InData) // Copies the data!
+		{}
+	};
+
+	static TLockFreePointerListUnordered<FQueuedEntry, PLATFORM_CACHE_LINE_SIZE> QueuedPuts;
+
+	int64 ResponseCode;
+	bool bSuccess;
+	bool bQueued;
+
+	static void ProcessQueuedPutsAndReleaseRequest(FRequestPool* Pool, FHttpRequest* Request, FDerivedDataCacheUsageStats& UsageStats)
+	{
+		while (Request)
+		{
+			// Make sure that whether we early exit or execute past the end of this scope that
+			// the request is released back to the pool.
+			{
+				ON_SCOPE_EXIT
+				{
+					Pool->ReleaseRequestToPool(Request);
+				};
+
+				if (ShouldAbortForShutdown())
+				{
+					return;
+				}
+
+				while (FQueuedEntry* Entry = QueuedPuts.Pop())
+				{
+					Request->Reset();
+					PerformPut(Request, *Entry->Namespace, *Entry->Bucket, *Entry->CacheKey, Entry->Data, UsageStats);
+					delete Entry;
+
+					if (ShouldAbortForShutdown())
+					{
+						return;
+					}
+				}
+			}
+
+			// An entry may have been queued while the request was being released.
+			if (QueuedPuts.IsEmpty())
+			{
+				break;
+			}
+
+			// Process the queue again if a request is free, otherwise the thread that got the request will process it.
+			Request = Pool->GetFreeRequest();
+		}
+	}
+
+	static int64 PerformPut(FHttpRequest* Request, const TCHAR* Namespace, const TCHAR* Bucket, const TCHAR* CacheKey, const TArrayView<const uint8> Data, FDerivedDataCacheUsageStats& UsageStats)
+	{
+		COOK_STAT(auto Timer = UsageStats.TimePut());
+
+		HashPayload(Request, Data);
+
+		TStringBuilder<256> Uri;
+		Uri.Appendf(TEXT("api/v1/c/ddc/%s/%s/%s"), Namespace, Bucket, CacheKey);
+
+		Request->PerformBlockingUpload<FHttpRequest::Put>(*Uri, Data);
+
+		const int64 ResponseCode = Request->GetResponseCode();
+		if (FHttpRequest::IsSuccessResponse(ResponseCode))
+		{
+			TRACE_COUNTER_ADD(HttpDDC_BytesSent, int64(Request->GetBytesSent()));
+			COOK_STAT(Timer.AddHit(Request->GetBytesSent()));
+		}
+
+		return Request->GetResponseCode();
+	}
+};
+
+TLockFreePointerListUnordered<FDataUploadHelper::FQueuedEntry, PLATFORM_CACHE_LINE_SIZE> FDataUploadHelper::QueuedPuts;
+
+#endif // WITH_DATAREQUEST_HELPER
+
+//----------------------------------------------------------------------------------------------------------
+// Certificate checking
+//----------------------------------------------------------------------------------------------------------
+
+#if WITH_SSL
+
+static int SslCertVerify(int PreverifyOk, X509_STORE_CTX* Context)
+{
+	if (PreverifyOk == 1)
+	{
+		SSL* Handle = static_cast<SSL*>(X509_STORE_CTX_get_ex_data(Context, SSL_get_ex_data_X509_STORE_CTX_idx()));
+		check(Handle);
+
+		SSL_CTX* SslContext = SSL_get_SSL_CTX(Handle);
+		check(SslContext);
+
+		FHttpRequest* Request = static_cast<FHttpRequest*>(SSL_CTX_get_app_data(SslContext));
+		check(Request);
+
+		const FString& Domain = Request->GetDomain();
+
+		if (!FSslModule::Get().GetCertificateManager().VerifySslCertificates(Context, Domain))
+		{
+			PreverifyOk = 0;
+		}
+	}
+
+	return PreverifyOk;
+}
+
+static CURLcode sslctx_function(CURL * curl, void * sslctx, void * parm)
+{
+	SSL_CTX* Context = static_cast<SSL_CTX*>(sslctx);
+	const ISslCertificateManager& CertificateManager = FSslModule::Get().GetCertificateManager();
+
+	CertificateManager.AddCertificatesToSslContext(Context);
+	SSL_CTX_set_verify(Context, SSL_CTX_get_verify_mode(Context), SslCertVerify);
+	SSL_CTX_set_app_data(Context, parm);
+
+	/* all set to go */
+	return CURLE_OK;
+}
+
+#endif //#if WITH_SSL
+
+//----------------------------------------------------------------------------------------------------------
+// Content parsing and checking
+//----------------------------------------------------------------------------------------------------------
+
+/**
+ * Verifies the integrity of the received data using supplied checksum.
+ * @param Hash received hash value.
+ * @param Namespace The namespace string used when originally fetching the request.
+ * @param Bucket The bucket string used when originally fetching the request.
+ * @param CacheKey The cache key string used when originally fetching the request.
+ * @param Payload Payload received.
+ * @return True if the data is correct, false if checksums doesn't match.
+ */
+bool VerifyPayload(const FSHAHash& Hash, const TCHAR* Namespace, const TCHAR* Bucket, const TCHAR* CacheKey, const TArray<uint8>& Payload)
+{
+	FSHAHash PayloadHash;
+	FSHA1::HashBuffer(Payload.GetData(), Payload.Num(), PayloadHash.Hash);
+
+	if (Hash != PayloadHash)
+	{
+		UE_LOG(LogDerivedDataCache,
+			Display,
+			TEXT("Checksum from server did not match received data (%s vs %s). Discarding cached result. Namespace: %s, Bucket: %s, Key: %s."),
+			*WriteToString<48>(Hash),
+			*WriteToString<48>(PayloadHash),
+			Namespace,
+			Bucket,
+			CacheKey
+		);
+		return false;
+	}
+
+	return true;
+}
+
+/**
+ * Verifies the integrity of the received data using supplied checksum.
+ * @param Hash received hash value.
+ * @param Namespace The namespace string used when originally fetching the request.
+ * @param Bucket The bucket string used when originally fetching the request.
+ * @param CacheKey The cache key string used when originally fetching the request.
+ * @param Payload Payload received.
+ * @return True if the data is correct, false if checksums doesn't match.
+ */
+bool VerifyPayload(const FIoHash& Hash, const TCHAR* Namespace, const TCHAR* Bucket, const TCHAR* CacheKey, const TArray<uint8>& Payload)
+{
+	FIoHash PayloadHash = FIoHash::HashBuffer(Payload.GetData(), Payload.Num());
+
+	if (Hash != PayloadHash)
+	{
+		UE_LOG(LogDerivedDataCache,
+			Display,
+			TEXT("Checksum from server did not match received data (%s vs %s). Discarding cached result. Namespace: %s, Bucket: %s, Key: %s."),
+			*WriteToString<48>(Hash),
+			*WriteToString<48>(PayloadHash),
+			Namespace,
+			Bucket,
+			CacheKey
+		);
+		return false;
+	}
+
+	return true;
+}
+
+
+/**
+ * Verifies the integrity of the received data using supplied checksum.
+ * @param Request Request that the data was be received with.
+ * @param Namespace The namespace string used when originally fetching the request.
+ * @param Bucket The bucket string used when originally fetching the request.
+ * @param CacheKey The cache key string used when originally fetching the request.
+ * @param Payload Payload received.
+ * @return True if the data is correct, false if checksums doesn't match.
+ */
+bool VerifyRequest(const FHttpRequest* Request, const TCHAR* Namespace, const TCHAR* Bucket, const TCHAR* CacheKey, const TArray<uint8>& Payload)
+{
+	FString ReceivedHashStr;
+	if (Request->GetHeader("X-Jupiter-Sha1", ReceivedHashStr))
+	{
+		FSHAHash ReceivedHash;
+		ReceivedHash.FromString(ReceivedHashStr);
+		return VerifyPayload(ReceivedHash, Namespace, Bucket, CacheKey, Payload);
+	}
+	if (Request->GetHeader("X-Jupiter-IoHash", ReceivedHashStr))
+	{
+		FIoHash ReceivedHash(ReceivedHashStr);
+		return VerifyPayload(ReceivedHash, Namespace, Bucket, CacheKey, Payload);
+	}
+	UE_LOG(LogDerivedDataCache, Warning, TEXT("%s: HTTP server did not send a content hash. Wrong server version?"), *Request->GetName());
+	return true;
+}
+
+/**
+ * Adds a checksum (as request header) for a given payload. Jupiter will use this to verify the integrity
+ * of the received data.
+ * @param Request Request that the data will be sent with.
+ * @param Payload Payload that will be sent.
+ * @return True on success, false on failure.
+ */
+bool HashPayload(FHttpRequest* Request, const TArrayView<const uint8> Payload)
+{
+	FIoHash PayloadHash = FIoHash::HashBuffer(Payload.GetData(), Payload.Num());
+	Request->SetHeader(TEXT("X-Jupiter-IoHash"), *WriteToString<48>(PayloadHash));
+	return true;
+}
+
+bool ShouldAbortForShutdown()
+{
+	return !GIsBuildMachine && FDerivedDataBackend::Get().IsShuttingDown();
+}
+
+TConstArrayView<uint8> MakeConstArrayView(FSharedBuffer Buffer)
+{
+	return TConstArrayView<uint8>(reinterpret_cast<const uint8*>(Buffer.GetData()), Buffer.GetSize());
+}
+
+static bool IsValueDataReady(FValue& Value, const ECachePolicy Policy)
+{
+	if (!EnumHasAnyFlags(Policy, ECachePolicy::Query))
+	{
+		Value = Value.RemoveData();
+		return true;
+	}
+
+	if (Value.HasData())
+	{
+		if (EnumHasAnyFlags(Policy, ECachePolicy::SkipData))
+		{
+			Value = Value.RemoveData();
+		}
+		return true;
+	}
+	return false;
+};
+
+//----------------------------------------------------------------------------------------------------------
+// FHttpAccessToken
+//----------------------------------------------------------------------------------------------------------
+
+FString FHttpAccessToken::GetHeader()
+{
+	Lock.ReadLock();
+	FString Header = FString::Printf(TEXT("Authorization: Bearer %s"), *Token);
+	Lock.ReadUnlock();
+	return Header;
+}
+
+void FHttpAccessToken::SetHeader(const TCHAR* InToken)
+{
+	Lock.WriteLock();
+	Token = InToken;
+	Serial++;
+	Lock.WriteUnlock();
+}
+
+uint32 FHttpAccessToken::GetSerial() const
+{
+	return Serial;
+}
+
+//----------------------------------------------------------------------------------------------------------
+// FHttpCacheStore
+//----------------------------------------------------------------------------------------------------------
+
+/**
+ * Backend for a HTTP based caching service (Jupiter).
+ **/
+class FHttpCacheStore final : public FDerivedDataBackendInterface
+{
+public:
+	
+	/**
+	 * Creates the backend, checks health status and attempts to acquire an access token.
+	 *
+	 * @param ServiceUrl			Base url to the service including schema.
+	 * @param Namespace				Namespace to use.
+	 * @param StructuredNamespace	Namespace to use for structured cache operations.
+	 * @param OAuthProvider			Url to OAuth provider, for example "https://myprovider.com/oauth2/v1/token".
+	 * @param OAuthClientId			OAuth client identifier.
+	 * @param OAuthData				OAuth form data to send to login service. Can either be the raw form data or a Windows network file address (starting with "\\").
+	 */
+	FHttpCacheStore(
+		const TCHAR* ServiceUrl, 
+		const TCHAR* Namespace, 
+		const TCHAR* StructuredNamespace, 
+		const TCHAR* OAuthProvider, 
+		const TCHAR* OAuthClientId, 
+		const TCHAR* OAuthData,
+		EBackendLegacyMode LegacyMode,
+		bool bReadOnly);
+
+	~FHttpCacheStore();
+
+	/**
+	 * Checks is backend is usable (reachable and accessible).
+	 * @return true if usable
+	 */
+	bool IsUsable() const { return bIsUsable; }
+
+	/** return true if this cache is writable **/
+	virtual bool IsWritable() const override
+	{
+		return !bReadOnly && bIsUsable;
+	}
+
+	virtual bool CachedDataProbablyExists(const TCHAR* CacheKey) override;
+	virtual TBitArray<> CachedDataProbablyExistsBatch(TConstArrayView<FString> CacheKeys) override;
+	virtual bool GetCachedData(const TCHAR* CacheKey, TArray<uint8>& OutData) override;
+	virtual EPutStatus PutCachedData(const TCHAR* CacheKey, TArrayView<const uint8> InData, bool bPutEvenIfExists) override;
+	virtual void RemoveCachedData(const TCHAR* CacheKey, bool bTransient) override;
+	virtual TSharedRef<FDerivedDataCacheStatsNode> GatherUsageStats() const override;
+	
+	virtual FString GetName() const override;
+	virtual TBitArray<> TryToPrefetch(TConstArrayView<FString> CacheKeys) override;
+	virtual bool WouldCache(const TCHAR* CacheKey, TArrayView<const uint8> InData) override;
+	virtual ESpeedClass GetSpeedClass() const override;
+	virtual bool ApplyDebugOptions(FBackendDebugOptions& InOptions) override;
+
+	void SetSpeedClass(ESpeedClass InSpeedClass) { SpeedClass = InSpeedClass; }
+
+	EBackendLegacyMode GetLegacyMode() const final { return LegacyMode; }
+
+	virtual void Put(
+		TConstArrayView<FCachePutRequest> Requests,
+		IRequestOwner& Owner,
+		FOnCachePutComplete&& OnComplete) override;
+
+	virtual void Get(
+		TConstArrayView<FCacheGetRequest> Requests,
+		IRequestOwner& Owner,
+		FOnCacheGetComplete&& OnComplete) override;
+	virtual void PutValue(
+		TConstArrayView<FCachePutValueRequest> Requests,
+		IRequestOwner& Owner,
+		FOnCachePutValueComplete&& OnComplete) override;
+	virtual void GetValue(
+		TConstArrayView<FCacheGetValueRequest> Requests,
+		IRequestOwner& Owner,
+		FOnCacheGetValueComplete&& OnComplete) override;
+	virtual void GetChunks(
+		TConstArrayView<FCacheGetChunkRequest> Requests,
+		IRequestOwner& Owner,
+		FOnCacheGetChunkComplete&& OnComplete) override;
+
+	static FHttpCacheStore* GetAny()
+	{
+		return AnyInstance;
+	}
+
+	const FString& GetDomain() const { return Domain; }
+	const FString& GetNamespace() const { return Namespace; }
+	const FString& GetStructuredNamespace() const { return StructuredNamespace; }
+	const FString& GetOAuthProvider() const { return OAuthProvider; }
+	const FString& GetOAuthClientId() const { return OAuthClientId; }
+	const FString& GetOAuthSecret() const { return OAuthSecret; }
+
+private:
+	FString Domain;
+	FString Namespace;
+	FString StructuredNamespace;
+	FString DefaultBucket;
+	FString OAuthProvider;
+	FString OAuthClientId;
+	FString OAuthSecret;
+	FCriticalSection AccessCs;
+	FDerivedDataCacheUsageStats UsageStats;
+	FBackendDebugOptions DebugOptions;
+	FCriticalSection MissedKeysCS;
+	TSet<FName> DebugMissedKeys;
+	TSet<FCacheKey> DebugMissedCacheKeys;
+	TUniquePtr<struct FRequestPool> GetRequestPools[2];
+	TUniquePtr<struct FRequestPool> PutRequestPools[2];
+	TUniquePtr<struct FHttpAccessToken> Access;
+	bool bIsUsable;
+	bool bReadOnly;
+	uint32 FailedLoginAttempts;
+	ESpeedClass SpeedClass;
+	EBackendLegacyMode LegacyMode;
+	static inline FHttpCacheStore* AnyInstance = nullptr;
+
+	struct FValueDebugContext
+	{
+		FStringView Name;
+		const FCacheKey& Key;
+		FString Id;
+	};
+
+	bool IsServiceReady();
+	bool AcquireAccessToken();
+	bool ShouldRetryOnError(int64 ResponseCode);
+	bool ShouldSimulateMiss(const TCHAR* InKey);
+	bool ShouldSimulateMiss(const FCacheKey& Key);
+
+	bool PutCacheRecord(FStringView Name, const FCacheRecord& Record, const FCacheRecordPolicy& Policy, uint64& OutWriteSize);
+	uint64 PutRef(const FCbPackage& Package, const FCacheKey& Key, FStringView Bucket, bool bFinalize, TArray<FIoHash>& OutNeededBlobHashes, bool& bOutPutCompletedSuccessfully);
+
+	FOptionalCacheRecord GetCacheRecordOnly(
+		FStringView Name,
+		const FCacheKey& Key,
+		const FCacheRecordPolicy& Policy);
+	FOptionalCacheRecord GetCacheRecord(
+		FStringView Name,
+		const FCacheKey& Key,
+		const FCacheRecordPolicy& Policy,
+		EStatus& OutStatus);
+
+	bool PutCacheValue(
+		FStringView Name,
+		const FCacheKey& Key,
+		const FValue& Value,
+		ECachePolicy Policy,
+		uint64& OutWriteSize);
+
+	bool GetCacheValue(
+		FStringView Name,
+		const FCacheKey& Key,
+		ECachePolicy Policy,
+		FValue& OutValue,
+		FHttpRequest* ExistingHttpRequest = nullptr);
+
+	template<typename ValueType, typename ValueHashGetterType, typename ValueDebugContextGetterType>
+	TBitArray<> TryGetCachedDataBatch(
+		TConstArrayView<ValueType> Values,
+		TArray<FCompressedBuffer>& OutBuffers,
+		ValueHashGetterType ValueHashGetter,
+		ValueDebugContextGetterType ValueDebugContextGetter,
+		FHttpRequest* ExistingHttpRequest = nullptr);
+	template<typename ValueType, typename ValueHashGetterType, typename ValueDebugContextGetterType>
+	TBitArray<> CachedDataProbablyExistsBatch(
+		TConstArrayView<ValueType> Values,
+		ValueHashGetterType ValueHashGetter,
+		ValueDebugContextGetterType ValueDebugContextGetter,
+		FHttpRequest* ExistingHttpRequest = nullptr);
+};
+
+FHttpCacheStore::FHttpCacheStore(
+	const TCHAR* InServiceUrl, 
+	const TCHAR* InNamespace, 
+	const TCHAR* InStructuredNamespace, 
+	const TCHAR* InOAuthProvider,
+	const TCHAR* InOAuthClientId,
+	const TCHAR* InOAuthSecret,
+	const EBackendLegacyMode InLegacyMode,
+	const bool bInReadOnly)
+	: Domain(InServiceUrl)
+	, Namespace(InNamespace)
+	, StructuredNamespace(InStructuredNamespace)
+	, DefaultBucket(TEXT("default"))
+	, OAuthProvider(InOAuthProvider)
+	, OAuthClientId(InOAuthClientId)
+	, OAuthSecret(InOAuthSecret)
+	, Access(nullptr)
+	, bIsUsable(false)
+	, bReadOnly(bInReadOnly)
+	, FailedLoginAttempts(0)
+	, SpeedClass(ESpeedClass::Slow)
+	, LegacyMode(InLegacyMode)
+{
+#if WITH_DATAREQUEST_HELPER
+	FDataRequestHelper::StaticInitialize();
+#endif
+	if (IsServiceReady() && AcquireAccessToken())
+	{
+		FString EffectiveDomain;
+		FString OriginalDomainPrefix;
+		TAnsiStringBuilder<64> DomainResolveName;
+
+		if (Domain.StartsWith(TEXT("http://")))
+		{
+			DomainResolveName << Domain.RightChop(7);
+			OriginalDomainPrefix = TEXT("http://");
+		}
+		else if (Domain.StartsWith(TEXT("https://")))
+		{
+			DomainResolveName << Domain.RightChop(8);
+			OriginalDomainPrefix = TEXT("https://");
+		}
+		else
+		{
+			DomainResolveName << Domain;
+		}
+
+		addrinfo* AddrResult = nullptr;
+		addrinfo AddrHints;
+		FMemory::Memset(&AddrHints, 0, sizeof(AddrHints));
+		AddrHints.ai_flags = AI_CANONNAME;
+		AddrHints.ai_family = AF_UNSPEC;
+		if (!::getaddrinfo(*DomainResolveName, nullptr, &AddrHints, &AddrResult))
+		{
+			if (AddrResult->ai_canonname)
+			{
+				// Swap the domain with a canonical name from DNS so that if we are using regional redirection, we pin to a region.
+				EffectiveDomain = OriginalDomainPrefix + ANSI_TO_TCHAR(AddrResult->ai_canonname);
+
+				UE_LOG(LogDerivedDataCache, Display,
+					TEXT("%s: Pinned to %s based on DNS canonical name."),
+					*Domain, *EffectiveDomain);
+			}
+			else
+			{
+				EffectiveDomain = Domain;
+			}
+
+			::freeaddrinfo(AddrResult);
+		}
+		else
+		{
+			EffectiveDomain = Domain;
+		}
+
+		GetRequestPools[0] = MakeUnique<FRequestPool>(*Domain, *EffectiveDomain, Access.Get(), UE_HTTPDDC_GET_REQUEST_POOL_SIZE);
+		GetRequestPools[1] = MakeUnique<FRequestPool>(*Domain, *EffectiveDomain, Access.Get(), 1);
+		PutRequestPools[0] = MakeUnique<FRequestPool>(*Domain, *EffectiveDomain, Access.Get(), UE_HTTPDDC_PUT_REQUEST_POOL_SIZE);
+		PutRequestPools[1] = MakeUnique<FRequestPool>(*Domain, *EffectiveDomain, Access.Get(), 1);
+		bIsUsable = true;
+	}
+
+	AnyInstance = this;
+}
+
+FHttpCacheStore::~FHttpCacheStore()
+{
+	if (AnyInstance == this)
+	{
+		AnyInstance = nullptr;
+	}
+#if WITH_DATAREQUEST_HELPER
+	FDataRequestHelper::StaticShutdown();
+#endif
+}
+
+FString FHttpCacheStore::GetName() const
+{
+	return Domain;
+}
+
+TBitArray<> FHttpCacheStore::TryToPrefetch(TConstArrayView<FString> CacheKeys)
+{
+	return CachedDataProbablyExistsBatch(CacheKeys);
+}
+
+bool FHttpCacheStore::WouldCache(const TCHAR* CacheKey, TArrayView<const uint8> InData)
+{
+	return IsWritable();
+}
+
+FHttpCacheStore::ESpeedClass FHttpCacheStore::GetSpeedClass() const
+{
+	return SpeedClass;
+}
+
+bool FHttpCacheStore::ApplyDebugOptions(FBackendDebugOptions& InOptions)
+{
+	DebugOptions = InOptions;
+	return true;
+}
+
+
+bool FHttpCacheStore::IsServiceReady()
+{
+	FHttpRequest Request(*Domain, *Domain, nullptr, false);
+	FHttpRequest::Result Result = Request.PerformBlockingDownload(TEXT("health/ready"), nullptr);
+	
+	if (Result == FHttpRequest::Success && Request.GetResponseCode() == 200)
+	{
+		UE_LOG(LogDerivedDataCache, Display, TEXT("%s: HTTP DDC service status: %s."), *Request.GetName(), *Request.GetResponseAsString());
+		return true;
+	}
+	else
+	{
+		UE_LOG(LogDerivedDataCache, Warning, TEXT("%s: Unable to reach HTTP DDC service at %s. Status: %d . Response: %s"), *Request.GetName(), *Domain, Request.GetResponseCode(), *Request.GetResponseAsString());
+	}
+
+	return false;
+}
+
+bool FHttpCacheStore::AcquireAccessToken()
+{
+	// Avoid spamming the this if the service is down
+	if (FailedLoginAttempts > UE_HTTPDDC_MAX_FAILED_LOGIN_ATTEMPTS)
+	{
+		return false;
+	}
+
+	ensureMsgf(OAuthProvider.StartsWith(TEXT("http://")) || OAuthProvider.StartsWith(TEXT("https://")),
+		TEXT("The OAuth provider %s is not valid. Needs to be a fully qualified url."),
+		*OAuthProvider
+	);
+
+	// In case many requests wants to update the token at the same time
+	// get the current serial while we wait to take the CS.
+	const uint32 WantsToUpdateTokenSerial = Access.IsValid() ? Access->GetSerial() : 0u;
+
+	{
+		FScopeLock Lock(&AccessCs);
+
+		// Check if someone has beaten us to update the token, then it 
+		// should now be valid.
+		if (Access.IsValid() && Access->GetSerial() > WantsToUpdateTokenSerial)
+		{
+			return true;
+		}
+
+		const uint32 SchemeEnd = OAuthProvider.Find(TEXT("://")) + 3;
+		const uint32 DomainEnd = OAuthProvider.Find(TEXT("/"), ESearchCase::CaseSensitive, ESearchDir::FromStart, SchemeEnd);
+		FString AuthDomain(DomainEnd, *OAuthProvider);
+		FString Uri(*OAuthProvider + DomainEnd + 1);
+
+		FHttpRequest Request(*AuthDomain, *AuthDomain, nullptr, false);
+
+		// If contents of the secret string is a file path, resolve and read form data.
+		if (OAuthSecret.StartsWith(TEXT("file://")))
+		{
+			FString FilePath = OAuthSecret.Mid(7, OAuthSecret.Len() - 7);
+			FString SecretFileContents;
+			if (FFileHelper::LoadFileToString(SecretFileContents, *FilePath))
+			{
+				// Overwrite the filepath with the actual content.
+				OAuthSecret = SecretFileContents;
+			}
+			else
+			{
+				UE_LOG(LogDerivedDataCache, Warning, TEXT("%s: Failed to read OAuth form data file (%s)."), *Request.GetName(), *OAuthSecret);
+				return false;
+			}
+		}
+
+		FString OAuthFormData = FString::Printf(
+			TEXT("client_id=%s&scope=cache_access&grant_type=client_credentials&client_secret=%s"),
+			*OAuthClientId,
+			*OAuthSecret
+		);
+
+		TArray<uint8> FormData;
+		auto OAuthFormDataUTF8 = FTCHARToUTF8(*OAuthFormData);
+		FormData.Append((uint8*)OAuthFormDataUTF8.Get(), OAuthFormDataUTF8.Length());
+
+		FHttpRequest::Result Result = Request.PerformBlockingUpload<FHttpRequest::Post>(*Uri, MakeArrayView(FormData));
+
+		if (Result == FHttpRequest::Success && Request.GetResponseCode() == 200)
+		{
+			TSharedPtr<FJsonObject> ResponseObject = Request.GetResponseAsJsonObject();
+			if (ResponseObject)
+			{
+				FString AccessTokenString;
+				int32 ExpiryTimeSeconds = 0;
+				int32 CurrentTimeSeconds = int32(FPlatformTime::ToSeconds(FPlatformTime::Cycles()));
+
+				if (ResponseObject->TryGetStringField(TEXT("access_token"), AccessTokenString) &&
+					ResponseObject->TryGetNumberField(TEXT("expires_in"), ExpiryTimeSeconds))
+				{
+					if (!Access)
+					{
+						Access = MakeUnique<FHttpAccessToken>();
+					}
+					Access->SetHeader(*AccessTokenString);
+					UE_LOG(LogDerivedDataCache, Display, TEXT("%s: Logged in to HTTP DDC services. Expires in %d seconds."), *Request.GetName(), ExpiryTimeSeconds);
+
+					//Schedule a refresh of the token ahead of expiry time (this will not work in commandlets)
+					if (!IsRunningCommandlet())
+					{
+						FTSTicker::GetCoreTicker().AddTicker(FTickerDelegate::CreateLambda(
+							[this](float DeltaTime)
+							{
+								this->AcquireAccessToken();
+								return false;
+							}
+						), ExpiryTimeSeconds - 20.0f);
+					}
+					// Reset failed login attempts, the service is indeed alive.
+					FailedLoginAttempts = 0;
+					return true;
+				}
+			}
+		}
+		else
+		{
+			UE_LOG(LogDerivedDataCache, Warning, TEXT("%s: Failed to log in to HTTP services. Server responed with code %d."), *Request.GetName(), Request.GetResponseCode());
+			FailedLoginAttempts++;
+		}
+	}
+	return false;
+}
+
+bool FHttpCacheStore::ShouldRetryOnError(int64 ResponseCode)
+{
+	// Access token might have expired, request a new token and try again.
+	if (ResponseCode == 401 && AcquireAccessToken())
+	{
+		return true;
+	}
+
+	// Too many requests, make a new attempt
+	if (ResponseCode == 429)
+	{
+		return true;
+	}
+
+	return false;
+}
+
+bool FHttpCacheStore::ShouldSimulateMiss(const TCHAR* InKey)
+{
+	if (DebugOptions.RandomMissRate == 0 && DebugOptions.SimulateMissTypes.IsEmpty())
+	{
+		return false;
+	}
+
+	const FName Key(InKey);
+	const uint32 Hash = GetTypeHash(Key);
+
+	if (FScopeLock Lock(&MissedKeysCS); DebugMissedKeys.ContainsByHash(Hash, Key))
+	{
+		return true;
+	}
+
+	if (DebugOptions.ShouldSimulateMiss(InKey))
+	{
+		FScopeLock Lock(&MissedKeysCS);
+		DebugMissedKeys.AddByHash(Hash, Key);
+		return true;
+	}
+
+	return false;
+}
+
+bool FHttpCacheStore::ShouldSimulateMiss(const FCacheKey& Key)
+{
+	if (DebugOptions.RandomMissRate == 0 && DebugOptions.SimulateMissTypes.IsEmpty())
+	{
+		return false;
+	}
+
+	const uint32 Hash = GetTypeHash(Key);
+
+	if (FScopeLock Lock(&MissedKeysCS); DebugMissedCacheKeys.ContainsByHash(Hash, Key))
+	{
+		return true;
+	}
+
+	if (DebugOptions.ShouldSimulateMiss(Key))
+	{
+		FScopeLock Lock(&MissedKeysCS);
+		DebugMissedCacheKeys.AddByHash(Hash, Key);
+		return true;
+	}
+
+	return false;
+}
+
+uint64 FHttpCacheStore::PutRef(const FCbPackage& Package, const FCacheKey& Key, FStringView Bucket, bool bFinalize, TArray<FIoHash>& OutNeededBlobHashes, bool& bOutPutCompletedSuccessfully)
+{
+	bOutPutCompletedSuccessfully = false;
+
+	uint64 BytesSent = 0;
+	TStringBuilder<256> RefsUri;
+	RefsUri << "api/v1/refs/" << StructuredNamespace << "/" << Bucket << "/" << Key.Hash;
+	if (bFinalize)
+	{
+		RefsUri << "/finalize/" << Package.GetObjectHash();
+	}
+
+	int64 ResponseCode = 0;
+	for (uint32 Attempts = 0; (Attempts < UE_HTTPDDC_MAX_ATTEMPTS) && !ShouldAbortForShutdown() && (Attempts == 0 || ShouldRetryOnError(ResponseCode)); ++Attempts)
+	{
+		FScopedRequestPtr Request(PutRequestPools[IsInGameThread()].Get());
+
+		if (bFinalize)
+		{
+			Request->PerformBlockingUpload<FHttpRequest::Post>(*RefsUri, TArrayView<const uint8>());
+		}
+		else
+		{
+			Request->SetHeader(TEXT("X-Jupiter-IoHash"), *WriteToString<48>(Package.GetObjectHash()));
+			Request->PerformBlockingUpload<FHttpRequest::PutCompactBinary>(*RefsUri, MakeConstArrayView(Package.GetObject().GetBuffer().ToShared()));
+		}
+		ResponseCode = Request->GetResponseCode();
+
+		if (FHttpRequest::IsSuccessResponse(ResponseCode))
+		{
+			BytesSent += Request->GetBytesSent();
+			if (TSharedPtr<FJsonObject> ResponseObject = Request->GetResponseAsJsonObject())
+			{
+				TArray<FString> NeedsArrayStrings;
+				ResponseObject->TryGetStringArrayField(TEXT("needs"), NeedsArrayStrings);
+
+				OutNeededBlobHashes.Empty(NeedsArrayStrings.Num());
+				for (const FString& NeededString : NeedsArrayStrings)
+				{
+					FIoHash BlobHash;
+					LexFromString(BlobHash, *NeededString);
+					if (!BlobHash.IsZero())
+					{
+						OutNeededBlobHashes.Add(BlobHash);
+					}
+				}
+			}
+			else
+			{
+				OutNeededBlobHashes.Empty();
+			}
+
+			bOutPutCompletedSuccessfully = true;
+			break;
+		}
+		else
+		{
+			OutNeededBlobHashes.Empty();
+		}
+	}
+
+	return BytesSent;
+}
+
+bool FHttpCacheStore::PutCacheRecord(
+	FStringView Name,
+	const FCacheRecord& Record,
+	const FCacheRecordPolicy& Policy,
+	uint64& OutWriteSize)
+{
+	OutWriteSize = 0;
+
+	if (!IsWritable())
+	{
+		UE_LOG(LogDerivedDataCache, VeryVerbose,
+			TEXT("%s: Skipped put of %s from '%.*s' because this cache store is read-only"),
+			*GetName(), *WriteToString<96>(Record.GetKey()), Name.Len(), Name.GetData());
+		return false;
+	}
+
+	const FCacheKey& Key = Record.GetKey();
+	const ECachePolicy RecordPolicy = Policy.GetRecordPolicy();
+
+	// Skip the request if storing to the cache is disabled.
+	// Http backends won't generally be "local" but including handling for this possibility for consistency
+	const ECachePolicy StoreFlag = SpeedClass == ESpeedClass::Local ? ECachePolicy::StoreLocal : ECachePolicy::StoreRemote;
+	if (!EnumHasAnyFlags(RecordPolicy, StoreFlag))
+	{
+		UE_LOG(LogDerivedDataCache, VeryVerbose, TEXT("%s: Skipped put of %s from '%.*s' due to cache policy"),
+			*GetName(), *WriteToString<96>(Key), Name.Len(), Name.GetData());
+		return false;
+	}
+
+	if (ShouldSimulateMiss(Key))
+	{
+		UE_LOG(LogDerivedDataCache, Verbose, TEXT("%s: Simulated miss for put of %s from '%.*s'"),
+			*GetName(), *WriteToString<96>(Key), Name.Len(), Name.GetData());
+		return false;
+	}
+
+	// TODO: Jupiter currently always overwrites.  It doesn't have a "write if not present" feature (for records or attachments),
+	//		 but would require one to implement all policy correctly.
+
+	FString Bucket(Key.Bucket.ToString());
+	Bucket.ToLowerInline();
+
+	bool bPutCompletedSuccessfully = false;
+
+	FCbPackage Package = Record.Save();
+	TArray<FIoHash> NeededBlobHashes;
+
+	// Initial record upload
+	size_t PutRefBytesSent = PutRef(Package, Record.GetKey(), Bucket, false, NeededBlobHashes, bPutCompletedSuccessfully);
+	OutWriteSize += PutRefBytesSent;
+
+	if (!bPutCompletedSuccessfully)
+	{
+		UE_LOG(LogDerivedDataCache, Warning, TEXT("%s: Failed to put reference object for put of %s from '%.*s'"),
+			*GetName(), *WriteToString<96>(Key), Name.Len(), Name.GetData());
+		return false;
+	}
+
+	// TODO: blob uploading and finalization should be replaced with a single batch compressed blob upload endpoint in the future.
+	TStringBuilder<128> ExpectedHashes;
+	bool bExpectedHashesSerialized = false;
+
+	// Needed blob upload (if any missing)
+	for (const FIoHash& NeededBlobHash : NeededBlobHashes)
+	{
+		TStringBuilder<256> CompressedBlobsUri;
+		CompressedBlobsUri << "api/v1/compressed-blobs/" << StructuredNamespace << "/" << NeededBlobHash;
+
+		if (const FCbAttachment* Attachment = Package.FindAttachment(NeededBlobHash))
+		{
+			FSharedBuffer TempBuffer;
+			TConstArrayView<uint8> BlobArrayView;
+			if (Attachment->IsCompressedBinary())
+			{
+				TempBuffer = Attachment->AsCompressedBinary().GetCompressed().ToShared();
+				BlobArrayView = MakeConstArrayView(TempBuffer);
+			}
+			else if (Attachment->IsBinary())
+			{
+				TempBuffer = FCompressedBuffer::Compress(Attachment->AsCompositeBinary()).GetCompressed().ToShared();
+				BlobArrayView = MakeConstArrayView(TempBuffer);
+			}
+			else
+			{
+				TempBuffer = FCompressedBuffer::Compress(Attachment->AsObject().GetBuffer()).GetCompressed().ToShared();
+				BlobArrayView = MakeConstArrayView(TempBuffer);
+			}
+
+			int64 ResponseCode = 0;
+			for (uint32 Attempts = 0; (Attempts < UE_HTTPDDC_MAX_ATTEMPTS) && !ShouldAbortForShutdown() && (Attempts == 0 || ShouldRetryOnError(ResponseCode)); ++Attempts)
+			{
+				FScopedRequestPtr Request(PutRequestPools[IsInGameThread()].Get());
+				Request->PerformBlockingUpload<FHttpRequest::PutCompressedBlob>(*CompressedBlobsUri, BlobArrayView);
+
+				ResponseCode = Request->GetResponseCode();
+
+				if (FHttpRequest::IsSuccessResponse(ResponseCode))
+				{
+					OutWriteSize += Request->GetBytesSent();
+					break;
+				}
+			}
+		}
+		else
+		{
+			if (!bExpectedHashesSerialized)
+			{
+				bool bFirstHash = true;
+				for (const FCbAttachment& PackageAttachment : Package.GetAttachments())
+				{
+					if (!bFirstHash)
+					{
+						ExpectedHashes << TEXT(", ");
+					}
+					ExpectedHashes << PackageAttachment.GetHash();
+					bFirstHash = false;
+				}
+				bExpectedHashesSerialized = true;
+			}
+			UE_LOG(LogDerivedDataCache, Warning, TEXT("%s: Server reported needed hash '%s' that is outside the set of expected hashes (%s) for put of %s from '%.*s'"),
+				*GetName(), *WriteToString<96>(NeededBlobHash), ExpectedHashes.ToString(), *WriteToString<96>(Key), Name.Len(), Name.GetData());
+		}
+	}
+
+	// Finalization (if any blobs needed)
+	if (!NeededBlobHashes.IsEmpty())
+	{
+		size_t FinalizeRefBytesSent = PutRef(Package, Record.GetKey(), Bucket, true, NeededBlobHashes, bPutCompletedSuccessfully);
+		OutWriteSize += FinalizeRefBytesSent;
+	}
+
+	return bPutCompletedSuccessfully && NeededBlobHashes.IsEmpty();
+}
+
+FOptionalCacheRecord FHttpCacheStore::GetCacheRecordOnly(
+	const FStringView Name,
+	const FCacheKey& Key,
+	const FCacheRecordPolicy& Policy)
+{
+	if (!IsUsable())
+	{
+		UE_LOG(LogDerivedDataCache, VeryVerbose,
+			TEXT("%s: Skipped get of %s from '%.*s' because this cache store is not available"),
+			*GetName(), *WriteToString<96>(Key), Name.Len(), Name.GetData());
+		return FOptionalCacheRecord();
+	}
+
+	// Skip the request if querying the cache is disabled.
+	const ECachePolicy QueryPolicy = SpeedClass == ESpeedClass::Local
+		? ECachePolicy::QueryLocal : ECachePolicy::QueryRemote;
+	if (!EnumHasAnyFlags(Policy.GetRecordPolicy(), QueryPolicy))
+	{
+		UE_LOG(LogDerivedDataCache, VeryVerbose, TEXT("%s: Skipped get of %s from '%.*s' due to cache policy"),
+			*GetName(), *WriteToString<96>(Key), Name.Len(), Name.GetData());
+		return FOptionalCacheRecord();
+	}
+
+	if (ShouldSimulateMiss(Key))
+	{
+		UE_LOG(LogDerivedDataCache, Verbose, TEXT("%s: Simulated miss for get of %s from '%.*s'"),
+			*GetName(), *WriteToString<96>(Key), Name.Len(), Name.GetData());
+		return FOptionalCacheRecord();
+	}
+
+	FScopedRequestPtr Request(GetRequestPools[IsInGameThread()].Get());
+
+	FOptionalCacheRecord Record;
+	{
+		FString Bucket(Key.Bucket.ToString());
+		Bucket.ToLowerInline();
+
+		TStringBuilder<256> RefsUri;
+		RefsUri << "api/v1/refs/" << StructuredNamespace << "/" << Bucket << "/" << Key.Hash;
+
+		bool bIsSuccessfulResponse = false;
+		FSharedBuffer ResponseBuffer;
+		int64 ResponseCode = 0;
+		for (uint32 Attempts = 0; (Attempts < UE_HTTPDDC_MAX_ATTEMPTS) && !ShouldAbortForShutdown() && (Attempts == 0 || ShouldRetryOnError(ResponseCode)); ++Attempts)
+		{
+			if (Attempts > 0)
+			{
+				Request->Reset();
+			}
+
+			TArray<uint8> ByteArray;
+			Request->SetHeader(TEXT("Accept"), TEXT("application/x-ue-cb"));
+			Request->PerformBlockingDownload(*RefsUri, &ByteArray, {401, 404});
+			ResponseCode = Request->GetResponseCode();
+
+			if (FHttpRequest::IsSuccessResponse(ResponseCode))
+			{
+				ResponseBuffer = MakeSharedBufferFromArray(MoveTemp(ByteArray));
+				bIsSuccessfulResponse = true;
+				break;
+			}
+		}
+
+		if (!bIsSuccessfulResponse)
+		{
+			UE_LOG(LogDerivedDataCache, Verbose, TEXT("%s: Cache miss with missing package for %s from '%.*s'"),
+				*GetName(), *WriteToString<96>(Key), Name.Len(), Name.GetData());
+			return Record;
+		}
+
+		if (ValidateCompactBinary(ResponseBuffer, ECbValidateMode::Default) != ECbValidateError::None)
+		{
+			UE_LOG(LogDerivedDataCache, Display, TEXT("%s: Cache miss with invalid package for %s from '%.*s'"),
+				*GetName(), *WriteToString<96>(Key), Name.Len(), Name.GetData());
+			return Record;
+		}
+
+		Record = FCacheRecord::Load(FCbPackage(FCbObject(ResponseBuffer)));
+		if (Record.IsNull())
+		{
+			UE_LOG(LogDerivedDataCache, Display, TEXT("%s: Cache miss with record load failure for %s from '%.*s'"),
+				*GetName(), *WriteToString<96>(Key), Name.Len(), Name.GetData());
+			return Record;
+		}
+	}
+
+	return Record;
+}
+
+bool FHttpCacheStore::PutCacheValue(
+	const FStringView Name,
+	const FCacheKey& Key,
+	const FValue& Value,
+	const ECachePolicy Policy,
+	uint64& OutWriteSize)
+{
+	if (!IsWritable())
+	{
+		UE_LOG(LogDerivedDataCache, VeryVerbose,
+			TEXT("%s: Skipped put of %s from '%.*s' because this cache store is read-only"),
+			*GetName(), *WriteToString<96>(Key), Name.Len(), Name.GetData());
+		return false;
+	}
+
+	// Skip the request if storing to the cache is disabled.
+	const ECachePolicy StoreFlag = SpeedClass == ESpeedClass::Local ? ECachePolicy::StoreLocal : ECachePolicy::StoreRemote;
+	if (!EnumHasAnyFlags(Policy, StoreFlag))
+	{
+		UE_LOG(LogDerivedDataCache, VeryVerbose, TEXT("%s: Skipped put of %s from '%.*s' due to cache policy"),
+			*GetName(), *WriteToString<96>(Key), Name.Len(), Name.GetData());
+		return false;
+	}
+
+	if (ShouldSimulateMiss(Key))
+	{
+		UE_LOG(LogDerivedDataCache, Verbose, TEXT("%s: Simulated miss for put of %s from '%.*s'"),
+			*GetName(), *WriteToString<96>(Key), Name.Len(), Name.GetData());
+		return false;
+	}
+
+	// TODO: Jupiter currently always overwrites.  It doesn't have a "write if not present" feature (for records or attachments),
+	//		 but would require one to implement all policy correctly.
+
+	FString Bucket(Key.Bucket.ToString());
+	Bucket.ToLowerInline();
+
+	bool bPutCompletedSuccessfully = false;
+
+	FCbWriter Writer;
+	Writer.BeginObject();
+	Writer.AddBinaryAttachment("RawHash", Value.GetRawHash());
+	Writer.AddInteger("RawSize", Value.GetRawSize());
+	Writer.EndObject();
+
+	FCbPackage Package(Writer.Save().AsObject());
+	TArray<FIoHash> NeededBlobHashes;
+
+	// Initial record upload
+	size_t PutRefBytesSent = PutRef(Package, Key, Bucket, false, NeededBlobHashes, bPutCompletedSuccessfully);
+	OutWriteSize += PutRefBytesSent;
+
+	if (!bPutCompletedSuccessfully)
+	{
+		UE_LOG(LogDerivedDataCache, Warning, TEXT("%s: Failed to put reference object for put of %s from '%.*s'"),
+			*GetName(), *WriteToString<96>(Key), Name.Len(), Name.GetData());
+		return false;
+	}
+
+	// Needed blob upload (if any missing)
+	if (!NeededBlobHashes.IsEmpty())
+	{
+		if (NeededBlobHashes.Num() != 1)
+		{
+			TStringBuilder<128> NeededHashString;
+			bool bFirstHash = true;
+			for (const FIoHash& NeededBlobHash : NeededBlobHashes)
+			{
+				if (!bFirstHash)
+				{
+					NeededHashString << TEXT(", ");
+				}
+				NeededHashString << NeededBlobHash;
+				bFirstHash = false;
+			}
+
+			UE_LOG(LogDerivedDataCache, Warning, TEXT("%s: Server reported unexpected needed hash quantity '%d' (%s) for put of %s from '%.*s'"),
+				*GetName(), NeededBlobHashes.Num(), NeededHashString.ToString(), *WriteToString<96>(Key), Name.Len(), Name.GetData());
+			return false;
+		}
+
+		if (NeededBlobHashes[0] != Value.GetRawHash())
+		{
+			UE_LOG(LogDerivedDataCache, Warning, TEXT("%s: Server reported needed hash '%s' that is outside the set of expected hashes (%s) for put of %s from '%.*s'"),
+				*GetName(), *WriteToString<96>(NeededBlobHashes[0]), *WriteToString<96>(Value.GetRawHash()), *WriteToString<96>(Key), Name.Len(), Name.GetData());
+			return false;
+		}
+
+		TStringBuilder<256> CompressedBlobsUri;
+		CompressedBlobsUri << "api/v1/compressed-blobs/" << StructuredNamespace << "/" << Value.GetRawHash();
+
+		FSharedBuffer TempBuffer = Value.GetData().GetCompressed().ToShared();
+		FScopedRequestPtr Request(PutRequestPools[IsInGameThread()].Get());
+
+		int64 ResponseCode = 0;
+		for (uint32 Attempts = 0; (Attempts < UE_HTTPDDC_MAX_ATTEMPTS) && !ShouldAbortForShutdown() && (Attempts == 0 || ShouldRetryOnError(ResponseCode)); ++Attempts)
+		{
+			if (Attempts > 0)
+			{
+				Request->Reset();
+			}
+
+			Request->PerformBlockingUpload<FHttpRequest::PutCompressedBlob>(*CompressedBlobsUri, MakeConstArrayView(TempBuffer));
+
+			ResponseCode = Request->GetResponseCode();
+
+			if (FHttpRequest::IsSuccessResponse(ResponseCode))
+			{
+				OutWriteSize += Request->GetBytesSent();
+				break;
+			}
+		}
+
+		OutWriteSize += PutRef(Package, Key, Bucket, true, NeededBlobHashes, bPutCompletedSuccessfully);
+	}
+
+
+	return bPutCompletedSuccessfully && NeededBlobHashes.IsEmpty();
+}
+
+bool FHttpCacheStore::GetCacheValue(
+	const FStringView Name,
+	const FCacheKey& Key,
+	const ECachePolicy Policy,
+	FValue& OutValue,
+	FHttpRequest* ExistingHttpRequest)
+{
+	if (!IsUsable())
+	{
+		UE_LOG(LogDerivedDataCache, VeryVerbose,
+			TEXT("%s: Skipped get of %s from '%.*s' because this cache store is not available"),
+			*GetName(), *WriteToString<96>(Key), Name.Len(), Name.GetData());
+		return false;
+	}
+
+	// Skip the request if querying the cache is disabled.
+	const ECachePolicy QueryFlag = SpeedClass == ESpeedClass::Local ? ECachePolicy::QueryLocal : ECachePolicy::QueryRemote;
+	if (!EnumHasAnyFlags(Policy, QueryFlag))
+	{
+		UE_LOG(LogDerivedDataCache, VeryVerbose, TEXT("%s: Skipped get of %s from '%.*s' due to cache policy"),
+			*GetName(), *WriteToString<96>(Key), Name.Len(), Name.GetData());
+		return false;
+	}
+
+	if (ShouldSimulateMiss(Key))
+	{
+		UE_LOG(LogDerivedDataCache, Verbose, TEXT("%s: Simulated miss for get of %s from '%.*s'"),
+			*GetName(), *WriteToString<96>(Key), Name.Len(), Name.GetData());
+		return false;
+	}
+
+	const bool bSkipData = EnumHasAnyFlags(Policy, ECachePolicy::SkipData);
+
+	FString Bucket(Key.Bucket.ToString());
+	Bucket.ToLowerInline();
+
+	TStringBuilder<256> RefsUri;
+	RefsUri << "api/v1/refs/" << StructuredNamespace << "/" << Bucket << "/" << Key.Hash;
+
+	FHttpRequest* Request = ExistingHttpRequest;
+	TOptional<FScopedRequestPtr> RequestPoolRequest;
+	if (Request == nullptr)
+	{
+		RequestPoolRequest.Emplace(GetRequestPools[IsInGameThread()].Get());
+		Request = RequestPoolRequest->Get();
+	}
+	else
+	{
+		Request->Reset();
+	}
+
+	bool bIsSuccessfulResponse = false;
+	FSharedBuffer ResponseBuffer;
+	int64 ResponseCode = 0;
+	for (uint32 Attempts = 0; (Attempts < UE_HTTPDDC_MAX_ATTEMPTS) && !ShouldAbortForShutdown() && (Attempts == 0 || ShouldRetryOnError(ResponseCode)); ++Attempts)
+	{
+		if (Attempts > 0)
+		{
+			Request->Reset();
+		}
+
+		TArray<uint8> ByteArray;
+		if (bSkipData)
+		{
+			Request->SetHeader(TEXT("Accept"), TEXT("application/x-ue-cb"));
+		}
+		else
+		{
+			Request->SetHeader(TEXT("Accept"), TEXT("application/x-jupiter-inline"));
+		}
+		Request->PerformBlockingDownload(*RefsUri, &ByteArray, { 401, 404 });
+		ResponseCode = Request->GetResponseCode();
+
+		if (FHttpRequest::IsSuccessResponse(ResponseCode))
+		{
+			ResponseBuffer = MakeSharedBufferFromArray(MoveTemp(ByteArray));
+			bIsSuccessfulResponse = true;
+			break;
+		}
+	}
+
+	if (!bIsSuccessfulResponse)
+	{
+		UE_LOG(LogDerivedDataCache, Verbose, TEXT("%s: Cache miss with missing package for %s from '%.*s'"),
+			*GetName(), *WriteToString<96>(Key), Name.Len(), Name.GetData());
+		return false;
+	}
+
+	if (bSkipData)
+	{
+		if (ValidateCompactBinary(ResponseBuffer, ECbValidateMode::Default) != ECbValidateError::None)
+		{
+			UE_LOG(LogDerivedDataCache, Display, TEXT("%s: Cache miss with invalid package for %s from '%.*s'"),
+				*GetName(), *WriteToString<96>(Key), Name.Len(), Name.GetData());
+			return false;
+		}
+
+		const FCbObjectView Object = FCbObject(ResponseBuffer);
+		const FIoHash RawHash = Object["RawHash"].AsHash();
+		const uint64 RawSize = Object["RawSize"].AsUInt64(MAX_uint64);
+		if (RawHash.IsZero() || RawSize == MAX_uint64)
+		{
+			UE_LOG(LogDerivedDataCache, Display, TEXT("%s: Cache miss with invalid value for %s from '%.*s'"),
+				*GetName(), *WriteToString<96>(Key), Name.Len(), Name.GetData());
+			return false;
+		}
+		OutValue = FValue(RawHash, RawSize);
+	}
+	else
+	{
+		FCompressedBuffer CompressedBuffer = FCompressedBuffer::FromCompressed(MoveTemp(ResponseBuffer));
+		if (!CompressedBuffer)
+		{
+			UE_LOG(LogDerivedDataCache, Display, TEXT("%s: Cache miss with invalid package for %s from '%.*s'"),
+				*GetName(), *WriteToString<96>(Key), Name.Len(), Name.GetData());
+			return false;
+		}
+		OutValue = FValue(CompressedBuffer);
+	}
+
+	return true;
+}
+
+FOptionalCacheRecord FHttpCacheStore::GetCacheRecord(
+	const FStringView Name,
+	const FCacheKey& Key,
+	const FCacheRecordPolicy& Policy,
+	EStatus& OutStatus)
+{
+	FOptionalCacheRecord Record = GetCacheRecordOnly(Name, Key, Policy);
+	if (Record.IsNull())
+	{
+		OutStatus = EStatus::Error;
+		return Record;
+	}
+
+	OutStatus = EStatus::Ok;
+
+	FCacheRecordBuilder RecordBuilder(Key);
+
+	if (!EnumHasAnyFlags(Policy.GetRecordPolicy(), ECachePolicy::SkipMeta))
+	{
+		RecordBuilder.SetMeta(FCbObject(Record.Get().GetMeta()));
+	}
+
+	// TODO: There is not currently a batched GET endpoint for Jupiter.  Once there is, all payload data should be fetched in one call.
+	//		 In the meantime, we try to keep the code structured in a way that is friendly to future batching of GETs.
+
+	TArray<FValueWithId> RequiredGets;
+	TArray<FValueWithId> RequiredHeads;
+
+	for (FValueWithId Value : Record.Get().GetValues())
+	{
+		const ECachePolicy ValuePolicy = Policy.GetValuePolicy(Value.GetId());
+		if (IsValueDataReady(Value, ValuePolicy))
+		{
+			RecordBuilder.AddValue(MoveTemp(Value));
+		}
+		else
+		{
+			if (EnumHasAnyFlags(ValuePolicy, ECachePolicy::SkipData))
+			{
+				RequiredHeads.Emplace(Value);
+			}
+			else
+			{
+				RequiredGets.Emplace(Value);
+			}
+		}
+	}
+
+	auto HashGetter = [](const FValueWithId& Value)
+	{
+		return Value.GetRawHash();
+	};
+
+	auto DebugContextGetter = [Name, &Key](const FValueWithId& Value)
+	{
+		return FValueDebugContext{ Name, Key, FString(*WriteToString<16>(Value.GetId())) };
+	};
+
+	if (CachedDataProbablyExistsBatch<FValueWithId>(RequiredHeads, HashGetter, DebugContextGetter).CountSetBits() != RequiredHeads.Num())
+	{
+		OutStatus = EStatus::Error;
+		return FOptionalCacheRecord();
+	}
+
+	TArray<FCompressedBuffer> FetchedBuffers;
+	if (TryGetCachedDataBatch<FValueWithId>(RequiredGets, FetchedBuffers, HashGetter, DebugContextGetter).CountSetBits() != RequiredGets.Num())
+	{
+		OutStatus = EStatus::Error;
+		return FOptionalCacheRecord();
+	}
+
+	for (int32 Index = 0; Index < RequiredHeads.Num(); ++Index)
+	{
+		RecordBuilder.AddValue(RequiredHeads[Index].RemoveData());
+	}
+
+	for (int32 Index = 0; Index < RequiredGets.Num(); ++Index)
+	{
+		RecordBuilder.AddValue(FValueWithId(RequiredGets[Index].GetId(), FetchedBuffers[Index]));
+	}
+
+	return RecordBuilder.Build();
+}
+
+template<typename ValueType, typename ValueHashGetterType, typename ValueDebugContextGetterType>
+TBitArray<> FHttpCacheStore::TryGetCachedDataBatch(
+	TConstArrayView<ValueType> Values,
+	TArray<FCompressedBuffer>& OutBuffers,
+	ValueHashGetterType ValueHashGetter,
+	ValueDebugContextGetterType ValueDebugContextGetter,
+	FHttpRequest* ExistingHttpRequest)
+{
+	FHttpRequest* Request = ExistingHttpRequest;
+	TOptional<FScopedRequestPtr> RequestPoolRequest;
+	if (Request == nullptr)
+	{
+		RequestPoolRequest.Emplace(GetRequestPools[IsInGameThread()].Get());
+		Request = RequestPoolRequest->Get();
+	}
+	else
+	{
+		Request->Reset();
+	}
+
+	bool bRequestNeedsReset = false;
+	TBitArray<> Results(true, Values.Num());
+	int32 ValueIndex = 0;
+	for (const ValueType& Value : Values)
+	{
+		const FIoHash& RawHash = ValueHashGetter(Value);
+		TStringBuilder<256> CompressedBlobsUri;
+		CompressedBlobsUri << "api/v1/compressed-blobs/" << StructuredNamespace << "/" << RawHash;
+
+		bool bHit = false;
+		FCompressedBuffer CompressedBuffer;
+		int64 ResponseCode = 0;
+		for (uint32 Attempts = 0; (Attempts < UE_HTTPDDC_MAX_ATTEMPTS) && !ShouldAbortForShutdown() && (Attempts == 0 || ShouldRetryOnError(ResponseCode)); ++Attempts)
+		{
+			if (bRequestNeedsReset)
+			{
+				Request->Reset();
+			}
+
+			TArray<uint8> ByteArray;
+			const FHttpRequest::Result Result = Request->PerformBlockingDownload(*CompressedBlobsUri, &ByteArray, {404});
+			ResponseCode = Request->GetResponseCode();
+			bRequestNeedsReset = true;
+
+			if (FHttpRequest::IsSuccessResponse(ResponseCode))
+			{
+				CompressedBuffer = FCompressedBuffer::FromCompressed(MakeSharedBufferFromArray(MoveTemp(ByteArray)));
+				bHit = true;
+				break;
+			}
+		}
+
+		if (!bHit)
+		{
+			FValueDebugContext DebugContext = ValueDebugContextGetter(Value);
+			UE_LOG(LogDerivedDataCache, Verbose,
+				TEXT("%s: Cache miss with missing value %s with hash %s for %s from '%.*s'"),
+				*GetName(), *DebugContext.Id, *WriteToString<48>(RawHash), *WriteToString<96>(DebugContext.Key),
+				DebugContext.Name.Len(), DebugContext.Name.GetData());
+			Results[ValueIndex] = false;
+		}
+		else if (CompressedBuffer.GetRawHash() != RawHash)
+		{
+			FValueDebugContext DebugContext = ValueDebugContextGetter(Value);
+			UE_LOG(LogDerivedDataCache, Display,
+				TEXT("%s: Cache miss with corrupted value %s with hash %s for %s from '%.*s'"),
+				*GetName(), *DebugContext.Id, *WriteToString<48>(RawHash),
+				*WriteToString<96>(DebugContext.Key), DebugContext.Name.Len(), DebugContext.Name.GetData());
+			Results[ValueIndex] = false;
+		}
+		else
+		{
+			OutBuffers.Add(CompressedBuffer);
+		}
+		++ValueIndex;
+	}
+	return Results;
+}
+
+template<typename ValueType, typename ValueHashGetterType, typename ValueDebugContextGetterType>
+TBitArray<> FHttpCacheStore::CachedDataProbablyExistsBatch(
+	TConstArrayView<ValueType> Values,
+	ValueHashGetterType ValueHashGetter,
+	ValueDebugContextGetterType ValueDebugContextGetter,
+	FHttpRequest* ExistingHttpRequest)
+{
+	if (Values.IsEmpty())
+	{
+		return TBitArray<>();
+	}
+
+	FHttpRequest* Request = ExistingHttpRequest;
+	TOptional<FScopedRequestPtr> RequestPoolRequest;
+	if (Request == nullptr)
+	{
+		RequestPoolRequest.Emplace(GetRequestPools[IsInGameThread()].Get());
+		Request = RequestPoolRequest->Get();
+	}
+	else
+	{
+		Request->Reset();
+	}
+
+	TStringBuilder<256> CompressedBlobsUri;
+	CompressedBlobsUri << "api/v1/compressed-blobs/" << StructuredNamespace << "/exists?";
+	bool bFirstItem = true;
+	for (const ValueType& Value : Values)
+	{
+		if (!bFirstItem)
+		{
+			CompressedBlobsUri << "&";
+		}
+		CompressedBlobsUri << "id=" << ValueHashGetter(Value);
+		bFirstItem = false;
+	}
+
+	int64 ResponseCode = 0;
+	for (uint32 Attempts = 0; (Attempts < UE_HTTPDDC_MAX_ATTEMPTS) && !ShouldAbortForShutdown() && (Attempts == 0 || ShouldRetryOnError(ResponseCode)); ++Attempts)
+	{
+		if (Attempts > 0)
+		{
+			Request->Reset();
+		}
+
+		TConstArrayView<uint8> DummyBuffer;
+		const FHttpRequest::Result Result = Request->PerformBlockingUpload<FHttpRequest::Post>(*CompressedBlobsUri, DummyBuffer);
+		ResponseCode = Request->GetResponseCode();
+
+		if (FHttpRequest::IsSuccessResponse(ResponseCode))
+		{
+			if (TSharedPtr<FJsonObject> ResponseObject = Request->GetResponseAsJsonObject())
+			{
+				TArray<FString> NeedsArrayStrings;
+				if (ResponseObject->TryGetStringArrayField(TEXT("needs"), NeedsArrayStrings))
+				{
+					if (NeedsArrayStrings.IsEmpty())
+					{
+						return TBitArray<>(true, Values.Num());
+					}
+				}
+
+				TBitArray<> Results(true, Values.Num());
+				for (const FString& NeedsString : NeedsArrayStrings)
+				{
+					FIoHash NeedHash;
+					LexFromString(NeedHash, *NeedsString);
+					for (int32 ValueIndex = 0; ValueIndex < Values.Num(); ++ValueIndex)
+					{
+						if (NeedHash == ValueHashGetter(Values[ValueIndex]))
+						{
+							Results[ValueIndex] = false;
+
+							FValueDebugContext DebugContext = ValueDebugContextGetter(Values[ValueIndex]);
+							UE_LOG(LogDerivedDataCache, Verbose,
+								TEXT("%s: Cache exists miss with missing value %s with hash %s for %s from '%.*s'"),
+								*GetName(), *DebugContext.Id, *NeedsString, *WriteToString<96>(DebugContext.Key),
+								DebugContext.Name.Len(), DebugContext.Name.GetData());
+						}
+					}
+				}
+
+				return Results;
+			}
+			else
+			{
+				UE_LOG(LogDerivedDataCache, Warning,
+					TEXT("%s: Cache exists returned invalid results."),
+					*GetName());
+				return TBitArray<>(false, Values.Num());
+			}
+
+			return TBitArray<>(false, Values.Num());
+		}
+	}
+
+	return TBitArray<>(false, Values.Num());
+}
+
+bool FHttpCacheStore::CachedDataProbablyExists(const TCHAR* CacheKey)
+{
+	TRACE_CPUPROFILER_EVENT_SCOPE(HttpDDC_Exist);
+	TRACE_COUNTER_ADD(HttpDDC_Exist, int64(1));
+	COOK_STAT(auto Timer = UsageStats.TimeProbablyExists());
+
+	if (ShouldSimulateMiss(CacheKey))
+	{
+		return false;
+	}
+
+#if WITH_DATAREQUEST_HELPER
+	// Retry request until we get an accepted response or exhaust allowed number of attempts.
+	for (int32 Attempts = 0; Attempts < UE_HTTPDDC_MAX_ATTEMPTS; ++Attempts)
+	{
+		FDataRequestHelper RequestHelper(GetRequestPools[IsInGameThread()].Get(), *Namespace, *DefaultBucket, CacheKey, nullptr);
+		const int64 ResponseCode = RequestHelper.GetResponseCode();
+
+		if (FHttpRequest::IsSuccessResponse(ResponseCode) && RequestHelper.IsSuccess())
+		{
+			COOK_STAT(Timer.AddHit(0));
+			return true;
+		}
+
+		if (!ShouldRetryOnError(ResponseCode))
+		{
+			return false;
+		}
+	}
+#else
+	FString Uri = FString::Printf(TEXT("api/v1/c/ddc/%s/%s/%s"), *Namespace, *DefaultBucket, CacheKey);
+
+	// Retry request until we get an accepted response or exhaust allowed number of attempts.
+	for (int32 Attempts = 0; Attempts < UE_HTTPDDC_MAX_ATTEMPTS; ++Attempts)
+	{
+		FScopedRequestPtr Request(GetRequestPools[IsInGameThread()].Get());
+		const FHttpRequest::Result Result = Request->PerformBlockingQuery<FHttpRequest::Head>(*Uri);
+		const int64 ResponseCode = Request->GetResponseCode();
+
+		if (FHttpRequest::IsSuccessResponse(ResponseCode) || ResponseCode == 400)
+		{
+			const bool bIsHit = (Result == FHttpRequest::Success && FHttpRequest::IsSuccessResponse(ResponseCode));
+			if (bIsHit)
+			{
+				TRACE_COUNTER_ADD(HttpDDC_ExistHit, int64(1));
+				COOK_STAT(Timer.AddHit(0));
+			}
+			return bIsHit;
+		}
+
+		if (!ShouldRetryOnError(ResponseCode))
+		{
+			break;
+		}
+	}
+#endif
+
+	return false;
+}
+
+TBitArray<> FHttpCacheStore::CachedDataProbablyExistsBatch(TConstArrayView<FString> CacheKeys)
+{
+	TRACE_CPUPROFILER_EVENT_SCOPE(HttpDDC_Exist);
+	TRACE_COUNTER_ADD(HttpDDC_Exist, int64(1));
+	COOK_STAT(auto Timer = UsageStats.TimeProbablyExists());
+#if WITH_DATAREQUEST_HELPER
+	for (int32 Attempts = 0; Attempts < UE_HTTPDDC_MAX_ATTEMPTS; ++Attempts)
+	{
+		FDataRequestHelper RequestHelper(GetRequestPools[IsInGameThread()].Get(), *Namespace, *DefaultBucket, CacheKeys);
+		const int64 ResponseCode = RequestHelper.GetResponseCode();
+
+		if (FHttpRequest::IsSuccessResponse(ResponseCode) && RequestHelper.IsSuccess())
+		{
+			COOK_STAT(Timer.AddHit(0));
+			TBitArray<> Results = RequestHelper.IsBatchSuccess();
+			int32 ResultIndex = 0;
+			for (const FString& CacheKey : CacheKeys)
+			{
+				if (ShouldSimulateMiss(*CacheKey))
+				{
+					Results[ResultIndex] = false;
+				}
+				ResultIndex++;
+			}
+
+			return Results;
+		}
+
+		if (!ShouldRetryOnError(ResponseCode))
+		{
+			TBitArray<> Results = RequestHelper.IsBatchSuccess();
+			int32 ResultIndex = 0;
+			for (const FString& CacheKey : CacheKeys)
+			{
+				if (ShouldSimulateMiss(*CacheKey))
+				{
+					Results[ResultIndex] = false;
+				}
+				ResultIndex++;
+			}
+
+			return Results;
+		}
+	}
+#else
+	const TCHAR* const Uri = TEXT("api/v1/c/ddc-rpc");
+
+	TAnsiStringBuilder<512> Body;
+	const FTCHARToUTF8 AnsiNamespace(*Namespace);
+	const FTCHARToUTF8 AnsiBucket(*DefaultBucket);
+	Body << "{\"Operations\":[";
+	for (const FString& CacheKey : CacheKeys)
+	{
+		Body << "{\"Namespace\":\"" << AnsiNamespace.Get() << "\",\"Bucket\":\"" << AnsiBucket.Get() << "\",";
+		Body << "\"Id\":\"" << FTCHARToUTF8(*CacheKey).Get() << "\",\"Op\":\"HEAD\"},";
+	}
+	Body.RemoveSuffix(1);
+	Body << "]}";
+
+	TConstArrayView<uint8> BodyView(reinterpret_cast<const uint8*>(Body.ToString()), Body.Len());
+
+	// Retry request until we get an accepted response or exhaust allowed number of attempts.
+	for (int32 Attempts = 0; Attempts < UE_HTTPDDC_MAX_ATTEMPTS; ++Attempts)
+	{
+		FScopedRequestPtr Request(GetRequestPools[IsInGameThread()].Get());
+		const FHttpRequest::Result Result = Request->PerformBlockingUpload<FHttpRequest::PostJson>(Uri, BodyView);
+		const int64 ResponseCode = Request->GetResponseCode();
+
+		if (Result == FHttpRequest::Success && ResponseCode == 200)
+		{
+			TArray<TSharedPtr<FJsonValue>> ResponseArray = Request->GetResponseAsJsonArray();
+
+			TBitArray<> Exists;
+			Exists.Reserve(CacheKeys.Num());
+			for (const FString& CacheKey : CacheKeys)
+			{
+				if (ShouldSimulateMiss(*CacheKey))
+				{
+					Exists.Add(false);
+				}
+				else
+				{
+					const TSharedPtr<FJsonValue>* FoundResponse = Algo::FindByPredicate(ResponseArray, [&CacheKey](const TSharedPtr<FJsonValue>& Response) {
+						FString Key;
+						Response->TryGetString(Key);
+						return Key == CacheKey;
+					});
+
+					Exists.Add(FoundResponse != nullptr);
+				}
+			}
+
+			if (Exists.CountSetBits() == CacheKeys.Num())
+			{
+				TRACE_COUNTER_ADD(HttpDDC_ExistHit, int64(1));
+				COOK_STAT(Timer.AddHit(0));
+			}
+			return Exists;
+		}
+
+		if (!ShouldRetryOnError(ResponseCode))
+		{
+			break;
+		}
+	}
+#endif
+
+	return TBitArray<>(false, CacheKeys.Num());
+}
+
+bool FHttpCacheStore::GetCachedData(const TCHAR* CacheKey, TArray<uint8>& OutData)
+{
+	TRACE_CPUPROFILER_EVENT_SCOPE(HttpDDC_Get);
+	TRACE_COUNTER_ADD(HttpDDC_Get, int64(1));
+	COOK_STAT(auto Timer = UsageStats.TimeGet());
+
+	if (ShouldSimulateMiss(CacheKey))
+	{
+		return false;
+	}
+
+#if WITH_DATAREQUEST_HELPER
+	// Retry request until we get an accepted response or exhaust allowed number of attempts.
+	for (int32 Attempts = 0; Attempts < UE_HTTPDDC_MAX_ATTEMPTS; ++Attempts)
+	{
+		FDataRequestHelper RequestHelper(GetRequestPools[IsInGameThread()].Get(), *Namespace, *DefaultBucket, CacheKey, &OutData);
+		const int64 ResponseCode = RequestHelper.GetResponseCode();
+
+		if (FHttpRequest::IsSuccessResponse(ResponseCode) && RequestHelper.IsSuccess())
+		{
+			COOK_STAT(Timer.AddHit(OutData.Num()));
+			check(OutData.Num() > 0);
+			return true;
+		}
+
+		if (!ShouldRetryOnError(ResponseCode))
+		{
+			return false;
+		}
+	}
+#else 
+	FString Uri = FString::Printf(TEXT("api/v1/c/ddc/%s/%s/%s.raw"), *Namespace, *DefaultBucket, CacheKey);
+
+	// Retry request until we get an accepted response or exhaust allowed number of attempts.
+	for (int32 Attempts = 0; Attempts < UE_HTTPDDC_MAX_ATTEMPTS; ++Attempts)
+	{
+		FScopedRequestPtr Request(GetRequestPools[IsInGameThread()].Get());
+		if (Request.IsValid())
+		{
+			FHttpRequest::Result Result = Request->PerformBlockingDownload(*Uri, &OutData);
+			const uint64 ResponseCode = Request->GetResponseCode();
+
+			// Request was successful, make sure we got all the expected data.
+			if (FHttpRequest::IsSuccessResponse(ResponseCode) && VerifyRequest(Request.Get(), *Namespace, *DefaultBucket, CacheKey, OutData))
+			{
+				TRACE_COUNTER_ADD(HttpDDC_GetHit, int64(1));
+				TRACE_COUNTER_ADD(HttpDDC_BytesReceived, int64(Request->GetBytesReceived()));
+				COOK_STAT(Timer.AddHit(Request->GetBytesReceived()));
+				return true;
+			}
+
+			if (!ShouldRetryOnError(ResponseCode))
+			{
+				return false;
+			}
+		}
+	}
+#endif
+
+	return false;
+}
+
+FDerivedDataBackendInterface::EPutStatus FHttpCacheStore::PutCachedData(const TCHAR* CacheKey, TArrayView<const uint8> InData, bool bPutEvenIfExists)
+{
+	TRACE_CPUPROFILER_EVENT_SCOPE(HttpDDC_Put);
+
+	if (!IsWritable())
+	{
+		UE_LOG(LogDerivedDataCache, Verbose, TEXT("%s is read only. Skipping put of %s"), *GetName(), CacheKey);
+		return EPutStatus::NotCached;
+	}
+
+	// don't put anything we pretended didn't exist
+	if (ShouldSimulateMiss(CacheKey))
+	{
+		return EPutStatus::Skipped;
+	}
+
+#if 0 // No longer WITH_DATAREQUEST_HELPER as async puts are unsupported except through the AsyncPutWrapper which expects the inner backend to perform the put synchronously
+	for (int32 Attempts = 0; Attempts < UE_HTTPDDC_MAX_ATTEMPTS; ++Attempts)
+	{
+		FDataUploadHelper Request(PutRequestPools[IsInGameThread()].Get(), *Namespace, *DefaultBucket,	CacheKey, InData, UsageStats);
+
+		if (ShouldAbortForShutdown())
+		{
+			return EPutStatus::NotCached;
+		}
+
+		const int64 ResponseCode = Request.GetResponseCode();
+
+		if (Request.IsSuccess() && (Request.IsQueued() || FHttpRequest::IsSuccessResponse(ResponseCode)))
+		{
+			return Request.IsQueued() ? EPutStatus::Executing : EPutStatus::Cached;
+		}
+
+		if (!ShouldRetryOnError(ResponseCode))
+		{
+			return EPutStatus::NotCached;
+		}
+	}
+#else
+	COOK_STAT(auto Timer = UsageStats.TimePut());
+
+	FString Uri = FString::Printf(TEXT("api/v1/c/ddc/%s/%s/%s"), *Namespace, *DefaultBucket, CacheKey);
+	int64 ResponseCode = 0; uint32 Attempts = 0;
+
+	// Retry request until we get an accepted response or exhaust allowed number of attempts.
+	while (ResponseCode == 0 && ++Attempts < UE_HTTPDDC_MAX_ATTEMPTS)
+	{
+		if (ShouldAbortForShutdown())
+		{
+			return EPutStatus::NotCached;
+		}
+
+		FScopedRequestPtr Request(PutRequestPools[IsInGameThread()].Get());
+		if (Request.IsValid())
+		{
+			// Append the content hash to the header
+			HashPayload(Request.Get(), InData);
+
+			Request->PerformBlockingUpload<FHttpRequest::Put>(*Uri, InData);
+			ResponseCode = Request->GetResponseCode();
+
+			if (FHttpRequest::IsSuccessResponse(ResponseCode))
+			{
+				TRACE_COUNTER_ADD(HttpDDC_BytesSent, int64(Request->GetBytesSent()));
+				COOK_STAT(Timer.AddHit(Request->GetBytesSent()));
+				return EPutStatus::Cached;
+			}
+
+			if (!ShouldRetryOnError(ResponseCode))
+			{
+				return EPutStatus::NotCached;
+			}
+
+			ResponseCode = 0;
+		}
+	}
+#endif // WITH_DATAREQUEST_HELPER
+
+	return EPutStatus::NotCached;
+}
+
+void FHttpCacheStore::RemoveCachedData(const TCHAR* CacheKey, bool bTransient)
+{
+	// do not remove transient data as Jupiter does its own verification of the content and cleans itself up
+	if (!IsWritable() || bTransient)
+		return;
+	
+	TRACE_CPUPROFILER_EVENT_SCOPE(HttpDDC_Remove);
+	FString Uri = FString::Printf(TEXT("api/v1/c/ddc/%s/%s/%s"), *Namespace, *DefaultBucket, CacheKey);
+	int64 ResponseCode = 0; uint32 Attempts = 0;
+
+	// Retry request until we get an accepted response or exhaust allowed number of attempts.
+	while (ResponseCode == 0 && ++Attempts < UE_HTTPDDC_MAX_ATTEMPTS)
+	{
+		FScopedRequestPtr Request(PutRequestPools[IsInGameThread()].Get());
+		if (Request.IsValid())
+		{
+			FHttpRequest::Result Result = Request->PerformBlockingQuery<FHttpRequest::Delete>(*Uri, {});
+			ResponseCode = Request->GetResponseCode();
+
+			if (ResponseCode == 200)
+			{
+				return;
+			}
+
+			if (!ShouldRetryOnError(ResponseCode))
+			{
+				return;
+			}
+
+			ResponseCode = 0;
+		}
+	}
+}
+
+TSharedRef<FDerivedDataCacheStatsNode> FHttpCacheStore::GatherUsageStats() const
+{
+	TSharedRef<FDerivedDataCacheStatsNode> Usage =
+		MakeShared<FDerivedDataCacheStatsNode>(TEXT("Horde Storage"), FString::Printf(TEXT("%s (%s)"), *Domain, *Namespace), /*bIsLocal*/ false);
+	Usage->Stats.Add(TEXT(""), UsageStats);
+	return Usage;
+}
+
+void FHttpCacheStore::Put(
+	const TConstArrayView<FCachePutRequest> Requests,
+	IRequestOwner& Owner,
+	FOnCachePutComplete&& OnComplete)
+{
+	for (const FCachePutRequest& Request : Requests)
+	{
+		const FCacheRecord& Record = Request.Record;
+		COOK_STAT(auto Timer = UsageStats.TimePut());
+		uint64 BytesSent = 0;
+		if (PutCacheRecord(Request.Name, Record, Request.Policy, BytesSent))
+		{
+			UE_LOG(LogDerivedDataCache, Verbose, TEXT("%s: Cache put complete for %s from '%s'"),
+				*GetName(), *WriteToString<96>(Record.GetKey()), *Request.Name);
+			TRACE_COUNTER_ADD(HttpDDC_BytesSent, int64(BytesSent));
+			COOK_STAT(Timer.AddHit(BytesSent));
+			OnComplete(Request.MakeResponse(EStatus::Ok));
+		}
+		else
+		{
+			OnComplete(Request.MakeResponse(EStatus::Error));
+		}
+	}
+}
+
+void FHttpCacheStore::Get(
+	const TConstArrayView<FCacheGetRequest> Requests,
+	IRequestOwner& Owner,
+	FOnCacheGetComplete&& OnComplete)
+{
+	for (const FCacheGetRequest& Request : Requests)
+	{
+		COOK_STAT(auto Timer = UsageStats.TimeGet());
+		EStatus Status = EStatus::Ok;
+		if (FOptionalCacheRecord Record = GetCacheRecord(Request.Name, Request.Key, Request.Policy, Status))
+		{
+			UE_LOG(LogDerivedDataCache, Verbose, TEXT("%s: Cache hit for %s from '%s'"),
+				*GetName(), *WriteToString<96>(Request.Key), *Request.Name);
+			TRACE_COUNTER_ADD(HttpDDC_BytesReceived, Private::GetCacheRecordCompressedSize(Record.Get()));
+			COOK_STAT(Timer.AddHit(Private::GetCacheRecordCompressedSize(Record.Get())));
+			OnComplete({Request.Name, MoveTemp(Record).Get(), Request.UserData, Status});
+		}
+		else
+		{
+			OnComplete(Request.MakeResponse(Status));
+		}
+	}
+}
+
+void FHttpCacheStore::PutValue(
+	const TConstArrayView<FCachePutValueRequest> Requests,
+	IRequestOwner& Owner,
+	FOnCachePutValueComplete&& OnComplete)
+{
+	for (const FCachePutValueRequest& Request : Requests)
+	{
+		COOK_STAT(auto Timer = UsageStats.TimePut());
+		uint64 WriteSize = 0;
+		if (PutCacheValue(Request.Name, Request.Key, Request.Value, Request.Policy, WriteSize))
+		{
+			UE_LOG(LogDerivedDataCache, Verbose, TEXT("%s: Cache put complete for %s from '%s'"),
+				*GetName(), *WriteToString<96>(Request.Key), *Request.Name);
+			TRACE_COUNTER_ADD(HttpDDC_BytesSent, WriteSize);
+			COOK_STAT(if (WriteSize) { Timer.AddHit(WriteSize); });
+			OnComplete(Request.MakeResponse(EStatus::Ok));
+		}
+		else
+		{
+			OnComplete(Request.MakeResponse(EStatus::Error));
+		}
+	}
+}
+
+void FHttpCacheStore::GetValue(
+	const TConstArrayView<FCacheGetValueRequest> Requests,
+	IRequestOwner& Owner,
+	FOnCacheGetValueComplete&& OnComplete)
+{
+	COOK_STAT(double StartTime = FPlatformTime::Seconds());
+	COOK_STAT(bool bIsInGameThread = IsInGameThread());
+
+	FScopedRequestPtr HttpRequest(GetRequestPools[IsInGameThread()].Get());
+	int64 HitBytes = 0;
+
+	for (const FCacheGetValueRequest& Request : Requests)
+	{
+		FValue Value;
+		if (!GetCacheValue(Request.Name, Request.Key, Request.Policy, Value, HttpRequest.Get()))
+		{
+			COOK_STAT(UsageStats.GetStats.Accumulate(FCookStats::CallStats::EHitOrMiss::Miss, FCookStats::CallStats::EStatType::Counter, 1l, bIsInGameThread));
+			OnComplete(Request.MakeResponse(EStatus::Error));
+		}
+		else
+		{
+			if (!IsValueDataReady(Value, Request.Policy) && !EnumHasAnyFlags(Request.Policy, ECachePolicy::SkipData))
+			{
+				// With inline fetching, expect we will always have a value we can use.  Even SkipData/Exists can rely on the blob existing if the ref is reported to exist.
+				UE_LOG(LogDerivedDataCache, Warning, TEXT("%s: Cache miss due to inlining failure for %s from '%s'"),
+					*GetName(), *WriteToString<96>(Request.Key), *Request.Name);
+				COOK_STAT(UsageStats.GetStats.Accumulate(FCookStats::CallStats::EHitOrMiss::Miss, FCookStats::CallStats::EStatType::Counter, 1l, bIsInGameThread));
+				OnComplete(Request.MakeResponse(EStatus::Error));
+			}
+			else
+			{
+				UE_LOG(LogDerivedDataCache, Verbose, TEXT("%s: Cache hit for %s from '%s'"),
+					*GetName(), *WriteToString<96>(Request.Key), *Request.Name);
+				uint64 ValueSize = Value.GetData().GetCompressedSize();
+				TRACE_COUNTER_ADD(HttpDDC_BytesReceived, ValueSize);
+				HitBytes += ValueSize;
+				COOK_STAT(UsageStats.GetStats.Accumulate(FCookStats::CallStats::EHitOrMiss::Hit, FCookStats::CallStats::EStatType::Counter, 1l, bIsInGameThread));
+				OnComplete({ Request.Name, Request.Key, Value, Request.UserData, EStatus::Ok });
+			}
+		}
+	}
+
+	COOK_STAT(const int64 CyclesUsed = int64((FPlatformTime::Seconds() - StartTime) / FPlatformTime::GetSecondsPerCycle()));
+	COOK_STAT(UsageStats.GetStats.Accumulate(FCookStats::CallStats::EHitOrMiss::Hit, FCookStats::CallStats::EStatType::Cycles, CyclesUsed, bIsInGameThread));
+	COOK_STAT(UsageStats.GetStats.Accumulate(FCookStats::CallStats::EHitOrMiss::Hit, FCookStats::CallStats::EStatType::Bytes, HitBytes, bIsInGameThread));
+}
+
+void FHttpCacheStore::GetChunks(
+	const TConstArrayView<FCacheGetChunkRequest> Requests,
+	IRequestOwner& Owner,
+	FOnCacheGetChunkComplete&& OnComplete)
+{
+	// TODO: This is inefficient because Jupiter doesn't allow us to get only part of a compressed blob, so we have to
+	//		 get the whole thing and then decompress only the portion we need.  Furthermore, because there is no propagation
+	//		 between cache stores during chunk requests, the fetched result won't end up in the local store.
+	//		 These efficiency issues will be addressed by changes to the Hierarchy that translate chunk requests that
+	//		 are missing in local/fast stores and have to be retrieved from slow stores into record requests instead.  That
+	//		 will make this code path unused/uncommon as Jupiter will most always be a slow store with a local/fast store in front of it.
+	//		 Regardless, to adhere to the functional contract, this implementation must exist.
+	TArray<FCacheGetChunkRequest, TInlineAllocator<16>> SortedRequests(Requests);
+	SortedRequests.StableSort(TChunkLess());
+
+	bool bHasValue = false;
+	FValue Value;
+	FValueId ValueId;
+	FCacheKey ValueKey;
+	FCompressedBuffer ValueBuffer;
+	FCompressedBufferReader ValueReader;
+	EStatus ValueStatus = EStatus::Error;
+	FOptionalCacheRecord Record;
+	for (const FCacheGetChunkRequest& Request : SortedRequests)
+	{
+		const bool bExistsOnly = EnumHasAnyFlags(Request.Policy, ECachePolicy::SkipData);
+		COOK_STAT(auto Timer = bExistsOnly ? UsageStats.TimeProbablyExists() : UsageStats.TimeGet());
+		if (!(bHasValue && ValueKey == Request.Key && ValueId == Request.Id) || ValueReader.HasSource() < !bExistsOnly)
+		{
+			ValueStatus = EStatus::Error;
+			ValueReader.ResetSource();
+			ValueKey = {};
+			ValueId.Reset();
+			Value.Reset();
+			bHasValue = false;
+			if (Request.Id.IsValid())
+			{
+				if (!(Record && Record.Get().GetKey() == Request.Key))
+				{
+					FCacheRecordPolicyBuilder PolicyBuilder(ECachePolicy::None);
+					PolicyBuilder.AddValuePolicy(Request.Id, Request.Policy);
+					Record.Reset();
+					Record = GetCacheRecordOnly(Request.Name, Request.Key, PolicyBuilder.Build());
+				}
+				if (Record)
+				{
+					const FValueWithId& ValueWithId = Record.Get().GetValue(Request.Id);
+					bHasValue = ValueWithId.IsValid();
+					Value = ValueWithId;
+					ValueId = Request.Id;
+					ValueKey = Request.Key;
+
+					if (IsValueDataReady(Value, Request.Policy))
+					{
+						ValueReader.SetSource(Value.GetData());
+					}
+					else
+					{
+						auto HashGetter = [](const FValueWithId& Value)
+						{
+							return Value.GetRawHash();
+						};
+
+						auto DebugContextGetter = [&Request](const FValueWithId& Value)
+						{
+							return FValueDebugContext{ Request.Name, Request.Key, FString(*WriteToString<16>(Value.GetId())) };
+						};
+
+						TArray<FCompressedBuffer> ValueBuffers;
+						if (TryGetCachedDataBatch<FValueWithId>(::MakeArrayView({ ValueWithId }), ValueBuffers, HashGetter, DebugContextGetter).CountSetBits() == 1)
+						{
+							ValueBuffer = ValueBuffers[0];
+							ValueReader.SetSource(ValueBuffer);
+						}
+						else
+						{
+							ValueBuffer.Reset();
+							ValueReader.ResetSource();
+						}
+					}
+				}
+			}
+			else
+			{
+				ValueKey = Request.Key;
+				bHasValue = GetCacheValue(Request.Name, Request.Key, Request.Policy, Value);
+				if (IsValueDataReady(Value, Request.Policy))
+				{
+					ValueReader.SetSource(Value.GetData());
+				}
+				else
+				{
+					auto HashGetter = [](const FValue& Value)
+					{
+						return Value.GetRawHash();
+					};
+
+					auto DebugContextGetter = [&Request](const FValue& Value)
+					{
+						return FValueDebugContext{ Request.Name, Request.Key, FString(TEXT("Default")) };
+					};
+
+					TArray<FCompressedBuffer> ValueBuffers;
+					if (TryGetCachedDataBatch<FValue>({ Value }, ValueBuffers, HashGetter, DebugContextGetter).CountSetBits() == 1)
+					{
+						ValueBuffer = ValueBuffers[0];
+						ValueReader.SetSource(ValueBuffer);
+					}
+					else
+					{
+						ValueBuffer.Reset();
+						ValueReader.ResetSource();
+					}
+				}
+			}
+		}
+		if (bHasValue)
+		{
+			const uint64 RawOffset = FMath::Min(Value.GetRawSize(), Request.RawOffset);
+			const uint64 RawSize = FMath::Min(Value.GetRawSize() - RawOffset, Request.RawSize);
+			UE_LOG(LogDerivedDataCache, Verbose, TEXT("%s: Cache hit for %s from '%s'"),
+				*GetName(), *WriteToString<96>(Request.Key, '/', Request.Id), *Request.Name);
+			COOK_STAT(Timer.AddHit(!bExistsOnly ? RawSize : 0));
+			FSharedBuffer Buffer;
+			if (!bExistsOnly)
+			{
+				Buffer = ValueReader.Decompress(RawOffset, RawSize);
+			}
+			const EStatus ChunkStatus = bExistsOnly || Buffer.GetSize() == RawSize ? EStatus::Ok : EStatus::Error;
+			OnComplete({Request.Name, Request.Key, Request.Id, Request.RawOffset,
+				RawSize, Value.GetRawHash(), MoveTemp(Buffer), Request.UserData, ChunkStatus});
+			continue;
+		}
+
+		OnComplete(Request.MakeResponse(EStatus::Error));
+	}
+}
+
+} // UE::DerivedData
+
+#endif // WITH_HTTP_DDC_BACKEND
+
+namespace UE::DerivedData
+{
+
+ILegacyCacheStore* CreateHttpCacheStore(
+	const TCHAR* NodeName,
+	const TCHAR* ServiceUrl,
+	const TCHAR* Namespace,
+	const TCHAR* StructuredNamespace,
+	const TCHAR* OAuthProvider,
+	const TCHAR* OAuthClientId,
+	const TCHAR* OAuthData,
+	const FDerivedDataBackendInterface::ESpeedClass* ForceSpeedClass,
+	EBackendLegacyMode LegacyMode,
+	bool bReadOnly)
+{
+#if WITH_HTTP_DDC_BACKEND
+	FHttpCacheStore* Backend = new FHttpCacheStore(ServiceUrl, Namespace, StructuredNamespace, OAuthProvider, OAuthClientId, OAuthData, LegacyMode, bReadOnly);
+	if (Backend->IsUsable())
+	{
+		return Backend;
+	}
+	UE_LOG(LogDerivedDataCache, Warning, TEXT("Node %s could not contact the service (%s), will not use it"), NodeName, ServiceUrl);
+	delete Backend;
+	return nullptr;
+#else
+	UE_LOG(LogDerivedDataCache, Warning, TEXT("HTTP backend is not yet supported in the current build configuration."));
+	return nullptr;
+#endif
+}
+
+FDerivedDataBackendInterface* GetAnyHttpCacheStore(
+	FString& OutDomain,
+	FString& OutOAuthProvider,
+	FString& OutOAuthClientId,
+	FString& OutOAuthSecret,
+	FString& OutNamespace,
+	FString& OutStructuredNamespace)
+{
+#if WITH_HTTP_DDC_BACKEND
+	if (FHttpCacheStore* HttpBackend = FHttpCacheStore::GetAny())
+	{
+		OutDomain = HttpBackend->GetDomain();
+		OutOAuthProvider = HttpBackend->GetOAuthProvider();
+		OutOAuthClientId = HttpBackend->GetOAuthClientId();
+		OutOAuthSecret = HttpBackend->GetOAuthSecret();
+		OutNamespace = HttpBackend->GetNamespace();
+		OutStructuredNamespace = HttpBackend->GetStructuredNamespace();
+
+		return HttpBackend;
+	}
+	return nullptr;
+#else
+	return nullptr;
+#endif
+}
+
+} // UE::DerivedData
diff --git a/Engine/Source/Developer/DerivedDataCache/Private/HttpDerivedDataCacheStore.cpp b/Engine/Source/Developer/DerivedDataCache/Private/HttpDerivedDataCacheStore.cpp
deleted file mode 100644
