Fixes for the implicit conversion of ints and chars to TCHARs in UTF-8 mode.

diff --git a/Engine/Source/Developer/DerivedDataCache/Private/HttpCacheStore.cpp b/Engine/Source/Developer/DerivedDataCache/Private/HttpCacheStore.cpp
--- a/Engine/Source/Developer/DerivedDataCache/Private/HttpCacheStore.cpp
+++ b/Engine/Source/Developer/DerivedDataCache/Private/HttpCacheStore.cpp
@@ -628,8 +628,8 @@ private:
 			{
 				// Print the response body if we got one, otherwise print header.
 				FString Response = GetAnsiBufferAsString(ResponseBuffer.Num() > 0 ? ResponseBuffer : ResponseHeader);
-				Response.ReplaceCharInline('\n', ' ');
-				Response.ReplaceCharInline('\r', ' ');
+				Response.ReplaceCharInline(TEXT('\n'), TEXT(' '));
+				Response.ReplaceCharInline(TEXT('\r'), TEXT(' '));
 				// Dont log access denied as error, since tokens can expire mid session
 				if (ResponseCode == 401)
 				{
diff --git a/Engine/Source/Developer/DerivedDataCache/Private/ZenCacheStore.cpp b/Engine/Source/Developer/DerivedDataCache/Private/ZenCacheStore.cpp
--- a/Engine/Source/Developer/DerivedDataCache/Private/ZenCacheStore.cpp
+++ b/Engine/Source/Developer/DerivedDataCache/Private/ZenCacheStore.cpp
@@ -486,14 +486,14 @@ void FZenCacheStore::AppendZenUri(const FCacheKey& CacheKey, const FValueId& Id,
 {
 	AppendZenUri(CacheKey, Out);
 	Out << TEXT('/') << Id;
-}
+		}
 
 void FZenCacheStore::AppendPolicyQueryString(ECachePolicy Policy, FStringBuilderBase& Uri)
-{
-	if (Policy != ECachePolicy::Default)
 	{
+	if (Policy != ECachePolicy::Default)
+		{
 		Uri << TEXT("?Policy=") << Policy;
-	}
+}
 }
 
 void FZenCacheStore::RemoveCachedData(const TCHAR* CacheKey, bool bTransient)
@@ -671,15 +671,15 @@ void FZenCacheStore::Put(
 			const FCacheRecord& Record = NextRequest.Record;
 			uint64 RecordSize = sizeof(FCacheKey) + Record.GetMeta().GetSize();
 			for (const FValueWithId& Value : Record.GetValues())
-			{
+	{
 				RecordSize += Value.GetData().GetCompressedSize();
 			}
 			BatchSize += RecordSize;
 			if (BatchSize > BatchPutMaxBytes)
-			{
+		{
 				BatchSize = RecordSize;
 				return EBatchView::NewBatch;
-			}
+		}
 			return EBatchView::Continue;
 		});
 
@@ -718,15 +718,15 @@ void FZenCacheStore::Put(
 					const FCacheRecord& Record = Request.Record;
 
 					BatchWriter.BeginObject();
-					{
+		{
 						BatchWriter.SetName("Record"_ASV);
 						Record.Save(BatchPackage, BatchWriter);
 						if (!Request.Policy.IsUniform() || Request.Policy.GetRecordPolicy() != BatchDefaultPolicy)
-						{
+			{
 							BatchWriter.SetName("Policy"_ASV);
 							Request.Policy.Save(BatchWriter);
-						}
-					}
+			}
+		}
 					BatchWriter.EndObject();
 				}
 				BatchWriter.EndArray();
@@ -760,7 +760,7 @@ void FZenCacheStore::Put(
 				const FCacheKey& Key = Request.Record.GetKey();
 				bool bPutSucceeded = ResponseField.AsBool();
 				if (ShouldSimulateMiss(Key))
-				{
+		{
 					UE_LOG(LogDerivedDataCache, Verbose, TEXT("%s: Simulated miss for put of %s from '%s'"),
 						*GetName(), *WriteToString<96>(Key), *Request.Name);
 					bPutSucceeded = false;
@@ -814,7 +814,7 @@ void FZenCacheStore::Get(
 		[this, &Requests, &Owner, &OnComplete, &OnHit, &OnMiss](int32 BatchFirst, int32 BatchLast)
 	{
 		TArrayView<const FCacheGetRequest> Batch(Requests.GetData() + BatchFirst, BatchLast - BatchFirst + 1);
-
+		
 		FCbWriter BatchRequest;
 		BatchRequest.BeginObject();
 		{
@@ -828,18 +828,18 @@ void FZenCacheStore::Get(
 				for (const FCacheGetRequest& Request : Batch)
 				{
 					BatchRequest.BeginObject();
-					{
+				{
 						BatchRequest << "Key"_ASV << Request.Key;
 						if (!Request.Policy.IsUniform() || Request.Policy.GetRecordPolicy() != BatchDefaultPolicy)
-						{
+					{
 							BatchRequest.SetName("Policy"_ASV);
 							Request.Policy.Save(BatchRequest);
 						}
 					}
-					BatchRequest.EndObject();
-				}
-				BatchRequest.EndArray();
-			}
+							BatchRequest.EndObject();
+						}
+						BatchRequest.EndArray();
+					}
 			BatchRequest.EndObject();
 		}
 		BatchRequest.EndObject();
@@ -861,12 +861,12 @@ void FZenCacheStore::Get(
 			for (FCbField RecordField : ResponseObj["Result"_ASV])
 			{
 				if (RequestIndex >= Batch.Num())
-				{
+			{
 					++RequestIndex;
 					continue;
 				}
 				const FCacheGetRequest& Request = Batch[RequestIndex++];
-
+				
 				const FCacheKey& Key = Request.Key;
 				FOptionalCacheRecord Record;
 
@@ -900,7 +900,7 @@ void FZenCacheStore::PutValue(
 	TConstArrayView<FCachePutValueRequest> Requests,
 	IRequestOwner& Owner,
 	FOnCachePutValueComplete&& OnComplete)
-{
+				{
 	TRACE_CPUPROFILER_EVENT_SCOPE(ZenDDC::PutValue);
 	COOK_STAT(auto Timer = UsageStats.TimePut());
 
@@ -1001,36 +1001,36 @@ void FZenCacheStore::PutValue(
 				bPutSucceeded ? OnHit(Request) : OnMiss(Request);
 			}
 			if (RequestIndex != Batch.Num())
-			{
+					{
 				UE_LOG(LogDerivedDataCache, Warning,
 					TEXT("Invalid response received from PutCacheValues rpc: %d results expected, received %d."), Batch.Num(), RequestIndex);
 			}
-		}
+					}
 		for (const FCachePutValueRequest& Request : Batch.RightChop(RequestIndex))
-		{
+					{
 			OnMiss(Request);
 		}
-	}
-}
+					}
+				}
 
 void FZenCacheStore::GetValue(
 	TConstArrayView<FCacheGetValueRequest> Requests,
 	IRequestOwner& Owner,
 	FOnCacheGetValueComplete&& OnComplete)
-{
+				{
 	TRACE_CPUPROFILER_EVENT_SCOPE(ZenDDC::GetValue);
 	TRACE_COUNTER_ADD(ZenDDC_CacheRecordRequestCount, int64(Requests.Num()));
 	COOK_STAT(auto Timer = UsageStats.TimeGet());
 	TRACE_COUNTER_ADD(ZenDDC_Get, (int64)Requests.Num());
-
+					
 	auto OnHit = [this, &OnComplete COOK_STAT(, &Timer)](const FCacheGetValueRequest& Request, FValue&& Value)
 	{
 		UE_LOG(LogDerivedDataCache, Verbose, TEXT("%s: Cache hit for '%s' from '%s'"),
 			*GetName(), *WriteToString<96>(Request.Key), *Request.Name);
-		TRACE_COUNTER_ADD(ZenDDC_GetHit, int64(1));
+					TRACE_COUNTER_ADD(ZenDDC_GetHit, int64(1));
 		int64 ReceivedSize = Value.GetData().GetCompressedSize();
-		TRACE_COUNTER_ADD(ZenDDC_BytesReceived, ReceivedSize);
-		COOK_STAT(Timer.AddHit(ReceivedSize));
+					TRACE_COUNTER_ADD(ZenDDC_BytesReceived, ReceivedSize);
+					COOK_STAT(Timer.AddHit(ReceivedSize));
 
 		OnComplete({Request.Name, Request.Key, MoveTemp(Value), Request.UserData, EStatus::Ok});
 	};
@@ -1060,13 +1060,13 @@ void FZenCacheStore::GetValue(
 					for (const FCacheGetValueRequest& Request : Batch)
 					{
 						BatchRequest.BeginObject();
-						{
+					{
 							BatchRequest << "Key"_ASV << Request.Key;
 							if (Request.Policy != BatchDefaultPolicy)
 							{
 								BatchRequest << "Policy"_ASV << WriteToString<128>(Request.Policy);
 							}
-						}
+					}
 						BatchRequest.EndObject();
 					}
 					BatchRequest.EndArray();
@@ -1086,25 +1086,25 @@ void FZenCacheStore::GetValue(
 
 			int32 RequestIndex = 0;
 			if (HttpResult == Zen::FZenHttpRequest::Result::Success)
-			{
+				{
 				const FCbObject& ResponseObj = BatchResponse.GetObject();
-
+					
 				for (FCbFieldView ResultField : ResponseObj["Result"_ASV])
 				{
 					if (RequestIndex >= Batch.Num())
 					{
 						++RequestIndex;
 						continue;
-					}
+				}
 					const FCacheGetValueRequest& Request = Batch[RequestIndex++];
-
+				
 					FCbObjectView ResultObj = ResultField.AsObjectView();
 					TOptional<FValue> Value;
 					if (ShouldSimulateMiss(Request.Key))
 					{
 						UE_LOG(LogDerivedDataCache, Verbose, TEXT("%s: Simulated miss for GetValue of '%s' from '%s'"),
 							*GetName(), *WriteToString<96>(Request.Key), *Request.Name);
-					}
+			}
 					else
 					{
 						FCbFieldView RawHashField = ResultObj["RawHash"];
@@ -1112,13 +1112,13 @@ void FZenCacheStore::GetValue(
 						if (const FCbAttachment* Attachment = BatchResponse.FindAttachment(RawHash))
 						{
 							Value.Emplace(Attachment->AsCompressedBinary());
-						}
-						else
-						{
+		}
+		else
+		{
 							FCbFieldView RawSizeField = ResultObj["RawSize"];
 							uint64 RawSize = RawSizeField.AsUInt64();
 							if (!RawSizeField.HasError() && !RawHashField.HasError())
-							{
+			{
 								Value.Emplace(RawHash, RawSize);
 							}
 						}
@@ -1134,11 +1134,11 @@ void FZenCacheStore::GetValue(
 			for (const FCacheGetValueRequest& Request : Batch.RightChop(RequestIndex))
 			{
 				OnMiss(Request);
-			}
-		});
+		}
+	});
 	TRACE_COUNTER_SUBTRACT(ZenDDC_CacheRecordRequestCount, int64(Requests.Num()));
 }
-
+	
 
 void FZenCacheStore::GetChunks(
 	TConstArrayView<FCacheGetChunkRequest> Requests,
@@ -1241,17 +1241,17 @@ void FZenCacheStore::GetChunks(
 					continue;
 				}
 				const FCacheGetChunkRequest& Request = Batch[RequestIndex++];
-
+					
 				FIoHash RawHash;
 				bool Succeeded = false;
 				uint64 RawSize = 0;
 				FCbObjectView ResultObject = ResultView.AsObjectView();
 				FSharedBuffer RequestedBytes;
 				if (ShouldSimulateMiss(Request.Key))
-				{
+					{
 					UE_LOG(LogDerivedDataCache, Verbose, TEXT("%s: Simulated miss for get of '%s' from '%s'"),
 						*GetName(), *WriteToString<96>(Request.Key, '/', Request.Id), *Request.Name);
-				}
+					}
 				else
 				{
 					FCbFieldView HashView = ResultObject["RawHash"_ASV];
@@ -1259,31 +1259,31 @@ void FZenCacheStore::GetChunks(
 					if (!HashView.HasError())
 					{
 						if (const FCbAttachment* Attachment = BatchResponse.FindAttachment(HashView.AsHash()))
-						{
+				{
 							FCompressedBuffer CompressedBuffer = Attachment->AsCompressedBinary();
 							if (CompressedBuffer)
-							{
+					{
 								TRACE_COUNTER_ADD(ZenDDC_BytesReceived, CompressedBuffer.GetCompressedSize());
 								RequestedBytes = FCompressedBufferReader(CompressedBuffer).Decompress(Request.RawOffset, Request.RawSize);
 								RawSize = RequestedBytes.GetSize();
 								Succeeded = true;
-							}
-						}
-						else
-						{
+					}
+				}
+				else
+				{
 							FCbFieldView RawSizeField = ResultObject["RawSize"_ASV];
 							uint64 TotalSize = RawSizeField.AsUInt64();
 							Succeeded = !RawSizeField.HasError();
 							if (Succeeded)
-							{
+					{
 								RawSize = FMath::Min(Request.RawSize, TotalSize - FMath::Min(Request.RawOffset, TotalSize));
-							}
-						}
 					}
 				}
-				Succeeded ? OnHit(Request, MoveTemp(RawHash), RawSize, MoveTemp(RequestedBytes)) : OnMiss(Request);
 			}
 		}
+				Succeeded ? OnHit(Request, MoveTemp(RawHash), RawSize, MoveTemp(RequestedBytes)) : OnMiss(Request);
+				}
+			}
 		for (const FCacheGetChunkRequest& Request : Batch.RightChop(RequestIndex))
 		{
 			OnMiss(Request);
@@ -1303,7 +1303,7 @@ ILegacyCacheStore* CreateZenCacheStore(const TCHAR* NodeName, const TCHAR* Servi
 	UE_LOG(LogDerivedDataCache, Warning, TEXT("%s could not contact the service (%s), will not use it."), NodeName, *Backend->GetName());
 	delete Backend;
 	return nullptr;
-}
+	}
 
 } // namespace UE::DerivedData
 
diff --git a/Engine/Source/Developer/Zen/Internal/ZenServerHttp.cpp b/Engine/Source/Developer/Zen/Internal/ZenServerHttp.cpp
--- a/Engine/Source/Developer/Zen/Internal/ZenServerHttp.cpp
+++ b/Engine/Source/Developer/Zen/Internal/ZenServerHttp.cpp
@@ -456,7 +456,7 @@ namespace UE::Zen {
 	{
 		LLM_SCOPE_BYTAG(ZenDDC);
 		// Strip any leading slashes because we compose the prefix and the suffix with a separating slash below
-		while (Uri.StartsWith('/'))
+		while (Uri.StartsWith(TEXT('/')))
 		{
 			Uri.RightChopInline(1);
 		}
@@ -600,8 +600,8 @@ namespace UE::Zen {
 			{
 				// Print the response body if we got one, otherwise print header.
 				FString Response = GetAnsiBufferAsString(ResponseBuffer.Num() > 0 ? ResponseBuffer : ResponseHeader);
-				Response.ReplaceCharInline('\n', ' ');
-				Response.ReplaceCharInline('\r', ' ');
+				Response.ReplaceCharInline(TEXT('\n'), TEXT(' '));
+				Response.ReplaceCharInline(TEXT('\r'), TEXT(' '));
 				UE_LOG(
 					LogZenHttp,
 					Error,
@@ -813,7 +813,7 @@ namespace UE::Zen {
 	FZenHttpRequestPool::FZenHttpRequestPool(FStringView InServiceUrl, uint32 PoolEntryCount)
 	{
 		int32 LastSlash = 0;
-		while (InServiceUrl.FindLastChar('/', /* out */ LastSlash) && (LastSlash == (InServiceUrl.Len() - 1)))
+		while (InServiceUrl.FindLastChar(TEXT('/'), /* out */ LastSlash) && (LastSlash == (InServiceUrl.Len() - 1)))
 		{
 			InServiceUrl.LeftChopInline(1);
 		}
diff --git a/Engine/Source/Developer/Zen/Internal/ZenServerInterface.cpp b/Engine/Source/Developer/Zen/Internal/ZenServerInterface.cpp
--- a/Engine/Source/Developer/Zen/Internal/ZenServerInterface.cpp
+++ b/Engine/Source/Developer/Zen/Internal/ZenServerInterface.cpp
@@ -406,9 +406,9 @@ ReadCbLockFile(FStringView FileName, FCbObject& OutLockObject)
 	FPathViews::ToAbsolutePath(FileName, FullFileNameBuilder);
 	for (TCHAR& Char : MakeArrayView(FullFileNameBuilder))
 	{
-		if (Char == '/')
+		if (Char == TEXT('/'))
 		{
-			Char = '\\';
+			Char = TEXT('\\');
 		}
 	}
 	if (FullFileNameBuilder.Len() >= MAX_PATH)
diff --git a/Engine/Source/Programs/UnrealHeaderTool/Private/UnrealSourceFile.cpp b/Engine/Source/Programs/UnrealHeaderTool/Private/UnrealSourceFile.cpp
--- a/Engine/Source/Programs/UnrealHeaderTool/Private/UnrealSourceFile.cpp
+++ b/Engine/Source/Programs/UnrealHeaderTool/Private/UnrealSourceFile.cpp
@@ -69,7 +69,7 @@ const FString& FUnrealSourceFile::GetFileId() const
 			}
 			else
 			{
-				Out.AppendChar('_');
+				Out.AppendChar(TEXT('_'));
 			}
 		}
 
diff --git a/Engine/Source/Programs/UnrealHeaderTool/Private/UnrealTypeDefinitionInfo.cpp b/Engine/Source/Programs/UnrealHeaderTool/Private/UnrealTypeDefinitionInfo.cpp
--- a/Engine/Source/Programs/UnrealHeaderTool/Private/UnrealTypeDefinitionInfo.cpp
+++ b/Engine/Source/Programs/UnrealHeaderTool/Private/UnrealTypeDefinitionInfo.cpp
@@ -1029,7 +1029,7 @@ void FUnrealPropertyDefinitionInfo::ExportCppDeclaration(FOutputDevice& Out, EEx
 				else
 				{
 					// export as a pointer if this is an optional out parm, reference if it's just an out parm, standard otherwise...
-					TCHAR ModifierString[2] = { 0,0 };
+					TCHAR ModifierString[2] = { TCHAR('\0'), TCHAR('\0') };
 					if (bCanHaveRef && (HasAnyPropertyFlags(CPF_OutParm | CPF_ReferenceParm) || bIsInterfaceProp))
 					{
 						ModifierString[0] = TEXT('&');
diff --git a/Engine/Source/Runtime/Core/Private/GenericPlatform/GenericPlatformOutputDevices.cpp b/Engine/Source/Runtime/Core/Private/GenericPlatform/GenericPlatformOutputDevices.cpp
--- a/Engine/Source/Runtime/Core/Private/GenericPlatform/GenericPlatformOutputDevices.cpp
+++ b/Engine/Source/Runtime/Core/Private/GenericPlatform/GenericPlatformOutputDevices.cpp
@@ -61,7 +61,7 @@ void FGenericPlatformOutputDevices::SetupOutputDevices()
 void FGenericPlatformOutputDevices::ResetCachedAbsoluteFilename()
 {
 	FScopeLock ScopeLock(&LogFilenameLock);
-	CachedAbsoluteFilename[0] = 0;
+	CachedAbsoluteFilename[0] = TEXT('\0');
 }
 
 void FGenericPlatformOutputDevices::OnLogFileOpened(const TCHAR* Pathname)
@@ -83,7 +83,7 @@ FString FGenericPlatformOutputDevices::GetAbsoluteLogFilename()
 		{
 			if (FParse::Value(FCommandLine::Get(), TEXT("ABSLOG="), LogFilename, bShouldStopOnSeparator))
 			{
-				CachedAbsoluteFilename[0] = 0;
+				CachedAbsoluteFilename[0] = TEXT('\0');
 			}
 		}
 
diff --git a/Engine/Source/Runtime/Core/Private/GenericPlatform/GenericPlatformStackWalk.cpp b/Engine/Source/Runtime/Core/Private/GenericPlatform/GenericPlatformStackWalk.cpp
--- a/Engine/Source/Runtime/Core/Private/GenericPlatform/GenericPlatformStackWalk.cpp
+++ b/Engine/Source/Runtime/Core/Private/GenericPlatform/GenericPlatformStackWalk.cpp
@@ -139,8 +139,8 @@ bool FGenericPlatformStackWalk::SymbolInfoToHumanReadableStringEx( const FProgra
 	// Strip module path.
 	int32 Pos0;
 	int32 Pos1;
-	SymbolInfo.ModuleName.FindLastChar('\\', Pos0);
-	SymbolInfo.ModuleName.FindLastChar('/', Pos1);
+	SymbolInfo.ModuleName.FindLastChar(TEXT('\\'), Pos0);
+	SymbolInfo.ModuleName.FindLastChar(TEXT('/'), Pos1);
 	const int32 RealPos = FMath::Max( Pos0, Pos1 );
 	const FString StrippedModuleName = RealPos > 0 ? SymbolInfo.ModuleName.RightChop(RealPos + 1) : SymbolInfo.ModuleName;
 
diff --git a/Engine/Source/Runtime/Core/Private/GenericPlatform/GenericPlatformTime.cpp b/Engine/Source/Runtime/Core/Private/GenericPlatform/GenericPlatformTime.cpp
--- a/Engine/Source/Runtime/Core/Private/GenericPlatform/GenericPlatformTime.cpp
+++ b/Engine/Source/Runtime/Core/Private/GenericPlatform/GenericPlatformTime.cpp
@@ -112,7 +112,7 @@ TCHAR* FGenericPlatformTime::StrTime( TCHAR* Dest, SIZE_T DestSize )
 const TCHAR* FGenericPlatformTime::StrTimestamp()
 {
 	static TCHAR Result[1024];
-	*Result = 0;
+	*Result = TEXT('\0');
 	StrDate( Result, UE_ARRAY_COUNT(Result) );
 	FCString::Strcat( Result, TEXT(" ") );
 	StrTime( Result + FCString::Strlen(Result), UE_ARRAY_COUNT(Result) - FCString::Strlen(Result) );
diff --git a/Engine/Source/Runtime/Core/Private/Hash/Blake3.cpp b/Engine/Source/Runtime/Core/Private/Hash/Blake3.cpp
--- a/Engine/Source/Runtime/Core/Private/Hash/Blake3.cpp
+++ b/Engine/Source/Runtime/Core/Private/Hash/Blake3.cpp
@@ -77,7 +77,7 @@ FString LexToString(const FBlake3Hash& Hash)
 	TArray<TCHAR, FString::AllocatorType>& CharArray = Output.GetCharArray();
 	CharArray.AddUninitialized(sizeof(FBlake3Hash::ByteArray) * 2 + 1);
 	UE::String::BytesToHexLower(Hash.GetBytes(), CharArray.GetData());
-	CharArray.Last() = 0;
+	CharArray.Last() = TCHAR('\0');
 	return Output;
 }
 
diff --git a/Engine/Source/Runtime/Core/Private/IO/IoHash.cpp b/Engine/Source/Runtime/Core/Private/IO/IoHash.cpp
--- a/Engine/Source/Runtime/Core/Private/IO/IoHash.cpp
+++ b/Engine/Source/Runtime/Core/Private/IO/IoHash.cpp
@@ -10,6 +10,6 @@ FString LexToString(const FIoHash& Hash)
 	TArray<TCHAR, FString::AllocatorType>& CharArray = Output.GetCharArray();
 	CharArray.AddUninitialized(sizeof(FIoHash::ByteArray) * 2 + 1);
 	UE::String::BytesToHexLower(Hash.GetBytes(), CharArray.GetData());
-	CharArray.Last() = 0;
+	CharArray.Last() = TEXT('\0');
 	return Output;
 }
diff --git a/Engine/Source/Runtime/Core/Private/IO/IoStatus.cpp b/Engine/Source/Runtime/Core/Private/IO/IoStatus.cpp
--- a/Engine/Source/Runtime/Core/Private/IO/IoStatus.cpp
+++ b/Engine/Source/Runtime/Core/Private/IO/IoStatus.cpp
@@ -21,7 +21,7 @@ FIoStatus::~FIoStatus()
 FIoStatus::FIoStatus(EIoErrorCode Code)
 :	ErrorCode(Code)
 {
-	ErrorMessage[0] = 0;
+	ErrorMessage[0] = TEXT('\0');
 }
 
 FIoStatus::FIoStatus(EIoErrorCode Code, const FStringView& InErrorMessage)
@@ -29,7 +29,7 @@ FIoStatus::FIoStatus(EIoErrorCode Code, const FStringView& InErrorMessage)
 {
 	const int32 ErrorMessageLength = FMath::Min(MaxErrorMessageLength - 1, InErrorMessage.Len());
 	FPlatformString::Convert(ErrorMessage, ErrorMessageLength, InErrorMessage.GetData(), ErrorMessageLength);
-	ErrorMessage[ErrorMessageLength] = 0;
+	ErrorMessage[ErrorMessageLength] = TEXT('\0');
 }
 
 FIoStatus& FIoStatus::operator=(const FIoStatus& Other)
@@ -43,7 +43,7 @@ FIoStatus& FIoStatus::operator=(const FIoStatus& Other)
 FIoStatus& FIoStatus::operator=(const EIoErrorCode InErrorCode)
 {
 	ErrorCode = InErrorCode;
-	ErrorMessage[0] = 0;
+	ErrorMessage[0] = TEXT('\0');
 
 	return *this;
 }
diff --git a/Engine/Source/Runtime/Core/Private/Internationalization/FastDecimalFormat.cpp b/Engine/Source/Runtime/Core/Private/Internationalization/FastDecimalFormat.cpp
--- a/Engine/Source/Runtime/Core/Private/Internationalization/FastDecimalFormat.cpp
+++ b/Engine/Source/Runtime/Core/Private/Internationalization/FastDecimalFormat.cpp
@@ -433,7 +433,7 @@ FString CultureInvariantDecimalToString(const double InVal, const TCHAR*& InBuff
 
 	// Scrape negative, apply at end
 	bool bIsNegative = false;
-	static const TCHAR EuropeanNegativePrefix = '-';
+	static const TCHAR EuropeanNegativePrefix = TEXT('-');
 	if (*InBuffer == EuropeanNegativePrefix)
 	{
 		bIsNegative = true;
@@ -445,7 +445,7 @@ FString CultureInvariantDecimalToString(const double InVal, const TCHAR*& InBuff
 	uint8 FractionalDigitsPrinted = 0;
 	while (InBuffer < InBufferEnd)
 	{
-		static const TCHAR EuropeanDecimal = '.';
+		static const TCHAR EuropeanDecimal = TEXT('.');
 		if (*InBuffer == EuropeanDecimal && InFormattingOptions.MaximumFractionalDigits > 0)
 		{
 			bParsedFractional = true;
@@ -551,7 +551,7 @@ void FractionalToString(const double InVal, const FDecimalNumberFormattingRules&
 	int32 FractionalPartLen = 0;
 	if (FractionalPart != 0.0)
 	{
-		FractionalPartLen = IntegralToString_UInt64ToString(static_cast<uint64>(FractionalPart), false, 0, 0, ' ', InFormattingRules.DigitCharacters, 0, InFormattingOptions.MaximumFractionalDigits, FractionalPartBuffer, UE_ARRAY_COUNT(FractionalPartBuffer));
+		FractionalPartLen = IntegralToString_UInt64ToString(static_cast<uint64>(FractionalPart), false, 0, 0, TEXT(' '), InFormattingRules.DigitCharacters, 0, InFormattingOptions.MaximumFractionalDigits, FractionalPartBuffer, UE_ARRAY_COUNT(FractionalPartBuffer));
 	
 		{
 			// Pad the fractional part with any leading zeros that may have been lost when the number was split
@@ -575,7 +575,7 @@ void FractionalToString(const double InVal, const FDecimalNumberFormattingRules&
 			--FractionalPartLen;
 		}
 	}
-	FractionalPartBuffer[FractionalPartLen] = 0;
+	FractionalPartBuffer[FractionalPartLen] = TCHAR('\0');
 
 	// Pad the fractional part with any zeros that may have been missed so far
 	{
@@ -584,7 +584,7 @@ void FractionalToString(const double InVal, const FDecimalNumberFormattingRules&
 		{
 			FractionalPartBuffer[FractionalPartLen++] = InFormattingRules.DigitCharacters[0];
 		}
-		FractionalPartBuffer[FractionalPartLen] = 0;
+		FractionalPartBuffer[FractionalPartLen] = TCHAR('\0');
 	}
 
 	BuildFinalString(bIsNegative, InFormattingOptions.AlwaysSign, InFormattingRules, IntegralPartBuffer, IntegralPartLen, FractionalPartBuffer, FractionalPartLen, OutString);
@@ -677,7 +677,7 @@ bool StringToIntegral_StringToUInt64(const TCHAR*& InBuffer, const TCHAR* InBuff
 	const bool bTestForOverflow = EnumHasAnyFlags(InParseFlags, EDecimalNumberParseFlags::TestLimits | EDecimalNumberParseFlags::ClampValue);
 
 	// Parse the number, stopping once we find the end of the string or a decimal separator
-	static const TCHAR EuropeanNumerals[] = { '0', '1', '2', '3', '4', '5', '6', '7', '8', '9' };
+	static const TCHAR EuropeanNumerals[] = { TEXT('0'), TEXT('1'), TEXT('2'), TEXT('3'), TEXT('4'), TEXT('5'), TEXT('6'), TEXT('7'), TEXT('8'), TEXT('9') };
 	bool bFoundUnexpectedNonNumericCharacter = false;
 	while (InBuffer < InBufferEnd && *InBuffer != InFormattingRules.DecimalSeparatorCharacter)
 	{
@@ -846,14 +846,14 @@ bool StringToCultureInvariantDecimal(const TCHAR*& InBuffer, const TCHAR* InBuff
 	bool bIsNegative = false;
 	InSignParser.ParseLeadingSign(InBuffer, bIsNegative);
 
-	static const TCHAR InvariantNegativePrefix = '-';
+	static const TCHAR InvariantNegativePrefix = TEXT('-');
 	if (bIsNegative)
 	{
 		OutInvariantDecimal += InvariantNegativePrefix;
 	}
 
 	// Parse the number, stopping once we find the end of the string or a decimal separator
-	static const TCHAR EuropeanNumerals[] = { '0', '1', '2', '3', '4', '5', '6', '7', '8', '9' };
+	static const TCHAR EuropeanNumerals[] = { TEXT('0'), TEXT('1'), TEXT('2'), TEXT('3'), TEXT('4'), TEXT('5'), TEXT('6'), TEXT('7'), TEXT('8'), TEXT('9') };
 	bool bFoundUnexpectedNonNumericCharacter = false;
 	while (InBuffer < InBufferEnd)
 	{
@@ -869,7 +869,7 @@ bool StringToCultureInvariantDecimal(const TCHAR*& InBuffer, const TCHAR* InBuff
 		}
 
 		// Walk over the decimal separator
-		static const TCHAR InvariantDecimal = '.';
+		static const TCHAR InvariantDecimal = TEXT('.');
 		if (*InBuffer == InFormattingRules.DecimalSeparatorCharacter)
 		{
 			if (!EnumHasAnyFlags(InParseFlags, EDecimalNumberParseFlags::AllowDecimalSeparators))
@@ -998,8 +998,8 @@ const FDecimalNumberFormattingRules& GetCultureAgnosticFormattingRules()
 		AgnosticFormattingRules.NegativePrefixString = TEXT("-");
 		AgnosticFormattingRules.PlusString = TEXT("+");
 		AgnosticFormattingRules.MinusString = TEXT("-");
-		AgnosticFormattingRules.GroupingSeparatorCharacter = ',';
-		AgnosticFormattingRules.DecimalSeparatorCharacter = '.';
+		AgnosticFormattingRules.GroupingSeparatorCharacter = TEXT(',');
+		AgnosticFormattingRules.DecimalSeparatorCharacter = TEXT('.');
 		AgnosticFormattingRules.PrimaryGroupingSize = 3;
 		AgnosticFormattingRules.SecondaryGroupingSize = 3;
 		return AgnosticFormattingRules;
diff --git a/Engine/Source/Runtime/Core/Private/Internationalization/LegacyCulture.cpp b/Engine/Source/Runtime/Core/Private/Internationalization/LegacyCulture.cpp
--- a/Engine/Source/Runtime/Core/Private/Internationalization/LegacyCulture.cpp
+++ b/Engine/Source/Runtime/Core/Private/Internationalization/LegacyCulture.cpp
@@ -76,7 +76,7 @@ FString FLegacyCultureImplementation::GetNativeLanguage() const
 {
 	int32 LastBracket = INDEX_NONE;
 	int32 FirstBracket = INDEX_NONE;
-	if ( NativeName.FindLastChar( ')', LastBracket ) && NativeName.FindChar( '(', FirstBracket ) && LastBracket != FirstBracket )
+	if ( NativeName.FindLastChar( TEXT(')'), LastBracket ) && NativeName.FindChar( TEXT('('), FirstBracket ) && LastBracket != FirstBracket )
 	{
 		return NativeName.Left( FirstBracket-1 );
 	}
@@ -87,7 +87,7 @@ FString FLegacyCultureImplementation::GetNativeRegion() const
 {
 	int32 LastBracket = INDEX_NONE;
 	int32 FirstBracket = INDEX_NONE;
-	if ( NativeName.FindLastChar( ')', LastBracket ) && NativeName.FindChar( '(', FirstBracket ) && LastBracket != FirstBracket )
+	if ( NativeName.FindLastChar( TEXT(')'), LastBracket ) && NativeName.FindChar( TEXT('('), FirstBracket ) && LastBracket != FirstBracket )
 	{
 		return NativeName.Mid( FirstBracket+1, LastBracket-FirstBracket-1 );
 	}
diff --git a/Engine/Source/Runtime/Core/Private/Internationalization/TextFormatter.cpp b/Engine/Source/Runtime/Core/Private/Internationalization/TextFormatter.cpp
--- a/Engine/Source/Runtime/Core/Private/Internationalization/TextFormatter.cpp
+++ b/Engine/Source/Runtime/Core/Private/Internationalization/TextFormatter.cpp
@@ -188,7 +188,7 @@ TOptional<FExpressionError> ParseArgumentModifier(const FTextFormatPatternDefini
 	// Parse out the argument modifier parameter text
 	TOptional<FStringToken> Parameters;
 	{
-		TCHAR QuoteChar = 0;
+		TCHAR QuoteChar = TEXT('\0');
 		int32 NumConsecutiveSlashes = 0;
 		Parameters = Stream.ParseToken([&](TCHAR InC)
 		{
@@ -202,7 +202,7 @@ TOptional<FExpressionError> ParseArgumentModifier(const FTextFormatPatternDefini
 				{
 					if (NumConsecutiveSlashes%2 == 0)
 					{
-						QuoteChar = 0;
+						QuoteChar = TEXT('\0');
 					}
 				}
 				else
diff --git a/Engine/Source/Runtime/Core/Private/Logging/LogMacros.cpp b/Engine/Source/Runtime/Core/Private/Logging/LogMacros.cpp
--- a/Engine/Source/Runtime/Core/Private/Logging/LogMacros.cpp
+++ b/Engine/Source/Runtime/Core/Private/Logging/LogMacros.cpp
@@ -64,7 +64,7 @@ void FMsg::LogfImpl(const ANSICHAR* File, int32 Line, const FLogCategoryName& Ca
 			GET_VARARGS(MsgLogfStaticBuffer, UE_ARRAY_COUNT(MsgLogfStaticBuffer), UE_ARRAY_COUNT(MsgLogfStaticBuffer) - 1, Fmt, Fmt);
 			// Copy the message to the stack-allocated buffer)
 			FCString::Strncpy(Message, MsgLogfStaticBuffer, UE_ARRAY_COUNT(Message) - 1);
-			Message[UE_ARRAY_COUNT(Message) - 1] = '\0';
+			Message[UE_ARRAY_COUNT(Message) - 1] = TEXT('\0');
 		}
 
 		const int32 NumStackFramesToIgnore = 1;
@@ -111,7 +111,7 @@ void FMsg::Logf_InternalImpl(const ANSICHAR* File, int32 Line, const FLogCategor
 			GET_VARARGS(MsgLogfStaticBuffer, UE_ARRAY_COUNT(MsgLogfStaticBuffer), UE_ARRAY_COUNT(MsgLogfStaticBuffer) - 1, Fmt, Fmt);
 			// Copy the message to the stack-allocated buffer)
 			FCString::Strncpy(Message, MsgLogfStaticBuffer, UE_ARRAY_COUNT(Message) - 1);
-			Message[UE_ARRAY_COUNT(Message) - 1] = '\0';
+			Message[UE_ARRAY_COUNT(Message) - 1] = TEXT('\0');
 		}
 
 		const int32 NumStackFramesToIgnore = 1;
diff --git a/Engine/Source/Runtime/Core/Private/Misc/AssertionMacros.cpp b/Engine/Source/Runtime/Core/Private/Misc/AssertionMacros.cpp
--- a/Engine/Source/Runtime/Core/Private/Misc/AssertionMacros.cpp
+++ b/Engine/Source/Runtime/Core/Private/Misc/AssertionMacros.cpp
@@ -252,7 +252,7 @@ void FDebug::LogFormattedMessageWithCallstack(const FName& InLogName, const ANSI
 		}
 
 		// cap it
-		*SingleLineWritePos = 0;
+		*SingleLineWritePos = TEXT('\0');
 
 		// prefix function lines with [Callstack] for parsing tools
 		const TCHAR* Prefix = (FCString::Strnicmp(LineStart, TEXT("0x"), 2) == 0) ? TEXT("[Callstack] ") : TEXT("");
@@ -489,8 +489,8 @@ FORCENOINLINE void FDebug::EnsureFailed(const ANSICHAR* Expr, const ANSICHAR* Fi
 
 			ReportEnsure(ErrorMsg, ProgramCounter);
 
-			GErrorHist[0] = 0;
-			GErrorExceptionDescription[0] = 0;
+			GErrorHist[0] = TEXT('\0');
+			GErrorExceptionDescription[0] = TEXT('\0');
 #endif
 		}
 	}
diff --git a/Engine/Source/Runtime/Core/Private/Misc/CommandLine.cpp b/Engine/Source/Runtime/Core/Private/Misc/CommandLine.cpp
--- a/Engine/Source/Runtime/Core/Private/Misc/CommandLine.cpp
+++ b/Engine/Source/Runtime/Core/Private/Misc/CommandLine.cpp
@@ -31,10 +31,10 @@ bool FCommandLine::IsInitialized()
 
 void FCommandLine::Reset()
 {
-	CmdLine[0] = '\0';
-	OriginalCmdLine[0] = '\0';
-	LoggingCmdLine[0] = '\0';
-	LoggingOriginalCmdLine[0] = '\0';
+	CmdLine[0] = TEXT('\0');
+	OriginalCmdLine[0] = TEXT('\0');
+	LoggingCmdLine[0] = TEXT('\0');
+	LoggingOriginalCmdLine[0] = TEXT('\0');
 	bIsInitialized = false;
 }
 
diff --git a/Engine/Source/Runtime/Core/Private/Misc/ConfigCacheIni.cpp b/Engine/Source/Runtime/Core/Private/Misc/ConfigCacheIni.cpp
--- a/Engine/Source/Runtime/Core/Private/Misc/ConfigCacheIni.cpp
+++ b/Engine/Source/Runtime/Core/Private/Misc/ConfigCacheIni.cpp
@@ -857,7 +857,7 @@ void FConfigFile::CombineFromBuffer(const FString& Buffer)
 		// Strip trailing spaces from the current line
 		while( *Start && FChar::IsWhitespace(Start[FCString::Strlen(Start)-1]) )
 		{
-			Start[FCString::Strlen(Start)-1] = 0;
+			Start[FCString::Strlen(Start)-1] = TEXT('\0');
 		}
 
 		// If the first character in the line is [ and last char is ], this line indicates a section name
@@ -865,7 +865,7 @@ void FConfigFile::CombineFromBuffer(const FString& Buffer)
 		{
 			// Remove the brackets
 			Start++;
-			Start[FCString::Strlen(Start)-1] = 0;
+			Start[FCString::Strlen(Start)-1] = TEXT('\0');
 
 			// If we don't have an existing section by this name, add one
 			CurrentSection = FindOrAddSection( Start );
@@ -890,7 +890,7 @@ void FConfigFile::CombineFromBuffer(const FString& Buffer)
 			if( Value )
 			{
 				// Terminate the property name, advancing past the =
-				*Value++ = 0;
+				*Value++ = TEXT('\0');
 
 				// strip leading whitespace from the property name
 				while ( *Start && FChar::IsWhitespace(*Start) )
@@ -912,13 +912,13 @@ void FConfigFile::CombineFromBuffer(const FString& Buffer)
 				}
 				else
 				{
-					Cmd=' ';
+					Cmd = TEXT(' ');
 				}
 
 				// Strip trailing spaces from the property name.
 				while( *Start && FChar::IsWhitespace(Start[FCString::Strlen(Start)-1]) )
 				{
-					Start[FCString::Strlen(Start)-1] = 0;
+					Start[FCString::Strlen(Start)-1] = TEXT('\0');
 				}
 
 				// Strip leading whitespace from the property value
@@ -930,7 +930,7 @@ void FConfigFile::CombineFromBuffer(const FString& Buffer)
 				// strip trailing whitespace from the property value
 				while( *Value && FChar::IsWhitespace(Value[FCString::Strlen(Value)-1]) )
 				{
-					Value[FCString::Strlen(Value)-1] = 0;
+					Value[FCString::Strlen(Value)-1] = TEXT('\0');
 				}
 
 				ProcessedValue.Reset();
@@ -1038,7 +1038,7 @@ void FConfigFile::ProcessInputFileContents(FStringView Contents)
 		// Strip trailing spaces from the current line
 		while( *Start && FChar::IsWhitespace(Start[FCString::Strlen(Start)-1]) )
 		{
-			Start[FCString::Strlen(Start)-1] = 0;
+			Start[FCString::Strlen(Start)-1] = TEXT('\0');
 		}
 
 		// If the first character in the line is [ and last char is ], this line indicates a section name
@@ -1046,7 +1046,7 @@ void FConfigFile::ProcessInputFileContents(FStringView Contents)
 		{
 			// Remove the brackets
 			Start++;
-			Start[FCString::Strlen(Start)-1] = 0;
+			Start[FCString::Strlen(Start)-1] = TEXT('\0');
 
 			// If we don't have an existing section by this name, add one
 			CurrentSection = FindOrAddSection( Start );
@@ -1067,7 +1067,7 @@ void FConfigFile::ProcessInputFileContents(FStringView Contents)
 			if( Value )
 			{
 				// Terminate the propertyname, advancing past the =
-				*Value++ = 0;
+				*Value++ = TEXT('\0');
 
 				// strip leading whitespace from the property name
 				while ( *Start && FChar::IsWhitespace(*Start) )
@@ -1075,7 +1075,7 @@ void FConfigFile::ProcessInputFileContents(FStringView Contents)
 
 				// Strip trailing spaces from the property name.
 				while( *Start && FChar::IsWhitespace(Start[FCString::Strlen(Start)-1]) )
-					Start[FCString::Strlen(Start)-1] = 0;
+					Start[FCString::Strlen(Start)-1] = TEXT('\0');
 
 				// Strip leading whitespace from the property value
 				while ( *Value && FChar::IsWhitespace(*Value) )
@@ -1083,7 +1083,7 @@ void FConfigFile::ProcessInputFileContents(FStringView Contents)
 
 				// strip trailing whitespace from the property value
 				while( *Value && FChar::IsWhitespace(Value[FCString::Strlen(Value)-1]) )
-					Value[FCString::Strlen(Value)-1] = 0;
+					Value[FCString::Strlen(Value)-1] = TEXT('\0');
 
 				// If this line is delimited by quotes
 				if( *Value=='\"' )
diff --git a/Engine/Source/Runtime/Core/Private/Misc/ExpressionParser.cpp b/Engine/Source/Runtime/Core/Private/Misc/ExpressionParser.cpp
--- a/Engine/Source/Runtime/Core/Private/Misc/ExpressionParser.cpp
+++ b/Engine/Source/Runtime/Core/Private/Misc/ExpressionParser.cpp
@@ -28,7 +28,7 @@ TCHAR FTokenStream::PeekChar(int32 Offset) const
 		return *(ReadPos + Offset);
 	}
 
-	return '\0';
+	return TEXT('\0');
 }
 
 int32 FTokenStream::CharsRemaining() const
diff --git a/Engine/Source/Runtime/Core/Private/Misc/ExpressionParserExamples.cpp b/Engine/Source/Runtime/Core/Private/Misc/ExpressionParserExamples.cpp
--- a/Engine/Source/Runtime/Core/Private/Misc/ExpressionParserExamples.cpp
+++ b/Engine/Source/Runtime/Core/Private/Misc/ExpressionParserExamples.cpp
@@ -33,7 +33,7 @@ public:
 		TokenDefinitions.DefineToken([](FExpressionTokenConsumer& Consumer){
 
 			// Called at the start of every new token - see if it's a + character
-			TOptional<FStringToken> Token = Consumer.GetStream().ParseSymbol('+');
+			TOptional<FStringToken> Token = Consumer.GetStream().ParseSymbol(TEXT('+'));
 			if (Token.IsSet())
 			{
 				// Add the token to the consumer. This moves the read position in the stream to the end of the token.
diff --git a/Engine/Source/Runtime/Core/Private/Misc/FeedbackContextMarkup.cpp b/Engine/Source/Runtime/Core/Private/Misc/FeedbackContextMarkup.cpp
--- a/Engine/Source/Runtime/Core/Private/Misc/FeedbackContextMarkup.cpp
+++ b/Engine/Source/Runtime/Core/Private/Misc/FeedbackContextMarkup.cpp
@@ -65,7 +65,7 @@ bool FFeedbackContextMarkup::PipeProcessOutput(const FText& Description, const F
 			BufferedText += FPlatformProcess::ReadPipe(PipeRead);
 
 			int32 EndOfLineIdx;
-			while(BufferedText.FindChar('\n', EndOfLineIdx))
+			while(BufferedText.FindChar(TEXT('\n'), EndOfLineIdx))
 			{
 				FString Line = BufferedText.Left(EndOfLineIdx);
 				Line.RemoveFromEnd(TEXT("\r"), ESearchCase::CaseSensitive);
diff --git a/Engine/Source/Runtime/Core/Private/Misc/FileHelper.cpp b/Engine/Source/Runtime/Core/Private/Misc/FileHelper.cpp
--- a/Engine/Source/Runtime/Core/Private/Misc/FileHelper.cpp
+++ b/Engine/Source/Runtime/Core/Private/Misc/FileHelper.cpp
@@ -134,7 +134,7 @@ void FFileHelper::BufferToString( FString& Result, const uint8* Buffer, int32 Si
 	else
 	{
 		// Else ensure null terminator is present
-		ResultArray.Last() = 0;
+		ResultArray.Last() = TCHAR('\0');
 
 		if (bIsUnicode)
 		{
diff --git a/Engine/Source/Runtime/Core/Private/Misc/Guid.cpp b/Engine/Source/Runtime/Core/Private/Misc/Guid.cpp
--- a/Engine/Source/Runtime/Core/Private/Misc/Guid.cpp
+++ b/Engine/Source/Runtime/Core/Private/Misc/Guid.cpp
@@ -105,10 +105,10 @@ FString FGuid::ToString(EGuidFormats Format) const
 
 		for(int32 i=OutCharArray.Num(); i<25; i++)
 		{
-			OutCharArray.Add('0');
+			OutCharArray.Add(TEXT('0'));
 		}
 
-		OutCharArray.Add(0);
+		OutCharArray.Add(TCHAR('\0'));
 
 		check(Result.Len() == 25);
 		return Result.Reverse();
diff --git a/Engine/Source/Runtime/Core/Private/Misc/OutputDeviceAnsiError.cpp b/Engine/Source/Runtime/Core/Private/Misc/OutputDeviceAnsiError.cpp
--- a/Engine/Source/Runtime/Core/Private/Misc/OutputDeviceAnsiError.cpp
+++ b/Engine/Source/Runtime/Core/Private/Misc/OutputDeviceAnsiError.cpp
@@ -76,7 +76,7 @@ void FOutputDeviceAnsiError::HandleError()
 	GIsRunning			= 0;
 	GIsCriticalError	= 1;
 	GLogConsole			= NULL;
-	GErrorHist[UE_ARRAY_COUNT(GErrorHist)-1]=0;
+	GErrorHist[UE_ARRAY_COUNT(GErrorHist)-1] = TCHAR('\0');
 
 	if (GLog)
 	{
diff --git a/Engine/Source/Runtime/Core/Private/Misc/OutputDeviceFile.cpp b/Engine/Source/Runtime/Core/Private/Misc/OutputDeviceFile.cpp
--- a/Engine/Source/Runtime/Core/Private/Misc/OutputDeviceFile.cpp
+++ b/Engine/Source/Runtime/Core/Private/Misc/OutputDeviceFile.cpp
@@ -349,7 +349,7 @@ FOutputDeviceFile::FOutputDeviceFile( const TCHAR* InFilename, bool bInDisableBa
 	}
 	else
 	{
-		Filename[0]	= 0;
+		Filename[0]	= TCHAR('\0');
 	}
 
 #if ALLOW_LOG_FILE && !NO_LOGGING
@@ -396,7 +396,7 @@ void FOutputDeviceFile::TearDown()
 	delete WriterArchive;
 	WriterArchive = nullptr;
 
-	Filename[0] = 0;
+	Filename[0] = TCHAR('\0');
 }
 
 /**
diff --git a/Engine/Source/Runtime/Core/Private/Misc/Parse.cpp b/Engine/Source/Runtime/Core/Private/Misc/Parse.cpp
--- a/Engine/Source/Runtime/Core/Private/Misc/Parse.cpp
+++ b/Engine/Source/Runtime/Core/Private/Misc/Parse.cpp
@@ -248,11 +248,11 @@ bool FParse::Value(
 			int32 QuoteCharactersToSkip = 1;
 			FCString::Strncpy(Value, Start + QuoteCharactersToSkip, MaxLen);
 
-			Value[MaxLen-1]=0;
+			Value[MaxLen - 1] = TCHAR('\0');
 			TCHAR* Temp = FCString::Strstr( Value, TEXT("\x22") );
 			if (Temp != nullptr)
 			{
-				*Temp = 0;
+				*Temp = TCHAR('\0');
 			}
 		}
 		else
@@ -262,16 +262,16 @@ bool FParse::Value(
 
 			// Non-quoted string without spaces.
 			FCString::Strncpy( Value, Start, MaxLen );
-			Value[MaxLen-1]=0;
+			Value[MaxLen - 1]= TCHAR('\0');
 			TCHAR* Temp;
-			Temp = FCString::Strstr( Value, TEXT(" ")  ); if( Temp ) *Temp=0;
-			Temp = FCString::Strstr( Value, TEXT("\r") ); if( Temp ) *Temp=0;
-			Temp = FCString::Strstr( Value, TEXT("\n") ); if( Temp ) *Temp=0;
-			Temp = FCString::Strstr( Value, TEXT("\t") ); if( Temp ) *Temp=0;
+			Temp = FCString::Strstr( Value, TEXT(" ")  ); if( Temp ) *Temp = TCHAR('\0');
+			Temp = FCString::Strstr( Value, TEXT("\r") ); if( Temp ) *Temp = TCHAR('\0');
+			Temp = FCString::Strstr( Value, TEXT("\n") ); if( Temp ) *Temp = TCHAR('\0');
+			Temp = FCString::Strstr( Value, TEXT("\t") ); if( Temp ) *Temp = TCHAR('\0');
 			if (bShouldStopOnSeparator)
 			{
-				Temp = FCString::Strstr( Value, TEXT(",")  ); if( Temp ) *Temp=0;
-				Temp = FCString::Strstr( Value, TEXT(")")  ); if( Temp ) *Temp=0;
+				Temp = FCString::Strstr( Value, TEXT(",")  ); if( Temp ) *Temp = TCHAR('\0');
+				Temp = FCString::Strstr( Value, TEXT(")")  ); if( Temp ) *Temp = TCHAR('\0');
 			}
 		}
 
@@ -323,7 +323,7 @@ bool FParse::Value( const TCHAR* Stream, const TCHAR* Match, FString& Value, boo
 	{
 		TArray<TCHAR, TInlineAllocator<4096>> ValueCharArray;
 		ValueCharArray.AddUninitialized(StreamLen + 1);
-		ValueCharArray[0] = 0;
+		ValueCharArray[0] = TCHAR('\0');
 
 		if( FParse::Value(Stream, Match, ValueCharArray.GetData(), ValueCharArray.Num(), bShouldStopOnSeparator) )
 		{
@@ -685,10 +685,10 @@ bool FParse::Value( const TCHAR* Stream, const TCHAR* Match, struct FGuid& Guid
 	if( FCString::Strlen(Temp)==32 )
 	{
 		TCHAR* End;
-		Guid.D = FCString::Strtoi( Temp+24, &End, 16 ); Temp[24]=0;
-		Guid.C = FCString::Strtoi( Temp+16, &End, 16 ); Temp[16]=0;
-		Guid.B = FCString::Strtoi( Temp+8,  &End, 16 ); Temp[8 ]=0;
-		Guid.A = FCString::Strtoi( Temp+0,  &End, 16 ); Temp[0 ]=0;
+		Guid.D = FCString::Strtoi( Temp+24, &End, 16 ); Temp[24] = TCHAR('\0');
+		Guid.C = FCString::Strtoi( Temp+16, &End, 16 ); Temp[16] = TCHAR('\0');
+		Guid.B = FCString::Strtoi( Temp+8,  &End, 16 ); Temp[8 ] = TCHAR('\0');
+		Guid.A = FCString::Strtoi( Temp+0,  &End, 16 ); Temp[0 ] = TCHAR('\0');
 	}
 	return true;
 }
@@ -847,8 +847,8 @@ bool FParse::Token( const TCHAR*& Str, TCHAR* Result, int32 MaxLen, bool UseEsca
 			}
 		}
 	}
-	Result[Len]=0;
-	return Len!=0;
+	Result[Len] = TCHAR('\0');
+	return Len != TCHAR('\0');
 }
 
 bool FParse::Token( const TCHAR*& Str, FString& Arg, bool UseEscape )
@@ -966,7 +966,7 @@ bool FParse::Line(const TCHAR** Stream, TCHAR* Result, int32 MaxLen, bool bExact
 	bool bIsQuoted = false;
 	bool bIgnore = false;
 
-	*Result=0;
+	*Result = TCHAR('\0');
 	while (**Stream != TEXT('\0') && **Stream != TEXT('\n') && **Stream != TEXT('\r') && --MaxLen > 0)
 	{
 		// Start of comments.
diff --git a/Engine/Source/Runtime/Core/Private/Misc/StringFormatter.cpp b/Engine/Source/Runtime/Core/Private/Misc/StringFormatter.cpp
--- a/Engine/Source/Runtime/Core/Private/Misc/StringFormatter.cpp
+++ b/Engine/Source/Runtime/Core/Private/Misc/StringFormatter.cpp
@@ -107,7 +107,7 @@ TOptional<FExpressionError> ParseIndex(FExpressionTokenConsumer& Consumer, bool
 {
 	auto& Stream = Consumer.GetStream();
 
-	TOptional<FStringToken> OpeningChar = Stream.ParseSymbol('{');
+	TOptional<FStringToken> OpeningChar = Stream.ParseSymbol(TEXT('{'));
 	if (!OpeningChar.IsSet())
 	{
 		return TOptional<FExpressionError>();
@@ -150,7 +150,7 @@ TOptional<FExpressionError> ParseIndex(FExpressionTokenConsumer& Consumer, bool
 	// Optional whitespace
 	Stream.ParseToken([](TCHAR InC) { return FChar::IsWhitespace(InC) ? EParseState::Continue : EParseState::StopBefore; }, &EntireToken);
 	
-	if (!Stream.ParseSymbol('}', &EntireToken).IsSet())
+	if (!Stream.ParseSymbol(TEXT('}'), &EntireToken).IsSet())
 	{
 		// Not a valid token
 		if (bEmitErrors)
@@ -172,7 +172,7 @@ TOptional<FExpressionError> ParseSpecifier(FExpressionTokenConsumer& Consumer, b
 {
 	auto& Stream = Consumer.GetStream();
 
-	TOptional<FStringToken> OpeningChar = Stream.ParseSymbol('{');
+	TOptional<FStringToken> OpeningChar = Stream.ParseSymbol(TEXT('{'));
 	if (!OpeningChar.IsSet())
 	{
 		return TOptional<FExpressionError>();
@@ -217,7 +217,7 @@ TOptional<FExpressionError> ParseSpecifier(FExpressionTokenConsumer& Consumer, b
 	// Optional whitespace
 	Stream.ParseToken([](TCHAR InC) { return FChar::IsWhitespace(InC) ? EParseState::Continue : EParseState::StopBefore; }, &EntireToken);
 
-	if (!Stream.ParseSymbol('}', &EntireToken).IsSet())
+	if (!Stream.ParseSymbol(TEXT('}'), &EntireToken).IsSet())
 	{		
 		// Not a valid token
 		if (bEmitErrors)
@@ -235,7 +235,7 @@ TOptional<FExpressionError> ParseSpecifier(FExpressionTokenConsumer& Consumer, b
 	return TOptional<FExpressionError>();
 }
 
-static const TCHAR EscapeChar = '`';
+static const TCHAR EscapeChar = TEXT('`');
 
 /** Parse an escaped character */
 TOptional<FExpressionError> ParseEscapedChar(FExpressionTokenConsumer& Consumer, bool bEmitErrors)
diff --git a/Engine/Source/Runtime/Core/Private/Misc/TextFilterExpressionEvaluator.cpp b/Engine/Source/Runtime/Core/Private/Misc/TextFilterExpressionEvaluator.cpp
--- a/Engine/Source/Runtime/Core/Private/Misc/TextFilterExpressionEvaluator.cpp
+++ b/Engine/Source/Runtime/Core/Private/Misc/TextFilterExpressionEvaluator.cpp
@@ -9,8 +9,8 @@ namespace TextFilterExpressionParser
 	 * This contains all the symbols that can define breaking points between text and an operator
 	 * Note: We don't include + and - in this list as these are valid to use inside text and numbers, and should be consumed as part of the text token
 	 */
-	static const TCHAR BasicTextBreakingCharacters[]	= { '(', ')', '!', '&', '|', ' ' };						// ETextFilterExpressionEvaluatorMode::BasicString
-	static const TCHAR ComplexTextBreakingCharacters[]	= { '(', ')', '=', ':', '<', '>', '!', '&', '|', ' ' };	// ETextFilterExpressionEvaluatorMode::Complex
+	static const TCHAR BasicTextBreakingCharacters[]	= { TEXT('('), TEXT(')'), TEXT('!'), TEXT('&'), TEXT('|'), TEXT(' ') };						// ETextFilterExpressionEvaluatorMode::BasicString
+	static const TCHAR ComplexTextBreakingCharacters[]	= { TEXT('('), TEXT(')'), TEXT('='), TEXT(':'), TEXT('<'), TEXT('>'), TEXT('!'), TEXT('&'), TEXT('|'), TEXT(' ') };	// ETextFilterExpressionEvaluatorMode::Complex
 
 	const TCHAR* FSubExpressionStart::Monikers[]		= { TEXT("(") };
 	const TCHAR* FSubExpressionEnd::Monikers[]			= { TEXT(")") };
@@ -62,8 +62,8 @@ namespace TextFilterExpressionParser
 	/** Transform the given string to remove any escape character sequences found in a quoted string */
 	void UnescapeQuotedString(FString& Str, const TCHAR InQuoteChar)
 	{
-		const TCHAR EscapedQuote[] = { '\\', InQuoteChar, 0 };
-		const TCHAR UnescapedQuote[] = { InQuoteChar, 0 };
+		const TCHAR EscapedQuote[] = { TEXT('\\'), InQuoteChar, TEXT('\0') };
+		const TCHAR UnescapedQuote[] = { InQuoteChar, TEXT('\0') };
 
 		// Unescape any literal quotes within the string
 		Str.ReplaceInline(EscapedQuote, UnescapedQuote);
@@ -127,7 +127,7 @@ namespace TextFilterExpressionParser
 		FString FinalString;
 		FString CurrentQuotedString;
 
-		TCHAR QuoteChar = 0;
+		TCHAR QuoteChar = TEXT('\0');
 		int32 NumConsecutiveSlashes = 0;
 		TOptional<FStringToken> TextToken = Stream.ParseToken([&](TCHAR InC)
 		{
@@ -162,7 +162,7 @@ namespace TextFilterExpressionParser
 					FinalString.Append(CurrentQuotedString);
 
 					CurrentQuotedString.Reset();
-					QuoteChar = 0;
+					QuoteChar = TEXT('\0');
 				}
 
 				if (InC == '\\')
diff --git a/Engine/Source/Runtime/Core/Private/Misc/VarargsHelper.h b/Engine/Source/Runtime/Core/Private/Misc/VarargsHelper.h
--- a/Engine/Source/Runtime/Core/Private/Misc/VarargsHelper.h
+++ b/Engine/Source/Runtime/Core/Private/Misc/VarargsHelper.h
@@ -35,7 +35,7 @@
 		GET_VARARGS_RESULT( Buffer, BufferSize, BufferSize-1, Fmt, Fmt, Result ); \
 		BufferSize *= 2; \
 	}; \
-	Buffer[Result] = 0; \
+	Buffer[Result] = TEXT('\0'); \
 	; \
 \
 	SerializeFunc; \
diff --git a/Engine/Source/Runtime/Core/Private/Serialization/Archive.cpp b/Engine/Source/Runtime/Core/Private/Serialization/Archive.cpp
--- a/Engine/Source/Runtime/Core/Private/Serialization/Archive.cpp
+++ b/Engine/Source/Runtime/Core/Private/Serialization/Archive.cpp
@@ -1319,7 +1319,7 @@ void FArchive::LogfImpl(const TCHAR* Fmt, ...)
 		GET_VARARGS_RESULT( Buffer, BufferSize, BufferSize-1, Fmt, Fmt, Result );
 		BufferSize *= 2;
 	};
-	Buffer[Result] = 0;
+	Buffer[Result] = TEXT('\0');
 
 	// Convert to ANSI and serialize as ANSI char.
 	for( int32 i=0; i<Result; i++ )
diff --git a/Engine/Source/Runtime/Core/Private/Serialization/Csv/CsvParser.cpp b/Engine/Source/Runtime/Core/Private/Serialization/Csv/CsvParser.cpp
--- a/Engine/Source/Runtime/Core/Private/Serialization/Csv/CsvParser.cpp
+++ b/Engine/Source/Runtime/Core/Private/Serialization/Csv/CsvParser.cpp
@@ -84,8 +84,8 @@ FCsvParser::EParseResult FCsvParser::ParseCell()
 				// We null terminate and leave the write pos pointing at the trailing closing quote 
 				// if present so it gets overwritten by any subsequent text in the cell
 				NumQuotes /= 2;
-				while(NumQuotes-- > 0) *(WriteAt++) = '"';
-				*WriteAt = '\0';
+				while(NumQuotes-- > 0) *(WriteAt++) = TEXT('"');
+				*WriteAt = TEXT('\0');
 
 				continue;
 			}
@@ -97,14 +97,14 @@ FCsvParser::EParseResult FCsvParser::ParseCell()
 			if (NewLineSize != 0)
 			{
 				// Null terminate the cell
-				*WriteAt = '\0';
+				*WriteAt = TEXT('\0');
 				ReadAt += NewLineSize;
 
 				return *ReadAt ? EParseResult::EndOfRow : EParseResult::EndOfString;
 			}
 			else if (*ReadAt == ',')
 			{
-				*WriteAt = '\0';
+				*WriteAt = TEXT('\0');
 				++ReadAt;
 
 				// We always return EndOfCell here as we still have another (potentially empty) cell to add
diff --git a/Engine/Source/Runtime/Core/Private/Serialization/MemoryImage.cpp b/Engine/Source/Runtime/Core/Private/Serialization/MemoryImage.cpp
--- a/Engine/Source/Runtime/Core/Private/Serialization/MemoryImage.cpp
+++ b/Engine/Source/Runtime/Core/Private/Serialization/MemoryImage.cpp
@@ -1186,7 +1186,7 @@ FHashedName::FHashedName(const FName& InName)
 			{
 				UpperNameBuffer.Wide[i] = FChar::ToUpper(NameBuffer.Wide[i]);
 			}
-			UpperNameBuffer.Wide[NameLength] = 0;
+			UpperNameBuffer.Wide[NameLength] = TEXT('\0');
 			const FTCHARToUTF8 UpperNameUTF8(UpperNameBuffer.Wide);
 			Hash = CityHash64WithSeed((const char*)UpperNameUTF8.Get(), UpperNameUTF8.Length(), InternalNumber);
 #if WITH_EDITORONLY_DATA
diff --git a/Engine/Source/Runtime/Core/Private/Tests/Containers/BitArrayTests.cpp b/Engine/Source/Runtime/Core/Private/Tests/Containers/BitArrayTests.cpp
--- a/Engine/Source/Runtime/Core/Private/Tests/Containers/BitArrayTests.cpp
+++ b/Engine/Source/Runtime/Core/Private/Tests/Containers/BitArrayTests.cpp
@@ -39,9 +39,9 @@ namespace BitArrayTest
 		{
 			if (Index != 0 && Index % 8 == 0)
 			{
-				Out.AppendChar(' ');
+				Out.AppendChar(TEXT(' '));
 			}
-			Out.AppendChar(It.GetValue() ? '1' : '0');
+			Out.AppendChar(It.GetValue() ? TEXT('1') : TEXT('0'));
 		}
 		return Out;
 	}
@@ -1195,15 +1195,15 @@ public:
 						bool OnesValue = ArrOnes[ResultIndex];
 						bOnesMatchesExpected = bOnesMatchesExpected & (OnesValue == OnesExpectedValue);
 						bZeroesMatchesExpected = bZeroesMatchesExpected & (ZeroesValue == ZeroesExpectedValue);
-						OnesExpectedText[DisplayIndex] = OnesExpectedValue ? '1' : '0';
-						OnesText[DisplayIndex] = OnesValue ? '1' : '0';
-						ZeroesExpectedText[DisplayIndex] = ZeroesExpectedValue ? '1' : '0';
-						ZeroesText[DisplayIndex] = ZeroesValue ? '1' : '0';
+						OnesExpectedText[DisplayIndex] = OnesExpectedValue ? TEXT('1') : TEXT('0');
+						OnesText[DisplayIndex] = OnesValue ? TEXT('1') : TEXT('0');
+						ZeroesExpectedText[DisplayIndex] = ZeroesExpectedValue ? TEXT('1') : TEXT('0');
+						ZeroesText[DisplayIndex] = ZeroesValue ? TEXT('1') : TEXT('0');
 					}
-					OnesExpectedText[DisplayIndex] = '\0';
-					OnesText[DisplayIndex] = '\0';
-					ZeroesExpectedText[DisplayIndex] = '\0';
-					ZeroesText[DisplayIndex] = '\0';
+					OnesExpectedText[DisplayIndex] = TEXT('\0');
+					OnesText[DisplayIndex] = TEXT('\0');
+					ZeroesExpectedText[DisplayIndex] = TEXT('\0');
+					ZeroesText[DisplayIndex] = TEXT('\0');
 					if (!bOnesMatchesExpected)
 					{
 						AddError(FString::Printf(TEXT("SetRange bool BitWidth=%d WriteOffset=%d, Ones\nExpected=%s\nActual  =%s"), BitWidth, WriteOffset, OnesExpectedText, OnesText));
@@ -2188,15 +2188,15 @@ protected:
 							bool ZeroesValue = (ZeroesVerify[ResultWord] & ResultMask) != 0;
 							bOnesMatchesExpected = bOnesMatchesExpected & (OnesValue == OnesExpectedValue);
 							bZeroesMatchesExpected = bZeroesMatchesExpected & (ZeroesValue == ZeroesExpectedValue);
-							OnesExpectedText[DisplayIndex] = OnesExpectedValue ? '1' : '0';
-							OnesText[DisplayIndex] = OnesValue ? '1' : '0';
-							ZeroesExpectedText[DisplayIndex] = ZeroesExpectedValue ? '1' : '0';
-							ZeroesText[DisplayIndex] = ZeroesValue ? '1' : '0';
+							OnesExpectedText[DisplayIndex] = OnesExpectedValue ? TEXT('1') : TEXT('0');
+							OnesText[DisplayIndex] = OnesValue ? TEXT('1') : TEXT('0');
+							ZeroesExpectedText[DisplayIndex] = ZeroesExpectedValue ? TEXT('1') : TEXT('0');
+							ZeroesText[DisplayIndex] = ZeroesValue ? TEXT('1') : TEXT('0');
 						}
-						OnesExpectedText[DisplayIndex] = '\0';
-						OnesText[DisplayIndex] = '\0';
-						ZeroesExpectedText[DisplayIndex] = '\0';
-						ZeroesText[DisplayIndex] = '\0';
+						OnesExpectedText[DisplayIndex] = TEXT('\0');
+						OnesText[DisplayIndex] = TEXT('\0');
+						ZeroesExpectedText[DisplayIndex] = TEXT('\0');
+						ZeroesText[DisplayIndex] = TEXT('\0');
 						if (!bOnesMatchesExpected)
 						{
 							AddError(FString::Printf(TEXT("MemmoveBitsWordOrder BitWidth=%d ReadOffset=%d, WriteOffset=%d, Overlap=%d, Ones\nExpected=%s\nActual  =%s"), BitWidth, ReadOffset, WriteOffset, Overlap, OnesExpectedText, OnesText));
diff --git a/Engine/Source/Runtime/Core/Private/UObject/UnrealNames.cpp b/Engine/Source/Runtime/Core/Private/UObject/UnrealNames.cpp
--- a/Engine/Source/Runtime/Core/Private/UObject/UnrealNames.cpp
+++ b/Engine/Source/Runtime/Core/Private/UObject/UnrealNames.cpp
@@ -2027,7 +2027,7 @@ void FNameEntry::GetUnterminatedName(TCHAR* OutName, uint32 OutLen) const
 void FNameEntry::GetName(TCHAR(&OutName)[NAME_SIZE]) const
 {
 	CopyAndConvertUnterminatedName(OutName);
-	OutName[GetNameLength()] = '\0';
+	OutName[GetNameLength()] = TEXT('\0');
 }
 
 void FNameEntry::CopyAndConvertUnterminatedName(TCHAR* OutName) const
@@ -3176,7 +3176,7 @@ uint32 FName::ToString(TCHAR* Out, uint32 OutSize) const
 
 	if (GetNumber() == NAME_NO_NUMBER_INTERNAL)
 	{
-		Out[NameLen] = '\0';
+		Out[NameLen] = TEXT('\0');
 		return NameLen;
 	}
 	else
@@ -3187,7 +3187,7 @@ uint32 FName::ToString(TCHAR* Out, uint32 OutSize) const
 		check(SuffixLen > 0 && OutSize > TotalLen);
 
 		FPlatformMemory::Memcpy(Out + NameLen, NumberSuffixStr, SuffixLen * sizeof(TCHAR));
-		Out[TotalLen] = '\0';
+		Out[TotalLen] = TEXT('\0');
 		return TotalLen;
 	}
 }
diff --git a/Engine/Source/Runtime/Core/Private/Windows/WindowsErrorOutputDevice.cpp b/Engine/Source/Runtime/Core/Private/Windows/WindowsErrorOutputDevice.cpp
--- a/Engine/Source/Runtime/Core/Private/Windows/WindowsErrorOutputDevice.cpp
+++ b/Engine/Source/Runtime/Core/Private/Windows/WindowsErrorOutputDevice.cpp
@@ -38,7 +38,7 @@ void FWindowsErrorOutputDevice::Serialize( const TCHAR* Msg, ELogVerbosity::Type
 		// First appError.
 		GIsCriticalError = 1;
 		TCHAR ErrorBuffer[1024];
-		ErrorBuffer[0] = 0;
+		ErrorBuffer[0] = TCHAR('\0');
 
 		// Windows error.
 		if (LastError == 0)
@@ -97,7 +97,7 @@ void FWindowsErrorOutputDevice::HandleError()
 	GIsRunning				= 0;
 	GIsCriticalError		= 1;
 	GLogConsole				= NULL;
-	GErrorHist[UE_ARRAY_COUNT(GErrorHist)-1]=0;
+	GErrorHist[UE_ARRAY_COUNT(GErrorHist) - 1] = TCHAR('\0');
 
 	// Trigger the OnSystemFailure hook if it exists
 	// make sure it happens after GIsGuarded is set to 0 in case this hook crashes
diff --git a/Engine/Source/Runtime/Core/Private/Windows/WindowsPlatformCrashContext.cpp b/Engine/Source/Runtime/Core/Private/Windows/WindowsPlatformCrashContext.cpp
--- a/Engine/Source/Runtime/Core/Private/Windows/WindowsPlatformCrashContext.cpp
+++ b/Engine/Source/Runtime/Core/Private/Windows/WindowsPlatformCrashContext.cpp
@@ -224,7 +224,7 @@ void FWindowsPlatformCrashContext::ConvertProgramCountersToStackFrames(
 				TCHAR* ModuleNameExt = FCString::Strrchr(ModuleName, '.');
 				if (ModuleNameExt != nullptr)
 				{
-					*ModuleNameExt = 0;
+					*ModuleNameExt = TCHAR('\0');
 				}
 			}
 			else
@@ -397,7 +397,7 @@ bool CreateCrashReportClientPath(TCHAR* OutClientPath, int32 MaxLength)
 		const TCHAR* BinariesDir = FPlatformProcess::GetBinariesSubdirectory();
 
 		// Find the path to crash reporter binary. Avoid creating FStrings.
-		*OutClientPath = 0;
+		*OutClientPath = TCHAR('\0');
 		FCString::Strncat(OutClientPath, EngineDir, MaxLength);
 		FCString::Strncat(OutClientPath, TEXT("Binaries/"), MaxLength);
 		FCString::Strncat(OutClientPath, BinariesDir, MaxLength);
diff --git a/Engine/Source/Runtime/Core/Private/Windows/WindowsPlatformFile.cpp b/Engine/Source/Runtime/Core/Private/Windows/WindowsPlatformFile.cpp
--- a/Engine/Source/Runtime/Core/Private/Windows/WindowsPlatformFile.cpp
+++ b/Engine/Source/Runtime/Core/Private/Windows/WindowsPlatformFile.cpp
@@ -947,9 +947,9 @@ class CORE_API FWindowsPlatformFile : public IPhysicalPlatformFile
 		// We now have a canonical, strict-valid, absolute Unreal Path.  Convert it to a Windows Path.
 		for (TCHAR& Char : TArrayView<TCHAR>(Path.GetData(), Path.Len()))
 		{
-			if (Char == '/')
+			if (Char == TEXT('/'))
 			{
-				Char = '\\';
+				Char = TEXT('\\');
 			}
 		}
 
diff --git a/Engine/Source/Runtime/Core/Private/Windows/WindowsPlatformMisc.cpp b/Engine/Source/Runtime/Core/Private/Windows/WindowsPlatformMisc.cpp
--- a/Engine/Source/Runtime/Core/Private/Windows/WindowsPlatformMisc.cpp
+++ b/Engine/Source/Runtime/Core/Private/Windows/WindowsPlatformMisc.cpp
@@ -325,8 +325,8 @@ int32 FWindowsOSVersionHelper::GetOSVersions( FString& OutOSVersionLabel, FStrin
 	TCHAR OSVersionLabel[128];
 	TCHAR OSSubVersionLabel[128];
 
-	OSVersionLabel[0] = 0;
-	OSSubVersionLabel[0] = 0;
+	OSVersionLabel[0] = TEXT('\0');
+	OSSubVersionLabel[0] = TEXT('\0');
 
 	int32 Result = GetOSVersionsHelper( OSVersionLabel, UE_ARRAY_COUNT(OSVersionLabel), OSSubVersionLabel, UE_ARRAY_COUNT(OSSubVersionLabel) );
 
@@ -694,7 +694,7 @@ void FWindowsPlatformMisc::GetEnvironmentVariable(const TCHAR* VariableName, TCH
 	uint32 Error = ::GetEnvironmentVariableW(VariableName, Result, ResultLength);
 	if (Error <= 0)
 	{		
-		*Result = 0;
+		*Result = TEXT('\0');
 	}
 }
 
@@ -2696,8 +2696,8 @@ void FWindowsPlatformMisc::GetOSVersions( FString& OutOSVersionLabel, FString& O
 	{
 		FOSVersionsInitializer()
 		{
-			OSVersionLabel[0] = 0;
-			OSSubVersionLabel[0] = 0;
+			OSVersionLabel[0] = TEXT('\0');
+			OSSubVersionLabel[0] = TEXT('\0');
 			GetOSVersionsHelper( OSVersionLabel, UE_ARRAY_COUNT(OSVersionLabel), OSSubVersionLabel, UE_ARRAY_COUNT(OSSubVersionLabel) );
 		}
 
@@ -2716,10 +2716,10 @@ FString FWindowsPlatformMisc::GetOSVersion()
 	{
 		FOSVersionInitializer()
 		{
-			CachedOSVersion[0] = 0;
+			CachedOSVersion[0] = TEXT('\0');
 			if (!GetOSVersionHelper(CachedOSVersion, UE_ARRAY_COUNT(CachedOSVersion)))
 			{
-				CachedOSVersion[0] = 0;
+				CachedOSVersion[0] = TEXT('\0');
 			}
 		}
 
diff --git a/Engine/Source/Runtime/Core/Private/Windows/WindowsPlatformProcess.cpp b/Engine/Source/Runtime/Core/Private/Windows/WindowsPlatformProcess.cpp
--- a/Engine/Source/Runtime/Core/Private/Windows/WindowsPlatformProcess.cpp
+++ b/Engine/Source/Runtime/Core/Private/Windows/WindowsPlatformProcess.cpp
@@ -1020,7 +1020,7 @@ const TCHAR* FWindowsPlatformProcess::BaseDir()
 					}
 				}
 			}
-			Result[StringLength] = 0;
+			Result[StringLength] = TCHAR('\0');
 
 			FString CollapseResult(Result);
 #ifdef UE_RELATIVE_BASE_DIR
@@ -1135,7 +1135,7 @@ const TCHAR* FWindowsPlatformProcess::UserName(bool bOnlyAlphaNumeric/* = true*/
 			for( c=ResultAlpha, d=ResultAlpha; *c!=0; c++ )
 				if( FChar::IsAlnum(*c) )
 					*d++ = *c;
-			*d++ = 0;
+			*d++ = TCHAR('\0');
 		}
 		return ResultAlpha;
 	}
@@ -1207,7 +1207,7 @@ const TCHAR* FWindowsPlatformProcess::ExecutablePath()
 	{
 		if ( !GetModuleFileName( hInstance, Result, UE_ARRAY_COUNT(Result) ) )
 		{
-			Result[0] = 0;
+			Result[0] = TCHAR('\0');
 		}
 	}
 	return Result;
@@ -1336,7 +1336,7 @@ bool FWindowsPlatformProcess::ResolveNetworkPath( FString InUNCPath, FString& Ou
 	// Get local machine name first and check if this UNC path points to local share
 	// (if it's not UNC path it will also fail this check)
 	uint32 ComputerNameSize = MAX_COMPUTERNAME_LENGTH;
-	TCHAR ComputerName[MAX_COMPUTERNAME_LENGTH + 3] = { '\\', '\\', '\0', };
+	TCHAR ComputerName[MAX_COMPUTERNAME_LENGTH + 3] = { TEXT('\\'), TEXT('\\'), TEXT('\0'), };
 
 	if ( GetComputerName( ComputerName + 2, (::DWORD*)&ComputerNameSize ) )
 	{
diff --git a/Engine/Source/Runtime/Core/Private/Windows/WindowsPlatformStackWalk.cpp b/Engine/Source/Runtime/Core/Private/Windows/WindowsPlatformStackWalk.cpp
--- a/Engine/Source/Runtime/Core/Private/Windows/WindowsPlatformStackWalk.cpp
+++ b/Engine/Source/Runtime/Core/Private/Windows/WindowsPlatformStackWalk.cpp
@@ -746,7 +746,7 @@ void LoadSymbolsForModule(HMODULE ModuleHandle, const FString& RemoteStorage)
 	{
 		if (!SearchPathList.IsEmpty())
 		{
-			SearchPathList.AppendChar(';');
+			SearchPathList.AppendChar(TEXT(';'));
 		}
 		SearchPathList.Append(RemoteStorage);
 	}
@@ -968,11 +968,11 @@ FString GetRemoteStorage(const FString& DownstreamStorage)
 		{
 			if (StorageIndex > 0) 
 			{
-				SymbolStorage.AppendChar(';');
+				SymbolStorage.AppendChar(TEXT(';'));
 			}
 			SymbolStorage.Append(TEXT("SRV*"));
 			SymbolStorage.Append(DownstreamStorage);
-			SymbolStorage.AppendChar('*');
+			SymbolStorage.AppendChar(TEXT('*'));
 			SymbolStorage.Append(RemoteStorage[StorageIndex]);
 		}
 		return SymbolStorage;
diff --git a/Engine/Source/Runtime/Core/Public/GenericPlatform/GenericPlatformDriver.h b/Engine/Source/Runtime/Core/Public/GenericPlatform/GenericPlatformDriver.h
--- a/Engine/Source/Runtime/Core/Public/GenericPlatform/GenericPlatformDriver.h
+++ b/Engine/Source/Runtime/Core/Public/GenericPlatform/GenericPlatformDriver.h
@@ -597,7 +597,7 @@ struct FGPUHardware
 			for (const FString& SuggestedVersion : SuggestedDriverVersions)
 			{
 				int32 Found = 0;
-				if (!SuggestedVersion.FindChar(';', Found))
+				if (!SuggestedVersion.FindChar(TEXT(';'), Found))
 				{
 					return SuggestedVersion;
 				}
diff --git a/Engine/Source/Runtime/Core/Public/UObject/NameTypes.h b/Engine/Source/Runtime/Core/Public/UObject/NameTypes.h
--- a/Engine/Source/Runtime/Core/Public/UObject/NameTypes.h
+++ b/Engine/Source/Runtime/Core/Public/UObject/NameTypes.h
@@ -139,7 +139,7 @@ typedef FNameEntryId NAME_INDEX;
 #define SUBOBJECT_DELIMITER				TEXT(":")
 
 /** this is the character used to separate a subobject root from its subobjects in a path name, as a char */
-#define SUBOBJECT_DELIMITER_CHAR		':'
+#define SUBOBJECT_DELIMITER_CHAR		TEXT(':')
 
 /** These are the characters that cannot be used in general FNames */
 #define INVALID_NAME_CHARACTERS			TEXT("\"' ,\n\r\t")
diff --git a/Engine/Source/Runtime/CoreUObject/Private/AssetRegistry/AssetBundleData.cpp b/Engine/Source/Runtime/CoreUObject/Private/AssetRegistry/AssetBundleData.cpp
--- a/Engine/Source/Runtime/CoreUObject/Private/AssetRegistry/AssetBundleData.cpp
+++ b/Engine/Source/Runtime/CoreUObject/Private/AssetRegistry/AssetBundleData.cpp
@@ -55,7 +55,7 @@ bool FAssetBundleEntry::ExportTextItem(FString& ValueStr, const FAssetBundleEntr
 			return false;
 		}
 		
-		ValueStr.AppendChar(',');
+		ValueStr.AppendChar(TEXT(','));
 	}
 
 	// Remove last comma
@@ -199,7 +199,7 @@ bool FAssetBundleData::ExportTextItem(FString& ValueStr, FAssetBundleData const&
 			ValueStr.LeftInline(OriginalLen);
 			return false;
 		}
-		ValueStr.AppendChar(',');
+		ValueStr.AppendChar(TEXT(','));
 	}
 
 	// Remove last comma
diff --git a/Engine/Source/Runtime/CoreUObject/Private/Misc/PackageName.cpp b/Engine/Source/Runtime/CoreUObject/Private/Misc/PackageName.cpp
--- a/Engine/Source/Runtime/CoreUObject/Private/Misc/PackageName.cpp
+++ b/Engine/Source/Runtime/CoreUObject/Private/Misc/PackageName.cpp
@@ -84,7 +84,7 @@ FString FPackageName::GetShortName(const FString& LongName)
 {
 	// Get everything after the last slash
 	int32 IndexOfLastSlash = INDEX_NONE;
-	LongName.FindLastChar('/', IndexOfLastSlash);
+	LongName.FindLastChar(TEXT('/'), IndexOfLastSlash);
 	return LongName.Mid(IndexOfLastSlash + 1);
 }
 
@@ -852,7 +852,7 @@ bool FPackageName::TryConvertToMountedPath(FStringView InFilePathOrPackageName,
 FString FPackageName::GetLongPackagePath(const FString& InLongPackageName)
 {
 	int32 IndexOfLastSlash = INDEX_NONE;
-	if (InLongPackageName.FindLastChar('/', IndexOfLastSlash))
+	if (InLongPackageName.FindLastChar(TEXT('/'), IndexOfLastSlash))
 	{
 		return InLongPackageName.Left(IndexOfLastSlash);
 	}
@@ -942,15 +942,15 @@ void FPackageName::SplitFullObjectPath(FStringView InFullObjectPath, FStringView
 	// then ClassName is empty and the remaining string is PackageName.ObjectName:SubObjectName
 	if (bDetectClassName)
 	{
-		ExtractBeforeDelim(' ', OutClassName);
+		ExtractBeforeDelim(TEXT(' '), OutClassName);
 	}
 	else
 	{
 		OutClassName.Reset();
 	}
-	if (ExtractBeforeDelim('.', OutPackageName))
+	if (ExtractBeforeDelim(TEXT('.'), OutPackageName))
 	{
-		if (ExtractBeforeDelim(':', OutObjectName))
+		if (ExtractBeforeDelim(TEXT(':'), OutObjectName))
 		{
 			OutSubObjectName = Remaining;
 		}
@@ -1130,7 +1130,7 @@ bool FPackageName::IsValidObjectPath(const FString& InObjectPath, FText* OutReas
 
 	// Check for package delimiter
 	int32 ObjectDelimiterIdx;
-	if (InObjectPath.FindChar('.', ObjectDelimiterIdx))
+	if (InObjectPath.FindChar(TEXT('.'), ObjectDelimiterIdx))
 	{
 		if (ObjectDelimiterIdx == InObjectPath.Len() - 1)
 		{
@@ -1173,7 +1173,7 @@ bool FPackageName::IsValidObjectPath(const FString& InObjectPath, FText* OutReas
 		}
 
 		int32 SlashIndex;
-		if (RemainingObjectPath.FindChar('/', SlashIndex))
+		if (RemainingObjectPath.FindChar(TEXT('/'), SlashIndex))
 		{
 			if (OutReason)
 			{
@@ -1845,7 +1845,7 @@ FString FPackageName::GetDelegateResolvedPackagePath(const FString& InSourcePack
 		FString ObjectName;
 		int32 DotIndex = INDEX_NONE;
 
-		if (PathName.FindChar('.', DotIndex))
+		if (PathName.FindChar(TEXT('.'), DotIndex))
 		{
 			ObjectName = PathName.Mid(DotIndex + 1);
 			PathName.LeftInline(DotIndex, false);
@@ -1868,7 +1868,7 @@ FString FPackageName::GetDelegateResolvedPackagePath(const FString& InSourcePack
 			if (ObjectName.Len())
 			{
 				int32 LastSlashIndex = INDEX_NONE;
-				if (PathName.FindLastChar('/', LastSlashIndex))
+				if (PathName.FindLastChar(TEXT('/'), LastSlashIndex))
 				{
 					ObjectName = PathName.Mid(LastSlashIndex + 1);
 				}
diff --git a/Engine/Source/Runtime/CoreUObject/Private/UObject/CoreRedirects.cpp b/Engine/Source/Runtime/CoreUObject/Private/UObject/CoreRedirects.cpp
--- a/Engine/Source/Runtime/CoreUObject/Private/UObject/CoreRedirects.cpp
+++ b/Engine/Source/Runtime/CoreUObject/Private/UObject/CoreRedirects.cpp
@@ -56,7 +56,7 @@ FCoreRedirectObjectName::FCoreRedirectObjectName(const class UObject* Object)
 				}
 				if (!OuterString.IsEmpty())
 				{
-					OuterString.AppendChar('.');
+					OuterString.AppendChar(TEXT('.'));
 				}
 				OuterString.Append(Outer->GetName());
 
@@ -182,10 +182,10 @@ bool FCoreRedirectObjectName::ExpandNames(const FString& InString, FName& OutNam
 	int32 FirstColonIndex = INDEX_NONE;
 	int32 LastColonIndex = INDEX_NONE;
 
-	FullString.FindChar('/', SlashIndex);
+	FullString.FindChar(TEXT('/'), SlashIndex);
 
-	FullString.FindChar('.', FirstPeriodIndex);
-	FullString.FindChar(':', FirstColonIndex);
+	FullString.FindChar(TEXT('.'), FirstPeriodIndex);
+	FullString.FindChar(TEXT(':'), FirstColonIndex);
 
 	if (FirstColonIndex != INDEX_NONE && (FirstPeriodIndex == INDEX_NONE || FirstColonIndex < FirstPeriodIndex))
 	{
@@ -207,8 +207,8 @@ bool FCoreRedirectObjectName::ExpandNames(const FString& InString, FName& OutNam
 		return true;
 	}
 
-	FullString.FindLastChar('.', LastPeriodIndex);
-	FullString.FindLastChar(':', LastColonIndex);
+	FullString.FindLastChar(TEXT('.'), LastPeriodIndex);
+	FullString.FindLastChar(TEXT(':'), LastColonIndex);
 
 	if (LastColonIndex != INDEX_NONE && (LastPeriodIndex == INDEX_NONE || LastColonIndex > LastPeriodIndex))
 	{
@@ -247,7 +247,7 @@ FString FCoreRedirectObjectName::CombineNames(FName NewName, FName NewOuter, FNa
 
 		int32 DelimIndex = INDEX_NONE;
 
-		if (OuterString.FindChar('.', DelimIndex) || OuterString.FindChar(':', DelimIndex))
+		if (OuterString.FindChar(TEXT('.'), DelimIndex) || OuterString.FindChar(TEXT(':'), DelimIndex))
 		{
 			if (NewPackage != NAME_None)
 			{
diff --git a/Engine/Source/Runtime/CoreUObject/Private/UObject/FieldPath.cpp b/Engine/Source/Runtime/CoreUObject/Private/UObject/FieldPath.cpp
--- a/Engine/Source/Runtime/CoreUObject/Private/UObject/FieldPath.cpp
+++ b/Engine/Source/Runtime/CoreUObject/Private/UObject/FieldPath.cpp
@@ -118,15 +118,15 @@ void FFieldPath::Generate(const TCHAR* InFieldPathString)
 		// Construct names
 		while (true)
 		{
-			if (*InFieldPathString == '.' || *InFieldPathString == SUBOBJECT_DELIMITER_CHAR || *InFieldPathString == '\0')
+			if (*InFieldPathString == TEXT('.') || *InFieldPathString == SUBOBJECT_DELIMITER_CHAR || *InFieldPathString == TEXT('\0'))
 			{
-				NameBuffer[NameIndex] = '\0';
+				NameBuffer[NameIndex] = TEXT('\0');
 				if (NameIndex > 0)
 				{
 					Result.Add(NameBuffer);
 					NameIndex = 0;
 				}
-				if (*InFieldPathString == '\0')
+				if (*InFieldPathString == TEXT('\0'))
 				{
 					break;
 				}
diff --git a/Engine/Source/Runtime/CoreUObject/Private/UObject/LinkerLoad.cpp b/Engine/Source/Runtime/CoreUObject/Private/UObject/LinkerLoad.cpp
--- a/Engine/Source/Runtime/CoreUObject/Private/UObject/LinkerLoad.cpp
+++ b/Engine/Source/Runtime/CoreUObject/Private/UObject/LinkerLoad.cpp
@@ -2055,8 +2055,8 @@ FString ExtractObjectName(const FString& InFullPath)
 {
 	FString ObjectName = InFullPath;
 	int32 LastDot, LastSemi;
-	InFullPath.FindLastChar('.', LastDot);
-	InFullPath.FindLastChar(':', LastSemi);
+	InFullPath.FindLastChar(TEXT('.'), LastDot);
+	InFullPath.FindLastChar(TEXT(':'), LastSemi);
 	int32 StartOfObjectName = FMath::Max(LastDot, LastSemi);
 	if (StartOfObjectName != INDEX_NONE)
 	{
diff --git a/Engine/Source/Runtime/CoreUObject/Private/UObject/Property.cpp b/Engine/Source/Runtime/CoreUObject/Private/UObject/Property.cpp
--- a/Engine/Source/Runtime/CoreUObject/Private/UObject/Property.cpp
+++ b/Engine/Source/Runtime/CoreUObject/Private/UObject/Property.cpp
@@ -593,7 +593,7 @@ FString FGCStackSizeHelper::GetPropertyPath() const
 {
 	FString Result;
 	const FProperty* PreviousProperty = nullptr;
-	const TCHAR DelimiterChar = '.';
+	const TCHAR DelimiterChar = TEXT('.');
 
 	for (int32 PropertyIndex = 0; PropertyIndex < PropertyStack.Num(); ++PropertyIndex)
 	{
@@ -617,7 +617,7 @@ FString FGCStackSizeHelper::GetPropertyPath() const
 
 bool FGCStackSizeHelper::ConvertPathToProperties(UClass* ObjectClass, const FName& InPropertyPath, TArray<FProperty*>& OutProperties)
 {
-	const TCHAR DelimiterChar = '.';
+	const TCHAR DelimiterChar = TEXT('.');
 	FString PropertyNameOrPath = InPropertyPath.ToString();
 	int32 DelimiterIndex = -1;
 	bool bFullPathConstructed = true;
@@ -645,7 +645,7 @@ bool FGCStackSizeHelper::ConvertPathToProperties(UClass* ObjectClass, const FNam
 			OutProperties.Add(FoundProperty);
 
 			int32 StartIndex = DelimiterIndex + 1;
-			const TCHAR DelimiterStr[] = { DelimiterChar, '\0' };
+			const TCHAR DelimiterStr[] = { DelimiterChar, TEXT('\0') };
 			do
 			{
 				// Determine the next property name
@@ -1054,7 +1054,7 @@ void FProperty::ExportCppDeclaration(FOutputDevice& Out, EExportedDeclaration::T
 				else
 				{
 					// export as a pointer if this is an optional out parm, reference if it's just an out parm, standard otherwise...
-					TCHAR ModifierString[2]={0,0};
+					TCHAR ModifierString[2] = { TCHAR('\0'), TCHAR('\0') };
 					if (bCanHaveRef && (HasAnyPropertyFlags(CPF_OutParm | CPF_ReferenceParm) || bIsInterfaceProp))
 					{
 						ModifierString[0] = TEXT('&');
diff --git a/Engine/Source/Runtime/CoreUObject/Private/UObject/UObjectBase.cpp b/Engine/Source/Runtime/CoreUObject/Private/UObject/UObjectBase.cpp
--- a/Engine/Source/Runtime/CoreUObject/Private/UObject/UObjectBase.cpp
+++ b/Engine/Source/Runtime/CoreUObject/Private/UObject/UObjectBase.cpp
@@ -1103,7 +1103,7 @@ const TCHAR* DebugPathName(UObject* Object)
 	{
 		// Hardcoded static array. This function is only used inside the debugger so it should be fine to return it.
 		static TCHAR PathName[1024];
-		PathName[0] = 0;
+		PathName[0] = TCHAR('\0');
 
 		// Keep track of how many outers we have as we need to print them in inverse order.
 		UObject*	TempObject = Object;
@@ -1153,7 +1153,7 @@ const TCHAR* DebugFullName(UObject* Object)
 	{
 		// Hardcoded static array. This function is only used inside the debugger so it should be fine to return it.
 		static TCHAR FullName[1024];
-		FullName[0]=0;
+		FullName[0] = TCHAR('\0');
 
 		// Class Full.Path.Name
 		FCString::Strcat( FullName, DebugFName(Object->GetClass()) );
diff --git a/Engine/Source/Runtime/CoreUObject/Public/AssetRegistry/AssetData.h b/Engine/Source/Runtime/CoreUObject/Public/AssetRegistry/AssetData.h
--- a/Engine/Source/Runtime/CoreUObject/Public/AssetRegistry/AssetData.h
+++ b/Engine/Source/Runtime/CoreUObject/Public/AssetRegistry/AssetData.h
@@ -219,7 +219,7 @@ public:
 	{
 		OutFullName.Reset();
 		AssetClass.AppendString(OutFullName);
-		OutFullName.AppendChar(' ');
+		OutFullName.AppendChar(TEXT(' '));
 		ObjectPath.AppendString(OutFullName);
 	}
 
@@ -228,7 +228,7 @@ public:
 	{
 		OutFullNameBuilder.Reset();
 		AssetClass.AppendString(OutFullNameBuilder);
-		OutFullNameBuilder.AppendChar(' ');
+		OutFullNameBuilder.AppendChar(TEXT*' '));
 		ObjectPath.AppendString(OutFullNameBuilder);
 	}
 
@@ -245,9 +245,9 @@ public:
 	{
 		OutExportTextName.Reset();
 		AssetClass.AppendString(OutExportTextName);
-		OutExportTextName.AppendChar('\'');
+		OutExportTextName.AppendChar(TEXT('\''));
 		ObjectPath.AppendString(OutExportTextName);
-		OutExportTextName.AppendChar('\'');
+		OutExportTextName.AppendChar(TEXT('\''));
 	}
 
 	/** Populates OutExportTextNameBuilder with the name for the asset in the form: Class'ObjectPath' */
@@ -255,9 +255,9 @@ public:
 	{
 		OutExportTextNameBuilder.Reset();
 		AssetClass.AppendString(OutExportTextNameBuilder);
-		OutExportTextNameBuilder.AppendChar('\'');
+		OutExportTextNameBuilder.AppendChar(TEXT('\''));
 		ObjectPath.AppendString(OutExportTextNameBuilder);
-		OutExportTextNameBuilder.AppendChar('\'');
+		OutExportTextNameBuilder.AppendChar(TEXT('\''));
 	}
 
 	/** Returns true if the this asset is a redirector. */
diff --git a/Engine/Source/Runtime/Launch/Private/LaunchEngineLoop.cpp b/Engine/Source/Runtime/Launch/Private/LaunchEngineLoop.cpp
--- a/Engine/Source/Runtime/Launch/Private/LaunchEngineLoop.cpp
+++ b/Engine/Source/Runtime/Launch/Private/LaunchEngineLoop.cpp
@@ -386,7 +386,7 @@ public:
 #if PLATFORM_WINDOWS
 			if (bIsConsoleOutput)
 			{
-				line.AppendChar('\n');
+				line.AppendChar(TEXT('\n'));
 
 				WriteConsoleW(GetStdHandle(STD_OUTPUT_HANDLE), *line, line.Len(), NULL, NULL);
 
diff --git a/Engine/Source/Runtime/Projects/Private/PluginManager.cpp b/Engine/Source/Runtime/Projects/Private/PluginManager.cpp
--- a/Engine/Source/Runtime/Projects/Private/PluginManager.cpp
+++ b/Engine/Source/Runtime/Projects/Private/PluginManager.cpp
@@ -113,9 +113,9 @@ FString FPlugin::GetMountedAssetPath() const
 {
 	FString Path;
 	Path.Reserve(Name.Len() + 2);
-	Path.AppendChar('/');
+	Path.AppendChar(TEXT('/'));
 	Path.Append(Name);
-	Path.AppendChar('/');
+	Path.AppendChar(TEXT('/'));
 	return Path;
 }
 
