Make use of batch existence check operation available on the Jupiter server for references.
Also allow receipt of uncompressed buffers when using the inlining functionality to get values back in the first round trip for references.

diff --git a/Engine/Source/Developer/DerivedDataCache/Private/HttpCacheStore.cpp b/Engine/Source/Developer/DerivedDataCache/Private/HttpCacheStore.cpp
--- a/Engine/Source/Developer/DerivedDataCache/Private/HttpCacheStore.cpp
+++ b/Engine/Source/Developer/DerivedDataCache/Private/HttpCacheStore.cpp
@@ -2223,6 +2223,12 @@ private:
 		ValueHashGetterType ValueHashGetter,
 		ValueDebugContextGetterType ValueDebugContextGetter,
 		FHttpRequest* ExistingHttpRequest = nullptr);
+	template<typename ValueRefType, typename ValueRefKeyGetterType, typename ValueRefDebugContextGetterType>
+	TArray<FValue> RefCachedDataProbablyExistsBatch(
+		TConstArrayView<ValueRefType> ValueRefs,
+		ValueRefKeyGetterType ValueRefKeyGetter,
+		ValueRefDebugContextGetterType ValueRefDebugContextGetter,
+		FHttpRequest* ExistingHttpRequest = nullptr);
 };
 
 FHttpCacheStore::FHttpCacheStore(
@@ -2571,7 +2577,22 @@ uint64 FHttpCacheStore::PutRef(const FCbPackage& Package, const FCacheKey& Key,
 		if (FHttpRequest::IsSuccessResponse(ResponseCode))
 		{
 			BytesSent += Request->GetBytesSent();
-			if (TSharedPtr<FJsonObject> ResponseObject = Request->GetResponseAsJsonObject())
+
+			// Useful when debugging issues related to compressed/uncompressed blobs being returned from Jupiter
+			const bool bPutRefBlobsAlways = false;
+
+			if (bPutRefBlobsAlways && !bFinalize)
+			{
+				Package.GetObject().IterateAttachments([&OutNeededBlobHashes](FCbFieldView AttachmentFieldView)
+					{
+						FIoHash AttachmentHash = AttachmentFieldView.AsHash();
+						if (!AttachmentHash.IsZero())
+						{
+							OutNeededBlobHashes.Add(AttachmentHash);
+						}
+					});
+			}
+			else if (TSharedPtr<FJsonObject> ResponseObject = Request->GetResponseAsJsonObject())
 			{
 				TArray<FString> NeedsArrayStrings;
 				ResponseObject->TryGetStringArrayField(TEXT("needs"), NeedsArrayStrings);
@@ -3057,7 +3078,21 @@ bool FHttpCacheStore::GetCacheValue(
 	}
 	else
 	{
-		FCompressedBuffer CompressedBuffer = FCompressedBuffer::FromCompressed(MoveTemp(ResponseBuffer));
+		FCompressedBuffer CompressedBuffer = FCompressedBuffer::FromCompressed(ResponseBuffer);
+		if (!CompressedBuffer)
+		{
+			FString ReceivedHashStr;
+			if (Request->GetHeader("X-Jupiter-InlinePayloadHash", ReceivedHashStr))
+			{
+				FIoHash ReceivedHash(ReceivedHashStr);
+				FIoHash ComputedHash = FIoHash::HashBuffer(ResponseBuffer.GetView());
+				if (ReceivedHash == ComputedHash)
+				{
+					CompressedBuffer = FCompressedBuffer::Compress(ResponseBuffer);
+				}
+			}
+		}
+
 		if (!CompressedBuffer)
 		{
 			UE_LOG(LogDerivedDataCache, Display, TEXT("%s: Cache miss with invalid package for %s from '%.*s'"),
@@ -3332,6 +3367,140 @@ TBitArray<> FHttpCacheStore::CachedDataProbablyExistsBatch(
 	return TBitArray<>(false, Values.Num());
 }
 
+template<typename ValueRefType, typename ValueRefKeyGetterType, typename ValueRefDebugContextGetterType>
+TArray<FValue> FHttpCacheStore::RefCachedDataProbablyExistsBatch(
+	TConstArrayView<ValueRefType> ValueRefs,
+	ValueRefKeyGetterType ValueRefKeyGetter,
+	ValueRefDebugContextGetterType ValueRefDebugContextGetter,
+	FHttpRequest* ExistingHttpRequest)
+{
+	if (ValueRefs.IsEmpty())
+	{
+		return TArray<FValue>();
+	}
+
+	FHttpRequest* Request = ExistingHttpRequest;
+	TOptional<FScopedRequestPtr> RequestPoolRequest;
+	if (Request == nullptr)
+	{
+		RequestPoolRequest.Emplace(GetRequestPools[IsInGameThread()].Get());
+		Request = RequestPoolRequest->Get();
+	}
+	else
+	{
+		Request->Reset();
+	}
+
+	TStringBuilder<256> RefsUri;
+	RefsUri << "api/v1/refs/" << StructuredNamespace;
+	FCbWriter RequestWriter;
+	RequestWriter.BeginObject();
+	RequestWriter.BeginArray("ops"_ASV);
+	uint32 OpIndex = 0;
+	for (const ValueRefType& ValueRef : ValueRefs)
+	{
+		RequestWriter.BeginObject();
+		RequestWriter.AddInteger("opId"_ASV, OpIndex);
+		RequestWriter.AddString("op"_ASV, "GET"_ASV);
+		FCacheKey Key = ValueRefKeyGetter(ValueRef);
+		FString Bucket(Key.Bucket.ToString());
+		Bucket.ToLowerInline();
+		RequestWriter.AddString("bucket"_ASV, Bucket);
+		RequestWriter.AddString("key"_ASV, LexToString(Key.Hash));
+		RequestWriter.EndObject();
+		++OpIndex;
+	}
+	RequestWriter.EndArray();
+	RequestWriter.EndObject();
+	FCbFieldIterator RequestFields = RequestWriter.Save();
+
+	TConstArrayView<uint8> BodyBuffer = MakeConstArrayView(RequestFields.GetOuterBuffer());
+
+	int64 ResponseCode = 0;
+	for (uint32 Attempts = 0; (Attempts < UE_HTTPDDC_MAX_ATTEMPTS) && !ShouldAbortForShutdown() && (Attempts == 0 || ShouldRetryOnError(ResponseCode)); ++Attempts)
+	{
+		if (Attempts > 0)
+		{
+			Request->Reset();
+		}
+
+		Request->SetHeader(TEXT("Accept"), TEXT("application/x-ue-cb"));
+		Request->SetHeader(TEXT("Content-Type"), TEXT("application/x-ue-cb"));
+		const FHttpRequest::Result Result = Request->PerformBlockingUpload<FHttpRequest::Post>(*RefsUri, BodyBuffer);
+		ResponseCode = Request->GetResponseCode();
+
+		if (FHttpRequest::IsSuccessResponse(ResponseCode))
+		{
+			FMemoryView ResponseView = MakeMemoryView(Request->GetResponseBuffer().GetData(), Request->GetResponseBuffer().Num());
+			if (ValidateCompactBinary(ResponseView, ECbValidateMode::Default) != ECbValidateError::None)
+			{
+				UE_LOG(LogDerivedDataCache, Warning,
+					TEXT("%s: Cache exists returned invalid results."),
+					*GetName());
+				TArray<FValue> RetVal;
+				RetVal.AddDefaulted(ValueRefs.Num());
+				return RetVal;
+			}
+
+			const FCbObjectView ResponseObject = FCbObjectView(Request->GetResponseBuffer().GetData());
+
+			FCbArrayView ResultsArrayView = ResponseObject["results"_ASV].AsArrayView();
+
+			if (ResultsArrayView.Num() != ValueRefs.Num())
+			{
+				UE_LOG(LogDerivedDataCache, Warning,
+					TEXT("%s: Cache exists returned unexpected quantity of results (expected %d, got %d)."),
+					*GetName(), ValueRefs.Num(), ResultsArrayView.Num());
+				TArray<FValue> RetVal;
+				RetVal.AddDefaulted(ValueRefs.Num());
+				return RetVal;
+			}
+			TArray<FValue> RetVal;
+			RetVal.AddDefaulted(ValueRefs.Num());
+			for (FCbFieldView ResultFieldView : ResultsArrayView)
+			{
+				FCbObjectView ResultObjectView = ResultFieldView.AsObjectView();
+				uint32 OpId = ResultObjectView["opId"_ASV].AsUInt32();
+				FCbObjectView ResponseObjectView = ResultObjectView["response"_ASV].AsObjectView();
+				int32 StatusCode = ResultObjectView["statusCode"_ASV].AsInt32();
+
+				if (OpId >= (uint32)RetVal.Num())
+				{
+					FValueDebugContext DebugContext = ValueRefDebugContextGetter(ValueRefs[OpId]);
+					UE_LOG(LogDerivedDataCache, Display, TEXT("%s: Cache miss with invalid op index %u for %s from '%.*s'"),
+						*GetName(), OpId, *WriteToString<96>(DebugContext.Key), DebugContext.Name.Len(), DebugContext.Name.GetData());
+					continue;
+				}
+
+				if (!FHttpRequest::IsSuccessResponse(StatusCode))
+				{
+					FValueDebugContext DebugContext = ValueRefDebugContextGetter(ValueRefs[OpId]);
+					UE_LOG(LogDerivedDataCache, Verbose, TEXT("%s: Cache miss with unsuccessful response code %d for %s from '%.*s'"),
+						*GetName(), StatusCode, *WriteToString<96>(DebugContext.Key), DebugContext.Name.Len(), DebugContext.Name.GetData());
+					continue;
+				}
+
+				const FIoHash RawHash = ResponseObjectView["RawHash"_ASV].AsHash();
+				const uint64 RawSize = ResponseObjectView["RawSize"_ASV].AsUInt64(MAX_uint64);
+				if (RawHash.IsZero() || RawSize == MAX_uint64)
+				{
+					FValueDebugContext DebugContext = ValueRefDebugContextGetter(ValueRefs[OpId]);
+					UE_LOG(LogDerivedDataCache, Display, TEXT("%s: Cache miss with invalid value for %s from '%.*s'"),
+						*GetName(), *WriteToString<96>(DebugContext.Key), DebugContext.Name.Len(), DebugContext.Name.GetData());
+					continue;
+				}
+
+				RetVal[OpId] = FValue(RawHash, RawSize);
+			}
+			return RetVal;
+		}
+	}
+
+	TArray<FValue> RetVal;
+	RetVal.AddDefaulted(ValueRefs.Num());
+	return RetVal;
+}
+
 bool FHttpCacheStore::CachedDataProbablyExists(const TCHAR* CacheKey)
 {
 	TRACE_CPUPROFILER_EVENT_SCOPE(HttpDDC_Exist);
@@ -3771,21 +3940,35 @@ void FHttpCacheStore::GetValue(
 	FScopedRequestPtr HttpRequest(GetRequestPools[IsInGameThread()].Get());
 	int64 HitBytes = 0;
 
+	bool bBatchExistsCandidate = true;
 	for (const FCacheGetValueRequest& Request : Requests)
 	{
-		FValue Value;
-		if (!GetCacheValue(Request.Name, Request.Key, Request.Policy, Value, HttpRequest.Get()))
+		if (!EnumHasAnyFlags(Request.Policy, ECachePolicy::SkipData))
 		{
-			COOK_STAT(UsageStats.GetStats.Accumulate(FCookStats::CallStats::EHitOrMiss::Miss, FCookStats::CallStats::EStatType::Counter, 1l, bIsInGameThread));
-			OnComplete(Request.MakeResponse(EStatus::Error));
+			bBatchExistsCandidate = false;
+			break;
 		}
-		else
+	}
+
+	if (bBatchExistsCandidate)
+	{
+		auto KeyGetter = [](const FCacheGetValueRequest& Request)
+		{
+			return Request.Key;
+		};
+
+		auto DebugContextGetter = [](const FCacheGetValueRequest& Request)
+		{
+			return FValueDebugContext{ Request.Name, Request.Key, FString(TEXT("Default")) };
+		};
+
+		TArray<FValue> Values = RefCachedDataProbablyExistsBatch<FCacheGetValueRequest>(Requests, KeyGetter, DebugContextGetter, HttpRequest.Get());
+
+		for (int32 RequestIndex = 0; RequestIndex < Requests.Num(); ++RequestIndex)
 		{
-			if (!IsValueDataReady(Value, Request.Policy) && !EnumHasAnyFlags(Request.Policy, ECachePolicy::SkipData))
+			const FCacheGetValueRequest& Request = Requests[RequestIndex];
+			if (Values[RequestIndex] == FValue::Null)
 			{
-				// With inline fetching, expect we will always have a value we can use.  Even SkipData/Exists can rely on the blob existing if the ref is reported to exist.
-				UE_LOG(LogDerivedDataCache, Warning, TEXT("%s: Cache miss due to inlining failure for %s from '%s'"),
-					*GetName(), *WriteToString<96>(Request.Key), *Request.Name);
 				COOK_STAT(UsageStats.GetStats.Accumulate(FCookStats::CallStats::EHitOrMiss::Miss, FCookStats::CallStats::EStatType::Counter, 1l, bIsInGameThread));
 				OnComplete(Request.MakeResponse(EStatus::Error));
 			}
@@ -3793,11 +3976,41 @@ void FHttpCacheStore::GetValue(
 			{
 				UE_LOG(LogDerivedDataCache, Verbose, TEXT("%s: Cache hit for %s from '%s'"),
 					*GetName(), *WriteToString<96>(Request.Key), *Request.Name);
-				uint64 ValueSize = Value.GetData().GetCompressedSize();
-				TRACE_COUNTER_ADD(HttpDDC_BytesReceived, ValueSize);
-				HitBytes += ValueSize;
 				COOK_STAT(UsageStats.GetStats.Accumulate(FCookStats::CallStats::EHitOrMiss::Hit, FCookStats::CallStats::EStatType::Counter, 1l, bIsInGameThread));
-				OnComplete({ Request.Name, Request.Key, Value, Request.UserData, EStatus::Ok });
+				OnComplete({ Request.Name, Request.Key, Values[RequestIndex], Request.UserData, EStatus::Ok });
+			}
+		}
+	}
+	else
+	{
+		for (const FCacheGetValueRequest& Request : Requests)
+		{
+			FValue Value;
+			if (!GetCacheValue(Request.Name, Request.Key, Request.Policy, Value, HttpRequest.Get()))
+			{
+				COOK_STAT(UsageStats.GetStats.Accumulate(FCookStats::CallStats::EHitOrMiss::Miss, FCookStats::CallStats::EStatType::Counter, 1l, bIsInGameThread));
+				OnComplete(Request.MakeResponse(EStatus::Error));
+			}
+			else
+			{
+				if (!IsValueDataReady(Value, Request.Policy) && !EnumHasAnyFlags(Request.Policy, ECachePolicy::SkipData))
+				{
+					// With inline fetching, expect we will always have a value we can use.  Even SkipData/Exists can rely on the blob existing if the ref is reported to exist.
+					UE_LOG(LogDerivedDataCache, Warning, TEXT("%s: Cache miss due to inlining failure for %s from '%s'"),
+						*GetName(), *WriteToString<96>(Request.Key), *Request.Name);
+					COOK_STAT(UsageStats.GetStats.Accumulate(FCookStats::CallStats::EHitOrMiss::Miss, FCookStats::CallStats::EStatType::Counter, 1l, bIsInGameThread));
+					OnComplete(Request.MakeResponse(EStatus::Error));
+				}
+				else
+				{
+					UE_LOG(LogDerivedDataCache, Verbose, TEXT("%s: Cache hit for %s from '%s'"),
+						*GetName(), *WriteToString<96>(Request.Key), *Request.Name);
+					uint64 ValueSize = Value.GetData().GetCompressedSize();
+					TRACE_COUNTER_ADD(HttpDDC_BytesReceived, ValueSize);
+					HitBytes += ValueSize;
+					COOK_STAT(UsageStats.GetStats.Accumulate(FCookStats::CallStats::EHitOrMiss::Hit, FCookStats::CallStats::EStatType::Counter, 1l, bIsInGameThread));
+					OnComplete({ Request.Name, Request.Key, Value, Request.UserData, EStatus::Ok });
+				}
 			}
 		}
 	}
