Resolve compile issue in UE Zen client caused by including an Internal header from a Public header.
This is being resolved by cherrypicking 19196840 by mark.lintott from UE5/Main (see original changelist details at the end of this cl description).  On top of the cherrypick, there was a change made to re-use a single Http request when fetching stats.  This avoids the overhead of doing name resolution repeatedly if the zen instance has a name instead of IP address.

Tested by:
- Running ShooterGame editor & observing UnrealInsights for ZenHttp_CurlPerform timers and ensuring they are not happening on the game thread when the editor is sitting idle (they are on background threads instead)
-Terminating ZenServer.exe while the editor is running and observing whether the framerate of the editor is negatively affected

Original changelist:
============
Blocking Zen HTTP stats request now runs async using futures
#[fyi] Zousar.Shaker

diff --git a/Engine/Source/Developer/DerivedDataCache/Private/ZenCacheStore.cpp b/Engine/Source/Developer/DerivedDataCache/Private/ZenCacheStore.cpp
--- a/Engine/Source/Developer/DerivedDataCache/Private/ZenCacheStore.cpp
+++ b/Engine/Source/Developer/DerivedDataCache/Private/ZenCacheStore.cpp
@@ -173,7 +173,6 @@ private:
 	mutable FDerivedDataCacheUsageStats UsageStats;
 	TUniquePtr<UE::Zen::FZenHttpRequestPool> RequestPool;
 	bool bIsUsable = false;
-	bool bIsRemote = false;
 	uint32 FailedLoginAttempts = 0;
 	uint32 MaxAttempts = 4;
 	int32 BatchPutMaxBytes = 1024*1024;
@@ -199,14 +198,10 @@ FZenCacheStore::FZenCacheStore(
 	{
 		RequestPool = MakeUnique<Zen::FZenHttpRequestPool>(ZenService.GetInstance().GetURL(), 32);
 		bIsUsable = true;
-		bIsRemote = false;
 
+		// Issue a request for stats as it will be fetched asynchronously and issuing now makes them available sooner for future callers.
 		Zen::FZenStats ZenStats;
-
-		if (ZenService.GetInstance().GetStats(ZenStats) == true)
-		{
-			 bIsRemote = ZenStats.UpstreamStats.EndPointStats.IsEmpty() == false;
-		}
+		ZenService.GetInstance().GetStats(ZenStats);
 	}
 
 	GConfig->GetInt(TEXT("Zen"), TEXT("BatchPutMaxBytes"), BatchPutMaxBytes, GEngineIni);
diff --git a/Engine/Source/Developer/Zen/Internal/ZenServerInterface.cpp b/Engine/Source/Developer/Zen/Internal/ZenServerInterface.cpp
--- a/Engine/Source/Developer/Zen/Internal/ZenServerInterface.cpp
+++ b/Engine/Source/Developer/Zen/Internal/ZenServerInterface.cpp
@@ -1073,12 +1073,18 @@ FZenServiceInstance::GetStats(FZenStats& Stats)
 #else
 		EAsyncExecution ThreadPool = EAsyncExecution::ThreadPool;
 #endif
+		if (!StatsHttpRequest.IsValid())
+		{
+			TStringBuilder<128> ZenDomain;
+			ZenDomain << HostName << TEXT(":") << Port;
+			StatsHttpRequest = MakePimpl<FZenHttpRequest>(ZenDomain.ToString(), false);
+		}
+
 		// We've not got any requests in flight and we've met a given time requirement for requests
 		StatsRequest = Async(ThreadPool, [this]
 			{
-				TStringBuilder<128> ZenDomain;
-				ZenDomain << HostName << TEXT(":") << Port;
-				UE::Zen::FZenHttpRequest Request(ZenDomain.ToString(), false);
+				UE::Zen::FZenHttpRequest& Request = *StatsHttpRequest.Get();
+				Request.Reset();
 
 				TArray64<uint8> GetBuffer;
 				FZenHttpRequest::Result Result = Request.PerformBlockingDownload(TEXT("/stats/z$"_SV), &GetBuffer, Zen::EContentType::CbObject);
diff --git a/Engine/Source/Developer/Zen/Public/ZenServerInterface.h b/Engine/Source/Developer/Zen/Public/ZenServerInterface.h
--- a/Engine/Source/Developer/Zen/Public/ZenServerInterface.h
+++ b/Engine/Source/Developer/Zen/Public/ZenServerInterface.h
@@ -10,6 +10,7 @@
 #include "Misc/TVariant.h"
 #include "Policies/PrettyJsonPrintPolicy.h"
 #include "Serialization/JsonWriter.h"
+#include "Templates/PimplPtr.h"
 #include "Templates/UniquePtr.h"
 #include "ZenGlobals.h"
 #include "Async/Future.h"
@@ -134,6 +135,7 @@ private:
 	FString ConditionalUpdateLocalInstall();
 	static bool AutoLaunch(const FServiceAutoLaunchSettings& InSettings, FString&& ExecutablePath, FString& OutHostName, uint16& OutPort);
 
+	mutable TPimplPtr<class FZenHttpRequest> StatsHttpRequest;
 	mutable TFuture<FZenStats> StatsRequest;
 	mutable FZenStats LastStats;
 	mutable uint64 LastStatsTime = 0;
