DDC: Fixed a crash with -NoThreading when upgrading legacy values

#rnx

diff --git a/Engine/Source/Developer/DerivedDataCache/Private/DerivedDataCacheStoreAsync.cpp b/Engine/Source/Developer/DerivedDataCache/Private/DerivedDataCacheStoreAsync.cpp
--- a/Engine/Source/Developer/DerivedDataCache/Private/DerivedDataCacheStoreAsync.cpp
+++ b/Engine/Source/Developer/DerivedDataCache/Private/DerivedDataCacheStoreAsync.cpp
@@ -290,8 +290,15 @@ void Private::ExecuteInCacheThreadPool(
 	IRequestOwner& Owner,
 	TUniqueFunction<void (IRequestOwner& Owner, bool bCancel)>&& Function)
 {
-	FDerivedDataAsyncWrapperRequest* Request = new FDerivedDataAsyncWrapperRequest(Owner, MoveTemp(Function));
-	Request->Start(Owner.GetPriority());
+	if (GCacheThreadPool)
+	{
+		FDerivedDataAsyncWrapperRequest* Request = new FDerivedDataAsyncWrapperRequest(Owner, MoveTemp(Function));
+		Request->Start(Owner.GetPriority());
+	}
+	else
+	{
+		Invoke(Function, Owner, /*bCancel*/ false);
+	}
 }
 
 template <typename RequestType, typename OnCompleteType, typename OnExecuteType>
