Incorporate http cache store client change for partner use.

diff --git a/Engine/Source/Developer/DerivedDataCache/Private/DerivedDataBackends.cpp b/Engine/Source/Developer/DerivedDataCache/Private/DerivedDataBackends.cpp
--- a/Engine/Source/Developer/DerivedDataCache/Private/DerivedDataBackends.cpp
+++ b/Engine/Source/Developer/DerivedDataCache/Private/DerivedDataBackends.cpp
@@ -54,6 +54,7 @@ ILegacyCacheStore* CreateFileSystemCacheStore(const TCHAR* CacheDirectory, const
 ILegacyCacheStore* CreateHttpCacheStore(
 	const TCHAR* NodeName,
 	const TCHAR* ServiceUrl,
+	bool bResolveHostCanonicalName,
 	const TCHAR* Namespace,
 	const TCHAR* StructuredNamespace,
 	const TCHAR* OAuthProvider,
@@ -770,6 +771,7 @@ public:
 		const FString& IniFilename,
 		const TCHAR* IniSection,
 		FString& Host,
+		bool& bResolveHostCanonicalName,
 		FString& Namespace,
 		FString& StructuredNamespace,
 		FString& OAuthProvider,
@@ -785,7 +787,7 @@ public:
 			const TCHAR* ServerSection = TEXT("HordeStorageServers");
 			if (GConfig->GetString(ServerSection, *ServerId, ServerEntry, IniFilename))
 			{
-				ParseHttpCacheParams(NodeName, *ServerEntry, IniFilename, IniSection, Host, Namespace, StructuredNamespace, OAuthProvider, OAuthClientId, OAuthSecret, LegacyMode, bReadOnly);
+				ParseHttpCacheParams(NodeName, *ServerEntry, IniFilename, IniSection, Host, bResolveHostCanonicalName, Namespace, StructuredNamespace, OAuthProvider, OAuthClientId, OAuthSecret, LegacyMode, bReadOnly);
 			}
 			else
 			{
@@ -815,6 +817,7 @@ public:
 			}
 		}
 
+		FParse::Bool(Entry, TEXT("ResolveHostCanonicalName="), bResolveHostCanonicalName);
 		FParse::Value(Entry, TEXT("Namespace="), Namespace);
 		FParse::Value(Entry, TEXT("StructuredNamespace="), StructuredNamespace);
 		FParse::Value(Entry, TEXT("OAuthProvider="), OAuthProvider);
@@ -841,6 +844,7 @@ public:
 		const TCHAR* IniSection)
 	{
 		FString Host;
+		bool bResolveHostCanonicalName = true;
 		FString Namespace;
 		FString StructuredNamespace;
 		FString OAuthProvider;
@@ -849,7 +853,7 @@ public:
 		EBackendLegacyMode LegacyMode = EBackendLegacyMode::ValueOnly;
 		bool bReadOnly = false;
 
-		ParseHttpCacheParams(NodeName, Entry, IniFilename, IniSection, Host, Namespace, StructuredNamespace, OAuthProvider, OAuthClientId, OAuthSecret, LegacyMode, bReadOnly);
+		ParseHttpCacheParams(NodeName, Entry, IniFilename, IniSection, Host, bResolveHostCanonicalName, Namespace, StructuredNamespace, OAuthProvider, OAuthClientId, OAuthSecret, LegacyMode, bReadOnly);
 
 		if (Host.IsEmpty())
 		{
@@ -880,16 +884,20 @@ public:
 			return MakeTuple(nullptr, ECacheStoreFlags::None);
 		}
 
-		if (OAuthClientId.IsEmpty())
+		// No need for OAuth client id and secret if using a local provider.
+		if (!OAuthProvider.StartsWith(TEXT("http://localhost")))
 		{
-			UE_LOG(LogDerivedDataCache, Error, TEXT("Node %s does not specify 'OAuthClientId'"), NodeName);
-			return MakeTuple(nullptr, ECacheStoreFlags::None);
-		}
+			if (OAuthClientId.IsEmpty())
+			{
+				UE_LOG(LogDerivedDataCache, Error, TEXT("Node %s does not specify 'OAuthClientId'"), NodeName);
+				return MakeTuple(nullptr, ECacheStoreFlags::None);
+			}
 
-		if (OAuthSecret.IsEmpty())
-		{
-			UE_LOG(LogDerivedDataCache, Error, TEXT("Node %s does not specify 'OAuthSecret'"), NodeName);
-			return MakeTuple(nullptr, ECacheStoreFlags::None);
+			if (OAuthSecret.IsEmpty())
+			{
+				UE_LOG(LogDerivedDataCache, Error, TEXT("Node %s does not specify 'OAuthSecret'"), NodeName);
+				return MakeTuple(nullptr, ECacheStoreFlags::None);
+			}
 		}
 
 		FDerivedDataBackendInterface::ESpeedClass ForceSpeedClass = FDerivedDataBackendInterface::ESpeedClass::Unknown;
@@ -924,7 +932,7 @@ public:
 		}
 
 		return MakeTuple(CreateHttpCacheStore(
-			NodeName, *Host, *Namespace, *StructuredNamespace, *OAuthProvider, *OAuthClientId, *OAuthSecret,
+			NodeName, *Host, bResolveHostCanonicalName, *Namespace, *StructuredNamespace, *OAuthProvider, *OAuthClientId, *OAuthSecret,
 			ForceSpeedClass == FDerivedDataBackendInterface::ESpeedClass::Unknown ? nullptr : &ForceSpeedClass, LegacyMode, bReadOnly),
 			ECacheStoreFlags::Remote | ECacheStoreFlags::Query | (bReadOnly ? ECacheStoreFlags::None : ECacheStoreFlags::Store));
 	}
diff --git a/Engine/Source/Developer/DerivedDataCache/Private/HttpCacheStore.cpp b/Engine/Source/Developer/DerivedDataCache/Private/HttpCacheStore.cpp
--- a/Engine/Source/Developer/DerivedDataCache/Private/HttpCacheStore.cpp
+++ b/Engine/Source/Developer/DerivedDataCache/Private/HttpCacheStore.cpp
@@ -2073,6 +2073,7 @@ public:
 	 */
 	FHttpCacheStore(
 		const TCHAR* ServiceUrl, 
+		bool bResolveHostCanonicalName,
 		const TCHAR* Namespace, 
 		const TCHAR* StructuredNamespace, 
 		const TCHAR* OAuthProvider, 
@@ -2233,6 +2234,7 @@ private:
 
 FHttpCacheStore::FHttpCacheStore(
 	const TCHAR* InServiceUrl, 
+	bool bResolveHostCanonicalName,
 	const TCHAR* InNamespace, 
 	const TCHAR* InStructuredNamespace, 
 	const TCHAR* InOAuthProvider,
@@ -2283,7 +2285,7 @@ FHttpCacheStore::FHttpCacheStore(
 		FMemory::Memset(&AddrHints, 0, sizeof(AddrHints));
 		AddrHints.ai_flags = AI_CANONNAME;
 		AddrHints.ai_family = AF_UNSPEC;
-		if (!::getaddrinfo(*DomainResolveName, nullptr, &AddrHints, &AddrResult))
+		if (bResolveHostCanonicalName && !::getaddrinfo(*DomainResolveName, nullptr, &AddrHints, &AddrResult))
 		{
 			if (AddrResult->ai_canonname)
 			{
@@ -2374,6 +2376,12 @@ bool FHttpCacheStore::IsServiceReady()
 
 bool FHttpCacheStore::AcquireAccessToken()
 {
+	if (Domain.StartsWith(TEXT("http://localhost")))
+	{
+		UE_LOG(LogDerivedDataCache, Log, TEXT("Connecting to a local host '%s', so skipping authorization"), *Domain);
+		return true;
+	}
+
 	// Avoid spamming the this if the service is down
 	if (FailedLoginAttempts > UE_HTTPDDC_MAX_FAILED_LOGIN_ATTEMPTS)
 	{
@@ -2405,35 +2413,46 @@ bool FHttpCacheStore::AcquireAccessToken()
 		FString Uri(*OAuthProvider + DomainEnd + 1);
 
 		FHttpRequest Request(*AuthDomain, *AuthDomain, nullptr, false);
-
-		// If contents of the secret string is a file path, resolve and read form data.
-		if (OAuthSecret.StartsWith(TEXT("file://")))
+		FHttpRequest::Result Result = FHttpRequest::Success;
+		if (OAuthProvider.StartsWith(TEXT("http://localhost")))
 		{
-			FString FilePath = OAuthSecret.Mid(7, OAuthSecret.Len() - 7);
-			FString SecretFileContents;
-			if (FFileHelper::LoadFileToString(SecretFileContents, *FilePath))
-			{
-				// Overwrite the filepath with the actual content.
-				OAuthSecret = SecretFileContents;
-			}
-			else
+			// Simple unauthenticated call to a local endpoint that mimics
+			// the result from an OIDC provider.
+			Result = Request.PerformBlockingDownload(*Uri, nullptr);
+		}
+		else
+		{
+			// Needs client id and secret to authenticate with an actual OIDC provider.
+
+			// If contents of the secret string is a file path, resolve and read form data.
+			if (OAuthSecret.StartsWith(TEXT("file://")))
 			{
-				UE_LOG(LogDerivedDataCache, Warning, TEXT("%s: Failed to read OAuth form data file (%s)."), *Request.GetName(), *OAuthSecret);
-				return false;
+				FString FilePath = OAuthSecret.Mid(7, OAuthSecret.Len() - 7);
+					FString SecretFileContents;
+				if (FFileHelper::LoadFileToString(SecretFileContents, *FilePath))
+				{
+					// Overwrite the filepath with the actual content.
+					OAuthSecret = SecretFileContents;
+				}
+				else
+				{
+					UE_LOG(LogDerivedDataCache, Warning, TEXT("%s: Failed to read OAuth form data file (%s)."), *Request.GetName(), *OAuthSecret);
+					return false;
+				}
 			}
-		}
 
-		FString OAuthFormData = FString::Printf(
-			TEXT("client_id=%s&scope=cache_access&grant_type=client_credentials&client_secret=%s"),
-			*OAuthClientId,
-			*OAuthSecret
-		);
+			FString OAuthFormData = FString::Printf(
+				TEXT("client_id=%s&scope=cache_access&grant_type=client_credentials&client_secret=%s"),
+				*OAuthClientId,
+				*OAuthSecret
+			);
 
-		TArray<uint8> FormData;
-		auto OAuthFormDataUTF8 = FTCHARToUTF8(*OAuthFormData);
-		FormData.Append((uint8*)OAuthFormDataUTF8.Get(), OAuthFormDataUTF8.Length());
+			TArray<uint8> FormData;
+			auto OAuthFormDataUTF8 = FTCHARToUTF8(*OAuthFormData);
+			FormData.Append((uint8*)OAuthFormDataUTF8.Get(), OAuthFormDataUTF8.Length());
 
-		FHttpRequest::Result Result = Request.PerformBlockingUpload<FHttpRequest::Post>(*Uri, MakeArrayView(FormData));
+			Result = Request.PerformBlockingUpload<FHttpRequest::Post>(*Uri, MakeArrayView(FormData));
+		}
 
 		if (Result == FHttpRequest::Success && Request.GetResponseCode() == 200)
 		{
@@ -4168,6 +4187,7 @@ namespace UE::DerivedData
 ILegacyCacheStore* CreateHttpCacheStore(
 	const TCHAR* NodeName,
 	const TCHAR* ServiceUrl,
+	bool bResolveHostCanonicalName,
 	const TCHAR* Namespace,
 	const TCHAR* StructuredNamespace,
 	const TCHAR* OAuthProvider,
@@ -4178,7 +4198,7 @@ ILegacyCacheStore* CreateHttpCacheStore(
 	bool bReadOnly)
 {
 #if WITH_HTTP_DDC_BACKEND
-	FHttpCacheStore* Backend = new FHttpCacheStore(ServiceUrl, Namespace, StructuredNamespace, OAuthProvider, OAuthClientId, OAuthData, LegacyMode, bReadOnly);
+	FHttpCacheStore* Backend = new FHttpCacheStore(ServiceUrl, bResolveHostCanonicalName, Namespace, StructuredNamespace, OAuthProvider, OAuthClientId, OAuthData, LegacyMode, bReadOnly);
 	if (Backend->IsUsable())
 	{
 		return Backend;
