Lowering log severity of unexpected results coming back from the Horde Storage server to avoid causing failures for operations that treat warnings as errors.

diff --git a/Engine/Source/Developer/DerivedDataCache/Private/HttpCacheStore.cpp b/Engine/Source/Developer/DerivedDataCache/Private/HttpCacheStore.cpp
--- a/Engine/Source/Developer/DerivedDataCache/Private/HttpCacheStore.cpp
+++ b/Engine/Source/Developer/DerivedDataCache/Private/HttpCacheStore.cpp
@@ -2697,7 +2697,7 @@ bool FHttpCacheStore::PutCacheRecord(
 
 	if (!bPutCompletedSuccessfully)
 	{
-		UE_LOG(LogDerivedDataCache, Warning, TEXT("%s: Failed to put reference object for put of %s from '%.*s'"),
+		UE_LOG(LogDerivedDataCache, Log, TEXT("%s: Failed to put reference object for put of %s from '%.*s'"),
 			*GetName(), *WriteToString<96>(Key), Name.Len(), Name.GetData());
 		return false;
 	}
@@ -2763,7 +2763,7 @@ bool FHttpCacheStore::PutCacheRecord(
 				}
 				bExpectedHashesSerialized = true;
 			}
-			UE_LOG(LogDerivedDataCache, Warning, TEXT("%s: Server reported needed hash '%s' that is outside the set of expected hashes (%s) for put of %s from '%.*s'"),
+			UE_LOG(LogDerivedDataCache, Log, TEXT("%s: Server reported needed hash '%s' that is outside the set of expected hashes (%s) for put of %s from '%.*s'"),
 				*GetName(), *WriteToString<96>(NeededBlobHash), ExpectedHashes.ToString(), *WriteToString<96>(Key), Name.Len(), Name.GetData());
 		}
 	}
@@ -2921,7 +2921,7 @@ bool FHttpCacheStore::PutCacheValue(
 
 	if (!bPutCompletedSuccessfully)
 	{
-		UE_LOG(LogDerivedDataCache, Warning, TEXT("%s: Failed to put reference object for put of %s from '%.*s'"),
+		UE_LOG(LogDerivedDataCache, Log, TEXT("%s: Failed to put reference object for put of %s from '%.*s'"),
 			*GetName(), *WriteToString<96>(Key), Name.Len(), Name.GetData());
 		return false;
 	}
@@ -2943,14 +2943,14 @@ bool FHttpCacheStore::PutCacheValue(
 				bFirstHash = false;
 			}
 
-			UE_LOG(LogDerivedDataCache, Warning, TEXT("%s: Server reported unexpected needed hash quantity '%d' (%s) for put of %s from '%.*s'"),
+			UE_LOG(LogDerivedDataCache, Log, TEXT("%s: Server reported unexpected needed hash quantity '%d' (%s) for put of %s from '%.*s'"),
 				*GetName(), NeededBlobHashes.Num(), NeededHashString.ToString(), *WriteToString<96>(Key), Name.Len(), Name.GetData());
 			return false;
 		}
 
 		if (NeededBlobHashes[0] != Value.GetRawHash())
 		{
-			UE_LOG(LogDerivedDataCache, Warning, TEXT("%s: Server reported needed hash '%s' that is outside the set of expected hashes (%s) for put of %s from '%.*s'"),
+			UE_LOG(LogDerivedDataCache, Log, TEXT("%s: Server reported needed hash '%s' that is outside the set of expected hashes (%s) for put of %s from '%.*s'"),
 				*GetName(), *WriteToString<96>(NeededBlobHashes[0]), *WriteToString<96>(Value.GetRawHash()), *WriteToString<96>(Key), Name.Len(), Name.GetData());
 			return false;
 		}
@@ -3373,7 +3373,7 @@ TBitArray<> FHttpCacheStore::CachedDataProbablyExistsBatch(
 			}
 			else
 			{
-				UE_LOG(LogDerivedDataCache, Warning,
+				UE_LOG(LogDerivedDataCache, Log,
 					TEXT("%s: Cache exists returned invalid results."),
 					*GetName());
 				return TBitArray<>(false, Values.Num());
@@ -3453,7 +3453,7 @@ TArray<FValue> FHttpCacheStore::RefCachedDataProbablyExistsBatch(
 			FMemoryView ResponseView = MakeMemoryView(Request->GetResponseBuffer().GetData(), Request->GetResponseBuffer().Num());
 			if (ValidateCompactBinary(ResponseView, ECbValidateMode::Default) != ECbValidateError::None)
 			{
-				UE_LOG(LogDerivedDataCache, Warning,
+				UE_LOG(LogDerivedDataCache, Log,
 					TEXT("%s: Cache exists returned invalid results."),
 					*GetName());
 				TArray<FValue> RetVal;
@@ -3467,7 +3467,7 @@ TArray<FValue> FHttpCacheStore::RefCachedDataProbablyExistsBatch(
 
 			if (ResultsArrayView.Num() != ValueRefs.Num())
 			{
-				UE_LOG(LogDerivedDataCache, Warning,
+				UE_LOG(LogDerivedDataCache, Log,
 					TEXT("%s: Cache exists returned unexpected quantity of results (expected %d, got %d)."),
 					*GetName(), ValueRefs.Num(), ResultsArrayView.Num());
 				TArray<FValue> RetVal;
@@ -4015,7 +4015,7 @@ void FHttpCacheStore::GetValue(
 				if (!IsValueDataReady(Value, Request.Policy) && !EnumHasAnyFlags(Request.Policy, ECachePolicy::SkipData))
 				{
 					// With inline fetching, expect we will always have a value we can use.  Even SkipData/Exists can rely on the blob existing if the ref is reported to exist.
-					UE_LOG(LogDerivedDataCache, Warning, TEXT("%s: Cache miss due to inlining failure for %s from '%s'"),
+					UE_LOG(LogDerivedDataCache, Log, TEXT("%s: Cache miss due to inlining failure for %s from '%s'"),
 						*GetName(), *WriteToString<96>(Request.Key), *Request.Name);
 					COOK_STAT(UsageStats.GetStats.Accumulate(FCookStats::CallStats::EHitOrMiss::Miss, FCookStats::CallStats::EStatType::Counter, 1l, bIsInGameThread));
 					OnComplete(Request.MakeResponse(EStatus::Error));
