Fix issue (using server-side feature in Horde Storage) with inconsistency to DDC content from Horde Storage.
This would manifest as a cache entry being deemed to "exist", but later when you attempt to fetch it, it would not exist.

diff --git a/Engine/Source/Developer/DerivedDataCache/Private/HttpCacheStore.cpp b/Engine/Source/Developer/DerivedDataCache/Private/HttpCacheStore.cpp
--- a/Engine/Source/Developer/DerivedDataCache/Private/HttpCacheStore.cpp
+++ b/Engine/Source/Developer/DerivedDataCache/Private/HttpCacheStore.cpp
@@ -3426,6 +3426,7 @@ TArray<FValue> FHttpCacheStore::RefCachedDataProbablyExistsBatch(
 		Bucket.ToLowerInline();
 		RequestWriter.AddString("bucket"_ASV, Bucket);
 		RequestWriter.AddString("key"_ASV, LexToString(Key.Hash));
+		RequestWriter.AddBool("resolveAttachments"_ASV, true);
 		RequestWriter.EndObject();
 		++OpIndex;
 	}
