DDC: Optimized legacy cache key validation and shortening overhead

- Validate and shorten legacy keys lazily.
- Change to ValueOnly mode for PakFile and Zen.

#rnx

diff --git a/Engine/Source/Developer/DerivedDataCache/Private/DerivedDataBackendInterface.h b/Engine/Source/Developer/DerivedDataCache/Private/DerivedDataBackendInterface.h
--- a/Engine/Source/Developer/DerivedDataCache/Private/DerivedDataBackendInterface.h
+++ b/Engine/Source/Developer/DerivedDataCache/Private/DerivedDataBackendInterface.h
@@ -281,7 +281,7 @@ public:
 	 *
 	 * @param PakFilename Pak filename
 	 */
-	virtual FDerivedDataBackendInterface* MountPakFile(const TCHAR* PakFilename) = 0;
+	virtual ILegacyCacheStore* MountPakFile(const TCHAR* PakFilename) = 0;
 
 	/**
 	 * Unmounts a read-only pak file.
diff --git a/Engine/Source/Developer/DerivedDataCache/Private/DerivedDataBackends.cpp b/Engine/Source/Developer/DerivedDataCache/Private/DerivedDataBackends.cpp
--- a/Engine/Source/Developer/DerivedDataCache/Private/DerivedDataBackends.cpp
+++ b/Engine/Source/Developer/DerivedDataCache/Private/DerivedDataBackends.cpp
@@ -1164,7 +1164,7 @@ public:
 		return *StaticGraph;
 	}
 
-	virtual FDerivedDataBackendInterface* MountPakFile(const TCHAR* PakFilename) override
+	virtual ILegacyCacheStore* MountPakFile(const TCHAR* PakFilename) override
 	{
 		// Assumptions: there's at least one read-only pak backend in the hierarchy
 		// and its parent is a hierarchical backend.
@@ -1313,8 +1313,8 @@ void FDerivedDataBackendInterface::LegacyPut(
 
 		checkf(CompositeValue.GetSize() < MAX_int32,
 			TEXT("Value is 2 GiB or greater, which is not supported for put of '%s' from '%s'"),
-			*Request.Key.GetFullKey(), *Request.Name);
-	
+			*Request.Key.GetShortKey(), *Request.Name);
+
 		UE_CLOG(Request.Key.HasShortKey(), LogDerivedDataCache, VeryVerbose,
 			TEXT("ShortenKey %s -> %s"), *Request.Key.GetFullKey(), *Request.Key.GetShortKey());
 
diff --git a/Engine/Source/Developer/DerivedDataCache/Private/DerivedDataCache.cpp b/Engine/Source/Developer/DerivedDataCache/Private/DerivedDataCache.cpp
--- a/Engine/Source/Developer/DerivedDataCache/Private/DerivedDataCache.cpp
+++ b/Engine/Source/Developer/DerivedDataCache/Private/DerivedDataCache.cpp
@@ -560,7 +560,6 @@ public:
 	{
 		DDC_SCOPE_CYCLE_COUNTER(DDC_GetSynchronous_Data);
 		UE_LOG(LogDerivedDataCache, VeryVerbose, TEXT("GetSynchronous %s from '%.*s'"), CacheKey, DebugContext.Len(), DebugContext.GetData());
-		ValidateCacheKey(CacheKey);
 		FAsyncTask<FBuildAsyncWorker> PendingTask(Backend, nullptr, CacheKey, DebugContext, true);
 		AddToAsyncCompletionCounter(1);
 		PendingTask.StartSynchronousTask();
@@ -584,7 +583,6 @@ public:
 		FScopeLock ScopeLock(&SynchronizationObject);
 		const uint32 Handle = NextHandle();
 		UE_LOG(LogDerivedDataCache, VeryVerbose, TEXT("GetAsynchronous %s from '%.*s', Handle %d"), CacheKey, DebugContext.Len(), DebugContext.GetData(), Handle);
-		ValidateCacheKey(CacheKey);
 		FAsyncTask<FBuildAsyncWorker>* AsyncTask = new FAsyncTask<FBuildAsyncWorker>(Backend, nullptr, CacheKey, DebugContext, false);
 		check(!PendingTasks.Contains(Handle));
 		PendingTasks.Add(Handle, AsyncTask);
@@ -598,7 +596,6 @@ public:
 	{
 		DDC_SCOPE_CYCLE_COUNTER(DDC_Put);
 		UE_LOG(LogDerivedDataCache, VeryVerbose, TEXT("Put %s from '%.*s'"), CacheKey, DebugContext.Len(), DebugContext.GetData());
-		ValidateCacheKey(CacheKey);
 		STAT(double ThisTime = 0);
 		{
 			SCOPE_SECONDS_COUNTER(ThisTime);
@@ -616,7 +613,7 @@ public:
 
 	virtual void MarkTransient(const TCHAR* CacheKey) override
 	{
-		ValidateCacheKey(CacheKey);
+		DDC_SCOPE_CYCLE_COUNTER(DDC_MarkTransient);
 		FLegacyCacheDeleteRequest LegacyRequest;
 		LegacyRequest.Key = FLegacyCacheKey(CacheKey, Backend->GetMaxKeyLength());
 		LegacyRequest.Name = LegacyRequest.Key.GetFullKey();
@@ -629,7 +626,6 @@ public:
 	virtual bool CachedDataProbablyExists(const TCHAR* CacheKey) override
 	{
 		DDC_SCOPE_CYCLE_COUNTER(DDC_CachedDataProbablyExists);
-		ValidateCacheKey(CacheKey);
 		bool bResult;
 		INC_DWORD_STAT(STAT_DDC_NumExist);
 		STAT(double ThisTime = 0);
@@ -830,12 +826,6 @@ private:
 		return Result;
 	}
 
-	static void ValidateCacheKey(const TCHAR* CacheKey)
-	{
-		checkf(Algo::AllOf(FStringView(CacheKey), IsValidCacheChar),
-			TEXT("Invalid characters in cache key %s. Use SanitizeCacheKey or BuildCacheKey to create valid keys."), CacheKey);
-	}
-
 	FDerivedDataBackend*		Backend;
 	/** Counter used to produce unique handles **/
 	FThreadSafeCounter			CurrentHandle;
@@ -855,6 +845,7 @@ public:
 		IRequestOwner& Owner,
 		FOnCachePutComplete&& OnComplete) final
 	{
+		DDC_SCOPE_CYCLE_COUNTER(DDC_Put);
 		return Backend->GetRoot().Put(Requests, Owner, OnComplete ? MoveTemp(OnComplete) : [](auto&&){});
 	}
 
@@ -863,6 +854,7 @@ public:
 		IRequestOwner& Owner,
 		FOnCacheGetComplete&& OnComplete) final
 	{
+		DDC_SCOPE_CYCLE_COUNTER(DDC_Get);
 		return Backend->GetRoot().Get(Requests, Owner, OnComplete ? MoveTemp(OnComplete) : [](auto&&){});
 	}
 
@@ -871,6 +863,7 @@ public:
 		IRequestOwner& Owner,
 		FOnCachePutValueComplete&& OnComplete) final
 	{
+		DDC_SCOPE_CYCLE_COUNTER(DDC_PutValue);
 		return Backend->GetRoot().PutValue(Requests, Owner, OnComplete ? MoveTemp(OnComplete) : [](auto&&){});
 	}
 
@@ -879,6 +872,7 @@ public:
 		IRequestOwner& Owner,
 		FOnCacheGetValueComplete&& OnComplete) final
 	{
+		DDC_SCOPE_CYCLE_COUNTER(DDC_GetValue);
 		return Backend->GetRoot().GetValue(Requests, Owner, OnComplete ? MoveTemp(OnComplete) : [](auto&&){});
 	}
 
@@ -887,6 +881,7 @@ public:
 		IRequestOwner& Owner,
 		FOnCacheGetChunkComplete&& OnComplete) final
 	{
+		DDC_SCOPE_CYCLE_COUNTER(DDC_GetChunks);
 		return Backend->GetRoot().GetChunks(Requests, Owner, OnComplete ? MoveTemp(OnComplete) : [](auto&&){});
 	}
 
diff --git a/Engine/Source/Developer/DerivedDataCache/Private/DerivedDataLegacyCacheStore.cpp b/Engine/Source/Developer/DerivedDataCache/Private/DerivedDataLegacyCacheStore.cpp
--- a/Engine/Source/Developer/DerivedDataCache/Private/DerivedDataLegacyCacheStore.cpp
+++ b/Engine/Source/Developer/DerivedDataCache/Private/DerivedDataLegacyCacheStore.cpp
@@ -2,12 +2,15 @@
 
 #include "DerivedDataLegacyCacheStore.h"
 
+#include "Algo/AllOf.h"
 #include "Containers/StringConv.h"
 #include "DerivedDataBackendInterface.h"
+#include "DerivedDataCacheInterface.h"
 #include "Misc/Crc.h"
 #include "Misc/ScopeRWLock.h"
 #include "Misc/SecureHash.h"
 #include "Misc/StringBuilder.h"
+#include "ProfilingDebugging/CpuProfilerTrace.h"
 #include "String/BytesToHex.h"
 #include "String/Find.h"
 
@@ -82,51 +85,82 @@ void ILegacyCacheStore::LegacyDelete(
 	CompleteWithStatus(Requests, OnComplete, EStatus::Error);
 }
 
-FLegacyCacheKey::FLegacyCacheKey(const FStringView InFullKey, const int32 MaxKeyLength)
+Private::FLegacyCacheKeyShared::FLegacyCacheKeyShared(const FStringView InFullKey, const int32 InMaxKeyLength)
 	: FullKey(InFullKey)
+	, MaxKeyLength(InMaxKeyLength)
 {
-	if (InFullKey.Len() > MaxKeyLength)
+}
+
+const FSharedString& Private::FLegacyCacheKeyShared::GetShortKey()
+{
+	if (FReadScopeLock ReadLock(Lock); !ShortKey.IsEmpty())
 	{
-		const auto FullKeyUCS2 = StringCast<UCS2CHAR>(InFullKey.GetData(), InFullKey.Len());
-		const int32 FullKeyLength = FullKeyUCS2.Length();
-		const uint32 CRCofPayload(FCrc::MemCrc32(FullKeyUCS2.Get(), FullKeyLength * sizeof(UCS2CHAR)));
+		return ShortKey;
+	}
+
+	TRACE_CPUPROFILER_EVENT_SCOPE(FLegacyCacheKey::GetShortKey);
 
-		uint8 Hash[FSHA1::DigestSize];
-		FSHA1 HashState;
-		HashState.Update((const uint8*)&FullKeyLength, sizeof(int32));
-		HashState.Update((const uint8*)&CRCofPayload, sizeof(uint32));
-		HashState.Update((const uint8*)FullKeyUCS2.Get(), FullKeyLength * sizeof(UCS2CHAR));
-		HashState.Final();
-		HashState.GetHash(Hash);
+	FWriteScopeLock WriteLock(Lock);
+	if (!ShortKey.IsEmpty())
+	{
+		return ShortKey;
+	}
 
-		TStringBuilder<128> ShortKeyBuilder;
-		ShortKeyBuilder << InFullKey.Left(MaxKeyLength - FSHA1::DigestSize * 2 - 2) << TEXT("__");
-		String::BytesToHex(Hash, ShortKeyBuilder);
+	const FStringView LocalFullKey = FullKey;
+	checkf(Algo::AllOf(LocalFullKey, FDerivedDataCacheInterface::IsValidCacheChar),
+		TEXT("Invalid characters in cache key %s. Use SanitizeCacheKey or BuildCacheKey to create valid keys."), *FullKey);
 
-		check(ShortKeyBuilder.Len() == MaxKeyLength && ShortKeyBuilder.Len() > 0);
-		ShortKey = ShortKeyBuilder;
+	if (LocalFullKey.Len() <= MaxKeyLength)
+	{
+		ShortKey = FullKey;
+		return ShortKey;
 	}
 
-	TStringBuilder<64> Bucket;
-	Bucket << TEXT("Legacy");
-	if (const int32 BucketEnd = String::FindFirstChar(InFullKey, TEXT('_')); BucketEnd != INDEX_NONE)
+	const auto FullKeyUCS2 = StringCast<UCS2CHAR>(LocalFullKey.GetData(), LocalFullKey.Len());
+	const int32 FullKeyLength = FullKeyUCS2.Length();
+	const uint32 CRCofPayload(FCrc::MemCrc32(FullKeyUCS2.Get(), FullKeyLength * sizeof(UCS2CHAR)));
+
+	uint8 Hash[FSHA1::DigestSize];
+	FSHA1 HashState;
+	HashState.Update((const uint8*)&FullKeyLength, sizeof(int32));
+	HashState.Update((const uint8*)&CRCofPayload, sizeof(uint32));
+	HashState.Update((const uint8*)FullKeyUCS2.Get(), FullKeyLength * sizeof(UCS2CHAR));
+	HashState.Final();
+	HashState.GetHash(Hash);
+
+	TStringBuilder<256> ShortKeyBuilder;
+	ShortKeyBuilder << LocalFullKey.Left(MaxKeyLength - FSHA1::DigestSize * 2 - 2) << TEXT("__");
+	String::BytesToHex(Hash, ShortKeyBuilder);
+
+	check(ShortKeyBuilder.Len() == MaxKeyLength && ShortKeyBuilder.Len() > 0);
+	ShortKey = ShortKeyBuilder;
+	return ShortKey;
+}
+
+FLegacyCacheKey::FLegacyCacheKey(const FStringView FullKey, const int32 MaxKeyLength)
+	: Shared(new Private::FLegacyCacheKeyShared(FullKey, MaxKeyLength))
+{
+	FTCHARToUTF8 FullKeyUtf8(FullKey);
+	TUtf8StringBuilder<64> Bucket;
+	Bucket << "Legacy";
+	if (const int32 BucketEnd = String::FindFirstChar(FullKeyUtf8, '_'); BucketEnd != INDEX_NONE)
 	{
-		Bucket << InFullKey.Left(BucketEnd);
+		Bucket << FUtf8StringView(FullKeyUtf8).Left(BucketEnd);
 	}
 	Key.Bucket = FCacheBucket(Bucket);
-	Key.Hash = FIoHash::HashBuffer(MakeMemoryView(FTCHARToUTF8(InFullKey)));
+	Key.Hash = FIoHash::HashBuffer(MakeMemoryView(FullKeyUtf8));
 }
 
 bool FLegacyCacheKey::ReadValueTrailer(FCompositeBuffer& Value) const
 {
-	if (!ShortKey.IsEmpty())
+	if (HasShortKey())
 	{
-		TUtf8StringBuilder<256> FullKeyUtf8;
-		FullKeyUtf8 << FullKey;
-		const FMemoryView CompareKey(FullKeyUtf8.ToString(), (FullKeyUtf8.Len() + 1) * sizeof(UTF8CHAR));
+		FUtf8StringBuilderBase FullKey;
+		FullKey << GetFullKey();
+		const FMemoryView CompareKey(FullKey.ToString(), (FullKey.Len() + 1) * sizeof(UTF8CHAR));
 		if (Value.GetSize() < CompareKey.GetSize())
 		{
-			UE_LOG(LogDerivedDataCache, Display, TEXT("FLegacyCacheKey: Hash collision or short value for key %s."), *FullKey);
+			UE_LOG(LogDerivedDataCache, Display, TEXT("FLegacyCacheKey: Hash collision or short value for key %s."), *GetFullKey());
 			return false;
 		}
 		FUniqueBuffer CopyBuffer;
@@ -134,7 +168,7 @@ bool FLegacyCacheKey::ReadValueTrailer(FCompositeBuffer& Value) const
 		const FMemoryView ValueKey = Value.ViewOrCopyRange(KeyOffset, CompareKey.GetSize(), CopyBuffer);
 		if (!CompareKey.EqualBytes(ValueKey))
 		{
-			UE_LOG(LogDerivedDataCache, Display, TEXT("FLegacyCacheKey: Hash collision for key %s."), *FullKey);
+			UE_LOG(LogDerivedDataCache, Display, TEXT("FLegacyCacheKey: Hash collision for key %s."), *GetFullKey());
 			return false;
 		}
 		Value = Value.Mid(0, KeyOffset);
@@ -144,11 +178,11 @@ bool FLegacyCacheKey::ReadValueTrailer(FCompositeBuffer& Value) const
 
 void FLegacyCacheKey::WriteValueTrailer(FCompositeBuffer& Value) const
 {
-	if (!ShortKey.IsEmpty())
+	if (HasShortKey())
 	{
-		TUtf8StringBuilder<256> FullKeyUtf8;
-		FullKeyUtf8 << FullKey;
-		Value = FCompositeBuffer(Value, FSharedBuffer::Clone(FullKeyUtf8.ToString(), (FullKeyUtf8.Len() + 1) * sizeof(UTF8CHAR)));
+		FUtf8StringBuilderBase FullKey;
+		FullKey << GetFullKey();
+		Value = FCompositeBuffer(Value, FSharedBuffer::Clone(FullKey.ToString(), (FullKey.Len() + 1) * sizeof(UTF8CHAR)));
 	}
 }
 
diff --git a/Engine/Source/Developer/DerivedDataCache/Private/DerivedDataLegacyCacheStore.h b/Engine/Source/Developer/DerivedDataCache/Private/DerivedDataLegacyCacheStore.h
--- a/Engine/Source/Developer/DerivedDataCache/Private/DerivedDataLegacyCacheStore.h
+++ b/Engine/Source/Developer/DerivedDataCache/Private/DerivedDataLegacyCacheStore.h
@@ -18,6 +18,7 @@ namespace UE::DerivedData { struct FLegacyCacheGetRequest; }
 namespace UE::DerivedData { struct FLegacyCacheGetResponse; }
 namespace UE::DerivedData { struct FLegacyCachePutRequest; }
 namespace UE::DerivedData { struct FLegacyCachePutResponse; }
+namespace UE::DerivedData::Private { class FLegacyCacheKeyShared; }
 namespace UE::DerivedData::Private { class FLegacyCacheValueShared; }
 
 DECLARE_LOG_CATEGORY_EXTERN(LogDerivedDataCache, Log, All);
@@ -52,24 +53,56 @@ public:
 	virtual bool LegacyDebugOptions(FBackendDebugOptions& Options) = 0;
 };
 
+class Private::FLegacyCacheKeyShared final
+{
+public:
+	FLegacyCacheKeyShared(FStringView FullKey, int32 MaxKeyLength);
+
+	FLegacyCacheKeyShared(const FLegacyCacheKeyShared&) = delete;
+	FLegacyCacheKeyShared& operator=(const FLegacyCacheKeyShared&) = delete;
+
+	inline bool HasShortKey() const { return FullKey.Len() > MaxKeyLength; }
+	inline const FSharedString& GetFullKey() const { return FullKey; }
+	const FSharedString& GetShortKey();
+
+	inline void AddRef()
+	{
+		ReferenceCount.fetch_add(1, std::memory_order_relaxed);
+	}
+
+	inline void Release()
+	{
+		if (ReferenceCount.fetch_sub(1, std::memory_order_acq_rel) == 1)
+		{
+			delete this;
+		}
+	}
+
+private:
+	FSharedString FullKey;
+	FSharedString ShortKey;
+	int32 MaxKeyLength;
+	std::atomic<uint32> ReferenceCount{0};
+	FRWLock Lock;
+};
+
 class FLegacyCacheKey
 {
 public:
 	FLegacyCacheKey() = default;
 	FLegacyCacheKey(FStringView FullKey, int32 MaxKeyLength);
 
-	const FCacheKey& GetKey() const { return Key; }
-	const FSharedString& GetFullKey() const { return FullKey; }
-	const FSharedString& GetShortKey() const { return ShortKey.IsEmpty() ? FullKey : ShortKey; }
-	bool HasShortKey() const { return !ShortKey.IsEmpty(); }
+	[[nodiscard]] inline const FCacheKey& GetKey() const { return Key; }
+	[[nodiscard]] inline const FSharedString& GetFullKey() const { return Shared ? Shared->GetFullKey() : FSharedString::Empty; }
+	[[nodiscard]] inline const FSharedString& GetShortKey() const { return Shared ? Shared->GetShortKey() : FSharedString::Empty; }
+	[[nodiscard]] inline bool HasShortKey() const { return Shared ? Shared->HasShortKey() : false; }
 
-	bool ReadValueTrailer(FCompositeBuffer& Value) const;
+	[[nodiscard]] bool ReadValueTrailer(FCompositeBuffer& Value) const;
 	void WriteValueTrailer(FCompositeBuffer& Value) const;
 
 private:
 	FCacheKey Key;
-	FSharedString FullKey;
-	FSharedString ShortKey;
+	TRefCountPtr<Private::FLegacyCacheKeyShared> Shared;
 };
 
 class Private::FLegacyCacheValueShared final
@@ -78,6 +111,9 @@ public:
 	explicit FLegacyCacheValueShared(const FValue& Value);
 	explicit FLegacyCacheValueShared(const FCompositeBuffer& RawData);
 
+	FLegacyCacheValueShared(const FLegacyCacheValueShared&) = delete;
+	FLegacyCacheValueShared& operator=(const FLegacyCacheValueShared&) = delete;
+
 	inline bool HasData() const { return Value.HasData() || RawData; }
 	const FValue& GetValue();
 	const FCompositeBuffer& GetRawData();
@@ -116,9 +152,9 @@ public:
 	[[nodiscard]] inline explicit operator bool() const { return !IsNull(); }
 	[[nodiscard]] inline bool IsNull() const { return !Shared; }
 
-	[[nodiscard]] bool HasData() const { return Shared ? Shared->HasData() : false; }
-	[[nodiscard]] const FValue& GetValue() const { return Shared ? Shared->GetValue() : FValue::Null; }
-	[[nodiscard]] const FCompositeBuffer& GetRawData() const { return Shared ? Shared->GetRawData() : FCompositeBuffer::Null; }
+	[[nodiscard]] inline bool HasData() const { return Shared ? Shared->HasData() : false; }
+	[[nodiscard]] inline const FValue& GetValue() const { return Shared ? Shared->GetValue() : FValue::Null; }
+	[[nodiscard]] inline const FCompositeBuffer& GetRawData() const { return Shared ? Shared->GetRawData() : FCompositeBuffer::Null; }
 	[[nodiscard]] inline FIoHash GetRawHash() const { return Shared ? Shared->GetRawHash() : FIoHash(); }
 	[[nodiscard]] inline uint64 GetRawSize() const { return Shared ? Shared->GetRawSize() : 0; }
 
diff --git a/Engine/Source/Developer/DerivedDataCache/Private/PakFileCacheStore.cpp b/Engine/Source/Developer/DerivedDataCache/Private/PakFileCacheStore.cpp
--- a/Engine/Source/Developer/DerivedDataCache/Private/PakFileCacheStore.cpp
+++ b/Engine/Source/Developer/DerivedDataCache/Private/PakFileCacheStore.cpp
@@ -121,7 +121,7 @@ public:
 
 	bool ApplyDebugOptions(FBackendDebugOptions& InOptions) final { return false; }
 
-	EBackendLegacyMode GetLegacyMode() const final { return EBackendLegacyMode::ValueWithLegacyFallback; }
+	EBackendLegacyMode GetLegacyMode() const final { return EBackendLegacyMode::ValueOnly; }
 
 	void Put(
 		TConstArrayView<FCachePutRequest> Requests,
diff --git a/Engine/Source/Developer/DerivedDataCache/Private/ZenCacheStore.cpp b/Engine/Source/Developer/DerivedDataCache/Private/ZenCacheStore.cpp
--- a/Engine/Source/Developer/DerivedDataCache/Private/ZenCacheStore.cpp
+++ b/Engine/Source/Developer/DerivedDataCache/Private/ZenCacheStore.cpp
@@ -113,7 +113,7 @@ public:
 	virtual bool WouldCache(const TCHAR* CacheKey, TArrayView<const uint8> InData) override;
 	virtual bool ApplyDebugOptions(FBackendDebugOptions& InOptions) override;
 
-	virtual EBackendLegacyMode GetLegacyMode() const override { return EBackendLegacyMode::ValueWithLegacyFallback; }
+	virtual EBackendLegacyMode GetLegacyMode() const override { return EBackendLegacyMode::ValueOnly; }
 
 	// ICacheStore
 
diff --git a/Engine/Source/Developer/DerivedDataCache/Public/DerivedDataSharedString.h b/Engine/Source/Developer/DerivedDataCache/Public/DerivedDataSharedString.h
--- a/Engine/Source/Developer/DerivedDataCache/Public/DerivedDataSharedString.h
+++ b/Engine/Source/Developer/DerivedDataCache/Public/DerivedDataSharedString.h
@@ -42,6 +42,8 @@ public:
 	[[nodiscard]] inline int32 Len() const { return Chars ? reinterpret_cast<const int32*>(Chars)[-1] : 0; }
 	[[nodiscard]] inline const CharType* operator*() const { return Chars ? Chars : &NullChar; }
 
+	static const TSharedString Empty;
+
 private:
 	static inline void AddRef(CharType* Chars);
 	static inline void Release(CharType* Chars);
@@ -51,6 +53,9 @@ private:
 	static constexpr inline CharType NullChar{};
 };
 
+template <typename CharType>
+inline const TSharedString<CharType> TSharedString<CharType>::Empty;
+
 template <typename CharType>
 constexpr inline const CharType* GetData(const TSharedString<CharType> String)
 {
