DDC: Fixed handling of metadata in the memory cache

#rnx

diff --git a/Engine/Source/Developer/DerivedDataCache/Private/MemoryCacheStore.cpp b/Engine/Source/Developer/DerivedDataCache/Private/MemoryCacheStore.cpp
--- a/Engine/Source/Developer/DerivedDataCache/Private/MemoryCacheStore.cpp
+++ b/Engine/Source/Developer/DerivedDataCache/Private/MemoryCacheStore.cpp
@@ -256,6 +256,10 @@ void FMemoryCacheStore::Put(
 		int64 RequiredSize = 0;
 		if (bHasExisting && !bReplaceExisting)
 		{
+			if (!Components.Meta && Record.GetMeta())
+			{
+				RequiredSize += Record.GetMeta().GetSize();
+			}
 			for (const FValueWithId& Value : Components.Values)
 			{
 				if (!Value.HasData())
@@ -326,6 +330,10 @@ void FMemoryCacheStore::Put(
 		}
 		else if (!bReplaceExisting)
 		{
+			if (!Components.Meta)
+			{
+				Components.Meta = Record.GetMeta();
+			}
 			for (FValueWithId& Value : Components.Values)
 			{
 				if (!Value.HasData())
@@ -398,6 +406,12 @@ void FMemoryCacheStore::Get(
 		if (Status == EStatus::Ok)
 		{
 			const ECachePolicy RecordPolicy = Policy.GetRecordPolicy();
+
+			if (!EnumHasAnyFlags(RecordPolicy, ECachePolicy::SkipMeta))
+			{
+				Builder.SetMeta(CopyTemp(Components.Meta));
+			}
+
 			for (const FValueWithId& Value : Components.Values)
 			{
 				const ECachePolicy ValuePolicy = Policy.GetValuePolicy(Value.GetId());
