DDC: Fixed timing scopes for FileSystem, PakFile, S3

Edited the timing scopes to exclude invocation of the completion callback.

#rnx

diff --git a/Engine/Source/Developer/DerivedDataCache/Private/FileSystemCacheStore.cpp b/Engine/Source/Developer/DerivedDataCache/Private/FileSystemCacheStore.cpp
--- a/Engine/Source/Developer/DerivedDataCache/Private/FileSystemCacheStore.cpp
+++ b/Engine/Source/Developer/DerivedDataCache/Private/FileSystemCacheStore.cpp
@@ -1556,24 +1556,27 @@ void FFileSystemCacheStore::Put(
 {
 	for (const FCachePutRequest& Request : Requests)
 	{
-		const FCacheRecord& Record = Request.Record;
-		TRACE_CPUPROFILER_EVENT_SCOPE(FileSystemDDC_Put);
-		TRACE_COUNTER_INCREMENT(FileSystemDDC_Put);
-		COOK_STAT(auto Timer = UsageStats.TimePut());
-		uint64 WriteSize = 0;
-		if (PutCacheRecord(Request.Name, Record, Request.Policy, WriteSize))
-		{
-			UE_LOG(LogDerivedDataCache, Verbose, TEXT("%s: Cache put complete for %s from '%s'"),
-				*CachePath, *WriteToString<96>(Record.GetKey()), *Request.Name);
-			TRACE_COUNTER_INCREMENT(FileSystemDDC_PutHit);
-			TRACE_COUNTER_ADD(FileSystemDDC_BytesWritten, WriteSize);
-			COOK_STAT(if (WriteSize) { Timer.AddHit(WriteSize); });
-			OnComplete(Request.MakeResponse(EStatus::Ok));
-		}
-		else
+		bool bOk;
 		{
-			OnComplete(Request.MakeResponse(EStatus::Error));
+			const FCacheRecord& Record = Request.Record;
+			TRACE_CPUPROFILER_EVENT_SCOPE(FileSystemDDC_Put);
+			TRACE_COUNTER_INCREMENT(FileSystemDDC_Put);
+			COOK_STAT(auto Timer = UsageStats.TimePut());
+			uint64 WriteSize = 0;
+			bOk = PutCacheRecord(Request.Name, Record, Request.Policy, WriteSize);
+			if (bOk)
+			{
+				UE_LOG(LogDerivedDataCache, Verbose, TEXT("%s: Cache put complete for %s from '%s'"),
+					*CachePath, *WriteToString<96>(Record.GetKey()), *Request.Name);
+				TRACE_COUNTER_INCREMENT(FileSystemDDC_PutHit);
+				TRACE_COUNTER_ADD(FileSystemDDC_BytesWritten, WriteSize);
+				if (WriteSize)
+				{
+					COOK_STAT(Timer.AddHit(WriteSize));
+				}
+			}
 		}
+		OnComplete(Request.MakeResponse(bOk ? EStatus::Ok : EStatus::Error));
 	}
 }
 
@@ -1584,23 +1587,26 @@ void FFileSystemCacheStore::Get(
 {
 	for (const FCacheGetRequest& Request : Requests)
 	{
-		TRACE_CPUPROFILER_EVENT_SCOPE(FileSystemDDC_Get);
-		TRACE_COUNTER_INCREMENT(FileSystemDDC_Get);
-		COOK_STAT(auto Timer = UsageStats.TimeGet());
-		EStatus Status = EStatus::Ok;
-		if (FOptionalCacheRecord Record = GetCacheRecord(Request.Name, Request.Key, Request.Policy, Status))
+		EStatus Status = EStatus::Error;
+		FOptionalCacheRecord Record;
 		{
-			UE_LOG(LogDerivedDataCache, Verbose, TEXT("%s: Cache hit for %s from '%s'"),
-				*CachePath, *WriteToString<96>(Request.Key), *Request.Name);
-			TRACE_COUNTER_INCREMENT(FileSystemDDC_GetHit);
-			TRACE_COUNTER_ADD(FileSystemDDC_BytesRead, Private::GetCacheRecordCompressedSize(Record.Get()));
-			COOK_STAT(Timer.AddHit(Private::GetCacheRecordCompressedSize(Record.Get())));
-			OnComplete({Request.Name, MoveTemp(Record).Get(), Request.UserData, Status});
-		}
-		else
-		{
-			OnComplete(Request.MakeResponse(Status));
+			TRACE_CPUPROFILER_EVENT_SCOPE(FileSystemDDC_Get);
+			TRACE_COUNTER_INCREMENT(FileSystemDDC_Get);
+			COOK_STAT(auto Timer = UsageStats.TimeGet());
+			if ((Record = GetCacheRecord(Request.Name, Request.Key, Request.Policy, Status)))
+			{
+				UE_LOG(LogDerivedDataCache, Verbose, TEXT("%s: Cache hit for %s from '%s'"),
+					*CachePath, *WriteToString<96>(Request.Key), *Request.Name);
+				TRACE_COUNTER_INCREMENT(FileSystemDDC_GetHit);
+				TRACE_COUNTER_ADD(FileSystemDDC_BytesRead, Private::GetCacheRecordCompressedSize(Record.Get()));
+				COOK_STAT(Timer.AddHit(Private::GetCacheRecordCompressedSize(Record.Get())));
+			}
+			else
+			{
+				Record = FCacheRecordBuilder(Request.Key).Build();
+			}
 		}
+		OnComplete({Request.Name, MoveTemp(Record).Get(), Request.UserData, Status});
 	}
 }
 
@@ -1611,23 +1617,26 @@ void FFileSystemCacheStore::PutValue(
 {
 	for (const FCachePutValueRequest& Request : Requests)
 	{
-		TRACE_CPUPROFILER_EVENT_SCOPE(FileSystemDDC_Put);
-		TRACE_COUNTER_INCREMENT(FileSystemDDC_Put);
-		COOK_STAT(auto Timer = UsageStats.TimePut());
-		uint64 WriteSize = 0;
-		if (PutCacheValue(Request.Name, Request.Key, Request.Value, Request.Policy, WriteSize))
-		{
-			UE_LOG(LogDerivedDataCache, Verbose, TEXT("%s: Cache put complete for %s from '%s'"),
-				*CachePath, *WriteToString<96>(Request.Key), *Request.Name);
-			TRACE_COUNTER_INCREMENT(FileSystemDDC_PutHit);
-			TRACE_COUNTER_ADD(FileSystemDDC_BytesWritten, WriteSize);
-			COOK_STAT(if (WriteSize) { Timer.AddHit(WriteSize); });
-			OnComplete(Request.MakeResponse(EStatus::Ok));
-		}
-		else
+		bool bOk;
 		{
-			OnComplete(Request.MakeResponse(EStatus::Ok));
+			TRACE_CPUPROFILER_EVENT_SCOPE(FileSystemDDC_PutValue);
+			TRACE_COUNTER_INCREMENT(FileSystemDDC_Put);
+			COOK_STAT(auto Timer = UsageStats.TimePut());
+			uint64 WriteSize = 0;
+			bOk = PutCacheValue(Request.Name, Request.Key, Request.Value, Request.Policy, WriteSize);
+			if (bOk)
+			{
+				UE_LOG(LogDerivedDataCache, Verbose, TEXT("%s: Cache put complete for %s from '%s'"),
+					*CachePath, *WriteToString<96>(Request.Key), *Request.Name);
+				TRACE_COUNTER_INCREMENT(FileSystemDDC_PutHit);
+				TRACE_COUNTER_ADD(FileSystemDDC_BytesWritten, WriteSize);
+				if (WriteSize)
+				{
+					COOK_STAT(Timer.AddHit(WriteSize));
+				}
+			}
 		}
+		OnComplete(Request.MakeResponse(bOk ? EStatus::Ok : EStatus::Error));
 	}
 }
 
@@ -1638,23 +1647,23 @@ void FFileSystemCacheStore::GetValue(
 {
 	for (const FCacheGetValueRequest& Request : Requests)
 	{
-		TRACE_CPUPROFILER_EVENT_SCOPE(FileSystemDDC_Get);
-		TRACE_COUNTER_INCREMENT(FileSystemDDC_Get);
-		COOK_STAT(auto Timer = UsageStats.TimeGet());
+		bool bOk;
 		FValue Value;
-		if (GetCacheValue(Request.Name, Request.Key, Request.Policy, Value))
-		{
-			UE_LOG(LogDerivedDataCache, Verbose, TEXT("%s: Cache hit for %s from '%s'"),
-				*CachePath, *WriteToString<96>(Request.Key), *Request.Name);
-			TRACE_COUNTER_INCREMENT(FileSystemDDC_GetHit);
-			TRACE_COUNTER_ADD(FileSystemDDC_BytesRead, Value.GetData().GetCompressedSize());
-			COOK_STAT(Timer.AddHit(Value.GetData().GetCompressedSize()));
-			OnComplete({Request.Name, Request.Key, Value, Request.UserData, EStatus::Ok});
-		}
-		else
 		{
-			OnComplete(Request.MakeResponse(EStatus::Error));
+			TRACE_CPUPROFILER_EVENT_SCOPE(FileSystemDDC_GetValue);
+			TRACE_COUNTER_INCREMENT(FileSystemDDC_Get);
+			COOK_STAT(auto Timer = UsageStats.TimeGet());
+			bOk = GetCacheValue(Request.Name, Request.Key, Request.Policy, Value);
+			if (bOk)
+			{
+				UE_LOG(LogDerivedDataCache, Verbose, TEXT("%s: Cache hit for %s from '%s'"),
+					*CachePath, *WriteToString<96>(Request.Key), *Request.Name);
+				TRACE_COUNTER_INCREMENT(FileSystemDDC_GetHit);
+				TRACE_COUNTER_ADD(FileSystemDDC_BytesRead, Value.GetData().GetCompressedSize());
+				COOK_STAT(Timer.AddHit(Value.GetData().GetCompressedSize()));
+			}
 		}
+		OnComplete({Request.Name, Request.Key, Value, Request.UserData, bOk ? EStatus::Ok : EStatus::Error});
 	}
 }
 
@@ -1675,70 +1684,71 @@ void FFileSystemCacheStore::GetChunks(
 	FOptionalCacheRecord Record;
 	for (const FCacheGetChunkRequest& Request : SortedRequests)
 	{
-		const bool bExistsOnly = EnumHasAnyFlags(Request.Policy, ECachePolicy::SkipData);
-		TRACE_CPUPROFILER_EVENT_SCOPE(FileSystemDDC_Get);
-		TRACE_COUNTER_INCREMENT(FileSystemDDC_Get);
-		COOK_STAT(auto Timer = bExistsOnly ? UsageStats.TimeProbablyExists() : UsageStats.TimeGet());
-		if (!(bHasValue && ValueKey == Request.Key && ValueId == Request.Id) || ValueReader.HasSource() < !bExistsOnly)
-		{
-			ValueReader.ResetSource();
-			ValueAr.Reset();
-			ValueKey = {};
-			ValueId.Reset();
-			Value.Reset();
-			bHasValue = false;
-			if (Request.Id.IsValid())
+		EStatus Status = EStatus::Error;
+		FSharedBuffer Buffer;
+		uint64 RawSize = 0;
+		{
+			TRACE_CPUPROFILER_EVENT_SCOPE(FileSystemDDC_GetChunks);
+			TRACE_COUNTER_INCREMENT(FileSystemDDC_Get);
+			const bool bExistsOnly = EnumHasAnyFlags(Request.Policy, ECachePolicy::SkipData);
+			COOK_STAT(auto Timer = bExistsOnly ? UsageStats.TimeProbablyExists() : UsageStats.TimeGet());
+			if (!(bHasValue && ValueKey == Request.Key && ValueId == Request.Id) || ValueReader.HasSource() < !bExistsOnly)
 			{
-				if (!(Record && Record.Get().GetKey() == Request.Key))
+				ValueReader.ResetSource();
+				ValueAr.Reset();
+				ValueKey = {};
+				ValueId.Reset();
+				Value.Reset();
+				bHasValue = false;
+				if (Request.Id.IsValid())
 				{
-					FCacheRecordPolicyBuilder PolicyBuilder(ECachePolicy::None);
-					PolicyBuilder.AddValuePolicy(Request.Id, Request.Policy);
-					Record.Reset();
-					Record = GetCacheRecordOnly(Request.Name, Request.Key, PolicyBuilder.Build());
+					if (!(Record && Record.Get().GetKey() == Request.Key))
+					{
+						FCacheRecordPolicyBuilder PolicyBuilder(ECachePolicy::None);
+						PolicyBuilder.AddValuePolicy(Request.Id, Request.Policy);
+						Record.Reset();
+						Record = GetCacheRecordOnly(Request.Name, Request.Key, PolicyBuilder.Build());
+					}
+					if (Record)
+					{
+						if (const FValueWithId& ValueWithId = Record.Get().GetValue(Request.Id))
+						{
+							bHasValue = true;
+							Value = ValueWithId;
+							ValueId = Request.Id;
+							ValueKey = Request.Key;
+							GetCacheContent(Request.Name, Request.Key, ValueId, Value, Request.Policy, ValueReader, ValueAr);
+						}
+					}
 				}
-				if (Record)
+				else
 				{
-					if (const FValueWithId& ValueWithId = Record.Get().GetValue(Request.Id))
+					ValueKey = Request.Key;
+					bHasValue = GetCacheValueOnly(Request.Name, Request.Key, Request.Policy, Value);
+					if (bHasValue)
 					{
-						bHasValue = true;
-						Value = ValueWithId;
-						ValueId = Request.Id;
-						ValueKey = Request.Key;
-						GetCacheContent(Request.Name, Request.Key, ValueId, Value, Request.Policy, ValueReader, ValueAr);
+						GetCacheContent(Request.Name, Request.Key, Request.Id, Value, Request.Policy, ValueReader, ValueAr);
 					}
 				}
 			}
-			else
+			if (bHasValue)
 			{
-				ValueKey = Request.Key;
-				bHasValue = GetCacheValueOnly(Request.Name, Request.Key, Request.Policy, Value);
-				if (bHasValue)
+				const uint64 RawOffset = FMath::Min(Value.GetRawSize(), Request.RawOffset);
+				RawSize = FMath::Min(Value.GetRawSize() - RawOffset, Request.RawSize);
+				UE_LOG(LogDerivedDataCache, Verbose, TEXT("%s: Cache hit for %s from '%s'"),
+					*CachePath, *WriteToString<96>(Request.Key, '/', Request.Id), *Request.Name);
+				TRACE_COUNTER_INCREMENT(FileSystemDDC_GetHit);
+				TRACE_COUNTER_ADD(FileSystemDDC_BytesRead, !bExistsOnly ? RawSize : 0);
+				COOK_STAT(Timer.AddHit(!bExistsOnly ? RawSize : 0));
+				if (!bExistsOnly)
 				{
-					GetCacheContent(Request.Name, Request.Key, Request.Id, Value, Request.Policy, ValueReader, ValueAr);
+					Buffer = ValueReader.Decompress(RawOffset, RawSize);
 				}
+				Status = bExistsOnly || Buffer.GetSize() == RawSize ? EStatus::Ok : EStatus::Error;
 			}
 		}
-		if (bHasValue)
-		{
-			const uint64 RawOffset = FMath::Min(Value.GetRawSize(), Request.RawOffset);
-			const uint64 RawSize = FMath::Min(Value.GetRawSize() - RawOffset, Request.RawSize);
-			UE_LOG(LogDerivedDataCache, Verbose, TEXT("%s: Cache hit for %s from '%s'"),
-				*CachePath, *WriteToString<96>(Request.Key, '/', Request.Id), *Request.Name);
-			TRACE_COUNTER_INCREMENT(FileSystemDDC_GetHit);
-			TRACE_COUNTER_ADD(FileSystemDDC_BytesRead, !bExistsOnly ? RawSize : 0);
-			COOK_STAT(Timer.AddHit(!bExistsOnly ? RawSize : 0));
-			FSharedBuffer Buffer;
-			if (!bExistsOnly)
-			{
-				Buffer = ValueReader.Decompress(RawOffset, RawSize);
-			}
-			const EStatus ChunkStatus = bExistsOnly || Buffer.GetSize() == RawSize ? EStatus::Ok : EStatus::Error;
-			OnComplete({Request.Name, Request.Key, Request.Id, Request.RawOffset,
-				RawSize, Value.GetRawHash(), MoveTemp(Buffer), Request.UserData, ChunkStatus});
-			continue;
-		}
-
-		OnComplete(Request.MakeResponse(EStatus::Error));
+		OnComplete({Request.Name, Request.Key, Request.Id, Request.RawOffset,
+			RawSize, Value.GetRawHash(), MoveTemp(Buffer), Request.UserData, Status});
 	}
 }
 
diff --git a/Engine/Source/Developer/DerivedDataCache/Private/PakFileCacheStore.cpp b/Engine/Source/Developer/DerivedDataCache/Private/PakFileCacheStore.cpp
--- a/Engine/Source/Developer/DerivedDataCache/Private/PakFileCacheStore.cpp
+++ b/Engine/Source/Developer/DerivedDataCache/Private/PakFileCacheStore.cpp
@@ -693,21 +693,24 @@ void FPakFileCacheStore::Put(
 {
 	for (const FCachePutRequest& Request : Requests)
 	{
-		const FCacheRecord& Record = Request.Record;
-		TRACE_CPUPROFILER_EVENT_SCOPE(PakFileDDC_Put);
-		COOK_STAT(auto Timer = UsageStats.TimePut());
-		uint64 WriteSize = 0;
-		if (PutCacheRecord(Request.Name, Record, Request.Policy, WriteSize))
-		{
-			UE_LOG(LogDerivedDataCache, Verbose, TEXT("%s: Cache put complete for %s from '%s'"),
-				*CachePath, *WriteToString<96>(Record.GetKey()), *Request.Name);
-			COOK_STAT(if (WriteSize) { Timer.AddHit(WriteSize); });
-			OnComplete(Request.MakeResponse(EStatus::Ok));
-		}
-		else
+		bool bOk;
 		{
-			OnComplete(Request.MakeResponse(EStatus::Error));
+			const FCacheRecord& Record = Request.Record;
+			TRACE_CPUPROFILER_EVENT_SCOPE(PakFileDDC_Put);
+			COOK_STAT(auto Timer = UsageStats.TimePut());
+			uint64 WriteSize = 0;
+			bOk = PutCacheRecord(Request.Name, Record, Request.Policy, WriteSize);
+			if (bOk)
+			{
+				UE_LOG(LogDerivedDataCache, Verbose, TEXT("%s: Cache put complete for %s from '%s'"),
+					*CachePath, *WriteToString<96>(Record.GetKey()), *Request.Name);
+				if (WriteSize)
+				{
+					COOK_STAT(Timer.AddHit(WriteSize));
+				}
+			}
 		}
+		OnComplete(Request.MakeResponse(bOk ? EStatus::Ok : EStatus::Error));
 	}
 }
 
@@ -718,20 +721,23 @@ void FPakFileCacheStore::Get(
 {
 	for (const FCacheGetRequest& Request : Requests)
 	{
-		TRACE_CPUPROFILER_EVENT_SCOPE(PakFileDDC_GetValue);
-		COOK_STAT(auto Timer = UsageStats.TimeGet());
-		EStatus Status = EStatus::Ok;
-		if (FOptionalCacheRecord Record = GetCacheRecord(Request.Name, Request.Key, Request.Policy, Status))
+		EStatus Status = EStatus::Error;
+		FOptionalCacheRecord Record;
 		{
-			UE_LOG(LogDerivedDataCache, Verbose, TEXT("%s: Cache hit for %s from '%s'"),
-				*CachePath, *WriteToString<96>(Request.Key), *Request.Name);
-			COOK_STAT(Timer.AddHit(Private::GetCacheRecordCompressedSize(Record.Get())));
-			OnComplete({Request.Name, MoveTemp(Record).Get(), Request.UserData, Status});
-		}
-		else
-		{
-			OnComplete(Request.MakeResponse(Status));
+			TRACE_CPUPROFILER_EVENT_SCOPE(PakFileDDC_Get);
+			COOK_STAT(auto Timer = UsageStats.TimeGet());
+			if ((Record = GetCacheRecord(Request.Name, Request.Key, Request.Policy, Status)))
+			{
+				UE_LOG(LogDerivedDataCache, Verbose, TEXT("%s: Cache hit for %s from '%s'"),
+					*CachePath, *WriteToString<96>(Request.Key), *Request.Name);
+				COOK_STAT(Timer.AddHit(Private::GetCacheRecordCompressedSize(Record.Get())));
+			}
+			else
+			{
+				Record = FCacheRecordBuilder(Request.Key).Build();
+			}
 		}
+		OnComplete({Request.Name, MoveTemp(Record).Get(), Request.UserData, Status});
 	}
 }
 
@@ -742,20 +748,23 @@ void FPakFileCacheStore::PutValue(
 {
 	for (const FCachePutValueRequest& Request : Requests)
 	{
-		TRACE_CPUPROFILER_EVENT_SCOPE(PakFileDDC_Put);
-		COOK_STAT(auto Timer = UsageStats.TimePut());
-		uint64 WriteSize = 0;
-		if (PutCacheValue(Request.Name, Request.Key, Request.Value, Request.Policy, WriteSize))
-		{
-			UE_LOG(LogDerivedDataCache, Verbose, TEXT("%s: Cache put complete for %s from '%s'"),
-				*CachePath, *WriteToString<96>(Request.Key), *Request.Name);
-			COOK_STAT(if (WriteSize) { Timer.AddHit(WriteSize); });
-			OnComplete(Request.MakeResponse(EStatus::Ok));
-		}
-		else
+		bool bOk;
 		{
-			OnComplete(Request.MakeResponse(EStatus::Error));
+			TRACE_CPUPROFILER_EVENT_SCOPE(PakFileDDC_PutValue);
+			COOK_STAT(auto Timer = UsageStats.TimePut());
+			uint64 WriteSize = 0;
+			bOk = PutCacheValue(Request.Name, Request.Key, Request.Value, Request.Policy, WriteSize);
+			if (bOk)
+			{
+				UE_LOG(LogDerivedDataCache, Verbose, TEXT("%s: Cache put complete for %s from '%s'"),
+					*CachePath, *WriteToString<96>(Request.Key), *Request.Name);
+				if (WriteSize)
+				{
+					COOK_STAT(Timer.AddHit(WriteSize));
+				}
+			}
 		}
+		OnComplete(Request.MakeResponse(bOk ? EStatus::Ok : EStatus::Error));
 	}
 }
 
@@ -766,20 +775,20 @@ void FPakFileCacheStore::GetValue(
 {
 	for (const FCacheGetValueRequest& Request : Requests)
 	{
-		TRACE_CPUPROFILER_EVENT_SCOPE(PakFileDDC_Get);
-		COOK_STAT(auto Timer = UsageStats.TimeGet());
+		bool bOk;
 		FValue Value;
-		if (GetCacheValue(Request.Name, Request.Key, Request.Policy, Value))
 		{
-			UE_LOG(LogDerivedDataCache, Verbose, TEXT("%s: Cache hit for %s from '%s'"),
-				*CachePath, *WriteToString<96>(Request.Key), *Request.Name);
-			COOK_STAT(Timer.AddHit(Value.GetData().GetCompressedSize()));
-			OnComplete({Request.Name, Request.Key, Value, Request.UserData, EStatus::Ok});
-		}
-		else
-		{
-			OnComplete(Request.MakeResponse(EStatus::Error));
+			TRACE_CPUPROFILER_EVENT_SCOPE(PakFileDDC_GetValue);
+			COOK_STAT(auto Timer = UsageStats.TimeGet());
+			bOk = GetCacheValue(Request.Name, Request.Key, Request.Policy, Value);
+			if (bOk)
+			{
+				UE_LOG(LogDerivedDataCache, Verbose, TEXT("%s: Cache hit for %s from '%s'"),
+					*CachePath, *WriteToString<96>(Request.Key), *Request.Name);
+				COOK_STAT(Timer.AddHit(Value.GetData().GetCompressedSize()));
+			}
 		}
+		OnComplete({Request.Name, Request.Key, Value, Request.UserData, bOk ? EStatus::Ok : EStatus::Error});
 	}
 }
 
@@ -800,67 +809,68 @@ void FPakFileCacheStore::GetChunks(
 	FOptionalCacheRecord Record;
 	for (const FCacheGetChunkRequest& Request : SortedRequests)
 	{
-		const bool bExistsOnly = EnumHasAnyFlags(Request.Policy, ECachePolicy::SkipData);
-		TRACE_CPUPROFILER_EVENT_SCOPE(PakFileDDC_Get);
-		COOK_STAT(auto Timer = bExistsOnly ? UsageStats.TimeProbablyExists() : UsageStats.TimeGet());
-		if (!(bHasValue && ValueKey == Request.Key && ValueId == Request.Id) || ValueReader.HasSource() < !bExistsOnly)
+		EStatus Status = EStatus::Error;
+		FSharedBuffer Buffer;
+		uint64 RawSize = 0;
 		{
-			ValueReader.ResetSource();
-			ValueAr.Reset();
-			ValueKey = {};
-			ValueId.Reset();
-			Value.Reset();
-			bHasValue = false;
-			if (Request.Id.IsValid())
+			TRACE_CPUPROFILER_EVENT_SCOPE(PakFileDDC_Get);
+			const bool bExistsOnly = EnumHasAnyFlags(Request.Policy, ECachePolicy::SkipData);
+			COOK_STAT(auto Timer = bExistsOnly ? UsageStats.TimeProbablyExists() : UsageStats.TimeGet());
+			if (!(bHasValue && ValueKey == Request.Key && ValueId == Request.Id) || ValueReader.HasSource() < !bExistsOnly)
 			{
-				if (!(Record && Record.Get().GetKey() == Request.Key))
+				ValueReader.ResetSource();
+				ValueAr.Reset();
+				ValueKey = {};
+				ValueId.Reset();
+				Value.Reset();
+				bHasValue = false;
+				if (Request.Id.IsValid())
 				{
-					FCacheRecordPolicyBuilder PolicyBuilder(ECachePolicy::None);
-					PolicyBuilder.AddValuePolicy(Request.Id, Request.Policy);
-					Record.Reset();
-					Record = GetCacheRecordOnly(Request.Name, Request.Key, PolicyBuilder.Build());
+					if (!(Record && Record.Get().GetKey() == Request.Key))
+					{
+						FCacheRecordPolicyBuilder PolicyBuilder(ECachePolicy::None);
+						PolicyBuilder.AddValuePolicy(Request.Id, Request.Policy);
+						Record.Reset();
+						Record = GetCacheRecordOnly(Request.Name, Request.Key, PolicyBuilder.Build());
+					}
+					if (Record)
+					{
+						if (const FValueWithId& ValueWithId = Record.Get().GetValue(Request.Id))
+						{
+							bHasValue = true;
+							Value = ValueWithId;
+							ValueId = Request.Id;
+							ValueKey = Request.Key;
+							GetCacheContent(Request.Name, Request.Key, ValueId, Value, Request.Policy, ValueReader, ValueAr);
+						}
+					}
 				}
-				if (Record)
+				else
 				{
-					if (const FValueWithId& ValueWithId = Record.Get().GetValue(Request.Id))
+					ValueKey = Request.Key;
+					bHasValue = GetCacheValueOnly(Request.Name, Request.Key, Request.Policy, Value);
+					if (bHasValue)
 					{
-						bHasValue = true;
-						Value = ValueWithId;
-						ValueId = Request.Id;
-						ValueKey = Request.Key;
-						GetCacheContent(Request.Name, Request.Key, ValueId, Value, Request.Policy, ValueReader, ValueAr);
+						GetCacheContent(Request.Name, Request.Key, Request.Id, Value, Request.Policy, ValueReader, ValueAr);
 					}
 				}
 			}
-			else
+			if (bHasValue)
 			{
-				ValueKey = Request.Key;
-				bHasValue = GetCacheValueOnly(Request.Name, Request.Key, Request.Policy, Value);
-				if (bHasValue)
+				const uint64 RawOffset = FMath::Min(Value.GetRawSize(), Request.RawOffset);
+				RawSize = FMath::Min(Value.GetRawSize() - RawOffset, Request.RawSize);
+				UE_LOG(LogDerivedDataCache, Verbose, TEXT("%s: Cache hit for %s from '%s'"),
+					*CachePath, *WriteToString<96>(Request.Key, '/', Request.Id), *Request.Name);
+				COOK_STAT(Timer.AddHit(!bExistsOnly ? RawSize : 0));
+				if (!bExistsOnly)
 				{
-					GetCacheContent(Request.Name, Request.Key, Request.Id, Value, Request.Policy, ValueReader, ValueAr);
+					Buffer = ValueReader.Decompress(RawOffset, RawSize);
 				}
+				Status = bExistsOnly || Buffer.GetSize() == RawSize ? EStatus::Ok : EStatus::Error;
 			}
 		}
-		if (bHasValue)
-		{
-			const uint64 RawOffset = FMath::Min(Value.GetRawSize(), Request.RawOffset);
-			const uint64 RawSize = FMath::Min(Value.GetRawSize() - RawOffset, Request.RawSize);
-			UE_LOG(LogDerivedDataCache, Verbose, TEXT("%s: Cache hit for %s from '%s'"),
-				*CachePath, *WriteToString<96>(Request.Key, '/', Request.Id), *Request.Name);
-			COOK_STAT(Timer.AddHit(!bExistsOnly ? RawSize : 0));
-			FSharedBuffer Buffer;
-			if (!bExistsOnly)
-			{
-				Buffer = ValueReader.Decompress(RawOffset, RawSize);
-			}
-			const EStatus ChunkStatus = bExistsOnly || Buffer.GetSize() == RawSize ? EStatus::Ok : EStatus::Error;
-			OnComplete({Request.Name, Request.Key, Request.Id, Request.RawOffset,
-				RawSize, Value.GetRawHash(), MoveTemp(Buffer), Request.UserData, ChunkStatus});
-			continue;
-		}
-
-		OnComplete(Request.MakeResponse(EStatus::Error));
+		OnComplete({Request.Name, Request.Key, Request.Id, Request.RawOffset,
+			RawSize, Value.GetRawHash(), MoveTemp(Buffer), Request.UserData, Status});
 	}
 }
 
diff --git a/Engine/Source/Developer/DerivedDataCache/Private/S3CacheStore.cpp b/Engine/Source/Developer/DerivedDataCache/Private/S3CacheStore.cpp
--- a/Engine/Source/Developer/DerivedDataCache/Private/S3CacheStore.cpp
+++ b/Engine/Source/Developer/DerivedDataCache/Private/S3CacheStore.cpp
@@ -1816,22 +1816,25 @@ void FS3CacheStore::Get(
 {
 	for (const FCacheGetRequest& Request : Requests)
 	{
-		TRACE_CPUPROFILER_EVENT_SCOPE(S3DDC_Get);
-		TRACE_COUNTER_INCREMENT(S3DDC_Get);
-		COOK_STAT(auto Timer = UsageStats.TimeGet());
-		EStatus Status = EStatus::Ok;
-		if (FOptionalCacheRecord Record = GetCacheRecord(Request.Name, Request.Key, Request.Policy, Status))
+		EStatus Status = EStatus::Error;
+		FOptionalCacheRecord Record;
 		{
-			UE_LOG(LogDerivedDataCache, Verbose, TEXT("%s: Cache hit for %s from '%s'"),
-				*GetName(), *WriteToString<96>(Request.Key), *Request.Name);
-			TRACE_COUNTER_INCREMENT(S3DDC_GetHit);
-			COOK_STAT(Timer.AddHit(Private::GetCacheRecordCompressedSize(Record.Get())));
-			OnComplete({Request.Name, MoveTemp(Record).Get(), Request.UserData, Status});
-		}
-		else
-		{
-			OnComplete(Request.MakeResponse(Status));
+			TRACE_CPUPROFILER_EVENT_SCOPE(S3DDC_Get);
+			TRACE_COUNTER_INCREMENT(S3DDC_Get);
+			COOK_STAT(auto Timer = UsageStats.TimeGet());
+			if ((Record = GetCacheRecord(Request.Name, Request.Key, Request.Policy, Status)))
+			{
+				UE_LOG(LogDerivedDataCache, Verbose, TEXT("%s: Cache hit for %s from '%s'"),
+					*GetName(), *WriteToString<96>(Request.Key), *Request.Name);
+				TRACE_COUNTER_INCREMENT(S3DDC_GetHit);
+				COOK_STAT(Timer.AddHit(Private::GetCacheRecordCompressedSize(Record.Get())));
+			}
+			else
+			{
+				Record = FCacheRecordBuilder(Request.Key).Build();
+			}
 		}
+		OnComplete({Request.Name, MoveTemp(Record).Get(), Request.UserData, Status});
 	}
 }
 
@@ -1850,22 +1853,22 @@ void FS3CacheStore::GetValue(
 {
 	for (const FCacheGetValueRequest& Request : Requests)
 	{
-		TRACE_CPUPROFILER_EVENT_SCOPE(S3DDC_Get);
-		TRACE_COUNTER_INCREMENT(S3DDC_Get);
-		COOK_STAT(auto Timer = UsageStats.TimeGet());
+		bool bOk;
 		FValue Value;
-		if (GetCacheValue(Request.Name, Request.Key, Request.Policy, Value))
-		{
-			UE_LOG(LogDerivedDataCache, Verbose, TEXT("%s: Cache hit for %s from '%s'"),
-				*GetName(), *WriteToString<96>(Request.Key), *Request.Name);
-			TRACE_COUNTER_INCREMENT(S3DDC_GetHit);
-			COOK_STAT(Timer.AddHit(Value.GetData().GetCompressedSize()));
-			OnComplete({Request.Name, Request.Key, Value, Request.UserData, EStatus::Ok});
-		}
-		else
 		{
-			OnComplete(Request.MakeResponse(EStatus::Error));
+			TRACE_CPUPROFILER_EVENT_SCOPE(S3DDC_GetValue);
+			TRACE_COUNTER_INCREMENT(S3DDC_Get);
+			COOK_STAT(auto Timer = UsageStats.TimeGet());
+			bOk = GetCacheValue(Request.Name, Request.Key, Request.Policy, Value);
+			if (bOk)
+			{
+				UE_LOG(LogDerivedDataCache, Verbose, TEXT("%s: Cache hit for %s from '%s'"),
+					*GetName(), *WriteToString<96>(Request.Key), *Request.Name);
+				TRACE_COUNTER_INCREMENT(S3DDC_GetHit);
+				COOK_STAT(Timer.AddHit(Value.GetData().GetCompressedSize()));
+			}
 		}
+		OnComplete({Request.Name, Request.Key, Value, Request.UserData, bOk ? EStatus::Ok : EStatus::Error});
 	}
 }
 
@@ -1886,69 +1889,70 @@ void FS3CacheStore::GetChunks(
 	FOptionalCacheRecord Record;
 	for (const FCacheGetChunkRequest& Request : SortedRequests)
 	{
-		const bool bExistsOnly = EnumHasAnyFlags(Request.Policy, ECachePolicy::SkipData);
-		TRACE_CPUPROFILER_EVENT_SCOPE(S3DDC_Get);
-		TRACE_COUNTER_INCREMENT(S3DDC_Get);
-		COOK_STAT(auto Timer = bExistsOnly ? UsageStats.TimeProbablyExists() : UsageStats.TimeGet());
-		if (!(bHasValue && ValueKey == Request.Key && ValueId == Request.Id) || ValueReader.HasSource() < !bExistsOnly)
+		EStatus Status = EStatus::Error;
+		FSharedBuffer Buffer;
+		uint64 RawSize = 0;
 		{
-			ValueReader.ResetSource();
-			ValueAr.Reset();
-			ValueKey = {};
-			ValueId.Reset();
-			Value.Reset();
-			bHasValue = false;
-			if (Request.Id.IsValid())
+			TRACE_CPUPROFILER_EVENT_SCOPE(S3DDC_GetChunks);
+			TRACE_COUNTER_INCREMENT(S3DDC_Get);
+			const bool bExistsOnly = EnumHasAnyFlags(Request.Policy, ECachePolicy::SkipData);
+			COOK_STAT(auto Timer = bExistsOnly ? UsageStats.TimeProbablyExists() : UsageStats.TimeGet());
+			if (!(bHasValue && ValueKey == Request.Key && ValueId == Request.Id) || ValueReader.HasSource() < !bExistsOnly)
 			{
-				if (!(Record && Record.Get().GetKey() == Request.Key))
+				ValueReader.ResetSource();
+				ValueAr.Reset();
+				ValueKey = {};
+				ValueId.Reset();
+				Value.Reset();
+				bHasValue = false;
+				if (Request.Id.IsValid())
 				{
-					FCacheRecordPolicyBuilder PolicyBuilder(ECachePolicy::None);
-					PolicyBuilder.AddValuePolicy(Request.Id, Request.Policy);
-					Record.Reset();
-					Record = GetCacheRecordOnly(Request.Name, Request.Key, PolicyBuilder.Build());
+					if (!(Record && Record.Get().GetKey() == Request.Key))
+					{
+						FCacheRecordPolicyBuilder PolicyBuilder(ECachePolicy::None);
+						PolicyBuilder.AddValuePolicy(Request.Id, Request.Policy);
+						Record.Reset();
+						Record = GetCacheRecordOnly(Request.Name, Request.Key, PolicyBuilder.Build());
+					}
+					if (Record)
+					{
+						if (const FValueWithId& ValueWithId = Record.Get().GetValue(Request.Id))
+						{
+							bHasValue = true;
+							Value = ValueWithId;
+							ValueId = Request.Id;
+							ValueKey = Request.Key;
+							GetCacheContent(Request.Name, Request.Key, ValueId, Value, Request.Policy, ValueReader, ValueAr);
+						}
+					}
 				}
-				if (Record)
+				else
 				{
-					if (const FValueWithId& ValueWithId = Record.Get().GetValue(Request.Id))
+					ValueKey = Request.Key;
+					bHasValue = GetCacheValueOnly(Request.Name, Request.Key, Request.Policy, Value);
+					if (bHasValue)
 					{
-						bHasValue = true;
-						Value = ValueWithId;
-						ValueId = Request.Id;
-						ValueKey = Request.Key;
-						GetCacheContent(Request.Name, Request.Key, ValueId, Value, Request.Policy, ValueReader, ValueAr);
+						GetCacheContent(Request.Name, Request.Key, Request.Id, Value, Request.Policy, ValueReader, ValueAr);
 					}
 				}
 			}
-			else
+			if (bHasValue)
 			{
-				ValueKey = Request.Key;
-				bHasValue = GetCacheValueOnly(Request.Name, Request.Key, Request.Policy, Value);
-				if (bHasValue)
+				const uint64 RawOffset = FMath::Min(Value.GetRawSize(), Request.RawOffset);
+				RawSize = FMath::Min(Value.GetRawSize() - RawOffset, Request.RawSize);
+				UE_LOG(LogDerivedDataCache, Verbose, TEXT("%s: Cache hit for %s from '%s'"),
+					*GetName(), *WriteToString<96>(Request.Key, '/', Request.Id), *Request.Name);
+				TRACE_COUNTER_INCREMENT(S3DDC_GetHit);
+				COOK_STAT(Timer.AddHit(!bExistsOnly ? RawSize : 0));
+				if (!bExistsOnly)
 				{
-					GetCacheContent(Request.Name, Request.Key, Request.Id, Value, Request.Policy, ValueReader, ValueAr);
+					Buffer = ValueReader.Decompress(RawOffset, RawSize);
 				}
+				Status = bExistsOnly || Buffer.GetSize() == RawSize ? EStatus::Ok : EStatus::Error;
 			}
 		}
-		if (bHasValue)
-		{
-			const uint64 RawOffset = FMath::Min(Value.GetRawSize(), Request.RawOffset);
-			const uint64 RawSize = FMath::Min(Value.GetRawSize() - RawOffset, Request.RawSize);
-			UE_LOG(LogDerivedDataCache, Verbose, TEXT("%s: Cache hit for %s from '%s'"),
-				*GetName(), *WriteToString<96>(Request.Key, '/', Request.Id), *Request.Name);
-			TRACE_COUNTER_INCREMENT(S3DDC_GetHit);
-			COOK_STAT(Timer.AddHit(!bExistsOnly ? RawSize : 0));
-			FSharedBuffer Buffer;
-			if (!bExistsOnly)
-			{
-				Buffer = ValueReader.Decompress(RawOffset, RawSize);
-			}
-			const EStatus ChunkStatus = bExistsOnly || Buffer.GetSize() == RawSize ? EStatus::Ok : EStatus::Error;
-			OnComplete({Request.Name, Request.Key, Request.Id, Request.RawOffset,
-				RawSize, Value.GetRawHash(), MoveTemp(Buffer), Request.UserData, ChunkStatus});
-			continue;
-		}
-
-		OnComplete(Request.MakeResponse(EStatus::Error));
+		OnComplete({Request.Name, Request.Key, Request.Id, Request.RawOffset,
+			RawSize, Value.GetRawHash(), MoveTemp(Buffer), Request.UserData, Status});
 	}
 }
 
