Optimized BuildCacheKey using FAsciiSet for ~4x speedup

Estimated to save 45s in a 1h50m cook of a larger internal project

diff --git a/Engine/Source/Developer/DerivedDataCache/Public/DerivedDataCacheInterface.h b/Engine/Source/Developer/DerivedDataCache/Public/DerivedDataCacheInterface.h
--- a/Engine/Source/Developer/DerivedDataCache/Public/DerivedDataCacheInterface.h
+++ b/Engine/Source/Developer/DerivedDataCache/Public/DerivedDataCacheInterface.h
@@ -6,6 +6,7 @@
 #include "CoreMinimal.h"
 #include "Containers/BitArray.h"
 #include "Containers/StringView.h"
+#include "Misc/AsciiSet.h"
 #include "DerivedDataCacheModule.h"
 
 class FDerivedDataCacheUsageStats;
@@ -67,12 +68,34 @@ public:
 	// Low Level Static Helpers
 	//--------------------------
 	
+private:
+	static inline constexpr FAsciiSet ValidCacheKeyChars = "ABCDEFGHIJKLMNOPQRSTUVWXYZabcdefghijklmnopqrstuvwxyz0123456789_$";
+
+	static void AppendSanitized(FString& Out, FStringView In)
+	{
+		while (true)
+		{
+			FStringView SanePrefix = FAsciiSet::FindPrefixWith(In, ValidCacheKeyChars);
+			Out.Append(SanePrefix);
+			In.RemovePrefix(SanePrefix.Len());
+
+			if (In.IsEmpty())
+			{
+				break;
+			}
+
+			Out.Appendf(TEXT("$%x"), uint32(In[0]));
+			In.RemovePrefix(1);
+		}
+	}
+
+public:
 	/**
 	 * Returns true if character is valid in a DDC cache key without escaping 
 	**/
 	static bool IsValidCacheChar(const TCHAR C)
 	{
-		return FChar::IsAlnum(C) || FChar::IsUnderscore(C) || C == TEXT('$');
+		return ValidCacheKeyChars.Contains(C);
 	}
 
 	/** 
@@ -82,43 +105,9 @@ public:
 	**/
 	static FString SanitizeCacheKey(const TCHAR* CacheKey)
 	{
-		FString Output;
-		FString Input(CacheKey);
-		int32 StartValid = 0;
-		int32 NumValid = 0;
-
-		for (int32 i = 0; i < Input.Len(); i++)
-		{
-			if ( IsValidCacheChar(Input[i]) )
-			{
-				NumValid++;
-			}
-			else
-			{
-				// Copy the valid range so far
-				Output += Input.Mid(StartValid, NumValid);
-
-				// Reset valid ranges
-				StartValid = i + 1;
-				NumValid = 0;
-
-				// Replace the invalid character with a special string
-				Output += FString::Printf(TEXT("$%x"), uint32(Input[i]));
-			}
-		}
-
-		// Just return the input if the entire string was valid
-		if (StartValid == 0 && NumValid == Input.Len())
-		{
-			return Input;
-		}
-		else if (NumValid > 0)
-		{
-			// Copy the remaining valid part
-			Output += Input.Mid(StartValid, NumValid);
-		}
-
-		return Output;
+		FString Out;
+		AppendSanitized(Out, FStringView(CacheKey));
+		return Out;
 	}
 
 	/** 
@@ -130,7 +119,19 @@ public:
 	**/
 	static FString BuildCacheKey(const TCHAR* PluginName, const TCHAR* VersionString, const TCHAR* PluginSpecificCacheKeySuffix)
 	{
-		return SanitizeCacheKey(*FString::Printf(TEXT("%s_%s_%s"), PluginName, VersionString, PluginSpecificCacheKeySuffix));
+		return BuildCacheKey(FStringView(PluginName), FStringView(VersionString), FStringView(PluginSpecificCacheKeySuffix));
+	}
+	
+	static FString BuildCacheKey(FStringView PluginName, FStringView VersionString, FStringView PluginSpecificCacheKeySuffix)
+	{
+		FString Out;
+		Out.Reserve(PluginName.Len() + 1 + VersionString.Len() + 1 + PluginSpecificCacheKeySuffix.Len());
+		AppendSanitized(Out, PluginName);
+		Out.AppendChar('_');
+		AppendSanitized(Out, VersionString);
+		Out.AppendChar('_');
+		AppendSanitized(Out, PluginSpecificCacheKeySuffix);
+		return Out;
 	}
 
 	//--------------------
