DDC: Added IRequestOwner::LaunchTask

- Added IRequestOwner::LaunchTask to launch a task as a request in the request owner.
- Renamed ExecuteInCacheThreadPool to LaunchTaskInCacheThreadPool.
- Made GCacheThreadPool private to one source file.
- Changed task scheduling throughout DDC to use these functions where appropriate.

diff --git a/Engine/Source/Developer/DerivedDataCache/Private/DerivedDataBackends.cpp b/Engine/Source/Developer/DerivedDataCache/Private/DerivedDataBackends.cpp
--- a/Engine/Source/Developer/DerivedDataCache/Private/DerivedDataBackends.cpp
+++ b/Engine/Source/Developer/DerivedDataCache/Private/DerivedDataBackends.cpp
@@ -1416,7 +1416,7 @@ void FDerivedDataBackendInterface::LegacyGet(
 			FLegacyCacheValue LegacyValue(FCompositeBuffer(MoveTemp(Value)));
 			if (LegacyValue.HasData() && LegacyMode == EBackendLegacyMode::ValueWithLegacyFallback)
 			{
-				Private::ExecuteInCacheThreadPool(AsyncOwner, [this, Request, LegacyValue](IRequestOwner& AsyncOwner, bool bCancel)
+				Private::LaunchTaskInCacheThreadPool(AsyncOwner, [this, &AsyncOwner, Request, LegacyValue]
 				{
 					ILegacyCacheStore::LegacyPut({{Request.Name, Request.Key, LegacyValue}}, AsyncOwner, [](auto&&){});
 				});
diff --git a/Engine/Source/Developer/DerivedDataCache/Private/DerivedDataBuildScheduler.cpp b/Engine/Source/Developer/DerivedDataCache/Private/DerivedDataBuildScheduler.cpp
--- a/Engine/Source/Developer/DerivedDataCache/Private/DerivedDataBuildScheduler.cpp
+++ b/Engine/Source/Developer/DerivedDataCache/Private/DerivedDataBuildScheduler.cpp
@@ -10,7 +10,6 @@
 #include "Experimental/Async/LazyEvent.h"
 #include "Misc/Guid.h"
 #include "Misc/ScopeLock.h"
-#include "Tasks/Task.h"
 
 namespace UE::DerivedData::Private
 {
@@ -19,24 +18,7 @@ namespace UE::DerivedData::Private
 
 static void ScheduleAsyncStep(IBuildJob& Job, IRequestOwner& Owner, const TCHAR* DebugName)
 {
-	class FStepAsyncRequest final : public FRequestBase
-	{
-	public:
-		Tasks::FTask Task;
-
-		void SetPriority(EPriority Priority) final {}
-		void Cancel() final { Task.Wait(); }
-		void Wait() final { Task.Wait(); }
-	};
-
-	FStepAsyncRequest* Request = new FStepAsyncRequest;
-
-	Tasks::FTaskEvent TaskEvent(TEXT("ScheduleAsyncStep"));
-	Request->Task = Tasks::Launch(DebugName,
-		[&Job, &Owner, Request] { Owner.End(Request, [&Job] { Job.StepExecution(); }); },
-		TaskEvent, Tasks::ETaskPriority::BackgroundNormal);
-	Owner.Begin(Request);
-	TaskEvent.Trigger();
+	Owner.LaunchTask(DebugName, [&Job] { Job.StepExecution(); });
 }
 
 ///////////////////////////////////////////////////////////////////////////////////////////////////////////////////////
diff --git a/Engine/Source/Developer/DerivedDataCache/Private/DerivedDataCache.cpp b/Engine/Source/Developer/DerivedDataCache/Private/DerivedDataCache.cpp
--- a/Engine/Source/Developer/DerivedDataCache/Private/DerivedDataCache.cpp
+++ b/Engine/Source/Developer/DerivedDataCache/Private/DerivedDataCache.cpp
@@ -14,7 +14,9 @@
 #include "DerivedDataCachePrivate.h"
 #include "DerivedDataCacheUsageStats.h"
 #include "DerivedDataPluginInterface.h"
+#include "DerivedDataRequest.h"
 #include "DerivedDataRequestOwner.h"
+#include "Experimental/Async/LazyEvent.h"
 #include "Features/IModularFeatures.h"
 #include "HAL/ThreadSafeCounter.h"
 #include "Misc/CommandLine.h"
@@ -192,6 +194,121 @@ namespace UE::DerivedData::Private
 
 FQueuedThreadPool* GCacheThreadPool;
 
+class FCacheThreadPoolTaskRequest final : public FRequestBase, private IQueuedWork
+{
+public:
+	inline FCacheThreadPoolTaskRequest(
+		IRequestOwner& InOwner,
+		TUniqueFunction<void ()>&& InTaskBody)
+		: Owner(InOwner)
+		, TaskBody(MoveTemp(InTaskBody))
+	{
+	}
+
+	inline void Start(EPriority Priority)
+	{
+		Owner.Begin(this);
+		DoneEvent.Reset();
+		GCacheThreadPool->AddQueuedWork(this, GetPriority(Priority));
+	}
+
+	inline void Execute()
+	{
+		FScopeCycleCounter Scope(GetStatId(), /*bAlways*/ true);
+		Owner.End(this, [this]
+		{
+			TaskBody();
+			DoneEvent.Trigger();
+		});
+		// DO NOT ACCESS ANY MEMBERS PAST THIS POINT!
+	}
+
+	// IRequest Interface
+
+	inline void SetPriority(EPriority Priority) final
+	{
+		if (GCacheThreadPool->RetractQueuedWork(this))
+		{
+			GCacheThreadPool->AddQueuedWork(this, GetPriority(Priority));
+		}
+	}
+
+	inline void Cancel() final
+	{
+		if (!DoneEvent.Wait(0))
+		{
+			if (GCacheThreadPool->RetractQueuedWork(this))
+			{
+				Abandon();
+			}
+			else
+			{
+				FScopeCycleCounter Scope(GetStatId());
+				DoneEvent.Wait();
+			}
+		}
+	}
+
+	inline void Wait() final
+	{
+		if (!DoneEvent.Wait(0))
+		{
+			if (GCacheThreadPool->RetractQueuedWork(this))
+			{
+				DoThreadedWork();
+			}
+			else
+			{
+				FScopeCycleCounter Scope(GetStatId());
+				DoneEvent.Wait();
+			}
+		}
+	}
+
+private:
+	static EQueuedWorkPriority GetPriority(EPriority Priority)
+	{
+		switch (Priority)
+		{
+		case EPriority::Blocking: return EQueuedWorkPriority::Blocking;
+		case EPriority::Highest:  return EQueuedWorkPriority::Highest;
+		case EPriority::High:     return EQueuedWorkPriority::High;
+		case EPriority::Normal:   return EQueuedWorkPriority::Normal;
+		case EPriority::Low:      return EQueuedWorkPriority::Low;
+		case EPriority::Lowest:   return EQueuedWorkPriority::Lowest;
+		default: checkNoEntry();  return EQueuedWorkPriority::Normal;
+		}
+	}
+
+	// IQueuedWork Interface
+
+	inline void DoThreadedWork() final { Execute(); }
+	inline void Abandon() final { Execute(); }
+
+	inline TStatId GetStatId() const
+	{
+		RETURN_QUICK_DECLARE_CYCLE_STAT(FCacheThreadPoolTaskRequest, STATGROUP_ThreadPoolAsyncTasks);
+	}
+
+private:
+	IRequestOwner& Owner;
+	TUniqueFunction<void ()> TaskBody;
+	FLazyEvent DoneEvent{EEventMode::ManualReset};
+};
+
+void LaunchTaskInCacheThreadPool(IRequestOwner& Owner, TUniqueFunction<void ()>&& TaskBody)
+{
+	if (GCacheThreadPool)
+	{
+		FCacheThreadPoolTaskRequest* Request = new FCacheThreadPoolTaskRequest(Owner, MoveTemp(TaskBody));
+		Request->Start(Owner.GetPriority());
+	}
+	else
+	{
+		TaskBody();
+	}
+}
+
 /**
  * Implementation of the derived data cache
  * This API is fully threadsafe
diff --git a/Engine/Source/Developer/DerivedDataCache/Private/DerivedDataCachePrivate.h b/Engine/Source/Developer/DerivedDataCache/Private/DerivedDataCachePrivate.h
--- a/Engine/Source/Developer/DerivedDataCache/Private/DerivedDataCachePrivate.h
+++ b/Engine/Source/Developer/DerivedDataCache/Private/DerivedDataCachePrivate.h
@@ -7,34 +7,25 @@
 #include "Misc/AsciiSet.h"
 
 class FDerivedDataCacheInterface;
-class FQueuedThreadPool;
 
 template <typename FuncType> class TUniqueFunction;
 
 namespace UE::DerivedData { class FCacheRecord; }
 namespace UE::DerivedData { class ICache; }
 namespace UE::DerivedData { class IRequestOwner; }
-namespace UE::DerivedData { struct FCacheKey; }
 
 namespace UE::DerivedData::Private
 {
 
-// Created by CreateCache
-extern FQueuedThreadPool* GCacheThreadPool;
-
 // Implemented in DerivedDataCache.cpp
 ICache* CreateCache(FDerivedDataCacheInterface** OutLegacyCache);
+void LaunchTaskInCacheThreadPool(IRequestOwner& Owner, TUniqueFunction<void ()>&& TaskBody);
 
 // Implemented in DerivedDataCacheRecord.cpp
 uint64 GetCacheRecordCompressedSize(const FCacheRecord& Record);
 uint64 GetCacheRecordTotalRawSize(const FCacheRecord& Record);
 uint64 GetCacheRecordRawSize(const FCacheRecord& Record);
 
-// Implemented in DerivedDataCacheStoreAsync.cpp
-void ExecuteInCacheThreadPool(
-	IRequestOwner& Owner,
-	TUniqueFunction<void (IRequestOwner& Owner, bool bCancel)>&& Function);
-
 ///////////////////////////////////////////////////////////////////////////////////////////////////////////////////////
 
 template <typename CharType>
diff --git a/Engine/Source/Developer/DerivedDataCache/Private/DerivedDataCacheStoreAsync.cpp b/Engine/Source/Developer/DerivedDataCache/Private/DerivedDataCacheStoreAsync.cpp
--- a/Engine/Source/Developer/DerivedDataCache/Private/DerivedDataCacheStoreAsync.cpp
+++ b/Engine/Source/Developer/DerivedDataCache/Private/DerivedDataCacheStoreAsync.cpp
@@ -6,15 +6,12 @@
 #include "DerivedDataCacheStore.h"
 #include "DerivedDataCacheUsageStats.h"
 #include "DerivedDataLegacyCacheStore.h"
-#include "DerivedDataRequest.h"
 #include "DerivedDataRequestOwner.h"
 #include "DerivedDataValueId.h"
-#include "Experimental/Async/LazyEvent.h"
 #include "Memory/SharedBuffer.h"
 #include "MemoryCacheStore.h"
 #include "ProfilingDebugging/CookStats.h"
 #include "Stats/Stats.h"
-#include "Tasks/Task.h"
 #include "Templates/Invoke.h"
 
 namespace UE::DerivedData
@@ -177,130 +174,6 @@ void FCacheStoreAsync::LegacyStats(FDerivedDataCacheStatsNode& OutNode)
 	}
 }
 
-class FDerivedDataAsyncWrapperRequest final : public FRequestBase, private IQueuedWork
-{
-public:
-	inline FDerivedDataAsyncWrapperRequest(
-		IRequestOwner& InOwner,
-		TUniqueFunction<void (IRequestOwner& Owner, bool bCancel)>&& InFunction)
-		: Owner(InOwner)
-		, Function(MoveTemp(InFunction))
-	{
-	}
-
-	inline void Start(EPriority Priority)
-	{
-		Owner.Begin(this);
-		DoneEvent.Reset();
-		Private::GCacheThreadPool->AddQueuedWork(this, GetPriority(Priority));
-	}
-
-	inline void Execute(bool bCancel)
-	{
-		FScopeCycleCounter Scope(GetStatId(), /*bAlways*/ true);
-		Owner.End(this, [this, bCancel]
-		{
-			Function(Owner, bCancel);
-			DoneEvent.Trigger();
-		});
-		// DO NOT ACCESS ANY MEMBERS PAST THIS POINT!
-	}
-
-	// IRequest Interface
-
-	inline void SetPriority(EPriority Priority) final
-	{
-		if (Private::GCacheThreadPool->RetractQueuedWork(this))
-		{
-			Private::GCacheThreadPool->AddQueuedWork(this, GetPriority(Priority));
-		}
-	}
-
-	inline void Cancel() final
-	{
-		if (!DoneEvent.Wait(0))
-		{
-			if (Private::GCacheThreadPool->RetractQueuedWork(this))
-			{
-				Abandon();
-			}
-			else
-			{
-				FScopeCycleCounter Scope(GetStatId());
-				DoneEvent.Wait();
-			}
-		}
-	}
-
-	inline void Wait() final
-	{
-		if (!DoneEvent.Wait(0))
-		{
-			if (Private::GCacheThreadPool->RetractQueuedWork(this))
-			{
-				DoThreadedWork();
-			}
-			else
-			{
-				FScopeCycleCounter Scope(GetStatId());
-				DoneEvent.Wait();
-			}
-		}
-	}
-
-private:
-	static EQueuedWorkPriority GetPriority(EPriority Priority)
-	{
-		switch (Priority)
-		{
-		case EPriority::Blocking: return EQueuedWorkPriority::Blocking;
-		case EPriority::Highest:  return EQueuedWorkPriority::Highest;
-		case EPriority::High:     return EQueuedWorkPriority::High;
-		case EPriority::Normal:   return EQueuedWorkPriority::Normal;
-		case EPriority::Low:      return EQueuedWorkPriority::Low;
-		case EPriority::Lowest:   return EQueuedWorkPriority::Lowest;
-		default: checkNoEntry();  return EQueuedWorkPriority::Normal;
-		}
-	}
-
-	// IQueuedWork Interface
-
-	inline void DoThreadedWork() final
-	{
-		Execute(/*bCancel*/ false);
-	}
-
-	inline void Abandon() final
-	{
-		Execute(/*bCancel*/ true);
-	}
-
-	inline TStatId GetStatId() const
-	{
-		RETURN_QUICK_DECLARE_CYCLE_STAT(FDerivedDataAsyncWrapperRequest, STATGROUP_ThreadPoolAsyncTasks);
-	}
-
-private:
-	IRequestOwner& Owner;
-	TUniqueFunction<void (IRequestOwner& Owner, bool bCancel)> Function;
-	FLazyEvent DoneEvent{EEventMode::ManualReset};
-};
-
-void Private::ExecuteInCacheThreadPool(
-	IRequestOwner& Owner,
-	TUniqueFunction<void (IRequestOwner& Owner, bool bCancel)>&& Function)
-{
-	if (GCacheThreadPool)
-	{
-		FDerivedDataAsyncWrapperRequest* Request = new FDerivedDataAsyncWrapperRequest(Owner, MoveTemp(Function));
-		Request->Start(Owner.GetPriority());
-	}
-	else
-	{
-		Invoke(Function, Owner, /*bCancel*/ false);
-	}
-}
-
 template <typename RequestType, typename OnCompleteType, typename OnExecuteType>
 void FCacheStoreAsync::Execute(
 	COOK_STAT(CookStatsFunction OnAddStats,)
@@ -323,19 +196,19 @@ void FCacheStoreAsync::Execute(
 	};
 
 	FDerivedDataBackend::Get().AddToAsyncCompletionCounter(Requests.Num());
-	if (Owner.GetPriority() == EPriority::Blocking || !Private::GCacheThreadPool)
+	if (Owner.GetPriority() == EPriority::Blocking)
 	{
 		return ExecuteWithStats(Requests, Owner, MoveTemp(OnComplete));
 	}
 
-	Private::ExecuteInCacheThreadPool(Owner,
+	Private::LaunchTaskInCacheThreadPool(Owner,
 		[this,
+		&Owner,
 		Requests = TArray<RequestType>(Requests),
 		OnComplete = MoveTemp(OnComplete),
-		ExecuteWithStats = MoveTemp(ExecuteWithStats)]
-		(IRequestOwner& Owner, bool bCancel) mutable
+		ExecuteWithStats = MoveTemp(ExecuteWithStats)]() mutable
 		{
-			if (!bCancel)
+			if (!Owner.IsCanceled())
 			{
 				ExecuteWithStats(Requests, Owner, MoveTemp(OnComplete));
 			}
diff --git a/Engine/Source/Developer/DerivedDataCache/Private/DerivedDataRequestOwner.cpp b/Engine/Source/Developer/DerivedDataCache/Private/DerivedDataRequestOwner.cpp
--- a/Engine/Source/Developer/DerivedDataCache/Private/DerivedDataRequestOwner.cpp
+++ b/Engine/Source/Developer/DerivedDataCache/Private/DerivedDataRequestOwner.cpp
@@ -9,7 +9,10 @@
 #include "HAL/CriticalSection.h"
 #include "Misc/ScopeExit.h"
 #include "Misc/ScopeRWLock.h"
+#include "Tasks/Task.h"
+#include "Templates/Function.h"
 #include "Templates/RefCounting.h"
+#include <atomic>
 
 namespace UE::DerivedData::Private
 {
@@ -276,4 +279,27 @@ bool FRequestOwner::Poll() const                        { return ToSharedOwner(O
 FRequestOwner::operator IRequest*()                     { return ToSharedOwner(Owner.Get()); }
 void FRequestOwner::Destroy(IRequestOwner& SharedOwner) { return ToSharedOwner(&SharedOwner)->Destroy(); }
 
+void IRequestOwner::LaunchTask(const TCHAR* DebugName, TUniqueFunction<void ()>&& TaskBody)
+{
+	using namespace Tasks;
+
+	struct FTaskRequest final : public FRequestBase
+	{
+		void SetPriority(EPriority Priority) final {}
+		void Cancel() final { Task.Wait(); }
+		void Wait() final { Task.Wait(); }
+		FTask Task;
+	};
+
+	Tasks::FTaskEvent TaskEvent(TEXT("LaunchTaskRequest"));
+	FTaskRequest* Request = new FTaskRequest;
+	Request->Task = Launch(
+		DebugName,
+		[this, Request, TaskBody = MoveTemp(TaskBody)] { End(Request, TaskBody); },
+		TaskEvent,
+		GetPriority() <= EPriority::Normal ? ETaskPriority::BackgroundNormal : ETaskPriority::BackgroundHigh);
+	Begin(Request);
+	TaskEvent.Trigger();
+}
+
 } // UE::DerivedData
diff --git a/Engine/Source/Developer/DerivedDataCache/Private/PakFileCacheStore.cpp b/Engine/Source/Developer/DerivedDataCache/Private/PakFileCacheStore.cpp
--- a/Engine/Source/Developer/DerivedDataCache/Private/PakFileCacheStore.cpp
+++ b/Engine/Source/Developer/DerivedDataCache/Private/PakFileCacheStore.cpp
@@ -10,7 +10,6 @@
 #include "DerivedDataCacheRecord.h"
 #include "DerivedDataCacheUsageStats.h"
 #include "DerivedDataChunk.h"
-#include "DerivedDataRequest.h"
 #include "DerivedDataRequestOwner.h"
 #include "DerivedDataValue.h"
 #include "HAL/CriticalSection.h"
@@ -30,7 +29,6 @@
 #include "Serialization/CompactBinaryWriter.h"
 #include "Serialization/MemoryReader.h"
 #include "Serialization/MemoryWriter.h"
-#include "Tasks/Task.h"
 #include "Templates/Greater.h"
 #include "Templates/UniquePtr.h"
 
@@ -1529,29 +1527,6 @@ bool FPakFileCacheStore::FileExists(const FStringView Path)
 	return CacheItems.ContainsByHash(PathHash, Path);
 }
 
-static void ScheduleAsyncRequest(IRequestOwner& Owner, const TCHAR* DebugName, TUniqueFunction<void (IRequestOwner& Owner)>&& Function)
-{
-	class FAsyncRequest final : public FRequestBase
-	{
-	public:
-		Tasks::FTask Task;
-		TUniqueFunction<void (IRequestOwner& Owner)> Function;
-
-		void SetPriority(EPriority Priority) final {}
-		void Cancel() final { Task.Wait(); }
-		void Wait() final { Task.Wait(); }
-	};
-
-	FAsyncRequest* Request = new FAsyncRequest;
-	Request->Function = MoveTemp(Function);
-	Tasks::FTaskEvent TaskEvent(TEXT("ScheduleAsyncRequest"));
-	Request->Task = Tasks::Launch(DebugName,
-		[Request, &Owner] { Owner.End(Request, [Request, &Owner] { Request->Function(Owner); }); },
-		TaskEvent, Tasks::ETaskPriority::BackgroundNormal);
-	Owner.Begin(Request);
-	TaskEvent.Trigger();
-}
-
 class FCompressedPakFileCacheStore final : public FPakFileCacheStore
 {
 public:
@@ -1620,8 +1595,8 @@ void FCompressedPakFileCacheStore::Put(
 	IRequestOwner& Owner,
 	FOnCachePutComplete&& OnComplete)
 {
-	ScheduleAsyncRequest(Owner, TEXT("PakFileDDC_Put"),
-		[this, Requests = TArray<FCachePutRequest, TInlineAllocator<1>>(Requests), OnComplete = MoveTemp(OnComplete)](IRequestOwner& Owner) mutable
+	Owner.LaunchTask(TEXT("PakFileDDC_Put"),
+		[this, &Owner, Requests = TArray<FCachePutRequest, TInlineAllocator<1>>(Requests), OnComplete = MoveTemp(OnComplete)]() mutable
 		{
 			for (FCachePutRequest& Request : Requests)
 			{
@@ -1633,8 +1608,8 @@ void FCompressedPakFileCacheStore::Put(
 				}
 				Request.Record = Builder.Build();
 			}
-			Private::ExecuteInCacheThreadPool(Owner,
-				[this, Requests = MoveTemp(Requests), OnComplete = MoveTemp(OnComplete)](IRequestOwner& Owner, bool bCancel) mutable
+			Private::LaunchTaskInCacheThreadPool(Owner,
+				[this, &Owner, Requests = MoveTemp(Requests), OnComplete = MoveTemp(OnComplete)]() mutable
 				{
 					FPakFileCacheStore::Put(Requests, Owner, MoveTemp(OnComplete));
 				});
@@ -1646,15 +1621,15 @@ void FCompressedPakFileCacheStore::PutValue(
 	IRequestOwner& Owner,
 	FOnCachePutValueComplete&& OnComplete)
 {
-	ScheduleAsyncRequest(Owner, TEXT("PakFileDDC_PutValue"),
-		[this, Requests = TArray<FCachePutValueRequest, TInlineAllocator<1>>(Requests), OnComplete = MoveTemp(OnComplete)](IRequestOwner& Owner) mutable
+	Owner.LaunchTask(TEXT("PakFileDDC_PutValue"),
+		[this, &Owner, Requests = TArray<FCachePutValueRequest, TInlineAllocator<1>>(Requests), OnComplete = MoveTemp(OnComplete)]() mutable
 		{
 			for (FCachePutValueRequest& Request : Requests)
 			{
 				Request.Value = Compress(Request.Value);
 			}
-			Private::ExecuteInCacheThreadPool(Owner,
-				[this, Requests = MoveTemp(Requests), OnComplete = MoveTemp(OnComplete)](IRequestOwner& Owner, bool bCancel) mutable
+			Private::LaunchTaskInCacheThreadPool(Owner,
+				[this, &Owner, Requests = MoveTemp(Requests), OnComplete = MoveTemp(OnComplete)]() mutable
 				{
 					FPakFileCacheStore::PutValue(Requests, Owner, MoveTemp(OnComplete));
 				});
diff --git a/Engine/Source/Developer/DerivedDataCache/Public/DerivedDataRequestOwner.h b/Engine/Source/Developer/DerivedDataCache/Public/DerivedDataRequestOwner.h
--- a/Engine/Source/Developer/DerivedDataCache/Public/DerivedDataRequestOwner.h
+++ b/Engine/Source/Developer/DerivedDataCache/Public/DerivedDataRequestOwner.h
@@ -11,6 +11,8 @@
 
 #define UE_API DERIVEDDATACACHE_API
 
+template <typename FuncType> class TUniqueFunction;
+
 namespace UE::DerivedData { class IRequest; }
 namespace UE::DerivedData { enum class EPriority : uint8; }
 
@@ -80,6 +82,16 @@ public:
 	virtual EPriority GetPriority() const = 0;
 	/** Returns whether the owner has been canceled, which new requests are expected to check. */
 	virtual bool IsCanceled() const = 0;
+
+	/**
+	 * Launches a task that executes the task body when scheduled.
+	 *
+	 * The task inherits the current priority of this request owner.
+	 * The task is launched as a new request within this request owner.
+	 *
+	 * @note The debug name is not copied and must remain valid for the duration of the task.
+	 */
+	UE_API void LaunchTask(const TCHAR* DebugName, TUniqueFunction<void ()>&& TaskBody);
 };
 
 /**
