Allow OAuth access scope to be configurable for a Horde Storage DDC store.  Supplied by partners.

#rnx

diff --git a/Engine/Source/Developer/DerivedDataCache/Private/DerivedDataBackends.cpp b/Engine/Source/Developer/DerivedDataCache/Private/DerivedDataBackends.cpp
--- a/Engine/Source/Developer/DerivedDataCache/Private/DerivedDataBackends.cpp
+++ b/Engine/Source/Developer/DerivedDataCache/Private/DerivedDataBackends.cpp
@@ -60,6 +60,7 @@ ILegacyCacheStore* CreateHttpCacheStore(
 	const TCHAR* OAuthProvider,
 	const TCHAR* OAuthClientId,
 	const TCHAR* OAuthData, 
+	const TCHAR* OAuthScope,
 	const FDerivedDataBackendInterface::ESpeedClass* ForceSpeedClass,
 	EBackendLegacyMode LegacyMode,
 	bool bReadOnly);
@@ -777,6 +778,7 @@ public:
 		FString& OAuthProvider,
 		FString& OAuthClientId,
 		FString& OAuthSecret,
+		FString& OAuthScope,
 		EBackendLegacyMode& LegacyMode,
 		bool& bReadOnly)
 	{
@@ -787,7 +789,7 @@ public:
 			const TCHAR* ServerSection = TEXT("HordeStorageServers");
 			if (GConfig->GetString(ServerSection, *ServerId, ServerEntry, IniFilename))
 			{
-				ParseHttpCacheParams(NodeName, *ServerEntry, IniFilename, IniSection, Host, bResolveHostCanonicalName, Namespace, StructuredNamespace, OAuthProvider, OAuthClientId, OAuthSecret, LegacyMode, bReadOnly);
+				ParseHttpCacheParams(NodeName, *ServerEntry, IniFilename, IniSection, Host, bResolveHostCanonicalName, Namespace, StructuredNamespace, OAuthProvider, OAuthClientId, OAuthSecret, OAuthScope, LegacyMode, bReadOnly);
 			}
 			else
 			{
@@ -821,8 +823,30 @@ public:
 		FParse::Value(Entry, TEXT("Namespace="), Namespace);
 		FParse::Value(Entry, TEXT("StructuredNamespace="), StructuredNamespace);
 		FParse::Value(Entry, TEXT("OAuthProvider="), OAuthProvider);
+		
+		FString CommandLineOAuthProvidorOverride;
+		if (FParse::Value(Entry, TEXT("CommandLineOAuthProviderOverride="), CommandLineOAuthProvidorOverride))
+		{
+			if (FParse::Value(FCommandLine::Get(), *(CommandLineOAuthProvidorOverride + TEXT("=")), OAuthProvider))
+			{
+				UE_LOG(LogDerivedDataCache, Log, TEXT("Node %s found command line override for OAuthProvider %s=%s"), NodeName, *CommandLineOAuthProvidorOverride, *OAuthProvider);
+			}
+		}
+		
 		FParse::Value(Entry, TEXT("OAuthClientId="), OAuthClientId);
 		FParse::Value(Entry, TEXT("OAuthSecret="), OAuthSecret);
+		
+		FString CommandLineOAuthSecretOverride;
+		if (FParse::Value(Entry, TEXT("CommandLineOAuthSecretOverride="), CommandLineOAuthSecretOverride))
+		{
+			if (FParse::Value(FCommandLine::Get(), *(CommandLineOAuthSecretOverride + TEXT("=")), OAuthSecret))
+			{
+				UE_LOG(LogDerivedDataCache, Log, TEXT("Node %s found command line override for OAuthSecret %s=%s"), NodeName, *CommandLineOAuthSecretOverride, *OAuthSecret);
+			}
+		}
+		
+		FParse::Value(Entry, TEXT("OAuthScope="), OAuthScope);
+		
 		FParse::Bool(Entry, TEXT("ReadOnly="), bReadOnly);
 
 		if (FString LegacyModeString; FParse::Value(Entry, TEXT("LegacyMode="), LegacyModeString))
@@ -850,10 +874,11 @@ public:
 		FString OAuthProvider;
 		FString OAuthClientId;
 		FString OAuthSecret;
+		FString OAuthScope;
 		EBackendLegacyMode LegacyMode = EBackendLegacyMode::ValueOnly;
 		bool bReadOnly = false;
 
-		ParseHttpCacheParams(NodeName, Entry, IniFilename, IniSection, Host, bResolveHostCanonicalName, Namespace, StructuredNamespace, OAuthProvider, OAuthClientId, OAuthSecret, LegacyMode, bReadOnly);
+		ParseHttpCacheParams(NodeName, Entry, IniFilename, IniSection, Host, bResolveHostCanonicalName, Namespace, StructuredNamespace, OAuthProvider, OAuthClientId, OAuthSecret, OAuthScope, LegacyMode, bReadOnly);
 
 		if (Host.IsEmpty())
 		{
@@ -900,6 +925,11 @@ public:
 			}
 		}
 
+		if (OAuthScope.IsEmpty())
+		{
+			OAuthScope = TEXT("cache_access");
+		}
+		
 		FDerivedDataBackendInterface::ESpeedClass ForceSpeedClass = FDerivedDataBackendInterface::ESpeedClass::Unknown;
 		FString ForceSpeedClassValue;
 		if (FParse::Value(FCommandLine::Get(), TEXT("HttpForceSpeedClass="), ForceSpeedClassValue))
@@ -932,7 +962,7 @@ public:
 		}
 
 		return MakeTuple(CreateHttpCacheStore(
-			NodeName, *Host, bResolveHostCanonicalName, *Namespace, *StructuredNamespace, *OAuthProvider, *OAuthClientId, *OAuthSecret,
+			NodeName, *Host, bResolveHostCanonicalName, *Namespace, *StructuredNamespace, *OAuthProvider, *OAuthClientId, *OAuthSecret, *OAuthScope,
 			ForceSpeedClass == FDerivedDataBackendInterface::ESpeedClass::Unknown ? nullptr : &ForceSpeedClass, LegacyMode, bReadOnly),
 			ECacheStoreFlags::Remote | ECacheStoreFlags::Query | (bReadOnly ? ECacheStoreFlags::None : ECacheStoreFlags::Store));
 	}
diff --git a/Engine/Source/Developer/DerivedDataCache/Private/HttpCacheStore.cpp b/Engine/Source/Developer/DerivedDataCache/Private/HttpCacheStore.cpp
--- a/Engine/Source/Developer/DerivedDataCache/Private/HttpCacheStore.cpp
+++ b/Engine/Source/Developer/DerivedDataCache/Private/HttpCacheStore.cpp
@@ -2070,6 +2070,7 @@ public:
 	 * @param OAuthProvider			Url to OAuth provider, for example "https://myprovider.com/oauth2/v1/token".
 	 * @param OAuthClientId			OAuth client identifier.
 	 * @param OAuthData				OAuth form data to send to login service. Can either be the raw form data or a Windows network file address (starting with "\\").
+	 * @param OAuthScope			OAuth scope identifier
 	 */
 	FHttpCacheStore(
 		const TCHAR* ServiceUrl, 
@@ -2079,6 +2080,7 @@ public:
 		const TCHAR* OAuthProvider, 
 		const TCHAR* OAuthClientId, 
 		const TCHAR* OAuthData,
+		const TCHAR* OAuthScope,
 		EBackendLegacyMode LegacyMode,
 		bool bReadOnly);
 
@@ -2146,6 +2148,7 @@ public:
 	const FString& GetOAuthProvider() const { return OAuthProvider; }
 	const FString& GetOAuthClientId() const { return OAuthClientId; }
 	const FString& GetOAuthSecret() const { return OAuthSecret; }
+	const FString& GetOAuthScope() const { return OAuthScope; }
 
 private:
 	FString Domain;
@@ -2155,6 +2158,7 @@ private:
 	FString OAuthProvider;
 	FString OAuthClientId;
 	FString OAuthSecret;
+	FString OAuthScope;
 	FCriticalSection AccessCs;
 	FDerivedDataCacheUsageStats UsageStats;
 	FBackendDebugOptions DebugOptions;
@@ -2240,6 +2244,7 @@ FHttpCacheStore::FHttpCacheStore(
 	const TCHAR* InOAuthProvider,
 	const TCHAR* InOAuthClientId,
 	const TCHAR* InOAuthSecret,
+	const TCHAR* InOAuthScope,
 	const EBackendLegacyMode InLegacyMode,
 	const bool bInReadOnly)
 	: Domain(InServiceUrl)
@@ -2249,6 +2254,7 @@ FHttpCacheStore::FHttpCacheStore(
 	, OAuthProvider(InOAuthProvider)
 	, OAuthClientId(InOAuthClientId)
 	, OAuthSecret(InOAuthSecret)
+	, OAuthScope(InOAuthScope)
 	, Access(nullptr)
 	, bIsUsable(false)
 	, bReadOnly(bInReadOnly)
@@ -2442,8 +2448,9 @@ bool FHttpCacheStore::AcquireAccessToken()
 			}
 
 			FString OAuthFormData = FString::Printf(
-				TEXT("client_id=%s&scope=cache_access&grant_type=client_credentials&client_secret=%s"),
+				TEXT("client_id=%s&scope=%s&grant_type=client_credentials&client_secret=%s"),
 				*OAuthClientId,
+				*OAuthScope,
 				*OAuthSecret
 			);
 
@@ -4194,12 +4201,13 @@ ILegacyCacheStore* CreateHttpCacheStore(
 	const TCHAR* OAuthProvider,
 	const TCHAR* OAuthClientId,
 	const TCHAR* OAuthData,
+	const TCHAR* OAuthScope,
 	const FDerivedDataBackendInterface::ESpeedClass* ForceSpeedClass,
 	EBackendLegacyMode LegacyMode,
 	bool bReadOnly)
 {
 #if WITH_HTTP_DDC_BACKEND
-	FHttpCacheStore* Backend = new FHttpCacheStore(ServiceUrl, bResolveHostCanonicalName, Namespace, StructuredNamespace, OAuthProvider, OAuthClientId, OAuthData, LegacyMode, bReadOnly);
+	FHttpCacheStore* Backend = new FHttpCacheStore(ServiceUrl, bResolveHostCanonicalName, Namespace, StructuredNamespace, OAuthProvider, OAuthClientId, OAuthData, OAuthScope, LegacyMode, bReadOnly);
 	if (Backend->IsUsable())
 	{
 		return Backend;
@@ -4218,6 +4226,7 @@ FDerivedDataBackendInterface* GetAnyHttpCacheStore(
 	FString& OutOAuthProvider,
 	FString& OutOAuthClientId,
 	FString& OutOAuthSecret,
+	FString& OutOAuthScope,
 	FString& OutNamespace,
 	FString& OutStructuredNamespace)
 {
@@ -4228,6 +4237,7 @@ FDerivedDataBackendInterface* GetAnyHttpCacheStore(
 		OutOAuthProvider = HttpBackend->GetOAuthProvider();
 		OutOAuthClientId = HttpBackend->GetOAuthClientId();
 		OutOAuthSecret = HttpBackend->GetOAuthSecret();
+		OutOAuthScope = HttpBackend->GetOAuthScope();
 		OutNamespace = HttpBackend->GetNamespace();
 		OutStructuredNamespace = HttpBackend->GetStructuredNamespace();
 
diff --git a/Engine/Source/Developer/DerivedDataCache/Private/Tests/HttpCacheStoreTest.cpp b/Engine/Source/Developer/DerivedDataCache/Private/Tests/HttpCacheStoreTest.cpp
--- a/Engine/Source/Developer/DerivedDataCache/Private/Tests/HttpCacheStoreTest.cpp
+++ b/Engine/Source/Developer/DerivedDataCache/Private/Tests/HttpCacheStoreTest.cpp
@@ -43,6 +43,7 @@ FDerivedDataBackendInterface* GetAnyHttpCacheStore(
 	FString& OutOAuthProvider,
 	FString& OutOAuthClientId,
 	FString& OutOAuthSecret,
+	FString& OutOAuthScope,
 	FString& OutNamespace,
 	FString& OutStructuredNamespace);
 
@@ -157,7 +158,7 @@ protected:
 	static FDerivedDataBackendInterface* GetTestBackend()
 	{
 		static FDerivedDataBackendInterface* CachedBackend = GetAnyHttpCacheStore(
-			TestDomain, TestOAuthProvider, TestOAuthClientId, TestOAuthSecret, TestNamespace, TestStructuredNamespace);
+			TestDomain, TestOAuthProvider, TestOAuthClientId, TestOAuthSecret, TestOAuthScope, TestNamespace, TestStructuredNamespace);
 		return CachedBackend;
 	}
 
@@ -486,6 +487,7 @@ protected:
 	static inline FString TestOAuthProvider;
 	static inline FString TestOAuthClientId;
 	static inline FString TestOAuthSecret;
+	static inline FString TestOAuthScope;
 	static inline FString TestNamespace;
 	static inline FString TestStructuredNamespace;
 };
