Fix bug  where the Content-Type header had multiple conflicting entries when posting compact binary during Horde Storage communications:

+		[1]	L"Content-Type: application/x-ue-cb"	FString
+		[2]	L"Content-Type: application/x-www-form-urlencoded"	FString

#rnx

diff --git a/Engine/Source/Developer/DerivedDataCache/Private/HttpCacheStore.cpp b/Engine/Source/Developer/DerivedDataCache/Private/HttpCacheStore.cpp
--- a/Engine/Source/Developer/DerivedDataCache/Private/HttpCacheStore.cpp
+++ b/Engine/Source/Developer/DerivedDataCache/Private/HttpCacheStore.cpp
@@ -180,6 +180,7 @@ public:
 		PutCompactBinary,
 		PutCompressedBlob,
 		Post,
+		PostCompactBinary,
 		PostJson,
 		Delete,
 		Head
@@ -313,7 +314,7 @@ public:
 	template<RequestVerb V>
 	Result PerformBlockingUpload(const TCHAR* Uri, TArrayView<const uint8> Buffer, TConstArrayView<long> ExpectedErrorCodes = {})
 	{
-		static_assert(V == Put || V == PutCompactBinary || V == PutCompressedBlob || V == Post || V == PostJson, "Upload should use either Put or Post verbs.");
+		static_assert(V == Put || V == PutCompactBinary || V == PutCompressedBlob || V == Post || V == PostCompactBinary || V == PostJson, "Upload should use either Put or Post verbs.");
 		
 		uint32 ContentLength = 0u;
 
@@ -338,13 +339,24 @@ public:
 			ContentLength = Buffer.Num();
 			ReadDataView = Buffer;
 		}
-		else if (V == Post || V == PostJson)
+		else if (V == Post || V == PostCompactBinary || V == PostJson)
 		{
 			curl_easy_setopt(Curl, CURLOPT_POST, 1L);
 			curl_easy_setopt(Curl, CURLOPT_INFILESIZE, Buffer.Num());
 			curl_easy_setopt(Curl, CURLOPT_READDATA, this);
 			curl_easy_setopt(Curl, CURLOPT_READFUNCTION, StaticReadFn);
-			Headers.Add(V == Post ? FString(TEXT("Content-Type: application/x-www-form-urlencoded")) : FString(TEXT("Content-Type: application/json")));
+			if constexpr (V == PostCompactBinary)
+			{
+				Headers.Add(FString(TEXT("Content-Type: application/x-ue-cb")));
+			}
+			else if constexpr (V == PostJson)
+			{
+				Headers.Add(FString(TEXT("Content-Type: application/json")));
+			}
+			else
+			{
+				Headers.Add(FString(TEXT("Content-Type: application/x-www-form-urlencoded")));
+			}
 			ContentLength = Buffer.Num();
 			ReadDataView = Buffer;
 		}
@@ -601,6 +613,7 @@ private:
 				AdditionalInfo = FString::Printf(TEXT("Sent: %d bytes."), BytesSent);
 				break;
 			case Post:
+			case PostCompactBinary:
 			case PostJson:
 				bSuccess = (ExpectedErrorCodes.Contains(ResponseCode) || IsSuccessResponse(ResponseCode));
 				VerbStr = TEXT("posting");
@@ -3452,8 +3465,7 @@ TArray<FValue> FHttpCacheStore::RefCachedDataProbablyExistsBatch(
 		}
 
 		Request->SetHeader(TEXT("Accept"), TEXT("application/x-ue-cb"));
-		Request->SetHeader(TEXT("Content-Type"), TEXT("application/x-ue-cb"));
-		const FHttpRequest::Result Result = Request->PerformBlockingUpload<FHttpRequest::Post>(*RefsUri, BodyBuffer);
+		const FHttpRequest::Result Result = Request->PerformBlockingUpload<FHttpRequest::PostCompactBinary>(*RefsUri, BodyBuffer);
 		ResponseCode = Request->GetResponseCode();
 
 		if (FHttpRequest::IsSuccessResponse(ResponseCode))
