DDC: Added -DDC-Verify arguments to control determinism verification

-DDC-Verify will verify determinism on every access to the cache.
-DDC-Verify=StaticMesh+Texture will verify determinism on every access static mesh and texture keys.
[REDACTED_EMAIL]+Texture will verify determinism on 12.5% of accesses to static mesh keys and every access to texture keys.
-DDC-VerifyRate=5 -DDC-Verify=Texture will verify detereminism on every access to texture keys and 5% of other keys.
-DDC-VerifyRate=5 -DDC-Verify will verify determinism on 5% of keys.
-DDC-VerifySalt=PositiveInt32 will set the salt used to match keys to verify. A key always has the same verification behavior with the same salt.

diff --git a/Engine/Source/Developer/DerivedDataCache/Internal/DerivedDataCacheKeyFilter.h b/Engine/Source/Developer/DerivedDataCache/Internal/DerivedDataCacheKeyFilter.h
new file mode 100644
--- /dev/null
+++ b/Engine/Source/Developer/DerivedDataCache/Internal/DerivedDataCacheKeyFilter.h
@@ -0,0 +1,60 @@
+// Copyright Epic Games, Inc. All Rights Reserved.
+
+#pragma once
+
+#include "CoreTypes.h"
+#include "Templates/PimplPtr.h"
+
+#define UE_API DERIVEDDATACACHE_API
+
+namespace UE::DerivedData { struct FCacheKey; }
+namespace UE::DerivedData::Private { struct FCacheKeyFilterState; }
+
+namespace UE::DerivedData
+{
+
+/** A type that filters cache keys by their bucket/type and at a configurable match rate. */
+class FCacheKeyFilter
+{
+public:
+	/**
+	 * Parses a filter from a configuration string.
+	 *
+	 * Example:
+	 * - Config = "[REDACTED_EMAIL]+StaticMesh -Unused -DDC-Verify=Audio@75"
+	 * - Prefix = "-DDC-Verify"
+	 * - DefaultRate = 50.0f
+	 * - Filter matches Texture at a rate of 2.5%, StaticMesh at a rate of 100%,
+	 *   Audio at a rate of 75%, and everything else at a rate of 50%.
+	 *
+	 * @param Config        The configuration in the form of command line parameters.
+	 * @param Prefix        The prefix of the parameter to parse the filter from.
+	 * @param DefaultRate   The default rate at which to apply the filter. [0, 100]
+	 */
+	UE_API static FCacheKeyFilter Parse(const TCHAR* Config, const TCHAR* Prefix, float DefaultRate = 0.0f);
+
+	/**
+	 * Accesses the salt that is combined with the cache key for filtering.
+	 *
+	 * The filter will always produce the same output for a fixed pair of salt and key. To change
+	 * the behavior of the filter between sessions, the salt is assigned randomly by default.
+	 */
+	UE_API void SetSalt(uint32 Salt);
+	UE_API uint32 GetSalt() const;
+
+	/** Returns true if the filter matches the key. The output is the same until the salt changes. */
+	UE_API bool IsMatch(const FCacheKey& Key) const;
+
+	/** Returns true if the filter matches the legacy key. The output is the same until the salt changes. */
+	UE_API bool IsLegacyMatch(const TCHAR* LegacyKey) const;
+
+	/** Returns true if this filter can ever match a cache key. */
+	inline explicit operator bool() const { return State.IsValid(); }
+
+private:
+	TPimplPtr<Private::FCacheKeyFilterState> State;
+};
+
+} // UE::DerivedData
+
+#undef UE_API
diff --git a/Engine/Source/Developer/DerivedDataCache/Private/DerivedDataBackends.cpp b/Engine/Source/Developer/DerivedDataCache/Private/DerivedDataBackends.cpp
--- a/Engine/Source/Developer/DerivedDataCache/Private/DerivedDataBackends.cpp
+++ b/Engine/Source/Developer/DerivedDataCache/Private/DerivedDataBackends.cpp
@@ -88,6 +88,7 @@ public:
 		, BootCache(nullptr)
 		, WritePakCache(nullptr)
 		, AsyncNode(nullptr)
+		, VerifyNode(nullptr)
 		, Hierarchy(nullptr)
 		, bUsingSharedDDC(false)
 		, bIsShuttingDown(false)
@@ -180,6 +181,19 @@ public:
 			RootNode.Key = AsyncNode;
 		}
 
+		// Create a Verify node when using -DDC-Verify[=Type1[@Rate2][+Type2[@Rate2]...]].
+		if (!VerifyNode)
+		{
+			FString VerifyArg;
+			if (FParse::Value(FCommandLine::Get(), TEXT("-DDC-Verify="), VerifyArg) ||
+				FParse::Param(FCommandLine::Get(), TEXT("-DDC-Verify")))
+			{
+				VerifyNode = CreateCacheStoreVerify(RootNode.Key, /*bPutOnError*/ false);
+				CreatedNodes.AddUnique(VerifyNode);
+				RootNode.Key = VerifyNode;
+			}
+		}
+
 		if (MaxKeyLength == 0)
 		{
 			MaxKeyLength = MAX_BACKEND_KEY_LENGTH;
@@ -276,7 +290,15 @@ public:
 				}
 				else if (NodeType == TEXT("Verify"))
 				{
-					ParsedNode = ParseVerify(*NodeName, *Entry, IniFilename, IniSection, InParsedNodes);
+					if (VerifyNode == nullptr)
+					{
+						ParsedNode = ParseVerify(*NodeName, *Entry, IniFilename, IniSection, InParsedNodes);
+						VerifyNode = ParsedNode.Key;
+					}
+					else
+					{
+						UE_LOG(LogDerivedDataCache, Warning, TEXT("Unable to create %s Verify because only one Verify node is supported."), *NodeName);
+					}
 				}
 				else if (NodeType == TEXT("ReadPak"))
 				{
@@ -1300,6 +1322,7 @@ private:
 	IMemoryCacheStore*	BootCache;
 	IPakFileCacheStore* WritePakCache;
 	ILegacyCacheStore*	AsyncNode;
+	ILegacyCacheStore*	VerifyNode;
 	ICacheStoreOwner* Hierarchy;
 	/** Support for multiple read only pak files. */
 	TArray<IPakFileCacheStore*> ReadPakCache;
diff --git a/Engine/Source/Developer/DerivedDataCache/Private/DerivedDataCacheKeyFilter.cpp b/Engine/Source/Developer/DerivedDataCache/Private/DerivedDataCacheKeyFilter.cpp
new file mode 100644
--- /dev/null
+++ b/Engine/Source/Developer/DerivedDataCache/Private/DerivedDataCacheKeyFilter.cpp
@@ -0,0 +1,199 @@
+// Copyright Epic Games, Inc. All Rights Reserved.
+
+#include "DerivedDataCacheKeyFilter.h"
+
+#include "Algo/Find.h"
+#include "Containers/StringConv.h"
+#include "DerivedDataCacheKey.h"
+#include "DerivedDataSharedString.h"
+#include "Hash/xxhash.h"
+#include "Math/UnrealMathUtility.h"
+#include "Misc/Guid.h"
+#include "Misc/Parse.h"
+#include "Misc/StringBuilder.h"
+#include "String/Find.h"
+#include "String/LexFromString.h"
+#include "String/ParseTokens.h"
+
+namespace UE::DerivedData
+{
+
+struct Private::FCacheKeyFilterState
+{
+	struct FTypeRate
+	{
+		FUtf8SharedString Type;
+		uint32 Rate = 0;
+	};
+
+	TArray<FTypeRate, TInlineAllocator<2>> Types;
+	uint32 DefaultRate = 0;
+	uint32 Salt = 0;
+
+	inline uint32 ApplySalt(const uint32 Hash) const
+	{
+		return Hash * Salt;
+	}
+
+	static inline uint32 ConvertMatchRate(const double Rate)
+	{
+		return uint32(0.01 * Rate * MAX_uint32);
+	}
+
+	static bool TryParseTypeRate(const FStringView TypeConfig, FTypeRate& OutTypeRate);
+};
+
+bool Private::FCacheKeyFilterState::TryParseTypeRate(const FStringView TypeConfig, FTypeRate& OutTypeRate)
+{
+	FStringView TypeView;
+	FStringView RateView = TEXTVIEW("100");
+
+	String::ParseTokens(TypeConfig, TEXT('@'), [&TypeView, &RateView](const FStringView TypeOrRate)
+	{
+		(TypeView.IsEmpty() ? TypeView : RateView) = TypeOrRate;
+	}, String::EParseTokensOptions::Trim);
+
+	double Rate;
+	LexFromString(Rate, RateView);
+	if (TypeView.IsEmpty() || Rate < 0.0 || Rate > 100.0)
+	{
+		return false;
+	}
+
+	OutTypeRate.Type = FUtf8StringView(FTCHARToUTF8(TypeView));
+	OutTypeRate.Rate = ConvertMatchRate(Rate);
+	return true;
+}
+
+FCacheKeyFilter FCacheKeyFilter::Parse(const TCHAR* const Config, const TCHAR* const Prefix, const float DefaultRate)
+{
+	using namespace UE::DerivedData::Private;
+
+	TArray<FCacheKeyFilterState::FTypeRate, TInlineAllocator<16>> Types;
+
+	FString TypeConfigArray;
+	for (FStringView ConfigView(Config), PrefixView(Prefix); FParse::Value(ConfigView.GetData(), Prefix, TypeConfigArray);)
+	{
+		const int32 PrefixIndex = String::FindFirst(ConfigView, PrefixView, ESearchCase::IgnoreCase);
+		ConfigView.RightChopInline(PrefixIndex + PrefixView.Len() + TypeConfigArray.Len());
+		String::ParseTokens(TypeConfigArray, TEXT('+'), [&Types](const FStringView TypeConfig)
+		{
+			FCacheKeyFilterState::FTypeRate TypeRate;
+			if (FCacheKeyFilterState::TryParseTypeRate(TypeConfig, TypeRate))
+			{
+				Types.Emplace(TypeRate);
+			}
+		}, String::EParseTokensOptions::SkipEmpty | String::EParseTokensOptions::Trim);
+	}
+
+	if (Types.IsEmpty() && DefaultRate == 0.0f)
+	{
+		return {};
+	}
+
+	FCacheKeyFilter Filter;
+	Filter.State = MakePimpl<FCacheKeyFilterState>();
+	Filter.State->Types = MoveTemp(Types);
+	Filter.State->DefaultRate = FCacheKeyFilterState::ConvertMatchRate(FMath::Clamp(DefaultRate, 0.0f, 100.0f));
+	Filter.SetSalt(0);
+	return Filter;
+}
+
+void FCacheKeyFilter::SetSalt(uint32 Salt)
+{
+	// Generate a random salt in the range [1, MAX_int32].
+	// Zero is invalid because the salt is multiplied with the key hash.
+	// Values above MAX_int32 are avoided because FParse::Value cannot parse values in that range.
+	Salt &= MAX_int32;
+	while (Salt == 0)
+	{
+		// A new guid is the most reliable random value that the engine can generate.
+		// Other options are not consistently seeded in a way that guarantees a random value when this executes.
+		const FGuid Guid = FGuid::NewGuid();
+		Salt = uint32(FXxHash64::HashBuffer(&Guid, sizeof(FGuid)).Hash & MAX_int32);
+	}
+	if (State)
+	{
+		State->Salt = Salt;
+	}
+}
+
+uint32 FCacheKeyFilter::GetSalt() const
+{
+	return State ? State->Salt : 0;
+}
+
+bool FCacheKeyFilter::IsMatch(const FCacheKey& Key) const
+{
+	using namespace UE::DerivedData::Private;
+
+	if (!State)
+	{
+		return false;
+	}
+
+	uint32 TargetRate = State->DefaultRate;
+
+	if (!State->Types.IsEmpty())
+	{
+		TUtf8StringBuilder<256> Bucket;
+		Bucket << Key.Bucket;
+		if (Bucket.ToView().StartsWith(UTF8TEXTVIEW("Legacy")))
+		{
+			Bucket.RemoveAt(0, UTF8TEXTVIEW("Legacy").Len());
+		}
+
+		using FTypeRate = FCacheKeyFilterState::FTypeRate;
+		if (const FTypeRate* TypeRate = Algo::FindBy(State->Types, Bucket.ToView(), &FTypeRate::Type))
+		{
+			TargetRate = TypeRate->Rate;
+		}
+	}
+
+	if (TargetRate == 0)
+	{
+		return false;
+	}
+	if (TargetRate == MAX_uint32)
+	{
+		return true;
+	}
+	return State->ApplySalt(GetTypeHash(Key.Hash)) < TargetRate;
+}
+
+bool FCacheKeyFilter::IsLegacyMatch(const TCHAR* LegacyKey) const
+{
+	using namespace UE::DerivedData::Private;
+
+	if (!State)
+	{
+		return false;
+	}
+
+	const FStringView LegacyKeyView(LegacyKey);
+	uint32 TargetRate = State->DefaultRate;
+
+	if (!State->Types.IsEmpty())
+	{
+		TUtf8StringBuilder<256> Bucket;
+		Bucket << LegacyKeyView.Left(String::FindFirstChar(LegacyKeyView, TEXT('_')));
+
+		using FTypeRate = FCacheKeyFilterState::FTypeRate;
+		if (const FTypeRate* TypeRate = Algo::FindBy(State->Types, Bucket.ToView(), &FTypeRate::Type))
+		{
+			TargetRate = TypeRate->Rate;
+		}
+	}
+
+	if (TargetRate == 0)
+	{
+		return false;
+	}
+	if (TargetRate == MAX_uint32)
+	{
+		return true;
+	}
+	return State->ApplySalt(uint32(FXxHash64::HashBuffer(MakeMemoryView(LegacyKeyView)).Hash)) < TargetRate;
+}
+
+} // UE::DerivedData
diff --git a/Engine/Source/Developer/DerivedDataCache/Private/DerivedDataCacheStoreVerify.cpp b/Engine/Source/Developer/DerivedDataCache/Private/DerivedDataCacheStoreVerify.cpp
--- a/Engine/Source/Developer/DerivedDataCache/Private/DerivedDataCacheStoreVerify.cpp
+++ b/Engine/Source/Developer/DerivedDataCache/Private/DerivedDataCacheStoreVerify.cpp
@@ -4,13 +4,17 @@
 #include "Compression/OodleDataCompression.h"
 #include "Containers/Set.h"
 #include "DerivedDataCacheKey.h"
+#include "DerivedDataCacheKeyFilter.h"
 #include "DerivedDataLegacyCacheStore.h"
 #include "HAL/CriticalSection.h"
+#include "Misc/CommandLine.h"
+#include "Misc/Parse.h"
 #include "Misc/Paths.h"
 #include "Misc/PathViews.h"
 #include "Misc/ScopeLock.h"
 #include "Misc/ScopeRWLock.h"
 #include "Misc/FileHelper.h"
+#include "String/Find.h"
 #include "Templates/Tuple.h"
 #include "Templates/UniquePtr.h"
 
@@ -30,6 +34,35 @@ public:
 		, bPutOnError(bInPutOnError)
 	{
 		check(InnerCache);
+
+		const TCHAR* const CommandLine = FCommandLine::Get();
+
+		const bool bDefaultMatch = FParse::Param(CommandLine, TEXT("-DDC-Verify")) ||
+			String::FindFirst(CommandLine, TEXT("-DDC-Verify="), ESearchCase::IgnoreCase) == INDEX_NONE;
+		float DefaultRate = bDefaultMatch ? 100.0f : 0.0f;
+		FParse::Value(CommandLine, TEXT("-DDC-VerifyRate="), DefaultRate);
+
+		Filter = FCacheKeyFilter::Parse(CommandLine, TEXT("-DDC-Verify="), DefaultRate);
+
+		uint32 Salt;
+		if (FParse::Value(CommandLine, TEXT("-DDC-VerifySalt="), Salt))
+		{
+			if (Salt == 0)
+			{
+				UE_LOG(LogDerivedDataCache, Warning,
+					TEXT("Verify: Ignoring salt of 0. The salt must be a positive integer."));
+			}
+			else
+			{
+				Filter.SetSalt(Salt);
+			}
+		}
+
+		if (Filter)
+		{
+			UE_LOG(LogDerivedDataCache, Display,
+				TEXT("Verify: Using salt -DDC-VerifySalt=%d to filter cache keys to verify."), Filter.GetSalt());
+		}
 	}
 
 	void Put(
@@ -138,8 +171,9 @@ private:
 
 private:
 	ILegacyCacheStore* InnerCache;
-	FCriticalSection SynchronizationObject;
+	FCriticalSection AlreadyTestedLock;
 	TSet<FCacheKey> AlreadyTested;
+	FCacheKeyFilter Filter;
 	bool bPutOnError;
 };
 
@@ -151,12 +185,16 @@ void FCacheStoreVerify::Put(
 	TUniquePtr<FVerifyPutState> State = MakeUnique<FVerifyPutState>();
 	State->VerifyRequests.Reserve(Requests.Num());
 	{
-		FScopeLock ScopeLock(&SynchronizationObject);
+		FScopeLock Lock(&AlreadyTestedLock);
 		for (const FCachePutRequest& Request : Requests)
 		{
-			bool bAlreadyTested = false;
-			AlreadyTested.Add(Request.Record.GetKey(), &bAlreadyTested);
-			(bAlreadyTested ? State->ForwardRequests : State->VerifyRequests).Add(Request);
+			const FCacheKey& Key = Request.Record.GetKey();
+			bool bForward = !Filter.IsMatch(Key);
+			if (!bForward)
+			{
+				AlreadyTested.Add(Key, &bForward);
+			}
+			(bForward ? State->ForwardRequests : State->VerifyRequests).Add(Request);
 		}
 	}
 
@@ -288,11 +326,11 @@ void FCacheStoreVerify::Get(
 	ForwardRequests.Reserve(Requests.Num());
 	VerifyRequests.Reserve(Requests.Num());
 	{
-		FScopeLock ScopeLock(&SynchronizationObject);
+		FScopeLock Lock(&AlreadyTestedLock);
 		for (const FCacheGetRequest& Request : Requests)
 		{
-			bool bAlreadyTested = AlreadyTested.Contains(Request.Key);
-			(bAlreadyTested ? ForwardRequests : VerifyRequests).Add(Request);
+			const bool bForward = !Filter.IsMatch(Request.Key) || AlreadyTested.Contains(Request.Key);
+			(bForward ? ForwardRequests : VerifyRequests).Add(Request);
 		}
 	}
 
@@ -312,12 +350,15 @@ void FCacheStoreVerify::PutValue(
 	TUniquePtr<FVerifyPutValueState> State = MakeUnique<FVerifyPutValueState>();
 	State->VerifyRequests.Reserve(Requests.Num());
 	{
-		FScopeLock ScopeLock(&SynchronizationObject);
+		FScopeLock Lock(&AlreadyTestedLock);
 		for (const FCachePutValueRequest& Request : Requests)
 		{
-			bool bAlreadyTested = false;
-			AlreadyTested.Add(Request.Key, &bAlreadyTested);
-			(bAlreadyTested ? State->ForwardRequests : State->VerifyRequests).Add(Request);
+			bool bForward = !Filter.IsMatch(Request.Key);
+			if (!bForward)
+			{
+				AlreadyTested.Add(Request.Key, &bForward);
+			}
+			(bForward ? State->ForwardRequests : State->VerifyRequests).Add(Request);
 		}
 	}
 
@@ -448,11 +489,11 @@ void FCacheStoreVerify::GetValue(
 	ForwardRequests.Reserve(Requests.Num());
 	VerifyRequests.Reserve(Requests.Num());
 	{
-		FScopeLock ScopeLock(&SynchronizationObject);
+		FScopeLock Lock(&AlreadyTestedLock);
 		for (const FCacheGetValueRequest& Request : Requests)
 		{
-			bool bAlreadyTested = AlreadyTested.Contains(Request.Key);
-			(bAlreadyTested ? ForwardRequests : VerifyRequests).Add(Request);
+			const bool bForward = !Filter.IsMatch(Request.Key) || AlreadyTested.Contains(Request.Key);
+			(bForward ? ForwardRequests : VerifyRequests).Add(Request);
 		}
 	}
 
@@ -474,11 +515,11 @@ void FCacheStoreVerify::GetChunks(
 	ForwardRequests.Reserve(Requests.Num());
 	VerifyRequests.Reserve(Requests.Num());
 	{
-		FScopeLock ScopeLock(&SynchronizationObject);
+		FScopeLock Lock(&AlreadyTestedLock);
 		for (const FCacheGetChunkRequest& Request : Requests)
 		{
-			bool bAlreadyTested = AlreadyTested.Contains(Request.Key);
-			(bAlreadyTested ? ForwardRequests : VerifyRequests).Add(Request);
+			const bool bForward = !Filter.IsMatch(Request.Key) || AlreadyTested.Contains(Request.Key);
+			(bForward ? ForwardRequests : VerifyRequests).Add(Request);
 		}
 	}
 
@@ -498,12 +539,16 @@ void FCacheStoreVerify::LegacyPut(
 	TUniquePtr<FVerifyLegacyPutState> State = MakeUnique<FVerifyLegacyPutState>();
 	State->VerifyRequests.Reserve(Requests.Num());
 	{
-		FScopeLock ScopeLock(&SynchronizationObject);
+		FScopeLock Lock(&AlreadyTestedLock);
 		for (const FLegacyCachePutRequest& Request : Requests)
 		{
-			bool bAlreadyTested = false;
-			AlreadyTested.Add(Request.Key.GetKey(), &bAlreadyTested);
-			(bAlreadyTested ? State->ForwardRequests : State->VerifyRequests).Add(Request);
+			const FCacheKey& Key = Request.Key.GetKey();
+			bool bForward = !Filter.IsMatch(Key);
+			if (!bForward)
+			{
+				AlreadyTested.Add(Key, &bForward);
+			}
+			(bForward ? State->ForwardRequests : State->VerifyRequests).Add(Request);
 		}
 	}
 
@@ -602,11 +647,12 @@ void FCacheStoreVerify::LegacyGet(
 	ForwardRequests.Reserve(Requests.Num());
 	VerifyRequests.Reserve(Requests.Num());
 	{
-		FScopeLock ScopeLock(&SynchronizationObject);
+		FScopeLock Lock(&AlreadyTestedLock);
 		for (const FLegacyCacheGetRequest& Request : Requests)
 		{
-			bool bAlreadyTested = AlreadyTested.Contains(Request.Key.GetKey());
-			(bAlreadyTested ? ForwardRequests : VerifyRequests).Add(Request);
+			const FCacheKey& Key = Request.Key.GetKey();
+			const bool bForward = !Filter.IsMatch(Key) || AlreadyTested.Contains(Key);
+			(bForward ? ForwardRequests : VerifyRequests).Add(Request);
 		}
 	}
 
