DDC: Extended -DDC-MissTypes to match -DDC-Verify

In the examples below, <Node> is the name of a node in the cache store graph, or All to apply to every node.
-DDC-<Node>-MissTypes=StaticMesh+Texture will simulate a miss on every access to to static mesh and texture keys.
-DDC-<Node>[REDACTED_EMAIL]+Texture will simulate misses on 12.5% of accesses to static mesh keys and every access to texture keys.
-DDC-<Node>-MissRate=5 -DDC-<Node>-MissTypes=Texture will simulate misses on every access to texture keys and 5% of other keys.
-DDC-<Node>-MissRate=5 will simulate misses on 5% of keys.
-DDC-<Node>-MissSalt=PositiveInt32 will set the salt used to match keys to simulate misses on. A key always has the same simulated miss behavior with the same salt.

diff --git a/Engine/Source/Developer/DerivedDataCache/Internal/DerivedDataCacheKeyFilter.h b/Engine/Source/Developer/DerivedDataCache/Internal/DerivedDataCacheKeyFilter.h
--- a/Engine/Source/Developer/DerivedDataCache/Internal/DerivedDataCacheKeyFilter.h
+++ b/Engine/Source/Developer/DerivedDataCache/Internal/DerivedDataCacheKeyFilter.h
@@ -52,7 +52,7 @@ public:
 	inline explicit operator bool() const { return State.IsValid(); }
 
 private:
-	TPimplPtr<Private::FCacheKeyFilterState> State;
+	TPimplPtr<Private::FCacheKeyFilterState, EPimplPtrMode::DeepCopy> State;
 };
 
 } // UE::DerivedData
diff --git a/Engine/Source/Developer/DerivedDataCache/Private/DerivedDataBackendInterface.h b/Engine/Source/Developer/DerivedDataCache/Private/DerivedDataBackendInterface.h
--- a/Engine/Source/Developer/DerivedDataCache/Private/DerivedDataBackendInterface.h
+++ b/Engine/Source/Developer/DerivedDataCache/Private/DerivedDataBackendInterface.h
@@ -7,6 +7,7 @@
 #include "Containers/BitArray.h"
 #include "Containers/StringView.h"
 #include "DerivedDataCache.h"
+#include "DerivedDataCacheKeyFilter.h"
 #include "DerivedDataLegacyCacheStore.h"
 #include "HAL/LowLevelMemTracker.h"
 #include "Stats/Stats.h"
@@ -63,31 +64,20 @@ enum class EBackendSpeedClass
 /** Debug options that can be applied to backends to simulate different behavior */
 struct FBackendDebugOptions
 {
-	/** Percentage of requests that should result in random misses */
-	int					RandomMissRate;
-
 	/** Apply behavior of this speed class */
 	EBackendSpeedClass	SpeedClass;
 
-	/** Types of DDC entries that should always be a miss */
-	TArray<FString>		SimulateMissTypes;
+	/** Filter to control the keys for which to simulate a miss. */
+	FCacheKeyFilter SimulateMissFilter;
 
 	/** State for simulated misses. */
 	TDontCopy<TPimplPtr<Private::FBackendDebugMissState>> SimulateMissState;
 
 	FBackendDebugOptions();
 
-	/** Fill in the provided structure based on the name of the node (e.g. 'shared') and the provided token stream */
+	/** Fill in the provided structure based on the name of the node (e.g. 'Shared') and the provided token stream. */
 	static bool ParseFromTokens(FBackendDebugOptions& OutOptions, const TCHAR* InNodeName, const TCHAR* InTokens);
 
-	/**
-	 * Returns true if, according to the properties of this struct, the provided key should be treated as a miss.
-	 * Implementing that miss and accounting for any behavior impact (e.g. skipping a subsequent put) is left to
-	 * each backend.
-	 */
-	bool ShouldSimulateMiss(const TCHAR* CacheKey) { return ShouldSimulateGetMiss(CacheKey); }
-	bool ShouldSimulateMiss(const FCacheKey& CacheKey) { return ShouldSimulateGetMiss(CacheKey); }
-
 	bool ShouldSimulatePutMiss(const FCacheKey& Key);
 	bool ShouldSimulateGetMiss(const FCacheKey& Key);
 
diff --git a/Engine/Source/Developer/DerivedDataCache/Private/DerivedDataBackends.cpp b/Engine/Source/Developer/DerivedDataCache/Private/DerivedDataBackends.cpp
--- a/Engine/Source/Developer/DerivedDataCache/Private/DerivedDataBackends.cpp
+++ b/Engine/Source/Developer/DerivedDataCache/Private/DerivedDataBackends.cpp
@@ -28,6 +28,7 @@
 #include "PakFileCacheStore.h"
 #include "ProfilingDebugging/CookStats.h"
 #include "Serialization/CompactBinaryPackage.h"
+#include "String/Find.h"
 #include <atomic>
 
 DEFINE_LOG_CATEGORY(LogDerivedDataCache);
@@ -352,7 +353,7 @@ public:
 			{
 				if (!ParsedNode.Key->LegacyDebugOptions(DebugOptions))
 				{
-					UE_LOG(LogDerivedDataCache, Warning, TEXT("Node %s is ignoring one or more -DDC-<NodeName>-Option debug options"), *NodeName);
+					UE_LOG(LogDerivedDataCache, Warning, TEXT("%s: Node is ignoring one or more -DDC-<NodeName>-Option debug options."), *NodeName);
 				}
 			}
 		}
@@ -377,7 +378,7 @@ public:
 
 		if( !PakFilename.Len() )
 		{
-			UE_LOG( LogDerivedDataCache, Log, TEXT("FDerivedDataBackendGraph:  %s pak cache Filename not found in *engine.ini, will not use a pak cache."), NodeName );
+			UE_LOG( LogDerivedDataCache, Log, TEXT("FDerivedDataBackendGraph: %s pak cache Filename not found in *engine.ini, will not use a pak cache."), NodeName );
 		}
 		else
 		{
@@ -402,7 +403,7 @@ public:
 				}
 				else
 				{
-					UE_LOG( LogDerivedDataCache, Log, TEXT("FDerivedDataBackendGraph:  %s pak cache file %s not found, will not use a pak cache."), NodeName, *PakFilename );
+					UE_LOG( LogDerivedDataCache, Log, TEXT("FDerivedDataBackendGraph: %s pak cache file %s not found, will not use a pak cache."), NodeName, *PakFilename );
 				}
 			}
 		}
@@ -1555,8 +1556,7 @@ struct Private::FBackendDebugMissState
 };
 
 FBackendDebugOptions::FBackendDebugOptions()
-	: RandomMissRate(0)
-	, SpeedClass(EBackendSpeedClass::Unknown)
+	: SpeedClass(EBackendSpeedClass::Unknown)
 {
 	SimulateMissState.Get() = MakePimpl<Private::FBackendDebugMissState>();
 }
@@ -1566,59 +1566,69 @@ FBackendDebugOptions::FBackendDebugOptions()
  */
 bool FBackendDebugOptions::ParseFromTokens(FDerivedDataBackendInterface::FBackendDebugOptions& OutOptions, const TCHAR* InNodeName, const TCHAR* InInputTokens)
 {
-	// check if the input stream has any ddc options for this node
-	FString PrefixKey = FString(TEXT("-ddc-")) + InNodeName;
-
-	if (FCString::Stristr(InInputTokens, *PrefixKey) == nullptr)
+	// Check if the input stream has any DDC options for this node.
+	TStringBuilder<64> Prefix;
+	Prefix.Append(TEXTVIEW("-DDC-")).Append(InNodeName).Append(TEXTVIEW("-"));
+	if (UE::String::FindFirst(InInputTokens, Prefix, ESearchCase::IgnoreCase) == INDEX_NONE)
 	{
-		// check if it has any -ddc-all- args
-		PrefixKey = FString(TEXT("-ddc-all"));
-
-		if (FCString::Stristr(InInputTokens, *PrefixKey) == nullptr)
+		// Check if the input stream has any -DDC-All- options.
+		Prefix.Reset();
+		Prefix.Append(TEXTVIEW("-DDC-All-"));
+		if (UE::String::FindFirst(InInputTokens, Prefix, ESearchCase::IgnoreCase) == INDEX_NONE)
 		{
 			return false;
 		}
 	}
 
-	// turn -arg= into arg= for parsing
-	PrefixKey.RightChopInline(1);
-
-	/** types that can be set to ignored (-ddc-<name>-misstypes="foo+bar" etc) */
-	// look for -ddc-local-misstype=AnimSeq+Audio -ddc-shared-misstype=AnimSeq+Audio 
-	FString ArgName = FString::Printf(TEXT("%s-misstypes="), *PrefixKey);
-
-	FString TempArg;
-	FParse::Value(InInputTokens, *ArgName, TempArg);
-	TempArg.ParseIntoArray(OutOptions.SimulateMissTypes, TEXT("+"), true);
+	const int32 PrefixLen = Prefix.Len();
 
-	// look for -ddc-local-missrate=, -ddc-shared-missrate= etc
-	ArgName = FString::Printf(TEXT("%s-missrate="), *PrefixKey);
-	int MissRate = 0;
-	FParse::Value(InInputTokens, *ArgName, OutOptions.RandomMissRate);
-
-	// look for -ddc-local-speed=, -ddc-shared-speed= etc
-	ArgName = FString::Printf(TEXT("%s-speed="), *PrefixKey);
-	if (FParse::Value(InInputTokens, *ArgName, TempArg))
+	// Look for -DDC-Local-Speed=, -DDC-Shared-Speed=, etc.
+	Prefix.Append(TEXTVIEW("Speed="));
+	FString SpeedClass;
+	if (FParse::Value(InInputTokens, *Prefix, SpeedClass))
 	{
-		if (!TempArg.IsEmpty())
-		{
-			LexFromString(OutOptions.SpeedClass, *TempArg);
-		}
+		LexFromString(OutOptions.SpeedClass, *SpeedClass);
 	}
+	Prefix.RemoveAt(PrefixLen, Prefix.Len() - PrefixLen);
+
+	// Look for -DDC-Local-MissRate=, -DDC-Shared-MissRate=, etc.
+	float MissRate = 0.0f;
+	Prefix.Append(TEXTVIEW("MissRate="));
+	FParse::Value(InInputTokens, *Prefix, MissRate);
+	Prefix.RemoveAt(PrefixLen, Prefix.Len() - PrefixLen);
+
+	// Look for -DDC-Local-MissTypes=AnimSeq+Audio, -DDC-Shared-MissType=AnimSeq+Audio, etc.
+	Prefix.Append(TEXTVIEW("MissTypes="));
+	OutOptions.SimulateMissFilter = FCacheKeyFilter::Parse(InInputTokens, *Prefix, MissRate);
+	Prefix.RemoveAt(PrefixLen, Prefix.Len() - PrefixLen);
+
+	// Look for -DDC-Local-MissSalt=, -DDC-Shared-MissSalt=, etc.
+	uint32 Salt = 0;
+	Prefix.Append(TEXTVIEW("MissSalt="));
+	if (FParse::Value(InInputTokens, *Prefix, Salt))
+	{
+		OutOptions.SimulateMissFilter.SetSalt(Salt);
+	}
+	if (OutOptions.SimulateMissFilter)
+	{
+		UE_LOG(LogDerivedDataCache, Display,
+			TEXT("%s: Using salt %s%u to filter cache keys to simulate misses on."),
+			InNodeName, *Prefix, OutOptions.SimulateMissFilter.GetSalt());
+	}
+	Prefix.RemoveAt(PrefixLen, Prefix.Len() - PrefixLen);
 
 	return true;
 }
 
 bool FBackendDebugOptions::ShouldSimulatePutMiss(const TCHAR* LegacyKey)
 {
-	if (RandomMissRate == 0 && SimulateMissTypes.IsEmpty())
+	if (!SimulateMissFilter.IsLegacyMatch(LegacyKey))
 	{
 		return false;
 	}
 
-	const FName Key(LegacyKey);
-
 	Private::FBackendDebugMissState& State = *SimulateMissState.Get();
+	const FName Key(LegacyKey);
 	const uint32 KeyHash = GetTypeHash(Key);
 
 	FScopeLock Lock(&State.Lock);
@@ -1628,36 +1638,22 @@ bool FBackendDebugOptions::ShouldSimulatePutMiss(const TCHAR* LegacyKey)
 
 bool FBackendDebugOptions::ShouldSimulateGetMiss(const TCHAR* LegacyKey)
 {
-	if (RandomMissRate == 0 && SimulateMissTypes.IsEmpty())
+	if (!SimulateMissFilter.IsLegacyMatch(LegacyKey))
 	{
 		return false;
 	}
 
-	const FStringView KeyView(LegacyKey);
-	const FName Key(KeyView);
-
-	bool bMiss = (RandomMissRate >= 100);
-	if (!bMiss && !SimulateMissTypes.IsEmpty())
-	{
-		const FStringView Bucket = KeyView.Left(KeyView.Find(TEXT("_")));
-		bMiss = SimulateMissTypes.Contains(Bucket);
-	}
-	if (!bMiss && RandomMissRate > 0)
-	{
-		bMiss = FMath::RandHelper(100) < RandomMissRate;
-	}
-
 	Private::FBackendDebugMissState& State = *SimulateMissState.Get();
+	const FName Key(LegacyKey);
 	const uint32 KeyHash = GetTypeHash(Key);
 
 	FScopeLock Lock(&State.Lock);
-	const EBackendDebugKeyState KeyState = bMiss ? EBackendDebugKeyState::MissGet : EBackendDebugKeyState::HitGet;
-	return State.LegacyKeys.FindOrAddByHash(KeyHash, Key, KeyState) == EBackendDebugKeyState::MissGet;
+	return State.LegacyKeys.FindOrAddByHash(KeyHash, Key, EBackendDebugKeyState::MissGet) == EBackendDebugKeyState::MissGet;
 }
 
 bool FBackendDebugOptions::ShouldSimulatePutMiss(const FCacheKey& Key)
 {
-	if (RandomMissRate == 0 && SimulateMissTypes.IsEmpty())
+	if (!SimulateMissFilter.IsMatch(Key))
 	{
 		return false;
 	}
@@ -1672,33 +1668,16 @@ bool FBackendDebugOptions::ShouldSimulatePutMiss(const FCacheKey& Key)
 
 bool FBackendDebugOptions::ShouldSimulateGetMiss(const FCacheKey& Key)
 {
-	if (RandomMissRate == 0 && SimulateMissTypes.IsEmpty())
+	if (!SimulateMissFilter.IsMatch(Key))
 	{
 		return false;
 	}
 
-	bool bMiss = (RandomMissRate >= 100);
-	if (!bMiss && !SimulateMissTypes.IsEmpty())
-	{
-		TStringBuilder<256> Bucket;
-		Bucket << Key.Bucket;
-		if (Bucket.ToView().StartsWith(TEXTVIEW("Legacy")))
-		{
-			Bucket.RemoveAt(0, TEXTVIEW("Legacy").Len());
-		}
-		bMiss = SimulateMissTypes.Contains(Bucket.ToView());
-	}
-	if (!bMiss && RandomMissRate > 0)
-	{
-		bMiss = FMath::RandHelper(100) < RandomMissRate;
-	}
-
 	Private::FBackendDebugMissState& State = *SimulateMissState.Get();
 	const uint32 KeyHash = GetTypeHash(Key);
 
 	FScopeLock Lock(&State.Lock);
-	const EBackendDebugKeyState KeyState = bMiss ? EBackendDebugKeyState::MissGet : EBackendDebugKeyState::HitGet;
-	return State.Keys.FindOrAddByHash(KeyHash, Key, KeyState) == EBackendDebugKeyState::MissGet;
+	return State.Keys.FindOrAddByHash(KeyHash, Key, EBackendDebugKeyState::MissGet) == EBackendDebugKeyState::MissGet;
 }
 
 } // UE::DerivedData
diff --git a/Engine/Source/Developer/DerivedDataCache/Private/DerivedDataCacheKeyFilter.cpp b/Engine/Source/Developer/DerivedDataCache/Private/DerivedDataCacheKeyFilter.cpp
--- a/Engine/Source/Developer/DerivedDataCache/Private/DerivedDataCacheKeyFilter.cpp
+++ b/Engine/Source/Developer/DerivedDataCache/Private/DerivedDataCacheKeyFilter.cpp
@@ -92,7 +92,7 @@ FCacheKeyFilter FCacheKeyFilter::Parse(const TCHAR* const Config, const TCHAR* c
 	}
 
 	FCacheKeyFilter Filter;
-	Filter.State = MakePimpl<FCacheKeyFilterState>();
+	Filter.State = MakePimpl<FCacheKeyFilterState, EPimplPtrMode::DeepCopy>();
 	Filter.State->Types = MoveTemp(Types);
 	Filter.State->DefaultRate = FCacheKeyFilterState::ConvertMatchRate(FMath::Clamp(DefaultRate, 0.0f, 100.0f));
 	Filter.SetSalt(0);
diff --git a/Engine/Source/Developer/DerivedDataCache/Private/DerivedDataCacheStoreVerify.cpp b/Engine/Source/Developer/DerivedDataCache/Private/DerivedDataCacheStoreVerify.cpp
--- a/Engine/Source/Developer/DerivedDataCache/Private/DerivedDataCacheStoreVerify.cpp
+++ b/Engine/Source/Developer/DerivedDataCache/Private/DerivedDataCacheStoreVerify.cpp
@@ -61,7 +61,7 @@ public:
 		if (Filter)
 		{
 			UE_LOG(LogDerivedDataCache, Display,
-				TEXT("Verify: Using salt -DDC-VerifySalt=%d to filter cache keys to verify."), Filter.GetSalt());
+				TEXT("Verify: Using salt -DDC-VerifySalt=%u to filter cache keys to verify."), Filter.GetSalt());
 		}
 	}
 
diff --git a/Engine/Source/Developer/DerivedDataCache/Private/HttpCacheStore.cpp b/Engine/Source/Developer/DerivedDataCache/Private/HttpCacheStore.cpp
--- a/Engine/Source/Developer/DerivedDataCache/Private/HttpCacheStore.cpp
+++ b/Engine/Source/Developer/DerivedDataCache/Private/HttpCacheStore.cpp
@@ -2175,9 +2175,6 @@ private:
 	FCriticalSection AccessCs;
 	FDerivedDataCacheUsageStats UsageStats;
 	FBackendDebugOptions DebugOptions;
-	FCriticalSection MissedKeysCS;
-	TSet<FName> DebugMissedKeys;
-	TSet<FCacheKey> DebugMissedCacheKeys;
 	TUniquePtr<struct FRequestPool> GetRequestPools[2];
 	TUniquePtr<struct FRequestPool> PutRequestPools[2];
 	TUniquePtr<struct FHttpAccessToken> Access;
@@ -2198,8 +2195,6 @@ private:
 	bool IsServiceReady();
 	bool AcquireAccessToken();
 	bool ShouldRetryOnError(int64 ResponseCode);
-	bool ShouldSimulateMiss(const TCHAR* InKey);
-	bool ShouldSimulateMiss(const FCacheKey& Key);
 
 	bool PutCacheRecord(FStringView Name, const FCacheRecord& Record, const FCacheRecordPolicy& Policy, uint64& OutWriteSize);
 	uint64 PutRef(const FCbPackage& Package, const FCacheKey& Key, FStringView Bucket, bool bFinalize, TArray<FIoHash>& OutNeededBlobHashes, bool& bOutPutCompletedSuccessfully);
@@ -2536,55 +2531,6 @@ bool FHttpCacheStore::ShouldRetryOnError(int64 ResponseCode)
 	return false;
 }
 
-bool FHttpCacheStore::ShouldSimulateMiss(const TCHAR* InKey)
-{
-	if (DebugOptions.RandomMissRate == 0 && DebugOptions.SimulateMissTypes.IsEmpty())
-	{
-		return false;
-	}
-
-	const FName Key(InKey);
-	const uint32 Hash = GetTypeHash(Key);
-
-	if (FScopeLock Lock(&MissedKeysCS); DebugMissedKeys.ContainsByHash(Hash, Key))
-	{
-		return true;
-	}
-
-	if (DebugOptions.ShouldSimulateMiss(InKey))
-	{
-		FScopeLock Lock(&MissedKeysCS);
-		DebugMissedKeys.AddByHash(Hash, Key);
-		return true;
-	}
-
-	return false;
-}
-
-bool FHttpCacheStore::ShouldSimulateMiss(const FCacheKey& Key)
-{
-	if (DebugOptions.RandomMissRate == 0 && DebugOptions.SimulateMissTypes.IsEmpty())
-	{
-		return false;
-	}
-
-	const uint32 Hash = GetTypeHash(Key);
-
-	if (FScopeLock Lock(&MissedKeysCS); DebugMissedCacheKeys.ContainsByHash(Hash, Key))
-	{
-		return true;
-	}
-
-	if (DebugOptions.ShouldSimulateMiss(Key))
-	{
-		FScopeLock Lock(&MissedKeysCS);
-		DebugMissedCacheKeys.AddByHash(Hash, Key);
-		return true;
-	}
-
-	return false;
-}
-
 uint64 FHttpCacheStore::PutRef(const FCbPackage& Package, const FCacheKey& Key, FStringView Bucket, bool bFinalize, TArray<FIoHash>& OutNeededBlobHashes, bool& bOutPutCompletedSuccessfully)
 {
 	bOutPutCompletedSuccessfully = false;
@@ -2693,7 +2639,7 @@ bool FHttpCacheStore::PutCacheRecord(
 		return false;
 	}
 
-	if (ShouldSimulateMiss(Key))
+	if (DebugOptions.ShouldSimulatePutMiss(Key))
 	{
 		UE_LOG(LogDerivedDataCache, Verbose, TEXT("%s: Simulated miss for put of %s from '%.*s'"),
 			*GetName(), *WriteToString<96>(Key), Name.Len(), Name.GetData());
@@ -2821,7 +2767,7 @@ FOptionalCacheRecord FHttpCacheStore::GetCacheRecordOnly(
 		return FOptionalCacheRecord();
 	}
 
-	if (ShouldSimulateMiss(Key))
+	if (DebugOptions.ShouldSimulateGetMiss(Key))
 	{
 		UE_LOG(LogDerivedDataCache, Verbose, TEXT("%s: Simulated miss for get of %s from '%.*s'"),
 			*GetName(), *WriteToString<96>(Key), Name.Len(), Name.GetData());
@@ -2911,7 +2857,7 @@ bool FHttpCacheStore::PutCacheValue(
 		return false;
 	}
 
-	if (ShouldSimulateMiss(Key))
+	if (DebugOptions.ShouldSimulatePutMiss(Key))
 	{
 		UE_LOG(LogDerivedDataCache, Verbose, TEXT("%s: Simulated miss for put of %s from '%.*s'"),
 			*GetName(), *WriteToString<96>(Key), Name.Len(), Name.GetData());
@@ -3031,7 +2977,7 @@ bool FHttpCacheStore::GetCacheValue(
 		return false;
 	}
 
-	if (ShouldSimulateMiss(Key))
+	if (DebugOptions.ShouldSimulateGetMiss(Key))
 	{
 		UE_LOG(LogDerivedDataCache, Verbose, TEXT("%s: Simulated miss for get of %s from '%.*s'"),
 			*GetName(), *WriteToString<96>(Key), Name.Len(), Name.GetData());
@@ -3546,7 +3492,7 @@ bool FHttpCacheStore::CachedDataProbablyExists(const TCHAR* CacheKey)
 	TRACE_COUNTER_ADD(HttpDDC_Exist, int64(1));
 	COOK_STAT(auto Timer = UsageStats.TimeProbablyExists());
 
-	if (ShouldSimulateMiss(CacheKey))
+	if (DebugOptions.ShouldSimulateGetMiss(CacheKey))
 	{
 		return false;
 	}
@@ -3618,7 +3564,7 @@ TBitArray<> FHttpCacheStore::CachedDataProbablyExistsBatch(TConstArrayView<FStri
 			int32 ResultIndex = 0;
 			for (const FString& CacheKey : CacheKeys)
 			{
-				if (ShouldSimulateMiss(*CacheKey))
+				if (DebugOptions.ShouldSimulateGetMiss(*CacheKey))
 				{
 					Results[ResultIndex] = false;
 				}
@@ -3634,7 +3580,7 @@ TBitArray<> FHttpCacheStore::CachedDataProbablyExistsBatch(TConstArrayView<FStri
 			int32 ResultIndex = 0;
 			for (const FString& CacheKey : CacheKeys)
 			{
-				if (ShouldSimulateMiss(*CacheKey))
+				if (DebugOptions.ShouldSimulateGetMiss(*CacheKey))
 				{
 					Results[ResultIndex] = false;
 				}
@@ -3676,7 +3622,7 @@ TBitArray<> FHttpCacheStore::CachedDataProbablyExistsBatch(TConstArrayView<FStri
 			Exists.Reserve(CacheKeys.Num());
 			for (const FString& CacheKey : CacheKeys)
 			{
-				if (ShouldSimulateMiss(*CacheKey))
+				if (DebugOptions.ShouldSimulateGetMiss(*CacheKey))
 				{
 					Exists.Add(false);
 				}
@@ -3716,7 +3662,7 @@ bool FHttpCacheStore::GetCachedData(const TCHAR* CacheKey, TArray<uint8>& OutDat
 	TRACE_COUNTER_ADD(HttpDDC_Get, int64(1));
 	COOK_STAT(auto Timer = UsageStats.TimeGet());
 
-	if (ShouldSimulateMiss(CacheKey))
+	if (DebugOptions.ShouldSimulateGetMiss(CacheKey))
 	{
 		return false;
 	}
@@ -3783,7 +3729,7 @@ FDerivedDataBackendInterface::EPutStatus FHttpCacheStore::PutCachedData(const TC
 	}
 
 	// don't put anything we pretended didn't exist
-	if (ShouldSimulateMiss(CacheKey))
+	if (DebugOptions.ShouldSimulatePutMiss(CacheKey))
 	{
 		return EPutStatus::Skipped;
 	}
diff --git a/Engine/Source/Developer/DerivedDataCache/Private/MemoryCacheStore.cpp b/Engine/Source/Developer/DerivedDataCache/Private/MemoryCacheStore.cpp
--- a/Engine/Source/Developer/DerivedDataCache/Private/MemoryCacheStore.cpp
+++ b/Engine/Source/Developer/DerivedDataCache/Private/MemoryCacheStore.cpp
@@ -13,7 +13,6 @@
 #include "DerivedDataValue.h"
 #include "HAL/FileManager.h"
 #include "Misc/ScopeExit.h"
-#include "Misc/ScopeLock.h"
 #include "Misc/ScopeRWLock.h"
 #include "ProfilingDebugging/CookStats.h"
 #include "Serialization/CompactBinary.h"
@@ -131,16 +130,7 @@ private:
 	bool bShuttingDown  = false;
 
 protected:
-
-	bool ShouldSimulateMiss(const FCacheKey& InKey);
-
-	/** Debug Options */
 	FBackendDebugOptions DebugOptions;
-
-	/** Keys we ignored due to miss rate settings */
-	FCriticalSection MissedKeysCS;
-	TSet<FName> DebugMissedKeys;
-	TSet<FCacheKey> DebugMissedCacheKeys;
 };
 
 FMemoryCacheStore::FMemoryCacheStore(const TCHAR* InName, int64 InMaxCacheSize, bool bInCanBeDisabled)
@@ -182,30 +172,6 @@ bool FMemoryCacheStore::LegacyDebugOptions(FBackendDebugOptions& InOptions)
 	return true;
 }
 
-bool FMemoryCacheStore::ShouldSimulateMiss(const FCacheKey& Key)
-{
-	if (!bCanBeDisabled || (DebugOptions.RandomMissRate == 0 && DebugOptions.SimulateMissTypes.IsEmpty()))
-	{
-		return false;
-	}
-
-	const uint32 Hash = GetTypeHash(Key);
-
-	if (FScopeLock Lock(&MissedKeysCS); DebugMissedCacheKeys.ContainsByHash(Hash, Key))
-	{
-		return true;
-	}
-
-	if (DebugOptions.ShouldSimulateMiss(Key))
-	{
-		FScopeLock Lock(&MissedKeysCS);
-		DebugMissedCacheKeys.AddByHash(Hash, Key);
-		return true;
-	}
-
-	return false;
-}
-
 void FMemoryCacheStore::Put(
 	const TConstArrayView<FCachePutRequest> Requests,
 	IRequestOwner& Owner,
@@ -233,7 +199,7 @@ void FMemoryCacheStore::Put(
 			continue;
 		}
 
-		if (ShouldSimulateMiss(Key))
+		if (bCanBeDisabled && DebugOptions.ShouldSimulatePutMiss(Key))
 		{
 			UE_LOG(LogDerivedDataCache, Verbose, TEXT("%s: Simulated miss for put of %s from '%s'"),
 				*Name, *WriteToString<96>(Key), *Request.Name);
@@ -387,7 +353,7 @@ void FMemoryCacheStore::Get(
 		FCacheRecordComponents Components;
 		EStatus Status = EStatus::Error;
 
-		if (ShouldSimulateMiss(Key))
+		if (bCanBeDisabled && DebugOptions.ShouldSimulateGetMiss(Key))
 		{
 			UE_LOG(LogDerivedDataCache, Verbose, TEXT("%s: Simulated miss for get of %s from '%s'"),
 				*Name, *WriteToString<96>(Key), *Request.Name);
@@ -480,7 +446,7 @@ void FMemoryCacheStore::PutValue(
 			continue;
 		}
 
-		if (ShouldSimulateMiss(Key))
+		if (bCanBeDisabled && DebugOptions.ShouldSimulatePutMiss(Key))
 		{
 			UE_LOG(LogDerivedDataCache, Verbose, TEXT("%s: Simulated miss for put of %s from '%s'"),
 				*Name, *WriteToString<96>(Key), *Request.Name);
@@ -546,7 +512,7 @@ void FMemoryCacheStore::GetValue(
 
 		FValue Value;
 		EStatus Status = EStatus::Error;
-		if (ShouldSimulateMiss(Key))
+		if (bCanBeDisabled && DebugOptions.ShouldSimulateGetMiss(Key))
 		{
 			UE_LOG(LogDerivedDataCache, Verbose, TEXT("%s: Simulated miss for get of %s from '%s'"),
 				*Name, *WriteToString<96>(Key), *Request.Name);
@@ -596,7 +562,7 @@ void FMemoryCacheStore::GetChunks(
 		bool bProcessHit = false;
 		const bool bExistsOnly = EnumHasAnyFlags(Request.Policy, ECachePolicy::SkipData);
 		COOK_STAT(auto Timer = bExistsOnly ? UsageStats.TimeProbablyExists() : UsageStats.TimeGet());
-		if (ShouldSimulateMiss(Request.Key))
+		if (bCanBeDisabled && DebugOptions.ShouldSimulateGetMiss(Request.Key))
 		{
 			UE_LOG(LogDerivedDataCache, Verbose, TEXT("%s: Simulated miss for get of %s from '%s'"),
 				*Name, *WriteToString<96>(Request.Key, '/', Request.Id), *Request.Name);
@@ -689,7 +655,7 @@ void FMemoryCacheStore::LegacyPut(
 			continue;
 		}
 
-		if (ShouldSimulateMiss(Key))
+		if (bCanBeDisabled && DebugOptions.ShouldSimulatePutMiss(Key))
 		{
 			UE_LOG(LogDerivedDataCache, Verbose, TEXT("%s: Simulated miss for put of %s from '%s'"),
 				*Name, *WriteToString<96>(Key), *Request.Name);
@@ -755,7 +721,7 @@ void FMemoryCacheStore::LegacyGet(
 
 		FLegacyCacheValue Value;
 		EStatus Status = EStatus::Error;
-		if (ShouldSimulateMiss(Key))
+		if (bCanBeDisabled && DebugOptions.ShouldSimulateGetMiss(Key))
 		{
 			UE_LOG(LogDerivedDataCache, Verbose, TEXT("%s: Simulated miss for get of %s from '%s'"),
 				*Name, *WriteToString<96>(Key), *Request.Name);
diff --git a/Engine/Source/Developer/DerivedDataCache/Private/PakFileCacheStore.cpp b/Engine/Source/Developer/DerivedDataCache/Private/PakFileCacheStore.cpp
--- a/Engine/Source/Developer/DerivedDataCache/Private/PakFileCacheStore.cpp
+++ b/Engine/Source/Developer/DerivedDataCache/Private/PakFileCacheStore.cpp
@@ -117,7 +117,7 @@ public:
 
 	bool WouldCache(const TCHAR* CacheKey, TArrayView<const uint8> InData) final { return true; }
 
-	bool ApplyDebugOptions(FBackendDebugOptions& InOptions) final { return false; }
+	bool ApplyDebugOptions(FBackendDebugOptions& InOptions) final;
 
 	EBackendLegacyMode GetLegacyMode() const final { return EBackendLegacyMode::ValueOnly; }
 
@@ -190,6 +190,7 @@ private:
 
 private:
 	FDerivedDataCacheUsageStats UsageStats;
+	FBackendDebugOptions DebugOptions;
 
 	struct FCacheValue
 	{
@@ -316,6 +317,10 @@ bool FPakFileCacheStore::BackfillLowerCacheLevels() const
 bool FPakFileCacheStore::CachedDataProbablyExists(const TCHAR* CacheKey)
 {
 	COOK_STAT(auto Timer = UsageStats.TimeProbablyExists());
+	if (bClosed || DebugOptions.ShouldSimulateGetMiss(CacheKey))
+	{
+		return false;
+	}
 	FReadScopeLock ScopeLock(SynchronizationObject);
 	bool Result = CacheItems.Contains(FString(CacheKey));
 	if (Result)
@@ -328,7 +333,7 @@ bool FPakFileCacheStore::CachedDataProbablyExists(const TCHAR* CacheKey)
 bool FPakFileCacheStore::GetCachedData(const TCHAR* CacheKey, TArray<uint8>& OutData)
 {
 	COOK_STAT(auto Timer = UsageStats.TimeGet());
-	if (bClosed)
+	if (bClosed || DebugOptions.ShouldSimulateGetMiss(CacheKey))
 	{
 		return false;
 	}
@@ -384,7 +389,7 @@ bool FPakFileCacheStore::GetCachedData(const TCHAR* CacheKey, TArray<uint8>& Out
 FDerivedDataBackendInterface::EPutStatus FPakFileCacheStore::PutCachedData(const TCHAR* CacheKey, TArrayView<const uint8> InData, bool bPutEvenIfExists)
 {
 	COOK_STAT(auto Timer = UsageStats.TimePut());
-	if (!IsWritable())
+	if (!IsWritable() || DebugOptions.ShouldSimulateGetMiss(CacheKey))
 	{
 		return EPutStatus::NotCached;
 	}
@@ -684,6 +689,12 @@ TSharedRef<FDerivedDataCacheStatsNode> FPakFileCacheStore::GatherUsageStats() co
 	return Usage;
 }
 
+bool FPakFileCacheStore::ApplyDebugOptions(FBackendDebugOptions& InOptions)
+{
+	DebugOptions = InOptions;
+	return true;
+}
+
 void FPakFileCacheStore::Put(
 	const TConstArrayView<FCachePutRequest> Requests,
 	IRequestOwner& Owner,
@@ -894,12 +905,12 @@ bool FPakFileCacheStore::PutCacheRecord(
 		return false;
 	}
 
-	//if (ShouldSimulateMiss(Key))
-	//{
-	//	UE_LOG(LogDerivedDataCache, Verbose, TEXT("%s: Simulated miss for put of %s from '%.*s'"),
-	//		*CachePath, *WriteToString<96>(Key), Name.Len(), Name.GetData());
-	//	return false;
-	//}
+	if (DebugOptions.ShouldSimulatePutMiss(Key))
+	{
+		UE_LOG(LogDerivedDataCache, Verbose, TEXT("%s: Simulated miss for put of %s from '%.*s'"),
+			*CachePath, *WriteToString<96>(Key), Name.Len(), Name.GetData());
+		return false;
+	}
 
 	TStringBuilder<256> Path;
 	FPathViews::Append(Path, TEXT("Buckets"), Key);
@@ -960,12 +971,12 @@ FOptionalCacheRecord FPakFileCacheStore::GetCacheRecordOnly(
 		return FOptionalCacheRecord();
 	}
 
-	//if (ShouldSimulateMiss(Key))
-	//{
-	//	UE_LOG(LogDerivedDataCache, Verbose, TEXT("%s: Simulated miss for get of %s from '%.*s'"),
-	//		*CachePath, *WriteToString<96>(Key), Name.Len(), Name.GetData());
-	//	return FOptionalCacheRecord();
-	//}
+	if (DebugOptions.ShouldSimulateGetMiss(Key))
+	{
+		UE_LOG(LogDerivedDataCache, Verbose, TEXT("%s: Simulated miss for get of %s from '%.*s'"),
+			*CachePath, *WriteToString<96>(Key), Name.Len(), Name.GetData());
+		return FOptionalCacheRecord();
+	}
 
 	TStringBuilder<256> Path;
 	FPathViews::Append(Path, TEXT("Buckets"), Key);
@@ -1076,12 +1087,12 @@ bool FPakFileCacheStore::PutCacheValue(
 		return false;
 	}
 
-	//if (ShouldSimulateMiss(Key))
-	//{
-	//	UE_LOG(LogDerivedDataCache, Verbose, TEXT("%s: Simulated miss for put of %s from '%.*s'"),
-	//		*CachePath, *WriteToString<96>(Key), Name.Len(), Name.GetData());
-	//	return false;
-	//}
+	if (DebugOptions.ShouldSimulatePutMiss(Key))
+	{
+		UE_LOG(LogDerivedDataCache, Verbose, TEXT("%s: Simulated miss for put of %s from '%.*s'"),
+			*CachePath, *WriteToString<96>(Key), Name.Len(), Name.GetData());
+		return false;
+	}
 
 	// Check if there is an existing value package.
 	bool bValueExists = false;
@@ -1157,12 +1168,12 @@ bool FPakFileCacheStore::GetCacheValueOnly(
 		return false;
 	}
 
-	//if (ShouldSimulateMiss(Key))
-	//{
-	//	UE_LOG(LogDerivedDataCache, Verbose, TEXT("%s: Simulated miss for get of %s from '%.*s'"),
-	//		*CachePath, *WriteToString<96>(Key), Name.Len(), Name.GetData());
-	//	return false;
-	//}
+	if (DebugOptions.ShouldSimulateGetMiss(Key))
+	{
+		UE_LOG(LogDerivedDataCache, Verbose, TEXT("%s: Simulated miss for get of %s from '%.*s'"),
+			*CachePath, *WriteToString<96>(Key), Name.Len(), Name.GetData());
+		return false;
+	}
 
 	TStringBuilder<256> Path;
 	FPathViews::Append(Path, TEXT("Buckets"), Key);
diff --git a/Engine/Source/Developer/DerivedDataCache/Private/S3CacheStore.cpp b/Engine/Source/Developer/DerivedDataCache/Private/S3CacheStore.cpp
--- a/Engine/Source/Developer/DerivedDataCache/Private/S3CacheStore.cpp
+++ b/Engine/Source/Developer/DerivedDataCache/Private/S3CacheStore.cpp
@@ -326,16 +326,7 @@ private:
 	void ReadBundle(FBundle& Bundle);
 	bool FindBundleEntry(const TCHAR* CacheKey, const FBundle*& OutBundle, const FBundleEntry*& OutBundleEntry) const;
 
-	/* Debug helpers */
-	bool DidSimulateMiss(const TCHAR* InKey);
-	bool ShouldSimulateMiss(const TCHAR* InKey);
-
-	/** Debug Options */
 	FBackendDebugOptions DebugOptions;
-
-	/** Keys we ignored due to miss rate settings */
-	FCriticalSection MissedKeysCS;
-	TSet<FName> DebugMissedKeys;
 };
 
 /**
@@ -1071,7 +1062,7 @@ bool FS3CacheStore::CachedDataProbablyExists(const TCHAR* CacheKey)
 	const FBundle* Bundle;
 	const FBundleEntry* BundleEntry;
 
-	if (ShouldSimulateMiss(CacheKey))
+	if (DebugOptions.ShouldSimulateGetMiss(CacheKey))
 	{
 		return false;
 	}
@@ -1093,7 +1084,7 @@ bool FS3CacheStore::GetCachedData(const TCHAR* CacheKey, TArray<uint8>& OutData)
 	TRACE_COUNTER_ADD(S3DDC_Get, int64(1));
 	COOK_STAT(auto Timer = UsageStats.TimeGet());
 
-	if (ShouldSimulateMiss(CacheKey))
+	if (DebugOptions.ShouldSimulateGetMiss(CacheKey))
 	{
 		return false;
 	}
@@ -1340,35 +1331,6 @@ bool FS3CacheStore::ApplyDebugOptions(FBackendDebugOptions& InOptions)
 	return true;
 }
 
-bool FS3CacheStore::DidSimulateMiss(const TCHAR* InKey)
-{
-	if (DebugOptions.RandomMissRate == 0 && DebugOptions.SimulateMissTypes.IsEmpty())
-	{
-		return false;
-	}
-	FScopeLock Lock(&MissedKeysCS);
-	return DebugMissedKeys.Contains(FName(InKey));
-}
-
-bool FS3CacheStore::ShouldSimulateMiss(const TCHAR* InKey)
-{
-	// once missed, always missed
-	if (DidSimulateMiss(InKey))
-	{
-		return true;
-	}
-
-	if (DebugOptions.ShouldSimulateMiss(InKey))
-	{
-		FScopeLock Lock(&MissedKeysCS);
-		UE_LOG(LogDerivedDataCache, Verbose, TEXT("Simulating miss in %s for %s"), *GetName(), InKey);
-		DebugMissedKeys.Add(FName(InKey));
-		return true;
-	}
-
-	return false;
-}
-
 FOptionalCacheRecord FS3CacheStore::GetCacheRecordOnly(
 	const FStringView Name,
 	const FCacheKey& Key,
@@ -1390,17 +1352,16 @@ FOptionalCacheRecord FS3CacheStore::GetCacheRecordOnly(
 		return FOptionalCacheRecord();
 	}
 
-	//if (ShouldSimulateMiss(Key))
-	//{
-	//	UE_LOG(LogDerivedDataCache, Verbose, TEXT("%s: Simulated miss for get of %s from '%.*s'"),
-	//		*GetName(), *WriteToString<96>(Key), Name.Len(), Name.GetData());
-	//	return FOptionalCacheRecord();
-	//}
+	if (DebugOptions.ShouldSimulateGetMiss(Key))
+	{
+		UE_LOG(LogDerivedDataCache, Verbose, TEXT("%s: Simulated miss for get of %s from '%.*s'"),
+			*GetName(), *WriteToString<96>(Key), Name.Len(), Name.GetData());
+		return FOptionalCacheRecord();
+	}
 
 	TStringBuilder<256> Path;
 	BuildCachePackagePath(Key, Path);
 
-
 	FOptionalCacheRecord Record;
 	{
 		FCbPackage Package;
@@ -1499,12 +1460,12 @@ bool FS3CacheStore::GetCacheValueOnly(
 		return false;
 	}
 
-	//if (ShouldSimulateMiss(Key))
-	//{
-	//	UE_LOG(LogDerivedDataCache, Verbose, TEXT("%s: Simulated miss for get of %s from '%.*s'"),
-	//		*CachePath, *WriteToString<96>(Key), Name.Len(), Name.GetData());
-	//	return false;
-	//}
+	if (DebugOptions.ShouldSimulateGetMiss(Key))
+	{
+		UE_LOG(LogDerivedDataCache, Verbose, TEXT("%s: Simulated miss for get of %s from '%.*s'"),
+			*GetName(), *WriteToString<96>(Key), Name.Len(), Name.GetData());
+		return false;
+	}
 
 	TStringBuilder<256> Path;
 	BuildCachePackagePath(Key, Path);
diff --git a/Engine/Source/Developer/DerivedDataCache/Private/ZenCacheStore.cpp b/Engine/Source/Developer/DerivedDataCache/Private/ZenCacheStore.cpp
--- a/Engine/Source/Developer/DerivedDataCache/Private/ZenCacheStore.cpp
+++ b/Engine/Source/Developer/DerivedDataCache/Private/ZenCacheStore.cpp
@@ -13,10 +13,8 @@
 #include "DerivedDataCacheRecord.h"
 #include "DerivedDataCacheUsageStats.h"
 #include "DerivedDataChunk.h"
-#include "HAL/CriticalSection.h"
 #include "Misc/ConfigCacheIni.h"
 #include "Misc/Optional.h"
-#include "Misc/ScopeLock.h"
 #include "ProfilingDebugging/CountersTrace.h"
 #include "ProfilingDebugging/CpuProfilerTrace.h"
 #include "Serialization/BufferArchive.h"
@@ -163,10 +161,6 @@ private:
 
 	static bool ShouldRetryOnError(int64 ResponseCode);
 
-	/* Debug helpers */
-	bool ShouldSimulateMiss(const TCHAR* InKey);
-	bool ShouldSimulateMiss(const FCacheKey& InKey);
-
 private:
 	FString Namespace;
 	UE::Zen::FScopeZenService ZenService;
@@ -181,11 +175,6 @@ private:
 
 	/** Debug Options */
 	FBackendDebugOptions DebugOptions;
-
-	/** Keys we ignored due to miss rate settings */
-	FCriticalSection MissedKeysCS;
-	TSet<FName> DebugMissedKeys;
-	TSet<FCacheKey> DebugMissedCacheKeys;
 };
 
 FZenCacheStore::FZenCacheStore(
@@ -246,7 +235,7 @@ bool FZenCacheStore::CachedDataProbablyExists(const TCHAR* CacheKey)
 	TRACE_COUNTER_ADD(ZenDDC_Exist, int64(1));
 	COOK_STAT(auto Timer = UsageStats.TimeProbablyExists());
 
-	if (ShouldSimulateMiss(CacheKey))
+	if (DebugOptions.ShouldSimulateGetMiss(CacheKey))
 	{
 		return false;
 	}
@@ -291,7 +280,7 @@ bool FZenCacheStore::GetCachedData(const TCHAR* CacheKey, TArray<uint8>& OutData
 	TRACE_COUNTER_ADD(ZenDDC_Get, int64(1));
 	COOK_STAT(auto Timer = UsageStats.TimeGet());
 
-	if (ShouldSimulateMiss(CacheKey))
+	if (DebugOptions.ShouldSimulateGetMiss(CacheKey))
 	{
 		return false;
 	}
@@ -420,7 +409,7 @@ FZenCacheStore::PutCachedData(const TCHAR* CacheKey, TArrayView<const uint8> InD
 {
 	TRACE_CPUPROFILER_EVENT_SCOPE(ZenDDC::PutCachedData);
 
-	if (ShouldSimulateMiss(CacheKey))
+	if (DebugOptions.ShouldSimulatePutMiss(CacheKey))
 	{
 		return EPutStatus::NotCached;
 	}
@@ -481,14 +470,14 @@ void FZenCacheStore::AppendZenUri(const FCacheKey& CacheKey, const FValueId& Id,
 {
 	AppendZenUri(CacheKey, Out);
 	Out << TEXT('/') << Id;
-		}
+}
 
 void FZenCacheStore::AppendPolicyQueryString(ECachePolicy Policy, FStringBuilderBase& Uri)
-	{
+{
 	if (Policy != ECachePolicy::Default)
-		{
+	{
 		Uri << TEXT("?Policy=") << Policy;
-}
+	}
 }
 
 void FZenCacheStore::RemoveCachedData(const TCHAR* CacheKey, bool bTransient)
@@ -601,56 +590,6 @@ bool FZenCacheStore::ApplyDebugOptions(FBackendDebugOptions& InOptions)
 	return true;
 }
 
-bool FZenCacheStore::ShouldSimulateMiss(const TCHAR* InKey)
-{
-	if (DebugOptions.RandomMissRate == 0 && DebugOptions.SimulateMissTypes.IsEmpty())
-	{
-		return false;
-	}
-
-	const FName Key(InKey);
-	const uint32 Hash = GetTypeHash(Key);
-
-	if (FScopeLock Lock(&MissedKeysCS); DebugMissedKeys.ContainsByHash(Hash, Key))
-	{
-		return true;
-	}
-
-	if (DebugOptions.ShouldSimulateMiss(InKey))
-	{
-		FScopeLock Lock(&MissedKeysCS);
-		UE_LOG(LogDerivedDataCache, Verbose, TEXT("Simulating miss in %s for %s"), *GetName(), InKey);
-		DebugMissedKeys.AddByHash(Hash, Key);
-		return true;
-	}
-
-	return false;
-}
-
-bool FZenCacheStore::ShouldSimulateMiss(const FCacheKey& Key)
-{
-	if (DebugOptions.RandomMissRate == 0 && DebugOptions.SimulateMissTypes.IsEmpty())
-	{
-		return false;
-	}
-
-	const uint32 Hash = GetTypeHash(Key);
-
-	if (FScopeLock Lock(&MissedKeysCS); DebugMissedCacheKeys.ContainsByHash(Hash, Key))
-	{
-		return true;
-	}
-
-	if (DebugOptions.ShouldSimulateMiss(Key))
-	{
-		FScopeLock Lock(&MissedKeysCS);
-		DebugMissedCacheKeys.AddByHash(Hash, Key);
-		return true;
-	}
-
-	return false;
-}
-
 void FZenCacheStore::Put(
 	const TConstArrayView<FCachePutRequest> Requests,
 	IRequestOwner& Owner,
@@ -666,15 +605,15 @@ void FZenCacheStore::Put(
 			const FCacheRecord& Record = NextRequest.Record;
 			uint64 RecordSize = sizeof(FCacheKey) + Record.GetMeta().GetSize();
 			for (const FValueWithId& Value : Record.GetValues())
-	{
+			{
 				RecordSize += Value.GetData().GetCompressedSize();
 			}
 			BatchSize += RecordSize;
 			if (BatchSize > BatchPutMaxBytes)
-		{
+			{
 				BatchSize = RecordSize;
 				return EBatchView::NewBatch;
-		}
+			}
 			return EBatchView::Continue;
 		});
 
@@ -713,15 +652,15 @@ void FZenCacheStore::Put(
 					const FCacheRecord& Record = Request.Record;
 
 					BatchWriter.BeginObject();
-		{
+					{
 						BatchWriter.SetName("Record"_ASV);
 						Record.Save(BatchPackage, BatchWriter);
 						if (!Request.Policy.IsUniform() || Request.Policy.GetRecordPolicy() != BatchDefaultPolicy)
-			{
+						{
 							BatchWriter.SetName("Policy"_ASV);
 							Request.Policy.Save(BatchWriter);
-			}
-		}
+						}
+					}
 					BatchWriter.EndObject();
 				}
 				BatchWriter.EndArray();
@@ -754,8 +693,8 @@ void FZenCacheStore::Put(
 
 				const FCacheKey& Key = Request.Record.GetKey();
 				bool bPutSucceeded = ResponseField.AsBool();
-				if (ShouldSimulateMiss(Key))
-		{
+				if (DebugOptions.ShouldSimulatePutMiss(Key))
+				{
 					UE_LOG(LogDerivedDataCache, Verbose, TEXT("%s: Simulated miss for put of %s from '%s'"),
 						*GetName(), *WriteToString<96>(Key), *Request.Name);
 					bPutSucceeded = false;
@@ -809,7 +748,7 @@ void FZenCacheStore::Get(
 		[this, &Requests, &Owner, &OnComplete, &OnHit, &OnMiss](int32 BatchFirst, int32 BatchLast)
 	{
 		TArrayView<const FCacheGetRequest> Batch(Requests.GetData() + BatchFirst, BatchLast - BatchFirst + 1);
-		
+
 		FCbWriter BatchRequest;
 		BatchRequest.BeginObject();
 		{
@@ -823,18 +762,18 @@ void FZenCacheStore::Get(
 				for (const FCacheGetRequest& Request : Batch)
 				{
 					BatchRequest.BeginObject();
-				{
+					{
 						BatchRequest << "Key"_ASV << Request.Key;
 						if (!Request.Policy.IsUniform() || Request.Policy.GetRecordPolicy() != BatchDefaultPolicy)
-					{
+						{
 							BatchRequest.SetName("Policy"_ASV);
 							Request.Policy.Save(BatchRequest);
 						}
 					}
-							BatchRequest.EndObject();
-						}
-						BatchRequest.EndArray();
-					}
+					BatchRequest.EndObject();
+				}
+				BatchRequest.EndArray();
+			}
 			BatchRequest.EndObject();
 		}
 		BatchRequest.EndObject();
@@ -856,16 +795,16 @@ void FZenCacheStore::Get(
 			for (FCbField RecordField : ResponseObj["Result"_ASV])
 			{
 				if (RequestIndex >= Batch.Num())
-			{
+				{
 					++RequestIndex;
 					continue;
 				}
 				const FCacheGetRequest& Request = Batch[RequestIndex++];
-				
+
 				const FCacheKey& Key = Request.Key;
 				FOptionalCacheRecord Record;
 
-				if (ShouldSimulateMiss(Key))
+				if (DebugOptions.ShouldSimulateGetMiss(Key))
 				{
 					UE_LOG(LogDerivedDataCache, Verbose, TEXT("%s: Simulated miss for put of '%s' from '%s'"),
 						*GetName(), *WriteToString<96>(Key), *Request.Name);
@@ -895,7 +834,7 @@ void FZenCacheStore::PutValue(
 	TConstArrayView<FCachePutValueRequest> Requests,
 	IRequestOwner& Owner,
 	FOnCachePutValueComplete&& OnComplete)
-				{
+{
 	TRACE_CPUPROFILER_EVENT_SCOPE(ZenDDC::PutValue);
 	COOK_STAT(auto Timer = UsageStats.TimePut());
 
@@ -987,7 +926,7 @@ void FZenCacheStore::PutValue(
 				const FCachePutValueRequest& Request = Batch[RequestIndex++];
 
 				bool bPutSucceeded = ResponseField.AsBool();
-				if (ShouldSimulateMiss(Request.Key))
+				if (DebugOptions.ShouldSimulatePutMiss(Request.Key))
 				{
 					UE_LOG(LogDerivedDataCache, Verbose, TEXT("%s: Simulated miss for PutValue of %s from '%s'"),
 						*GetName(), *WriteToString<96>(Request.Key), *Request.Name);
@@ -996,17 +935,17 @@ void FZenCacheStore::PutValue(
 				bPutSucceeded ? OnHit(Request) : OnMiss(Request);
 			}
 			if (RequestIndex != Batch.Num())
-					{
+			{
 				UE_LOG(LogDerivedDataCache, Warning,
 					TEXT("Invalid response received from PutCacheValues rpc: %d results expected, received %d."), Batch.Num(), RequestIndex);
 			}
-					}
+		}
 		for (const FCachePutValueRequest& Request : Batch.RightChop(RequestIndex))
-					{
+		{
 			OnMiss(Request);
 		}
-					}
-				}
+	}
+}
 
 void FZenCacheStore::GetValue(
 	TConstArrayView<FCacheGetValueRequest> Requests,
@@ -1081,25 +1020,25 @@ void FZenCacheStore::GetValue(
 
 			int32 RequestIndex = 0;
 			if (HttpResult == Zen::FZenHttpRequest::Result::Success)
-				{
+			{
 				const FCbObject& ResponseObj = BatchResponse.GetObject();
-					
+
 				for (FCbFieldView ResultField : ResponseObj["Result"_ASV])
 				{
 					if (RequestIndex >= Batch.Num())
 					{
 						++RequestIndex;
 						continue;
-				}
+					}
 					const FCacheGetValueRequest& Request = Batch[RequestIndex++];
-				
+
 					FCbObjectView ResultObj = ResultField.AsObjectView();
 					TOptional<FValue> Value;
-					if (ShouldSimulateMiss(Request.Key))
+					if (DebugOptions.ShouldSimulateGetMiss(Request.Key))
 					{
 						UE_LOG(LogDerivedDataCache, Verbose, TEXT("%s: Simulated miss for GetValue of '%s' from '%s'"),
 							*GetName(), *WriteToString<96>(Request.Key), *Request.Name);
-			}
+					}
 					else
 					{
 						FCbFieldView RawHashField = ResultObj["RawHash"];
@@ -1107,13 +1046,13 @@ void FZenCacheStore::GetValue(
 						if (const FCbAttachment* Attachment = BatchResponse.FindAttachment(RawHash))
 						{
 							Value.Emplace(Attachment->AsCompressedBinary());
-		}
-		else
-		{
+						}
+						else
+						{
 							FCbFieldView RawSizeField = ResultObj["RawSize"];
 							uint64 RawSize = RawSizeField.AsUInt64();
 							if (!RawSizeField.HasError() && !RawHashField.HasError())
-			{
+							{
 								Value.Emplace(RawHash, RawSize);
 							}
 						}
@@ -1129,11 +1068,10 @@ void FZenCacheStore::GetValue(
 			for (const FCacheGetValueRequest& Request : Batch.RightChop(RequestIndex))
 			{
 				OnMiss(Request);
-		}
-	});
+			}
+		});
 	TRACE_COUNTER_SUBTRACT(ZenDDC_CacheRecordRequestCount, int64(Requests.Num()));
 }
-	
 
 void FZenCacheStore::GetChunks(
 	TConstArrayView<FCacheGetChunkRequest> Requests,
@@ -1242,11 +1180,11 @@ void FZenCacheStore::GetChunks(
 				uint64 RawSize = 0;
 				FCbObjectView ResultObject = ResultView.AsObjectView();
 				FSharedBuffer RequestedBytes;
-				if (ShouldSimulateMiss(Request.Key))
-					{
+				if (DebugOptions.ShouldSimulateGetMiss(Request.Key))
+				{
 					UE_LOG(LogDerivedDataCache, Verbose, TEXT("%s: Simulated miss for get of '%s' from '%s'"),
 						*GetName(), *WriteToString<96>(Request.Key, '/', Request.Id), *Request.Name);
-					}
+				}
 				else
 				{
 					FCbFieldView HashView = ResultObject["RawHash"_ASV];
@@ -1254,31 +1192,31 @@ void FZenCacheStore::GetChunks(
 					if (!HashView.HasError())
 					{
 						if (const FCbAttachment* Attachment = BatchResponse.FindAttachment(HashView.AsHash()))
-				{
+						{
 							FCompressedBuffer CompressedBuffer = Attachment->AsCompressedBinary();
 							if (CompressedBuffer)
-					{
+							{
 								TRACE_COUNTER_ADD(ZenDDC_BytesReceived, CompressedBuffer.GetCompressedSize());
 								RequestedBytes = FCompressedBufferReader(CompressedBuffer).Decompress(Request.RawOffset, Request.RawSize);
 								RawSize = RequestedBytes.GetSize();
 								Succeeded = true;
-					}
-				}
-				else
-				{
+							}
+						}
+						else
+						{
 							FCbFieldView RawSizeField = ResultObject["RawSize"_ASV];
 							uint64 TotalSize = RawSizeField.AsUInt64();
 							Succeeded = !RawSizeField.HasError();
 							if (Succeeded)
-					{
+							{
 								RawSize = FMath::Min(Request.RawSize, TotalSize - FMath::Min(Request.RawOffset, TotalSize));
+							}
+						}
 					}
 				}
-			}
-		}
 				Succeeded ? OnHit(Request, MoveTemp(RawHash), RawSize, MoveTemp(RequestedBytes)) : OnMiss(Request);
-				}
 			}
+		}
 		for (const FCacheGetChunkRequest& Request : Batch.RightChop(RequestIndex))
 		{
 			OnMiss(Request);
@@ -1298,7 +1236,7 @@ ILegacyCacheStore* CreateZenCacheStore(const TCHAR* NodeName, const TCHAR* Servi
 	UE_LOG(LogDerivedDataCache, Warning, TEXT("%s could not contact the service (%s), will not use it."), NodeName, *Backend->GetName());
 	delete Backend;
 	return nullptr;
-	}
+}
 
 } // namespace UE::DerivedData
 
