Make DDC-Verify successes Verbose

diff --git a/Engine/Source/Developer/DerivedDataCache/Private/DerivedDataCacheStoreVerify.cpp b/Engine/Source/Developer/DerivedDataCache/Private/DerivedDataCacheStoreVerify.cpp
--- a/Engine/Source/Developer/DerivedDataCache/Private/DerivedDataCacheStoreVerify.cpp
+++ b/Engine/Source/Developer/DerivedDataCache/Private/DerivedDataCacheStoreVerify.cpp
@@ -235,7 +235,7 @@ void FCacheStoreVerify::GetMetaComplete(IRequestOwner& Owner, FVerifyPutState* S
 		};
 		if (Algo::CompareBy(Request.Record.GetValues(), Response.Record.GetValues(), MakeValueTuple))
 		{
-			UE_LOG(LogDerivedDataCache, Log,
+			UE_LOG(LogDerivedDataCache, Verbose,
 				TEXT("Verify: Data in the cache matches newly generated data for %s from '%s'."),
 				*WriteToString<96>(Request.Record.GetKey()), *Request.Name);
 			State->OnComplete(Request.MakeResponse(EStatus::Ok));
@@ -271,7 +271,7 @@ void FCacheStoreVerify::GetDataComplete(IRequestOwner& Owner, FVerifyPutState* S
 	{
 		if (CompareRecords(Request.Record, Response.Record, Request.Name))
 		{
-			UE_LOG(LogDerivedDataCache, Log,
+			UE_LOG(LogDerivedDataCache, Verbose,
 				TEXT("Verify: Data in the cache matches newly generated data for %s from '%s'."),
 				*WriteToString<96>(Request.Record.GetKey()), *Request.Name);
 			State->OnComplete(Request.MakeResponse(EStatus::Ok));
@@ -394,7 +394,7 @@ void FCacheStoreVerify::GetMetaComplete(IRequestOwner& Owner, FVerifyPutValueSta
 	{
 		if (Request.Value.GetRawHash() == Response.Value.GetRawHash())
 		{
-			UE_LOG(LogDerivedDataCache, Log,
+			UE_LOG(LogDerivedDataCache, Verbose,
 				TEXT("Verify: Data in the cache matches newly generated data for %s from '%s'."),
 				*WriteToString<96>(Request.Key), *Request.Name);
 			State->OnComplete(Request.MakeResponse(EStatus::Ok));
@@ -428,7 +428,7 @@ void FCacheStoreVerify::GetDataComplete(IRequestOwner& Owner, FVerifyPutValueSta
 	{
 		if (Request.Value.GetRawHash() == Response.Value.GetRawHash())
 		{
-			UE_LOG(LogDerivedDataCache, Log,
+			UE_LOG(LogDerivedDataCache, Verbose,
 				TEXT("Verify: Data in the cache matches newly generated data for %s from '%s'."),
 				*WriteToString<96>(Request.Key), *Request.Name);
 			State->OnComplete(Request.MakeResponse(EStatus::Ok));
@@ -586,7 +586,7 @@ void FCacheStoreVerify::GetDataComplete(IRequestOwner& Owner, FVerifyLegacyPutSt
 		const FLegacyCacheValue& OldValue = Response.Value;
 		if (NewValue.GetRawHash() == OldValue.GetRawHash())
 		{
-			UE_LOG(LogDerivedDataCache, Log,
+			UE_LOG(LogDerivedDataCache, Verbose,
 				TEXT("Verify: Data in the cache matches newly generated data for %s from '%s'."),
 				*WriteToString<96>(Request.Key.GetKey()), *Request.Name);
 			State->OnComplete(Request.MakeResponse(EStatus::Ok));
