Allow async operations for Jupiter DDC backend.

diff --git a/Engine/Source/Developer/DerivedDataCache/Private/HttpCacheStore.cpp b/Engine/Source/Developer/DerivedDataCache/Private/HttpCacheStore.cpp
--- a/Engine/Source/Developer/DerivedDataCache/Private/HttpCacheStore.cpp
+++ b/Engine/Source/Developer/DerivedDataCache/Private/HttpCacheStore.cpp
@@ -16,6 +16,7 @@
 #include "Algo/Find.h"
 #include "Algo/Transform.h"
 #include "Compression/CompressedBuffer.h"
+#include "Containers/DepletableMpscQueue.h"
 #include "Containers/StaticArray.h"
 #include "Containers/StringView.h"
 #include "Containers/Ticker.h"
@@ -24,14 +25,21 @@
 #include "DerivedDataCacheRecord.h"
 #include "DerivedDataCacheUsageStats.h"
 #include "DerivedDataChunk.h"
+#include "DerivedDataRequest.h"
+#include "DerivedDataRequestOwner.h"
 #include "DerivedDataValue.h"
 #include "Dom/JsonObject.h"
+#include "Experimental/Async/LazyEvent.h"
 #include "Experimental/Containers/FAAArrayQueue.h"
 #include "GenericPlatform/GenericPlatformFile.h"
+#include "HAL/IConsoleManager.h"
 #include "HAL/PlatformFileManager.h"
+#include "HAL/Runnable.h"
+#include "HAL/RunnableThread.h"
 #include "IO/IoHash.h"
 #include "Memory/SharedBuffer.h"
 #include "Misc/FileHelper.h"
+#include "Misc/Optional.h"
 #include "Misc/ScopeLock.h"
 #include "Misc/SecureHash.h"
 #include "Misc/StringBuilder.h"
@@ -46,6 +54,7 @@
 #include "Serialization/JsonReader.h"
 #include "Serialization/JsonSerializer.h"
 #include "Serialization/JsonWriter.h"
+#include "Tasks/Task.h"
 
 #if WITH_SSL
 #include "Ssl.h"
@@ -66,6 +75,7 @@
 #define UE_HTTPDDC_HTTP_DEBUG 0
 #define UE_HTTPDDC_GET_REQUEST_POOL_SIZE 48
 #define UE_HTTPDDC_PUT_REQUEST_POOL_SIZE 16
+#define UE_HTTPDDC_NONBLOCKING_REQUEST_POOL_SIZE 128
 #define UE_HTTPDDC_MAX_FAILED_LOGIN_ATTEMPTS 16
 #define UE_HTTPDDC_MAX_ATTEMPTS 4
 #define UE_HTTPDDC_MAX_BUFFER_RESERVE 104857600u
@@ -87,6 +97,13 @@ TRACE_DECLARE_INT_COUNTER(HttpDDC_PutHit, TEXT("HttpDDC Put Hit"));
 TRACE_DECLARE_INT_COUNTER(HttpDDC_BytesReceived, TEXT("HttpDDC Bytes Received"));
 TRACE_DECLARE_INT_COUNTER(HttpDDC_BytesSent, TEXT("HttpDDC Bytes Sent"));
 
+static bool bHttpEnableAsync = true;
+static FAutoConsoleVariableRef CVarHttpEnableAsync(
+	TEXT("DDC.Http.EnableAsync"),
+	bHttpEnableAsync,
+	TEXT("If true, async operations are permitted, otherwise all operations are forced to be synchronous."),
+	ECVF_Default);
+
 static CURLcode sslctx_function(CURL * curl, void * sslctx, void * parm);
 
 typedef TSharedPtr<class IHttpRequest> FHttpRequestPtr;
@@ -111,6 +128,7 @@ private:
 struct FHttpSharedData
 {
 	FHttpSharedData()
+	: PendingRequestEvent(EEventMode::AutoReset)
 	{
 		FMemory::Memset(WriteLocked, 0, sizeof(WriteLocked));
 		CurlShare = curl_share_init();
@@ -119,14 +137,24 @@ struct FHttpSharedData
 		curl_share_setopt(CurlShare, CURLSHOPT_UNLOCKFUNC, UnlockFn);
 		curl_share_setopt(CurlShare, CURLSHOPT_SHARE, CURL_LOCK_DATA_DNS);
 		curl_share_setopt(CurlShare, CURLSHOPT_SHARE, CURL_LOCK_DATA_SSL_SESSION);
+		CurlMulti = curl_multi_init();
+		curl_multi_setopt(CurlMulti, CURLMOPT_PIPELINING, CURLPIPE_MULTIPLEX);
 	}
 
 	~FHttpSharedData()
 	{
+		curl_multi_cleanup(CurlMulti);
 		curl_share_cleanup(CurlShare);
 	}
 
 	CURLSH* CurlShare;
+	CURLM* CurlMulti;
+	TDepletableMpscQueue<CURL*> PendingRequestAdditions;
+	FLazyEvent PendingRequestEvent;
+	FRunnableThread* AsyncServiceThread = nullptr;
+	TUniquePtr<class FHttpCacheStoreRunnable> AsyncRunnable;
+	std::atomic<bool> bAsyncThreadStarting = false;
+	static inline std::atomic<bool> bAsyncThreadShutdownRequested = false;
 
 private:
 	FRWLock Locks[CURL_LOCK_DATA_LAST];
@@ -160,6 +188,33 @@ private:
 	}
 };
 
+template <typename T>
+class TRefCountedUniqueFunction final : public FThreadSafeRefCountedObject
+{
+public:
+	explicit TRefCountedUniqueFunction(T&& InFunction) : Function(MoveTemp(InFunction))
+	{
+	}
+
+	const T& GetFunction() const { return Function; }
+
+private:
+	T Function;
+};
+
+class FHttpCacheStoreRunnable final : public FRunnable
+{
+public:
+	explicit FHttpCacheStoreRunnable(FHttpSharedData& InSharedData) : SharedData(InSharedData)
+	{
+	}
+
+	virtual uint32 Run() override final;
+
+protected:
+	FHttpSharedData& SharedData;
+};
+
 /**
  * Minimal HTTP request type wrapping CURL without the need for managers. This request
  * is written to allow reuse of request objects, in order to allow connections to be reused.
@@ -196,8 +251,17 @@ public:
 		FailedTimeout
 	};
 
+	enum class ECompletionBehavior : uint8
+	{
+		Done,
+		Retry
+	};
+
+	using FOnHttpRequestComplete = TUniqueFunction<ECompletionBehavior(Result HttpResult, FHttpRequest* Request)>;
+
 	FHttpRequest(const TCHAR* InDomain, const TCHAR* InEffectiveDomain, FHttpAccessToken* InAuthorizationToken, bool bInLogErrors)
-		: bLogErrors(bInLogErrors)
+		: AsyncData(nullptr)
+		, bLogErrors(bInLogErrors)
 		, Domain(InDomain)
 		, EffectiveDomain(InEffectiveDomain)
 		, AuthorizationToken(InAuthorizationToken)
@@ -209,6 +273,7 @@ public:
 	~FHttpRequest()
 	{
 		curl_easy_cleanup(Curl);
+		check(!AsyncData);
 	}
 
 	/**
@@ -220,16 +285,19 @@ public:
 		ResponseHeader.Reset();
 		ResponseBuffer.Reset();
 		ResponseCode = 0;
-		ReadDataView = TArrayView<const uint8>();
+		ReadDataView = FMemoryView();
 		WriteDataBufferPtr = nullptr;
 		WriteHeaderBufferPtr = nullptr;
 		BytesSent = 0;
 		BytesReceived = 0;
+		Attempts = 0;
 		CurlResult = CURL_LAST;
 
-		static FHttpSharedData SharedData;
+		static FHttpSharedData StaticSharedData;
+		SharedData = &StaticSharedData;
 
 		curl_easy_reset(Curl);
+		check(!AsyncData);
 
 		// Options that are always set for all connections.
 #if UE_HTTPDDC_HTTP_REQUEST_TIMEOUT_ENABLED
@@ -238,7 +306,7 @@ public:
 		curl_easy_setopt(Curl, CURLOPT_FOLLOWLOCATION, 1L);
 		curl_easy_setopt(Curl, CURLOPT_NOSIGNAL, 1L);
 		curl_easy_setopt(Curl, CURLOPT_DNS_CACHE_TIMEOUT, 300L); // Don't re-resolve every minute
-		curl_easy_setopt(Curl, CURLOPT_SHARE, SharedData.CurlShare);
+		curl_easy_setopt(Curl, CURLOPT_SHARE, SharedData->CurlShare);
 		// SSL options
 		curl_easy_setopt(Curl, CURLOPT_USE_SSL, CURLUSESSL_ALL);
 		curl_easy_setopt(Curl, CURLOPT_SSL_VERIFYPEER, 1);
@@ -267,6 +335,19 @@ public:
 		curl_easy_setopt(Curl, CURLOPT_DEBUGFUNCTION, StaticDebugCallback);
 		curl_easy_setopt(Curl, CURLOPT_VERBOSE, 1L);
 #endif
+		curl_easy_setopt(Curl, CURLOPT_PRIVATE, this);
+		curl_easy_setopt(Curl, CURLOPT_HTTP_VERSION, CURL_HTTP_VERSION_2_0);
+	}
+
+	void PrepareToRetry()
+	{
+		ResponseHeader.Reset();
+		ResponseBuffer.Reset();
+		ResponseCode = 0;
+		BytesSent = 0;
+		BytesReceived = 0;
+		CurlResult = CURL_LAST;
+		++Attempts;
 	}
 
 	/** Gets the domain name for this request */
@@ -299,8 +380,14 @@ public:
 		return BytesReceived;
 	}
 
+	/** Returns the number of attempts we've made issuing this request (currently tracked for async requests only). */
+	const size_t GetAttempts() const
+	{
+		return Attempts;
+	}
+
 	/** Returns the number of bytes sent during this request (headers withstanding). */
-	const size_t  GetBytesSent() const
+	const size_t GetBytesSent() const
 	{
 		return BytesSent;
 	}
@@ -316,7 +403,7 @@ public:
 	{
 		static_assert(V == Put || V == PutCompactBinary || V == PutCompressedBlob || V == Post || V == PostCompactBinary || V == PostJson, "Upload should use either Put or Post verbs.");
 		
-		uint32 ContentLength = 0u;
+		uint64 ContentLength = 0u;
 
 		if constexpr (V == Put || V == PutCompactBinary || V == PutCompressedBlob)
 		{
@@ -337,9 +424,9 @@ public:
 				Headers.Add(FString(TEXT("Content-Type: application/octet-stream")));
 			}
 			ContentLength = Buffer.Num();
-			ReadDataView = Buffer;
+			ReadDataView = MakeMemoryView(Buffer);
 		}
-		else if (V == Post || V == PostCompactBinary || V == PostJson)
+		else if constexpr (V == Post || V == PostCompactBinary || V == PostJson)
 		{
 			curl_easy_setopt(Curl, CURLOPT_POST, 1L);
 			curl_easy_setopt(Curl, CURLOPT_INFILESIZE, Buffer.Num());
@@ -358,11 +445,66 @@ public:
 				Headers.Add(FString(TEXT("Content-Type: application/x-www-form-urlencoded")));
 			}
 			ContentLength = Buffer.Num();
-			ReadDataView = Buffer;
+			ReadDataView = MakeMemoryView(Buffer);
 		}
 
 		return PerformBlocking(Uri, V, ContentLength, ExpectedErrorCodes);
 	}
+	
+	template<RequestVerb V>
+	void EnqueueAsyncUpload(IRequestOwner& Owner, struct FRequestPool* Pool, const TCHAR* Uri, FSharedBuffer Buffer, FOnHttpRequestComplete&& OnComplete, TConstArrayView<long> ExpectedErrorCodes = {})
+	{
+		static_assert(V == Put || V == PutCompactBinary || V == PutCompressedBlob || V == Post || V == PostCompactBinary || V == PostJson, "Upload should use either Put or Post verbs.");
+		
+		uint64 ContentLength = 0u;
+
+		if constexpr (V == Put || V == PutCompactBinary || V == PutCompressedBlob)
+		{
+			curl_easy_setopt(Curl, CURLOPT_UPLOAD, 1L);
+			curl_easy_setopt(Curl, CURLOPT_INFILESIZE, Buffer.GetSize());
+			curl_easy_setopt(Curl, CURLOPT_READDATA, this);
+			curl_easy_setopt(Curl, CURLOPT_READFUNCTION, StaticReadFn);
+			if constexpr (V == PutCompactBinary)
+			{
+				Headers.Add(FString(TEXT("Content-Type: application/x-ue-cb")));
+			}
+			else if constexpr (V == PutCompressedBlob)
+			{
+				Headers.Add(FString(TEXT("Content-Type: application/x-ue-comp")));
+			}
+			else
+			{
+				Headers.Add(FString(TEXT("Content-Type: application/octet-stream")));
+			}
+			ReadSharedBuffer = Buffer;
+			ContentLength = Buffer.GetSize();
+			ReadDataView = Buffer.GetView();
+		}
+		else if constexpr (V == Post || V == PostCompactBinary || V == PostJson)
+		{
+			curl_easy_setopt(Curl, CURLOPT_POST, 1L);
+			curl_easy_setopt(Curl, CURLOPT_INFILESIZE, Buffer.GetSize());
+			curl_easy_setopt(Curl, CURLOPT_READDATA, this);
+			curl_easy_setopt(Curl, CURLOPT_READFUNCTION, StaticReadFn);
+			if constexpr (V == PostCompactBinary)
+			{
+				Headers.Add(FString(TEXT("Content-Type: application/x-ue-cb")));
+			}
+			else if constexpr (V == PostJson)
+			{
+				Headers.Add(FString(TEXT("Content-Type: application/json")));
+			}
+			else
+			{
+				Headers.Add(FString(TEXT("Content-Type: application/x-www-form-urlencoded")));
+			}
+			ReadSharedBuffer = Buffer;
+			ContentLength = Buffer.GetSize();
+			ReadDataView = Buffer.GetView();
+		}
+
+		return EnqueueAsync(Owner, Pool, Uri, V, ContentLength, MoveTemp(OnComplete), ExpectedErrorCodes);
+	}
 
 	/**
 	 * Download an url into a buffer using the request.
@@ -379,6 +521,13 @@ public:
 		return PerformBlocking(Uri, Get, 0u, ExpectedErrorCodes);
 	}
 
+	void EnqueueAsyncDownload(IRequestOwner& Owner, struct FRequestPool* Pool, const TCHAR* Uri, FOnHttpRequestComplete&& OnComplete, TConstArrayView<long> ExpectedErrorCodes = {400})
+	{
+		curl_easy_setopt(Curl, CURLOPT_HTTPGET, 1L);
+
+		return EnqueueAsync(Owner, Pool, Uri, Get, 0u, MoveTemp(OnComplete), ExpectedErrorCodes);
+	}
+
 	/**
 	 * Query an url using the request. Queries can use either "Head" or "Delete" verbs.
 	 * @param Uri Url to use.
@@ -401,6 +550,23 @@ public:
 		return PerformBlocking(Uri, V, 0u, ExpectedErrorCodes);
 	}
 
+	template<RequestVerb V>
+	void EnqueueAsyncQuery(const TCHAR* Uri, FOnHttpRequestComplete&& OnComplete, TConstArrayView<long> ExpectedErrorCodes = {400})
+	{
+		static_assert(V == Head || V == Delete, "Queries should use either Head or Delete verbs.");
+
+		if (V == Delete)
+		{
+			curl_easy_setopt(Curl, CURLOPT_CUSTOMREQUEST, "DELETE");
+		}
+		else if (V == Head)
+		{
+			curl_easy_setopt(Curl, CURLOPT_NOBODY, 1L);
+		}
+
+		return EnqueueAsync(Uri, V, 0u, MoveTemp(OnComplete), ExpectedErrorCodes);
+	}
+
 	/**
 	 * Set a header to send with the request.
 	 */
@@ -442,6 +608,11 @@ public:
 		return ResponseBuffer;
 	}
 
+	FSharedBuffer MoveResponseBufferToShared()
+	{
+		return MakeSharedBufferFromArray(MoveTemp(ResponseBuffer));
+	}
+
 	/**
 	 * Returns the response buffer as a string. Note that is the request is performed
 	 * with an external buffer as target buffer this string will be empty.
@@ -491,6 +662,8 @@ public:
 		return JsonArray;
 	}
 
+	void CompleteAsync(CURLcode Result);
+
 	/** Will return true if the response code is considered a success */
 	static bool IsSuccessResponse(long ResponseCode)
 	{
@@ -498,16 +671,57 @@ public:
 		return ResponseCode >= 100 && ResponseCode < 300;
 	}
 
+	static bool AllowAsync();
+
 private:
 
+	struct FAsyncRequestData final : public FRequestBase
+	{
+		IRequestOwner* Owner = nullptr;
+		FRequestPool* Pool = nullptr;
+		curl_slist* CurlHeaders = nullptr;
+		FString Uri;
+		RequestVerb Verb;
+		TArray<long, TInlineAllocator<4>> ExpectedErrorCodes;
+		FOnHttpRequestComplete OnComplete;
+		FLazyEvent Event {EEventMode::ManualReset};
+
+		void Reset()
+		{
+			if (CurlHeaders)
+			{
+				curl_slist_free_all(CurlHeaders);
+				CurlHeaders = nullptr;
+			}
+			Uri.Empty();
+			ExpectedErrorCodes.Empty();
+		}
+
+		void SetPriority(EPriority Priority) final {}
+
+		void Cancel() final
+		{
+			Event.Wait();
+		}
+
+		void Wait() final
+		{
+			Event.Wait();
+		}
+	};
+
 	CURL*					Curl;
 	CURLcode				CurlResult;
+	FHttpSharedData*		SharedData;
+	FAsyncRequestData*		AsyncData;
 	long					ResponseCode;
 	size_t					BytesSent;
 	size_t					BytesReceived;
+	size_t					Attempts;
 	bool					bLogErrors;
 
-	TArrayView<const uint8>	ReadDataView;
+	FSharedBuffer			ReadSharedBuffer;
+	FMemoryView				ReadDataView;
 	TArray<uint8>*			WriteDataBufferPtr;
 	TArray<uint8>*			WriteHeaderBufferPtr;
 	
@@ -518,22 +732,13 @@ private:
 	FString					EffectiveDomain;
 	FHttpAccessToken*		AuthorizationToken;
 
-	/**
-	 * Performs the request, blocking until finished.
-	 * @param Uri Address on the domain to query
-	 * @param Verb HTTP verb to use
-	 * @param Buffer Optional buffer to directly receive the result of the request.
-	 * If unset the response body will be stored in the request.
-	 */
-	Result PerformBlocking(const TCHAR* Uri, RequestVerb Verb, uint32 ContentLength, TConstArrayView<long> ExpectedErrorCodes)
+	curl_slist* PrepareToIssueRequest(const TCHAR* Uri, uint64 ContentLength)
 	{
 		static const char* CommonHeaders[] = {
 			"User-Agent: Unreal Engine",
 			nullptr
 		};
 
-		TRACE_CPUPROFILER_EVENT_SCOPE(HttpDDC_CurlPerform);
-
 		// Setup request options
 		FString Url = FString::Printf(TEXT("%s/%s"), *EffectiveDomain, Uri);
 		curl_easy_setopt(Curl, CURLOPT_URL, TCHAR_TO_ANSI(*Url));
@@ -567,16 +772,29 @@ private:
 			CurlHeaders = curl_slist_append(CurlHeaders, TCHAR_TO_ANSI(*Header));
 		}
 		curl_easy_setopt(Curl, CURLOPT_HTTPHEADER, CurlHeaders);
+		return CurlHeaders;
+	}
+
+	/**
+	 * Performs the request, blocking until finished.
+	 * @param Uri Address on the domain to query
+	 * @param Verb HTTP verb to use
+	 * @param ContentLength The number of bytes being uploaded for the body of this request.
+	 * @param ExpectedErrorCodes An array of expected return codes outside of the success range that should NOT be logged as an abnormal/exceptional outcome.
+	 * If unset the response body will be stored in the request.
+	 */
+	Result PerformBlocking(const TCHAR* Uri, RequestVerb Verb, uint64 ContentLength, TConstArrayView<long> ExpectedErrorCodes)
+	{
+		TRACE_CPUPROFILER_EVENT_SCOPE(HttpDDC_CurlPerform);
+
+		// Build headers list
+		curl_slist* CurlHeaders = PrepareToIssueRequest(Uri, ContentLength);
 
 		// Shots fired!
 		CurlResult = curl_easy_perform(Curl);
 
 		// Get response code
-		bool bRedirected = false;
-		if (CURLE_OK == curl_easy_getinfo(Curl, CURLINFO_RESPONSE_CODE, &ResponseCode))
-		{
-			bRedirected = (ResponseCode >= 300 && ResponseCode < 400);
-		}
+		curl_easy_getinfo(Curl, CURLINFO_RESPONSE_CODE, &ResponseCode);
 
 		LogResult(CurlResult, Uri, Verb, ExpectedErrorCodes);
 
@@ -586,6 +804,8 @@ private:
 		return CurlResult == CURLE_OK ? Success : Failed;
 	}
 
+	void EnqueueAsync(IRequestOwner& Owner, FRequestPool* Pool, const TCHAR* Uri, RequestVerb Verb, uint64 ContentLength, FOnHttpRequestComplete&& OnComplete, TConstArrayView<long> ExpectedErrorCodes);
+
 	void LogResult(CURLcode Result, const TCHAR* Uri, RequestVerb Verb, TConstArrayView<long> ExpectedErrorCodes) const
 	{
 		if (Result == CURLE_OK)
@@ -745,17 +965,11 @@ private:
 	static size_t StaticReadFn(void* Ptr, size_t SizeInBlocks, size_t BlockSizeInBytes, void* UserData)
 	{
 		FHttpRequest* Request = static_cast<FHttpRequest*>(UserData);
-		TArrayView<const uint8>& ReadDataView = Request->ReadDataView;
-
-		const size_t Offset = Request->BytesSent;
-		const size_t ReadSize = FMath::Min((size_t)ReadDataView.Num() - Offset, SizeInBlocks * BlockSizeInBytes);
-		check(ReadDataView.Num() >= Offset + ReadSize);
-
-		FMemory::Memcpy(Ptr, ReadDataView.GetData() + Offset, ReadSize);
-		Request->BytesSent += ReadSize;
-		return ReadSize;
-		
-		return 0;
+		const size_t MaxReadSize = SizeInBlocks * BlockSizeInBytes;
+		const FMemoryView SourceView = Request->ReadDataView.Mid(Request->BytesSent, MaxReadSize);
+		MakeMemoryView(Ptr, MaxReadSize).CopyFrom(SourceView);
+		Request->BytesSent += SourceView.GetSize();
+		return SourceView.GetSize();
 	}
 
 	static size_t StaticWriteHeaderFn(void* Ptr, size_t SizeInBlocks, size_t BlockSizeInBytes, void* UserData)
@@ -823,11 +1037,11 @@ private:
 		{
 		case SEEK_SET: NewPosition = Offset; break;
 		case SEEK_CUR: NewPosition = Request->BytesSent + Offset; break;
-		case SEEK_END: NewPosition = Request->ReadDataView.Num() + Offset; break;
+		case SEEK_END: NewPosition = Request->ReadDataView.GetSize() + Offset; break;
 		}
 
 		// Make sure we don't seek outside of the buffer
-		if (NewPosition < 0 || NewPosition >= Request->ReadDataView.Num())
+		if (NewPosition < 0 || NewPosition >= Request->ReadDataView.GetSize())
 		{
 			return CURL_SEEKFUNC_FAIL;
 		}
@@ -860,24 +1074,28 @@ bool ShouldAbortForShutdown();
  */
 struct FRequestPool
 {
-	FRequestPool(const TCHAR* InServiceUrl, const TCHAR* InEffectiveServiceUrl, FHttpAccessToken* InAuthorizationToken, uint32 PoolSize)
+	FRequestPool(const TCHAR* InServiceUrl, const TCHAR* InEffectiveServiceUrl, FHttpAccessToken* InAuthorizationToken, uint32 PoolSize, uint32 InOverflowLimit = 0)
+	: ActiveOverflowRequests(0)
+	, OverflowLimit(InOverflowLimit)
 	{
 		Pool.AddUninitialized(PoolSize);
+		Requests.AddUninitialized(PoolSize);
 		for (uint8 i = 0; i < Pool.Num(); ++i)
 		{
 			Pool[i].Usage = 0u;
-			Pool[i].Request = new FHttpRequest(InServiceUrl, InEffectiveServiceUrl, InAuthorizationToken, true);
+			Pool[i].Request = new(&Requests[i]) FHttpRequest(InServiceUrl, InEffectiveServiceUrl, InAuthorizationToken, true);
 		}
-		
+
+		InitData = MakeUnique<FInitData>(InServiceUrl, InEffectiveServiceUrl, InAuthorizationToken);
 	}
 
 	~FRequestPool()
 	{
+		check(ActiveOverflowRequests.load() == 0);
 		for (uint8 i = 0; i < Pool.Num(); ++i)
 		{
 			// No requests should be in use by now.
 			check(Pool[i].Usage.load(std::memory_order_acquire) == 0u);
-			delete Pool[i].Request;
 		}
 	}
 
@@ -886,7 +1104,7 @@ struct FRequestPool
 	 * "owned by the caller and need to release it to the pool when work has been completed.
 	 * @return Usable request instance if one is available, otherwise null.
 	 */
-	FHttpRequest* GetFreeRequest()
+	FHttpRequest* GetFreeRequest(bool bUnboundedOverflow = false)
 	{
 		for (uint8 i = 0; i < Pool.Num(); ++i)
 		{
@@ -900,6 +1118,17 @@ struct FRequestPool
 				}
 			}
 		}
+		if (bUnboundedOverflow || (OverflowLimit > 0))
+		{
+			// The use of two operations here (load, then fetch_add) implies that we can exceed the overflow limit because the combined operation
+			// is not atomic.  This is acceptable for our use case.  If we wanted to enforce the hard limit, we could use a loop instead.
+			if (bUnboundedOverflow || (ActiveOverflowRequests.load(std::memory_order_relaxed) < OverflowLimit))
+			{
+				// Create an overflow request (outside of the pre-allocated range of requests)
+				ActiveOverflowRequests.fetch_add(1, std::memory_order_relaxed);
+				return new FHttpRequest(*InitData->ServiceUrl, *InitData->EffectiveServiceUrl, InitData->AccessToken, true);
+			}
+		}
 		return nullptr;
 	}
 
@@ -943,11 +1172,11 @@ struct FRequestPool
 	 * "owned by the caller and need to release it to the pool when work has been completed.
 	 * @return Usable request instance.
 	 */
-	FHttpRequest* WaitForFreeRequest()
+	FHttpRequest* WaitForFreeRequest(bool bUnboundedOverflow = false)
 	{
 		TRACE_CPUPROFILER_EVENT_SCOPE(HttpDDC_WaitForConnPool);
 
-		FHttpRequest* Request = GetFreeRequest();
+		FHttpRequest* Request = GetFreeRequest(bUnboundedOverflow);
 		if (Request == nullptr)
 		{
 			// Make it a fair by allowing each thread to register itself in a FIFO
@@ -962,7 +1191,7 @@ struct FRequestPool
 			{
 				// While waiting, allow us to check if a race occurred and a request has been freed
 				// between the time we checked for free requests and the time we queued ourself as a Waiter.
-				if ((Request = GetFreeRequest()) != nullptr)
+				if ((Request = GetFreeRequest(bUnboundedOverflow)) != nullptr)
 				{
 					// We abandon the FWaiter, it will be freed by the next dequeue
 					// and if it has a request, it will be queued back to the pool.
@@ -985,6 +1214,14 @@ struct FRequestPool
 	 */
 	void ReleaseRequestToPool(FHttpRequest* Request)
 	{
+		if ((Request < Requests.GetData()) || (Request >= (Requests.GetData() + Requests.Num())))
+		{
+			// For overflow requests (outside of the pre-allocated range of requests), just delete it immediately
+			delete Request;
+			ActiveOverflowRequests.fetch_sub(1, std::memory_order_relaxed);
+			return;
+		}
+
 		for (uint8 i = 0; i < Pool.Num(); ++i)
 		{
 			if (Pool[i].Request == Request)
@@ -1014,6 +1251,12 @@ struct FRequestPool
 	 */
 	void MakeRequestShared(FHttpRequest* Request, uint8 Users)
 	{
+		if ((Request < Requests.GetData()) || (Request >= (Requests.GetData() + Requests.Num())))
+		{
+			// Overflow requests (outside of the pre-allocated range of requests), cannot be made shared
+			check(false);
+		}
+
 		check(Users != 0);
 		for (uint8 i = 0; i < Pool.Num(); ++i)
 		{
@@ -1035,7 +1278,24 @@ private:
 	};
 
 	TArray<FEntry> Pool;
+	TArray<FHttpRequest> Requests;
 	FAAArrayQueue<FWaiter> Waiters;
+	std::atomic<uint32> ActiveOverflowRequests;
+	struct FInitData
+	{
+		FString ServiceUrl;
+		FString EffectiveServiceUrl;
+		FHttpAccessToken* AccessToken;
+
+		FInitData(const TCHAR* InServiceUrl, const TCHAR* InEffectiveServiceUrl, FHttpAccessToken* InAccessToken)
+		: ServiceUrl(InServiceUrl)
+		, EffectiveServiceUrl(InEffectiveServiceUrl)
+		, AccessToken(InAccessToken)
+		{
+		}
+	};
+	TUniquePtr<const FInitData> InitData;
+	const uint32 OverflowLimit;
 
 	FRequestPool() = delete;
 };
@@ -1053,7 +1313,8 @@ public:
 	FScopedRequestPtr(FRequestPool* InPool)
 		: Request(InPool->WaitForFreeRequest())
 		, Pool(InPool)
-	{}
+	{
+	}
 
 	~FScopedRequestPtr()
 	{
@@ -1387,7 +1648,9 @@ private:
 				const uint32 Reserved = FMath::Min((uint32)UE_HTTPDDC_BATCH_SIZE, Batch.Reserved.fetch_add(UE_HTTPDDC_BATCH_SIZE, std::memory_order_acquire));
 
 				// Give other threads time to copy their data to batch
-				while (Batch.Ready.load(std::memory_order_acquire) < Reserved) {}
+				while (Batch.Ready.load(std::memory_order_acquire) < Reserved)
+				{
+				}
 
 				// Increment request ref count to reflect all waiting threads
 				InPool->MakeRequestShared(Batch.Request, Reserved);
@@ -1400,7 +1663,9 @@ private:
 
 				// Store away the request and wait until other threads have too
 				Request = Batch.Request;
-				while (Batch.Ready.load(std::memory_order_acquire) > 1) {}
+				while (Batch.Ready.load(std::memory_order_acquire) > 1)
+				{
+				}
 
 				//Reset batch for next use
 				Batch.Complete->Reset();
@@ -1510,7 +1775,8 @@ private:
 			, StartEntryIdx(InEntryIdx)
 			, StartKeyIdx(InKeyIdx)
 			, RequestResult(InRequestResult)
-		{}
+		{
+		}
 
 		bool FindNext(int32& EntryIdx, int32& KeyIdx)
 		{
@@ -1530,7 +1796,8 @@ private:
 					KeyIdx = CurrentKeyIdx;
 					return true;
 				}
-			} while (AdvanceIndices(CurrentEntryIdx, CurrentKeyIdx));
+			}
+			while (AdvanceIndices(CurrentEntryIdx, CurrentKeyIdx));
 
 			return false;
 		}
@@ -1781,7 +2048,8 @@ private:
 			, Bucket(InBucket)
 			, CacheKey(InCacheKey)
 			, Data(InData) // Copies the data!
-		{}
+		{
+		}
 	};
 
 	static TLockFreePointerListUnordered<FQueuedEntry, PLATFORM_CACHE_LINE_SIZE> QueuedPuts;
@@ -2165,6 +2433,7 @@ public:
 
 private:
 	FString Domain;
+	FString EffectiveDomain;
 	FString Namespace;
 	FString StructuredNamespace;
 	FString DefaultBucket;
@@ -2177,6 +2446,7 @@ private:
 	FBackendDebugOptions DebugOptions;
 	TUniquePtr<struct FRequestPool> GetRequestPools[2];
 	TUniquePtr<struct FRequestPool> PutRequestPools[2];
+	TUniquePtr<struct FRequestPool> NonBlockingRequestPools;
 	TUniquePtr<struct FHttpAccessToken> Access;
 	bool bIsUsable;
 	bool bReadOnly;
@@ -2185,395 +2455,383 @@ private:
 	EBackendLegacyMode LegacyMode;
 	static inline FHttpCacheStore* AnyInstance = nullptr;
 
-	struct FValueDebugContext
-	{
-		FStringView Name;
-		const FCacheKey& Key;
-		FString Id;
-	};
-
 	bool IsServiceReady();
 	bool AcquireAccessToken();
-	bool ShouldRetryOnError(int64 ResponseCode);
+	bool ShouldRetryOnError(FHttpRequest::Result Result, int64 ResponseCode);
+	bool ShouldRetryOnError(int64 ResponseCode) { return ShouldRetryOnError(FHttpRequest::Success, ResponseCode); }
 
-	bool PutCacheRecord(FStringView Name, const FCacheRecord& Record, const FCacheRecordPolicy& Policy, uint64& OutWriteSize);
-	uint64 PutRef(const FCbPackage& Package, const FCacheKey& Key, FStringView Bucket, bool bFinalize, TArray<FIoHash>& OutNeededBlobHashes, bool& bOutPutCompletedSuccessfully);
+	enum class OperationCategory
+	{
+		Get,
+		Put
+	};
+	template<OperationCategory Category>
+	FHttpRequest* WaitForHttpRequestForOwner(IRequestOwner& Owner, bool bUnboundedOverflow, FRequestPool*& OutPool)
+	{
+		if (!FHttpRequest::AllowAsync())
+		{
+			if (Category == OperationCategory::Get)
+			{
+				OutPool = GetRequestPools[IsInGameThread()].Get();
+			}
+			else
+			{
+				OutPool = PutRequestPools[IsInGameThread()].Get();
+			}
+			return OutPool->WaitForFreeRequest();
+		}
+		else
+		{
+			OutPool = NonBlockingRequestPools.Get();
+			return OutPool->WaitForFreeRequest(bUnboundedOverflow);
+		}
+	}
 
-	FOptionalCacheRecord GetCacheRecordOnly(
-		FStringView Name,
+	struct FGetCacheRecordOnlyResponse
+	{
+		FSharedString Name;
+		FCacheKey Key;
+		uint64 UserData = 0;
+		uint64 BytesReceived = 0;
+		FOptionalCacheRecord Record;
+		EStatus Status = EStatus::Error;
+	};
+	using FOnGetCacheRecordOnlyComplete = TUniqueFunction<void(FGetCacheRecordOnlyResponse&& Response)>;
+	void GetCacheRecordOnlyAsync(
+		IRequestOwner& Owner,
+		const FSharedString& Name,
 		const FCacheKey& Key,
-		const FCacheRecordPolicy& Policy);
-	FOptionalCacheRecord GetCacheRecord(
-		FStringView Name,
+		const FCacheRecordPolicy& Policy,
+		uint64 UserData,
+		FOnGetCacheRecordOnlyComplete&& OnComplete);
+
+	void GetCacheRecordAsync(
+		IRequestOwner& Owner,
+		const FSharedString& Name,
 		const FCacheKey& Key,
 		const FCacheRecordPolicy& Policy,
-		EStatus& OutStatus);
+		uint64 UserData,
+		TUniqueFunction<void(FCacheGetResponse&& Response, uint64 BytesReceived)>&& OnComplete);
+
+	void PutCacheRecordAsync(
+		IRequestOwner& Owner,
+		const FSharedString& Name,
+		const FCacheRecord& Record,
+		const FCacheRecordPolicy& Policy,
+		uint64 UserData,
+		TUniqueFunction<void(FCachePutResponse&& Response, uint64 BytesSent)>&& OnComplete);
 
-	bool PutCacheValue(
-		FStringView Name,
+	void PutCacheValueAsync(
+		IRequestOwner& Owner,
+		const FSharedString& Name,
 		const FCacheKey& Key,
 		const FValue& Value,
 		ECachePolicy Policy,
-		uint64& OutWriteSize);
+		uint64 UserData,
+		TUniqueFunction<void(FCachePutValueResponse&& Response, uint64 BytesSent)>&& OnComplete);
 
-	bool GetCacheValue(
-		FStringView Name,
+	void GetCacheValueAsync(
+		IRequestOwner& Owner,
+		FSharedString Name,
 		const FCacheKey& Key,
 		ECachePolicy Policy,
-		FValue& OutValue,
-		FHttpRequest* ExistingHttpRequest = nullptr);
+		uint64 UserData,
+		FOnCacheGetValueComplete&& OnComplete);
 
-	template<typename ValueType, typename ValueHashGetterType, typename ValueDebugContextGetterType>
-	TBitArray<> TryGetCachedDataBatch(
-		TConstArrayView<ValueType> Values,
-		TArray<FCompressedBuffer>& OutBuffers,
-		ValueHashGetterType ValueHashGetter,
-		ValueDebugContextGetterType ValueDebugContextGetter,
-		FHttpRequest* ExistingHttpRequest = nullptr);
-	template<typename ValueType, typename ValueHashGetterType, typename ValueDebugContextGetterType>
-	TBitArray<> CachedDataProbablyExistsBatch(
-		TConstArrayView<ValueType> Values,
-		ValueHashGetterType ValueHashGetter,
-		ValueDebugContextGetterType ValueDebugContextGetter,
-		FHttpRequest* ExistingHttpRequest = nullptr);
-	template<typename ValueRefType, typename ValueRefKeyGetterType, typename ValueRefDebugContextGetterType>
-	TArray<FValue> RefCachedDataProbablyExistsBatch(
-		TConstArrayView<ValueRefType> ValueRefs,
-		ValueRefKeyGetterType ValueRefKeyGetter,
-		ValueRefDebugContextGetterType ValueRefDebugContextGetter,
-		FHttpRequest* ExistingHttpRequest = nullptr);
+	void RefCachedDataProbablyExistsBatchAsync(
+		IRequestOwner& Owner,
+		TConstArrayView<FCacheGetValueRequest> ValueRefs,
+		FOnCacheGetValueComplete&& OnComplete);
+
+	class FPutPackageOp;
+	class FGetRecordOp;
 };
 
-FHttpCacheStore::FHttpCacheStore(
-	const TCHAR* InServiceUrl, 
-	bool bResolveHostCanonicalName,
-	const TCHAR* InNamespace, 
-	const TCHAR* InStructuredNamespace, 
-	const TCHAR* InOAuthProvider,
-	const TCHAR* InOAuthClientId,
-	const TCHAR* InOAuthSecret,
-	const TCHAR* InOAuthScope,
-	const EBackendLegacyMode InLegacyMode,
-	const bool bInReadOnly)
-	: Domain(InServiceUrl)
-	, Namespace(InNamespace)
-	, StructuredNamespace(InStructuredNamespace)
-	, DefaultBucket(TEXT("default"))
-	, OAuthProvider(InOAuthProvider)
-	, OAuthClientId(InOAuthClientId)
-	, OAuthSecret(InOAuthSecret)
-	, OAuthScope(InOAuthScope)
-	, Access(nullptr)
-	, bIsUsable(false)
-	, bReadOnly(bInReadOnly)
-	, FailedLoginAttempts(0)
-	, SpeedClass(ESpeedClass::Slow)
-	, LegacyMode(InLegacyMode)
+//----------------------------------------------------------------------------------------------------------
+// FHttpCacheStore::FPutPackageOp
+//----------------------------------------------------------------------------------------------------------
+class FHttpCacheStore::FPutPackageOp final : public FThreadSafeRefCountedObject
 {
-#if WITH_DATAREQUEST_HELPER
-	FDataRequestHelper::StaticInitialize();
-#endif
-	if (IsServiceReady() && AcquireAccessToken())
+public:
+
+	struct FCachePutPackageResponse
 	{
-		FString EffectiveDomain;
-		FString OriginalDomainPrefix;
-		TAnsiStringBuilder<64> DomainResolveName;
+		FSharedString Name;
+		FCacheKey Key;
+		uint64 UserData = 0;
+		uint64 BytesSent = 0;
+		EStatus Status = EStatus::Error;
+	};
+	using FOnCachePutPackageComplete = TUniqueFunction<void(FCachePutPackageResponse&& Response)>;
 
-		if (Domain.StartsWith(TEXT("http://")))
-		{
-			DomainResolveName << Domain.RightChop(7);
-			OriginalDomainPrefix = TEXT("http://");
-		}
-		else if (Domain.StartsWith(TEXT("https://")))
-		{
-			DomainResolveName << Domain.RightChop(8);
-			OriginalDomainPrefix = TEXT("https://");
-		}
-		else
-		{
-			DomainResolveName << Domain;
-		}
-
-		addrinfo* AddrResult = nullptr;
-		addrinfo AddrHints;
-		FMemory::Memset(&AddrHints, 0, sizeof(AddrHints));
-		AddrHints.ai_flags = AI_CANONNAME;
-		AddrHints.ai_family = AF_UNSPEC;
-		if (bResolveHostCanonicalName && !::getaddrinfo(*DomainResolveName, nullptr, &AddrHints, &AddrResult))
-		{
-			if (AddrResult->ai_canonname)
-			{
-				// Swap the domain with a canonical name from DNS so that if we are using regional redirection, we pin to a region.
-				EffectiveDomain = OriginalDomainPrefix + ANSI_TO_TCHAR(AddrResult->ai_canonname);
-
-				UE_LOG(LogDerivedDataCache, Display,
-					TEXT("%s: Pinned to %s based on DNS canonical name."),
-					*Domain, *EffectiveDomain);
-			}
-			else
-			{
-				EffectiveDomain = Domain;
-			}
-
-			::freeaddrinfo(AddrResult);
-		}
-		else
-		{
-			EffectiveDomain = Domain;
-		}
-
-		GetRequestPools[0] = MakeUnique<FRequestPool>(*Domain, *EffectiveDomain, Access.Get(), UE_HTTPDDC_GET_REQUEST_POOL_SIZE);
-		GetRequestPools[1] = MakeUnique<FRequestPool>(*Domain, *EffectiveDomain, Access.Get(), 1);
-		PutRequestPools[0] = MakeUnique<FRequestPool>(*Domain, *EffectiveDomain, Access.Get(), UE_HTTPDDC_PUT_REQUEST_POOL_SIZE);
-		PutRequestPools[1] = MakeUnique<FRequestPool>(*Domain, *EffectiveDomain, Access.Get(), 1);
-		bIsUsable = true;
-	}
-
-	AnyInstance = this;
-}
-
-FHttpCacheStore::~FHttpCacheStore()
-{
-	if (AnyInstance == this)
-	{
-		AnyInstance = nullptr;
-	}
-#if WITH_DATAREQUEST_HELPER
-	FDataRequestHelper::StaticShutdown();
-#endif
-}
-
-FString FHttpCacheStore::GetName() const
-{
-	return Domain;
-}
-
-TBitArray<> FHttpCacheStore::TryToPrefetch(TConstArrayView<FString> CacheKeys)
-{
-	return CachedDataProbablyExistsBatch(CacheKeys);
-}
-
-bool FHttpCacheStore::WouldCache(const TCHAR* CacheKey, TArrayView<const uint8> InData)
-{
-	return IsWritable();
-}
+	/** Performs a multi-request operation for uploading a package of content. */
+	static void PutPackage(
+		FHttpCacheStore& CacheStore,
+		IRequestOwner& Owner,
+		const FSharedString& Name,
+		FCacheKey Key,
+		FCbPackage&& Package,
+		FCacheRecordPolicy Policy,
+		uint64 UserData,
+		FOnCachePutPackageComplete&& OnComplete);
 
-FHttpCacheStore::ESpeedClass FHttpCacheStore::GetSpeedClass() const
-{
-	return SpeedClass;
-}
+private:
+	FHttpCacheStore& CacheStore;
+	IRequestOwner& Owner;
+	const FSharedString Name;
+	const FCacheKey Key;
+	const uint64 UserData;
+	std::atomic<uint64> BytesSent;
+	const FCbObject PackageObject;
+	const FIoHash PackageObjectHash;
+	const uint32 TotalBlobUploads;
+	std::atomic<uint32> SuccessfulBlobUploads;
+	std::atomic<uint32> PendingBlobUploads;
+	FOnCachePutPackageComplete OnComplete;
+
+	struct FCachePutRefResponse
+	{
+		FSharedString Name;
+		FCacheKey Key;
+		uint64 UserData = 0;
+		uint64 BytesSent = 0;
+		TConstArrayView<FIoHash> NeededBlobHashes;
+		EStatus Status = EStatus::Error;
+	};
+	using FOnCachePutRefComplete = TUniqueFunction<void(FCachePutRefResponse&& Response)>;
+
+	FPutPackageOp(
+		FHttpCacheStore& InCacheStore,
+		IRequestOwner& InOwner,
+		const FSharedString& InName,
+		const FCacheKey& InKey,
+		uint64 InUserData,
+		uint64 InBytesSent,
+		const FCbObject& InPackageObject,
+		const FIoHash& InPackageObjectHash,
+		uint32 InTotalBlobUploads,
+		FOnCachePutPackageComplete&& InOnComplete);
+
+	static void PutRefAsync(
+		FHttpCacheStore& CacheStore,
+		IRequestOwner& Owner,
+		const FSharedString& Name,
+		FCacheKey Key,
+		FCbObject Object,
+		FIoHash ObjectHash,
+		uint64 UserData,
+		bool bFinalize,
+		FOnCachePutRefComplete&& OnComplete);
+
+	static void OnPackagePutRefComplete(
+		FHttpCacheStore& CacheStore,
+		IRequestOwner& Owner,
+		const FSharedString& Name,
+		const FCacheKey& Key,
+		FCbPackage&& Package,
+		FCacheRecordPolicy Policy,
+		uint64 UserData,
+		FOnCachePutPackageComplete&& OnComplete,
+		FCachePutRefResponse&& Response);
 
-bool FHttpCacheStore::ApplyDebugOptions(FBackendDebugOptions& InOptions)
-{
-	DebugOptions = InOptions;
-	return true;
-}
+	FHttpRequest::ECompletionBehavior OnCompressedBlobUploadComplete(
+		FHttpRequest::Result HttpResult,
+		FHttpRequest* Request);
 
+	void OnPutRefFinalizationComplete(
+		FCachePutRefResponse&& Response);
 
-bool FHttpCacheStore::IsServiceReady()
-{
-	FHttpRequest Request(*Domain, *Domain, nullptr, false);
-	FHttpRequest::Result Result = Request.PerformBlockingDownload(TEXT("health/ready"), nullptr);
-	
-	if (Result == FHttpRequest::Success && Request.GetResponseCode() == 200)
-	{
-		UE_LOG(LogDerivedDataCache, Display, TEXT("%s: HTTP DDC service status: %s."), *Request.GetName(), *Request.GetResponseAsString());
-		return true;
-	}
-	else
+	FCachePutPackageResponse MakeResponse(uint64 InBytesSent, EStatus Status)
 	{
-		UE_LOG(LogDerivedDataCache, Warning, TEXT("%s: Unable to reach HTTP DDC service at %s. Status: %d . Response: %s"), *Request.GetName(), *Domain, Request.GetResponseCode(), *Request.GetResponseAsString());
-	}
-
-	return false;
-}
+		return FCachePutPackageResponse{ Name, Key, UserData, InBytesSent, Status };
+	};
+};
 
-bool FHttpCacheStore::AcquireAccessToken()
+//----------------------------------------------------------------------------------------------------------
+// FHttpCacheStore::FGetRecordOp
+//----------------------------------------------------------------------------------------------------------
+class FHttpCacheStore::FGetRecordOp final : public FThreadSafeRefCountedObject
 {
-	if (Domain.StartsWith(TEXT("http://localhost")))
-	{
-		UE_LOG(LogDerivedDataCache, Log, TEXT("Connecting to a local host '%s', so skipping authorization"), *Domain);
-		return true;
-	}
+public:
+	/** Performs a multi-request operation for downloading a record. */
+	static void GetRecord(
+		FHttpCacheStore& CacheStore,
+		IRequestOwner& Owner,
+		const FSharedString& Name,
+		const FCacheKey& Key,
+		const FCacheRecordPolicy& Policy,
+		uint64 UserData,
+		TUniqueFunction<void(FCacheGetResponse&& Response, uint64 BytesReceived)>&& OnComplete);
 
-	// Avoid spamming the this if the service is down
-	if (FailedLoginAttempts > UE_HTTPDDC_MAX_FAILED_LOGIN_ATTEMPTS)
+	struct FGetCachedDataBatchResponse
 	{
-		return false;
-	}
-
-	ensureMsgf(OAuthProvider.StartsWith(TEXT("http://")) || OAuthProvider.StartsWith(TEXT("https://")),
-		TEXT("The OAuth provider %s is not valid. Needs to be a fully qualified url."),
-		*OAuthProvider
-	);
+		FSharedString Name;
+		FCacheKey Key;
+		int32 ValueIndex;
+		uint64 BytesReceived = 0;
+		FCompressedBuffer DataBuffer;
+		EStatus Status = EStatus::Error;
+	};
+	using FOnGetCachedDataBatchComplete = TUniqueFunction<void(FGetCachedDataBatchResponse&& Response)>;
 
-	// In case many requests wants to update the token at the same time
-	// get the current serial while we wait to take the CS.
-	const uint32 WantsToUpdateTokenSerial = Access.IsValid() ? Access->GetSerial() : 0u;
+	/** Utility method for fetching a batch of value data. */
+	template<typename ValueType, typename ValueIdGetterType>
+	static void GetDataBatch(
+		FHttpCacheStore& CacheStore,
+		IRequestOwner& Owner,
+		FSharedString Name,
+		const FCacheKey& Key,
+		TConstArrayView<ValueType> Values,
+		ValueIdGetterType ValueIdGetter,
+		FOnGetCachedDataBatchComplete&& OnComplete);
+private:
+	FHttpCacheStore& CacheStore;
+	IRequestOwner& Owner;
+	const FSharedString Name;
+	const FCacheKey Key;
+	const uint64 UserData;
+	std::atomic<uint64> BytesReceived;
+	TArray<FCompressedBuffer> FetchedBuffers;
+	const TArray<FValueWithId> RequiredGets;
+	const TArray<FValueWithId> RequiredHeads;
+	FCacheRecordBuilder RecordBuilder;
+	const uint32 TotalOperations;
+	std::atomic<uint32> SuccessfulOperations;
+	std::atomic<uint32> PendingOperations;
+	TUniqueFunction<void(FCacheGetResponse&& Response, uint64 BytesReceived)> OnComplete;
+
+	FGetRecordOp(
+		FHttpCacheStore& InCacheStore,
+		IRequestOwner& InOwner,
+		const FSharedString& InName,
+		const FCacheKey& InKey,
+		uint64 InUserData,
+		uint64 InBytesReceived,
+		TArray<FValueWithId>&& InRequiredGets,
+		TArray<FValueWithId>&& InRequiredHeads,
+		FCacheRecordBuilder&& InRecordBuilder,
+		TUniqueFunction<void(FCacheGetResponse&& Response, uint64 BytesReceived)>&& InOnComplete);
+
+	static void OnOnlyRecordComplete(
+		FHttpCacheStore& CacheStore,
+		IRequestOwner& Owner,
+		const FCacheRecordPolicy& Policy,
+		TUniqueFunction<void(FCacheGetResponse&& Response, uint64 BytesReceived)>&& OnComplete,
+		FGetCacheRecordOnlyResponse&& Response);
 
+	struct FCachedDataProbablyExistsBatchResponse
 	{
-		FScopeLock Lock(&AccessCs);
-
-		// Check if someone has beaten us to update the token, then it 
-		// should now be valid.
-		if (Access.IsValid() && Access->GetSerial() > WantsToUpdateTokenSerial)
-		{
-			return true;
-		}
-
-		const uint32 SchemeEnd = OAuthProvider.Find(TEXT("://")) + 3;
-		const uint32 DomainEnd = OAuthProvider.Find(TEXT("/"), ESearchCase::CaseSensitive, ESearchDir::FromStart, SchemeEnd);
-		FString AuthDomain(DomainEnd, *OAuthProvider);
-		FString Uri(*OAuthProvider + DomainEnd + 1);
-
-		FHttpRequest Request(*AuthDomain, *AuthDomain, nullptr, false);
-		FHttpRequest::Result Result = FHttpRequest::Success;
-		if (OAuthProvider.StartsWith(TEXT("http://localhost")))
-		{
-			// Simple unauthenticated call to a local endpoint that mimics
-			// the result from an OIDC provider.
-			Result = Request.PerformBlockingDownload(*Uri, nullptr);
-		}
-		else
-		{
-			// Needs client id and secret to authenticate with an actual OIDC provider.
-
-			// If contents of the secret string is a file path, resolve and read form data.
-			if (OAuthSecret.StartsWith(TEXT("file://")))
-			{
-				FString FilePath = OAuthSecret.Mid(7, OAuthSecret.Len() - 7);
-					FString SecretFileContents;
-				if (FFileHelper::LoadFileToString(SecretFileContents, *FilePath))
-				{
-					// Overwrite the filepath with the actual content.
-					OAuthSecret = SecretFileContents;
-				}
-				else
-				{
-					UE_LOG(LogDerivedDataCache, Warning, TEXT("%s: Failed to read OAuth form data file (%s)."), *Request.GetName(), *OAuthSecret);
-					return false;
-				}
-			}
-
-			FString OAuthFormData = FString::Printf(
-				TEXT("client_id=%s&scope=%s&grant_type=client_credentials&client_secret=%s"),
-				*OAuthClientId,
-				*OAuthScope,
-				*OAuthSecret
-			);
-
-			TArray<uint8> FormData;
-			auto OAuthFormDataUTF8 = FTCHARToUTF8(*OAuthFormData);
-			FormData.Append((uint8*)OAuthFormDataUTF8.Get(), OAuthFormDataUTF8.Length());
-
-			Result = Request.PerformBlockingUpload<FHttpRequest::Post>(*Uri, MakeArrayView(FormData));
-		}
+		FSharedString Name;
+		FCacheKey Key;
+		int32 ValueIndex;
+		EStatus Status = EStatus::Error;
+	};
+	using FOnCachedDataProbablyExistsBatchComplete = TUniqueFunction<void(FCachedDataProbablyExistsBatchResponse&& Response)>;
+	void DataProbablyExistsBatch(
+		TConstArrayView<FValueWithId> Values,
+		FOnCachedDataProbablyExistsBatchComplete&& OnComplete);
 
-		if (Result == FHttpRequest::Success && Request.GetResponseCode() == 200)
-		{
-			TSharedPtr<FJsonObject> ResponseObject = Request.GetResponseAsJsonObject();
-			if (ResponseObject)
-			{
-				FString AccessTokenString;
-				int32 ExpiryTimeSeconds = 0;
-				int32 CurrentTimeSeconds = int32(FPlatformTime::ToSeconds(FPlatformTime::Cycles()));
+	void FinishDataStep(bool bSuccess, uint64 InBytesReceived);
+};
 
-				if (ResponseObject->TryGetStringField(TEXT("access_token"), AccessTokenString) &&
-					ResponseObject->TryGetNumberField(TEXT("expires_in"), ExpiryTimeSeconds))
-				{
-					if (!Access)
-					{
-						Access = MakeUnique<FHttpAccessToken>();
-					}
-					Access->SetHeader(*AccessTokenString);
-					UE_LOG(LogDerivedDataCache, Display, TEXT("%s: Logged in to HTTP DDC services. Expires in %d seconds."), *Request.GetName(), ExpiryTimeSeconds);
+void FHttpCacheStore::FPutPackageOp::PutPackage(
+	FHttpCacheStore& CacheStore,
+	IRequestOwner& Owner,
+	const FSharedString& Name,
+	FCacheKey Key,
+	FCbPackage&& Package,
+	FCacheRecordPolicy Policy,
+	uint64 UserData,
+	FOnCachePutPackageComplete&& OnComplete)
+{
+	// TODO: Jupiter currently always overwrites.  It doesn't have a "write if not present" feature (for records or attachments),
+	//		 but would require one to implement all policy correctly.
 
-					//Schedule a refresh of the token ahead of expiry time (this will not work in commandlets)
-					if (!IsRunningCommandlet())
-					{
-						FTSTicker::GetCoreTicker().AddTicker(FTickerDelegate::CreateLambda(
-							[this](float DeltaTime)
-							{
-								this->AcquireAccessToken();
-								return false;
-							}
-						), ExpiryTimeSeconds - 20.0f);
-					}
-					// Reset failed login attempts, the service is indeed alive.
-					FailedLoginAttempts = 0;
-					return true;
-				}
-			}
-		}
-		else
+	// Initial record upload
+	PutRefAsync(CacheStore, Owner, Name, Key, Package.GetObject(), Package.GetObjectHash(), UserData, false,
+		[&CacheStore, &Owner, Name = FSharedString(Name), Key, Package = MoveTemp(Package), Policy, UserData, OnComplete = MoveTemp(OnComplete)](FCachePutRefResponse&& Response) mutable
 		{
-			UE_LOG(LogDerivedDataCache, Warning, TEXT("%s: Failed to log in to HTTP services. Server responed with code %d."), *Request.GetName(), Request.GetResponseCode());
-			FailedLoginAttempts++;
-		}
-	}
-	return false;
+			return OnPackagePutRefComplete(CacheStore, Owner, Name, Key, MoveTemp(Package), Policy, UserData, MoveTemp(OnComplete), MoveTemp(Response));
+		});
 }
 
-bool FHttpCacheStore::ShouldRetryOnError(int64 ResponseCode)
+FHttpCacheStore::FPutPackageOp::FPutPackageOp(
+	FHttpCacheStore& InCacheStore,
+	IRequestOwner& InOwner,
+	const FSharedString& InName,
+	const FCacheKey& InKey,
+	uint64 InUserData,
+	uint64 InBytesSent,
+	const FCbObject& InPackageObject,
+	const FIoHash& InPackageObjectHash,
+	uint32 InTotalBlobUploads,
+	FOnCachePutPackageComplete&& InOnComplete)
+	: CacheStore(InCacheStore)
+	, Owner(InOwner)
+	, Name(InName)
+	, Key(InKey)
+	, UserData(InUserData)
+	, BytesSent(InBytesSent)
+	, PackageObject(InPackageObject)
+	, PackageObjectHash(InPackageObjectHash)
+	, TotalBlobUploads(InTotalBlobUploads)
+	, SuccessfulBlobUploads(0)
+	, PendingBlobUploads(InTotalBlobUploads)
+	, OnComplete(MoveTemp(InOnComplete))
 {
-	// Access token might have expired, request a new token and try again.
-	if (ResponseCode == 401 && AcquireAccessToken())
-	{
-		return true;
-	}
-
-	// Too many requests, make a new attempt
-	if (ResponseCode == 429)
-	{
-		return true;
-	}
-
-	return false;
 }
 
-uint64 FHttpCacheStore::PutRef(const FCbPackage& Package, const FCacheKey& Key, FStringView Bucket, bool bFinalize, TArray<FIoHash>& OutNeededBlobHashes, bool& bOutPutCompletedSuccessfully)
+void FHttpCacheStore::FPutPackageOp::PutRefAsync(
+	FHttpCacheStore& CacheStore,
+	IRequestOwner& Owner,
+	const FSharedString& Name,
+	FCacheKey Key,
+	FCbObject Object,
+	FIoHash ObjectHash,
+	uint64 UserData,
+	bool bFinalize,
+	FOnCachePutRefComplete&& OnComplete)
 {
-	bOutPutCompletedSuccessfully = false;
+	FString Bucket(Key.Bucket.ToString());
+	Bucket.ToLowerInline();
 
-	uint64 BytesSent = 0;
 	TStringBuilder<256> RefsUri;
-	RefsUri << "api/v1/refs/" << StructuredNamespace << "/" << Bucket << "/" << Key.Hash;
+	RefsUri << "api/v1/refs/" << CacheStore.StructuredNamespace << "/" << Bucket << "/" << Key.Hash;
 	if (bFinalize)
 	{
-		RefsUri << "/finalize/" << Package.GetObjectHash();
+		RefsUri << "/finalize/" << ObjectHash;
 	}
 
-	int64 ResponseCode = 0;
-	for (uint32 Attempts = 0; (Attempts < UE_HTTPDDC_MAX_ATTEMPTS) && !ShouldAbortForShutdown() && (Attempts == 0 || ShouldRetryOnError(ResponseCode)); ++Attempts)
+	FRequestPool* Pool = nullptr;
+	FHttpRequest* Request = CacheStore.WaitForHttpRequestForOwner<OperationCategory::Put>(Owner, bFinalize /* bUnboundedOverflow */, Pool);
+
+	auto OnHttpRequestComplete = [&CacheStore, &Owner, Name = FSharedString(Name), Key, Object, UserData, bFinalize, OnComplete = MoveTemp(OnComplete)]
+	(FHttpRequest::Result HttpResult, FHttpRequest* Request)
 	{
-		FScopedRequestPtr Request(PutRequestPools[IsInGameThread()].Get());
+		TRACE_CPUPROFILER_EVENT_SCOPE(HttpDDC_PutRefAsync_OnHttpRequestComplete);
 
-		if (bFinalize)
-		{
-			Request->PerformBlockingUpload<FHttpRequest::Post>(*RefsUri, TArrayView<const uint8>());
-		}
-		else
+		if (Owner.IsCanceled())
 		{
-			Request->SetHeader(TEXT("X-Jupiter-IoHash"), *WriteToString<48>(Package.GetObjectHash()));
-			Request->PerformBlockingUpload<FHttpRequest::PutCompactBinary>(*RefsUri, MakeConstArrayView(Package.GetObject().GetBuffer().ToShared()));
+			OnComplete({ Name, Key, UserData, Request->GetBytesSent(), {}, EStatus::Canceled });
+			return FHttpRequest::ECompletionBehavior::Done;
 		}
-		ResponseCode = Request->GetResponseCode();
 
+		int64 ResponseCode = Request->GetResponseCode();
 		if (FHttpRequest::IsSuccessResponse(ResponseCode))
 		{
-			BytesSent += Request->GetBytesSent();
+			TArray<FIoHash> NeededBlobHashes;
 
 			// Useful when debugging issues related to compressed/uncompressed blobs being returned from Jupiter
 			const bool bPutRefBlobsAlways = false;
 
 			if (bPutRefBlobsAlways && !bFinalize)
 			{
-				Package.GetObject().IterateAttachments([&OutNeededBlobHashes](FCbFieldView AttachmentFieldView)
+				Object.IterateAttachments([&NeededBlobHashes](FCbFieldView AttachmentFieldView)
 					{
 						FIoHash AttachmentHash = AttachmentFieldView.AsHash();
 						if (!AttachmentHash.IsZero())
 						{
-							OutNeededBlobHashes.Add(AttachmentHash);
+							NeededBlobHashes.Add(AttachmentHash);
 						}
 					});
 			}
@@ -2582,136 +2840,98 @@ uint64 FHttpCacheStore::PutRef(const FCbPackage& Package, const FCacheKey& Key,
 				TArray<FString> NeedsArrayStrings;
 				ResponseObject->TryGetStringArrayField(TEXT("needs"), NeedsArrayStrings);
 
-				OutNeededBlobHashes.Empty(NeedsArrayStrings.Num());
+				NeededBlobHashes.Reserve(NeedsArrayStrings.Num());
 				for (const FString& NeededString : NeedsArrayStrings)
 				{
 					FIoHash BlobHash;
 					LexFromString(BlobHash, *NeededString);
 					if (!BlobHash.IsZero())
 					{
-						OutNeededBlobHashes.Add(BlobHash);
+						NeededBlobHashes.Add(BlobHash);
 					}
 				}
 			}
-			else
-			{
-				OutNeededBlobHashes.Empty();
-			}
 
-			bOutPutCompletedSuccessfully = true;
-			break;
+			OnComplete({ Name, Key, UserData, Request->GetBytesSent(), NeededBlobHashes, EStatus::Ok });
+			return FHttpRequest::ECompletionBehavior::Done;
 		}
-		else
+
+		if (!ShouldAbortForShutdown() && CacheStore.ShouldRetryOnError(HttpResult, ResponseCode) && ((Request->GetAttempts()+1) < UE_HTTPDDC_MAX_ATTEMPTS))
 		{
-			OutNeededBlobHashes.Empty();
+			return FHttpRequest::ECompletionBehavior::Retry;
 		}
-	}
-
-	return BytesSent;
-}
 
-bool FHttpCacheStore::PutCacheRecord(
-	FStringView Name,
-	const FCacheRecord& Record,
-	const FCacheRecordPolicy& Policy,
-	uint64& OutWriteSize)
-{
-	OutWriteSize = 0;
+		OnComplete({ Name, Key, UserData, Request->GetBytesSent(), {}, EStatus::Error });
+		return FHttpRequest::ECompletionBehavior::Done;
+	};
 
-	if (!IsWritable())
+	if (bFinalize)
 	{
-		UE_LOG(LogDerivedDataCache, VeryVerbose,
-			TEXT("%s: Skipped put of %s from '%.*s' because this cache store is read-only"),
-			*GetName(), *WriteToString<96>(Record.GetKey()), Name.Len(), Name.GetData());
-		return false;
+		Request->EnqueueAsyncUpload<FHttpRequest::Post>(Owner, Pool, *RefsUri, FSharedBuffer(), MoveTemp(OnHttpRequestComplete));
 	}
-
-	const FCacheKey& Key = Record.GetKey();
-	const ECachePolicy RecordPolicy = Policy.GetRecordPolicy();
-
-	// Skip the request if storing to the cache is disabled.
-	// Http backends won't generally be "local" but including handling for this possibility for consistency
-	const ECachePolicy StoreFlag = SpeedClass == ESpeedClass::Local ? ECachePolicy::StoreLocal : ECachePolicy::StoreRemote;
-	if (!EnumHasAnyFlags(RecordPolicy, StoreFlag))
+	else
 	{
-		UE_LOG(LogDerivedDataCache, VeryVerbose, TEXT("%s: Skipped put of %s from '%.*s' due to cache policy"),
-			*GetName(), *WriteToString<96>(Key), Name.Len(), Name.GetData());
-		return false;
+		Request->SetHeader(TEXT("X-Jupiter-IoHash"), *WriteToString<48>(ObjectHash));
+		Request->EnqueueAsyncUpload<FHttpRequest::PutCompactBinary>(Owner, Pool, *RefsUri, Object.GetBuffer().ToShared(), MoveTemp(OnHttpRequestComplete));
 	}
+}
 
-	if (DebugOptions.ShouldSimulatePutMiss(Key))
+void FHttpCacheStore::FPutPackageOp::OnPackagePutRefComplete(
+	FHttpCacheStore& CacheStore,
+	IRequestOwner& Owner,
+	const FSharedString& Name,
+	const FCacheKey& Key,
+	FCbPackage&& Package,
+	FCacheRecordPolicy Policy,
+	uint64 UserData,
+	FOnCachePutPackageComplete&& OnComplete,
+	FCachePutRefResponse&& Response)
+{
+	if (Response.Status != EStatus::Ok)
 	{
-		UE_LOG(LogDerivedDataCache, Verbose, TEXT("%s: Simulated miss for put of %s from '%.*s'"),
-			*GetName(), *WriteToString<96>(Key), Name.Len(), Name.GetData());
-		return false;
+		if (Response.Status == EStatus::Error)
+		{
+			UE_LOG(LogDerivedDataCache, Log, TEXT("%s: Failed to put reference object for put of %s from '%s'"),
+				*CacheStore.GetName(), *WriteToString<96>(Response.Key), *Response.Name);
+		}
+		return OnComplete(FCachePutPackageResponse{ Name, Key, UserData, Response.BytesSent, Response.Status });
 	}
 
-	// TODO: Jupiter currently always overwrites.  It doesn't have a "write if not present" feature (for records or attachments),
-	//		 but would require one to implement all policy correctly.
-
-	FString Bucket(Key.Bucket.ToString());
-	Bucket.ToLowerInline();
-
-	bool bPutCompletedSuccessfully = false;
-
-	FCbPackage Package = Record.Save();
-	TArray<FIoHash> NeededBlobHashes;
-
-	// Initial record upload
-	size_t PutRefBytesSent = PutRef(Package, Record.GetKey(), Bucket, false, NeededBlobHashes, bPutCompletedSuccessfully);
-	OutWriteSize += PutRefBytesSent;
-
-	if (!bPutCompletedSuccessfully)
+	struct FCompressedBlobUpload
 	{
-		UE_LOG(LogDerivedDataCache, Log, TEXT("%s: Failed to put reference object for put of %s from '%.*s'"),
-			*GetName(), *WriteToString<96>(Key), Name.Len(), Name.GetData());
-		return false;
-	}
+		FIoHash Hash;
+		FSharedBuffer BlobBuffer;
+		FCompressedBlobUpload(const FIoHash& InHash, FSharedBuffer&& InBlobBuffer) : Hash(InHash), BlobBuffer(InBlobBuffer)
+		{
+		}
+	};
+
+	TArray<FCompressedBlobUpload> CompressedBlobUploads;
 
 	// TODO: blob uploading and finalization should be replaced with a single batch compressed blob upload endpoint in the future.
 	TStringBuilder<128> ExpectedHashes;
 	bool bExpectedHashesSerialized = false;
 
 	// Needed blob upload (if any missing)
-	for (const FIoHash& NeededBlobHash : NeededBlobHashes)
+	for (const FIoHash& NeededBlobHash : Response.NeededBlobHashes)
 	{
-		TStringBuilder<256> CompressedBlobsUri;
-		CompressedBlobsUri << "api/v1/compressed-blobs/" << StructuredNamespace << "/" << NeededBlobHash;
-
 		if (const FCbAttachment* Attachment = Package.FindAttachment(NeededBlobHash))
 		{
 			FSharedBuffer TempBuffer;
-			TConstArrayView<uint8> BlobArrayView;
 			if (Attachment->IsCompressedBinary())
 			{
 				TempBuffer = Attachment->AsCompressedBinary().GetCompressed().ToShared();
-				BlobArrayView = MakeConstArrayView(TempBuffer);
 			}
 			else if (Attachment->IsBinary())
 			{
-				TempBuffer = FCompressedBuffer::Compress(Attachment->AsCompositeBinary()).GetCompressed().ToShared();
-				BlobArrayView = MakeConstArrayView(TempBuffer);
+				TempBuffer = FValue::Compress(Attachment->AsCompositeBinary()).GetData().GetCompressed().ToShared();
 			}
 			else
 			{
-				TempBuffer = FCompressedBuffer::Compress(Attachment->AsObject().GetBuffer()).GetCompressed().ToShared();
-				BlobArrayView = MakeConstArrayView(TempBuffer);
+				TempBuffer = FValue::Compress(Attachment->AsObject().GetBuffer()).GetData().GetCompressed().ToShared();
 			}
 
-			int64 ResponseCode = 0;
-			for (uint32 Attempts = 0; (Attempts < UE_HTTPDDC_MAX_ATTEMPTS) && !ShouldAbortForShutdown() && (Attempts == 0 || ShouldRetryOnError(ResponseCode)); ++Attempts)
-			{
-				FScopedRequestPtr Request(PutRequestPools[IsInGameThread()].Get());
-				Request->PerformBlockingUpload<FHttpRequest::PutCompressedBlob>(*CompressedBlobsUri, BlobArrayView);
-
-				ResponseCode = Request->GetResponseCode();
-
-				if (FHttpRequest::IsSuccessResponse(ResponseCode))
-				{
-					OutWriteSize += Request->GetBytesSent();
-					break;
-				}
-			}
+			CompressedBlobUploads.Emplace(NeededBlobHash, MoveTemp(TempBuffer));
 		}
 		else
 		{
@@ -2729,651 +2949,1143 @@ bool FHttpCacheStore::PutCacheRecord(
 				}
 				bExpectedHashesSerialized = true;
 			}
-			UE_LOG(LogDerivedDataCache, Log, TEXT("%s: Server reported needed hash '%s' that is outside the set of expected hashes (%s) for put of %s from '%.*s'"),
-				*GetName(), *WriteToString<96>(NeededBlobHash), ExpectedHashes.ToString(), *WriteToString<96>(Key), Name.Len(), Name.GetData());
+			UE_LOG(LogDerivedDataCache, Log, TEXT("%s: Server reported needed hash '%s' that is outside the set of expected hashes (%s) for put of %s from '%s'"),
+				*CacheStore.GetName(), *WriteToString<96>(NeededBlobHash), ExpectedHashes.ToString(), *WriteToString<96>(Response.Key), *Response.Name);
 		}
 	}
 
-	// Finalization (if any blobs needed)
-	if (!NeededBlobHashes.IsEmpty())
+	if (CompressedBlobUploads.IsEmpty())
 	{
-		size_t FinalizeRefBytesSent = PutRef(Package, Record.GetKey(), Bucket, true, NeededBlobHashes, bPutCompletedSuccessfully);
-		OutWriteSize += FinalizeRefBytesSent;
+		// No blobs need to be uploaded.  No finalization necessary.
+		return OnComplete(FCachePutPackageResponse{ Name, Key, UserData, Response.BytesSent, EStatus::Ok });
 	}
 
-	return bPutCompletedSuccessfully && NeededBlobHashes.IsEmpty();
+	// Having this be a ref ensures we don't have the op reach 0 ref count as we queue up multiple operations which MAY execute synchronously
+	TRefCountPtr<FPutPackageOp> PutPackageOp = new FPutPackageOp(
+		CacheStore,
+		Owner,
+		Response.Name,
+		Response.Key,
+		Response.UserData,
+		Response.BytesSent,
+		Package.GetObject(),
+		Package.GetObjectHash(),
+		(uint32)CompressedBlobUploads.Num(),
+		MoveTemp(OnComplete)
+	);
+
+	FRequestBarrier Barrier(Owner);
+	for (const FCompressedBlobUpload& CompressedBlobUpload : CompressedBlobUploads)
+	{
+		TStringBuilder<256> CompressedBlobsUri;
+		CompressedBlobsUri << "api/v1/compressed-blobs/" << CacheStore.StructuredNamespace << "/" << CompressedBlobUpload.Hash;
+
+		FRequestPool* Pool = nullptr;
+		FHttpRequest* Request = CacheStore.WaitForHttpRequestForOwner<OperationCategory::Put>(Owner, true /* bUnboundedOverflow */, Pool);
+		Request->EnqueueAsyncUpload<FHttpRequest::PutCompressedBlob>(Owner, Pool, *CompressedBlobsUri, CompressedBlobUpload.BlobBuffer,
+			[PutPackageOp](FHttpRequest::Result HttpResult, FHttpRequest* Request)
+			{
+				return PutPackageOp->OnCompressedBlobUploadComplete(HttpResult, Request);
+			});
+	}
 }
 
-FOptionalCacheRecord FHttpCacheStore::GetCacheRecordOnly(
-	const FStringView Name,
-	const FCacheKey& Key,
-	const FCacheRecordPolicy& Policy)
+FHttpRequest::ECompletionBehavior FHttpCacheStore::FPutPackageOp::OnCompressedBlobUploadComplete(
+	FHttpRequest::Result HttpResult,
+	FHttpRequest* Request)
 {
-	if (!IsUsable())
+	int64 ResponseCode = Request->GetResponseCode();
+	bool bIsSuccessResponse = FHttpRequest::IsSuccessResponse(ResponseCode);
+
+	if (!bIsSuccessResponse && !ShouldAbortForShutdown() && !Owner.IsCanceled() && CacheStore.ShouldRetryOnError(HttpResult, ResponseCode) && ((Request->GetAttempts()+1) < UE_HTTPDDC_MAX_ATTEMPTS))
 	{
-		UE_LOG(LogDerivedDataCache, VeryVerbose,
-			TEXT("%s: Skipped get of %s from '%.*s' because this cache store is not available"),
-			*GetName(), *WriteToString<96>(Key), Name.Len(), Name.GetData());
-		return FOptionalCacheRecord();
+		return FHttpRequest::ECompletionBehavior::Retry;
 	}
 
-	// Skip the request if querying the cache is disabled.
-	const ECachePolicy QueryPolicy = SpeedClass == ESpeedClass::Local
-		? ECachePolicy::QueryLocal : ECachePolicy::QueryRemote;
-	if (!EnumHasAnyFlags(Policy.GetRecordPolicy(), QueryPolicy))
+	BytesSent.fetch_add(Request->GetBytesSent(), std::memory_order_relaxed);
+	if (bIsSuccessResponse)
 	{
-		UE_LOG(LogDerivedDataCache, VeryVerbose, TEXT("%s: Skipped get of %s from '%.*s' due to cache policy"),
-			*GetName(), *WriteToString<96>(Key), Name.Len(), Name.GetData());
-		return FOptionalCacheRecord();
+		SuccessfulBlobUploads.fetch_add(1, std::memory_order_relaxed);
 	}
 
-	if (DebugOptions.ShouldSimulateGetMiss(Key))
+	if (PendingBlobUploads.fetch_sub(1, std::memory_order_relaxed) == 1)
 	{
-		UE_LOG(LogDerivedDataCache, Verbose, TEXT("%s: Simulated miss for get of %s from '%.*s'"),
-			*GetName(), *WriteToString<96>(Key), Name.Len(), Name.GetData());
-		return FOptionalCacheRecord();
+		if (Owner.IsCanceled())
+		{
+			OnComplete(MakeResponse(BytesSent.load(std::memory_order_relaxed), EStatus::Canceled));
+			return FHttpRequest::ECompletionBehavior::Done;
+		}
+
+		uint32 LocalSuccessfulBlobUploads = SuccessfulBlobUploads.load(std::memory_order_relaxed);
+		if (LocalSuccessfulBlobUploads == TotalBlobUploads)
+		{
+			// Perform finalization
+			PutRefAsync(CacheStore, Owner, Name, Key, PackageObject, PackageObjectHash, UserData, true,
+				[this](FCachePutRefResponse&& Response)
+				{
+					return OnPutRefFinalizationComplete(MoveTemp(Response));
+				});
+		}
+		else
+		{
+			uint32 FailedBlobUploads = (uint32)(TotalBlobUploads - LocalSuccessfulBlobUploads);
+			UE_LOG(LogDerivedDataCache, Log, TEXT("%s: Failed to put %d/%d blobs for put of %s from '%s'"),
+				*CacheStore.GetName(), FailedBlobUploads, TotalBlobUploads, *WriteToString<96>(Key), *Name);
+			OnComplete(MakeResponse(BytesSent.load(std::memory_order_relaxed), EStatus::Error));
+		}
 	}
+	return FHttpRequest::ECompletionBehavior::Done;
+}
 
-	FScopedRequestPtr Request(GetRequestPools[IsInGameThread()].Get());
+void FHttpCacheStore::FPutPackageOp::OnPutRefFinalizationComplete(
+	FCachePutRefResponse&& Response)
+{
+	BytesSent.fetch_add(Response.BytesSent, std::memory_order_relaxed);
 
-	FOptionalCacheRecord Record;
+	if (Response.Status == EStatus::Error)
 	{
-		FString Bucket(Key.Bucket.ToString());
-		Bucket.ToLowerInline();
+		UE_LOG(LogDerivedDataCache, Log, TEXT("%s: Failed to finalize reference object for put of %s from '%s'"),
+			*CacheStore.GetName(), *WriteToString<96>(Key), *Name);
+	}
 
-		TStringBuilder<256> RefsUri;
-		RefsUri << "api/v1/refs/" << StructuredNamespace << "/" << Bucket << "/" << Key.Hash;
+	return OnComplete(MakeResponse(BytesSent.load(std::memory_order_relaxed), Response.Status));
+}
 
-		bool bIsSuccessfulResponse = false;
-		FSharedBuffer ResponseBuffer;
-		int64 ResponseCode = 0;
-		for (uint32 Attempts = 0; (Attempts < UE_HTTPDDC_MAX_ATTEMPTS) && !ShouldAbortForShutdown() && (Attempts == 0 || ShouldRetryOnError(ResponseCode)); ++Attempts)
+void FHttpCacheStore::FGetRecordOp::GetRecord(
+	FHttpCacheStore& CacheStore,
+	IRequestOwner& Owner,
+	const FSharedString& Name,
+	const FCacheKey& Key,
+	const FCacheRecordPolicy& Policy,
+	uint64 UserData,
+	TUniqueFunction<void(FCacheGetResponse&& Response, uint64 BytesReceived)>&& OnComplete)
+{
+	CacheStore.GetCacheRecordOnlyAsync(Owner, Name, Key, Policy, UserData, [&CacheStore, &Owner, Policy = FCacheRecordPolicy(Policy), OnComplete = MoveTemp(OnComplete)](FGetCacheRecordOnlyResponse&& Response) mutable
+	{
+		OnOnlyRecordComplete(CacheStore, Owner, Policy, MoveTemp(OnComplete), MoveTemp(Response));
+	});
+}
+
+template<typename ValueType, typename ValueIdGetterType>
+void FHttpCacheStore::FGetRecordOp::GetDataBatch(
+	FHttpCacheStore& CacheStore,
+	IRequestOwner& Owner,
+	FSharedString Name,
+	const FCacheKey& Key,
+	TConstArrayView<ValueType> Values,
+	ValueIdGetterType ValueIdGetter,
+	FOnGetCachedDataBatchComplete&& OnComplete)
+{
+	if (Values.IsEmpty())
+	{
+		return;
+	}
+
+	FRequestBarrier Barrier(Owner);
+	TRefCountedUniqueFunction<FOnGetCachedDataBatchComplete>* CompletionFunction = new TRefCountedUniqueFunction<FOnGetCachedDataBatchComplete>(MoveTemp(OnComplete));
+	TRefCountPtr<TRefCountedUniqueFunction<FOnGetCachedDataBatchComplete>> BatchOnCompleteRef(CompletionFunction);
+	for (int32 ValueIndex = 0; ValueIndex < Values.Num(); ++ValueIndex)
+	{
+		const ValueType& Value = Values[ValueIndex];
+		const FIoHash& RawHash = Value.GetRawHash();
+
+		FRequestPool* Pool = nullptr;
+		FHttpRequest* Request = CacheStore.WaitForHttpRequestForOwner<OperationCategory::Get>(Owner, true /* bUnboundedOverflow */, Pool);
+
+		auto OnHttpRequestComplete = [&CacheStore, &Owner, Name = FSharedString(Name), Key = FCacheKey(Key), ValueIndex, Value = Value.RemoveData(), ValueIdGetter, OnCompletePtr = TRefCountPtr<TRefCountedUniqueFunction<FOnGetCachedDataBatchComplete>>(CompletionFunction)]
+		(FHttpRequest::Result HttpResult, FHttpRequest* Request)
 		{
-			if (Attempts > 0)
+			TRACE_CPUPROFILER_EVENT_SCOPE(HttpDDC_GetDataBatch_OnHttpRequestComplete);
+
+			int64 ResponseCode = Request->GetResponseCode();
+			bool bHit = false;
+			FCompressedBuffer CompressedBuffer;
+			if (FHttpRequest::IsSuccessResponse(ResponseCode))
 			{
-				Request->Reset();
+				FString ReceivedContentType;
+				if (Request->GetHeader("Content-Type", ReceivedContentType))
+				{
+					if (ReceivedContentType == TEXT("application/x-ue-comp"))
+					{
+						CompressedBuffer = FCompressedBuffer::FromCompressed(Request->MoveResponseBufferToShared());
+						bHit = true;
+					}
+					else if (ReceivedContentType == TEXT("application/octet-stream"))
+					{
+						CompressedBuffer = FValue::Compress(Request->MoveResponseBufferToShared()).GetData();
+						bHit = true;
+					}
+					else
+					{
+						bHit = false;
+					}
+				}
+				else
+				{
+					CompressedBuffer = FCompressedBuffer::FromCompressed(Request->MoveResponseBufferToShared());
+					bHit = true;
+				}
 			}
 
-			TArray<uint8> ByteArray;
-			Request->SetHeader(TEXT("Accept"), TEXT("application/x-ue-cb"));
-			Request->PerformBlockingDownload(*RefsUri, &ByteArray, {401, 404});
-			ResponseCode = Request->GetResponseCode();
+			if (!ShouldAbortForShutdown() && !Owner.IsCanceled() && CacheStore.ShouldRetryOnError(HttpResult, ResponseCode) && ((Request->GetAttempts()+1) < UE_HTTPDDC_MAX_ATTEMPTS))
+			{
+				return FHttpRequest::ECompletionBehavior::Retry;
+			}
 
-			if (FHttpRequest::IsSuccessResponse(ResponseCode))
+			if (!bHit)
 			{
-				ResponseBuffer = MakeSharedBufferFromArray(MoveTemp(ByteArray));
-				bIsSuccessfulResponse = true;
-				break;
+				UE_LOG(LogDerivedDataCache, Verbose,
+					TEXT("%s: Cache miss with missing value %s with hash %s for %s from '%s'"),
+					*CacheStore.GetName(), *ValueIdGetter(Value), *WriteToString<48>(Value.GetRawHash()), *WriteToString<96>(Key),
+					*Name);
+				OnCompletePtr->GetFunction()({ Name, Key, ValueIndex, Request->GetBytesReceived(), {}, EStatus::Error });
+			}
+			else if (CompressedBuffer.GetRawHash() != Value.GetRawHash())
+			{
+				UE_LOG(LogDerivedDataCache, Display,
+					TEXT("%s: Cache miss with corrupted value %s with hash %s for %s from '%s'"),
+					*CacheStore.GetName(), *ValueIdGetter(Value), *WriteToString<48>(Value.GetRawHash()),
+					*WriteToString<96>(Key), *Name);
+				OnCompletePtr->GetFunction()({ Name, Key, ValueIndex, Request->GetBytesReceived(), {}, EStatus::Error });
 			}
+			else
+			{
+				OnCompletePtr->GetFunction()({ Name, Key, ValueIndex, Request->GetBytesReceived(), MoveTemp(CompressedBuffer), EStatus::Ok });
+			}
+
+			return FHttpRequest::ECompletionBehavior::Done;
+		};
+
+		TStringBuilder<256> CompressedBlobsUri;
+		CompressedBlobsUri << "api/v1/compressed-blobs/" << CacheStore.StructuredNamespace << "/" << RawHash;
+		Request->SetHeader(TEXT("Accept"), TEXT("*/*"));
+		Request->EnqueueAsyncDownload(Owner, Pool, *CompressedBlobsUri, MoveTemp(OnHttpRequestComplete), { 404 });
+	}
+}
+
+FHttpCacheStore::FGetRecordOp::FGetRecordOp(
+	FHttpCacheStore& InCacheStore,
+	IRequestOwner& InOwner,
+	const FSharedString& InName,
+	const FCacheKey& InKey,
+	uint64 InUserData,
+	uint64 InBytesReceived,
+	TArray<FValueWithId>&& InRequiredGets,
+	TArray<FValueWithId>&& InRequiredHeads,
+	FCacheRecordBuilder&& InRecordBuilder,
+	TUniqueFunction<void(FCacheGetResponse&& Response, uint64 BytesReceived)>&& InOnComplete)
+	: CacheStore(InCacheStore)
+	, Owner(InOwner)
+	, Name(InName)
+	, Key(InKey)
+	, UserData(InUserData)
+	, BytesReceived(InBytesReceived)
+	, RequiredGets(MoveTemp(InRequiredGets))
+	, RequiredHeads(MoveTemp(InRequiredHeads))
+	, RecordBuilder(MoveTemp(InRecordBuilder))
+	, TotalOperations(RequiredGets.Num() + RequiredHeads.Num())
+	, SuccessfulOperations(0)
+	, PendingOperations(TotalOperations)
+	, OnComplete(MoveTemp(InOnComplete))
+{
+	FetchedBuffers.AddDefaulted(RequiredGets.Num());
+}
+
+void FHttpCacheStore::FGetRecordOp::OnOnlyRecordComplete(
+	FHttpCacheStore& CacheStore,
+	IRequestOwner& Owner,
+	const FCacheRecordPolicy& Policy,
+	TUniqueFunction<void(FCacheGetResponse&& Response, uint64 BytesReceived)>&& OnComplete,
+	FGetCacheRecordOnlyResponse&& Response)
+{
+	FCacheRecordBuilder RecordBuilder(Response.Key);
+	if (Response.Status != EStatus::Ok)
+	{
+		return OnComplete({ Response.Name, RecordBuilder.Build(), Response.UserData, Response.Status }, Response.BytesReceived);
+	}
+
+	if (!EnumHasAnyFlags(Policy.GetRecordPolicy(), ECachePolicy::SkipMeta))
+	{
+		RecordBuilder.SetMeta(FCbObject(Response.Record.Get().GetMeta()));
+	}
+
+	// TODO: There is not currently a batched GET endpoint for Jupiter.  Once there is, all payload data should be fetched in one call.
+	//		 In the meantime, we try to keep the code structured in a way that is friendly to future batching of GETs.
+
+	TArray<FValueWithId> RequiredGets;
+	TArray<FValueWithId> RequiredHeads;
+
+	for (FValueWithId Value : Response.Record.Get().GetValues())
+	{
+		const ECachePolicy ValuePolicy = Policy.GetValuePolicy(Value.GetId());
+		if (IsValueDataReady(Value, ValuePolicy))
+		{
+			RecordBuilder.AddValue(MoveTemp(Value));
 		}
+		else
+		{
+			if (EnumHasAnyFlags(ValuePolicy, ECachePolicy::SkipData))
+			{
+				RequiredHeads.Emplace(Value);
+			}
+			else
+			{
+				RequiredGets.Emplace(Value);
+			}
+		}
+	}
+
+	if (RequiredGets.IsEmpty() && RequiredHeads.IsEmpty())
+	{
+		return OnComplete({ Response.Name, RecordBuilder.Build(), Response.UserData, Response.Status }, Response.BytesReceived);
+	}
+
+	// Having this be a ref ensures we don't have the op reach 0 ref count in between the start of the exist batch operation and the get batch operation
+	TRefCountPtr<FGetRecordOp> GetRecordOp = new FGetRecordOp(
+		CacheStore,
+		Owner,
+		Response.Name,
+		Response.Key,
+		Response.UserData,
+		Response.BytesReceived,
+		MoveTemp(RequiredGets),
+		MoveTemp(RequiredHeads),
+		MoveTemp(RecordBuilder),
+		MoveTemp(OnComplete)
+	);
 
-		if (!bIsSuccessfulResponse)
+	auto IdGetter = [](const FValueWithId& Value)
+	{
+		return FString(WriteToString<16>(Value.GetId()));
+	};
+
+	{
+		FRequestBarrier Barrier(Owner);
+		GetRecordOp->DataProbablyExistsBatch(GetRecordOp->RequiredHeads, [GetRecordOp](FCachedDataProbablyExistsBatchResponse&& Response)
 		{
-			UE_LOG(LogDerivedDataCache, Verbose, TEXT("%s: Cache miss with missing package for %s from '%.*s'"),
-				*GetName(), *WriteToString<96>(Key), Name.Len(), Name.GetData());
-			return Record;
+			GetRecordOp->FinishDataStep(Response.Status == EStatus::Ok, 0);
+		});
+
+		GetDataBatch<FValueWithId>(CacheStore, Owner, Response.Name, Response.Key, GetRecordOp->RequiredGets, IdGetter, [GetRecordOp](FGetCachedDataBatchResponse&& Response)
+		{
+			GetRecordOp->FetchedBuffers[Response.ValueIndex] = MoveTemp(Response.DataBuffer);
+			GetRecordOp->FinishDataStep(Response.Status == EStatus::Ok, Response.BytesReceived);
+		});
+
+	}
+}
+
+void FHttpCacheStore::FGetRecordOp::DataProbablyExistsBatch(
+	TConstArrayView<FValueWithId> Values,
+	FOnCachedDataProbablyExistsBatchComplete&& InOnComplete)
+{
+	if (Values.IsEmpty())
+	{
+		return;
+	}
+
+	FRequestPool* Pool = nullptr;
+	FHttpRequest* Request = CacheStore.WaitForHttpRequestForOwner<OperationCategory::Get>(Owner, true /* bUnboundedOverflow */, Pool);
+
+	TStringBuilder<256> CompressedBlobsUri;
+	CompressedBlobsUri << "api/v1/compressed-blobs/" << CacheStore.StructuredNamespace << "/exists?";
+	bool bFirstItem = true;
+	for (const FValueWithId& Value : Values)
+	{
+		if (!bFirstItem)
+		{
+			CompressedBlobsUri << "&";
 		}
+		CompressedBlobsUri << "id=" << Value.GetRawHash();
+		bFirstItem = false;
+	}
 
-		if (ValidateCompactBinary(ResponseBuffer, ECbValidateMode::Default) != ECbValidateError::None)
+	auto OnHttpRequestComplete = [this, Values = TArray<FValueWithId>(Values), InOnComplete = MoveTemp(InOnComplete)](FHttpRequest::Result HttpResult, FHttpRequest* Request)
+	{
+		TRACE_CPUPROFILER_EVENT_SCOPE(HttpDDC_DataProbablyExistsBatch_OnHttpRequestComplete);
+
+		int64 ResponseCode = Request->GetResponseCode();
+		if (FHttpRequest::IsSuccessResponse(ResponseCode))
 		{
-			UE_LOG(LogDerivedDataCache, Display, TEXT("%s: Cache miss with invalid package for %s from '%.*s'"),
-				*GetName(), *WriteToString<96>(Key), Name.Len(), Name.GetData());
-			return Record;
+			if (TSharedPtr<FJsonObject> ResponseObject = Request->GetResponseAsJsonObject())
+			{
+				TArray<FString> NeedsArrayStrings;
+				if (ResponseObject->TryGetStringArrayField(TEXT("needs"), NeedsArrayStrings))
+				{
+					if (NeedsArrayStrings.IsEmpty())
+					{
+						for (int32 ValueIndex = 0; ValueIndex < Values.Num(); ++ValueIndex)
+						{
+							const FValueWithId& Value = Values[ValueIndex];
+							UE_LOG(LogDerivedDataCache, Verbose,
+								TEXT("%s: Cache exists miss with missing value %s with hash %s for %s from '%s'"),
+								*CacheStore.GetName(), *WriteToString<16>(Value.GetId()), *WriteToString<48>(Value.GetRawHash()), *WriteToString<96>(Key),
+								*Name);
+							InOnComplete({ Name, Key, ValueIndex, EStatus::Error });
+						}
+						return FHttpRequest::ECompletionBehavior::Done;
+					}
+				}
+
+				TBitArray<> ResultStatus(true, Values.Num());
+				for (const FString& NeedsString : NeedsArrayStrings)
+				{
+					FIoHash NeedHash;
+					LexFromString(NeedHash, *NeedsString);
+					for (int32 ValueIndex = 0; ValueIndex < Values.Num(); ++ValueIndex)
+					{
+						const FValueWithId& Value = Values[ValueIndex];
+						if (ResultStatus[ValueIndex] && NeedHash == Value.GetRawHash())
+						{
+							ResultStatus[ValueIndex] = false;
+							break;
+						}
+					}
+				}
+
+				for (int32 ValueIndex = 0; ValueIndex < Values.Num(); ++ValueIndex)
+				{
+					const FValueWithId& Value = Values[ValueIndex];
+
+					if (ResultStatus[ValueIndex])
+					{
+						UE_LOG(LogDerivedDataCache, VeryVerbose,
+							TEXT("%s: Cache exists hit for %s with hash %s for %s from '%s'"),
+							*CacheStore.GetName(), *WriteToString<16>(Value.GetId()), *WriteToString<48>(Value.GetRawHash()), *WriteToString<96>(Key),
+							*Name);
+						InOnComplete({ Name, Key, ValueIndex, EStatus::Ok });
+					}
+					else
+					{
+						UE_LOG(LogDerivedDataCache, Verbose,
+							TEXT("%s: Cache exists miss with missing value %s with hash %s for %s from '%s'"),
+							*CacheStore.GetName(), *WriteToString<16>(Value.GetId()), *WriteToString<48>(Value.GetRawHash()), *WriteToString<96>(Key),
+							*Name);
+						InOnComplete({ Name, Key, ValueIndex, EStatus::Error });
+					}
+				}
+			}
+			else
+			{
+				for (int32 ValueIndex = 0; ValueIndex < Values.Num(); ++ValueIndex)
+				{
+					UE_LOG(LogDerivedDataCache, Log,
+						TEXT("%s: Cache exists returned invalid results."),
+						*CacheStore.GetName());
+					InOnComplete({ Name, Key, ValueIndex, EStatus::Error });
+				}
+			}
+
+			return FHttpRequest::ECompletionBehavior::Done;
 		}
 
-		Record = FCacheRecord::Load(FCbPackage(FCbObject(ResponseBuffer)));
-		if (Record.IsNull())
+		if (!ShouldAbortForShutdown() && !Owner.IsCanceled() && CacheStore.ShouldRetryOnError(HttpResult, ResponseCode) && ((Request->GetAttempts()+1) < UE_HTTPDDC_MAX_ATTEMPTS))
 		{
-			UE_LOG(LogDerivedDataCache, Display, TEXT("%s: Cache miss with record load failure for %s from '%.*s'"),
-				*GetName(), *WriteToString<96>(Key), Name.Len(), Name.GetData());
-			return Record;
+			return FHttpRequest::ECompletionBehavior::Retry;
 		}
-	}
 
-	return Record;
+		for (int32 ValueIndex = 0; ValueIndex < Values.Num(); ++ValueIndex)
+		{
+			const FValueWithId& Value = Values[ValueIndex];
+			UE_LOG(LogDerivedDataCache, Verbose, TEXT("%s: Cache miss with failed HTTP request for %s from '%s'"),
+				*CacheStore.GetName(), *WriteToString<96>(Key), *Name);
+			InOnComplete({Name, Key, ValueIndex, EStatus::Error});
+		}
+		return FHttpRequest::ECompletionBehavior::Done;
+	};
+	FSharedBuffer DummyBuffer;
+	Request->EnqueueAsyncUpload<FHttpRequest::Post>(Owner, Pool, *CompressedBlobsUri, DummyBuffer, MoveTemp(OnHttpRequestComplete));
 }
 
-bool FHttpCacheStore::PutCacheValue(
-	const FStringView Name,
-	const FCacheKey& Key,
-	const FValue& Value,
-	const ECachePolicy Policy,
-	uint64& OutWriteSize)
+void FHttpCacheStore::FGetRecordOp::FinishDataStep(bool bSuccess, uint64 InBytesReceived)
 {
-	if (!IsWritable())
+	BytesReceived.fetch_add(InBytesReceived, std::memory_order_relaxed);
+	if (bSuccess)
 	{
-		UE_LOG(LogDerivedDataCache, VeryVerbose,
-			TEXT("%s: Skipped put of %s from '%.*s' because this cache store is read-only"),
-			*GetName(), *WriteToString<96>(Key), Name.Len(), Name.GetData());
-		return false;
+		SuccessfulOperations.fetch_add(1, std::memory_order_relaxed);
 	}
 
-	// Skip the request if storing to the cache is disabled.
-	const ECachePolicy StoreFlag = SpeedClass == ESpeedClass::Local ? ECachePolicy::StoreLocal : ECachePolicy::StoreRemote;
-	if (!EnumHasAnyFlags(Policy, StoreFlag))
+	if (PendingOperations.fetch_sub(1, std::memory_order_acq_rel) == 1)
 	{
-		UE_LOG(LogDerivedDataCache, VeryVerbose, TEXT("%s: Skipped put of %s from '%.*s' due to cache policy"),
-			*GetName(), *WriteToString<96>(Key), Name.Len(), Name.GetData());
-		return false;
+		EStatus Status = EStatus::Error;
+		uint32 LocalSuccessfulOperations = SuccessfulOperations.load(std::memory_order_relaxed);
+		if (LocalSuccessfulOperations == TotalOperations)
+		{
+			for (int32 Index = 0; Index < RequiredHeads.Num(); ++Index)
+			{
+				RecordBuilder.AddValue(RequiredHeads[Index].RemoveData());
+			}
+
+			for (int32 Index = 0; Index < RequiredGets.Num(); ++Index)
+			{
+				RecordBuilder.AddValue(FValueWithId(RequiredGets[Index].GetId(), FetchedBuffers[Index]));
+			}
+			Status = EStatus::Ok;
+		}
+		OnComplete({Name, RecordBuilder.Build(), UserData, Status}, BytesReceived.load(std::memory_order_relaxed));
 	}
+}
 
-	if (DebugOptions.ShouldSimulatePutMiss(Key))
+FHttpCacheStore::FHttpCacheStore(
+	const TCHAR* InServiceUrl, 
+	bool bResolveHostCanonicalName,
+	const TCHAR* InNamespace, 
+	const TCHAR* InStructuredNamespace, 
+	const TCHAR* InOAuthProvider,
+	const TCHAR* InOAuthClientId,
+	const TCHAR* InOAuthSecret,
+	const TCHAR* InOAuthScope,
+	const EBackendLegacyMode InLegacyMode,
+	const bool bInReadOnly)
+	: Domain(InServiceUrl)
+	, EffectiveDomain(InServiceUrl)
+	, Namespace(InNamespace)
+	, StructuredNamespace(InStructuredNamespace)
+	, DefaultBucket(TEXT("default"))
+	, OAuthProvider(InOAuthProvider)
+	, OAuthClientId(InOAuthClientId)
+	, OAuthSecret(InOAuthSecret)
+	, OAuthScope(InOAuthScope)
+	, Access(nullptr)
+	, bIsUsable(false)
+	, bReadOnly(bInReadOnly)
+	, FailedLoginAttempts(0)
+	, SpeedClass(ESpeedClass::Slow)
+	, LegacyMode(InLegacyMode)
+{
+#if WITH_DATAREQUEST_HELPER
+	FDataRequestHelper::StaticInitialize();
+#endif
+	if (IsServiceReady() && AcquireAccessToken())
 	{
-		UE_LOG(LogDerivedDataCache, Verbose, TEXT("%s: Simulated miss for put of %s from '%.*s'"),
-			*GetName(), *WriteToString<96>(Key), Name.Len(), Name.GetData());
-		return false;
+		FString OriginalDomainPrefix;
+		TAnsiStringBuilder<64> DomainResolveName;
+
+		if (Domain.StartsWith(TEXT("http://")))
+		{
+			DomainResolveName << Domain.RightChop(7);
+			OriginalDomainPrefix = TEXT("http://");
+		}
+		else if (Domain.StartsWith(TEXT("https://")))
+		{
+			DomainResolveName << Domain.RightChop(8);
+			OriginalDomainPrefix = TEXT("https://");
+		}
+		else
+		{
+			DomainResolveName << Domain;
+		}
+
+		addrinfo* AddrResult = nullptr;
+		addrinfo AddrHints;
+		FMemory::Memset(&AddrHints, 0, sizeof(AddrHints));
+		AddrHints.ai_flags = AI_CANONNAME;
+		AddrHints.ai_family = AF_UNSPEC;
+		if (bResolveHostCanonicalName && !::getaddrinfo(*DomainResolveName, nullptr, &AddrHints, &AddrResult))
+		{
+			if (AddrResult->ai_canonname)
+			{
+				// Swap the domain with a canonical name from DNS so that if we are using regional redirection, we pin to a region.
+				EffectiveDomain = OriginalDomainPrefix + ANSI_TO_TCHAR(AddrResult->ai_canonname);
+
+				UE_LOG(LogDerivedDataCache, Display,
+					TEXT("%s: Pinned to %s based on DNS canonical name."),
+					*Domain, *EffectiveDomain);
+			}
+			else
+			{
+				EffectiveDomain = Domain;
+			}
+
+			::freeaddrinfo(AddrResult);
+		}
+		else
+		{
+			EffectiveDomain = Domain;
+		}
+
+		GetRequestPools[0] = MakeUnique<FRequestPool>(*Domain, *EffectiveDomain, Access.Get(), UE_HTTPDDC_GET_REQUEST_POOL_SIZE);
+		GetRequestPools[1] = MakeUnique<FRequestPool>(*Domain, *EffectiveDomain, Access.Get(), UE_HTTPDDC_GET_REQUEST_POOL_SIZE);
+		PutRequestPools[0] = MakeUnique<FRequestPool>(*Domain, *EffectiveDomain, Access.Get(), UE_HTTPDDC_PUT_REQUEST_POOL_SIZE);
+		PutRequestPools[1] = MakeUnique<FRequestPool>(*Domain, *EffectiveDomain, Access.Get(), UE_HTTPDDC_PUT_REQUEST_POOL_SIZE);
+		// Allowing the non-blocking requests to overflow to double their pre-allocated size before we start waiting for one to free up.
+		NonBlockingRequestPools = MakeUnique<FRequestPool>(*Domain, *EffectiveDomain, Access.Get(), UE_HTTPDDC_NONBLOCKING_REQUEST_POOL_SIZE, UE_HTTPDDC_NONBLOCKING_REQUEST_POOL_SIZE);
+		bIsUsable = true;
 	}
 
-	// TODO: Jupiter currently always overwrites.  It doesn't have a "write if not present" feature (for records or attachments),
-	//		 but would require one to implement all policy correctly.
+	AnyInstance = this;
+}
 
-	FString Bucket(Key.Bucket.ToString());
-	Bucket.ToLowerInline();
+FHttpCacheStore::~FHttpCacheStore()
+{
+	if (AnyInstance == this)
+	{
+		AnyInstance = nullptr;
+	}
+#if WITH_DATAREQUEST_HELPER
+	FDataRequestHelper::StaticShutdown();
+#endif
+	FHttpSharedData::bAsyncThreadShutdownRequested.store(true, std::memory_order_relaxed);
+}
 
-	bool bPutCompletedSuccessfully = false;
+FString FHttpCacheStore::GetName() const
+{
+	return Domain;
+}
 
-	FCbWriter Writer;
-	Writer.BeginObject();
-	Writer.AddBinaryAttachment("RawHash", Value.GetRawHash());
-	Writer.AddInteger("RawSize", Value.GetRawSize());
-	Writer.EndObject();
+TBitArray<> FHttpCacheStore::TryToPrefetch(TConstArrayView<FString> CacheKeys)
+{
+	return CachedDataProbablyExistsBatch(CacheKeys);
+}
 
-	FCbPackage Package(Writer.Save().AsObject());
-	TArray<FIoHash> NeededBlobHashes;
+bool FHttpCacheStore::WouldCache(const TCHAR* CacheKey, TArrayView<const uint8> InData)
+{
+	return IsWritable();
+}
 
-	// Initial record upload
-	size_t PutRefBytesSent = PutRef(Package, Key, Bucket, false, NeededBlobHashes, bPutCompletedSuccessfully);
-	OutWriteSize += PutRefBytesSent;
+FHttpCacheStore::ESpeedClass FHttpCacheStore::GetSpeedClass() const
+{
+	return SpeedClass;
+}
 
-	if (!bPutCompletedSuccessfully)
+bool FHttpCacheStore::ApplyDebugOptions(FBackendDebugOptions& InOptions)
+{
+	DebugOptions = InOptions;
+	return true;
+}
+
+
+bool FHttpCacheStore::IsServiceReady()
+{
+	FHttpRequest Request(*Domain, *Domain, nullptr, false);
+	FHttpRequest::Result Result = Request.PerformBlockingDownload(TEXT("health/ready"), nullptr);
+	
+	if (Result == FHttpRequest::Success && Request.GetResponseCode() == 200)
+	{
+		UE_LOG(LogDerivedDataCache, Display, TEXT("%s: HTTP DDC service status: %s."), *Request.GetName(), *Request.GetResponseAsString());
+		return true;
+	}
+	else
+	{
+		UE_LOG(LogDerivedDataCache, Warning, TEXT("%s: Unable to reach HTTP DDC service at %s. Status: %d . Response: %s"), *Request.GetName(), *Domain, Request.GetResponseCode(), *Request.GetResponseAsString());
+	}
+
+	return false;
+}
+
+bool FHttpCacheStore::AcquireAccessToken()
+{
+	if (Domain.StartsWith(TEXT("http://localhost")))
+	{
+		UE_LOG(LogDerivedDataCache, Log, TEXT("Connecting to a local host '%s', so skipping authorization"), *Domain);
+		return true;
+	}
+
+	// Avoid spamming the this if the service is down
+	if (FailedLoginAttempts > UE_HTTPDDC_MAX_FAILED_LOGIN_ATTEMPTS)
 	{
-		UE_LOG(LogDerivedDataCache, Log, TEXT("%s: Failed to put reference object for put of %s from '%.*s'"),
-			*GetName(), *WriteToString<96>(Key), Name.Len(), Name.GetData());
 		return false;
 	}
 
-	// Needed blob upload (if any missing)
-	if (!NeededBlobHashes.IsEmpty())
+	ensureMsgf(OAuthProvider.StartsWith(TEXT("http://")) || OAuthProvider.StartsWith(TEXT("https://")),
+		TEXT("The OAuth provider %s is not valid. Needs to be a fully qualified url."),
+		*OAuthProvider
+	);
+
+	// In case many requests wants to update the token at the same time
+	// get the current serial while we wait to take the CS.
+	const uint32 WantsToUpdateTokenSerial = Access.IsValid() ? Access->GetSerial() : 0u;
+
 	{
-		if (NeededBlobHashes.Num() != 1)
+		FScopeLock Lock(&AccessCs);
+
+		// Check if someone has beaten us to update the token, then it 
+		// should now be valid.
+		if (Access.IsValid() && Access->GetSerial() > WantsToUpdateTokenSerial)
+		{
+			return true;
+		}
+
+		const uint32 SchemeEnd = OAuthProvider.Find(TEXT("://")) + 3;
+		const uint32 DomainEnd = OAuthProvider.Find(TEXT("/"), ESearchCase::CaseSensitive, ESearchDir::FromStart, SchemeEnd);
+		FString AuthDomain(DomainEnd, *OAuthProvider);
+		FString Uri(*OAuthProvider + DomainEnd + 1);
+
+		FHttpRequest Request(*AuthDomain, *AuthDomain, nullptr, false);
+		FHttpRequest::Result Result = FHttpRequest::Success;
+		if (OAuthProvider.StartsWith(TEXT("http://localhost")))
+		{
+			// Simple unauthenticated call to a local endpoint that mimics
+			// the result from an OIDC provider.
+			Result = Request.PerformBlockingDownload(*Uri, nullptr);
+		}
+		else
 		{
-			TStringBuilder<128> NeededHashString;
-			bool bFirstHash = true;
-			for (const FIoHash& NeededBlobHash : NeededBlobHashes)
+			// Needs client id and secret to authenticate with an actual OIDC provider.
+
+			// If contents of the secret string is a file path, resolve and read form data.
+			if (OAuthSecret.StartsWith(TEXT("file://")))
 			{
-				if (!bFirstHash)
+				FString FilePath = OAuthSecret.Mid(7, OAuthSecret.Len() - 7);
+					FString SecretFileContents;
+				if (FFileHelper::LoadFileToString(SecretFileContents, *FilePath))
+				{
+					// Overwrite the filepath with the actual content.
+					OAuthSecret = SecretFileContents;
+				}
+				else
 				{
-					NeededHashString << TEXT(", ");
+					UE_LOG(LogDerivedDataCache, Warning, TEXT("%s: Failed to read OAuth form data file (%s)."), *Request.GetName(), *OAuthSecret);
+					return false;
 				}
-				NeededHashString << NeededBlobHash;
-				bFirstHash = false;
 			}
 
-			UE_LOG(LogDerivedDataCache, Log, TEXT("%s: Server reported unexpected needed hash quantity '%d' (%s) for put of %s from '%.*s'"),
-				*GetName(), NeededBlobHashes.Num(), NeededHashString.ToString(), *WriteToString<96>(Key), Name.Len(), Name.GetData());
-			return false;
-		}
-
-		if (NeededBlobHashes[0] != Value.GetRawHash())
-		{
-			UE_LOG(LogDerivedDataCache, Log, TEXT("%s: Server reported needed hash '%s' that is outside the set of expected hashes (%s) for put of %s from '%.*s'"),
-				*GetName(), *WriteToString<96>(NeededBlobHashes[0]), *WriteToString<96>(Value.GetRawHash()), *WriteToString<96>(Key), Name.Len(), Name.GetData());
-			return false;
-		}
+			FString OAuthFormData = FString::Printf(
+				TEXT("client_id=%s&scope=%s&grant_type=client_credentials&client_secret=%s"),
+				*OAuthClientId,
+				*OAuthScope,
+				*OAuthSecret
+			);
 
-		TStringBuilder<256> CompressedBlobsUri;
-		CompressedBlobsUri << "api/v1/compressed-blobs/" << StructuredNamespace << "/" << Value.GetRawHash();
+			TArray<uint8> FormData;
+			auto OAuthFormDataUTF8 = FTCHARToUTF8(*OAuthFormData);
+			FormData.Append((uint8*)OAuthFormDataUTF8.Get(), OAuthFormDataUTF8.Length());
 
-		FSharedBuffer TempBuffer = Value.GetData().GetCompressed().ToShared();
-		FScopedRequestPtr Request(PutRequestPools[IsInGameThread()].Get());
+			Result = Request.PerformBlockingUpload<FHttpRequest::Post>(*Uri, MakeArrayView(FormData));
+		}
 
-		int64 ResponseCode = 0;
-		for (uint32 Attempts = 0; (Attempts < UE_HTTPDDC_MAX_ATTEMPTS) && !ShouldAbortForShutdown() && (Attempts == 0 || ShouldRetryOnError(ResponseCode)); ++Attempts)
+		if (Result == FHttpRequest::Success && Request.GetResponseCode() == 200)
 		{
-			if (Attempts > 0)
+			TSharedPtr<FJsonObject> ResponseObject = Request.GetResponseAsJsonObject();
+			if (ResponseObject)
 			{
-				Request->Reset();
-			}
-
-			Request->PerformBlockingUpload<FHttpRequest::PutCompressedBlob>(*CompressedBlobsUri, MakeConstArrayView(TempBuffer));
+				FString AccessTokenString;
+				int32 ExpiryTimeSeconds = 0;
+				int32 CurrentTimeSeconds = int32(FPlatformTime::ToSeconds(FPlatformTime::Cycles()));
 
-			ResponseCode = Request->GetResponseCode();
+				if (ResponseObject->TryGetStringField(TEXT("access_token"), AccessTokenString) &&
+					ResponseObject->TryGetNumberField(TEXT("expires_in"), ExpiryTimeSeconds))
+				{
+					if (!Access)
+					{
+						Access = MakeUnique<FHttpAccessToken>();
+					}
+					Access->SetHeader(*AccessTokenString);
+					UE_LOG(LogDerivedDataCache, Display, TEXT("%s: Logged in to HTTP DDC services. Expires in %d seconds."), *Request.GetName(), ExpiryTimeSeconds);
 
-			if (FHttpRequest::IsSuccessResponse(ResponseCode))
-			{
-				OutWriteSize += Request->GetBytesSent();
-				break;
+					//Schedule a refresh of the token ahead of expiry time (this will not work in commandlets)
+					if (!IsRunningCommandlet())
+					{
+						FTSTicker::GetCoreTicker().AddTicker(FTickerDelegate::CreateLambda(
+							[this](float DeltaTime)
+							{
+								this->AcquireAccessToken();
+								return false;
+							}
+						), ExpiryTimeSeconds - 20.0f);
+					}
+					// Reset failed login attempts, the service is indeed alive.
+					FailedLoginAttempts = 0;
+					return true;
+				}
 			}
 		}
+		else
+		{
+			UE_LOG(LogDerivedDataCache, Warning, TEXT("%s: Failed to log in to HTTP services. Server responed with code %d."), *Request.GetName(), Request.GetResponseCode());
+			FailedLoginAttempts++;
+		}
+	}
+	return false;
+}
+
+bool FHttpCacheStore::ShouldRetryOnError(FHttpRequest::Result Result, int64 ResponseCode)
+{
+	if (Result == FHttpRequest::FailedTimeout)
+	{
+		return true;
+	}
 
-		OutWriteSize += PutRef(Package, Key, Bucket, true, NeededBlobHashes, bPutCompletedSuccessfully);
+	// Access token might have expired, request a new token and try again.
+	if (ResponseCode == 401 && AcquireAccessToken())
+	{
+		return true;
 	}
 
+	// Too many requests, make a new attempt
+	if (ResponseCode == 429)
+	{
+		return true;
+	}
 
-	return bPutCompletedSuccessfully && NeededBlobHashes.IsEmpty();
+	return false;
 }
 
-bool FHttpCacheStore::GetCacheValue(
-	const FStringView Name,
+void FHttpCacheStore::GetCacheRecordOnlyAsync(
+	IRequestOwner& Owner,
+	const FSharedString& Name,
 	const FCacheKey& Key,
-	const ECachePolicy Policy,
-	FValue& OutValue,
-	FHttpRequest* ExistingHttpRequest)
+	const FCacheRecordPolicy& Policy,
+	uint64 UserData,
+	FOnGetCacheRecordOnlyComplete&& OnComplete)
 {
+	auto MakeResponse = [Name = FSharedString(Name), Key, UserData](uint64 BytesReceived, EStatus Status)
+	{
+		return FGetCacheRecordOnlyResponse{ Name, Key, UserData, BytesReceived, {}, Status };
+	};
+
 	if (!IsUsable())
 	{
 		UE_LOG(LogDerivedDataCache, VeryVerbose,
-			TEXT("%s: Skipped get of %s from '%.*s' because this cache store is not available"),
-			*GetName(), *WriteToString<96>(Key), Name.Len(), Name.GetData());
-		return false;
+			TEXT("%s: Skipped get of %s from '%s' because this cache store is not available"),
+			*GetName(), *WriteToString<96>(Key), *Name);
+		return OnComplete(MakeResponse(0, EStatus::Error));
 	}
 
 	// Skip the request if querying the cache is disabled.
-	const ECachePolicy QueryFlag = SpeedClass == ESpeedClass::Local ? ECachePolicy::QueryLocal : ECachePolicy::QueryRemote;
-	if (!EnumHasAnyFlags(Policy, QueryFlag))
+	const ECachePolicy QueryPolicy = SpeedClass == ESpeedClass::Local
+		? ECachePolicy::QueryLocal : ECachePolicy::QueryRemote;
+	if (!EnumHasAnyFlags(Policy.GetRecordPolicy(), QueryPolicy))
 	{
-		UE_LOG(LogDerivedDataCache, VeryVerbose, TEXT("%s: Skipped get of %s from '%.*s' due to cache policy"),
-			*GetName(), *WriteToString<96>(Key), Name.Len(), Name.GetData());
-		return false;
+		UE_LOG(LogDerivedDataCache, VeryVerbose, TEXT("%s: Skipped get of %s from '%s' due to cache policy"),
+			*GetName(), *WriteToString<96>(Key), *Name);
+		return OnComplete(MakeResponse(0, EStatus::Error));
 	}
 
 	if (DebugOptions.ShouldSimulateGetMiss(Key))
 	{
-		UE_LOG(LogDerivedDataCache, Verbose, TEXT("%s: Simulated miss for get of %s from '%.*s'"),
-			*GetName(), *WriteToString<96>(Key), Name.Len(), Name.GetData());
-		return false;
+		UE_LOG(LogDerivedDataCache, Verbose, TEXT("%s: Simulated miss for get of %s from '%s'"),
+			*GetName(), *WriteToString<96>(Key), *Name);
+		return OnComplete(MakeResponse(0, EStatus::Error));
 	}
 
-	const bool bSkipData = EnumHasAnyFlags(Policy, ECachePolicy::SkipData);
-
 	FString Bucket(Key.Bucket.ToString());
 	Bucket.ToLowerInline();
 
 	TStringBuilder<256> RefsUri;
 	RefsUri << "api/v1/refs/" << StructuredNamespace << "/" << Bucket << "/" << Key.Hash;
 
-	FHttpRequest* Request = ExistingHttpRequest;
-	TOptional<FScopedRequestPtr> RequestPoolRequest;
-	if (Request == nullptr)
-	{
-		RequestPoolRequest.Emplace(GetRequestPools[IsInGameThread()].Get());
-		Request = RequestPoolRequest->Get();
-	}
-	else
-	{
-		Request->Reset();
-	}
+	FRequestPool* Pool = nullptr;
+	FHttpRequest* Request = WaitForHttpRequestForOwner<OperationCategory::Get>(Owner, false /* bUnboundedOverflow */, Pool);
 
-	bool bIsSuccessfulResponse = false;
-	FSharedBuffer ResponseBuffer;
-	int64 ResponseCode = 0;
-	for (uint32 Attempts = 0; (Attempts < UE_HTTPDDC_MAX_ATTEMPTS) && !ShouldAbortForShutdown() && (Attempts == 0 || ShouldRetryOnError(ResponseCode)); ++Attempts)
+	auto OnHttpRequestComplete = [this, &Owner, Name = FSharedString(Name), Key, UserData, OnComplete = MoveTemp(OnComplete)]
+	(FHttpRequest::Result HttpResult, FHttpRequest* Request)
 	{
-		if (Attempts > 0)
-		{
-			Request->Reset();
-		}
-
-		TArray<uint8> ByteArray;
-		if (bSkipData)
-		{
-			Request->SetHeader(TEXT("Accept"), TEXT("application/x-ue-cb"));
-		}
-		else
-		{
-			Request->SetHeader(TEXT("Accept"), TEXT("application/x-jupiter-inline"));
-		}
-		Request->PerformBlockingDownload(*RefsUri, &ByteArray, { 401, 404 });
-		ResponseCode = Request->GetResponseCode();
+		TRACE_CPUPROFILER_EVENT_SCOPE(HttpDDC_GetCacheRecordOnlyAsync_OnHttpRequestComplete);
 
+		int64 ResponseCode = Request->GetResponseCode();
 		if (FHttpRequest::IsSuccessResponse(ResponseCode))
 		{
-			ResponseBuffer = MakeSharedBufferFromArray(MoveTemp(ByteArray));
-			bIsSuccessfulResponse = true;
-			break;
-		}
-	}
-
-	if (!bIsSuccessfulResponse)
-	{
-		UE_LOG(LogDerivedDataCache, Verbose, TEXT("%s: Cache miss with missing package for %s from '%.*s'"),
-			*GetName(), *WriteToString<96>(Key), Name.Len(), Name.GetData());
-		return false;
-	}
+			FSharedBuffer ResponseBuffer = Request->MoveResponseBufferToShared();
 
-	if (bSkipData)
-	{
-		if (ValidateCompactBinary(ResponseBuffer, ECbValidateMode::Default) != ECbValidateError::None)
-		{
-			UE_LOG(LogDerivedDataCache, Display, TEXT("%s: Cache miss with invalid package for %s from '%.*s'"),
-				*GetName(), *WriteToString<96>(Key), Name.Len(), Name.GetData());
-			return false;
-		}
+			if (ValidateCompactBinary(ResponseBuffer, ECbValidateMode::Default) != ECbValidateError::None)
+			{
+				UE_LOG(LogDerivedDataCache, Log, TEXT("%s: Cache miss with invalid package for %s from '%s'"),
+					*GetName(), *WriteToString<96>(Key), *Name);
+				OnComplete({ Name, Key, UserData, Request->GetBytesReceived(), {}, EStatus::Error });
+				return FHttpRequest::ECompletionBehavior::Done;
+			}
 
-		const FCbObjectView Object = FCbObject(ResponseBuffer);
-		const FIoHash RawHash = Object["RawHash"].AsHash();
-		const uint64 RawSize = Object["RawSize"].AsUInt64(MAX_uint64);
-		if (RawHash.IsZero() || RawSize == MAX_uint64)
-		{
-			UE_LOG(LogDerivedDataCache, Display, TEXT("%s: Cache miss with invalid value for %s from '%.*s'"),
-				*GetName(), *WriteToString<96>(Key), Name.Len(), Name.GetData());
-			return false;
-		}
-		OutValue = FValue(RawHash, RawSize);
-	}
-	else
-	{
-		FCompressedBuffer CompressedBuffer = FCompressedBuffer::FromCompressed(ResponseBuffer);
-		if (!CompressedBuffer)
-		{
-			FString ReceivedHashStr;
-			if (Request->GetHeader("X-Jupiter-InlinePayloadHash", ReceivedHashStr))
+			FOptionalCacheRecord Record = FCacheRecord::Load(FCbPackage(FCbObject(ResponseBuffer)));
+			if (Record.IsNull())
 			{
-				FIoHash ReceivedHash(ReceivedHashStr);
-				FIoHash ComputedHash = FIoHash::HashBuffer(ResponseBuffer.GetView());
-				if (ReceivedHash == ComputedHash)
-				{
-					CompressedBuffer = FCompressedBuffer::Compress(ResponseBuffer);
-				}
+				UE_LOG(LogDerivedDataCache, Log, TEXT("%s: Cache miss with record load failure for %s from '%s'"),
+					*GetName(), *WriteToString<96>(Key), *Name);
+				OnComplete({ Name, Key, UserData, Request->GetBytesReceived(), {}, EStatus::Error });
+				return FHttpRequest::ECompletionBehavior::Done;
 			}
+
+			OnComplete({ Name, Key, UserData, Request->GetBytesReceived(), MoveTemp(Record), EStatus::Ok });
+			return FHttpRequest::ECompletionBehavior::Done;
 		}
 
-		if (!CompressedBuffer)
+		if (!ShouldAbortForShutdown() && !Owner.IsCanceled() && ShouldRetryOnError(HttpResult, ResponseCode) && ((Request->GetAttempts()+1) < UE_HTTPDDC_MAX_ATTEMPTS))
 		{
-			UE_LOG(LogDerivedDataCache, Display, TEXT("%s: Cache miss with invalid package for %s from '%.*s'"),
-				*GetName(), *WriteToString<96>(Key), Name.Len(), Name.GetData());
-			return false;
+			return FHttpRequest::ECompletionBehavior::Retry;
 		}
-		OutValue = FValue(CompressedBuffer);
-	}
 
-	return true;
+		UE_LOG(LogDerivedDataCache, Verbose, TEXT("%s: Cache miss with missing package for %s from '%s'"),
+			*GetName(), *WriteToString<96>(Key), *Name);
+		OnComplete({ Name, Key, UserData, Request->GetBytesReceived(), {}, EStatus::Error });
+		return FHttpRequest::ECompletionBehavior::Done;
+	};
+	Request->SetHeader(TEXT("Accept"), TEXT("application/x-ue-cb"));
+	Request->EnqueueAsyncDownload(Owner, Pool, *RefsUri, MoveTemp(OnHttpRequestComplete), { 401, 404 });
 }
 
-FOptionalCacheRecord FHttpCacheStore::GetCacheRecord(
-	const FStringView Name,
-	const FCacheKey& Key,
+void FHttpCacheStore::PutCacheRecordAsync(
+	IRequestOwner& Owner,
+	const FSharedString& Name,
+	const FCacheRecord& Record,
 	const FCacheRecordPolicy& Policy,
-	EStatus& OutStatus)
+	uint64 UserData,
+	TUniqueFunction<void(FCachePutResponse&& Response, uint64 BytesSent)>&& OnComplete)
 {
-	FOptionalCacheRecord Record = GetCacheRecordOnly(Name, Key, Policy);
-	if (Record.IsNull())
+	const FCacheKey& Key = Record.GetKey();
+	auto MakeResponse = [Name = FSharedString(Name), Key = FCacheKey(Key), UserData](EStatus Status)
 	{
-		OutStatus = EStatus::Error;
-		return Record;
-	}
+		return FCachePutResponse{ Name, Key, UserData, Status };
+	};
 
-	OutStatus = EStatus::Ok;
+	if (!IsWritable())
+	{
+		UE_LOG(LogDerivedDataCache, VeryVerbose,
+			TEXT("%s: Skipped put of %s from '%s' because this cache store is read-only"),
+			*GetName(), *WriteToString<96>(Key), *Name);
+		return OnComplete(MakeResponse(EStatus::Error), 0);
+	}
 
-	FCacheRecordBuilder RecordBuilder(Key);
+	// Skip the request if storing to the cache is disabled.
+	const ECachePolicy RecordPolicy = Policy.GetRecordPolicy();
+	const ECachePolicy StoreFlag = SpeedClass == ESpeedClass::Local ? ECachePolicy::StoreLocal : ECachePolicy::StoreRemote;
+	if (!EnumHasAnyFlags(RecordPolicy, StoreFlag))
+	{
+		UE_LOG(LogDerivedDataCache, VeryVerbose, TEXT("%s: Skipped put of %s from '%s' due to cache policy"),
+			*GetName(), *WriteToString<96>(Key), *Name);
+		return OnComplete(MakeResponse(EStatus::Error), 0);
+	}
 
-	if (!EnumHasAnyFlags(Policy.GetRecordPolicy(), ECachePolicy::SkipMeta))
+	if (DebugOptions.ShouldSimulatePutMiss(Key))
 	{
-		RecordBuilder.SetMeta(FCbObject(Record.Get().GetMeta()));
+		UE_LOG(LogDerivedDataCache, Verbose, TEXT("%s: Simulated miss for put of %s from '%s'"),
+			*GetName(), *WriteToString<96>(Key), *Name);
+		return OnComplete(MakeResponse(EStatus::Error), 0);
 	}
 
-	// TODO: There is not currently a batched GET endpoint for Jupiter.  Once there is, all payload data should be fetched in one call.
-	//		 In the meantime, we try to keep the code structured in a way that is friendly to future batching of GETs.
+	// TODO: Jupiter currently always overwrites.  It doesn't have a "write if not present" feature (for records or attachments),
+	//		 but would require one to implement all policy correctly.
 
-	TArray<FValueWithId> RequiredGets;
-	TArray<FValueWithId> RequiredHeads;
+	FString Bucket(Key.Bucket.ToString());
+	Bucket.ToLowerInline();
 
-	for (FValueWithId Value : Record.Get().GetValues())
-	{
-		const ECachePolicy ValuePolicy = Policy.GetValuePolicy(Value.GetId());
-		if (IsValueDataReady(Value, ValuePolicy))
-		{
-			RecordBuilder.AddValue(MoveTemp(Value));
-		}
-		else
-		{
-			if (EnumHasAnyFlags(ValuePolicy, ECachePolicy::SkipData))
-			{
-				RequiredHeads.Emplace(Value);
-			}
-			else
-			{
-				RequiredGets.Emplace(Value);
-			}
-		}
-	}
+	FCbPackage Package = Record.Save();
 
-	auto HashGetter = [](const FValueWithId& Value)
+	FPutPackageOp::PutPackage(*this, Owner, Name, Key, MoveTemp(Package), Policy, UserData, [MakeResponse = MoveTemp(MakeResponse), OnComplete = MoveTemp(OnComplete)](FPutPackageOp::FCachePutPackageResponse&& Response)
 	{
-		return Value.GetRawHash();
-	};
+		OnComplete(MakeResponse(Response.Status), Response.BytesSent);
+	});
+}
 
-	auto DebugContextGetter = [Name, &Key](const FValueWithId& Value)
+void FHttpCacheStore::PutCacheValueAsync(
+	IRequestOwner& Owner,
+	const FSharedString& Name,
+	const FCacheKey& Key,
+	const FValue& Value,
+	const ECachePolicy Policy,
+	uint64 UserData,
+	TUniqueFunction<void(FCachePutValueResponse&& Response, uint64 BytesSent)>&& OnComplete)
+{
+	auto MakeResponse = [Name = FSharedString(Name), Key = FCacheKey(Key), UserData](EStatus Status)
 	{
-		return FValueDebugContext{ Name, Key, FString(*WriteToString<16>(Value.GetId())) };
+		return FCachePutValueResponse{ Name, Key, UserData, Status };
 	};
 
-	if (CachedDataProbablyExistsBatch<FValueWithId>(RequiredHeads, HashGetter, DebugContextGetter).CountSetBits() != RequiredHeads.Num())
-	{
-		OutStatus = EStatus::Error;
-		return FOptionalCacheRecord();
-	}
-
-	TArray<FCompressedBuffer> FetchedBuffers;
-	if (TryGetCachedDataBatch<FValueWithId>(RequiredGets, FetchedBuffers, HashGetter, DebugContextGetter).CountSetBits() != RequiredGets.Num())
+	if (!IsWritable())
 	{
-		OutStatus = EStatus::Error;
-		return FOptionalCacheRecord();
+		UE_LOG(LogDerivedDataCache, VeryVerbose,
+			TEXT("%s: Skipped put of %s from '%s' because this cache store is read-only"),
+			*GetName(), *WriteToString<96>(Key), *Name);
+		return OnComplete(MakeResponse(EStatus::Error), 0);
 	}
 
-	for (int32 Index = 0; Index < RequiredHeads.Num(); ++Index)
+	// Skip the request if storing to the cache is disabled.
+	const ECachePolicy StoreFlag = SpeedClass == ESpeedClass::Local ? ECachePolicy::StoreLocal : ECachePolicy::StoreRemote;
+	if (!EnumHasAnyFlags(Policy, StoreFlag))
 	{
-		RecordBuilder.AddValue(RequiredHeads[Index].RemoveData());
+		UE_LOG(LogDerivedDataCache, VeryVerbose, TEXT("%s: Skipped put of %s from '%s' due to cache policy"),
+			*GetName(), *WriteToString<96>(Key), *Name);
+		return OnComplete(MakeResponse(EStatus::Error), 0);
 	}
 
-	for (int32 Index = 0; Index < RequiredGets.Num(); ++Index)
+	if (DebugOptions.ShouldSimulatePutMiss(Key))
 	{
-		RecordBuilder.AddValue(FValueWithId(RequiredGets[Index].GetId(), FetchedBuffers[Index]));
+		UE_LOG(LogDerivedDataCache, Verbose, TEXT("%s: Simulated miss for put of %s from '%s'"),
+			*GetName(), *WriteToString<96>(Key), *Name);
+		return OnComplete(MakeResponse(EStatus::Error), 0);
 	}
 
-	return RecordBuilder.Build();
-}
-
-template<typename ValueType, typename ValueHashGetterType, typename ValueDebugContextGetterType>
-TBitArray<> FHttpCacheStore::TryGetCachedDataBatch(
-	TConstArrayView<ValueType> Values,
-	TArray<FCompressedBuffer>& OutBuffers,
-	ValueHashGetterType ValueHashGetter,
-	ValueDebugContextGetterType ValueDebugContextGetter,
-	FHttpRequest* ExistingHttpRequest)
-{
-	FHttpRequest* Request = ExistingHttpRequest;
-	TOptional<FScopedRequestPtr> RequestPoolRequest;
-	if (Request == nullptr)
-	{
-		RequestPoolRequest.Emplace(GetRequestPools[IsInGameThread()].Get());
-		Request = RequestPoolRequest->Get();
-	}
-	else
-	{
-		Request->Reset();
-	}
+	// TODO: Jupiter currently always overwrites.  It doesn't have a "write if not present" feature (for records or attachments),
+	//		 but would require one to implement all policy correctly.
 
-	bool bRequestNeedsReset = false;
-	TBitArray<> Results(true, Values.Num());
-	int32 ValueIndex = 0;
-	for (const ValueType& Value : Values)
-	{
-		const FIoHash& RawHash = ValueHashGetter(Value);
-		TStringBuilder<256> CompressedBlobsUri;
-		CompressedBlobsUri << "api/v1/compressed-blobs/" << StructuredNamespace << "/" << RawHash;
+	FString Bucket(Key.Bucket.ToString());
+	Bucket.ToLowerInline();
 
-		bool bHit = false;
-		FCompressedBuffer CompressedBuffer;
-		int64 ResponseCode = 0;
-		for (uint32 Attempts = 0; (Attempts < UE_HTTPDDC_MAX_ATTEMPTS) && !ShouldAbortForShutdown() && (Attempts == 0 || ShouldRetryOnError(ResponseCode)); ++Attempts)
-		{
-			if (bRequestNeedsReset)
-			{
-				Request->Reset();
-			}
+	FCbWriter Writer;
+	Writer.BeginObject();
+	Writer.AddBinaryAttachment("RawHash", Value.GetRawHash());
+	Writer.AddInteger("RawSize", Value.GetRawSize());
+	Writer.EndObject();
 
-			TArray<uint8> ByteArray;
-			const FHttpRequest::Result Result = Request->PerformBlockingDownload(*CompressedBlobsUri, &ByteArray, {404});
-			ResponseCode = Request->GetResponseCode();
-			bRequestNeedsReset = true;
+	FCbPackage Package(Writer.Save().AsObject());
+	Package.AddAttachment(FCbAttachment(Value.GetData()));
 
-			if (FHttpRequest::IsSuccessResponse(ResponseCode))
-			{
-				CompressedBuffer = FCompressedBuffer::FromCompressed(MakeSharedBufferFromArray(MoveTemp(ByteArray)));
-				bHit = true;
-				break;
-			}
-		}
+	FPutPackageOp::PutPackage(*this, Owner, Name, Key, MoveTemp(Package), Policy, UserData, [MakeResponse = MoveTemp(MakeResponse), OnComplete = MoveTemp(OnComplete)](FPutPackageOp::FCachePutPackageResponse&& Response)
+	{
+		OnComplete(MakeResponse(Response.Status), Response.BytesSent);
+	});
+}
 
-		if (!bHit)
-		{
-			FValueDebugContext DebugContext = ValueDebugContextGetter(Value);
-			UE_LOG(LogDerivedDataCache, Verbose,
-				TEXT("%s: Cache miss with missing value %s with hash %s for %s from '%.*s'"),
-				*GetName(), *DebugContext.Id, *WriteToString<48>(RawHash), *WriteToString<96>(DebugContext.Key),
-				DebugContext.Name.Len(), DebugContext.Name.GetData());
-			Results[ValueIndex] = false;
-		}
-		else if (CompressedBuffer.GetRawHash() != RawHash)
-		{
-			FValueDebugContext DebugContext = ValueDebugContextGetter(Value);
-			UE_LOG(LogDerivedDataCache, Display,
-				TEXT("%s: Cache miss with corrupted value %s with hash %s for %s from '%.*s'"),
-				*GetName(), *DebugContext.Id, *WriteToString<48>(RawHash),
-				*WriteToString<96>(DebugContext.Key), DebugContext.Name.Len(), DebugContext.Name.GetData());
-			Results[ValueIndex] = false;
-		}
-		else
-		{
-			OutBuffers.Add(CompressedBuffer);
-		}
-		++ValueIndex;
+void FHttpCacheStore::GetCacheValueAsync(
+	IRequestOwner& Owner,
+	FSharedString Name,
+	const FCacheKey& Key,
+	ECachePolicy Policy,
+	uint64 UserData,
+	FOnCacheGetValueComplete&& OnComplete)
+{
+	if (!IsUsable())
+	{
+		UE_LOG(LogDerivedDataCache, VeryVerbose,
+			TEXT("%s: Skipped get of %s from '%s' because this cache store is not available"),
+			*GetName(), *WriteToString<96>(Key), *Name);
+		OnComplete({Name, Key, {}, UserData, EStatus::Error});
+		return;
 	}
-	return Results;
-}
 
-template<typename ValueType, typename ValueHashGetterType, typename ValueDebugContextGetterType>
-TBitArray<> FHttpCacheStore::CachedDataProbablyExistsBatch(
-	TConstArrayView<ValueType> Values,
-	ValueHashGetterType ValueHashGetter,
-	ValueDebugContextGetterType ValueDebugContextGetter,
-	FHttpRequest* ExistingHttpRequest)
-{
-	if (Values.IsEmpty())
+	// Skip the request if querying the cache is disabled.
+	const ECachePolicy QueryFlag = SpeedClass == ESpeedClass::Local ? ECachePolicy::QueryLocal : ECachePolicy::QueryRemote;
+	if (!EnumHasAnyFlags(Policy, QueryFlag))
 	{
-		return TBitArray<>();
+		UE_LOG(LogDerivedDataCache, VeryVerbose, TEXT("%s: Skipped get of %s from '%s' due to cache policy"),
+			*GetName(), *WriteToString<96>(Key), *Name);
+		OnComplete({Name, Key, {}, UserData, EStatus::Error});
+		return;
 	}
 
-	FHttpRequest* Request = ExistingHttpRequest;
-	TOptional<FScopedRequestPtr> RequestPoolRequest;
-	if (Request == nullptr)
+	if (DebugOptions.ShouldSimulateGetMiss(Key))
 	{
-		RequestPoolRequest.Emplace(GetRequestPools[IsInGameThread()].Get());
-		Request = RequestPoolRequest->Get();
+		UE_LOG(LogDerivedDataCache, Verbose, TEXT("%s: Simulated miss for get of %s from '%s'"),
+			*GetName(), *WriteToString<96>(Key), *Name);
+		OnComplete({Name, Key, {}, UserData, EStatus::Error});
+		return;
 	}
-	else
+
+	const bool bSkipData = EnumHasAnyFlags(Policy, ECachePolicy::SkipData);
+
+	FString Bucket(Key.Bucket.ToString());
+	Bucket.ToLowerInline();
+
+	TStringBuilder<256> RefsUri;
+	RefsUri << "api/v1/refs/" << StructuredNamespace << "/" << Bucket << "/" << Key.Hash;
+
+	FRequestPool* Pool = nullptr;
+	FHttpRequest* Request = WaitForHttpRequestForOwner<OperationCategory::Get>(Owner, false /* bUnboundedOverflow */, Pool);
+	if (bSkipData)
 	{
-		Request->Reset();
+		Request->SetHeader(TEXT("Accept"), TEXT("application/x-ue-cb"));
 	}
-
-	TStringBuilder<256> CompressedBlobsUri;
-	CompressedBlobsUri << "api/v1/compressed-blobs/" << StructuredNamespace << "/exists?";
-	bool bFirstItem = true;
-	for (const ValueType& Value : Values)
+	else
 	{
-		if (!bFirstItem)
-		{
-			CompressedBlobsUri << "&";
-		}
-		CompressedBlobsUri << "id=" << ValueHashGetter(Value);
-		bFirstItem = false;
+		Request->SetHeader(TEXT("Accept"), TEXT("application/x-jupiter-inline"));
 	}
 
-	int64 ResponseCode = 0;
-	for (uint32 Attempts = 0; (Attempts < UE_HTTPDDC_MAX_ATTEMPTS) && !ShouldAbortForShutdown() && (Attempts == 0 || ShouldRetryOnError(ResponseCode)); ++Attempts)
+	auto OnHttpRequestComplete = [this, &Owner, Name, Key, UserData, bSkipData, OnComplete = MoveTemp(OnComplete)] (FHttpRequest::Result HttpResult, FHttpRequest* Request)
 	{
-		if (Attempts > 0)
-		{
-			Request->Reset();
-		}
-
-		TConstArrayView<uint8> DummyBuffer;
-		const FHttpRequest::Result Result = Request->PerformBlockingUpload<FHttpRequest::Post>(*CompressedBlobsUri, DummyBuffer);
-		ResponseCode = Request->GetResponseCode();
+		TRACE_CPUPROFILER_EVENT_SCOPE(HttpDDC_GetCacheValueAsync_OnHttpRequestComplete);
 
+		int64 ResponseCode = Request->GetResponseCode();
 		if (FHttpRequest::IsSuccessResponse(ResponseCode))
 		{
-			if (TSharedPtr<FJsonObject> ResponseObject = Request->GetResponseAsJsonObject())
+			FValue ResultValue;
+			FSharedBuffer ResponseBuffer = Request->MoveResponseBufferToShared();
+
+			if (bSkipData)
 			{
-				TArray<FString> NeedsArrayStrings;
-				if (ResponseObject->TryGetStringArrayField(TEXT("needs"), NeedsArrayStrings))
+				if (ValidateCompactBinary(ResponseBuffer, ECbValidateMode::Default) != ECbValidateError::None)
 				{
-					if (NeedsArrayStrings.IsEmpty())
-					{
-						return TBitArray<>(true, Values.Num());
-					}
+					UE_LOG(LogDerivedDataCache, Display, TEXT("%s: Cache miss with invalid package for %s from '%s'"),
+						*GetName(), *WriteToString<96>(Key), *Name);
+					OnComplete({Name, Key, {}, UserData, EStatus::Error});
+					return FHttpRequest::ECompletionBehavior::Done;
 				}
 
-				TBitArray<> Results(true, Values.Num());
-				for (const FString& NeedsString : NeedsArrayStrings)
+				const FCbObjectView Object = FCbObject(ResponseBuffer);
+				const FIoHash RawHash = Object["RawHash"].AsHash();
+				const uint64 RawSize = Object["RawSize"].AsUInt64(MAX_uint64);
+				if (RawHash.IsZero() || RawSize == MAX_uint64)
 				{
-					FIoHash NeedHash;
-					LexFromString(NeedHash, *NeedsString);
-					for (int32 ValueIndex = 0; ValueIndex < Values.Num(); ++ValueIndex)
+					UE_LOG(LogDerivedDataCache, Display, TEXT("%s: Cache miss with invalid value for %s from '%'"),
+						*GetName(), *WriteToString<96>(Key), *Name);
+					OnComplete({Name, Key, {}, UserData, EStatus::Error});
+					return FHttpRequest::ECompletionBehavior::Done;
+				}
+				ResultValue = FValue(RawHash, RawSize);
+			}
+			else
+			{
+				FCompressedBuffer CompressedBuffer = FCompressedBuffer::FromCompressed(ResponseBuffer);
+				if (!CompressedBuffer)
+				{
+					FString ReceivedHashStr;
+					if (Request->GetHeader("X-Jupiter-InlinePayloadHash", ReceivedHashStr))
 					{
-						if (NeedHash == ValueHashGetter(Values[ValueIndex]))
+						FIoHash ReceivedHash(ReceivedHashStr);
+						FIoHash ComputedHash = FIoHash::HashBuffer(ResponseBuffer.GetView());
+						if (ReceivedHash == ComputedHash)
 						{
-							Results[ValueIndex] = false;
-
-							FValueDebugContext DebugContext = ValueDebugContextGetter(Values[ValueIndex]);
-							UE_LOG(LogDerivedDataCache, Verbose,
-								TEXT("%s: Cache exists miss with missing value %s with hash %s for %s from '%.*s'"),
-								*GetName(), *DebugContext.Id, *NeedsString, *WriteToString<96>(DebugContext.Key),
-								DebugContext.Name.Len(), DebugContext.Name.GetData());
+							CompressedBuffer = FCompressedBuffer::Compress(ResponseBuffer);
 						}
 					}
 				}
 
-				return Results;
-			}
-			else
-			{
-				UE_LOG(LogDerivedDataCache, Log,
-					TEXT("%s: Cache exists returned invalid results."),
-					*GetName());
-				return TBitArray<>(false, Values.Num());
+				if (!CompressedBuffer)
+				{
+					UE_LOG(LogDerivedDataCache, Display, TEXT("%s: Cache miss with invalid package for %s from '%s'"),
+						*GetName(), *WriteToString<96>(Key), *Name);
+					OnComplete({Name, Key, {}, UserData, EStatus::Error});
+					return FHttpRequest::ECompletionBehavior::Done;
+				}
+				ResultValue = FValue(CompressedBuffer);
 			}
-
-			return TBitArray<>(false, Values.Num());
+			OnComplete({Name, Key, ResultValue, UserData, EStatus::Ok});
+			return FHttpRequest::ECompletionBehavior::Done;
 		}
-	}
 
-	return TBitArray<>(false, Values.Num());
+		 if (!ShouldAbortForShutdown() && !Owner.IsCanceled() && ShouldRetryOnError(HttpResult, ResponseCode) && ((Request->GetAttempts()+1) < UE_HTTPDDC_MAX_ATTEMPTS))
+		 {
+			return FHttpRequest::ECompletionBehavior::Retry;
+		 }
+
+		UE_LOG(LogDerivedDataCache, Verbose, TEXT("%s: Cache miss with failed HTTP request for %s from '%s'"),
+			*GetName(), *WriteToString<96>(Key), *Name);
+		OnComplete({Name, Key, {}, UserData, EStatus::Error});
+		return FHttpRequest::ECompletionBehavior::Done;
+	};
+
+	Request->EnqueueAsyncDownload(Owner, Pool, *RefsUri, MoveTemp(OnHttpRequestComplete), { 401, 404 });
+}
+
+void FHttpCacheStore::GetCacheRecordAsync(
+	IRequestOwner& Owner,
+	const FSharedString& Name,
+	const FCacheKey& Key,
+	const FCacheRecordPolicy& Policy,
+	uint64 UserData,
+	TUniqueFunction<void(FCacheGetResponse&& Response, uint64 BytesReceived)>&& OnComplete)
+{
+	FGetRecordOp::GetRecord(*this, Owner, Name, Key, Policy, UserData, MoveTemp(OnComplete));
 }
 
-template<typename ValueRefType, typename ValueRefKeyGetterType, typename ValueRefDebugContextGetterType>
-TArray<FValue> FHttpCacheStore::RefCachedDataProbablyExistsBatch(
-	TConstArrayView<ValueRefType> ValueRefs,
-	ValueRefKeyGetterType ValueRefKeyGetter,
-	ValueRefDebugContextGetterType ValueRefDebugContextGetter,
-	FHttpRequest* ExistingHttpRequest)
+void FHttpCacheStore::RefCachedDataProbablyExistsBatchAsync(
+	IRequestOwner& Owner,
+	TConstArrayView<FCacheGetValueRequest> ValueRefs,
+	FOnCacheGetValueComplete&& OnComplete)
 {
 	if (ValueRefs.IsEmpty())
 	{
-		return TArray<FValue>();
+		return;
 	}
 
-	FHttpRequest* Request = ExistingHttpRequest;
-	TOptional<FScopedRequestPtr> RequestPoolRequest;
-	if (Request == nullptr)
+	if (!IsUsable())
 	{
-		RequestPoolRequest.Emplace(GetRequestPools[IsInGameThread()].Get());
-		Request = RequestPoolRequest->Get();
-	}
-	else
+		for (const FCacheGetValueRequest& ValueRef : ValueRefs)
 	{
-		Request->Reset();
+			UE_LOG(LogDerivedDataCache, VeryVerbose,
+				TEXT("%s: Skipped exists check of %s from '%s' because this cache store is not available"),
+				*GetName(), *WriteToString<96>(ValueRef.Key), *ValueRef.Name);
+			OnComplete(ValueRef.MakeResponse(EStatus::Error));
+	}
+		return;
 	}
 
 	TStringBuilder<256> RefsUri;
@@ -3382,12 +4094,12 @@ TArray<FValue> FHttpCacheStore::RefCachedDataProbablyExistsBatch(
 	RequestWriter.BeginObject();
 	RequestWriter.BeginArray("ops"_ASV);
 	uint32 OpIndex = 0;
-	for (const ValueRefType& ValueRef : ValueRefs)
+	for (const FCacheGetValueRequest& ValueRef : ValueRefs)
 	{
 		RequestWriter.BeginObject();
 		RequestWriter.AddInteger("opId"_ASV, OpIndex);
 		RequestWriter.AddString("op"_ASV, "GET"_ASV);
-		FCacheKey Key = ValueRefKeyGetter(ValueRef);
+		FCacheKey Key = ValueRef.Key;
 		FString Bucket(Key.Bucket.ToString());
 		Bucket.ToLowerInline();
 		RequestWriter.AddString("bucket"_ASV, Bucket);
@@ -3400,31 +4112,30 @@ TArray<FValue> FHttpCacheStore::RefCachedDataProbablyExistsBatch(
 	RequestWriter.EndObject();
 	FCbFieldIterator RequestFields = RequestWriter.Save();
 
-	TConstArrayView<uint8> BodyBuffer = MakeConstArrayView(RequestFields.GetOuterBuffer());
 
-	int64 ResponseCode = 0;
-	for (uint32 Attempts = 0; (Attempts < UE_HTTPDDC_MAX_ATTEMPTS) && !ShouldAbortForShutdown() && (Attempts == 0 || ShouldRetryOnError(ResponseCode)); ++Attempts)
-	{
-		if (Attempts > 0)
-		{
-			Request->Reset();
-		}
+	FRequestPool* Pool = nullptr;
+	FHttpRequest* Request = WaitForHttpRequestForOwner<OperationCategory::Get>(Owner, false /* bUnboundedOverflow */, Pool);
 
-		Request->SetHeader(TEXT("Accept"), TEXT("application/x-ue-cb"));
-		const FHttpRequest::Result Result = Request->PerformBlockingUpload<FHttpRequest::PostCompactBinary>(*RefsUri, BodyBuffer);
-		ResponseCode = Request->GetResponseCode();
+	Request->SetHeader(TEXT("Accept"), TEXT("application/x-ue-cb"));
+
+	auto OnHttpRequestComplete = [this, &Owner, ValueRefs = TArray<FCacheGetValueRequest>(ValueRefs), OnComplete = MoveTemp(OnComplete)](FHttpRequest::Result HttpResult, FHttpRequest* Request)
+	{
+		TRACE_CPUPROFILER_EVENT_SCOPE(HttpDDC_RefCachedDataProbablyExistsBatchAsync_OnHttpRequestComplete);
 
+		int64 ResponseCode = Request->GetResponseCode();
 		if (FHttpRequest::IsSuccessResponse(ResponseCode))
 		{
 			FMemoryView ResponseView = MakeMemoryView(Request->GetResponseBuffer().GetData(), Request->GetResponseBuffer().Num());
 			if (ValidateCompactBinary(ResponseView, ECbValidateMode::Default) != ECbValidateError::None)
 			{
+				for (const FCacheGetValueRequest& ValueRef : ValueRefs)
+				{
 				UE_LOG(LogDerivedDataCache, Log,
 					TEXT("%s: Cache exists returned invalid results."),
 					*GetName());
-				TArray<FValue> RetVal;
-				RetVal.AddDefaulted(ValueRefs.Num());
-				return RetVal;
+					OnComplete(ValueRef.MakeResponse(EStatus::Error));
+				}
+				return FHttpRequest::ECompletionBehavior::Done;
 			}
 
 			const FCbObjectView ResponseObject = FCbObjectView(Request->GetResponseBuffer().GetData());
@@ -3433,15 +4144,16 @@ TArray<FValue> FHttpCacheStore::RefCachedDataProbablyExistsBatch(
 
 			if (ResultsArrayView.Num() != ValueRefs.Num())
 			{
+				for (const FCacheGetValueRequest& ValueRef : ValueRefs)
+				{
 				UE_LOG(LogDerivedDataCache, Log,
 					TEXT("%s: Cache exists returned unexpected quantity of results (expected %d, got %d)."),
 					*GetName(), ValueRefs.Num(), ResultsArrayView.Num());
-				TArray<FValue> RetVal;
-				RetVal.AddDefaulted(ValueRefs.Num());
-				return RetVal;
+					OnComplete(ValueRef.MakeResponse(EStatus::Error));
+				}
+				return FHttpRequest::ECompletionBehavior::Done;
 			}
-			TArray<FValue> RetVal;
-			RetVal.AddDefaulted(ValueRefs.Num());
+
 			for (FCbFieldView ResultFieldView : ResultsArrayView)
 			{
 				FCbObjectView ResultObjectView = ResultFieldView.AsObjectView();
@@ -3449,19 +4161,29 @@ TArray<FValue> FHttpCacheStore::RefCachedDataProbablyExistsBatch(
 				FCbObjectView ResponseObjectView = ResultObjectView["response"_ASV].AsObjectView();
 				int32 StatusCode = ResultObjectView["statusCode"_ASV].AsInt32();
 
-				if (OpId >= (uint32)RetVal.Num())
+				if (OpId >= (uint32)ValueRefs.Num())
 				{
-					FValueDebugContext DebugContext = ValueRefDebugContextGetter(ValueRefs[OpId]);
-					UE_LOG(LogDerivedDataCache, Display, TEXT("%s: Cache miss with invalid op index %u for %s from '%.*s'"),
-						*GetName(), OpId, *WriteToString<96>(DebugContext.Key), DebugContext.Name.Len(), DebugContext.Name.GetData());
+					UE_LOG(LogDerivedDataCache, Display, TEXT("%s: Encountered invalid opId %d while querying %d values"),
+						*GetName(), OpId, ValueRefs.Num());
 					continue;
 				}
 
+				const FCacheGetValueRequest& ValueRef = ValueRefs[OpId];
+
 				if (!FHttpRequest::IsSuccessResponse(StatusCode))
 				{
-					FValueDebugContext DebugContext = ValueRefDebugContextGetter(ValueRefs[OpId]);
-					UE_LOG(LogDerivedDataCache, Verbose, TEXT("%s: Cache miss with unsuccessful response code %d for %s from '%.*s'"),
-						*GetName(), StatusCode, *WriteToString<96>(DebugContext.Key), DebugContext.Name.Len(), DebugContext.Name.GetData());
+					UE_LOG(LogDerivedDataCache, Verbose, TEXT("%s: Cache miss with unsuccessful response code %d for %s from '%s'"),
+						*GetName(), StatusCode, *WriteToString<96>(ValueRef.Key), *ValueRef.Name);
+					OnComplete(ValueRef.MakeResponse(EStatus::Error));
+					continue;
+				}
+
+				const ECachePolicy QueryFlag = SpeedClass == ESpeedClass::Local ? ECachePolicy::QueryLocal : ECachePolicy::QueryRemote;
+				if (!EnumHasAnyFlags(ValueRef.Policy, QueryFlag))
+				{
+					UE_LOG(LogDerivedDataCache, VeryVerbose, TEXT("%s: Skipped exists check of %s from '%s' due to cache policy"),
+						*GetName(), *WriteToString<96>(ValueRef.Key), *ValueRef.Name);
+					OnComplete(ValueRef.MakeResponse(EStatus::Error));
 					continue;
 				}
 
@@ -3469,21 +4191,32 @@ TArray<FValue> FHttpCacheStore::RefCachedDataProbablyExistsBatch(
 				const uint64 RawSize = ResponseObjectView["RawSize"_ASV].AsUInt64(MAX_uint64);
 				if (RawHash.IsZero() || RawSize == MAX_uint64)
 				{
-					FValueDebugContext DebugContext = ValueRefDebugContextGetter(ValueRefs[OpId]);
-					UE_LOG(LogDerivedDataCache, Display, TEXT("%s: Cache miss with invalid value for %s from '%.*s'"),
-						*GetName(), *WriteToString<96>(DebugContext.Key), DebugContext.Name.Len(), DebugContext.Name.GetData());
+					UE_LOG(LogDerivedDataCache, Display, TEXT("%s: Cache miss with invalid value for %s from '%s'"),
+						*GetName(), *WriteToString<96>(ValueRef.Key), *ValueRef.Name);
+					OnComplete(ValueRef.MakeResponse(EStatus::Error));
 					continue;
 				}
 
-				RetVal[OpId] = FValue(RawHash, RawSize);
+				OnComplete({ValueRef.Name, ValueRef.Key, FValue(RawHash, RawSize), ValueRef.UserData, EStatus::Ok});
 			}
-			return RetVal;
+			return FHttpRequest::ECompletionBehavior::Done;
 		}
-	}
 
-	TArray<FValue> RetVal;
-	RetVal.AddDefaulted(ValueRefs.Num());
-	return RetVal;
+		if (!ShouldAbortForShutdown() && !Owner.IsCanceled() && ShouldRetryOnError(HttpResult, ResponseCode) && ((Request->GetAttempts()+1) < UE_HTTPDDC_MAX_ATTEMPTS))
+		{
+			return FHttpRequest::ECompletionBehavior::Retry;
+		}
+
+		for (const FCacheGetValueRequest& ValueRef : ValueRefs)
+		{
+			UE_LOG(LogDerivedDataCache, Verbose, TEXT("%s: Cache miss with failed HTTP request for %s from '%s'"),
+				*GetName(), *WriteToString<96>(ValueRef.Key), *ValueRef.Name);
+			OnComplete(ValueRef.MakeResponse(EStatus::Error));
+		}
+		return FHttpRequest::ECompletionBehavior::Done;
+	};
+
+	Request->EnqueueAsyncUpload<FHttpRequest::PostCompactBinary>(Owner, Pool, *RefsUri, RequestFields.GetOuterBuffer(), MoveTemp(OnHttpRequestComplete));
 }
 
 bool FHttpCacheStore::CachedDataProbablyExists(const TCHAR* CacheKey)
@@ -3658,7 +4391,7 @@ TBitArray<> FHttpCacheStore::CachedDataProbablyExistsBatch(TConstArrayView<FStri
 
 bool FHttpCacheStore::GetCachedData(const TCHAR* CacheKey, TArray<uint8>& OutData)
 {
-	TRACE_CPUPROFILER_EVENT_SCOPE(HttpDDC_Get);
+	TRACE_CPUPROFILER_EVENT_SCOPE(HttpDDC_GetCachedData);
 	TRACE_COUNTER_ADD(HttpDDC_Get, int64(1));
 	COOK_STAT(auto Timer = UsageStats.TimeGet());
 
@@ -3720,7 +4453,7 @@ bool FHttpCacheStore::GetCachedData(const TCHAR* CacheKey, TArray<uint8>& OutDat
 
 FDerivedDataBackendInterface::EPutStatus FHttpCacheStore::PutCachedData(const TCHAR* CacheKey, TArrayView<const uint8> InData, bool bPutEvenIfExists)
 {
-	TRACE_CPUPROFILER_EVENT_SCOPE(HttpDDC_Put);
+	TRACE_CPUPROFILER_EVENT_SCOPE(HttpDDC_PutCachedData);
 
 	if (!IsWritable())
 	{
@@ -3846,23 +4579,21 @@ void FHttpCacheStore::Put(
 	IRequestOwner& Owner,
 	FOnCachePutComplete&& OnComplete)
 {
+	TRACE_CPUPROFILER_EVENT_SCOPE(HttpDDC_Put);
+	FRequestBarrier Barrier(Owner);
+	TRefCountedUniqueFunction<FOnCachePutComplete>* CompletionFunction = new TRefCountedUniqueFunction<FOnCachePutComplete>(MoveTemp(OnComplete));
+	TRefCountPtr<TRefCountedUniqueFunction<FOnCachePutComplete>> BatchOnCompleteRef(CompletionFunction);
 	for (const FCachePutRequest& Request : Requests)
 	{
-		const FCacheRecord& Record = Request.Record;
-		COOK_STAT(auto Timer = UsageStats.TimePut());
-		uint64 BytesSent = 0;
-		if (PutCacheRecord(Request.Name, Record, Request.Policy, BytesSent))
-		{
-			UE_LOG(LogDerivedDataCache, Verbose, TEXT("%s: Cache put complete for %s from '%s'"),
-				*GetName(), *WriteToString<96>(Record.GetKey()), *Request.Name);
-			TRACE_COUNTER_ADD(HttpDDC_BytesSent, int64(BytesSent));
-			COOK_STAT(Timer.AddHit(BytesSent));
-			OnComplete(Request.MakeResponse(EStatus::Ok));
-		}
-		else
+		PutCacheRecordAsync(Owner, Request.Name, Request.Record, Request.Policy, Request.UserData, [COOK_STAT(Timer = UsageStats.TimePut(), ) OnCompletePtr = TRefCountPtr<TRefCountedUniqueFunction<FOnCachePutComplete>>(CompletionFunction)](FCachePutResponse&& Response, uint64 BytesSent) mutable
 		{
-			OnComplete(Request.MakeResponse(EStatus::Error));
-		}
+			TRACE_COUNTER_ADD(HttpDDC_BytesSent, BytesSent);
+			if (Response.Status == EStatus::Ok)
+			{
+				COOK_STAT(if (BytesSent) { Timer.AddHit(BytesSent); });
+			}
+			OnCompletePtr->GetFunction()(MoveTemp(Response));
+		});
 	}
 }
 
@@ -3871,22 +4602,21 @@ void FHttpCacheStore::Get(
 	IRequestOwner& Owner,
 	FOnCacheGetComplete&& OnComplete)
 {
+	TRACE_CPUPROFILER_EVENT_SCOPE(HttpDDC_Get);
+	FRequestBarrier Barrier(Owner);
+	TRefCountedUniqueFunction<FOnCacheGetComplete>* CompletionFunction = new TRefCountedUniqueFunction<FOnCacheGetComplete>(MoveTemp(OnComplete));
+	TRefCountPtr<TRefCountedUniqueFunction<FOnCacheGetComplete>> BatchOnCompleteRef(CompletionFunction);
 	for (const FCacheGetRequest& Request : Requests)
 	{
-		COOK_STAT(auto Timer = UsageStats.TimeGet());
-		EStatus Status = EStatus::Ok;
-		if (FOptionalCacheRecord Record = GetCacheRecord(Request.Name, Request.Key, Request.Policy, Status))
-		{
-			UE_LOG(LogDerivedDataCache, Verbose, TEXT("%s: Cache hit for %s from '%s'"),
-				*GetName(), *WriteToString<96>(Request.Key), *Request.Name);
-			TRACE_COUNTER_ADD(HttpDDC_BytesReceived, Private::GetCacheRecordCompressedSize(Record.Get()));
-			COOK_STAT(Timer.AddHit(Private::GetCacheRecordCompressedSize(Record.Get())));
-			OnComplete({Request.Name, MoveTemp(Record).Get(), Request.UserData, Status});
-		}
-		else
+		GetCacheRecordAsync(Owner, Request.Name, Request.Key, Request.Policy, Request.UserData, [COOK_STAT(Timer = UsageStats.TimePut(), ) OnCompletePtr = TRefCountPtr<TRefCountedUniqueFunction<FOnCacheGetComplete>>(CompletionFunction)](FCacheGetResponse&& Response, uint64 BytesReceived) mutable
 		{
-			OnComplete(Request.MakeResponse(Status));
-		}
+			TRACE_COUNTER_ADD(HttpDDC_BytesReceived, BytesReceived);
+			if (Response.Status == EStatus::Ok)
+			{
+				COOK_STAT(Timer.AddHit(BytesReceived););
+			}
+			OnCompletePtr->GetFunction()(MoveTemp(Response));
+		});
 	}
 }
 
@@ -3895,22 +4625,21 @@ void FHttpCacheStore::PutValue(
 	IRequestOwner& Owner,
 	FOnCachePutValueComplete&& OnComplete)
 {
+	TRACE_CPUPROFILER_EVENT_SCOPE(HttpDDC_PutValue);
+	FRequestBarrier Barrier(Owner);
+	TRefCountedUniqueFunction<FOnCachePutValueComplete>* CompletionFunction = new TRefCountedUniqueFunction<FOnCachePutValueComplete>(MoveTemp(OnComplete));
+	TRefCountPtr<TRefCountedUniqueFunction<FOnCachePutValueComplete>> BatchOnCompleteRef(CompletionFunction);
 	for (const FCachePutValueRequest& Request : Requests)
 	{
-		COOK_STAT(auto Timer = UsageStats.TimePut());
-		uint64 WriteSize = 0;
-		if (PutCacheValue(Request.Name, Request.Key, Request.Value, Request.Policy, WriteSize))
-		{
-			UE_LOG(LogDerivedDataCache, Verbose, TEXT("%s: Cache put complete for %s from '%s'"),
-				*GetName(), *WriteToString<96>(Request.Key), *Request.Name);
-			TRACE_COUNTER_ADD(HttpDDC_BytesSent, WriteSize);
-			COOK_STAT(if (WriteSize) { Timer.AddHit(WriteSize); });
-			OnComplete(Request.MakeResponse(EStatus::Ok));
-		}
-		else
+		PutCacheValueAsync(Owner, Request.Name, Request.Key, Request.Value, Request.Policy, Request.UserData, [COOK_STAT(Timer = UsageStats.TimePut(),) OnCompletePtr = TRefCountPtr<TRefCountedUniqueFunction<FOnCachePutValueComplete>>(CompletionFunction)](FCachePutValueResponse&& Response, uint64 BytesSent) mutable
 		{
-			OnComplete(Request.MakeResponse(EStatus::Error));
-		}
+			TRACE_COUNTER_ADD(HttpDDC_BytesSent, BytesSent);
+			if (Response.Status == EStatus::Ok)
+			{
+				COOK_STAT(if (BytesSent) { Timer.AddHit(BytesSent); });
+			}
+			OnCompletePtr->GetFunction()(MoveTemp(Response));
+		});
 	}
 }
 
@@ -3919,12 +4648,10 @@ void FHttpCacheStore::GetValue(
 	IRequestOwner& Owner,
 	FOnCacheGetValueComplete&& OnComplete)
 {
+	TRACE_CPUPROFILER_EVENT_SCOPE(HttpDDC_GetValue);
 	COOK_STAT(double StartTime = FPlatformTime::Seconds());
 	COOK_STAT(bool bIsInGameThread = IsInGameThread());
 
-	FScopedRequestPtr HttpRequest(GetRequestPools[IsInGameThread()].Get());
-	int64 HitBytes = 0;
-
 	bool bBatchExistsCandidate = true;
 	for (const FCacheGetValueRequest& Request : Requests)
 	{
@@ -3934,75 +4661,74 @@ void FHttpCacheStore::GetValue(
 			break;
 		}
 	}
-
 	if (bBatchExistsCandidate)
 	{
-		auto KeyGetter = [](const FCacheGetValueRequest& Request)
-		{
-			return Request.Key;
-		};
-
-		auto DebugContextGetter = [](const FCacheGetValueRequest& Request)
-		{
-			return FValueDebugContext{ Request.Name, Request.Key, FString(TEXT("Default")) };
-		};
-
-		TArray<FValue> Values = RefCachedDataProbablyExistsBatch<FCacheGetValueRequest>(Requests, KeyGetter, DebugContextGetter, HttpRequest.Get());
-
-		for (int32 RequestIndex = 0; RequestIndex < Requests.Num(); ++RequestIndex)
+		RefCachedDataProbablyExistsBatchAsync(Owner, Requests,
+		[this, COOK_STAT(StartTime, bIsInGameThread, ) OnComplete = MoveTemp(OnComplete)](FCacheGetValueResponse&& Response)
 		{
-			const FCacheGetValueRequest& Request = Requests[RequestIndex];
-			if (Values[RequestIndex] == FValue::Null)
+			if (Response.Status != EStatus::Ok)
 			{
 				COOK_STAT(UsageStats.GetStats.Accumulate(FCookStats::CallStats::EHitOrMiss::Miss, FCookStats::CallStats::EStatType::Counter, 1l, bIsInGameThread));
-				OnComplete(Request.MakeResponse(EStatus::Error));
+				OnComplete(MoveTemp(Response));
 			}
 			else
 			{
 				UE_LOG(LogDerivedDataCache, Verbose, TEXT("%s: Cache hit for %s from '%s'"),
-					*GetName(), *WriteToString<96>(Request.Key), *Request.Name);
+					*GetName(), *WriteToString<96>(Response.Key), *Response.Name);
 				COOK_STAT(UsageStats.GetStats.Accumulate(FCookStats::CallStats::EHitOrMiss::Hit, FCookStats::CallStats::EStatType::Counter, 1l, bIsInGameThread));
-				OnComplete({ Request.Name, Request.Key, Values[RequestIndex], Request.UserData, EStatus::Ok });
+				OnComplete(MoveTemp(Response));
 			}
-		}
+
+			COOK_STAT(const int64 CyclesUsed = int64((FPlatformTime::Seconds() - StartTime) / FPlatformTime::GetSecondsPerCycle()));
+			COOK_STAT(UsageStats.GetStats.Accumulate(FCookStats::CallStats::EHitOrMiss::Hit, FCookStats::CallStats::EStatType::Cycles, CyclesUsed, bIsInGameThread));
+		});
 	}
 	else
 	{
+		FRequestBarrier Barrier(Owner);
+		TRefCountedUniqueFunction<FOnCacheGetValueComplete>* CompletionFunction = new TRefCountedUniqueFunction<FOnCacheGetValueComplete>(MoveTemp(OnComplete));
+		TRefCountPtr<TRefCountedUniqueFunction<FOnCacheGetValueComplete>> BatchOnCompleteRef(CompletionFunction);
+		int64 HitBytes = 0;
 		for (const FCacheGetValueRequest& Request : Requests)
 		{
-			FValue Value;
-			if (!GetCacheValue(Request.Name, Request.Key, Request.Policy, Value, HttpRequest.Get()))
-			{
-				COOK_STAT(UsageStats.GetStats.Accumulate(FCookStats::CallStats::EHitOrMiss::Miss, FCookStats::CallStats::EStatType::Counter, 1l, bIsInGameThread));
-				OnComplete(Request.MakeResponse(EStatus::Error));
-			}
-			else
+			GetCacheValueAsync(Owner, Request.Name, Request.Key, Request.Policy, Request.UserData,
+			[this, COOK_STAT(StartTime, bIsInGameThread,) Policy = Request.Policy, OnCompletePtr = TRefCountPtr<TRefCountedUniqueFunction<FOnCacheGetValueComplete>>(CompletionFunction)] (FCacheGetValueResponse&& Response)
 			{
-				if (!IsValueDataReady(Value, Request.Policy) && !EnumHasAnyFlags(Request.Policy, ECachePolicy::SkipData))
+				const FOnCacheGetValueComplete& OnComplete = OnCompletePtr->GetFunction();
+				check(OnComplete);
+				if (Response.Status != EStatus::Ok)
 				{
-					// With inline fetching, expect we will always have a value we can use.  Even SkipData/Exists can rely on the blob existing if the ref is reported to exist.
-					UE_LOG(LogDerivedDataCache, Log, TEXT("%s: Cache miss due to inlining failure for %s from '%s'"),
-						*GetName(), *WriteToString<96>(Request.Key), *Request.Name);
 					COOK_STAT(UsageStats.GetStats.Accumulate(FCookStats::CallStats::EHitOrMiss::Miss, FCookStats::CallStats::EStatType::Counter, 1l, bIsInGameThread));
-					OnComplete(Request.MakeResponse(EStatus::Error));
+					OnComplete(MoveTemp(Response));
 				}
 				else
 				{
-					UE_LOG(LogDerivedDataCache, Verbose, TEXT("%s: Cache hit for %s from '%s'"),
-						*GetName(), *WriteToString<96>(Request.Key), *Request.Name);
-					uint64 ValueSize = Value.GetData().GetCompressedSize();
-					TRACE_COUNTER_ADD(HttpDDC_BytesReceived, ValueSize);
-					HitBytes += ValueSize;
-					COOK_STAT(UsageStats.GetStats.Accumulate(FCookStats::CallStats::EHitOrMiss::Hit, FCookStats::CallStats::EStatType::Counter, 1l, bIsInGameThread));
-					OnComplete({ Request.Name, Request.Key, Value, Request.UserData, EStatus::Ok });
+					if (!IsValueDataReady(Response.Value, Policy) && !EnumHasAnyFlags(Policy, ECachePolicy::SkipData))
+					{
+						// With inline fetching, expect we will always have a value we can use.  Even SkipData/Exists can rely on the blob existing if the ref is reported to exist.
+						UE_LOG(LogDerivedDataCache, Log, TEXT("%s: Cache miss due to inlining failure for %s from '%s'"),
+									*GetName(), *WriteToString<96>(Response.Key), *Response.Name);
+						COOK_STAT(UsageStats.GetStats.Accumulate(FCookStats::CallStats::EHitOrMiss::Miss, FCookStats::CallStats::EStatType::Counter, 1l, bIsInGameThread));
+						OnComplete(MoveTemp(Response));
+					}
+					else
+					{
+						UE_LOG(LogDerivedDataCache, Verbose, TEXT("%s: Cache hit for %s from '%s'"),
+									*GetName(), *WriteToString<96>(Response.Key), *Response.Name);
+						uint64 ValueSize = Response.Value.GetData().GetCompressedSize();
+						TRACE_COUNTER_ADD(HttpDDC_BytesReceived, ValueSize);
+						COOK_STAT(UsageStats.GetStats.Accumulate(FCookStats::CallStats::EHitOrMiss::Hit, FCookStats::CallStats::EStatType::Counter, 1l, bIsInGameThread));
+						OnComplete({ Response.Name, Response.Key, Response.Value, Response.UserData, EStatus::Ok });
+						
+						COOK_STAT(UsageStats.GetStats.Accumulate(FCookStats::CallStats::EHitOrMiss::Hit, FCookStats::CallStats::EStatType::Bytes, ValueSize, bIsInGameThread));
+					}
 				}
-			}
+				COOK_STAT(const int64 CyclesUsed = int64((FPlatformTime::Seconds() - StartTime) / FPlatformTime::GetSecondsPerCycle()));
+				COOK_STAT(UsageStats.GetStats.Accumulate(FCookStats::CallStats::EHitOrMiss::Hit, FCookStats::CallStats::EStatType::Cycles, CyclesUsed, bIsInGameThread));
+			});
 		}
 	}
 
-	COOK_STAT(const int64 CyclesUsed = int64((FPlatformTime::Seconds() - StartTime) / FPlatformTime::GetSecondsPerCycle()));
-	COOK_STAT(UsageStats.GetStats.Accumulate(FCookStats::CallStats::EHitOrMiss::Hit, FCookStats::CallStats::EStatType::Cycles, CyclesUsed, bIsInGameThread));
-	COOK_STAT(UsageStats.GetStats.Accumulate(FCookStats::CallStats::EHitOrMiss::Hit, FCookStats::CallStats::EStatType::Bytes, HitBytes, bIsInGameThread));
 }
 
 void FHttpCacheStore::GetChunks(
@@ -4010,6 +4736,7 @@ void FHttpCacheStore::GetChunks(
 	IRequestOwner& Owner,
 	FOnCacheGetChunkComplete&& OnComplete)
 {
+	TRACE_CPUPROFILER_EVENT_SCOPE(HttpDDC_GetChunks);
 	// TODO: This is inefficient because Jupiter doesn't allow us to get only part of a compressed blob, so we have to
 	//		 get the whole thing and then decompress only the portion we need.  Furthermore, because there is no propagation
 	//		 between cache stores during chunk requests, the fetched result won't end up in the local store.
@@ -4047,7 +4774,13 @@ void FHttpCacheStore::GetChunks(
 					FCacheRecordPolicyBuilder PolicyBuilder(ECachePolicy::None);
 					PolicyBuilder.AddValuePolicy(Request.Id, Request.Policy);
 					Record.Reset();
-					Record = GetCacheRecordOnly(Request.Name, Request.Key, PolicyBuilder.Build());
+
+					FRequestOwner BlockingOwner(EPriority::Blocking);
+					GetCacheRecordOnlyAsync(BlockingOwner, Request.Name, Request.Key, PolicyBuilder.Build(), 0, [&Record](FGetCacheRecordOnlyResponse&& Response)
+						{
+							Record = MoveTemp(Response.Record);
+						});
+					BlockingOwner.Wait();
 				}
 				if (Record)
 				{
@@ -4063,20 +4796,27 @@ void FHttpCacheStore::GetChunks(
 					}
 					else
 					{
-						auto HashGetter = [](const FValueWithId& Value)
+						auto IdGetter = [](const FValueWithId& Value)
 						{
-							return Value.GetRawHash();
+							return FString(WriteToString<16>(Value.GetId()));
 						};
 
-						auto DebugContextGetter = [&Request](const FValueWithId& Value)
-						{
-							return FValueDebugContext{ Request.Name, Request.Key, FString(*WriteToString<16>(Value.GetId())) };
-						};
+						FRequestOwner BlockingOwner(EPriority::Blocking);
+						bool bSucceeded = false;
+						FCompressedBuffer NewBuffer;
+						FGetRecordOp::GetDataBatch(*this, BlockingOwner, Request.Name, Request.Key, ::MakeArrayView({ ValueWithId }), IdGetter, [&bSucceeded, &NewBuffer](FGetRecordOp::FGetCachedDataBatchResponse&& Response)
+							{
+								if (Response.Status == EStatus::Ok)
+								{
+									bSucceeded = true;
+									NewBuffer = MoveTemp(Response.DataBuffer);
+								}
+							});
+						BlockingOwner.Wait();
 
-						TArray<FCompressedBuffer> ValueBuffers;
-						if (TryGetCachedDataBatch<FValueWithId>(::MakeArrayView({ ValueWithId }), ValueBuffers, HashGetter, DebugContextGetter).CountSetBits() == 1)
+						if (bSucceeded)
 						{
-							ValueBuffer = ValueBuffers[0];
+							ValueBuffer = MoveTemp(NewBuffer);
 							ValueReader.SetSource(ValueBuffer);
 						}
 						else
@@ -4090,27 +4830,46 @@ void FHttpCacheStore::GetChunks(
 			else
 			{
 				ValueKey = Request.Key;
-				bHasValue = GetCacheValue(Request.Name, Request.Key, Request.Policy, Value);
+
+				{
+					FRequestOwner BlockingOwner(EPriority::Blocking);
+					bool bSucceeded = false;
+					GetCacheValueAsync(BlockingOwner, Request.Name, Request.Key, Request.Policy, 0, [&bSucceeded, &Value](FCacheGetValueResponse&& Response)
+						{
+							Value = MoveTemp(Response.Value);
+							bSucceeded = Response.Status == EStatus::Ok;
+						});
+					BlockingOwner.Wait();
+					bHasValue = bSucceeded;
+				}
+
 				if (IsValueDataReady(Value, Request.Policy))
 				{
 					ValueReader.SetSource(Value.GetData());
 				}
 				else
 				{
-					auto HashGetter = [](const FValue& Value)
+					auto IdGetter = [](const FValue& Value)
 					{
-						return Value.GetRawHash();
+						return FString(TEXT("Default"));
 					};
 
-					auto DebugContextGetter = [&Request](const FValue& Value)
-					{
-						return FValueDebugContext{ Request.Name, Request.Key, FString(TEXT("Default")) };
-					};
+					FRequestOwner BlockingOwner(EPriority::Blocking);
+					bool bSucceeded = false;
+					FCompressedBuffer NewBuffer;
+					FGetRecordOp::GetDataBatch(*this, BlockingOwner, Request.Name, Request.Key, ::MakeArrayView({ Value }), IdGetter, [&bSucceeded, &NewBuffer](FGetRecordOp::FGetCachedDataBatchResponse&& Response)
+						{
+							if (Response.Status == EStatus::Ok)
+							{
+								bSucceeded = true;
+								NewBuffer = MoveTemp(Response.DataBuffer);
+							}
+						});
+					BlockingOwner.Wait();
 
-					TArray<FCompressedBuffer> ValueBuffers;
-					if (TryGetCachedDataBatch<FValue>({ Value }, ValueBuffers, HashGetter, DebugContextGetter).CountSetBits() == 1)
+					if (bSucceeded)
 					{
-						ValueBuffer = ValueBuffers[0];
+						ValueBuffer = MoveTemp(NewBuffer);
 						ValueReader.SetSource(ValueBuffer);
 					}
 					else
@@ -4143,6 +4902,185 @@ void FHttpCacheStore::GetChunks(
 	}
 }
 
+bool FHttpRequest::AllowAsync()
+{
+	if (!FGenericPlatformProcess::SupportsMultithreading() || !bHttpEnableAsync)
+	{
+		return false;
+	}
+
+	return true;
+}
+
+void FHttpRequest::EnqueueAsync(IRequestOwner& Owner, FRequestPool* Pool, const TCHAR* Uri, RequestVerb Verb, uint64 ContentLength, FOnHttpRequestComplete&& OnComplete, TConstArrayView<long> ExpectedErrorCodes)
+{
+	if (!AllowAsync())
+	{
+		while (OnComplete(PerformBlocking(Uri, Verb, ContentLength, ExpectedErrorCodes), this) == ECompletionBehavior::Retry)
+		{
+			PrepareToRetry();
+		}
+		if (Pool)
+		{
+			Pool->ReleaseRequestToPool(this);
+		}
+		return;
+	}
+
+	TRACE_CPUPROFILER_EVENT_SCOPE(HttpDDC_CurlEnqueueAsync);
+	AsyncData = new FAsyncRequestData;
+	AsyncData->Owner = &Owner;
+	AsyncData->Pool = Pool;
+	AsyncData->CurlHeaders = PrepareToIssueRequest(Uri, ContentLength);
+	AsyncData->Uri = Uri;
+	AsyncData->Verb = Verb;
+	AsyncData->ExpectedErrorCodes = ExpectedErrorCodes;
+	AsyncData->OnComplete = MoveTemp(OnComplete);
+	AsyncData->Owner->Begin(AsyncData);
+
+	if (SharedData->PendingRequestAdditions.EnqueueAndReturnWasEmpty(Curl))
+	{
+		SharedData->PendingRequestEvent.Trigger();
+	}
+
+	bool ExistingValue = false;
+	if (SharedData->bAsyncThreadStarting.compare_exchange_weak(ExistingValue, true, std::memory_order_relaxed))
+	{
+		SharedData->AsyncRunnable = MakeUnique<FHttpCacheStoreRunnable>(*SharedData);
+		SharedData->AsyncServiceThread = FRunnableThread::Create(SharedData->AsyncRunnable.Get(), TEXT("HttpCacheStoreThread"), 64 * 1024, TPri_Normal);
+	}
+}
+
+void FHttpRequest::CompleteAsync(CURLcode Result)
+{
+	CurlResult = Result;
+
+	// Get response code
+	curl_easy_getinfo(Curl, CURLINFO_RESPONSE_CODE, &ResponseCode);
+
+	LogResult(Result, *AsyncData->Uri, AsyncData->Verb, AsyncData->ExpectedErrorCodes);
+
+	ECompletionBehavior Behavior;
+	{
+		FRequestBarrier Barrier(*AsyncData->Owner, ERequestBarrierFlags::Priority);
+		FHttpRequest::Result HttpResult;
+		switch (CurlResult)
+		{
+		case CURLE_OK:
+			HttpResult = Success;
+			break;
+		case CURLE_OPERATION_TIMEDOUT:
+			HttpResult = FailedTimeout;
+			break;
+		default:
+			HttpResult = Failed;
+			break;
+		}
+		Behavior = AsyncData->OnComplete(HttpResult, this);
+	}
+
+	if (Behavior == ECompletionBehavior::Retry)
+	{
+		PrepareToRetry();
+		if (SharedData->PendingRequestAdditions.EnqueueAndReturnWasEmpty(Curl))
+		{
+			SharedData->PendingRequestEvent.Trigger();
+		}
+	}
+	else
+	{
+		// Clean up
+		curl_slist_free_all(AsyncData->CurlHeaders);
+
+		AsyncData->Owner->End(AsyncData, [this, Behavior]
+			{
+				AsyncData->Event.Trigger();
+				FRequestPool* Pool = AsyncData->Pool;
+				AsyncData = nullptr;
+				if (Pool)
+				{
+					Pool->ReleaseRequestToPool(this);
+				}
+			});
+	}
+}
+
+uint32 FHttpCacheStoreRunnable::Run()
+{
+	int ActiveTransfers = 0;
+
+	auto ProcessPendingRequests = [this, &ActiveTransfers]()
+	{
+		int CurrentActiveTransfers = -1;
+
+		do
+		{
+			SharedData.PendingRequestAdditions.Deplete([this, &ActiveTransfers](CURL* Curl)
+				{
+					curl_multi_add_handle(SharedData.CurlMulti, Curl);
+					++ActiveTransfers;
+				});
+
+			curl_multi_perform(SharedData.CurlMulti, &CurrentActiveTransfers);
+
+			if (CurrentActiveTransfers == 0 || ActiveTransfers != CurrentActiveTransfers)
+			{
+				for (;;)
+				{
+					int MsgsStillInQueue = 0;	// may use that to impose some upper limit we may spend in that loop
+					CURLMsg* Message = curl_multi_info_read(SharedData.CurlMulti, &MsgsStillInQueue);
+
+					if (!Message)
+					{
+						break;
+					}
+
+					// find out which requests have completed
+					if (Message->msg == CURLMSG_DONE)
+					{
+						CURL* CompletedHandle = Message->easy_handle;
+						curl_multi_remove_handle(SharedData.CurlMulti, CompletedHandle);
+
+						void* PrivateData = nullptr;
+						curl_easy_getinfo(CompletedHandle, CURLINFO_PRIVATE, &PrivateData);
+						FHttpRequest* CompletedRequest = (FHttpRequest*)PrivateData;
+
+						if (CompletedRequest)
+						{
+							// It is important that the CompleteAsync call doesn't happen on this thread as it is possible it will block waiting
+							// for a free HTTP request, and if that happens on this thread, we can deadlock as no HTTP requests will become 
+							// available while this thread is blocked.
+							UE::Tasks::Launch(TEXT("FHttpRequest::CompleteAsync"), [CompletedRequest, Result = Message->data.result]() mutable
+							{
+								CompletedRequest->CompleteAsync(Result);
+							});
+						}
+					}
+				}
+				ActiveTransfers = CurrentActiveTransfers;
+			}
+
+			if (CurrentActiveTransfers > 0)
+			{
+				curl_multi_wait(SharedData.CurlMulti, nullptr, 0, 1, nullptr);
+			}
+		}
+		while (CurrentActiveTransfers > 0);
+	};
+
+	do
+	{
+		ProcessPendingRequests();
+		SharedData.PendingRequestEvent.Wait(100);
+	}
+	while (!FHttpSharedData::bAsyncThreadShutdownRequested.load(std::memory_order_relaxed));
+
+	// Process last requests before shutdown.  May want these to be aborted instead.
+	ProcessPendingRequests();
+
+	return 0;
+}
+
 } // UE::DerivedData
 
 #endif // WITH_HTTP_DDC_BACKEND
