DDC: Added ConvertLegacyCacheKey to construct new cache keys from legacy cache keys

The function is shared with the legacy cache store to allow callers of the legacy API to migrate to the value API using this conversion function without invalidating existing cache keys.

diff --git a/Engine/Source/Developer/DerivedDataCache/Private/DerivedDataCacheKey.cpp b/Engine/Source/Developer/DerivedDataCache/Private/DerivedDataCacheKey.cpp
--- a/Engine/Source/Developer/DerivedDataCache/Private/DerivedDataCacheKey.cpp
+++ b/Engine/Source/Developer/DerivedDataCache/Private/DerivedDataCacheKey.cpp
@@ -5,10 +5,12 @@
 #include "Containers/Set.h"
 #include "Containers/StringConv.h"
 #include "DerivedDataCachePrivate.h"
+#include "IO/IoHash.h"
 #include "Math/UnrealMathUtility.h"
 #include "Misc/ScopeRWLock.h"
 #include "Misc/StringBuilder.h"
 #include "Serialization/CompactBinaryWriter.h"
+#include "String/Find.h"
 
 namespace UE::DerivedData::Private
 {
@@ -125,11 +127,23 @@ FCacheBucket::FCacheBucket(FWideStringView InName)
 FCbWriter& operator<<(FCbWriter& Writer, const FCacheKey& Key)
 {
 	Writer.BeginObject();
-	Writer.AddString("Bucket"_ASV, Key.Bucket.ToString());
-	Writer.AddHash("Hash"_ASV, Key.Hash);
+	Writer.AddString(ANSITEXTVIEW("Bucket"), Key.Bucket.ToString());
+	Writer.AddHash(ANSITEXTVIEW("Hash"), Key.Hash);
 	Writer.EndObject();
 	return Writer;
 }
 
+FCacheKey ConvertLegacyCacheKey(const FStringView Key)
+{
+	FTCHARToUTF8 Utf8Key(Key);
+	TUtf8StringBuilder<64> Utf8Bucket;
+	Utf8Bucket << ANSITEXTVIEW("Legacy");
+	if (const int32 BucketEnd = String::FindFirstChar(Utf8Key, '_'); BucketEnd != INDEX_NONE)
+	{
+		Utf8Bucket << FUtf8StringView(Utf8Key).Left(BucketEnd);
+	}
+	const FCacheBucket Bucket(Utf8Bucket);
+	return {Bucket, FIoHash::HashBuffer(MakeMemoryView(Utf8Key))};
+}
 
 } // UE::DerivedData
diff --git a/Engine/Source/Developer/DerivedDataCache/Private/DerivedDataLegacyCacheStore.cpp b/Engine/Source/Developer/DerivedDataCache/Private/DerivedDataLegacyCacheStore.cpp
--- a/Engine/Source/Developer/DerivedDataCache/Private/DerivedDataLegacyCacheStore.cpp
+++ b/Engine/Source/Developer/DerivedDataCache/Private/DerivedDataLegacyCacheStore.cpp
@@ -3,7 +3,6 @@
 #include "DerivedDataLegacyCacheStore.h"
 
 #include "Algo/AllOf.h"
-#include "Containers/StringConv.h"
 #include "DerivedDataBackendInterface.h"
 #include "DerivedDataCacheInterface.h"
 #include "Misc/Crc.h"
@@ -12,7 +11,6 @@
 #include "Misc/StringBuilder.h"
 #include "ProfilingDebugging/CpuProfilerTrace.h"
 #include "String/BytesToHex.h"
-#include "String/Find.h"
 
 namespace UE::DerivedData
 {
@@ -138,17 +136,9 @@ const FSharedString& Private::FLegacyCacheKeyShared::GetShortKey()
 }
 
 FLegacyCacheKey::FLegacyCacheKey(const FStringView FullKey, const int32 MaxKeyLength)
-	: Shared(new Private::FLegacyCacheKeyShared(FullKey, MaxKeyLength))
+	: Key(ConvertLegacyCacheKey(FullKey))
+	, Shared(new Private::FLegacyCacheKeyShared(FullKey, MaxKeyLength))
 {
-	FTCHARToUTF8 FullKeyUtf8(FullKey);
-	TUtf8StringBuilder<64> Bucket;
-	Bucket << "Legacy";
-	if (const int32 BucketEnd = String::FindFirstChar(FullKeyUtf8, '_'); BucketEnd != INDEX_NONE)
-	{
-		Bucket << FUtf8StringView(FullKeyUtf8).Left(BucketEnd);
-	}
-	Key.Bucket = FCacheBucket(Bucket);
-	Key.Hash = FIoHash::HashBuffer(MakeMemoryView(FullKeyUtf8));
 }
 
 bool FLegacyCacheKey::ReadValueTrailer(FCompositeBuffer& Value) const
diff --git a/Engine/Source/Developer/DerivedDataCache/Public/DerivedDataCacheKey.h b/Engine/Source/Developer/DerivedDataCache/Public/DerivedDataCacheKey.h
--- a/Engine/Source/Developer/DerivedDataCache/Public/DerivedDataCacheKey.h
+++ b/Engine/Source/Developer/DerivedDataCache/Public/DerivedDataCacheKey.h
@@ -127,6 +127,15 @@ UE_API FCbWriter& operator<<(FCbWriter& Writer, const FCacheKey& Key);
 
 ///////////////////////////////////////////////////////////////////////////////////////////////////////////////////////
 
+/**
+ * Constructs a cache key from a legacy cache key.
+ *
+ * Legacy keys are typically of the form DATATYPENAME_VersionGuid_TypeSpecificSuffix.
+ */
+UE_API FCacheKey ConvertLegacyCacheKey(FStringView Key);
+
+///////////////////////////////////////////////////////////////////////////////////////////////////////////////////////
+
 } // UE::DerivedData
 
 #undef UE_API
