DDC Verify cache miss is just a message not a warning

diff --git a/Engine/Source/Developer/DerivedDataCache/Private/DerivedDataCacheStoreVerify.cpp b/Engine/Source/Developer/DerivedDataCache/Private/DerivedDataCacheStoreVerify.cpp
--- a/Engine/Source/Developer/DerivedDataCache/Private/DerivedDataCacheStoreVerify.cpp
+++ b/Engine/Source/Developer/DerivedDataCache/Private/DerivedDataCacheStoreVerify.cpp
@@ -410,7 +410,7 @@ void FCacheStoreVerify::GetMetaComplete(IRequestOwner& Owner, FVerifyPutValueSta
 	}
 	else
 	{
-		UE_LOG(LogDerivedDataCache, Warning,
+		UE_LOG(LogDerivedDataCache, Display,
 			TEXT("Verify: Cache did not contain a value for %s from '%s'."),
 			*WriteToString<96>(Request.Key), *Request.Name);
 		FWriteScopeLock Lock(State->Lock);
@@ -456,7 +456,7 @@ void FCacheStoreVerify::GetDataComplete(IRequestOwner& Owner, FVerifyPutValueSta
 	}
 	else
 	{
-		UE_LOG(LogDerivedDataCache, Warning,
+		UE_LOG(LogDerivedDataCache, Display,
 			TEXT("Verify: Cache did not contain a value for %s from '%s'."),
 			*WriteToString<96>(Request.Key), *Request.Name);
 		FWriteScopeLock Lock(State->Lock);
@@ -614,7 +614,7 @@ void FCacheStoreVerify::GetDataComplete(IRequestOwner& Owner, FVerifyLegacyPutSt
 	}
 	else
 	{
-		UE_LOG(LogDerivedDataCache, Warning,
+		UE_LOG(LogDerivedDataCache, Display,
 			TEXT("Verify: Cache did not contain a value for %s from '%s'."),
 			*WriteToString<96>(Request.Key.GetKey()), *Request.Name);
 		FWriteScopeLock Lock(State->Lock);
