DDC: Fixed incorrect skipping of get propagation past nodes with the flag StopStore

The wrong response status was being checked, and caused subsequent nodes to be skipped even if the StopStore node did not contain the data, if any previous node contained the data.

#rnx

diff --git a/Engine/Source/Developer/DerivedDataCache/Private/DerivedDataCacheStoreHierarchy.cpp b/Engine/Source/Developer/DerivedDataCache/Private/DerivedDataCacheStoreHierarchy.cpp
--- a/Engine/Source/Developer/DerivedDataCache/Private/DerivedDataCacheStoreHierarchy.cpp
+++ b/Engine/Source/Developer/DerivedDataCache/Private/DerivedDataCacheStoreHierarchy.cpp
@@ -761,7 +761,7 @@ void FCacheStoreHierarchy::TGetBatch<Params>::CompleteRequest(FGetResponse&& Res
 		OnComplete(MoveTemp(Response));
 	}
 
-	if (State.Response.Status == EStatus::Ok)
+	if (Response.Status == EStatus::Ok)
 	{
 		if (EnumHasAnyFlags(Node.CacheFlags, ECacheStoreFlags::StopStore))
 		{
