DDC: Changed FileSystem PutValue to match Put

Except when replacing an existing value, load the existing value and check for determinism, and replace the existing value in the presence of an invalid value or non-deterministic value with missing content. This resolves move collisions being reported when a value existed without its content.

#rnx

diff --git a/Engine/Source/Developer/DerivedDataCache/Private/FileSystemCacheStore.cpp b/Engine/Source/Developer/DerivedDataCache/Private/FileSystemCacheStore.cpp
--- a/Engine/Source/Developer/DerivedDataCache/Private/FileSystemCacheStore.cpp
+++ b/Engine/Source/Developer/DerivedDataCache/Private/FileSystemCacheStore.cpp
@@ -1792,24 +1792,24 @@ bool FFileSystemCacheStore::PutCacheRecord(
 	FCbPackage ExistingPackage;
 	const ECachePolicy QueryFlag = SpeedClass == ESpeedClass::Local ? ECachePolicy::QueryLocal : ECachePolicy::QueryRemote;
 	bool bReplaceExisting = !EnumHasAnyFlags(RecordPolicy, QueryFlag);
-	bool bSaveRecord = bReplaceExisting;
-	if (const bool bLoadRecord = !bReplaceExisting || !Algo::AllOf(Record.GetValues(), &FValue::HasData))
+	bool bSavePackage = bReplaceExisting;
+	if (const bool bLoadPackage = !bReplaceExisting || !Algo::AllOf(Record.GetValues(), &FValue::HasData))
 	{
-		// Load the existing package to take its inline attachments into account.
-		bSaveRecord |= !LoadFileWithHash(Path, Name, [&ExistingPackage](FArchive& Ar) { ExistingPackage.TryLoad(Ar); });
-		if (!bSaveRecord)
+		// Load the existing package to take its attachments into account.
+		// Save the new package if there is no existing package or it fails to load.
+		bSavePackage |= !LoadFileWithHash(Path, Name, [&ExistingPackage](FArchive& Ar) { ExistingPackage.TryLoad(Ar); });
+		if (!bSavePackage)
 		{
-			// Load the existing record to detect determinism issues.
+			// Save the new package if the existing package is invalid.
 			const FOptionalCacheRecord ExistingRecord = FCacheRecord::Load(ExistingPackage);
-			bSaveRecord |= !ExistingRecord;
+			bSavePackage |= !ExistingRecord;
 			const auto MakeValueTuple = [](const FValueWithId& Value) -> TTuple<FValueId, FIoHash>
 			{
 				return MakeTuple(Value.GetId(), Value.GetRawHash());
 			};
 			if (ExistingRecord && !Algo::CompareBy(ExistingRecord.Get().GetValues(), Record.GetValues(), MakeValueTuple))
 			{
-				// Values differ between the existing record and the new record.
-				// Overwrite the existing record if the cache is missing content for any of its values.
+				// Content differs between the existing record and the new record.
 				UE_LOG(LogDerivedDataCache, Display, TEXT("%s: Cache put found non-deterministic record for %s from '%.*s'"),
 					*CachePath, *WriteToString<96>(Key), Name.Len(), Name.GetData());
 				const auto HasValueContent = [this, Name, &Key](const FValueWithId& Value) -> bool
@@ -1825,16 +1825,17 @@ bool FFileSystemCacheStore::PutCacheRecord(
 					}
 					return true;
 				};
-				bSaveRecord |= !Algo::AllOf(ExistingRecord.Get().GetValues(), HasValueContent);
-				bReplaceExisting |= bSaveRecord;
+				// Save the new package because the existing package differs and is missing content.
+				bSavePackage |= !Algo::AllOf(ExistingRecord.Get().GetValues(), HasValueContent);
 			}
+			bReplaceExisting |= bSavePackage;
 		}
 	}
 
 	// Serialize the record to a package and remove attachments that will be stored externally.
 	FCbPackage Package = Record.Save();
 	TArray<FCompressedBuffer, TInlineAllocator<8>> ExternalContent;
-	if (ExistingPackage && !bSaveRecord)
+	if (ExistingPackage && !bSavePackage)
 	{
 		// Mirror the existing internal/external attachment storage.
 		TArray<FCompressedBuffer, TInlineAllocator<8>> AllContent;
@@ -1900,8 +1901,12 @@ bool FFileSystemCacheStore::PutCacheRecord(
 	}
 
 	// Save the record package to storage.
-	const auto WriteRecord = [&](FArchive& Ar) { Package.Save(Ar); OutWriteSize += uint64(Ar.TotalSize()); };
-	if (bSaveRecord && !SaveFileWithHash(Path, Name, WriteRecord, bReplaceExisting))
+	const auto WritePackage = [&Package, &OutWriteSize](FArchive& Ar)
+	{
+		Package.Save(Ar);
+		OutWriteSize += uint64(Ar.TotalSize());
+	};
+	if (bSavePackage && !SaveFileWithHash(Path, Name, WritePackage, bReplaceExisting))
 	{
 		return false;
 	}
@@ -2069,20 +2074,61 @@ bool FFileSystemCacheStore::PutCacheValue(
 	}
 
 	// Check if there is an existing value package.
-	bool bValueExists = false;
 	FCbPackage ExistingPackage;
 	TStringBuilder<256> Path;
 	BuildCachePackagePath(Key, Path);
 	const ECachePolicy QueryFlag = SpeedClass == ESpeedClass::Local ? ECachePolicy::QueryLocal : ECachePolicy::QueryRemote;
-	const bool bReplaceExisting = !EnumHasAnyFlags(Policy, QueryFlag);
-	if (!bReplaceExisting)
-	{
-		bValueExists = LoadFileWithHash(Path, Name, [&ExistingPackage](FArchive& Ar) { ExistingPackage.TryLoad(Ar); })
-			&& (ExistingPackage.FindAttachment(Value.GetRawHash()) || GetCacheContentExists(Key, Value.GetRawHash()));
+	bool bReplaceExisting = !EnumHasAnyFlags(Policy, QueryFlag);
+	bool bSavePackage = bReplaceExisting;
+	if (const bool bLoadPackage = !bReplaceExisting || !Value.HasData())
+	{
+		// Load the existing package to take its attachments into account.
+		// Save the new package if there is no existing package or it fails to load.
+		bSavePackage |= !LoadFileWithHash(Path, Name, [&ExistingPackage](FArchive& Ar) { ExistingPackage.TryLoad(Ar); });
+		if (!bSavePackage)
+		{
+			const FCbObjectView Object = ExistingPackage.GetObject();
+			const FIoHash RawHash = Object["RawHash"].AsHash();
+			const uint64 RawSize = Object["RawSize"].AsUInt64(MAX_uint64);
+			if (RawHash.IsZero() || RawSize == MAX_uint64)
+			{
+				// Save the new package because the existing package is invalid.
+				UE_LOG(LogDerivedDataCache, Display, TEXT("%s: Cache put found invalid existing value for %s from '%.*s'"),
+					*CachePath, *WriteToString<96>(Key), Name.Len(), Name.GetData());
+				bSavePackage = true;
+			}
+			else if (!(RawHash == Value.GetRawHash() && RawSize == Value.GetRawSize()))
+			{
+				// Content differs between the existing value and the new value.
+				UE_LOG(LogDerivedDataCache, Display, TEXT("%s: Cache put found non-deterministic value "
+					"with new hash %s and existing hash %s for %s from '%.*s'"),
+					*CachePath, *WriteToString<48>(Value.GetRawHash()), *WriteToString<48>(RawHash),
+					*WriteToString<96>(Key), Name.Len(), Name.GetData());
+				if (!ExistingPackage.FindAttachment(RawHash) && !GetCacheContentExists(Key, RawHash))
+				{
+					// Save the new package because the existing package differs and is missing content.
+					UE_LOG(LogDerivedDataCache, Log,
+						TEXT("%s: Cache put of non-deterministic value will overwrite existing value due to "
+						"missing value with hash %s for %s from '%.*s'"),
+						*CachePath, *WriteToString<48>(RawHash), *WriteToString<96>(Key), Name.Len(), Name.GetData());
+					bSavePackage = true;
+				}
+			}
+			bReplaceExisting |= bSavePackage;
+		}
 	}
 
 	// Save the value to a package and save the data to external content depending on its size.
-	if (!bValueExists)
+	FCbPackage Package;
+	TArray<FCompressedBuffer, TInlineAllocator<1>> ExternalContent;
+	if (ExistingPackage && !bSavePackage)
+	{
+		if (Value.HasData() && !ExistingPackage.FindAttachment(Value.GetRawHash()))
+		{
+			ExternalContent.Add(Value.GetData());
+		}
+	}
+	else
 	{
 		FCbWriter Writer;
 		Writer.BeginObject();
@@ -2090,7 +2136,7 @@ bool FFileSystemCacheStore::PutCacheValue(
 		Writer.AddInteger("RawSize", Value.GetRawSize());
 		Writer.EndObject();
 
-		FCbPackage Package(Writer.Save().AsObject());
+		Package.SetObject(Writer.Save().AsObject());
 		if (!Value.HasData())
 		{
 			// Verify that the content exists in storage.
@@ -2108,21 +2154,30 @@ bool FFileSystemCacheStore::PutCacheValue(
 		}
 		else
 		{
-			// Save the external content to storage.
-			uint64 WriteSize = 0;
-			if (!PutCacheContent(Name, Value.GetData(), WriteSize))
-			{
-				return false;
-			}
-			OutWriteSize += WriteSize;
+			ExternalContent.Add(Value.GetData());
 		}
+	}
 
-		// Save the value package to storage.
-		const auto WritePackage = [&](FArchive& Ar) { Package.Save(Ar); OutWriteSize += uint64(Ar.TotalSize()); };
-		if (!SaveFileWithHash(Path, Name, WritePackage, bReplaceExisting))
+	// Save the external content to storage.
+	for (FCompressedBuffer& Content : ExternalContent)
+	{
+		uint64 WriteSize = 0;
+		if (!PutCacheContent(Name, Content, OutWriteSize))
 		{
 			return false;
 		}
+		OutWriteSize += WriteSize;
+	}
+
+	// Save the value package to storage.
+	const auto WritePackage = [&Package, &OutWriteSize](FArchive& Ar)
+	{
+		Package.Save(Ar);
+		OutWriteSize += uint64(Ar.TotalSize());
+	};
+	if (bSavePackage && !SaveFileWithHash(Path, Name, WritePackage, bReplaceExisting))
+	{
+		return false;
 	}
 
 	if (AccessLogWriter)
