DDC: Fixed duplicate puts caused by the in-flight memory cache

Without using StopStore, get hits on the in-flight cache would propagate as puts to subsequent nodes in the hierarchy. Since the in-flight cache is only used as temporary storage while a put is in flight, this caused duplicates of the in-flight puts, which manifests as move collision log spam from the file system cache store.

#rnx

diff --git a/Engine/Source/Developer/DerivedDataCache/Private/DerivedDataCacheStoreHierarchy.cpp b/Engine/Source/Developer/DerivedDataCache/Private/DerivedDataCacheStoreHierarchy.cpp
--- a/Engine/Source/Developer/DerivedDataCache/Private/DerivedDataCacheStoreHierarchy.cpp
+++ b/Engine/Source/Developer/DerivedDataCache/Private/DerivedDataCacheStoreHierarchy.cpp
@@ -221,7 +221,7 @@ FCacheStoreHierarchy::FCacheStoreHierarchy(IMemoryCacheStore* InMemoryCache)
 {
 	if (MemoryCache)
 	{
-		Add(MemoryCache, ECacheStoreFlags::Local | ECacheStoreFlags::Query);
+		Add(MemoryCache, ECacheStoreFlags::Local | ECacheStoreFlags::Query | ECacheStoreFlags::StopStore);
 	}
 }
 
