DDC: Back out 19757142 because it broke cache puts

#rnx

diff --git a/Engine/Source/Developer/DerivedDataCache/Private/DerivedDataCacheStoreHierarchy.cpp b/Engine/Source/Developer/DerivedDataCache/Private/DerivedDataCacheStoreHierarchy.cpp
--- a/Engine/Source/Developer/DerivedDataCache/Private/DerivedDataCacheStoreHierarchy.cpp
+++ b/Engine/Source/Developer/DerivedDataCache/Private/DerivedDataCacheStoreHierarchy.cpp
@@ -221,7 +221,7 @@ FCacheStoreHierarchy::FCacheStoreHierarchy(IMemoryCacheStore* InMemoryCache)
 {
 	if (MemoryCache)
 	{
-		Add(MemoryCache, ECacheStoreFlags::Local | ECacheStoreFlags::Query | ECacheStoreFlags::StopStore);
+		Add(MemoryCache, ECacheStoreFlags::Local | ECacheStoreFlags::Query);
 	}
 }
 
