DDC: Fixed incorrect cache store flags on dynamically mounted pak files

#rnx

diff --git a/Engine/Source/Developer/DerivedDataCache/Private/DerivedDataBackends.cpp b/Engine/Source/Developer/DerivedDataCache/Private/DerivedDataBackends.cpp
--- a/Engine/Source/Developer/DerivedDataCache/Private/DerivedDataBackends.cpp
+++ b/Engine/Source/Developer/DerivedDataCache/Private/DerivedDataBackends.cpp
@@ -1226,7 +1226,7 @@ public:
 		{
 			ReadPak = CreatePakFileCacheStore(PakFilename, /*bWriting*/ false, /*bCompressed*/ false);
 
-			Hierarchy->Add(ReadPak, ECacheStoreFlags::Local | ECacheStoreFlags::Store | ECacheStoreFlags::StopStore);
+			Hierarchy->Add(ReadPak, ECacheStoreFlags::Local | ECacheStoreFlags::Query | ECacheStoreFlags::StopStore);
 			CreatedNodes.AddUnique(ReadPak);
 			ReadPakCache.Add(ReadPak);
 		}
