DDC: Fixed duplicate puts caused by the in-flight memory cache

Without using StopGetStore, get hits on the in-flight cache would propagate as puts to subsequent nodes in the hierarchy. Since the in-flight cache is only used as temporary storage while a put is in flight, this caused duplicates of the in-flight puts, which manifests as move collision log spam from the file system cache store.

#rnx

diff --git a/Engine/Source/Developer/DerivedDataCache/Private/DerivedDataCacheStore.h b/Engine/Source/Developer/DerivedDataCache/Private/DerivedDataCacheStore.h
--- a/Engine/Source/Developer/DerivedDataCache/Private/DerivedDataCacheStore.h
+++ b/Engine/Source/Developer/DerivedDataCache/Private/DerivedDataCacheStore.h
@@ -106,8 +106,12 @@ enum class ECacheStoreFlags : uint32
 	/** Accepts requests with a policy of StoreLocal or StoreRemote. Needs matching Local/Remote flag. */
 	Store           = 1 << 3,
 
-	/** Requests to store records or values contained by this cache store will not propagate any further. */
-	StopStore       = 1 << 4,
+	/** A put of a record or value contained by this cache store will not store to later cache stores. */
+	StopPutStore    = 1 << 4,
+	/** A get of a record or value contained by this cache store will not store to later cache stores. */
+	StopGetStore    = 1 << 5,
+	/** A record or value contained by this cache store will not store to later cache stores. */
+	StopStore       = StopPutStore | StopGetStore,
 };
 
 ENUM_CLASS_FLAGS(ECacheStoreFlags);
diff --git a/Engine/Source/Developer/DerivedDataCache/Private/DerivedDataCacheStoreHierarchy.cpp b/Engine/Source/Developer/DerivedDataCache/Private/DerivedDataCacheStoreHierarchy.cpp
--- a/Engine/Source/Developer/DerivedDataCache/Private/DerivedDataCacheStoreHierarchy.cpp
+++ b/Engine/Source/Developer/DerivedDataCache/Private/DerivedDataCacheStoreHierarchy.cpp
@@ -221,7 +221,7 @@ FCacheStoreHierarchy::FCacheStoreHierarchy(IMemoryCacheStore* InMemoryCache)
 {
 	if (MemoryCache)
 	{
-		Add(MemoryCache, ECacheStoreFlags::Local | ECacheStoreFlags::Query);
+		Add(MemoryCache, ECacheStoreFlags::Local | ECacheStoreFlags::Query | ECacheStoreFlags::StopGetStore);
 	}
 }
 
@@ -446,7 +446,7 @@ bool FCacheStoreHierarchy::TPutBatch<Params>::DispatchGetRequests()
 	NodeGetIndex = NodePutIndex;
 
 	const FCacheStoreNode& Node = Hierarchy.Nodes[NodeGetIndex];
-	if (!EnumHasAnyFlags(Node.CacheFlags, ECacheStoreFlags::StopStore))
+	if (!EnumHasAnyFlags(Node.CacheFlags, ECacheStoreFlags::StopPutStore))
 	{
 		return false;
 	}
@@ -668,7 +668,7 @@ void FCacheStoreHierarchy::TGetBatch<Params>::DispatchRequests()
 				{
 					AsyncNodeRequests.Add(MakePutRequest(Response, Request));
 				}
-				else if (EnumHasAnyFlags(Node.CacheFlags, ECacheStoreFlags::StopStore) && CanQuery(GetCombinedPolicy(Request.Policy), Node.CacheFlags))
+				else if (EnumHasAnyFlags(Node.CacheFlags, ECacheStoreFlags::StopGetStore) && CanQuery(GetCombinedPolicy(Request.Policy), Node.CacheFlags))
 				{
 					NodeRequests.Add({Request.Name, Request.Key, AddPolicy(Request.Policy, ECachePolicy::SkipData), StateIndex});
 				}
@@ -763,9 +763,9 @@ void FCacheStoreHierarchy::TGetBatch<Params>::CompleteRequest(FGetResponse&& Res
 
 	if (Response.Status == EStatus::Ok)
 	{
-		if (EnumHasAnyFlags(Node.CacheFlags, ECacheStoreFlags::StopStore))
+		if (EnumHasAnyFlags(Node.CacheFlags, ECacheStoreFlags::StopGetStore))
 		{
-			// Never store to later later nodes if this node has StopStore.
+			// Never store to later nodes.
 			State.Request.Policy = RemovePolicy(State.Request.Policy, ECachePolicy::Default);
 		}
 		else
