DDC: Added more documentation to the cache policy

#rnx

diff --git a/Engine/Source/Developer/DerivedDataCache/Public/DerivedDataCachePolicy.h b/Engine/Source/Developer/DerivedDataCache/Public/DerivedDataCachePolicy.h
--- a/Engine/Source/Developer/DerivedDataCache/Public/DerivedDataCachePolicy.h
+++ b/Engine/Source/Developer/DerivedDataCache/Public/DerivedDataCachePolicy.h
@@ -33,9 +33,12 @@ namespace UE::DerivedData
  * Get(Query | StoreLocal): Fetch from any cache. Store response to any local caches if missing.
  * Get(Query | SkipData): Check for existence in any cache. Do not store to any caches if missing.
  * Get(Default | SkipData): Check for existence in any cache. Store the response to any caches if missing.
+ * Get(Default | PartialRecord): Fetch from any cache, and return a partial record if values are missing data.
  *
- * Put(Default): Store to every cache.
+ * Put(Default): Store to every cache, and do not overwrite existing valid records or values.
+ * Put(Store): Store to every cache, and overwrite existing records or values.
  * Put(Local): Store to every local cache, skipping remote caches.
+ * Put(Default | PartialRecord): Store to every cache, even if the record has missing data for its values.
  */
 enum class ECachePolicy : uint32
 {
@@ -84,10 +87,13 @@ enum class ECachePolicy : uint32
 	PartialRecord   = 1 << 6,
 
 	/**
-	 * Keep records in the cache for at least the duration of the session.
+	 * Keep records and values in the cache for at least the duration of the session.
 	 *
-	 * This is a hint that the record may be accessed again in this session. This is mainly meant
-	 * to be used when subsequent accesses will not tolerate a cache miss.
+	 * This flag hints that the records and values that it is applied to may be accessed again in
+	 * this session. The cache will make an effort to prevent eviction of the records and values,
+	 * though an absolute guarantee is impossible.
+	 *
+	 * This flag is meant to be used when subsequent accesses will not tolerate a cache miss.
 	 */
 	KeepAlive       = 1 << 7,
 };
@@ -199,7 +205,7 @@ public:
 	{
 	}
 
-	/** Adds a cache policy override for a value. */
+	/** Adds a cache policy override for a value. Must contain only flags in FCacheRecordValue::PolicyMask. */
 	UE_API void AddValuePolicy(const FCacheValuePolicy& Value);
 	inline void AddValuePolicy(const FValueId& Id, ECachePolicy Policy) { AddValuePolicy({Id, Policy}); }
 
