Replaced uses of _ASV with ANSITEXTVIEW

#rnx

diff --git a/Engine/Source/Developer/DerivedDataCache/Private/DerivedDataBuildAction.cpp b/Engine/Source/Developer/DerivedDataCache/Private/DerivedDataBuildAction.cpp
--- a/Engine/Source/Developer/DerivedDataCache/Private/DerivedDataBuildAction.cpp
+++ b/Engine/Source/Developer/DerivedDataCache/Private/DerivedDataBuildAction.cpp
@@ -150,13 +150,13 @@ FBuildActionInternal::FBuildActionInternal(FBuildActionBuilderInternal&& ActionB
 
 	TCbWriter<2048> Writer;
 	Writer.BeginObject();
-	Writer.AddString("Function"_ASV, Function);
-	Writer.AddUuid("FunctionVersion"_ASV, FunctionVersion);
-	Writer.AddUuid("BuildSystemVersion"_ASV, BuildSystemVersion);
+	Writer.AddString(ANSITEXTVIEW("Function"), Function);
+	Writer.AddUuid(ANSITEXTVIEW("FunctionVersion"), FunctionVersion);
+	Writer.AddUuid(ANSITEXTVIEW("BuildSystemVersion"), BuildSystemVersion);
 
 	if (bHasConstants)
 	{
-		Writer.BeginObject("Constants"_ASV);
+		Writer.BeginObject(ANSITEXTVIEW("Constants"));
 		for (const TPair<FUtf8SharedString, FBuildActionBuilderInternal::InputType>& Pair : ActionBuilder.Inputs)
 		{
 			if (Pair.Value.IsType<FCbObject>())
@@ -169,15 +169,15 @@ FBuildActionInternal::FBuildActionInternal(FBuildActionBuilderInternal&& ActionB
 
 	if (bHasInputs)
 	{
-		Writer.BeginObject("Inputs"_ASV);
+		Writer.BeginObject(ANSITEXTVIEW("Inputs"));
 		for (const TPair<FUtf8SharedString, FBuildActionBuilderInternal::InputType>& Pair : ActionBuilder.Inputs)
 		{
 			if (Pair.Value.IsType<TTuple<FIoHash, uint64>>())
 			{
 				const TTuple<FIoHash, uint64>& Input = Pair.Value.Get<TTuple<FIoHash, uint64>>();
 				Writer.BeginObject(Pair.Key);
-				Writer.AddBinaryAttachment("RawHash"_ASV, Input.Get<FIoHash>());
-				Writer.AddInteger("RawSize"_ASV, Input.Get<uint64>());
+				Writer.AddBinaryAttachment(ANSITEXTVIEW("RawHash"), Input.Get<FIoHash>());
+				Writer.AddInteger(ANSITEXTVIEW("RawSize"), Input.Get<uint64>());
 				Writer.EndObject();
 			}
 		}
@@ -191,9 +191,9 @@ FBuildActionInternal::FBuildActionInternal(FBuildActionBuilderInternal&& ActionB
 
 FBuildActionInternal::FBuildActionInternal(const FSharedString& InName, FCbObject&& InAction, bool& bOutIsValid)
 	: Name(InName)
-	, Function(InAction.FindView("Function"_ASV).AsString())
-	, FunctionVersion(InAction.FindView("FunctionVersion"_ASV).AsUuid())
-	, BuildSystemVersion(InAction.FindView("BuildSystemVersion"_ASV).AsUuid())
+	, Function(InAction.FindView(ANSITEXTVIEW("Function")).AsString())
+	, FunctionVersion(InAction.FindView(ANSITEXTVIEW("FunctionVersion")).AsUuid())
+	, BuildSystemVersion(InAction.FindView(ANSITEXTVIEW("BuildSystemVersion")).AsUuid())
 	, Action(MoveTemp(InAction))
 	, Key{Action.GetHash()}
 {
@@ -203,29 +203,29 @@ FBuildActionInternal::FBuildActionInternal(const FSharedString& InName, FCbObjec
 		&& IsValidBuildFunctionName(Function)
 		&& FunctionVersion.IsValid()
 		&& BuildSystemVersion.IsValid()
-		&& Algo::AllOf(Action.AsView()["Constants"_ASV],
+		&& Algo::AllOf(Action.AsView()[ANSITEXTVIEW("Constants")],
 			[](FCbFieldView Field) { return Field.GetName().Len() > 0 && Field.IsObject(); })
-		&& Algo::AllOf(Action.AsView()["Inputs"_ASV], [](FCbFieldView Field)
+		&& Algo::AllOf(Action.AsView()[ANSITEXTVIEW("Inputs")], [](FCbFieldView Field)
 			{
 				return Field.GetName().Len() > 0 && Field.IsObject()
-					&& Field["RawHash"_ASV].IsBinaryAttachment()
-					&& Field["RawSize"_ASV].IsInteger();
+					&& Field[ANSITEXTVIEW("RawHash")].IsBinaryAttachment()
+					&& Field[ANSITEXTVIEW("RawSize")].IsInteger();
 			});
 }
 
 bool FBuildActionInternal::HasConstants() const
 {
-	return Action["Constants"_ASV].HasValue();
+	return Action[ANSITEXTVIEW("Constants")].HasValue();
 }
 
 bool FBuildActionInternal::HasInputs() const
 {
-	return Action["Inputs"_ASV].HasValue();
+	return Action[ANSITEXTVIEW("Inputs")].HasValue();
 }
 
 void FBuildActionInternal::IterateConstants(TFunctionRef<void (FUtf8StringView Key, FCbObject&& Value)> Visitor) const
 {
-	for (FCbField Field : Action["Constants"_ASV])
+	for (FCbField Field : Action[ANSITEXTVIEW("Constants")])
 	{
 		Visitor(Field.GetName(), Field.AsObject());
 	}
@@ -233,9 +233,9 @@ void FBuildActionInternal::IterateConstants(TFunctionRef<void (FUtf8StringView K
 
 void FBuildActionInternal::IterateInputs(TFunctionRef<void (FUtf8StringView Key, const FIoHash& RawHash, uint64 RawSize)> Visitor) const
 {
-	for (FCbFieldView Field : Action.AsView()["Inputs"_ASV])
+	for (FCbFieldView Field : Action.AsView()[ANSITEXTVIEW("Inputs")])
 	{
-		Visitor(Field.GetName(), Field["RawHash"_ASV].AsHash(), Field["RawSize"_ASV].AsUInt64());
+		Visitor(Field.GetName(), Field[ANSITEXTVIEW("RawHash")].AsHash(), Field[ANSITEXTVIEW("RawSize")].AsUInt64());
 	}
 }
 
diff --git a/Engine/Source/Developer/DerivedDataCache/Private/DerivedDataBuildDefinition.cpp b/Engine/Source/Developer/DerivedDataCache/Private/DerivedDataBuildDefinition.cpp
--- a/Engine/Source/Developer/DerivedDataCache/Private/DerivedDataBuildDefinition.cpp
+++ b/Engine/Source/Developer/DerivedDataCache/Private/DerivedDataBuildDefinition.cpp
@@ -169,11 +169,11 @@ FBuildDefinitionInternal::FBuildDefinitionInternal(FBuildDefinitionBuilderIntern
 
 	TCbWriter<2048> Writer;
 	Writer.BeginObject();
-	Writer.AddString("Function"_ASV, Function);
+	Writer.AddString(ANSITEXTVIEW("Function"), Function);
 
 	if (bHasConstants)
 	{
-		Writer.BeginObject("Constants"_ASV);
+		Writer.BeginObject(ANSITEXTVIEW("Constants"));
 		for (const TPair<FUtf8SharedString, FBuildDefinitionBuilderInternal::InputType>& Pair : DefinitionBuilder.Inputs)
 		{
 			if (Pair.Value.IsType<FCbObject>())
@@ -186,20 +186,20 @@ FBuildDefinitionInternal::FBuildDefinitionInternal(FBuildDefinitionBuilderIntern
 
 	if (bHasInputs)
 	{
-		Writer.BeginObject("Inputs"_ASV);
+		Writer.BeginObject(ANSITEXTVIEW("Inputs"));
 	}
 
 	if (bHasBuilds)
 	{
-		Writer.BeginObject("Builds"_ASV);
+		Writer.BeginObject(ANSITEXTVIEW("Builds"));
 		for (const TPair<FUtf8SharedString, FBuildDefinitionBuilderInternal::InputType>& Pair : DefinitionBuilder.Inputs)
 		{
 			if (Pair.Value.IsType<FBuildValueKey>())
 			{
 				const FBuildValueKey& ValueKey = Pair.Value.Get<FBuildValueKey>();
 				Writer.BeginObject(Pair.Key);
-				Writer.AddHash("Build"_ASV, ValueKey.BuildKey.Hash);
-				Writer.AddObjectId("Id"_ASV, ValueKey.Id);
+				Writer.AddHash(ANSITEXTVIEW("Build"), ValueKey.BuildKey.Hash);
+				Writer.AddObjectId(ANSITEXTVIEW("Id"), ValueKey.Id);
 				Writer.EndObject();
 			}
 		}
@@ -208,7 +208,7 @@ FBuildDefinitionInternal::FBuildDefinitionInternal(FBuildDefinitionBuilderIntern
 
 	if (bHasBulkData)
 	{
-		Writer.BeginObject("BulkData"_ASV);
+		Writer.BeginObject(ANSITEXTVIEW("BulkData"));
 		for (const TPair<FUtf8SharedString, FBuildDefinitionBuilderInternal::InputType>& Pair : DefinitionBuilder.Inputs)
 		{
 			if (Pair.Value.IsType<FGuid>())
@@ -221,7 +221,7 @@ FBuildDefinitionInternal::FBuildDefinitionInternal(FBuildDefinitionBuilderIntern
 
 	if (bHasFiles)
 	{
-		Writer.BeginObject("Files"_ASV);
+		Writer.BeginObject(ANSITEXTVIEW("Files"));
 		for (const TPair<FUtf8SharedString, FBuildDefinitionBuilderInternal::InputType>& Pair : DefinitionBuilder.Inputs)
 		{
 			if (Pair.Value.IsType<FUtf8SharedString>())
@@ -234,7 +234,7 @@ FBuildDefinitionInternal::FBuildDefinitionInternal(FBuildDefinitionBuilderIntern
 
 	if (bHasHashes)
 	{
-		Writer.BeginObject("Hashes"_ASV);
+		Writer.BeginObject(ANSITEXTVIEW("Hashes"));
 		for (const TPair<FUtf8SharedString, FBuildDefinitionBuilderInternal::InputType>& Pair : DefinitionBuilder.Inputs)
 		{
 			if (Pair.Value.IsType<FIoHash>())
@@ -257,7 +257,7 @@ FBuildDefinitionInternal::FBuildDefinitionInternal(FBuildDefinitionBuilderIntern
 
 FBuildDefinitionInternal::FBuildDefinitionInternal(const FSharedString& InName, FCbObject&& InDefinition, bool& bOutIsValid)
 	: Name(InName)
-	, Function(InDefinition.FindView("Function"_ASV).AsString())
+	, Function(InDefinition.FindView(ANSITEXTVIEW("Function")).AsString())
 	, Definition(MoveTemp(InDefinition))
 	, Key{Definition.GetHash()}
 {
@@ -265,35 +265,35 @@ FBuildDefinitionInternal::FBuildDefinitionInternal(const FSharedString& InName,
 	Definition.MakeOwned();
 	bOutIsValid = Definition
 		&& IsValidBuildFunctionName(Function)
-		&& Algo::AllOf(Definition.AsView()["Constants"_ASV],
+		&& Algo::AllOf(Definition.AsView()[ANSITEXTVIEW("Constants")],
 			[](FCbFieldView Field) { return Field.GetName().Len() > 0 && Field.IsObject(); })
-		&& Algo::AllOf(Definition.AsView()["Inputs"_ASV]["Builds"_ASV], [](FCbFieldView Field)
+		&& Algo::AllOf(Definition.AsView()[ANSITEXTVIEW("Inputs")][ANSITEXTVIEW("Builds")], [](FCbFieldView Field)
 			{
 				return Field.GetName().Len() > 0 && Field.IsObject()
-					&& Field["Build"_ASV].IsHash()
-					&& Field["Id"_ASV].IsObjectId();
+					&& Field[ANSITEXTVIEW("Build")].IsHash()
+					&& Field[ANSITEXTVIEW("Id")].IsObjectId();
 			})
-		&& Algo::AllOf(Definition.AsView()["Inputs"_ASV]["BulkData"_ASV],
+		&& Algo::AllOf(Definition.AsView()[ANSITEXTVIEW("Inputs")][ANSITEXTVIEW("BulkData")],
 			[](FCbFieldView Field) { return Field.GetName().Len() > 0 && Field.IsUuid(); })
-		&& Algo::AllOf(Definition.AsView()["Inputs"_ASV]["Files"_ASV],
+		&& Algo::AllOf(Definition.AsView()[ANSITEXTVIEW("Inputs")][ANSITEXTVIEW("Files")],
 			[](FCbFieldView Field) { return Field.GetName().Len() > 0 && Field.AsString().Len() > 0; })
-		&& Algo::AllOf(Definition.AsView()["Inputs"_ASV]["Hashes"_ASV],
+		&& Algo::AllOf(Definition.AsView()[ANSITEXTVIEW("Inputs")][ANSITEXTVIEW("Hashes")],
 			[](FCbFieldView Field) { return Field.GetName().Len() > 0 && Field.IsBinaryAttachment(); });
 }
 
 bool FBuildDefinitionInternal::HasConstants() const
 {
-	return Definition["Constants"_ASV].HasValue();
+	return Definition[ANSITEXTVIEW("Constants")].HasValue();
 }
 
 bool FBuildDefinitionInternal::HasInputs() const
 {
-	return Definition["Inputs"_ASV].HasValue();
+	return Definition[ANSITEXTVIEW("Inputs")].HasValue();
 }
 
 void FBuildDefinitionInternal::IterateConstants(TFunctionRef<void (FUtf8StringView Key, FCbObject&& Value)> Visitor) const
 {
-	for (FCbField Field : Definition["Constants"_ASV])
+	for (FCbField Field : Definition[ANSITEXTVIEW("Constants")])
 	{
 		Visitor(Field.GetName(), Field.AsObject());
 	}
@@ -301,17 +301,17 @@ void FBuildDefinitionInternal::IterateConstants(TFunctionRef<void (FUtf8StringVi
 
 void FBuildDefinitionInternal::IterateInputBuilds(TFunctionRef<void (FUtf8StringView Key, const FBuildValueKey& ValueKey)> Visitor) const
 {
-	for (FCbFieldView Field : Definition.AsView()["Inputs"_ASV]["Builds"_ASV])
+	for (FCbFieldView Field : Definition.AsView()[ANSITEXTVIEW("Inputs")][ANSITEXTVIEW("Builds")])
 	{
-		const FBuildKey BuildKey{Field["Build"_ASV].AsHash()};
-		const FValueId Id = Field["Id"_ASV].AsObjectId();
+		const FBuildKey BuildKey{Field[ANSITEXTVIEW("Build")].AsHash()};
+		const FValueId Id = Field[ANSITEXTVIEW("Id")].AsObjectId();
 		Visitor(Field.GetName(), FBuildValueKey{BuildKey, Id});
 	}
 }
 
 void FBuildDefinitionInternal::IterateInputBulkData(TFunctionRef<void (FUtf8StringView Key, const FGuid& BulkDataId)> Visitor) const
 {
-	for (FCbFieldView Field : Definition.AsView()["Inputs"_ASV]["BulkData"_ASV])
+	for (FCbFieldView Field : Definition.AsView()[ANSITEXTVIEW("Inputs")][ANSITEXTVIEW("BulkData")])
 	{
 		Visitor(Field.GetName(), Field.AsUuid());
 	}
@@ -319,7 +319,7 @@ void FBuildDefinitionInternal::IterateInputBulkData(TFunctionRef<void (FUtf8Stri
 
 void FBuildDefinitionInternal::IterateInputFiles(TFunctionRef<void (FUtf8StringView Key, FUtf8StringView Path)> Visitor) const
 {
-	for (FCbFieldView Field : Definition.AsView()["Inputs"_ASV]["Files"_ASV])
+	for (FCbFieldView Field : Definition.AsView()[ANSITEXTVIEW("Inputs")][ANSITEXTVIEW("Files")])
 	{
 		Visitor(Field.GetName(), Field.AsString());
 	}
@@ -327,7 +327,7 @@ void FBuildDefinitionInternal::IterateInputFiles(TFunctionRef<void (FUtf8StringV
 
 void FBuildDefinitionInternal::IterateInputHashes(TFunctionRef<void (FUtf8StringView Key, const FIoHash& RawHash)> Visitor) const
 {
-	for (FCbFieldView Field : Definition.AsView()["Inputs"_ASV]["Hashes"_ASV])
+	for (FCbFieldView Field : Definition.AsView()[ANSITEXTVIEW("Inputs")][ANSITEXTVIEW("Hashes")])
 	{
 		Visitor(Field.GetName(), Field.AsBinaryAttachment());
 	}
diff --git a/Engine/Source/Developer/DerivedDataCache/Private/DerivedDataBuildJob.cpp b/Engine/Source/Developer/DerivedDataCache/Private/DerivedDataBuildJob.cpp
--- a/Engine/Source/Developer/DerivedDataCache/Private/DerivedDataBuildJob.cpp
+++ b/Engine/Source/Developer/DerivedDataCache/Private/DerivedDataBuildJob.cpp
@@ -401,25 +401,25 @@ void FBuildJob::CreateContext()
 	const IBuildFunction* Function = BuildSystem.GetFunctionRegistry().FindFunction(FunctionName);
 	if (BuildSystem.GetVersion() != Action.Get().GetBuildSystemVersion())
 	{
-		return CompleteWithError(WriteToUtf8String<192>("Failed because the build system is version "_ASV,
-			BuildSystem.GetVersion(), " when version "_ASV,
-			Action.Get().GetBuildSystemVersion(), " is expected."_ASV));
+		return CompleteWithError(WriteToUtf8String<192>(ANSITEXTVIEW("Failed because the build system is version "),
+			BuildSystem.GetVersion(), ANSITEXTVIEW(" when version "),
+			Action.Get().GetBuildSystemVersion(), ANSITEXTVIEW(" is expected.")));
 	}
 	else if (!Function)
 	{
-		return CompleteWithError(WriteToUtf8String<128>("Failed because the function "_ASV, FunctionName,
-			" was not found."_ASV));
+		return CompleteWithError(WriteToUtf8String<128>(ANSITEXTVIEW("Failed because the function "), FunctionName,
+			ANSITEXTVIEW(" was not found.")));
 	}
 	else if (!Function->GetVersion().IsValid())
 	{
-		return CompleteWithError(WriteToUtf8String<128>("Failed because the function "_ASV, FunctionName,
-			" has a version of zero."_ASV));
+		return CompleteWithError(WriteToUtf8String<128>(ANSITEXTVIEW("Failed because the function "), FunctionName,
+			ANSITEXTVIEW(" has a version of zero.")));
 	}
 	else if (Function->GetVersion() != Action.Get().GetFunctionVersion())
 	{
-		return CompleteWithError(WriteToUtf8String<192>("Failed because the function "_ASV, FunctionName,
-			" is version "_ASV, Function->GetVersion(), " when version "_ASV,
-			Action.Get().GetFunctionVersion(), " is expected."_ASV));
+		return CompleteWithError(WriteToUtf8String<192>(ANSITEXTVIEW("Failed because the function "), FunctionName,
+			ANSITEXTVIEW(" is version "), Function->GetVersion(), ANSITEXTVIEW(" when version "),
+			Action.Get().GetFunctionVersion(), ANSITEXTVIEW(" is expected.")));
 	}
 	else
 	{
@@ -533,7 +533,7 @@ void FBuildJob::EndCacheQuery(FCacheGetResponse&& Response)
 {
 	if (Response.Status == EStatus::Canceled)
 	{
-		return CompleteWithError("Build was canceled."_ASV);
+		return CompleteWithError(ANSITEXTVIEW("Build was canceled."));
 	}
 	if (Response.Status == EStatus::Ok)
 	{
@@ -545,7 +545,7 @@ void FBuildJob::EndCacheQuery(FCacheGetResponse&& Response)
 	}
 	if (!EnumHasAnyFlags(BuildPolicy.GetCombinedPolicy() & Context->GetBuildPolicyMask(), EBuildPolicy::Build))
 	{
-		return CompleteWithError("Failed to fetch from the cache and build policy does not allow execution."_ASV);
+		return CompleteWithError(ANSITEXTVIEW("Failed to fetch from the cache and build policy does not allow execution."));
 	}
 	return AdvanceToState(EBuildJobState::ExecuteRemote);
 }
@@ -618,11 +618,11 @@ void FBuildJob::EnterResolveKey()
 	}
 	if (DefinitionKey != FBuildKey::Empty)
 	{
-		return CompleteWithError("Failed to resolve null key."_ASV);
+		return CompleteWithError(ANSITEXTVIEW("Failed to resolve null key."));
 	}
 	if (!InputResolver)
 	{
-		return CompleteWithError("Failed to resolve key due to null input resolver."_ASV);
+		return CompleteWithError(ANSITEXTVIEW("Failed to resolve key due to null input resolver."));
 	}
 }
 
@@ -637,7 +637,7 @@ void FBuildJob::EndResolveKey(FBuildKeyResolvedParams&& Params)
 {
 	if (Params.Status == EStatus::Canceled)
 	{
-		return CompleteWithError("Build was canceled."_ASV);
+		return CompleteWithError(ANSITEXTVIEW("Build was canceled."));
 	}
 	if (Params.Status == EStatus::Ok && Params.Definition)
 	{
@@ -645,7 +645,7 @@ void FBuildJob::EndResolveKey(FBuildKeyResolvedParams&& Params)
 	}
 	else
 	{
-		return CompleteWithError(WriteToUtf8String<128>("Failed to resolve key "_ASV, Params.Key, "."_ASV));
+		return CompleteWithError(WriteToUtf8String<128>(ANSITEXTVIEW("Failed to resolve key "), Params.Key, ANSITEXTVIEW(".")));
 	}
 }
 
@@ -659,7 +659,7 @@ void FBuildJob::EnterResolveInputMeta()
 	}
 	if (!InputResolver)
 	{
-		return CompleteWithError("Failed to resolve input metadata due to null input resolver."_ASV);
+		return CompleteWithError(ANSITEXTVIEW("Failed to resolve input metadata due to null input resolver."));
 	}
 }
 
@@ -674,7 +674,7 @@ void FBuildJob::EndResolveInputMeta(FBuildInputMetaResolvedParams&& Params)
 {
 	if (Params.Status == EStatus::Canceled)
 	{
-		return CompleteWithError("Build was canceled."_ASV);
+		return CompleteWithError(ANSITEXTVIEW("Build was canceled."));
 	}
 	if (Params.Status == EStatus::Ok)
 	{
@@ -682,7 +682,7 @@ void FBuildJob::EndResolveInputMeta(FBuildInputMetaResolvedParams&& Params)
 	}
 	else
 	{
-		return CompleteWithError("Failed to resolve input metadata."_ASV);
+		return CompleteWithError(ANSITEXTVIEW("Failed to resolve input metadata."));
 	}
 }
 
@@ -732,7 +732,7 @@ void FBuildJob::EnterResolveInputData()
 
 	if (!InputResolver)
 	{
-		return CompleteWithError("Failed to resolve input data due to null input resolver."_ASV);
+		return CompleteWithError(ANSITEXTVIEW("Failed to resolve input data due to null input resolver."));
 	}
 }
 
@@ -757,7 +757,7 @@ void FBuildJob::EndResolveInputData(FBuildInputDataResolvedParams&& Params)
 {
 	if (Params.Status == EStatus::Canceled)
 	{
-		return CompleteWithError("Build was canceled."_ASV);
+		return CompleteWithError(ANSITEXTVIEW("Build was canceled."));
 	}
 	if (Params.Status == EStatus::Ok)
 	{
@@ -779,7 +779,7 @@ void FBuildJob::EndResolveInputData(FBuildInputDataResolvedParams&& Params)
 	}
 	else
 	{
-		return CompleteWithError("Failed to resolve input data."_ASV);
+		return CompleteWithError(ANSITEXTVIEW("Failed to resolve input data."));
 	}
 }
 
@@ -833,7 +833,7 @@ void FBuildJob::EndExecuteRemote(FBuildWorkerActionCompleteParams&& Params)
 {
 	if (Params.Status == EStatus::Canceled)
 	{
-		return CompleteWithError("Build was canceled."_ASV);
+		return CompleteWithError(ANSITEXTVIEW("Build was canceled."));
 	}
 	if (Params.Output)
 	{
@@ -888,13 +888,13 @@ void FBuildJob::EnterExecuteLocal()
 		}
 		else if (EnumHasAnyFlags(BuildStatus, EBuildStatus::BuildTryRemote))
 		{
-			return CompleteWithError("Failed because build policy does not allow local execution, "
-				"and remote execution failed to build."_ASV);
+			return CompleteWithError(ANSITEXTVIEW("Failed because build policy does not allow local execution, "
+				"and remote execution failed to build."));
 		}
 		else
 		{
-			return CompleteWithError("Failed because build policy does not allow local execution, "
-				"and remote execution was not available."_ASV);
+			return CompleteWithError(ANSITEXTVIEW("Failed because build policy does not allow local execution, "
+				"and remote execution was not available."));
 		}
 	}
 
@@ -907,9 +907,10 @@ void FBuildJob::EnterExecuteLocal()
 		}
 		else
 		{
-			CompleteWithError(WriteToUtf8String<256>("Failed because input '"_ASV, Key, "' with hash "_ASV,
-				RawHash, " ("_ASV, RawSize, " bytes) was resolved with hash "_ASV, Buffer.GetRawHash(),
-				" ("_ASV, Buffer.GetRawSize(), " bytes)."_ASV));
+			CompleteWithError(WriteToUtf8String<256>(ANSITEXTVIEW("Failed because input '"), Key,
+				ANSITEXTVIEW("' with hash "), RawHash, ANSITEXTVIEW(" ("), RawSize,
+				ANSITEXTVIEW(" bytes) was resolved with hash "), Buffer.GetRawHash(),
+				ANSITEXTVIEW(" ("), Buffer.GetRawSize(), ANSITEXTVIEW(" bytes).")));
 		}
 	});
 	Action.Reset();
@@ -1034,7 +1035,7 @@ void FBuildJob::CompleteWithError(FUtf8StringView Error)
 	{
 		return;
 	}
-	OutputBuilder.AddLog({"LogDerivedDataBuild"_ASV, Error, EBuildOutputLogLevel::Error});
+	OutputBuilder.AddLog({ANSITEXTVIEW("LogDerivedDataBuild"), Error, EBuildOutputLogLevel::Error});
 	return SetOutputNoCheck(OutputBuilder.Build(), EBuildJobState::Complete);
 }
 
@@ -1054,7 +1055,7 @@ void FBuildJob::AdvanceToState(EBuildJobState NewState)
 			NewState = EBuildJobState::Complete;
 			if (Output.IsNull())
 			{
-				OutputBuilder.AddLog({"LogDerivedDataBuild"_ASV, "Build was canceled."_ASV, EBuildOutputLogLevel::Error});
+				OutputBuilder.AddLog({ANSITEXTVIEW("LogDerivedDataBuild"), ANSITEXTVIEW("Build was canceled."), EBuildOutputLogLevel::Error});
 			}
 		}
 
diff --git a/Engine/Source/Developer/DerivedDataCache/Private/DerivedDataBuildJobContext.cpp b/Engine/Source/Developer/DerivedDataCache/Private/DerivedDataBuildJobContext.cpp
--- a/Engine/Source/Developer/DerivedDataCache/Private/DerivedDataBuildJobContext.cpp
+++ b/Engine/Source/Developer/DerivedDataCache/Private/DerivedDataBuildJobContext.cpp
@@ -93,7 +93,7 @@ FSharedBuffer FBuildJobContext::FindInput(FUtf8StringView Key) const
 				<< TEXT(" and raw size ") << Input->GetRawSize() << TEXT(" but has raw hash ") << RawHash
 				<< TEXT(" and raw size ") << Buffer.GetSize() << TEXT(" after decompression for build of '")
 				<< Job.GetName() << TEXT("' by ") << Job.GetFunction() << TEXT(".");
-			OutputBuilder.AddLog({"LogDerivedDataBuild"_ASV, FTCHARToUTF8(Error), EBuildOutputLogLevel::Error});
+			OutputBuilder.AddLog({ANSITEXTVIEW("LogDerivedDataBuild"), FTCHARToUTF8(Error), EBuildOutputLogLevel::Error});
 			UE_LOG(LogDerivedDataBuild, Error, TEXT("%.*s"), Error.Len(), Error.GetData());
 		}
 	}
diff --git a/Engine/Source/Developer/DerivedDataCache/Private/DerivedDataCachePolicy.cpp b/Engine/Source/Developer/DerivedDataCache/Private/DerivedDataCachePolicy.cpp
--- a/Engine/Source/Developer/DerivedDataCache/Private/DerivedDataCachePolicy.cpp
+++ b/Engine/Source/Developer/DerivedDataCache/Private/DerivedDataCachePolicy.cpp
@@ -186,15 +186,15 @@ FCacheRecordPolicy FCacheRecordPolicy::Transform(const TFunctionRef<ECachePolicy
 void FCacheRecordPolicy::Save(FCbWriter& Writer) const
 {
 	Writer.BeginObject();
-	Writer.AddString("BasePolicy"_ASV, WriteToUtf8String<128>(GetBasePolicy()));
+	Writer.AddString(ANSITEXTVIEW("BasePolicy"), WriteToUtf8String<128>(GetBasePolicy()));
 	if (!IsUniform())
 	{
-		Writer.BeginArray("ValuePolicies"_ASV);
+		Writer.BeginArray(ANSITEXTVIEW("ValuePolicies"));
 		for (const FCacheValuePolicy& Value : GetValuePolicies())
 		{
 			Writer.BeginObject();
-			Writer.AddObjectId("Id"_ASV, Value.Id);
-			Writer.AddString("Policy"_ASV, WriteToUtf8String<128>(Value.Policy));
+			Writer.AddObjectId(ANSITEXTVIEW("Id"), Value.Id);
+			Writer.AddString(ANSITEXTVIEW("Policy"), WriteToUtf8String<128>(Value.Policy));
 			Writer.EndObject();
 		}
 		Writer.EndArray();
@@ -204,17 +204,17 @@ void FCacheRecordPolicy::Save(FCbWriter& Writer) const
 
 FOptionalCacheRecordPolicy FCacheRecordPolicy::Load(const FCbObjectView Object)
 {
-	const FUtf8StringView BasePolicyText = Object["BasePolicy"_ASV].AsString();
+	const FUtf8StringView BasePolicyText = Object[ANSITEXTVIEW("BasePolicy")].AsString();
 	if (BasePolicyText.IsEmpty())
 	{
 		return {};
 	}
 
 	FCacheRecordPolicyBuilder Builder(ParseCachePolicy(BasePolicyText));
-	for (const FCbFieldView Value : Object["ValuePolicies"_ASV])
+	for (const FCbFieldView Value : Object[ANSITEXTVIEW("ValuePolicies")])
 	{
-		const FValueId Id = Value["Id"_ASV].AsObjectId();
-		const FUtf8StringView PolicyText = Value["Policy"_ASV].AsString();
+		const FValueId Id = Value[ANSITEXTVIEW("Id")].AsObjectId();
+		const FUtf8StringView PolicyText = Value[ANSITEXTVIEW("Policy")].AsString();
 		if (Id.IsNull() || PolicyText.IsEmpty())
 		{
 			return {};
diff --git a/Engine/Source/Developer/DerivedDataCache/Private/DerivedDataCacheRecord.cpp b/Engine/Source/Developer/DerivedDataCache/Private/DerivedDataCacheRecord.cpp
--- a/Engine/Source/Developer/DerivedDataCache/Private/DerivedDataCacheRecord.cpp
+++ b/Engine/Source/Developer/DerivedDataCache/Private/DerivedDataCacheRecord.cpp
@@ -169,16 +169,16 @@ FCbPackage FCacheRecord::Save() const
 void FCacheRecord::Save(FCbPackage& Attachments, FCbWriter& Writer) const
 {
 	Writer.BeginObject();
-	Writer << "Key"_ASV << GetKey();
+	Writer << ANSITEXTVIEW("Key") << GetKey();
 
 	if (const FCbObject& Meta = GetMeta())
 	{
-		Writer.AddObject("Meta"_ASV, Meta);
+		Writer.AddObject(ANSITEXTVIEW("Meta"), Meta);
 	}
 	TConstArrayView<FValueWithId> Values = GetValues();
 	if (!Values.IsEmpty())
 	{
-		Writer.BeginArray("Values"_ASV);
+		Writer.BeginArray(ANSITEXTVIEW("Values"));
 		for (const FValueWithId& Value : Values)
 		{
 			if (Value.HasData())
@@ -186,9 +186,9 @@ void FCacheRecord::Save(FCbPackage& Attachments, FCbWriter& Writer) const
 				Attachments.AddAttachment(FCbAttachment(Value.GetData()));
 			}
 			Writer.BeginObject();
-			Writer.AddObjectId("Id"_ASV, Value.GetId());
-			Writer.AddBinaryAttachment("RawHash"_ASV, Value.GetRawHash());
-			Writer.AddInteger("RawSize"_ASV, Value.GetRawSize());
+			Writer.AddObjectId(ANSITEXTVIEW("Id"), Value.GetId());
+			Writer.AddBinaryAttachment(ANSITEXTVIEW("RawHash"), Value.GetRawHash());
+			Writer.AddInteger(ANSITEXTVIEW("RawSize"), Value.GetRawSize());
 			Writer.EndObject();
 		}
 		Writer.EndArray();
@@ -201,13 +201,13 @@ FOptionalCacheRecord FCacheRecord::Load(const FCbPackage& Attachments, const FCb
 	const FCbObjectView ObjectView = Object;
 
 	// Check for the previous format of cache record. Remove this check in 5.1.
-	if (ObjectView["Value"_ASV] || ObjectView["Attachments"_ASV])
+	if (ObjectView[ANSITEXTVIEW("Value")] || ObjectView[ANSITEXTVIEW("Attachments")])
 	{
 		return FOptionalCacheRecord();
 	}
 
 	FCacheKey Key;
-	FCbObjectView KeyObject = ObjectView["Key"_ASV].AsObjectView();
+	FCbObjectView KeyObject = ObjectView[ANSITEXTVIEW("Key")].AsObjectView();
 	auto TrySetBucketName = [](FUtf8StringView Name, FCacheKey& Key)
 	{
 		if (Private::IsValidCacheBucketName(Name))
@@ -217,24 +217,24 @@ FOptionalCacheRecord FCacheRecord::Load(const FCbPackage& Attachments, const FCb
 		}
 		return false;
 	};
-	if (!TrySetBucketName(KeyObject["Bucket"_ASV].AsString(), Key))
+	if (!TrySetBucketName(KeyObject[ANSITEXTVIEW("Bucket")].AsString(), Key))
 	{
 		return FOptionalCacheRecord();
 	}
-	Key.Hash = KeyObject["Hash"_ASV].AsHash();
+	Key.Hash = KeyObject[ANSITEXTVIEW("Hash")].AsHash();
 
 	FCacheRecordBuilder Builder(Key);
 
-	Builder.SetMeta(Object["Meta"_ASV].AsObject());
+	Builder.SetMeta(Object[ANSITEXTVIEW("Meta")].AsObject());
 
 	auto LoadValue = [&Attachments](const FCbObjectView& ValueObject)
 	{
-		const FValueId Id = ValueObject["Id"_ASV].AsObjectId();
+		const FValueId Id = ValueObject[ANSITEXTVIEW("Id")].AsObjectId();
 		if (Id.IsNull())
 		{
 			return FValueWithId();
 		}
-		const FIoHash RawHash = ValueObject["RawHash"_ASV].AsHash();
+		const FIoHash RawHash = ValueObject[ANSITEXTVIEW("RawHash")].AsHash();
 		if (const FCbAttachment* Attachment = Attachments.FindAttachment(RawHash))
 		{
 			if (const FCompressedBuffer& Compressed = Attachment->AsCompressedBinary())
@@ -242,7 +242,7 @@ FOptionalCacheRecord FCacheRecord::Load(const FCbPackage& Attachments, const FCb
 				return FValueWithId(Id, Compressed);
 			}
 		}
-		const uint64 RawSize = ValueObject["RawSize"_ASV].AsUInt64(MAX_uint64);
+		const uint64 RawSize = ValueObject[ANSITEXTVIEW("RawSize")].AsUInt64(MAX_uint64);
 		if (!RawHash.IsZero() && RawSize != MAX_uint64)
 		{
 			return FValueWithId(Id, RawHash, RawSize);
@@ -253,7 +253,7 @@ FOptionalCacheRecord FCacheRecord::Load(const FCbPackage& Attachments, const FCb
 		}
 	};
 
-	for (FCbFieldView ValueField : ObjectView["Values"_ASV])
+	for (FCbFieldView ValueField : ObjectView[ANSITEXTVIEW("Values")])
 	{
 		FValueWithId Value = LoadValue(ValueField.AsObjectView());
 		if (!Value)
diff --git a/Engine/Source/Developer/DerivedDataCache/Private/HttpCacheStore.cpp b/Engine/Source/Developer/DerivedDataCache/Private/HttpCacheStore.cpp
--- a/Engine/Source/Developer/DerivedDataCache/Private/HttpCacheStore.cpp
+++ b/Engine/Source/Developer/DerivedDataCache/Private/HttpCacheStore.cpp
@@ -4092,19 +4092,19 @@ void FHttpCacheStore::RefCachedDataProbablyExistsBatchAsync(
 	RefsUri << "api/v1/refs/" << StructuredNamespace;
 	FCbWriter RequestWriter;
 	RequestWriter.BeginObject();
-	RequestWriter.BeginArray("ops"_ASV);
+	RequestWriter.BeginArray(ANSITEXTVIEW("ops"));
 	uint32 OpIndex = 0;
 	for (const FCacheGetValueRequest& ValueRef : ValueRefs)
 	{
 		RequestWriter.BeginObject();
-		RequestWriter.AddInteger("opId"_ASV, OpIndex);
-		RequestWriter.AddString("op"_ASV, "GET"_ASV);
+		RequestWriter.AddInteger(ANSITEXTVIEW("opId"), OpIndex);
+		RequestWriter.AddString(ANSITEXTVIEW("op"), ANSITEXTVIEW("GET"));
 		FCacheKey Key = ValueRef.Key;
 		FString Bucket(Key.Bucket.ToString());
 		Bucket.ToLowerInline();
-		RequestWriter.AddString("bucket"_ASV, Bucket);
-		RequestWriter.AddString("key"_ASV, LexToString(Key.Hash));
-		RequestWriter.AddBool("resolveAttachments"_ASV, true);
+		RequestWriter.AddString(ANSITEXTVIEW("bucket"), Bucket);
+		RequestWriter.AddString(ANSITEXTVIEW("key"), LexToString(Key.Hash));
+		RequestWriter.AddBool(ANSITEXTVIEW("resolveAttachments"), true);
 		RequestWriter.EndObject();
 		++OpIndex;
 	}
@@ -4140,7 +4140,7 @@ void FHttpCacheStore::RefCachedDataProbablyExistsBatchAsync(
 
 			const FCbObjectView ResponseObject = FCbObjectView(Request->GetResponseBuffer().GetData());
 
-			FCbArrayView ResultsArrayView = ResponseObject["results"_ASV].AsArrayView();
+			FCbArrayView ResultsArrayView = ResponseObject[ANSITEXTVIEW("results")].AsArrayView();
 
 			if (ResultsArrayView.Num() != ValueRefs.Num())
 			{
@@ -4157,9 +4157,9 @@ void FHttpCacheStore::RefCachedDataProbablyExistsBatchAsync(
 			for (FCbFieldView ResultFieldView : ResultsArrayView)
 			{
 				FCbObjectView ResultObjectView = ResultFieldView.AsObjectView();
-				uint32 OpId = ResultObjectView["opId"_ASV].AsUInt32();
-				FCbObjectView ResponseObjectView = ResultObjectView["response"_ASV].AsObjectView();
-				int32 StatusCode = ResultObjectView["statusCode"_ASV].AsInt32();
+				uint32 OpId = ResultObjectView[ANSITEXTVIEW("opId")].AsUInt32();
+				FCbObjectView ResponseObjectView = ResultObjectView[ANSITEXTVIEW("response")].AsObjectView();
+				int32 StatusCode = ResultObjectView[ANSITEXTVIEW("statusCode")].AsInt32();
 
 				if (OpId >= (uint32)ValueRefs.Num())
 				{
@@ -4187,8 +4187,8 @@ void FHttpCacheStore::RefCachedDataProbablyExistsBatchAsync(
 					continue;
 				}
 
-				const FIoHash RawHash = ResponseObjectView["RawHash"_ASV].AsHash();
-				const uint64 RawSize = ResponseObjectView["RawSize"_ASV].AsUInt64(MAX_uint64);
+				const FIoHash RawHash = ResponseObjectView[ANSITEXTVIEW("RawHash")].AsHash();
+				const uint64 RawSize = ResponseObjectView[ANSITEXTVIEW("RawSize")].AsUInt64(MAX_uint64);
 				if (RawHash.IsZero() || RawSize == MAX_uint64)
 				{
 					UE_LOG(LogDerivedDataCache, Display, TEXT("%s: Cache miss with invalid value for %s from '%s'"),
diff --git a/Engine/Source/Developer/DerivedDataCache/Private/Tests/HttpCacheStoreTest.cpp b/Engine/Source/Developer/DerivedDataCache/Private/Tests/HttpCacheStoreTest.cpp
--- a/Engine/Source/Developer/DerivedDataCache/Private/Tests/HttpCacheStoreTest.cpp
+++ b/Engine/Source/Developer/DerivedDataCache/Private/Tests/HttpCacheStoreTest.cpp
@@ -796,7 +796,7 @@ bool CacheStore::RunTest(const FString& Parameters)
 	{
 		TCbWriter<64> MetaWriter;
 		MetaWriter.BeginObject();
-		MetaWriter.AddInteger("MetaKey"_ASV, 42);
+		MetaWriter.AddInteger(ANSITEXTVIEW("MetaKey"), 42);
 		MetaWriter.EndObject();
 		FCbObject MetaObject = MetaWriter.Save().AsObject();
 
diff --git a/Engine/Source/Developer/DerivedDataCache/Private/ZenCacheStore.cpp b/Engine/Source/Developer/DerivedDataCache/Private/ZenCacheStore.cpp
--- a/Engine/Source/Developer/DerivedDataCache/Private/ZenCacheStore.cpp
+++ b/Engine/Source/Developer/DerivedDataCache/Private/ZenCacheStore.cpp
@@ -640,24 +640,24 @@ void FZenCacheStore::Put(
 		FCbWriter BatchWriter;
 		BatchWriter.BeginObject();
 		{
-			BatchWriter << "Method"_ASV << "PutCacheRecords";
-			BatchWriter.BeginObject("Params"_ASV);
+			BatchWriter << ANSITEXTVIEW("Method") << "PutCacheRecords";
+			BatchWriter.BeginObject(ANSITEXTVIEW("Params"));
 			{
 				ECachePolicy BatchDefaultPolicy = Batch[0].Policy.GetRecordPolicy();
-				BatchWriter << "DefaultPolicy"_ASV << *WriteToString<128>(BatchDefaultPolicy);
+				BatchWriter << ANSITEXTVIEW("DefaultPolicy") << *WriteToString<128>(BatchDefaultPolicy);
 
-				BatchWriter.BeginArray("Requests"_ASV);
+				BatchWriter.BeginArray(ANSITEXTVIEW("Requests"));
 				for (const FCachePutRequest& Request : Batch)
 				{
 					const FCacheRecord& Record = Request.Record;
 
 					BatchWriter.BeginObject();
 					{
-						BatchWriter.SetName("Record"_ASV);
+						BatchWriter.SetName(ANSITEXTVIEW("Record"));
 						Record.Save(BatchPackage, BatchWriter);
 						if (!Request.Policy.IsUniform() || Request.Policy.GetRecordPolicy() != BatchDefaultPolicy)
 						{
-							BatchWriter.SetName("Policy"_ASV);
+							BatchWriter.SetName(ANSITEXTVIEW("Policy"));
 							Request.Policy.Save(BatchWriter);
 						}
 					}
@@ -675,14 +675,14 @@ void FZenCacheStore::Put(
 
 		{
 			Zen::FZenScopedRequestPtr Request(RequestPool.Get());
-			HttpResult = Request->PerformRpc(TEXT("/z$/$rpc"_SV), BatchPackage, BatchResponse);
+			HttpResult = Request->PerformRpc(TEXTVIEW("/z$/$rpc"), BatchPackage, BatchResponse);
 		}
 
 		int32 RequestIndex = 0;
 		if (HttpResult == Zen::FZenHttpRequest::Result::Success)
 		{
 			const FCbObject& ResponseObj = BatchResponse.GetObject();
-			for (FCbField ResponseField : ResponseObj["Result"_ASV])
+			for (FCbField ResponseField : ResponseObj[ANSITEXTVIEW("Result")])
 			{
 				if (RequestIndex >= Batch.Num())
 				{
@@ -752,21 +752,21 @@ void FZenCacheStore::Get(
 		FCbWriter BatchRequest;
 		BatchRequest.BeginObject();
 		{
-			BatchRequest << "Method"_ASV << "GetCacheRecords"_ASV;
-			BatchRequest.BeginObject("Params"_ASV);
+			BatchRequest << ANSITEXTVIEW("Method") << ANSITEXTVIEW("GetCacheRecords");
+			BatchRequest.BeginObject(ANSITEXTVIEW("Params"));
 			{
 				ECachePolicy BatchDefaultPolicy = Batch[0].Policy.GetRecordPolicy();
-				BatchRequest << "DefaultPolicy"_ASV << *WriteToString<128>(BatchDefaultPolicy);
+				BatchRequest << ANSITEXTVIEW("DefaultPolicy") << *WriteToString<128>(BatchDefaultPolicy);
 
-				BatchRequest.BeginArray("Requests"_ASV);
+				BatchRequest.BeginArray(ANSITEXTVIEW("Requests"));
 				for (const FCacheGetRequest& Request : Batch)
 				{
 					BatchRequest.BeginObject();
 					{
-						BatchRequest << "Key"_ASV << Request.Key;
+						BatchRequest << ANSITEXTVIEW("Key") << Request.Key;
 						if (!Request.Policy.IsUniform() || Request.Policy.GetRecordPolicy() != BatchDefaultPolicy)
 						{
-							BatchRequest.SetName("Policy"_ASV);
+							BatchRequest.SetName(ANSITEXTVIEW("Policy"));
 							Request.Policy.Save(BatchRequest);
 						}
 					}
@@ -792,7 +792,7 @@ void FZenCacheStore::Get(
 		{
 			const FCbObject& ResponseObj = BatchResponse.GetObject();
 			
-			for (FCbField RecordField : ResponseObj["Result"_ASV])
+			for (FCbField RecordField : ResponseObj[ANSITEXTVIEW("Result")])
 			{
 				if (RequestIndex >= Batch.Num())
 				{
@@ -873,18 +873,18 @@ void FZenCacheStore::PutValue(
 		FCbWriter BatchWriter;
 		BatchWriter.BeginObject();
 		{
-			BatchWriter << "Method"_ASV << "PutCacheValues"_ASV;
-			BatchWriter.BeginObject("Params"_ASV);
+			BatchWriter << ANSITEXTVIEW("Method") << ANSITEXTVIEW("PutCacheValues");
+			BatchWriter.BeginObject(ANSITEXTVIEW("Params"));
 			{
 				ECachePolicy BatchDefaultPolicy = Batch[0].Policy;
-				BatchWriter << "DefaultPolicy"_ASV << *WriteToString<128>(BatchDefaultPolicy);
+				BatchWriter << ANSITEXTVIEW("DefaultPolicy") << *WriteToString<128>(BatchDefaultPolicy);
 
 				BatchWriter.BeginArray("Requests");
 				for (const FCachePutValueRequest& Request : Batch)
 				{
 					BatchWriter.BeginObject();
 					{
-						BatchWriter << "Key"_ASV << Request.Key;
+						BatchWriter << ANSITEXTVIEW("Key") << Request.Key;
 						const FValue& Value = Request.Value;
 						BatchWriter.AddBinaryAttachment("RawHash", Value.GetRawHash());
 						if (Value.HasData())
@@ -893,7 +893,7 @@ void FZenCacheStore::PutValue(
 						}
 						if (Request.Policy != BatchDefaultPolicy)
 						{
-							BatchWriter << "Policy"_ASV << WriteToString<128>(Request.Policy);
+							BatchWriter << ANSITEXTVIEW("Policy") << WriteToString<128>(Request.Policy);
 						}
 					}
 					BatchWriter.EndObject();
@@ -909,14 +909,14 @@ void FZenCacheStore::PutValue(
 		Zen::FZenHttpRequest::Result HttpResult = Zen::FZenHttpRequest::Result::Failed;
 		{
 			Zen::FZenScopedRequestPtr Request(RequestPool.Get());
-			HttpResult = Request->PerformRpc(TEXT("/z$/$rpc"_SV), BatchPackage, BatchResponse);
+			HttpResult = Request->PerformRpc(TEXTVIEW("/z$/$rpc"), BatchPackage, BatchResponse);
 		}
 
 		int32 RequestIndex = 0;
 		if (HttpResult == Zen::FZenHttpRequest::Result::Success)
 		{
 			const FCbObject& ResponseObj = BatchResponse.GetObject();
-			for (FCbField ResponseField : ResponseObj["Result"_ASV])
+			for (FCbField ResponseField : ResponseObj[ANSITEXTVIEW("Result")])
 			{
 				if (RequestIndex >= Batch.Num())
 				{
@@ -984,21 +984,21 @@ void FZenCacheStore::GetValue(
 			FCbWriter BatchRequest;
 			BatchRequest.BeginObject();
 			{
-				BatchRequest << "Method"_ASV << "GetCacheValues"_ASV;
-				BatchRequest.BeginObject("Params"_ASV);
+				BatchRequest << ANSITEXTVIEW("Method") << ANSITEXTVIEW("GetCacheValues");
+				BatchRequest.BeginObject(ANSITEXTVIEW("Params"));
 				{
 					ECachePolicy BatchDefaultPolicy = Batch[0].Policy;
-					BatchRequest << "DefaultPolicy"_ASV << *WriteToString<128>(BatchDefaultPolicy);
+					BatchRequest << ANSITEXTVIEW("DefaultPolicy") << *WriteToString<128>(BatchDefaultPolicy);
 
 					BatchRequest.BeginArray("Requests");
 					for (const FCacheGetValueRequest& Request : Batch)
 					{
 						BatchRequest.BeginObject();
 					{
-							BatchRequest << "Key"_ASV << Request.Key;
+							BatchRequest << ANSITEXTVIEW("Key") << Request.Key;
 							if (Request.Policy != BatchDefaultPolicy)
 							{
-								BatchRequest << "Policy"_ASV << WriteToString<128>(Request.Policy);
+								BatchRequest << ANSITEXTVIEW("Policy") << WriteToString<128>(Request.Policy);
 							}
 					}
 						BatchRequest.EndObject();
@@ -1015,7 +1015,7 @@ void FZenCacheStore::GetValue(
 			{
 				LLM_SCOPE_BYTAG(UntaggedDDCResult);
 				Zen::FZenScopedRequestPtr Request(RequestPool.Get());
-				HttpResult = Request->PerformRpc(TEXT("/z$/$rpc"_SV), BatchRequest.Save().AsObject(), BatchResponse);
+				HttpResult = Request->PerformRpc(TEXTVIEW("/z$/$rpc"), BatchRequest.Save().AsObject(), BatchResponse);
 			}
 
 			int32 RequestIndex = 0;
@@ -1023,7 +1023,7 @@ void FZenCacheStore::GetValue(
 			{
 				const FCbObject& ResponseObj = BatchResponse.GetObject();
 
-				for (FCbFieldView ResultField : ResponseObj["Result"_ASV])
+				for (FCbFieldView ResultField : ResponseObj[ANSITEXTVIEW("Result")])
 				{
 					if (RequestIndex >= Batch.Num())
 					{
@@ -1110,38 +1110,38 @@ void FZenCacheStore::GetChunks(
 		FCbWriter BatchRequest;
 		BatchRequest.BeginObject();
 		{
-			BatchRequest << "Method"_ASV << "GetCacheChunks";
-			BatchRequest.AddInteger("MethodVersion"_ASV, 1);
-			BatchRequest.BeginObject("Params"_ASV);
+			BatchRequest << ANSITEXTVIEW("Method") << "GetCacheChunks";
+			BatchRequest.AddInteger(ANSITEXTVIEW("MethodVersion"), 1);
+			BatchRequest.BeginObject(ANSITEXTVIEW("Params"));
 			{
 				ECachePolicy DefaultPolicy = Batch[0].Policy;
-				BatchRequest << "DefaultPolicy"_ASV << WriteToString<128>(DefaultPolicy);
-				BatchRequest.BeginArray("ChunkRequests"_ASV);
+				BatchRequest << ANSITEXTVIEW("DefaultPolicy") << WriteToString<128>(DefaultPolicy);
+				BatchRequest.BeginArray(ANSITEXTVIEW("ChunkRequests"));
 				for (const FCacheGetChunkRequest& Request : Batch)
 				{
 					BatchRequest.BeginObject();
 					
-					BatchRequest << "Key"_ASV << Request.Key;
+					BatchRequest << ANSITEXTVIEW("Key") << Request.Key;
 
 					if (Request.Id.IsValid())
 					{
-						BatchRequest.AddObjectId("ValueId"_ASV, Request.Id);
+						BatchRequest.AddObjectId(ANSITEXTVIEW("ValueId"), Request.Id);
 					}
 					if (Request.RawOffset != 0)
 					{
-						BatchRequest << "RawOffset"_ASV << Request.RawOffset;
+						BatchRequest << ANSITEXTVIEW("RawOffset") << Request.RawOffset;
 					}
 					if (Request.RawSize != MAX_uint64)
 					{
-						BatchRequest << "RawSize"_ASV << Request.RawSize;
+						BatchRequest << ANSITEXTVIEW("RawSize") << Request.RawSize;
 					}
 					if (!Request.RawHash.IsZero())
 					{
-						BatchRequest << "ChunkId"_ASV << Request.RawHash;
+						BatchRequest << ANSITEXTVIEW("ChunkId") << Request.RawHash;
 					}
 					if (Request.Policy != DefaultPolicy)
 					{
-						BatchRequest << "Policy"_ASV << WriteToString<128>(Request.Policy);
+						BatchRequest << ANSITEXTVIEW("Policy") << WriteToString<128>(Request.Policy);
 					}
 
 					BatchRequest.EndObject();
@@ -1166,7 +1166,7 @@ void FZenCacheStore::GetChunks(
 		{
 			const FCbObject& ResponseObj = BatchResponse.GetObject();
 
-			for (FCbFieldView ResultView : ResponseObj["Result"_ASV])
+			for (FCbFieldView ResultView : ResponseObj[ANSITEXTVIEW("Result")])
 			{
 				if (RequestIndex >= Batch.Num())
 				{
@@ -1187,7 +1187,7 @@ void FZenCacheStore::GetChunks(
 				}
 				else
 				{
-					FCbFieldView HashView = ResultObject["RawHash"_ASV];
+					FCbFieldView HashView = ResultObject[ANSITEXTVIEW("RawHash")];
 					RawHash = HashView.AsHash();
 					if (!HashView.HasError())
 					{
@@ -1204,7 +1204,7 @@ void FZenCacheStore::GetChunks(
 						}
 						else
 						{
-							FCbFieldView RawSizeField = ResultObject["RawSize"_ASV];
+							FCbFieldView RawSizeField = ResultObject[ANSITEXTVIEW("RawSize")];
 							uint64 TotalSize = RawSizeField.AsUInt64();
 							Succeeded = !RawSizeField.HasError();
 							if (Succeeded)
diff --git a/Engine/Source/Developer/DerivedDataCache/Public/DerivedDataBuildKey.h b/Engine/Source/Developer/DerivedDataCache/Public/DerivedDataBuildKey.h
--- a/Engine/Source/Developer/DerivedDataCache/Public/DerivedDataBuildKey.h
+++ b/Engine/Source/Developer/DerivedDataCache/Public/DerivedDataBuildKey.h
@@ -70,7 +70,7 @@ inline uint32 GetTypeHash(const FBuildKey& Key)
 template <typename CharType>
 inline TStringBuilderBase<CharType>& operator<<(TStringBuilderBase<CharType>& Builder, const FBuildKey& Key)
 {
-	return Builder << "Build/"_ASV << Key.Hash;
+	return Builder << ANSITEXTVIEW("Build/") << Key.Hash;
 }
 
 ///////////////////////////////////////////////////////////////////////////////////////////////////////////////////////
@@ -98,7 +98,7 @@ inline uint32 GetTypeHash(const FBuildActionKey& Key)
 template <typename CharType>
 inline TStringBuilderBase<CharType>& operator<<(TStringBuilderBase<CharType>& Builder, const FBuildActionKey& Key)
 {
-	return Builder << "BuildAction/"_ASV << Key.Hash;
+	return Builder << ANSITEXTVIEW("BuildAction/") << Key.Hash;
 }
 
 ///////////////////////////////////////////////////////////////////////////////////////////////////////////////////////
diff --git a/Engine/Source/Developer/TextureBuild/Private/TextureBuildFunction.cpp b/Engine/Source/Developer/TextureBuild/Private/TextureBuildFunction.cpp
--- a/Engine/Source/Developer/TextureBuild/Private/TextureBuildFunction.cpp
+++ b/Engine/Source/Developer/TextureBuild/Private/TextureBuildFunction.cpp
@@ -326,7 +326,7 @@ FGuid FTextureBuildFunction::GetVersion() const
 
 void FTextureBuildFunction::Configure(UE::DerivedData::FBuildConfigContext& Context) const
 {
-	Context.SetCacheBucket(UE::DerivedData::FCacheBucket("Texture"_ASV));
+	Context.SetCacheBucket(UE::DerivedData::FCacheBucket(ANSITEXTVIEW("Texture")));
 	Context.SetRequiredMemory(EstimateTextureBuildMemoryUsage(Context.FindConstant(UTF8TEXTVIEW("Settings"))));
 }
 
@@ -406,7 +406,7 @@ void FTextureBuildFunction::Build(UE::DerivedData::FBuildContext& Context) const
 		FCbWriter DescriptionWriter;
 		DescriptionWriter.BeginObject();
 
-		DescriptionWriter.BeginArray("Size"_ASV);
+		DescriptionWriter.BeginArray(ANSITEXTVIEW("Size"));
 		DescriptionWriter.AddInteger(CompressedMips[0].SizeX);
 		DescriptionWriter.AddInteger(CompressedMips[0].SizeY);
 		// this is brittle ; CompressedMips[] should report NumSlices itself
@@ -414,14 +414,14 @@ void FTextureBuildFunction::Build(UE::DerivedData::FBuildContext& Context) const
 		DescriptionWriter.AddInteger(NumSlices);
 		DescriptionWriter.EndArray();
 
-		DescriptionWriter.AddString("PixelFormat"_ASV, GetPixelFormatString((EPixelFormat)CompressedMips[0].PixelFormat));
-		DescriptionWriter.AddBool("bCubeMap"_ASV, BuildSettings.bCubemap);
-		DescriptionWriter.AddInteger("ExtData"_ASV, ExtData);
-		DescriptionWriter.AddInteger("NumMips"_ASV, MipCount);
-		DescriptionWriter.AddInteger("NumStreamingMips"_ASV, FirstInlineMip);
-		DescriptionWriter.AddInteger("NumMipsInTail"_ASV, NumMipsInTail);
+		DescriptionWriter.AddString(ANSITEXTVIEW("PixelFormat"), GetPixelFormatString((EPixelFormat)CompressedMips[0].PixelFormat));
+		DescriptionWriter.AddBool(ANSITEXTVIEW("bCubeMap"), BuildSettings.bCubemap);
+		DescriptionWriter.AddInteger(ANSITEXTVIEW("ExtData"), ExtData);
+		DescriptionWriter.AddInteger(ANSITEXTVIEW("NumMips"), MipCount);
+		DescriptionWriter.AddInteger(ANSITEXTVIEW("NumStreamingMips"), FirstInlineMip);
+		DescriptionWriter.AddInteger(ANSITEXTVIEW("NumMipsInTail"), NumMipsInTail);
 
-		DescriptionWriter.BeginArray("Mips"_ASV);
+		DescriptionWriter.BeginArray(ANSITEXTVIEW("Mips"));
 		int64 MipOffset = 0;
 		for (int32 MipIndex = 0; MipIndex < MipCount; ++MipIndex)
 		{
@@ -429,7 +429,7 @@ void FTextureBuildFunction::Build(UE::DerivedData::FBuildContext& Context) const
 
 			DescriptionWriter.BeginObject();
 			
-			DescriptionWriter.BeginArray("Size"_ASV);
+			DescriptionWriter.BeginArray(ANSITEXTVIEW("Size"));
 			DescriptionWriter.AddInteger(CompressedMip.SizeX);
 			DescriptionWriter.AddInteger(CompressedMip.SizeY);
 			DescriptionWriter.AddInteger(CompressedMip.SizeZ);
@@ -443,12 +443,12 @@ void FTextureBuildFunction::Build(UE::DerivedData::FBuildContext& Context) const
 			const bool bIsInlineMip = MipIndex >= FirstInlineMip;
 
 			EFileRegionType FileRegion = FFileRegion::SelectType(EPixelFormat(CompressedMip.PixelFormat));
-			DescriptionWriter.AddInteger("FileRegion"_ASV, static_cast<int32>(FileRegion));
+			DescriptionWriter.AddInteger(ANSITEXTVIEW("FileRegion"), static_cast<int32>(FileRegion));
 			if (bIsInlineMip)
 			{
-				DescriptionWriter.AddInteger("MipOffset"_ASV, MipOffset);
+				DescriptionWriter.AddInteger(ANSITEXTVIEW("MipOffset"), MipOffset);
 			}
-			DescriptionWriter.AddInteger("NumBytes"_ASV, CompressedMip.RawData.Num());
+			DescriptionWriter.AddInteger(ANSITEXTVIEW("NumBytes"), CompressedMip.RawData.Num());
 
 			DescriptionWriter.EndObject();
 
@@ -461,14 +461,14 @@ void FTextureBuildFunction::Build(UE::DerivedData::FBuildContext& Context) const
 
 		DescriptionWriter.EndObject();
 		FCbObject DescriptionObject = DescriptionWriter.Save().AsObject();
-		Context.AddValue(UE::DerivedData::FValueId::FromName("Description"_ASV), DescriptionObject);
+		Context.AddValue(UE::DerivedData::FValueId::FromName(ANSITEXTVIEW("Description")), DescriptionObject);
 	}
 
 	// Streaming mips
 	for (int32 MipIndex = 0; MipIndex < FirstInlineMip; ++MipIndex)
 	{
 		TAnsiStringBuilder<16> MipName;
-		MipName << "Mip"_ASV << MipIndex;
+		MipName << ANSITEXTVIEW("Mip") << MipIndex;
 
 		FSharedBuffer MipData = MakeSharedBufferFromArray(MoveTemp(CompressedMips[MipIndex].RawData));
 		Context.AddValue(UE::DerivedData::FValueId::FromName(MipName), MipData);
@@ -484,6 +484,6 @@ void FTextureBuildFunction::Build(UE::DerivedData::FBuildContext& Context) const
 	FCompositeBuffer MipTail(MipTailComponents);
 	if (MipTail.GetSize() > 0)
 	{
-		Context.AddValue(UE::DerivedData::FValueId::FromName("MipTail"_ASV), MipTail);
+		Context.AddValue(UE::DerivedData::FValueId::FromName(ANSITEXTVIEW("MipTail")), MipTail);
 	}
 }
diff --git a/Engine/Source/Developer/TraceServices/Private/AnalysisCache.cpp b/Engine/Source/Developer/TraceServices/Private/AnalysisCache.cpp
--- a/Engine/Source/Developer/TraceServices/Private/AnalysisCache.cpp
+++ b/Engine/Source/Developer/TraceServices/Private/AnalysisCache.cpp
@@ -137,18 +137,18 @@ bool FAnalysisCache::FFileContents::Save()
 	FCbWriter Writer;
 	Writer.BeginObject();
 	Writer << "Version" << CurrentVersion;
-	Writer.BeginArray("Index"_ASV);
+	Writer.BeginArray(ANSITEXTVIEW("Index"));
 	for (auto Entry : IndexEntries)
 	{
 		Writer.BeginObject();
-		Writer << "N"_ASV << Entry.Name;
-		Writer << "I"_ASV << Entry.Id;
-		Writer << "F"_ASV << Entry.Flags;
-		Writer.AddBinary("UD"_ASV, &Entry.UserData, UserDataSize);
+		Writer << ANSITEXTVIEW("N") << Entry.Name;
+		Writer << ANSITEXTVIEW("I") << Entry.Id;
+		Writer << ANSITEXTVIEW("F") << Entry.Flags;
+		Writer.AddBinary(ANSITEXTVIEW("UD"), &Entry.UserData, UserDataSize);
 		Writer.EndObject();
 	}
 	Writer.EndArray();
-	Writer.BeginArray("Blocks"_ASV);
+	Writer.BeginArray(ANSITEXTVIEW("Blocks"));
 	for (const FBlockEntry& Entry : Blocks)
 	{
 		Writer.AddBinary(&Entry, sizeof(FBlockEntry));
@@ -195,15 +195,15 @@ bool FAnalysisCache::FFileContents::Load()
 		return false;
 	}
 
-	FCbArrayView IndexArray = Package.GetObject().Find("Index"_ASV).AsArrayView();
+	FCbArrayView IndexArray = Package.GetObject().Find(ANSITEXTVIEW("Index")).AsArrayView();
 	IndexEntries.Reserve(IndexArray.Num());
 	for (FCbFieldView IndexEntry : IndexArray)
 	{
 		FCbObjectView IndexEntryObj = IndexEntry.AsObjectView();
-		FUtf8StringView NameView = IndexEntryObj["N"_ASV].AsString();
-		const uint32 Id = IndexEntryObj["I"_ASV].AsUInt32();
-		const uint32 Flags = IndexEntryObj["F"_ASV].AsUInt32();
-		FMemoryView UserData = IndexEntryObj["UD"_ASV].AsBinaryView();
+		FUtf8StringView NameView = IndexEntryObj[ANSITEXTVIEW("N")].AsString();
+		const uint32 Id = IndexEntryObj[ANSITEXTVIEW("I")].AsUInt32();
+		const uint32 Flags = IndexEntryObj[ANSITEXTVIEW("F")].AsUInt32();
+		FMemoryView UserData = IndexEntryObj[ANSITEXTVIEW("UD")].AsBinaryView();
 		FIndexEntry& Entry = IndexEntries.AddZeroed_GetRef();
 		Entry.Name = FString(NameView);
 		Entry.Id = Id;
@@ -213,7 +213,7 @@ bool FAnalysisCache::FFileContents::Load()
 		check(Remainder.GetSize() == 0);
 	}
 
-	FCbArrayView BlockArray = Package.GetObject().Find("Blocks"_ASV).AsArrayView();
+	FCbArrayView BlockArray = Package.GetObject().Find(ANSITEXTVIEW("Blocks")).AsArrayView();
 	Blocks.Reserve(BlockArray.Num());
 	
 	for (FCbFieldView BlockEntryView : BlockArray)
diff --git a/Engine/Source/Developer/Zen/Internal/ZenClient.cpp b/Engine/Source/Developer/Zen/Internal/ZenClient.cpp
--- a/Engine/Source/Developer/Zen/Internal/ZenClient.cpp
+++ b/Engine/Source/Developer/Zen/Internal/ZenClient.cpp
@@ -40,32 +40,32 @@ bool SendData(FSocket& Socket, FMemoryView Data)
 
 void CreateHttpUpgradeRequest(const TArray<FString>& Protocols, FAnsiStringBuilderBase& OutRequest)
 {
-	const FAnsiStringView Endpoint = "/zen"_ASV;
+	const FAnsiStringView Endpoint = ANSITEXTVIEW("/zen");
 
-	OutRequest << "GET "_ASV << Endpoint << " HTTP/1.1\r\n"_ASV;
-	OutRequest << "Upgrade: websocket\r\n"_ASV;
-	OutRequest << "Connection: upgrade\r\n"_ASV;
-	OutRequest << "Sec-WebSocket-Key: dGhlIHNhbXBsZSBub25jZQ==\r\n"_ASV;
+	OutRequest << ANSITEXTVIEW("GET ") << Endpoint << ANSITEXTVIEW(" HTTP/1.1\r\n");
+	OutRequest << ANSITEXTVIEW("Upgrade: websocket\r\n");
+	OutRequest << ANSITEXTVIEW("Connection: upgrade\r\n");
+	OutRequest << ANSITEXTVIEW("Sec-WebSocket-Key: dGhlIHNhbXBsZSBub25jZQ==\r\n");
 
 	if (Protocols.Num())
 	{
-		OutRequest << "Sec-WebSocket-Protocol: "_ASV;
+		OutRequest << ANSITEXTVIEW("Sec-WebSocket-Protocol: ");
 
 		for (int32 Idx = 0, Count = Protocols.Num(); Idx < Count; ++Idx)
 		{
 			if (Idx > 0)
 			{
-				OutRequest << " "_ASV;
+				OutRequest << ANSITEXTVIEW(" ");
 			}
 
 			FTCHARToUTF8 Utf8(*Protocols[Idx]);
 			OutRequest << FAnsiStringView(Utf8.Get(), Utf8.Length());
 		}
 		
-		OutRequest << "\r\n"_ASV;
+		OutRequest << ANSITEXTVIEW("\r\n");
 	}
 
-	OutRequest << "\r\n"_ASV;
+	OutRequest << ANSITEXTVIEW("\r\n");
 }
 
 bool ReceieveHttpUpgradeResponse(FSocket& Socket, FAnsiStringBuilderBase& Response)
@@ -98,7 +98,7 @@ bool ReceieveHttpUpgradeResponse(FSocket& Socket, FAnsiStringBuilderBase& Respon
 				return true;
 			}
 
-			Response << Line << "\n"_ASV;
+			Response << Line << ANSITEXTVIEW("\n");
 		}
 	}
 }
@@ -111,7 +111,7 @@ bool ValidateHandshake(FAnsiStringView UpgradeResponse, FString& OutReason)
 		FAnsiStringView StatusLine = UpgradeResponse.Left(NewLineIndex);
 		UpgradeResponse.RightChopInline(NewLineIndex + 1);
 
-		if (StatusLine != "HTTP/1.1 101 Switching Protocols"_ASV)
+		if (StatusLine != ANSITEXTVIEW("HTTP/1.1 101 Switching Protocols"))
 		{
 			FUTF8ToTCHAR Str(StatusLine.GetData(), StatusLine.Len());
 			OutReason = Str.Get(); 
@@ -483,7 +483,7 @@ bool FZenClient::Connect(FStringView Host, int32 Port)
 
 void FZenClient::Disconnect()
 {
-	CloseConnection(TEXT(""_SV));
+	CloseConnection(TEXTVIEW(""));
 }
 
 bool FZenClient::SendRequestMessage(FZenWebSocketMessage&& RequestMessage, FOnStreamResponse&& OnResponse)
diff --git a/Engine/Source/Editor/UnrealEd/Private/DerivedDataBuildRemoteExecutor.cpp b/Engine/Source/Editor/UnrealEd/Private/DerivedDataBuildRemoteExecutor.cpp
--- a/Engine/Source/Editor/UnrealEd/Private/DerivedDataBuildRemoteExecutor.cpp
+++ b/Engine/Source/Editor/UnrealEd/Private/DerivedDataBuildRemoteExecutor.cpp
@@ -241,7 +241,7 @@ public:
 
 	TConstArrayView<FStringView> GetHostPlatforms() const final
 	{
-		static constexpr FStringView HostPlatforms[]{TEXT("Win64"_SV), TEXT("Linux"_SV), TEXT("Mac"_SV)};
+		static constexpr FStringView HostPlatforms[]{TEXTVIEW("Win64"), TEXTVIEW("Linux"), TEXTVIEW("Mac")};
 		return HostPlatforms;
 	}
 
@@ -509,28 +509,28 @@ FCbObject FRemoteBuildExecutionRequest::BuildWorkerDescriptor(const FBuildWorker
 	FCbWriter WorkerDescriptor;
 	WorkerDescriptor.BeginObject();
 
-	WorkerDescriptor.AddString("name"_ASV, Worker.GetName());
-	WorkerDescriptor.AddString("path"_ASV, Worker.GetPath());
-	WorkerDescriptor.AddString("host"_ASV, Worker.GetHostPlatform());
-	WorkerDescriptor.AddUuid("buildsystem_version"_ASV, Worker.GetBuildSystemVersion());
-	WorkerDescriptor.AddInteger("timeout"_ASV, TimeoutSeconds);
-	WorkerDescriptor.AddInteger("cores"_ASV, 1);
-	//WorkerDescriptor.AddInteger("memory"_ASV, 1 * 1024 * 1024 * 1024);
+	WorkerDescriptor.AddString(ANSITEXTVIEW("name"), Worker.GetName());
+	WorkerDescriptor.AddString(ANSITEXTVIEW("path"), Worker.GetPath());
+	WorkerDescriptor.AddString(ANSITEXTVIEW("host"), Worker.GetHostPlatform());
+	WorkerDescriptor.AddUuid(ANSITEXTVIEW("buildsystem_version"), Worker.GetBuildSystemVersion());
+	WorkerDescriptor.AddInteger(ANSITEXTVIEW("timeout"), TimeoutSeconds);
+	WorkerDescriptor.AddInteger(ANSITEXTVIEW("cores"), 1);
+	//WorkerDescriptor.AddInteger(ANSITEXTVIEW("memory"), 1 * 1024 * 1024 * 1024);
 
-	WorkerDescriptor.BeginArray("environment"_ASV);
+	WorkerDescriptor.BeginArray(ANSITEXTVIEW("environment"));
 	Worker.IterateEnvironment([&WorkerDescriptor](FStringView Name, FStringView Value)
 		{
 			WorkerDescriptor.AddString(WriteToString<256>(Name, "=", Value));
 		});
 	WorkerDescriptor.EndArray();
 
-	WorkerDescriptor.BeginArray("executables"_ASV);
+	WorkerDescriptor.BeginArray(ANSITEXTVIEW("executables"));
 	Worker.IterateExecutables([&WorkerDescriptor](FStringView Key, const FIoHash& RawHash, uint64 RawSize)
 		{
 			WorkerDescriptor.BeginObject();
-			WorkerDescriptor.AddString("name"_ASV, Key);
-			WorkerDescriptor.AddBinaryAttachment("hash"_ASV, RawHash);
-			WorkerDescriptor.AddInteger("size"_ASV, RawSize);
+			WorkerDescriptor.AddString(ANSITEXTVIEW("name"), Key);
+			WorkerDescriptor.AddBinaryAttachment(ANSITEXTVIEW("hash"), RawHash);
+			WorkerDescriptor.AddInteger(ANSITEXTVIEW("size"), RawSize);
 			WorkerDescriptor.EndObject();
 		});
 	WorkerDescriptor.EndArray();
@@ -539,23 +539,23 @@ FCbObject FRemoteBuildExecutionRequest::BuildWorkerDescriptor(const FBuildWorker
 	Worker.IterateFiles([&WorkerDescriptor](FStringView Key, const FIoHash& RawHash, uint64 RawSize)
 		{
 			WorkerDescriptor.BeginObject();
-			WorkerDescriptor.AddString("name"_ASV, Key);
-			WorkerDescriptor.AddBinaryAttachment("hash"_ASV, RawHash);
-			WorkerDescriptor.AddInteger("size"_ASV, RawSize);
+			WorkerDescriptor.AddString(ANSITEXTVIEW("name"), Key);
+			WorkerDescriptor.AddBinaryAttachment(ANSITEXTVIEW("hash"), RawHash);
+			WorkerDescriptor.AddInteger(ANSITEXTVIEW("size"), RawSize);
 			WorkerDescriptor.EndObject();
 		});
 	WorkerDescriptor.EndArray();
 
-	WorkerDescriptor.BeginArray("dirs"_ASV);
+	WorkerDescriptor.BeginArray(ANSITEXTVIEW("dirs"));
 	WorkerDescriptor.AddString(WriteToString<256>("Engine/Binaries/", Worker.GetHostPlatform()));
 	WorkerDescriptor.EndArray();
 
-	WorkerDescriptor.BeginArray("functions"_ASV);
+	WorkerDescriptor.BeginArray(ANSITEXTVIEW("functions"));
 	Worker.IterateFunctions([&WorkerDescriptor](FUtf8StringView Name, const FGuid& Version)
 		{
 			WorkerDescriptor.BeginObject();
-			WorkerDescriptor.AddString("name"_ASV, Name);
-			WorkerDescriptor.AddUuid("version"_ASV, Version);
+			WorkerDescriptor.AddString(ANSITEXTVIEW("name"), Name);
+			WorkerDescriptor.AddUuid(ANSITEXTVIEW("version"), Version);
 			WorkerDescriptor.EndObject();
 		});
 	WorkerDescriptor.EndArray();
diff --git a/Engine/Source/Runtime/Core/Private/Serialization/CompactBinarySerialization.cpp b/Engine/Source/Runtime/Core/Private/Serialization/CompactBinarySerialization.cpp
--- a/Engine/Source/Runtime/Core/Private/Serialization/CompactBinarySerialization.cpp
+++ b/Engine/Source/Runtime/Core/Private/Serialization/CompactBinarySerialization.cpp
@@ -282,13 +282,13 @@ public:
 		if (FUtf8StringView Name = Field.GetName(); !Name.IsEmpty())
 		{
 			AppendQuotedString(Name);
-			Builder << ": "_ASV;
+			Builder << ANSITEXTVIEW(": ");
 		}
 
 		switch (FCbValue Accessor = Field.GetValue(); Accessor.GetType())
 		{
 		case ECbFieldType::Null:
-			Builder << "null"_ASV;
+			Builder << ANSITEXTVIEW("null");
 			break;
 		case ECbFieldType::Object:
 		case ECbFieldType::UniformObject:
@@ -341,10 +341,10 @@ public:
 			Builder.Appendf(UTF8TEXT("%.17g"), Accessor.AsFloat64());
 			break;
 		case ECbFieldType::BoolFalse:
-			Builder << "false"_ASV;
+			Builder << ANSITEXTVIEW("false");
 			break;
 		case ECbFieldType::BoolTrue:
-			Builder << "true"_ASV;
+			Builder << ANSITEXTVIEW("true");
 			break;
 		case ECbFieldType::ObjectAttachment:
 		case ECbFieldType::BinaryAttachment:
@@ -439,13 +439,13 @@ private:
 			{
 				switch (Char)
 				{
-				case '\\': Builder << "\\\\"_ASV; break;
-				case '\"': Builder << "\\\""_ASV; break;
-				case '\b': Builder << "\\b"_ASV; break;
-				case '\f': Builder << "\\f"_ASV; break;
-				case '\n': Builder << "\\n"_ASV; break;
-				case '\r': Builder << "\\r"_ASV; break;
-				case '\t': Builder << "\\t"_ASV; break;
+				case '\\': Builder << ANSITEXTVIEW("\\\\"); break;
+				case '\"': Builder << ANSITEXTVIEW("\\\""); break;
+				case '\b': Builder << ANSITEXTVIEW("\\b"); break;
+				case '\f': Builder << ANSITEXTVIEW("\\f"); break;
+				case '\n': Builder << ANSITEXTVIEW("\\n"); break;
+				case '\r': Builder << ANSITEXTVIEW("\\r"); break;
+				case '\t': Builder << ANSITEXTVIEW("\\t"); break;
 				default:
 					Builder.Appendf(UTF8TEXT("\\u%04x"), uint32(Char));
 					break;
diff --git a/Engine/Source/Runtime/Core/Private/Tests/Containers/StringViewTest.cpp b/Engine/Source/Runtime/Core/Private/Tests/Containers/StringViewTest.cpp
--- a/Engine/Source/Runtime/Core/Private/Tests/Containers/StringViewTest.cpp
+++ b/Engine/Source/Runtime/Core/Private/Tests/Containers/StringViewTest.cpp
@@ -106,8 +106,8 @@ bool FStringViewTestCtor::RunTest(const FString& Parameters)
 	// Create using string view literals
 	{
 		FStringView View = TEXTVIEW("Test");
-		FAnsiStringView ViewAnsi = "Test"_ASV;
-		FWideStringView ViewWide = WIDETEXT("Test"_WSV);
+		FAnsiStringView ViewAnsi = ANSITEXTVIEW("Test");
+		FWideStringView ViewWide = WIDETEXTVIEW("Test");
 	}
 
 	// Verify that class template argument deduction is working
diff --git a/Engine/Source/Runtime/Core/Private/Tests/Misc/AsciiSetTest.cpp b/Engine/Source/Runtime/Core/Private/Tests/Misc/AsciiSetTest.cpp
--- a/Engine/Source/Runtime/Core/Private/Tests/Misc/AsciiSetTest.cpp
+++ b/Engine/Source/Runtime/Core/Private/Tests/Misc/AsciiSetTest.cpp
@@ -47,18 +47,18 @@ bool TAsciiSetTest::RunTest(const FString& Parameters)
 	TestEqual(TEXT("AdvanceToLast"),	*FAsciiSet::FindLastOrEnd("NonWhitespaceNonWhitespace", Whitespaces), '\0');
 
 	constexpr FAsciiSet Lowercase("abcdefghijklmnopqrstuvwxyz");
-	TestEqual(TEXT("TrimPrefixWithout"),		FAsciiSet::TrimPrefixWithout("ABcdEF"_ASV, Lowercase), "cdEF"_ASV);
-	TestEqual(TEXT("FindPrefixWithout"),		FAsciiSet::FindPrefixWithout("ABcdEF"_ASV, Lowercase), "AB"_ASV);
-	TestEqual(TEXT("TrimSuffixWithout"),		FAsciiSet::TrimSuffixWithout("ABcdEF"_ASV, Lowercase), "ABcd"_ASV);
-	TestEqual(TEXT("FindSuffixWithout"),		FAsciiSet::FindSuffixWithout("ABcdEF"_ASV, Lowercase), "EF"_ASV);
-	TestEqual(TEXT("TrimPrefixWithout none"),	FAsciiSet::TrimPrefixWithout("same"_ASV, Lowercase), "same"_ASV);
-	TestEqual(TEXT("FindPrefixWithout none"),	FAsciiSet::FindPrefixWithout("same"_ASV, Lowercase), ""_ASV);
-	TestEqual(TEXT("TrimSuffixWithout none"),	FAsciiSet::TrimSuffixWithout("same"_ASV, Lowercase), "same"_ASV);
-	TestEqual(TEXT("FindSuffixWithout none"),	FAsciiSet::FindSuffixWithout("same"_ASV, Lowercase), ""_ASV);
-	TestEqual(TEXT("TrimPrefixWithout empty"),	FAsciiSet::TrimPrefixWithout(""_ASV, Lowercase), ""_ASV);
-	TestEqual(TEXT("FindPrefixWithout empty"),	FAsciiSet::FindPrefixWithout(""_ASV, Lowercase), ""_ASV);
-	TestEqual(TEXT("TrimSuffixWithout empty"),	FAsciiSet::TrimSuffixWithout(""_ASV, Lowercase), ""_ASV);
-	TestEqual(TEXT("FindSuffixWithout empty"),	FAsciiSet::FindSuffixWithout(""_ASV, Lowercase), ""_ASV);
+	TestEqual(TEXT("TrimPrefixWithout"),		FAsciiSet::TrimPrefixWithout(ANSITEXTVIEW("ABcdEF"), Lowercase), ANSITEXTVIEW("cdEF"));
+	TestEqual(TEXT("FindPrefixWithout"),		FAsciiSet::FindPrefixWithout(ANSITEXTVIEW("ABcdEF"), Lowercase), ANSITEXTVIEW("AB"));
+	TestEqual(TEXT("TrimSuffixWithout"),		FAsciiSet::TrimSuffixWithout(ANSITEXTVIEW("ABcdEF"), Lowercase), ANSITEXTVIEW("ABcd"));
+	TestEqual(TEXT("FindSuffixWithout"),		FAsciiSet::FindSuffixWithout(ANSITEXTVIEW("ABcdEF"), Lowercase), ANSITEXTVIEW("EF"));
+	TestEqual(TEXT("TrimPrefixWithout none"),	FAsciiSet::TrimPrefixWithout(ANSITEXTVIEW("same"), Lowercase), ANSITEXTVIEW("same"));
+	TestEqual(TEXT("FindPrefixWithout none"),	FAsciiSet::FindPrefixWithout(ANSITEXTVIEW("same"), Lowercase), ANSITEXTVIEW(""));
+	TestEqual(TEXT("TrimSuffixWithout none"),	FAsciiSet::TrimSuffixWithout(ANSITEXTVIEW("same"), Lowercase), ANSITEXTVIEW("same"));
+	TestEqual(TEXT("FindSuffixWithout none"),	FAsciiSet::FindSuffixWithout(ANSITEXTVIEW("same"), Lowercase), ANSITEXTVIEW(""));
+	TestEqual(TEXT("TrimPrefixWithout empty"),	FAsciiSet::TrimPrefixWithout(ANSITEXTVIEW(""), Lowercase), ANSITEXTVIEW(""));
+	TestEqual(TEXT("FindPrefixWithout empty"),	FAsciiSet::FindPrefixWithout(ANSITEXTVIEW(""), Lowercase), ANSITEXTVIEW(""));
+	TestEqual(TEXT("TrimSuffixWithout empty"),	FAsciiSet::TrimSuffixWithout(ANSITEXTVIEW(""), Lowercase), ANSITEXTVIEW(""));
+	TestEqual(TEXT("FindSuffixWithout empty"),	FAsciiSet::FindSuffixWithout(ANSITEXTVIEW(""), Lowercase), ANSITEXTVIEW(""));
 
 
 	auto TestHasFunctions = [&](auto MakeString)
diff --git a/Engine/Source/Runtime/Core/Private/Tests/Serialization/CompactBinaryPackageTest.cpp b/Engine/Source/Runtime/Core/Private/Tests/Serialization/CompactBinaryPackageTest.cpp
--- a/Engine/Source/Runtime/Core/Private/Tests/Serialization/CompactBinaryPackageTest.cpp
+++ b/Engine/Source/Runtime/Core/Private/Tests/Serialization/CompactBinaryPackageTest.cpp
@@ -95,7 +95,7 @@ bool FCbAttachmentTest::RunTest(const FString& Parameters)
 	{
 		FCbWriter Writer;
 		Writer.BeginObject();
-		Writer << "Name"_ASV << 42;
+		Writer << ANSITEXTVIEW("Name") << 42;
 		Writer.EndObject();
 		FCbObject Object = Writer.Save().AsObject();
 		FCbAttachment Attachment(Object);
@@ -130,7 +130,7 @@ bool FCbAttachmentTest::RunTest(const FString& Parameters)
 	{
 		FCbWriter Writer;
 		Writer.BeginObject();
-		Writer << "Name"_ASV << 42;
+		Writer << ANSITEXTVIEW("Name") << 42;
 		Writer.EndObject();
 		FCbObject Object = Writer.Save().AsObject();
 		FCbObject ObjectView = FCbObject::MakeView(Object);
@@ -200,7 +200,7 @@ bool FCbAttachmentTest::RunTest(const FString& Parameters)
 	{
 		FCbWriter ValueWriter;
 		ValueWriter.BeginObject();
-		ValueWriter << "Name"_ASV << 42;
+		ValueWriter << ANSITEXTVIEW("Name") << 42;
 		ValueWriter.EndObject();
 		const FCbObject Value = ValueWriter.Save().AsObject();
 		TestEqual(TEXT("FCbAttachment(LoadObjectView).Validate"),
diff --git a/Engine/Source/Runtime/Core/Private/Tests/Serialization/CompactBinaryValidationTest.cpp b/Engine/Source/Runtime/Core/Private/Tests/Serialization/CompactBinaryValidationTest.cpp
--- a/Engine/Source/Runtime/Core/Private/Tests/Serialization/CompactBinaryValidationTest.cpp
+++ b/Engine/Source/Runtime/Core/Private/Tests/Serialization/CompactBinaryValidationTest.cpp
@@ -367,7 +367,7 @@ bool FCbValidateAttachmentTest::RunTest(const FString& Parameters)
 	{
 		TCbWriter<256> Writer;
 		Writer.BeginObject();
-		Writer.AddNull("N"_ASV);
+		Writer.AddNull(ANSITEXTVIEW("N"));
 		Writer.EndObject();
 		FBufferArchive Buffer;
 		FCbAttachment(Writer.Save().AsObject()).Save(Buffer);
@@ -401,7 +401,7 @@ bool FCbValidateAttachmentTest::RunTest(const FString& Parameters)
 		FBufferArchive Buffer;
 		FCbWriter Writer;
 		Writer.AddBinaryAttachment(FIoHash::HashBuffer(BinaryView));
-		Writer.SetName("Name"_ASV).AddBinary(BinaryView);
+		Writer.SetName(ANSITEXTVIEW("Name")).AddBinary(BinaryView);
 		Writer.Save(Buffer);
 		TestEqual(TEXT("ValidateCompactBinaryAttachment(NameOnValue)"), ValidateCompactBinaryAttachment(MakeMemoryView(Buffer), ECbValidateMode::All), ECbValidateError::InvalidPackageFormat);
 		TestEqual(TEXT("ValidateCompactBinaryAttachment(NameOnValue, Mode)"), ValidateCompactBinaryAttachment(MakeMemoryView(Buffer), ECbValidateMode::All & ~ECbValidateMode::Package), ECbValidateError::None);
@@ -409,7 +409,7 @@ bool FCbValidateAttachmentTest::RunTest(const FString& Parameters)
 	{
 		FBufferArchive Buffer;
 		FCbWriter Writer;
-		Writer.SetName("Name"_ASV).AddBinaryAttachment(FIoHash::HashBuffer(BinaryView));
+		Writer.SetName(ANSITEXTVIEW("Name")).AddBinaryAttachment(FIoHash::HashBuffer(BinaryView));
 		Writer.AddBinary(BinaryView);
 		Writer.Save(Buffer);
 		TestEqual(TEXT("ValidateCompactBinaryAttachment(NameOnHash)"), ValidateCompactBinaryAttachment(MakeMemoryView(Buffer), ECbValidateMode::All), ECbValidateError::InvalidPackageFormat);
@@ -532,7 +532,7 @@ bool FCbValidatePackageTest::RunTest(const FString& Parameters)
 		FBufferArchive Buffer;
 		TCbWriter<256> Writer;
 		Writer.AddHash(CommonRootObjectHash);
-		Writer.SetName("Name"_ASV);
+		Writer.SetName(ANSITEXTVIEW("Name"));
 		Writer.AddObject(CommonRootObject);
 		Writer.AddNull();
 		Writer.Save(Buffer);
@@ -542,7 +542,7 @@ bool FCbValidatePackageTest::RunTest(const FString& Parameters)
 	{
 		FBufferArchive Buffer;
 		TCbWriter<256> Writer;
-		Writer.SetName("Name"_ASV);
+		Writer.SetName(ANSITEXTVIEW("Name"));
 		Writer.AddHash(CommonRootObjectHash);
 		Writer.AddObject(CommonRootObject);
 		Writer.AddNull();
@@ -556,7 +556,7 @@ bool FCbValidatePackageTest::RunTest(const FString& Parameters)
 		Writer.AddHash(CommonRootObjectHash);
 		Writer.AddObject(CommonRootObject);
 		Writer.AddBinaryAttachment(BinaryBufferHash);
-		Writer.SetName("Name"_ASV).AddBinary(BinaryBuffer);
+		Writer.SetName(ANSITEXTVIEW("Name")).AddBinary(BinaryBuffer);
 		Writer.AddNull();
 		Writer.Save(Buffer);
 		TestEqual(TEXT("ValidateCompactBinaryPackage(NameOnAttachment)"), ValidateCompactBinaryPackage(MakeMemoryView(Buffer), ECbValidateMode::All), ECbValidateError::InvalidPackageFormat);
diff --git a/Engine/Source/Runtime/Core/Public/Serialization/CompactBinary.h b/Engine/Source/Runtime/Core/Public/Serialization/CompactBinary.h
--- a/Engine/Source/Runtime/Core/Public/Serialization/CompactBinary.h
+++ b/Engine/Source/Runtime/Core/Public/Serialization/CompactBinary.h
@@ -56,7 +56,7 @@
  *         }
  *     }
  *
- *     for (FCbFieldView Format : Params.FindView("Formats"_ASV))
+ *     for (FCbFieldView Format : Params.FindView(ANSITEXTVIEW("Formats")))
  *     {
  *         BeginCompress(FName(Format.AsString()));
  *     }
diff --git a/Engine/Source/Runtime/StorageServerClient/Private/StorageServerConnection.cpp b/Engine/Source/Runtime/StorageServerClient/Private/StorageServerConnection.cpp
--- a/Engine/Source/Runtime/StorageServerClient/Private/StorageServerConnection.cpp
+++ b/Engine/Source/Runtime/StorageServerClient/Private/StorageServerConnection.cpp
@@ -609,11 +609,11 @@ bool FStorageServerConnection::ReadChunkRequest(const FIoChunkId& ChunkId, uint6
 	{
 		if (HaveQuery)
 		{
-			ResourceBuilder.Append("&"_ASV);
+			ResourceBuilder.Append(ANSITEXTVIEW("&"));
 		}
 		else
 		{
-			ResourceBuilder.Append("?"_ASV);
+			ResourceBuilder.Append(ANSITEXTVIEW("?"));
 			HaveQuery = true;
 		}
 	};
diff --git a/Engine/Source/Runtime/StorageServerClient/Public/StorageServerConnection.h b/Engine/Source/Runtime/StorageServerClient/Public/StorageServerConnection.h
--- a/Engine/Source/Runtime/StorageServerClient/Public/StorageServerConnection.h
+++ b/Engine/Source/Runtime/StorageServerClient/Public/StorageServerConnection.h
@@ -31,27 +31,27 @@ inline FAnsiStringView GetMimeTypeString(EStorageServerContentType ContentType)
 	switch (ContentType)
 	{
 		case EStorageServerContentType::CbObject:
-			return "application/ue-x-cb"_ASV;
+			return ANSITEXTVIEW("application/ue-x-cb");
 		case EStorageServerContentType::Binary:
-			return "application/octet-stream"_ASV;
+			return ANSITEXTVIEW("application/octet-stream");
 		case EStorageServerContentType::CompressedBinary:
-			return "application/x-ue-comp"_ASV;
+			return ANSITEXTVIEW("application/x-ue-comp");
 		default:
-			return "unknown"_ASV;
+			return ANSITEXTVIEW("unknown");
 	};
 };
 
 inline EStorageServerContentType GetMimeType(const FAnsiStringView& ContentType)
 {
-	if (ContentType == "application/octet-stream"_ASV)
+	if (ContentType == ANSITEXTVIEW("application/octet-stream"))
 	{
 		return EStorageServerContentType::Binary;
 	}
-	else if (ContentType == "application/x-ue-comp"_ASV)
+	else if (ContentType == ANSITEXTVIEW("application/x-ue-comp"))
 	{
 		return EStorageServerContentType::CompressedBinary;
 	}
-	else if (ContentType == "application/x-ue-cb"_ASV)
+	else if (ContentType == ANSITEXTVIEW("application/x-ue-cb"))
 	{
 		return EStorageServerContentType::CbObject;
 	}
