DDC: Added a timing scope to the flush (deletion of everything) of the file system cache

#rnx

diff --git a/Engine/Source/Developer/DerivedDataCache/Private/FileSystemCacheStore.cpp b/Engine/Source/Developer/DerivedDataCache/Private/FileSystemCacheStore.cpp
--- a/Engine/Source/Developer/DerivedDataCache/Private/FileSystemCacheStore.cpp
+++ b/Engine/Source/Developer/DerivedDataCache/Private/FileSystemCacheStore.cpp
@@ -916,6 +916,7 @@ FFileSystemCacheStore::FFileSystemCacheStore(
 	bool bFlush = false;
 	if (!bReadOnly && FParse::Bool(InParams, TEXT("Flush="), bFlush) && bFlush)
 	{
+		TRACE_CPUPROFILER_EVENT_SCOPE(FileSystemDDC_Flush);
 		IFileManager::Get().DeleteDirectory(*(CachePath / TEXT("")), /*bRequireExists*/ false, /*bTree*/ true);
 	}
 
