DDC: Added and clarified wait timing scopes in the Horde Storage cache client

#rnx

diff --git a/Engine/Source/Developer/DerivedDataCache/Private/HttpCacheStore.cpp b/Engine/Source/Developer/DerivedDataCache/Private/HttpCacheStore.cpp
--- a/Engine/Source/Developer/DerivedDataCache/Private/HttpCacheStore.cpp
+++ b/Engine/Source/Developer/DerivedDataCache/Private/HttpCacheStore.cpp
@@ -701,11 +701,13 @@ private:
 
 		void Cancel() final
 		{
+			TRACE_CPUPROFILER_EVENT_SCOPE(HttpDDC_Cancel);
 			Event.Wait();
 		}
 
 		void Wait() final
 		{
+			TRACE_CPUPROFILER_EVENT_SCOPE(HttpDDC_Wait);
 			Event.Wait();
 		}
 	};
@@ -1174,7 +1176,7 @@ struct FRequestPool
 	 */
 	FHttpRequest* WaitForFreeRequest(bool bUnboundedOverflow = false)
 	{
-		TRACE_CPUPROFILER_EVENT_SCOPE(HttpDDC_WaitForConnPool);
+		TRACE_CPUPROFILER_EVENT_SCOPE(HttpDDC_WaitForConnection);
 
 		FHttpRequest* Request = GetFreeRequest(bUnboundedOverflow);
 		if (Request == nullptr)
@@ -1677,7 +1679,7 @@ private:
 			{
 				// Wait until "driver" has done query
 				{
-					TRACE_CPUPROFILER_EVENT_SCOPE(WaitForMasterOfBatch);
+					TRACE_CPUPROFILER_EVENT_SCOPE(HttpDDC_WaitForMasterOfBatch);
 					Batch.Complete->Wait(~0);
 				}
 
