DDC: Fixed simulated cache miss salt differing for every node when using -DDC-All-Miss options

#rnx

diff --git a/Engine/Source/Developer/DerivedDataCache/Private/DerivedDataBackends.cpp b/Engine/Source/Developer/DerivedDataCache/Private/DerivedDataBackends.cpp
--- a/Engine/Source/Developer/DerivedDataCache/Private/DerivedDataBackends.cpp
+++ b/Engine/Source/Developer/DerivedDataCache/Private/DerivedDataBackends.cpp
@@ -107,6 +107,9 @@ public:
 
 		check(IsInGameThread()); // we pretty much need this to be initialized from the main thread...it uses GConfig, etc
 		check(GConfig && GConfig->IsReadyForUse());
+
+		bHasDefaultDebugOptions = FBackendDebugOptions::ParseFromTokens(DefaultDebugOptions, TEXT("All"), FCommandLine::Get());
+
 		RootCache = nullptr;
 		FParsedNodeMap ParsedNodes;
 
@@ -353,7 +356,14 @@ public:
 			{
 				if (!ParsedNode.Key->LegacyDebugOptions(DebugOptions))
 				{
-					UE_LOG(LogDerivedDataCache, Warning, TEXT("%s: Node is ignoring one or more -DDC-<NodeName>-Option debug options."), *NodeName);
+					UE_LOG(LogDerivedDataCache, Warning, TEXT("%s: Node is ignoring one or more -DDC-%s-<Option> debug options."), *NodeName, *NodeName);
+				}
+			}
+			else if (bHasDefaultDebugOptions)
+			{
+				if (!ParsedNode.Key->LegacyDebugOptions(DebugOptions))
+				{
+					UE_LOG(LogDerivedDataCache, Warning, TEXT("%s: Node is ignoring one or more -DDC-All-<Option> debug options."), *NodeName);
 				}
 			}
 		}
@@ -1331,8 +1341,12 @@ private:
 	/** List of directories used by the DDC */
 	TArray<FString> Directories;
 
+	FBackendDebugOptions DefaultDebugOptions;
+
 	int32 MaxKeyLength = 0;
 
+	bool bHasDefaultDebugOptions = false;
+
 	/** Whether a shared cache is in use */
 	bool bUsingSharedDDC;
 
@@ -1571,13 +1585,7 @@ bool FBackendDebugOptions::ParseFromTokens(FDerivedDataBackendInterface::FBacken
 	Prefix.Append(TEXTVIEW("-DDC-")).Append(InNodeName).Append(TEXTVIEW("-"));
 	if (UE::String::FindFirst(InInputTokens, Prefix, ESearchCase::IgnoreCase) == INDEX_NONE)
 	{
-		// Check if the input stream has any -DDC-All- options.
-		Prefix.Reset();
-		Prefix.Append(TEXTVIEW("-DDC-All-"));
-		if (UE::String::FindFirst(InInputTokens, Prefix, ESearchCase::IgnoreCase) == INDEX_NONE)
-		{
-			return false;
-		}
+		return false;
 	}
 
 	const int32 PrefixLen = Prefix.Len();
