DDC: Improved file system cache store maintenance

- Added deletion of empty directories.
- Applied the rate limit to directory scans.
- Reduced the default scan rate by a factor of 10.

diff --git a/Engine/Config/BaseEngine.ini b/Engine/Config/BaseEngine.ini
--- a/Engine/Config/BaseEngine.ini
+++ b/Engine/Config/BaseEngine.ini
@@ -2192,7 +2192,7 @@ Hierarchy=(Type=Hierarchical, Inner=Pak, Inner=EnginePak, Inner=Local, Inner=Sha
 ;		ValueOnly: Look only for compressed data in the 5.0+ directory layout.
 ;		LegacyOnly: Look only for uncompressed data in the pre-5.0 directory layout.
 ;		ValueWithLegacyFallback: Look in the 5.0+ location but fall back to the pre-5.0 location on miss.
-Local=(Type=FileSystem, PurgeTransient=true, UnusedFileAge=34, PromptIfMissing=true, Path=%ENGINEDIR%DerivedDataCache, EnvPathOverride=UE-LocalDataCachePath, EditorOverrideSetting=LocalDerivedDataCache)
+Local=(Type=FileSystem, PurgeTransient=true, UnusedFileAge=34, PromptIfMissing=true, Path=%ENGINEDIR%DerivedDataCache, EnvPathOverride=UE-LocalDataCachePath, EditorOverrideSetting=LocalDerivedDataCache, CommandLineOverride=LocalDataCachePath)
 ; Configure the shared DDC that is accessed after local. It's a filesystem DDC and the parameters are explained above
 Shared=(Type=FileSystem, UnusedFileAge=10, FoldersToClean=10, ConsiderSlowAt=70, Path=?EpicDDC, EnvPathOverride=UE-SharedDataCachePath, EditorOverrideSetting=SharedDerivedDataCache, CommandLineOverride=SharedDataCachePath)
 ; Configure a Project Pak node. This is a pre-generated DDC data file for the project that can be distributed to reduce runtime fetches/generation
@@ -2306,7 +2306,7 @@ Path=%GAMEDIR%DerivedDataCache/VT
 
 [DDCCleanup]
 TimeToWaitAfterInit=120
-MaxFileChecksPerSec=10000
+MaxFileChecksPerSec=1000
 
 [/Script/Engine.LocalPlayer]
 AspectRatioAxisConstraint=AspectRatio_MaintainXFOV
diff --git a/Engine/Source/Developer/DerivedDataCache/Private/FileSystemCacheStore.cpp b/Engine/Source/Developer/DerivedDataCache/Private/FileSystemCacheStore.cpp
--- a/Engine/Source/Developer/DerivedDataCache/Private/FileSystemCacheStore.cpp
+++ b/Engine/Source/Developer/DerivedDataCache/Private/FileSystemCacheStore.cpp
@@ -204,8 +204,8 @@ struct FFileSystemCacheStoreMaintainerParams
 {
 	/** Files older than this will be deleted. */
 	FTimespan MaxFileAge = FTimespan::FromDays(15.0);
-	/** Limits the number of files scanned in one second. */
-	uint32 MaxFileScanRate = MAX_uint32;
+	/** Limits the number of paths scanned in one second. */
+	uint32 MaxScanRate = MAX_uint32;
 	/** Limits the number of directories scanned in each cache bucket or content root. */
 	uint32 MaxDirectoryScanCount = MAX_uint32;
 	/** Minimum duration between the start of consecutive scans. Use MaxValue to scan only once. */
@@ -240,7 +240,11 @@ private:
 
 	void ResetRoots();
 
-	void ProcessFile(const TCHAR* Path, const FFileStatData& Stat);
+	void ProcessDirectory(const TCHAR* Path);
+	void ProcessFile(const TCHAR* Path, const FFileStatData& Stat, bool& bOutDeletedFile);
+	void ProcessWait();
+
+	void DeleteDirectory(const TCHAR* Path);
 
 private:
 	struct FRoot;
@@ -255,8 +259,8 @@ private:
 	bool bExit = false;
 	/** True when maintenance is expected to exit at the end of the scan. */
 	bool bExitAfterScan = false;
-	/** Ignore the file scan rate for one maintenance scan. */
-	bool bIgnoreFileScanRate = false;
+	/** Ignore the scan rate for one maintenance scan. */
+	bool bIgnoreScanRate = false;
 
 	uint32 ProcessCount = 0;
 	uint32 DeleteCount = 0;
@@ -338,14 +342,14 @@ FFileSystemCacheStoreMaintainer::~FFileSystemCacheStoreMaintainer()
 
 void FFileSystemCacheStoreMaintainer::BoostPriority()
 {
-	bIgnoreFileScanRate = true;
+	bIgnoreScanRate = true;
 	WaitEvent->Trigger();
 }
 
 void FFileSystemCacheStoreMaintainer::Tick()
 {
 	// Scan once and exit if the priority has been boosted.
-	if (bIgnoreFileScanRate)
+	if (bIgnoreScanRate)
 	{
 		bExitAfterScan = true;
 		Loop();
@@ -368,7 +372,7 @@ void FFileSystemCacheStoreMaintainer::Loop()
 		Scan();
 		bIdle = true;
 		IdleEvent.Trigger();
-		bIgnoreFileScanRate = false;
+		bIgnoreScanRate = false;
 		const FDateTime ScanEnd = FDateTime::Now();
 
 		UE_LOG(LogDerivedDataCache, Log,
@@ -488,11 +492,7 @@ void FFileSystemCacheStoreMaintainer::ScanHashRoot(const uint32 RootIndex)
 
 	TStringBuilder<256> Path;
 	Path.Appendf(TEXT("%s/%02x/%02x"), *Root.Path, IndexLevel0, IndexLevel1);
-	FileManager.IterateDirectoryStat(*Path, [this](const TCHAR* const Path, const FFileStatData& Stat) -> bool
-	{
-		ProcessFile(Path, Stat);
-		return !bExit;
-	});
+	ProcessDirectory(*Path);
 
 	bScanned = true;
 }
@@ -513,6 +513,10 @@ TStaticBitArray<256> FFileSystemCacheStoreMaintainer::ScanHashDirectory(FStringB
 		}
 		return !bExit;
 	});
+	if (Exists.FindFirstSetBit() == INDEX_NONE)
+	{
+		DeleteDirectory(*BasePath);
+	}
 	return Exists;
 }
 
@@ -528,6 +532,10 @@ TStaticBitArray<10> FFileSystemCacheStoreMaintainer::ScanLegacyDirectory(FString
 		}
 		return !bExit;
 	});
+	if (Exists.FindFirstSetBit() == INDEX_NONE && BasePath.Len() > CachePath.Len())
+	{
+		DeleteDirectory(*BasePath);
+	}
 	return Exists;
 }
 
@@ -592,11 +600,7 @@ void FFileSystemCacheStoreMaintainer::ScanLegacyRoot()
 
 	TStringBuilder<256> Path;
 	FPathViews::Append(Path, CachePath, IndexLevel0, IndexLevel1, IndexLevel2);
-	FileManager.IterateDirectoryStat(*Path, [this](const TCHAR* const Path, const FFileStatData& Stat) -> bool
-	{
-		ProcessFile(Path, Stat);
-		return !bExit;
-	});
+	ProcessDirectory(*Path);
 
 	bScanned = true;
 }
@@ -607,8 +611,30 @@ void FFileSystemCacheStoreMaintainer::ResetRoots()
 	LegacyRoot.Reset();
 }
 
-void FFileSystemCacheStoreMaintainer::ProcessFile(const TCHAR* const Path, const FFileStatData& Stat)
+void FFileSystemCacheStoreMaintainer::ProcessDirectory(const TCHAR* const Path)
 {
+	bool bTryDelete = true;
+
+	FileManager.IterateDirectoryStat(Path, [this, &bTryDelete](const TCHAR* const Path, const FFileStatData& Stat) -> bool
+	{
+		bool bDeletedFile = false;
+		ProcessFile(Path, Stat, bDeletedFile);
+		bTryDelete &= bDeletedFile;
+		return !bExit;
+	});
+
+	if (bTryDelete)
+	{
+		DeleteDirectory(Path);
+	}
+
+	ProcessWait();
+}
+
+void FFileSystemCacheStoreMaintainer::ProcessFile(const TCHAR* const Path, const FFileStatData& Stat, bool& bOutDeletedFile)
+{
+	bOutDeletedFile = false;
+
 	if (Stat.bIsDirectory)
 	{
 		return;
@@ -621,6 +647,7 @@ void FFileSystemCacheStoreMaintainer::ProcessFile(const TCHAR* const Path, const
 		DeleteSize += Stat.FileSize > 0 ? uint64(Stat.FileSize) : 0;
 		if (FileManager.Delete(Path, /*bRequireExists*/ false, /*bEvenReadOnly*/ false, /*bQuiet*/ true))
 		{
+			bOutDeletedFile = true;
 			UE_LOG(LogDerivedDataCache, VeryVerbose,
 				TEXT("%s: Maintenance deleted file %s that was last modified at %s."),
 				*CachePath, Path, *Stat.ModificationTime.ToIso8601());
@@ -633,7 +660,12 @@ void FFileSystemCacheStoreMaintainer::ProcessFile(const TCHAR* const Path, const
 		}
 	}
 
-	if (!bExit && !bIgnoreFileScanRate && Params.MaxFileScanRate && ++ProcessCount % Params.MaxFileScanRate == 0)
+	ProcessWait();
+}
+
+void FFileSystemCacheStoreMaintainer::ProcessWait()
+{
+	if (!bExit && !bIgnoreScanRate && Params.MaxScanRate && ++ProcessCount % Params.MaxScanRate == 0)
 	{
 		const double BatchEndTime = FPlatformTime::Seconds();
 		if (const double BatchWaitTime = 1.0 - (BatchEndTime - BatchStartTime); BatchWaitTime > 0.0)
@@ -648,6 +680,14 @@ void FFileSystemCacheStoreMaintainer::ProcessFile(const TCHAR* const Path, const
 	}
 }
 
+void FFileSystemCacheStoreMaintainer::DeleteDirectory(const TCHAR* Path)
+{
+	if (FileManager.DeleteDirectory(Path))
+	{
+		UE_LOG(LogDerivedDataCache, VeryVerbose, TEXT("%s: Maintenance deleted empty directory %s."), *CachePath, Path);
+	}
+}
+
 ///////////////////////////////////////////////////////////////////////////////////////////////////////////////////////
 
 class FAccessLogWriter
@@ -821,7 +861,8 @@ private:
 	[[nodiscard]] bool LoadFileWithHash(FStringBuilderBase& Path, FStringView DebugName, TFunctionRef<void (FArchive&)> ReadFunction) const;
 	[[nodiscard]] bool SaveFile(FStringBuilderBase& Path, FStringView DebugName, TFunctionRef<void (FArchive&)> WriteFunction, bool bReplaceExisting = false) const;
 	[[nodiscard]] bool LoadFile(FStringBuilderBase& Path, FStringView DebugName, TFunctionRef<void (FArchive&)> ReadFunction) const;
-	[[nodiscard]] TUniquePtr<FArchive> OpenFile(FStringBuilderBase& Path, FStringView DebugName) const;
+	[[nodiscard]] TUniquePtr<FArchive> OpenFileWrite(FStringBuilderBase& Path, FStringView DebugName) const;
+	[[nodiscard]] TUniquePtr<FArchive> OpenFileRead(FStringBuilderBase& Path, FStringView DebugName) const;
 
 	[[nodiscard]] bool FileExists(FStringBuilderBase& Path) const;
 
@@ -1020,12 +1061,12 @@ FFileSystemCacheStore::FFileSystemCacheStore(
 				MaintainerParams.MaxFileAge = FTimespan::FromDays(DaysToDeleteUnusedFiles);
 				if (bDeleteUnused)
 				{
-					if (!FParse::Value(InParams, TEXT("MaxFileChecksPerSec="), MaintainerParams.MaxFileScanRate))
+					if (!FParse::Value(InParams, TEXT("MaxFileChecksPerSec="), MaintainerParams.MaxScanRate))
 					{
 						int32 MaxFileScanRate;
 						if (GConfig->GetInt(TEXT("DDCCleanup"), TEXT("MaxFileChecksPerSec"), MaxFileScanRate, GEngineIni))
 						{
-							MaintainerParams.MaxFileScanRate = uint32(MaxFileScanRate);
+							MaintainerParams.MaxScanRate = uint32(MaxFileScanRate);
 						}
 					}
 					FParse::Value(InParams, TEXT("FoldersToClean="), MaintainerParams.MaxDirectoryScanCount);
@@ -2424,7 +2465,7 @@ void FFileSystemCacheStore::GetCacheContent(
 	}
 	else
 	{
-		OutArchive = OpenFile(Path, Name);
+		OutArchive = OpenFileRead(Path, Name);
 		if (OutArchive)
 		{
 			Reader.SetSource(*OutArchive);
@@ -2522,7 +2563,7 @@ bool FFileSystemCacheStore::SaveFile(
 	int64 WriteSize = 0;
 	bool bError = false;
 
-	if (TUniquePtr<FArchive> Ar{IFileManager::Get().CreateFileWriter(*TempPath, FILEWRITE_Silent)})
+	if (TUniquePtr<FArchive> Ar = OpenFileWrite(TempPath, DebugName))
 	{
 		WriteFunction(*Ar);
 		WriteSize = Ar->Tell();
@@ -2575,7 +2616,7 @@ bool FFileSystemCacheStore::LoadFile(
 	int64 ReadSize = 0;
 	bool bError = false;
 
-	if (TUniquePtr<FArchive> Ar = OpenFile(Path, DebugName))
+	if (TUniquePtr<FArchive> Ar = OpenFileRead(Path, DebugName))
 	{
 		ReadFunction(*Ar);
 		ReadSize = Ar->Tell();
@@ -2612,7 +2653,21 @@ bool FFileSystemCacheStore::LoadFile(
 	return !bError && ReadSize > 0;
 }
 
-TUniquePtr<FArchive> FFileSystemCacheStore::OpenFile(FStringBuilderBase& Path, const FStringView DebugName) const
+TUniquePtr<FArchive> FFileSystemCacheStore::OpenFileWrite(FStringBuilderBase& Path, const FStringView DebugName) const
+{
+	// Retry to handle a race where the directory is deleted while the file is being created.
+	constexpr int32 MaxAttemptCount = 3;
+	for (int32 AttemptCount = 0; AttemptCount < MaxAttemptCount; ++AttemptCount)
+	{
+		if (TUniquePtr<FArchive> Ar{IFileManager::Get().CreateFileWriter(*Path, FILEWRITE_Silent)})
+		{
+			return Ar;
+		}
+	}
+	return nullptr;
+}
+
+TUniquePtr<FArchive> FFileSystemCacheStore::OpenFileRead(FStringBuilderBase& Path, const FStringView DebugName) const
 {
 	// Check for existence before reading because it may update the modification time and avoid the
 	// file being deleted by a cache cleanup thread or process.
