Fixed bug in handling of -DDC-All-* commandline arguments.

#rnx

diff --git a/Engine/Source/Developer/DerivedDataCache/Private/DerivedDataBackends.cpp b/Engine/Source/Developer/DerivedDataCache/Private/DerivedDataBackends.cpp
--- a/Engine/Source/Developer/DerivedDataCache/Private/DerivedDataBackends.cpp
+++ b/Engine/Source/Developer/DerivedDataCache/Private/DerivedDataBackends.cpp
@@ -361,7 +361,7 @@ public:
 			}
 			else if (bHasDefaultDebugOptions)
 			{
-				if (!ParsedNode.Key->LegacyDebugOptions(DebugOptions))
+				if (!ParsedNode.Key->LegacyDebugOptions(DefaultDebugOptions))
 				{
 					UE_LOG(LogDerivedDataCache, Warning, TEXT("%s: Node is ignoring one or more -DDC-All-<Option> debug options."), *NodeName);
 				}
