Avoid querying Horde Storage during GetChunks calls for "Null hash" compressed blob value when you've experienced a miss for the record it belongs to.

diff --git a/Engine/Source/Developer/DerivedDataCache/Private/HttpCacheStore.cpp b/Engine/Source/Developer/DerivedDataCache/Private/HttpCacheStore.cpp
--- a/Engine/Source/Developer/DerivedDataCache/Private/HttpCacheStore.cpp
+++ b/Engine/Source/Developer/DerivedDataCache/Private/HttpCacheStore.cpp
@@ -4845,41 +4845,49 @@ void FHttpCacheStore::GetChunks(
 					bHasValue = bSucceeded;
 				}
 
-				if (IsValueDataReady(Value, Request.Policy))
+				if (bHasValue)
 				{
-					ValueReader.SetSource(Value.GetData());
-				}
-				else
-				{
-					auto IdGetter = [](const FValue& Value)
-					{
-						return FString(TEXT("Default"));
-					};
-
-					FRequestOwner BlockingOwner(EPriority::Blocking);
-					bool bSucceeded = false;
-					FCompressedBuffer NewBuffer;
-					FGetRecordOp::GetDataBatch(*this, BlockingOwner, Request.Name, Request.Key, ::MakeArrayView({ Value }), IdGetter, [&bSucceeded, &NewBuffer](FGetRecordOp::FGetCachedDataBatchResponse&& Response)
-						{
-							if (Response.Status == EStatus::Ok)
-							{
-								bSucceeded = true;
-								NewBuffer = MoveTemp(Response.DataBuffer);
-							}
-						});
-					BlockingOwner.Wait();
-
-					if (bSucceeded)
+					if (IsValueDataReady(Value, Request.Policy))
 					{
-						ValueBuffer = MoveTemp(NewBuffer);
-						ValueReader.SetSource(ValueBuffer);
+						ValueReader.SetSource(Value.GetData());
 					}
 					else
 					{
-						ValueBuffer.Reset();
-						ValueReader.ResetSource();
+						auto IdGetter = [](const FValue& Value)
+						{
+							return FString(TEXT("Default"));
+						};
+
+						FRequestOwner BlockingOwner(EPriority::Blocking);
+						bool bSucceeded = false;
+						FCompressedBuffer NewBuffer;
+						FGetRecordOp::GetDataBatch(*this, BlockingOwner, Request.Name, Request.Key, ::MakeArrayView({ Value }), IdGetter, [&bSucceeded, &NewBuffer](FGetRecordOp::FGetCachedDataBatchResponse&& Response)
+							{
+								if (Response.Status == EStatus::Ok)
+								{
+									bSucceeded = true;
+									NewBuffer = MoveTemp(Response.DataBuffer);
+								}
+							});
+						BlockingOwner.Wait();
+
+						if (bSucceeded)
+						{
+							ValueBuffer = MoveTemp(NewBuffer);
+							ValueReader.SetSource(ValueBuffer);
+						}
+						else
+						{
+							ValueBuffer.Reset();
+							ValueReader.ResetSource();
+						}
 					}
 				}
+				else
+				{
+					ValueBuffer.Reset();
+					ValueReader.ResetSource();
+				}
 			}
 		}
 		if (bHasValue)
