Limit concurrent connections when using the Curl multi API for Jupiter DDC operations.

diff --git a/Engine/Source/Developer/DerivedDataCache/Private/HttpCacheStore.cpp b/Engine/Source/Developer/DerivedDataCache/Private/HttpCacheStore.cpp
--- a/Engine/Source/Developer/DerivedDataCache/Private/HttpCacheStore.cpp
+++ b/Engine/Source/Developer/DerivedDataCache/Private/HttpCacheStore.cpp
@@ -76,6 +76,11 @@
 #define UE_HTTPDDC_GET_REQUEST_POOL_SIZE 48
 #define UE_HTTPDDC_PUT_REQUEST_POOL_SIZE 16
 #define UE_HTTPDDC_NONBLOCKING_REQUEST_POOL_SIZE 128
+// After 32 concurrent connections, we'd rather that requests on a multi handle wait for another request
+// to free up rather than attempting to start a new connection and then have the connection cache
+// retreat and close a connection later.  This will be mitigated with HTTP/2 support in the future which
+// will allow for multiplexing of multiple requests on the same connection (not just connection re-use).
+#define UE_HTTPDDC_CONNECTION_LIMIT 32L
 #define UE_HTTPDDC_MAX_FAILED_LOGIN_ATTEMPTS 16
 #define UE_HTTPDDC_MAX_ATTEMPTS 4
 #define UE_HTTPDDC_MAX_BUFFER_RESERVE 104857600u
@@ -139,6 +144,10 @@ struct FHttpSharedData
 		curl_share_setopt(CurlShare, CURLSHOPT_SHARE, CURL_LOCK_DATA_SSL_SESSION);
 		CurlMulti = curl_multi_init();
 		curl_multi_setopt(CurlMulti, CURLMOPT_PIPELINING, CURLPIPE_MULTIPLEX);
+		// Absolute number of concurrent connections permitted
+		curl_multi_setopt(CurlMulti, CURLMOPT_MAX_TOTAL_CONNECTIONS, UE_HTTPDDC_CONNECTION_LIMIT);
+		// Number of concurrent connections held in the connection cache
+		curl_multi_setopt(CurlMulti, CURLMOPT_MAXCONNECTS, UE_HTTPDDC_CONNECTION_LIMIT);
 	}
 
 	~FHttpSharedData()
