DDC: Optimized the file system cache for SMB

- Reading files from the Game Thread now launches a task to check and update the modification time, which approximately halves the latency of these reads.
- Writing files now only asks to delete the temporary file upon failure, as it added unnecessary latency before by trying to delete it after it was renamed.

diff --git a/Engine/Source/Developer/DerivedDataCache/Private/FileSystemCacheStore.cpp b/Engine/Source/Developer/DerivedDataCache/Private/FileSystemCacheStore.cpp
--- a/Engine/Source/Developer/DerivedDataCache/Private/FileSystemCacheStore.cpp
+++ b/Engine/Source/Developer/DerivedDataCache/Private/FileSystemCacheStore.cpp
@@ -16,6 +16,7 @@
 #include "DerivedDataCacheRecord.h"
 #include "DerivedDataCacheUsageStats.h"
 #include "DerivedDataChunk.h"
+#include "DerivedDataRequestOwner.h"
 #include "DerivedDataValue.h"
 #include "Experimental/Async/LazyEvent.h"
 #include "Features/IModularFeatures.h"
@@ -2555,45 +2556,35 @@ bool FFileSystemCacheStore::SaveFile(
 	TStringBuilder<256> TempPath;
 	TempPath << FPathViews::GetPath(Path) << TEXT("/Temp.") << FGuid::NewGuid();
 
-	ON_SCOPE_EXIT
-	{
-		IFileManager::Get().Delete(*TempPath, /*bRequireExists*/ false, /*bEvenReadOnly*/ false, /*bQuiet*/ true);
-	};
-
-	int64 WriteSize = 0;
-	bool bError = false;
-
-	if (TUniquePtr<FArchive> Ar = OpenFileWrite(TempPath, DebugName))
-	{
-		WriteFunction(*Ar);
-		WriteSize = Ar->Tell();
-		bError = !Ar->Close();
-	}
-
-	if (bError || WriteSize == 0)
+	TUniquePtr<FArchive> Ar = OpenFileWrite(TempPath, DebugName);
+	if (!Ar)
 	{
 		UE_LOG(LogDerivedDataCache, Warning,
-			TEXT("%s: Failed to write to temp file %s when saving %s from '%.*s'. Error 0x%08x"),
+			TEXT("%s: Failed to open temp file %s for writing when saving %s from '%.*s'. Error 0x%08x."),
 			*CachePath, *TempPath, *Path, DebugName.Len(), DebugName.GetData(), FPlatformMisc::GetLastError());
 		return false;
 	}
 
-	if (IFileManager::Get().FileSize(*TempPath) != WriteSize)
+	WriteFunction(*Ar);
+	const int64 WriteSize = Ar->Tell();
+
+	if (!Ar->Close() || WriteSize == 0 || WriteSize != IFileManager::Get().FileSize(*TempPath))
 	{
 		UE_LOG(LogDerivedDataCache, Warning,
-			TEXT("%s: Failed to write to temp file %s when saving %s from '%.*s'. ")
-			TEXT("File is %" INT64_FMT " bytes when %" INT64_FMT " bytes are expected."),
-			*CachePath, *TempPath, *Path, DebugName.Len(), DebugName.GetData(),
+			TEXT("%s: Failed to write to temp file %s when saving %s from '%.*s'. Error 0x%08x. "
+			"File is %" INT64_FMT " bytes when %" INT64_FMT " bytes are expected."),
+			*CachePath, *TempPath, *Path, DebugName.Len(), DebugName.GetData(), FPlatformMisc::GetLastError(),
 			IFileManager::Get().FileSize(*TempPath), WriteSize);
+		IFileManager::Get().Delete(*TempPath, /*bRequireExists*/ false, /*bEvenReadOnly*/ false, /*bQuiet*/ true);
 		return false;
 	}
 
-	if (!IFileManager::Get().Move(*Path, *TempPath,
-		bReplaceExisting, /*bEvenIfReadOnly*/ false, /*bAttributes*/ false, /*bDoNotRetryOrError*/ true))
+	if (!IFileManager::Get().Move(*Path, *TempPath, bReplaceExisting, /*bEvenIfReadOnly*/ false, /*bAttributes*/ false, /*bDoNotRetryOrError*/ true))
 	{
 		UE_LOG(LogDerivedDataCache, Log,
 			TEXT("%s: Move collision when writing file %s from '%.*s'."),
 			*CachePath, *Path, DebugName.Len(), DebugName.GetData());
+		IFileManager::Get().Delete(*TempPath, /*bRequireExists*/ false, /*bEvenReadOnly*/ false, /*bQuiet*/ true);
 	}
 
 	const double WriteDuration = FPlatformTime::Seconds() - StartTime;
@@ -2613,41 +2604,43 @@ bool FFileSystemCacheStore::LoadFile(
 	check(IsUsable());
 	const double StartTime = FPlatformTime::Seconds();
 
-	int64 ReadSize = 0;
-	bool bError = false;
+	TUniquePtr<FArchive> Ar = OpenFileRead(Path, DebugName);
+	if (!Ar)
+	{
+		return false;
+	}
 
-	if (TUniquePtr<FArchive> Ar = OpenFileRead(Path, DebugName))
+	ReadFunction(*Ar);
+	const int64 ReadSize = Ar->Tell();
+	const bool bError = !Ar->Close();
+
+	if (bError)
 	{
-		ReadFunction(*Ar);
-		ReadSize = Ar->Tell();
-		bError = !Ar->Close();
+		UE_LOG(LogDerivedDataCache, Display,
+			TEXT("%s: Failed to load file %s from '%.*s'."),
+			*CachePath, *Path, DebugName.Len(), DebugName.GetData());
 
-		if (bError)
+		if (!bReadOnly)
 		{
-			UE_LOG(LogDerivedDataCache, Display,
-				TEXT("%s: Failed to load file %s from '%.*s'."),
-				*CachePath, *Path, DebugName.Len(), DebugName.GetData());
+			IFileManager::Get().Delete(*Path, /*bRequireExists*/ false, /*bEvenReadOnly*/ false, /*bQuiet*/ true);
 		}
 	}
-
-	const double ReadDuration = FPlatformTime::Seconds() - StartTime;
-	const double ReadSpeed = ReadDuration > 0.001 ? (ReadSize / ReadDuration) / (1024.0 * 1024.0) : 0.0;
-
-	if (!GIsBuildMachine && ReadDuration > 5.0)
+	else
 	{
-		// Slower than 0.5 MiB/s?
-		UE_CLOG(ReadSpeed < 0.5, LogDerivedDataCache, Warning,
-			TEXT("%s: Loading %s from '%.*s' is very slow (%.2f MiB/s); consider disabling this cache backend"),
-			*CachePath, *Path, DebugName.Len(), DebugName.GetData(), ReadSpeed);
-	}
+		const double ReadDuration = FPlatformTime::Seconds() - StartTime;
+		const double ReadSpeed = ReadDuration > 0.001 ? (ReadSize / ReadDuration) / (1024.0 * 1024.0) : 0.0;
 
-	UE_LOG(LogDerivedDataCache, VeryVerbose,
-		TEXT("%s: Loaded %s from '%.*s' (%" INT64_FMT " bytes, %.02f secs, %.2f MiB/s)"),
-		*CachePath, *Path, DebugName.Len(), DebugName.GetData(), ReadSize, ReadDuration, ReadSpeed);
+		UE_LOG(LogDerivedDataCache, VeryVerbose,
+			TEXT("%s: Loaded %s from '%.*s' (%" INT64_FMT " bytes, %.02f secs, %.2f MiB/s)"),
+			*CachePath, *Path, DebugName.Len(), DebugName.GetData(), ReadSize, ReadDuration, ReadSpeed);
 
-	if (bError && !bReadOnly)
-	{
-		IFileManager::Get().Delete(*Path, /*bRequireExists*/ false, /*bEvenReadOnly*/ false, /*bQuiet*/ true);
+		if (!GIsBuildMachine && ReadDuration > 5.0)
+		{
+			// Slower than 0.5 MiB/s?
+			UE_CLOG(ReadSpeed < 0.5, LogDerivedDataCache, Warning,
+				TEXT("%s: Loading %s from '%.*s' is very slow (%.2f MiB/s); consider disabling this cache store."),
+				*CachePath, *Path, DebugName.Len(), DebugName.GetData(), ReadSpeed);
+		}
 	}
 
 	return !bError && ReadSize > 0;
@@ -2669,11 +2662,23 @@ TUniquePtr<FArchive> FFileSystemCacheStore::OpenFileWrite(FStringBuilderBase& Pa
 
 TUniquePtr<FArchive> FFileSystemCacheStore::OpenFileRead(FStringBuilderBase& Path, const FStringView DebugName) const
 {
-	// Check for existence before reading because it may update the modification time and avoid the
-	// file being deleted by a cache cleanup thread or process.
-	if (!FileExists(Path))
+	// Checking for existence may update the modification time to avoid the file being evicted from the cache.
+	// Reduce Game Thread overhead by executing the update on a worker thread if the path implies higher latency.
+	if (IsInGameThread() && FStringView(CachePath).StartsWith(TEXTVIEW("//"), ESearchCase::CaseSensitive))
 	{
-		return nullptr;
+		FRequestOwner AsyncOwner(EPriority::Normal);
+		Private::LaunchTaskInCacheThreadPool(AsyncOwner, [this, Path = WriteToString<256>(Path)]() mutable
+		{
+			(void)FileExists(Path);
+		});
+		AsyncOwner.KeepAlive();
+	}
+	else
+	{
+		if (!FileExists(Path))
+		{
+			return nullptr;
+		}
 	}
 
 	return TUniquePtr<FArchive>(IFileManager::Get().CreateFileReader(*Path, FILEREAD_Silent));
