DDC: Added propagation of cache records that contain metadata with no values

diff --git a/Engine/Source/Developer/DerivedDataCache/Private/DerivedDataCacheStoreHierarchy.cpp b/Engine/Source/Developer/DerivedDataCache/Private/DerivedDataCacheStoreHierarchy.cpp
--- a/Engine/Source/Developer/DerivedDataCache/Private/DerivedDataCacheStoreHierarchy.cpp
+++ b/Engine/Source/Developer/DerivedDataCache/Private/DerivedDataCacheStoreHierarchy.cpp
@@ -800,7 +800,7 @@ auto FCacheStoreHierarchy::FCacheRecordBatchParams::Get() -> GetFunctionType
 template <>
 bool FCacheStoreHierarchy::FCacheRecordBatchParams::HasResponseData(const FCacheGetResponse& Response)
 {
-	return Algo::AnyOf(Response.Record.GetValues(), &FValue::HasData);
+	return Algo::AnyOf(Response.Record.GetValues(), &FValue::HasData) || Response.Record.GetMeta();
 }
 
 template <>
