DDC: Fixed HttpCacheStore to shut down its thread

FHttpSharedData is no longer constructed as a function-local static because that makes shutdown more difficult to handle correctly. The cache store now owns the shared data and destroys it when it shuts down.

Includes greater encapsulation of the shared data and migration from FRunnable/FRunnableThread to FThread.

#rnx

diff --git a/Engine/Source/Developer/DerivedDataCache/Private/HttpCacheStore.cpp b/Engine/Source/Developer/DerivedDataCache/Private/HttpCacheStore.cpp
--- a/Engine/Source/Developer/DerivedDataCache/Private/HttpCacheStore.cpp
+++ b/Engine/Source/Developer/DerivedDataCache/Private/HttpCacheStore.cpp
@@ -34,8 +34,7 @@
 #include "GenericPlatform/GenericPlatformFile.h"
 #include "HAL/IConsoleManager.h"
 #include "HAL/PlatformFileManager.h"
-#include "HAL/Runnable.h"
-#include "HAL/RunnableThread.h"
+#include "HAL/Thread.h"
 #include "IO/IoHash.h"
 #include "Memory/SharedBuffer.h"
 #include "Misc/FileHelper.h"
@@ -109,6 +108,9 @@ static CURLcode sslctx_function(CURL * curl, void * sslctx, void * parm);
 typedef TSharedPtr<class IHttpRequest> FHttpRequestPtr;
 typedef TSharedPtr<class IHttpResponse, ESPMode::ThreadSafe> FHttpResponsePtr;
 
+class FHttpCacheStoreRunnable;
+struct FRequestPool;
+
 /**
  * Encapsulation for access token shared by all requests.
  */
@@ -125,12 +127,12 @@ private:
 	uint32		Serial;
 };
 
-struct FHttpSharedData
+class FHttpSharedData
 {
+public:
 	FHttpSharedData()
 	: PendingRequestEvent(EEventMode::AutoReset)
 	{
-		FMemory::Memset(WriteLocked, 0, sizeof(WriteLocked));
 		CurlShare = curl_share_init();
 		curl_share_setopt(CurlShare, CURLSHOPT_USERDATA, this);
 		curl_share_setopt(CurlShare, CURLSHOPT_LOCKFUNC, LockFn);
@@ -143,22 +145,31 @@ struct FHttpSharedData
 
 	~FHttpSharedData()
 	{
+		bAsyncThreadShutdownRequested.store(true, std::memory_order_relaxed);
+		if (AsyncServiceThread.IsJoinable())
+		{
+			AsyncServiceThread.Join();
+		}
+
 		curl_multi_cleanup(CurlMulti);
 		curl_share_cleanup(CurlShare);
 	}
 
+	void AddRequest(CURL* Curl);
+
+	CURLSH* GetCurlShare() const { return CurlShare; }
+
+private:
 	CURLSH* CurlShare;
 	CURLM* CurlMulti;
 	TDepletableMpscQueue<CURL*> PendingRequestAdditions;
 	FLazyEvent PendingRequestEvent;
-	FRunnableThread* AsyncServiceThread = nullptr;
-	TUniquePtr<class FHttpCacheStoreRunnable> AsyncRunnable;
+	FThread AsyncServiceThread;
 	std::atomic<bool> bAsyncThreadStarting = false;
-	static inline std::atomic<bool> bAsyncThreadShutdownRequested = false;
+	std::atomic<bool> bAsyncThreadShutdownRequested = false;
 
-private:
 	FRWLock Locks[CURL_LOCK_DATA_LAST];
-	bool WriteLocked[CURL_LOCK_DATA_LAST];
+	bool WriteLocked[CURL_LOCK_DATA_LAST]{};
 
 	static void LockFn(CURL* Handle, curl_lock_data Data, curl_lock_access Access, void* User)
 	{
@@ -173,6 +184,7 @@ private:
 			SharedData->WriteLocked[Data] = true;
 		}
 	}
+
 	static void UnlockFn(CURL* Handle, curl_lock_data Data, void* User)
 	{
 		FHttpSharedData* SharedData = (FHttpSharedData*)User;
@@ -186,6 +198,8 @@ private:
 			SharedData->Locks[Data].WriteUnlock();
 		}
 	}
+
+	void ProcessAsyncRequests();
 };
 
 template <typename T>
@@ -202,19 +216,6 @@ private:
 	T Function;
 };
 
-class FHttpCacheStoreRunnable final : public FRunnable
-{
-public:
-	explicit FHttpCacheStoreRunnable(FHttpSharedData& InSharedData) : SharedData(InSharedData)
-	{
-	}
-
-	virtual uint32 Run() override final;
-
-protected:
-	FHttpSharedData& SharedData;
-};
-
 /**
  * Minimal HTTP request type wrapping CURL without the need for managers. This request
  * is written to allow reuse of request objects, in order to allow connections to be reused.
@@ -259,8 +260,9 @@ public:
 
 	using FOnHttpRequestComplete = TUniqueFunction<ECompletionBehavior(Result HttpResult, FHttpRequest* Request)>;
 
-	FHttpRequest(const TCHAR* InDomain, const TCHAR* InEffectiveDomain, FHttpAccessToken* InAuthorizationToken, bool bInLogErrors)
-		: AsyncData(nullptr)
+	FHttpRequest(const TCHAR* InDomain, const TCHAR* InEffectiveDomain, FHttpAccessToken* InAuthorizationToken, FHttpSharedData* InSharedData, bool bInLogErrors)
+		: SharedData(InSharedData)
+		, AsyncData(nullptr)
 		, bLogErrors(bInLogErrors)
 		, Domain(InDomain)
 		, EffectiveDomain(InEffectiveDomain)
@@ -293,9 +295,6 @@ public:
 		Attempts = 0;
 		CurlResult = CURL_LAST;
 
-		static FHttpSharedData StaticSharedData;
-		SharedData = &StaticSharedData;
-
 		curl_easy_reset(Curl);
 		check(!AsyncData);
 
@@ -306,7 +305,7 @@ public:
 		curl_easy_setopt(Curl, CURLOPT_FOLLOWLOCATION, 1L);
 		curl_easy_setopt(Curl, CURLOPT_NOSIGNAL, 1L);
 		curl_easy_setopt(Curl, CURLOPT_DNS_CACHE_TIMEOUT, 300L); // Don't re-resolve every minute
-		curl_easy_setopt(Curl, CURLOPT_SHARE, SharedData->CurlShare);
+		curl_easy_setopt(Curl, CURLOPT_SHARE, SharedData ? SharedData->GetCurlShare() : nullptr);
 		// SSL options
 		curl_easy_setopt(Curl, CURLOPT_USE_SSL, CURLUSESSL_ALL);
 		curl_easy_setopt(Curl, CURLOPT_SSL_VERIFYPEER, 1);
@@ -452,7 +451,7 @@ public:
 	}
 	
 	template<RequestVerb V>
-	void EnqueueAsyncUpload(IRequestOwner& Owner, struct FRequestPool* Pool, const TCHAR* Uri, FSharedBuffer Buffer, FOnHttpRequestComplete&& OnComplete, TConstArrayView<long> ExpectedErrorCodes = {})
+	void EnqueueAsyncUpload(IRequestOwner& Owner, FRequestPool* Pool, const TCHAR* Uri, FSharedBuffer Buffer, FOnHttpRequestComplete&& OnComplete, TConstArrayView<long> ExpectedErrorCodes = {})
 	{
 		static_assert(V == Put || V == PutCompactBinary || V == PutCompressedBlob || V == Post || V == PostCompactBinary || V == PostJson, "Upload should use either Put or Post verbs.");
 		
@@ -521,7 +520,7 @@ public:
 		return PerformBlocking(Uri, Get, 0u, ExpectedErrorCodes);
 	}
 
-	void EnqueueAsyncDownload(IRequestOwner& Owner, struct FRequestPool* Pool, const TCHAR* Uri, FOnHttpRequestComplete&& OnComplete, TConstArrayView<long> ExpectedErrorCodes = {400})
+	void EnqueueAsyncDownload(IRequestOwner& Owner, FRequestPool* Pool, const TCHAR* Uri, FOnHttpRequestComplete&& OnComplete, TConstArrayView<long> ExpectedErrorCodes = {400})
 	{
 		curl_easy_setopt(Curl, CURLOPT_HTTPGET, 1L);
 
@@ -892,8 +891,6 @@ private:
 						*Response
 					);
 				}
-
-
 			}
 		}
 		else if(bLogErrors)
@@ -1076,7 +1073,7 @@ bool ShouldAbortForShutdown();
  */
 struct FRequestPool
 {
-	FRequestPool(const TCHAR* InServiceUrl, const TCHAR* InEffectiveServiceUrl, FHttpAccessToken* InAuthorizationToken, uint32 PoolSize, uint32 InOverflowLimit = 0)
+	FRequestPool(const TCHAR* InServiceUrl, const TCHAR* InEffectiveServiceUrl, FHttpAccessToken* InAuthorizationToken, FHttpSharedData* InSharedData, uint32 PoolSize, uint32 InOverflowLimit = 0)
 	: ActiveOverflowRequests(0)
 	, OverflowLimit(InOverflowLimit)
 	{
@@ -1085,10 +1082,10 @@ struct FRequestPool
 		for (uint8 i = 0; i < Pool.Num(); ++i)
 		{
 			Pool[i].Usage = 0u;
-			Pool[i].Request = new(&Requests[i]) FHttpRequest(InServiceUrl, InEffectiveServiceUrl, InAuthorizationToken, true);
+			Pool[i].Request = new(&Requests[i]) FHttpRequest(InServiceUrl, InEffectiveServiceUrl, InAuthorizationToken, InSharedData, true);
 		}
 
-		InitData = MakeUnique<FInitData>(InServiceUrl, InEffectiveServiceUrl, InAuthorizationToken);
+		InitData = MakeUnique<FInitData>(InServiceUrl, InEffectiveServiceUrl, InAuthorizationToken, InSharedData);
 	}
 
 	~FRequestPool()
@@ -1128,7 +1125,7 @@ struct FRequestPool
 			{
 				// Create an overflow request (outside of the pre-allocated range of requests)
 				ActiveOverflowRequests.fetch_add(1, std::memory_order_relaxed);
-				return new FHttpRequest(*InitData->ServiceUrl, *InitData->EffectiveServiceUrl, InitData->AccessToken, true);
+				return new FHttpRequest(*InitData->ServiceUrl, *InitData->EffectiveServiceUrl, InitData->AccessToken, InitData->SharedData, true);
 			}
 		}
 		return nullptr;
@@ -1288,11 +1285,13 @@ private:
 		FString ServiceUrl;
 		FString EffectiveServiceUrl;
 		FHttpAccessToken* AccessToken;
+		FHttpSharedData* SharedData;
 
-		FInitData(const TCHAR* InServiceUrl, const TCHAR* InEffectiveServiceUrl, FHttpAccessToken* InAccessToken)
+		FInitData(const TCHAR* InServiceUrl, const TCHAR* InEffectiveServiceUrl, FHttpAccessToken* InAccessToken, FHttpSharedData* InSharedData)
 		: ServiceUrl(InServiceUrl)
 		, EffectiveServiceUrl(InEffectiveServiceUrl)
 		, AccessToken(InAccessToken)
+		, SharedData(InSharedData)
 		{
 		}
 	};
@@ -2446,10 +2445,11 @@ private:
 	FCriticalSection AccessCs;
 	FDerivedDataCacheUsageStats UsageStats;
 	FBackendDebugOptions DebugOptions;
-	TUniquePtr<struct FRequestPool> GetRequestPools[2];
-	TUniquePtr<struct FRequestPool> PutRequestPools[2];
-	TUniquePtr<struct FRequestPool> NonBlockingRequestPools;
-	TUniquePtr<struct FHttpAccessToken> Access;
+	TUniquePtr<FHttpSharedData> SharedData;
+	TUniquePtr<FRequestPool> GetRequestPools[2];
+	TUniquePtr<FRequestPool> PutRequestPools[2];
+	TUniquePtr<FRequestPool> NonBlockingRequestPools;
+	TUniquePtr<FHttpAccessToken> Access;
 	bool bIsUsable;
 	bool bReadOnly;
 	uint32 FailedLoginAttempts;
@@ -3458,6 +3458,7 @@ FHttpCacheStore::FHttpCacheStore(
 #if WITH_DATAREQUEST_HELPER
 	FDataRequestHelper::StaticInitialize();
 #endif
+	SharedData = MakeUnique<FHttpSharedData>();
 	if (IsServiceReady() && AcquireAccessToken())
 	{
 		FString OriginalDomainPrefix;
@@ -3506,12 +3507,12 @@ FHttpCacheStore::FHttpCacheStore(
 			EffectiveDomain = Domain;
 		}
 
-		GetRequestPools[0] = MakeUnique<FRequestPool>(*Domain, *EffectiveDomain, Access.Get(), UE_HTTPDDC_GET_REQUEST_POOL_SIZE);
-		GetRequestPools[1] = MakeUnique<FRequestPool>(*Domain, *EffectiveDomain, Access.Get(), UE_HTTPDDC_GET_REQUEST_POOL_SIZE);
-		PutRequestPools[0] = MakeUnique<FRequestPool>(*Domain, *EffectiveDomain, Access.Get(), UE_HTTPDDC_PUT_REQUEST_POOL_SIZE);
-		PutRequestPools[1] = MakeUnique<FRequestPool>(*Domain, *EffectiveDomain, Access.Get(), UE_HTTPDDC_PUT_REQUEST_POOL_SIZE);
+		GetRequestPools[0] = MakeUnique<FRequestPool>(*Domain, *EffectiveDomain, Access.Get(), SharedData.Get(), UE_HTTPDDC_GET_REQUEST_POOL_SIZE);
+		GetRequestPools[1] = MakeUnique<FRequestPool>(*Domain, *EffectiveDomain, Access.Get(), SharedData.Get(), UE_HTTPDDC_GET_REQUEST_POOL_SIZE);
+		PutRequestPools[0] = MakeUnique<FRequestPool>(*Domain, *EffectiveDomain, Access.Get(), SharedData.Get(), UE_HTTPDDC_PUT_REQUEST_POOL_SIZE);
+		PutRequestPools[1] = MakeUnique<FRequestPool>(*Domain, *EffectiveDomain, Access.Get(), SharedData.Get(), UE_HTTPDDC_PUT_REQUEST_POOL_SIZE);
 		// Allowing the non-blocking requests to overflow to double their pre-allocated size before we start waiting for one to free up.
-		NonBlockingRequestPools = MakeUnique<FRequestPool>(*Domain, *EffectiveDomain, Access.Get(), UE_HTTPDDC_NONBLOCKING_REQUEST_POOL_SIZE, UE_HTTPDDC_NONBLOCKING_REQUEST_POOL_SIZE);
+		NonBlockingRequestPools = MakeUnique<FRequestPool>(*Domain, *EffectiveDomain, Access.Get(), SharedData.Get(), UE_HTTPDDC_NONBLOCKING_REQUEST_POOL_SIZE, UE_HTTPDDC_NONBLOCKING_REQUEST_POOL_SIZE);
 		bIsUsable = true;
 	}
 
@@ -3527,7 +3528,6 @@ FHttpCacheStore::~FHttpCacheStore()
 #if WITH_DATAREQUEST_HELPER
 	FDataRequestHelper::StaticShutdown();
 #endif
-	FHttpSharedData::bAsyncThreadShutdownRequested.store(true, std::memory_order_relaxed);
 }
 
 FString FHttpCacheStore::GetName() const
@@ -3559,7 +3559,7 @@ bool FHttpCacheStore::ApplyDebugOptions(FBackendDebugOptions& InOptions)
 
 bool FHttpCacheStore::IsServiceReady()
 {
-	FHttpRequest Request(*Domain, *Domain, nullptr, false);
+	FHttpRequest Request(*Domain, *Domain, nullptr, SharedData.Get(), false);
 	FHttpRequest::Result Result = Request.PerformBlockingDownload(TEXT("health/ready"), nullptr);
 	
 	if (Result == FHttpRequest::Success && Request.GetResponseCode() == 200)
@@ -3613,7 +3613,7 @@ bool FHttpCacheStore::AcquireAccessToken()
 		FString AuthDomain(DomainEnd, *OAuthProvider);
 		FString Uri(*OAuthProvider + DomainEnd + 1);
 
-		FHttpRequest Request(*AuthDomain, *AuthDomain, nullptr, false);
+		FHttpRequest Request(*AuthDomain, *AuthDomain, nullptr, SharedData.Get(), false);
 		FHttpRequest::Result Result = FHttpRequest::Success;
 		if (OAuthProvider.StartsWith(TEXT("http://localhost")))
 		{
@@ -4148,10 +4148,10 @@ void FHttpCacheStore::RefCachedDataProbablyExistsBatchAsync(
 			{
 				for (const FCacheGetValueRequest& ValueRef : ValueRefs)
 				{
-				UE_LOG(LogDerivedDataCache, Log,
-					TEXT("%s: Cache exists returned unexpected quantity of results (expected %d, got %d)."),
-					*GetName(), ValueRefs.Num(), ResultsArrayView.Num());
-					OnComplete(ValueRef.MakeResponse(EStatus::Error));
+					UE_LOG(LogDerivedDataCache, Log,
+						TEXT("%s: Cache exists returned unexpected quantity of results (expected %d, got %d)."),
+						*GetName(), ValueRefs.Num(), ResultsArrayView.Num());
+						OnComplete(ValueRef.MakeResponse(EStatus::Error));
 				}
 				return FHttpRequest::ECompletionBehavior::Done;
 			}
@@ -4948,17 +4948,7 @@ void FHttpRequest::EnqueueAsync(IRequestOwner& Owner, FRequestPool* Pool, const
 	AsyncData->OnComplete = MoveTemp(OnComplete);
 	AsyncData->Owner->Begin(AsyncData);
 
-	if (SharedData->PendingRequestAdditions.EnqueueAndReturnWasEmpty(Curl))
-	{
-		SharedData->PendingRequestEvent.Trigger();
-	}
-
-	bool ExistingValue = false;
-	if (SharedData->bAsyncThreadStarting.compare_exchange_weak(ExistingValue, true, std::memory_order_relaxed))
-	{
-		SharedData->AsyncRunnable = MakeUnique<FHttpCacheStoreRunnable>(*SharedData);
-		SharedData->AsyncServiceThread = FRunnableThread::Create(SharedData->AsyncRunnable.Get(), TEXT("HttpCacheStoreThread"), 64 * 1024, TPri_Normal);
-	}
+	SharedData->AddRequest(Curl);
 }
 
 void FHttpRequest::CompleteAsync(CURLcode Result)
@@ -4992,10 +4982,7 @@ void FHttpRequest::CompleteAsync(CURLcode Result)
 	if (Behavior == ECompletionBehavior::Retry)
 	{
 		PrepareToRetry();
-		if (SharedData->PendingRequestAdditions.EnqueueAndReturnWasEmpty(Curl))
-		{
-			SharedData->PendingRequestEvent.Trigger();
-		}
+		SharedData->AddRequest(Curl);
 	}
 	else
 	{
@@ -5015,7 +5002,20 @@ void FHttpRequest::CompleteAsync(CURLcode Result)
 	}
 }
 
-uint32 FHttpCacheStoreRunnable::Run()
+void FHttpSharedData::AddRequest(CURL* Curl)
+{
+	if (PendingRequestAdditions.EnqueueAndReturnWasEmpty(Curl))
+	{
+		PendingRequestEvent.Trigger();
+	}
+
+	if (!bAsyncThreadStarting.load(std::memory_order_relaxed) && !bAsyncThreadStarting.exchange(true, std::memory_order_relaxed))
+	{
+		AsyncServiceThread = FThread(TEXT("HttpCacheStore"), [this] { ProcessAsyncRequests(); }, 64 * 1024, TPri_Normal);
+	}
+}
+
+void FHttpSharedData::ProcessAsyncRequests()
 {
 	int ActiveTransfers = 0;
 
@@ -5025,20 +5025,20 @@ uint32 FHttpCacheStoreRunnable::Run()
 
 		do
 		{
-			SharedData.PendingRequestAdditions.Deplete([this, &ActiveTransfers](CURL* Curl)
+			PendingRequestAdditions.Deplete([this, &ActiveTransfers](CURL* Curl)
 				{
-					curl_multi_add_handle(SharedData.CurlMulti, Curl);
+					curl_multi_add_handle(CurlMulti, Curl);
 					++ActiveTransfers;
 				});
 
-			curl_multi_perform(SharedData.CurlMulti, &CurrentActiveTransfers);
+			curl_multi_perform(CurlMulti, &CurrentActiveTransfers);
 
 			if (CurrentActiveTransfers == 0 || ActiveTransfers != CurrentActiveTransfers)
 			{
 				for (;;)
 				{
 					int MsgsStillInQueue = 0;	// may use that to impose some upper limit we may spend in that loop
-					CURLMsg* Message = curl_multi_info_read(SharedData.CurlMulti, &MsgsStillInQueue);
+					CURLMsg* Message = curl_multi_info_read(CurlMulti, &MsgsStillInQueue);
 
 					if (!Message)
 					{
@@ -5049,7 +5049,7 @@ uint32 FHttpCacheStoreRunnable::Run()
 					if (Message->msg == CURLMSG_DONE)
 					{
 						CURL* CompletedHandle = Message->easy_handle;
-						curl_multi_remove_handle(SharedData.CurlMulti, CompletedHandle);
+						curl_multi_remove_handle(CurlMulti, CompletedHandle);
 
 						void* PrivateData = nullptr;
 						curl_easy_getinfo(CompletedHandle, CURLINFO_PRIVATE, &PrivateData);
@@ -5072,7 +5072,7 @@ uint32 FHttpCacheStoreRunnable::Run()
 
 			if (CurrentActiveTransfers > 0)
 			{
-				curl_multi_wait(SharedData.CurlMulti, nullptr, 0, 1, nullptr);
+				curl_multi_wait(CurlMulti, nullptr, 0, 1, nullptr);
 			}
 		}
 		while (CurrentActiveTransfers > 0);
@@ -5081,14 +5081,12 @@ uint32 FHttpCacheStoreRunnable::Run()
 	do
 	{
 		ProcessPendingRequests();
-		SharedData.PendingRequestEvent.Wait(100);
+		PendingRequestEvent.Wait(100);
 	}
 	while (!FHttpSharedData::bAsyncThreadShutdownRequested.load(std::memory_order_relaxed));
 
 	// Process last requests before shutdown.  May want these to be aborted instead.
 	ProcessPendingRequests();
-
-	return 0;
 }
 
 } // UE::DerivedData
