Ensure that http headers can be found in a case-insensitive way as expected by the callers of GetHeader.

#rnx

diff --git a/Engine/Source/Developer/DerivedDataCache/Private/HttpCacheStore.cpp b/Engine/Source/Developer/DerivedDataCache/Private/HttpCacheStore.cpp
--- a/Engine/Source/Developer/DerivedDataCache/Private/HttpCacheStore.cpp
+++ b/Engine/Source/Developer/DerivedDataCache/Private/HttpCacheStore.cpp
@@ -37,6 +37,7 @@
 #include "HAL/Thread.h"
 #include "IO/IoHash.h"
 #include "Memory/SharedBuffer.h"
+#include "Misc/CString.h"
 #include "Misc/FileHelper.h"
 #include "Misc/Optional.h"
 #include "Misc/ScopeLock.h"
@@ -586,7 +587,7 @@ public:
 		size_t HeaderLen = strlen(Header);
 
 		// Find the header key in the (ANSI) response buffer. If not found we can exist immediately
-		if (const ANSICHAR* Found = strstr(HeadersBuffer, Header))
+		if (const ANSICHAR* Found = FCStringAnsi::Stristr(HeadersBuffer, Header))
 		{
 			const ANSICHAR* Linebreak = strchr(Found, '\r');
 			const ANSICHAR* ValueStart = Found + HeaderLen + 2; //colon and space
