Improve frame rate stability in editor by reducing the priority of the DDC IO threads which are now doing CPU intensive compression to avoid them preempting more important threads (GT/RT/RHI)

#rnx

diff --git a/Engine/Source/Developer/DerivedDataCache/Private/DerivedDataCache.cpp b/Engine/Source/Developer/DerivedDataCache/Private/DerivedDataCache.cpp
--- a/Engine/Source/Developer/DerivedDataCache/Private/DerivedDataCache.cpp
+++ b/Engine/Source/Developer/DerivedDataCache/Private/DerivedDataCache.cpp
@@ -535,7 +535,12 @@ public:
 		{
 			GCacheThreadPool = FQueuedThreadPool::Allocate();
 			const int32 ThreadCount = FPlatformMisc::NumberOfIOWorkerThreadsToSpawn();
+#if WITH_EDITOR
+			// Use normal priority to avoid preempting GT/RT/RHI and other more important threads with CPU processing (i.e. compression) happening on the IO Threads in editor.
+			verify(GCacheThreadPool->Create(ThreadCount, 96 * 1024, TPri_Normal, TEXT("DDC IO ThreadPool")));
+#else
 			verify(GCacheThreadPool->Create(ThreadCount, 96 * 1024, TPri_AboveNormal, TEXT("DDC IO ThreadPool")));
+#endif
 		}
 
 		Backend = FDerivedDataBackend::Create();
