Separate re-usable HTTP types from the HttpCacheStore to allow for re-use in the future by other cache stores.

#rnx

diff --git a/Engine/Source/Developer/DerivedDataCache/Private/Http/HttpClient.cpp b/Engine/Source/Developer/DerivedDataCache/Private/Http/HttpClient.cpp
new file mode 100644
--- /dev/null
+++ b/Engine/Source/Developer/DerivedDataCache/Private/Http/HttpClient.cpp
@@ -0,0 +1,316 @@
+// Copyright Epic Games, Inc. All Rights Reserved.
+
+#include "HttpClient.h"
+#include "Tasks/Task.h"
+
+#define UE_HTTPDDC_BACKEND_WAIT_INTERVAL 0.01f
+#define UE_HTTPDDC_BACKEND_WAIT_INTERVAL_MS ((uint32)(UE_HTTPDDC_BACKEND_WAIT_INTERVAL*1000))
+
+namespace UE
+{
+
+FString FHttpAccessToken::GetHeader()
+{
+	Lock.ReadLock();
+	FString Header = FString::Printf(TEXT("Authorization: Bearer %s"), *Token);
+	Lock.ReadUnlock();
+	return Header;
+}
+
+void FHttpAccessToken::SetHeader(const TCHAR* InToken)
+{
+	Lock.WriteLock();
+	Token = InToken;
+	Serial++;
+	Lock.WriteUnlock();
+}
+
+uint32 FHttpAccessToken::GetSerial() const
+{
+	return Serial;
+}
+
+FHttpSharedData::FHttpSharedData()
+: PendingRequestEvent(EEventMode::AutoReset)
+{
+	CurlShare = curl_share_init();
+	curl_share_setopt(CurlShare, CURLSHOPT_USERDATA, this);
+	curl_share_setopt(CurlShare, CURLSHOPT_LOCKFUNC, LockFn);
+	curl_share_setopt(CurlShare, CURLSHOPT_UNLOCKFUNC, UnlockFn);
+	curl_share_setopt(CurlShare, CURLSHOPT_SHARE, CURL_LOCK_DATA_DNS);
+	curl_share_setopt(CurlShare, CURLSHOPT_SHARE, CURL_LOCK_DATA_SSL_SESSION);
+	CurlMulti = curl_multi_init();
+	curl_multi_setopt(CurlMulti, CURLMOPT_PIPELINING, CURLPIPE_MULTIPLEX);
+}
+
+FHttpSharedData::~FHttpSharedData()
+{
+	bAsyncThreadShutdownRequested.store(true, std::memory_order_relaxed);
+	if (AsyncServiceThread.IsJoinable())
+	{
+		AsyncServiceThread.Join();
+	}
+
+	curl_multi_cleanup(CurlMulti);
+	curl_share_cleanup(CurlShare);
+}
+
+void FHttpSharedData::AddRequest(CURL* Curl)
+{
+	if (PendingRequestAdditions.EnqueueAndReturnWasEmpty(Curl))
+	{
+		PendingRequestEvent.Trigger();
+	}
+
+	if (!bAsyncThreadStarting.load(std::memory_order_relaxed) && !bAsyncThreadStarting.exchange(true, std::memory_order_relaxed))
+	{
+		AsyncServiceThread = FThread(TEXT("HttpCacheStore"), [this] { ProcessAsyncRequests(); }, 64 * 1024, TPri_Normal);
+	}
+}
+
+void FHttpSharedData::LockFn(CURL* Handle, curl_lock_data Data, curl_lock_access Access, void* User)
+{
+	FHttpSharedData* SharedData = (FHttpSharedData*)User;
+	if (Access == CURL_LOCK_ACCESS_SHARED)
+	{
+		SharedData->Locks[Data].ReadLock();
+	}
+	else
+	{
+		SharedData->Locks[Data].WriteLock();
+		SharedData->WriteLocked[Data] = true;
+	}
+}
+
+void FHttpSharedData::UnlockFn(CURL* Handle, curl_lock_data Data, void* User)
+{
+	FHttpSharedData* SharedData = (FHttpSharedData*)User;
+	if (!SharedData->WriteLocked[Data])
+	{
+		SharedData->Locks[Data].ReadUnlock();
+	}
+	else
+	{
+		SharedData->WriteLocked[Data] = false;
+		SharedData->Locks[Data].WriteUnlock();
+	}
+}
+
+void FHttpSharedData::ProcessAsyncRequests()
+{
+	int ActiveTransfers = 0;
+
+	auto ProcessPendingRequests = [this, &ActiveTransfers]()
+	{
+		int CurrentActiveTransfers = -1;
+
+		do
+		{
+			PendingRequestAdditions.Deplete([this, &ActiveTransfers](CURL* Curl)
+				{
+					curl_multi_add_handle(CurlMulti, Curl);
+					++ActiveTransfers;
+				});
+
+			curl_multi_perform(CurlMulti, &CurrentActiveTransfers);
+
+			if (CurrentActiveTransfers == 0 || ActiveTransfers != CurrentActiveTransfers)
+			{
+				for (;;)
+				{
+					int MsgsStillInQueue = 0;	// may use that to impose some upper limit we may spend in that loop
+					CURLMsg* Message = curl_multi_info_read(CurlMulti, &MsgsStillInQueue);
+
+					if (!Message)
+					{
+						break;
+					}
+
+					// find out which requests have completed
+					if (Message->msg == CURLMSG_DONE)
+					{
+						CURL* CompletedHandle = Message->easy_handle;
+						curl_multi_remove_handle(CurlMulti, CompletedHandle);
+
+						void* PrivateData = nullptr;
+						curl_easy_getinfo(CompletedHandle, CURLINFO_PRIVATE, &PrivateData);
+						FHttpRequest* CompletedRequest = (FHttpRequest*)PrivateData;
+
+						if (CompletedRequest)
+						{
+							// It is important that the CompleteAsync call doesn't happen on this thread as it is possible it will block waiting
+							// for a free HTTP request, and if that happens on this thread, we can deadlock as no HTTP requests will become 
+							// available while this thread is blocked.
+							UE::Tasks::Launch(TEXT("FHttpRequest::CompleteAsync"), [CompletedRequest, Result = Message->data.result]() mutable
+							{
+								CompletedRequest->CompleteAsync(Result);
+							});
+						}
+					}
+				}
+				ActiveTransfers = CurrentActiveTransfers;
+			}
+
+			if (CurrentActiveTransfers > 0)
+			{
+				curl_multi_wait(CurlMulti, nullptr, 0, 1, nullptr);
+			}
+		}
+		while (CurrentActiveTransfers > 0);
+	};
+
+	do
+	{
+		ProcessPendingRequests();
+		PendingRequestEvent.Wait(100);
+	}
+	while (!FHttpSharedData::bAsyncThreadShutdownRequested.load(std::memory_order_relaxed));
+
+	// Process last requests before shutdown.  May want these to be aborted instead.
+	ProcessPendingRequests();
+}
+
+FHttpRequestPool::FHttpRequestPool(const TCHAR* InServiceUrl, const TCHAR* InEffectiveServiceUrl, FHttpAccessToken* InAuthorizationToken, FHttpSharedData* InSharedData, uint32 PoolSize, uint32 InOverflowLimit)
+: ActiveOverflowRequests(0)
+, OverflowLimit(InOverflowLimit)
+{
+	Pool.AddUninitialized(PoolSize);
+	Requests.AddUninitialized(PoolSize);
+	for (uint8 i = 0; i < Pool.Num(); ++i)
+	{
+		Pool[i].Usage = 0u;
+		Pool[i].Request = new(&Requests[i]) FHttpRequest(InServiceUrl, InEffectiveServiceUrl, InAuthorizationToken, InSharedData, true);
+	}
+
+	InitData = MakeUnique<FInitData>(InServiceUrl, InEffectiveServiceUrl, InAuthorizationToken, InSharedData);
+}
+
+FHttpRequestPool::~FHttpRequestPool()
+{
+	check(ActiveOverflowRequests.load() == 0);
+	for (uint8 i = 0; i < Pool.Num(); ++i)
+	{
+		// No requests should be in use by now.
+		check(Pool[i].Usage.load(std::memory_order_acquire) == 0u);
+	}
+}
+
+FHttpRequest* FHttpRequestPool::GetFreeRequest(bool bUnboundedOverflow)
+{
+	for (uint8 i = 0; i < Pool.Num(); ++i)
+	{
+		if (!Pool[i].Usage.load(std::memory_order_relaxed))
+		{
+			uint8 Expected = 0u;
+			if (Pool[i].Usage.compare_exchange_strong(Expected, 1u))
+			{
+				Pool[i].Request->Reset();
+				return Pool[i].Request;
+			}
+		}
+	}
+	if (bUnboundedOverflow || (OverflowLimit > 0))
+	{
+		// The use of two operations here (load, then fetch_add) implies that we can exceed the overflow limit because the combined operation
+		// is not atomic.  This is acceptable for our use case.  If we wanted to enforce the hard limit, we could use a loop instead.
+		if (bUnboundedOverflow || (ActiveOverflowRequests.load(std::memory_order_relaxed) < OverflowLimit))
+		{
+			// Create an overflow request (outside of the pre-allocated range of requests)
+			ActiveOverflowRequests.fetch_add(1, std::memory_order_relaxed);
+			return new FHttpRequest(*InitData->ServiceUrl, *InitData->EffectiveServiceUrl, InitData->AccessToken, InitData->SharedData, true);
+		}
+	}
+	return nullptr;
+}
+
+FHttpRequest* FHttpRequestPool::WaitForFreeRequest(bool bUnboundedOverflow)
+{
+	TRACE_CPUPROFILER_EVENT_SCOPE(HttpDDC_WaitForConnection);
+
+	FHttpRequest* Request = GetFreeRequest(bUnboundedOverflow);
+	if (Request == nullptr)
+	{
+		// Make it a fair by allowing each thread to register itself in a FIFO
+		// so that the first thread to start waiting is the first one to get a request.
+		FWaiter* Waiter = new FWaiter(this);
+		Waiter->AddRef(); // One ref for the thread that will dequeue
+		Waiter->AddRef(); // One ref for us
+
+		Waiters.enqueue(Waiter);
+
+		while (!Waiter->Wait(UE_HTTPDDC_BACKEND_WAIT_INTERVAL_MS))
+		{
+			// While waiting, allow us to check if a race occurred and a request has been freed
+			// between the time we checked for free requests and the time we queued ourself as a Waiter.
+			if ((Request = GetFreeRequest(bUnboundedOverflow)) != nullptr)
+			{
+				// We abandon the FWaiter, it will be freed by the next dequeue
+				// and if it has a request, it will be queued back to the pool.
+				Waiter->Release();
+				return Request;
+			}
+		}
+
+		Request = Waiter->Request.exchange(nullptr);
+		Request->Reset();
+		Waiter->Release();
+	}
+	check(Request);
+	return Request;
+}
+
+void FHttpRequestPool::ReleaseRequestToPool(FHttpRequest* Request)
+{
+	if ((Request < Requests.GetData()) || (Request >= (Requests.GetData() + Requests.Num())))
+	{
+		// For overflow requests (outside of the pre-allocated range of requests), just delete it immediately
+		delete Request;
+		ActiveOverflowRequests.fetch_sub(1, std::memory_order_relaxed);
+		return;
+	}
+
+	for (uint8 i = 0; i < Pool.Num(); ++i)
+	{
+		if (Pool[i].Request == Request)
+		{
+			// If only 1 user is remaining, we can give it to a waiter
+			// instead of releasing it back to the pool.
+			if (Pool[i].Usage == 1u)
+			{
+				if (FWaiter* Waiter = Waiters.dequeue())
+				{
+					Waiter->Request = Request;
+					Waiter->Trigger();
+					Waiter->Release();
+					return;
+				}
+			}
+			
+			Pool[i].Usage--;
+			return;
+		}
+	}
+	check(false);
+}
+
+void FHttpRequestPool::MakeRequestShared(FHttpRequest* Request, uint8 Users)
+{
+	if ((Request < Requests.GetData()) || (Request >= (Requests.GetData() + Requests.Num())))
+	{
+		// Overflow requests (outside of the pre-allocated range of requests), cannot be made shared
+		check(false);
+	}
+
+	check(Users != 0);
+	for (uint8 i = 0; i < Pool.Num(); ++i)
+	{
+		if (Pool[i].Request == Request)
+		{
+			Pool[i].Usage = Users;
+			return;
+		}
+	}
+	check(false);
+}
+
+} // UE
diff --git a/Engine/Source/Developer/DerivedDataCache/Private/Http/HttpClient.h b/Engine/Source/Developer/DerivedDataCache/Private/Http/HttpClient.h
new file mode 100644
--- /dev/null
+++ b/Engine/Source/Developer/DerivedDataCache/Private/Http/HttpClient.h
@@ -0,0 +1,660 @@
+// Copyright Epic Games, Inc. All Rights Reserved.
+
+#pragma once
+
+#include "CoreTypes.h"
+#include "Containers/Array.h"
+#include "Containers/Array.h"
+#include "Containers/ArrayView.h"
+#include "Containers/DepletableMpscQueue.h"
+#include "Containers/UnrealString.h"
+#include "DerivedDataRequest.h"
+#include "DerivedDataRequestOwner.h"
+#include "DerivedDataRequestTypes.h"
+#include "Dom/JsonObject.h"
+#include "Experimental/Async/LazyEvent.h"
+#include "HAL/CriticalSection.h"
+#include "HAL/CriticalSection.h"
+#include "HAL/Thread.h"
+#include "Memory/MemoryView.h"
+#include "Memory/SharedBuffer.h"
+#include "Templates/SharedPointer.h"
+#include <atomic>
+
+#if WITH_SSL
+#include "Ssl.h"
+#include <openssl/ssl.h>
+#endif
+
+#if PLATFORM_MICROSOFT
+#include "Microsoft/AllowMicrosoftPlatformTypes.h"
+#endif
+
+// if there is a platform-specific include then it must be used in place of forward declation in case it defines CURL_STRICTER
+#if defined(PLATFORM_CURL_INCLUDE)
+#include PLATFORM_CURL_INCLUDE
+#else
+#include "curl/curl.h"
+#endif //defined(PLATFORM_CURL_INCLUDE)
+
+#if PLATFORM_MICROSOFT
+#include "Microsoft/HideMicrosoftPlatformTypes.h"
+#endif
+
+namespace UE
+{
+
+class FHttpRequestPool;
+
+/**
+ * Encapsulation for access token shared by all requests.
+ */
+class FHttpAccessToken
+{
+public:
+	FHttpAccessToken() = default;
+	FString GetHeader();
+	void SetHeader(const TCHAR*);
+	uint32 GetSerial() const;
+private:
+	FRWLock		Lock;
+	FString		Token;
+	uint32		Serial;
+};
+
+class FHttpSharedData
+{
+public:
+	FHttpSharedData();
+	~FHttpSharedData();
+
+	void AddRequest(CURL* Curl);
+	CURLSH* GetCurlShare() const { return CurlShare; }
+
+private:
+	CURLSH* CurlShare;
+	CURLM* CurlMulti;
+	UE::TDepletableMpscQueue<CURL*> PendingRequestAdditions;
+	UE::FLazyEvent PendingRequestEvent;
+	FThread AsyncServiceThread;
+	std::atomic<bool> bAsyncThreadStarting = false;
+	std::atomic<bool> bAsyncThreadShutdownRequested = false;
+
+	FRWLock Locks[CURL_LOCK_DATA_LAST];
+	bool WriteLocked[CURL_LOCK_DATA_LAST]{};
+
+	static void LockFn(CURL* Handle, curl_lock_data Data, curl_lock_access Access, void* User);
+	static void UnlockFn(CURL* Handle, curl_lock_data Data, void* User);
+	void ProcessAsyncRequests();
+};
+
+
+/**
+ * Minimal HTTP request type wrapping CURL without the need for managers. This request
+ * is written to allow reuse of request objects, in order to allow connections to be reused.
+ *
+ * CURL has a global library initialization (curl_global_init). We rely on this happening in 
+ * the Online/HTTP library which is a dependency on this module.
+ */
+class FHttpRequest
+{
+public:
+	/**
+	 * Supported request verbs
+	 */
+	enum RequestVerb
+	{
+		Get,
+		Put,
+		PutCompactBinary,
+		PutCompressedBlob,
+		Post,
+		PostCompactBinary,
+		PostJson,
+		Delete,
+		Head
+	};
+
+	/**
+	 * Convenience result type interpreted from HTTP response code.
+	 */
+	enum Result
+	{
+		Success,
+		Failed,
+		FailedTimeout
+	};
+
+	enum class ECompletionBehavior : uint8
+	{
+		Done,
+		Retry
+	};
+
+	using FOnHttpRequestComplete = TUniqueFunction<ECompletionBehavior(Result HttpResult, FHttpRequest* Request)>;
+
+	FHttpRequest(const TCHAR* InDomain, const TCHAR* InEffectiveDomain, FHttpAccessToken* InAuthorizationToken, FHttpSharedData* InSharedData, bool bInLogErrors);
+	~FHttpRequest();
+
+	/**
+	 * Resets all options on the request except those that should always be set.
+	 */
+	void Reset();
+
+	void PrepareToRetry();
+
+	/** Gets the domain name for this request */
+	const FString& GetName() const
+	{
+		return Domain;
+	}
+
+	/** Gets the domain name for this request */
+	const FString& GetDomain() const
+	{
+		return Domain;
+	}
+
+	/** Gets the effective domain name for this request */
+	const FString& GetEffectiveDomain() const
+	{
+		return EffectiveDomain;
+	}
+
+	/** Returns the HTTP response code.*/
+	const int64 GetResponseCode() const
+	{
+		return ResponseCode;
+	}
+
+	/** Returns the number of bytes received this request (headers withstanding). */
+	const size_t GetBytesReceived() const
+	{
+		return BytesReceived;
+	}
+
+	/** Returns the number of attempts we've made issuing this request (currently tracked for async requests only). */
+	const size_t GetAttempts() const
+	{
+		return Attempts;
+	}
+
+	/** Returns the number of bytes sent during this request (headers withstanding). */
+	const size_t GetBytesSent() const
+	{
+		return BytesSent;
+	}
+
+	/**
+	 * Upload buffer using the request, using either "Put" or "Post" verbs.
+	 * @param Uri Url to use.
+	 * @param Buffer Data to upload
+	 * @return Result of the request
+	 */
+	template<RequestVerb V>
+	Result PerformBlockingUpload(const TCHAR* Uri, TArrayView<const uint8> Buffer, TConstArrayView<long> ExpectedErrorCodes = {})
+	{
+		static_assert(V == Put || V == PutCompactBinary || V == PutCompressedBlob || V == Post || V == PostCompactBinary || V == PostJson, "Upload should use either Put or Post verbs.");
+		
+		uint64 ContentLength = 0u;
+
+		if constexpr (V == Put || V == PutCompactBinary || V == PutCompressedBlob)
+		{
+			curl_easy_setopt(Curl, CURLOPT_UPLOAD, 1L);
+			curl_easy_setopt(Curl, CURLOPT_INFILESIZE, Buffer.Num());
+			curl_easy_setopt(Curl, CURLOPT_READDATA, this);
+			curl_easy_setopt(Curl, CURLOPT_READFUNCTION, StaticReadFn);
+			if constexpr (V == PutCompactBinary)
+			{
+				Headers.Add(FString(TEXT("Content-Type: application/x-ue-cb")));
+			}
+			else if constexpr (V == PutCompressedBlob)
+			{
+				Headers.Add(FString(TEXT("Content-Type: application/x-ue-comp")));
+			}
+			else
+			{
+				Headers.Add(FString(TEXT("Content-Type: application/octet-stream")));
+			}
+			ContentLength = Buffer.Num();
+			ReadDataView = MakeMemoryView(Buffer);
+		}
+		else if constexpr (V == Post || V == PostCompactBinary || V == PostJson)
+		{
+			curl_easy_setopt(Curl, CURLOPT_POST, 1L);
+			curl_easy_setopt(Curl, CURLOPT_INFILESIZE, Buffer.Num());
+			curl_easy_setopt(Curl, CURLOPT_READDATA, this);
+			curl_easy_setopt(Curl, CURLOPT_READFUNCTION, StaticReadFn);
+			if constexpr (V == PostCompactBinary)
+			{
+				Headers.Add(FString(TEXT("Content-Type: application/x-ue-cb")));
+			}
+			else if constexpr (V == PostJson)
+			{
+				Headers.Add(FString(TEXT("Content-Type: application/json")));
+			}
+			else
+			{
+				Headers.Add(FString(TEXT("Content-Type: application/x-www-form-urlencoded")));
+			}
+			ContentLength = Buffer.Num();
+			ReadDataView = MakeMemoryView(Buffer);
+		}
+
+		return PerformBlocking(Uri, V, ContentLength, ExpectedErrorCodes);
+	}
+	
+	template<RequestVerb V>
+	void EnqueueAsyncUpload(UE::DerivedData::IRequestOwner& Owner, FHttpRequestPool* Pool, const TCHAR* Uri, FSharedBuffer Buffer, FOnHttpRequestComplete&& OnComplete, TConstArrayView<long> ExpectedErrorCodes = {})
+	{
+		static_assert(V == Put || V == PutCompactBinary || V == PutCompressedBlob || V == Post || V == PostCompactBinary || V == PostJson, "Upload should use either Put or Post verbs.");
+		
+		uint64 ContentLength = 0u;
+
+		if constexpr (V == Put || V == PutCompactBinary || V == PutCompressedBlob)
+		{
+			curl_easy_setopt(Curl, CURLOPT_UPLOAD, 1L);
+			curl_easy_setopt(Curl, CURLOPT_INFILESIZE, Buffer.GetSize());
+			curl_easy_setopt(Curl, CURLOPT_READDATA, this);
+			curl_easy_setopt(Curl, CURLOPT_READFUNCTION, StaticReadFn);
+			if constexpr (V == PutCompactBinary)
+			{
+				Headers.Add(FString(TEXT("Content-Type: application/x-ue-cb")));
+			}
+			else if constexpr (V == PutCompressedBlob)
+			{
+				Headers.Add(FString(TEXT("Content-Type: application/x-ue-comp")));
+			}
+			else
+			{
+				Headers.Add(FString(TEXT("Content-Type: application/octet-stream")));
+			}
+			ReadSharedBuffer = Buffer;
+			ContentLength = Buffer.GetSize();
+			ReadDataView = Buffer.GetView();
+		}
+		else if constexpr (V == Post || V == PostCompactBinary || V == PostJson)
+		{
+			curl_easy_setopt(Curl, CURLOPT_POST, 1L);
+			curl_easy_setopt(Curl, CURLOPT_INFILESIZE, Buffer.GetSize());
+			curl_easy_setopt(Curl, CURLOPT_READDATA, this);
+			curl_easy_setopt(Curl, CURLOPT_READFUNCTION, StaticReadFn);
+			if constexpr (V == PostCompactBinary)
+			{
+				Headers.Add(FString(TEXT("Content-Type: application/x-ue-cb")));
+			}
+			else if constexpr (V == PostJson)
+			{
+				Headers.Add(FString(TEXT("Content-Type: application/json")));
+			}
+			else
+			{
+				Headers.Add(FString(TEXT("Content-Type: application/x-www-form-urlencoded")));
+			}
+			ReadSharedBuffer = Buffer;
+			ContentLength = Buffer.GetSize();
+			ReadDataView = Buffer.GetView();
+		}
+
+		return EnqueueAsync(Owner, Pool, Uri, V, ContentLength, MoveTemp(OnComplete), ExpectedErrorCodes);
+	}
+
+	/**
+	 * Download an url into a buffer using the request.
+	 * @param Uri Url to use.
+	 * @param Buffer Optional buffer where data should be downloaded to. If empty downloaded data will
+	 * be stored in an internal buffer and accessed GetResponse* methods.
+	 * @return Result of the request
+	 */
+	Result PerformBlockingDownload(const TCHAR* Uri, TArray<uint8>* Buffer, TConstArrayView<long> ExpectedErrorCodes = {400});
+
+	void EnqueueAsyncDownload(UE::DerivedData::IRequestOwner& Owner, FHttpRequestPool* Pool, const TCHAR* Uri, FOnHttpRequestComplete&& OnComplete, TConstArrayView<long> ExpectedErrorCodes = {400});
+
+	/**
+	 * Query an url using the request. Queries can use either "Head" or "Delete" verbs.
+	 * @param Uri Url to use.
+	 * @return Result of the request
+	 */
+	template<RequestVerb V>
+	Result PerformBlockingQuery(const TCHAR* Uri, TConstArrayView<long> ExpectedErrorCodes = {400})
+	{
+		static_assert(V == Head || V == Delete, "Queries should use either Head or Delete verbs.");
+
+		if (V == Delete)
+		{
+			curl_easy_setopt(Curl, CURLOPT_CUSTOMREQUEST, "DELETE");
+		}
+		else if (V == Head)
+		{
+			curl_easy_setopt(Curl, CURLOPT_NOBODY, 1L);
+		}
+
+		return PerformBlocking(Uri, V, 0u, ExpectedErrorCodes);
+	}
+
+	template<RequestVerb V>
+	void EnqueueAsyncQuery(UE::DerivedData::IRequestOwner& Owner, FHttpRequestPool* Pool, const TCHAR* Uri, FOnHttpRequestComplete&& OnComplete, TConstArrayView<long> ExpectedErrorCodes = {400})
+	{
+		static_assert(V == Head || V == Delete, "Queries should use either Head or Delete verbs.");
+
+		if (V == Delete)
+		{
+			curl_easy_setopt(Curl, CURLOPT_CUSTOMREQUEST, "DELETE");
+		}
+		else if (V == Head)
+		{
+			curl_easy_setopt(Curl, CURLOPT_NOBODY, 1L);
+		}
+
+		return EnqueueAsync(Owner, Pool, Uri, V, 0u, MoveTemp(OnComplete), ExpectedErrorCodes);
+	}
+
+	/**
+	 * Set a header to send with the request.
+	 */
+	void SetHeader(const TCHAR* Header, const TCHAR* Value);
+
+	/**
+	 * Attempts to find the header from the response. Returns false if header is not present.
+	 */
+	bool GetHeader(const ANSICHAR* Header, FString& OutValue) const;
+
+	/**
+	 * Returns the response buffer. Note that is the request is performed
+	 * with an external buffer as target buffer this string will be empty.
+	 */
+	const TArray<uint8>& GetResponseBuffer() const
+	{
+		return ResponseBuffer;
+	}
+
+	FSharedBuffer MoveResponseBufferToShared()
+	{
+		return MakeSharedBufferFromArray(MoveTemp(ResponseBuffer));
+	}
+
+	/**
+	 * Returns the response buffer as a string. Note that is the request is performed
+	 * with an external buffer as target buffer this string will be empty.
+	 */
+	FString GetResponseAsString() const
+	{
+		return GetAnsiBufferAsString(ResponseBuffer);
+	}
+
+	/**
+	 * Returns the response header as a string.
+	 */
+	FString GetResponseHeaderAsString()
+	{
+		return GetAnsiBufferAsString(ResponseHeader);
+	}
+
+	/**
+	 * Tries to parse the response buffer as a JsonObject. Return empty pointer if 
+	 * parse error occurs.
+	 */
+	TSharedPtr<FJsonObject> GetResponseAsJsonObject() const;
+
+	/**
+	 * Tries to parse the response buffer as a JsonArray. Return empty array if
+	 * parse error occurs.
+	 */
+	TArray<TSharedPtr<FJsonValue>> GetResponseAsJsonArray() const;
+
+	void CompleteAsync(CURLcode Result);
+
+	/** Will return true if the response code is considered a success */
+	static bool IsSuccessResponse(long ResponseCode)
+	{
+		// We consider anything in the 1XX or 2XX range a success
+		return ResponseCode >= 100 && ResponseCode < 300;
+	}
+
+	static bool AllowAsync();
+
+private:
+
+	struct FAsyncRequestData final : public UE::DerivedData::FRequestBase
+	{
+		UE::DerivedData::IRequestOwner* Owner = nullptr;
+		FHttpRequestPool* Pool = nullptr;
+		curl_slist* CurlHeaders = nullptr;
+		FString Uri;
+		RequestVerb Verb;
+		TArray<long, TInlineAllocator<4>> ExpectedErrorCodes;
+		FOnHttpRequestComplete OnComplete;
+		UE::FLazyEvent Event {EEventMode::ManualReset};
+
+		void Reset()
+		{
+			if (CurlHeaders)
+			{
+				curl_slist_free_all(CurlHeaders);
+				CurlHeaders = nullptr;
+			}
+			Uri.Empty();
+			ExpectedErrorCodes.Empty();
+		}
+
+		void SetPriority(UE::DerivedData::EPriority Priority) final {}
+
+		void Cancel() final
+		{
+			TRACE_CPUPROFILER_EVENT_SCOPE(HttpDDC_Cancel);
+			Event.Wait();
+		}
+
+		void Wait() final
+		{
+			TRACE_CPUPROFILER_EVENT_SCOPE(HttpDDC_Wait);
+			Event.Wait();
+		}
+	};
+
+	CURL*					Curl;
+	CURLcode				CurlResult;
+	FHttpSharedData*		SharedData;
+	FAsyncRequestData*		AsyncData;
+	long					ResponseCode;
+	size_t					BytesSent;
+	size_t					BytesReceived;
+	size_t					Attempts;
+	bool					bLogErrors;
+
+	FSharedBuffer			ReadSharedBuffer;
+	FMemoryView				ReadDataView;
+	TArray<uint8>*			WriteDataBufferPtr;
+	TArray<uint8>*			WriteHeaderBufferPtr;
+	
+	TArray<uint8>			ResponseHeader;
+	TArray<uint8>			ResponseBuffer;
+	TArray<FString>			Headers;
+	FString					Domain;
+	FString					EffectiveDomain;
+	FHttpAccessToken*		AuthorizationToken;
+
+	curl_slist* PrepareToIssueRequest(const TCHAR* Uri, uint64 ContentLength);
+
+	/**
+	 * Performs the request, blocking until finished.
+	 * @param Uri Address on the domain to query
+	 * @param Verb HTTP verb to use
+	 * @param ContentLength The number of bytes being uploaded for the body of this request.
+	 * @param ExpectedErrorCodes An array of expected return codes outside of the success range that should NOT be logged as an abnormal/exceptional outcome.
+	 * If unset the response body will be stored in the request.
+	 */
+	Result PerformBlocking(const TCHAR* Uri, RequestVerb Verb, uint64 ContentLength, TConstArrayView<long> ExpectedErrorCodes);
+
+	void EnqueueAsync(UE::DerivedData::IRequestOwner& Owner, FHttpRequestPool* Pool, const TCHAR* Uri, RequestVerb Verb, uint64 ContentLength, FOnHttpRequestComplete&& OnComplete, TConstArrayView<long> ExpectedErrorCodes);
+
+	void LogResult(CURLcode Result, const TCHAR* Uri, RequestVerb Verb, TConstArrayView<long> ExpectedErrorCodes) const;
+
+	FString GetAnsiBufferAsString(const TArray<uint8>& Buffer) const
+	{
+		// Content is NOT null-terminated; we need to specify lengths here
+		FUTF8ToTCHAR TCHARData(reinterpret_cast<const ANSICHAR*>(Buffer.GetData()), Buffer.Num());
+		return FString(TCHARData.Length(), TCHARData.Get());
+	}
+
+	//----------------------------------------------------------------------------------------------------------
+	// Certificate checking
+	//----------------------------------------------------------------------------------------------------------
+	#if WITH_SSL
+	static int SslCertVerify(int PreverifyOk, X509_STORE_CTX* Context);
+	static CURLcode StaticSSLCTXFn(CURL * curl, void * sslctx, void * parm);
+	#endif //#if WITH_SSL
+
+	static size_t StaticDebugCallback(CURL * Handle, curl_infotype DebugInfoType, char * DebugInfo, size_t DebugInfoSize, void* UserData);
+	static size_t StaticReadFn(void* Ptr, size_t SizeInBlocks, size_t BlockSizeInBytes, void* UserData);
+	static size_t StaticWriteHeaderFn(void* Ptr, size_t SizeInBlocks, size_t BlockSizeInBytes, void* UserData);
+	static size_t StaticWriteBodyFn(void* Ptr, size_t SizeInBlocks, size_t BlockSizeInBytes, void* UserData);
+	static size_t StaticSeekFn(void* UserData, curl_off_t Offset, int Origin);
+
+};
+
+/**
+ * Pool that manages a fixed set of requests. Users are required to release requests that have been 
+ * acquired. Usable with \ref FScopedRequestPtr which handles this automatically.
+ */
+class FHttpRequestPool
+{
+public:
+	FHttpRequestPool(const TCHAR* InServiceUrl, const TCHAR* InEffectiveServiceUrl, FHttpAccessToken* InAuthorizationToken, FHttpSharedData* InSharedData, uint32 PoolSize, uint32 InOverflowLimit = 0);
+	~FHttpRequestPool();
+
+	/**
+	 * Attempts to get a request is free. Once a request has been returned it is
+	 * "owned by the caller and need to release it to the pool when work has been completed.
+	 * @return Usable request instance if one is available, otherwise null.
+	 */
+	FHttpRequest* GetFreeRequest(bool bUnboundedOverflow = false);
+
+	class FWaiter : public FThreadSafeRefCountedObject
+	{
+	public:
+		std::atomic<FHttpRequest*> Request{ nullptr };
+
+		FWaiter(FHttpRequestPool* InPool)
+			: Event(FPlatformProcess::GetSynchEventFromPool(true))
+			, Pool(InPool)
+		{
+		}
+		
+		bool Wait(uint32 TimeMS)
+		{
+			return Event->Wait(TimeMS);
+		}
+
+		void Trigger()
+		{
+			Event->Trigger();
+		}
+	private:
+		~FWaiter()
+		{
+			FPlatformProcess::ReturnSynchEventToPool(Event);
+
+			if (Request)
+			{
+				Pool->ReleaseRequestToPool(Request.exchange(nullptr));
+			}
+		}
+
+		FEvent* Event;
+		FHttpRequestPool* Pool;
+	};
+
+	/**
+	 * Block until a request is free. Once a request has been returned it is 
+	 * "owned by the caller and need to release it to the pool when work has been completed.
+	 * @return Usable request instance.
+	 */
+	FHttpRequest* WaitForFreeRequest(bool bUnboundedOverflow = false);
+
+	/**
+	 * Release request to the pool.
+	 * @param Request Request that should be freed. Note that any buffer owened by the request can now be reset.
+	 */
+	void ReleaseRequestToPool(FHttpRequest* Request);
+
+	/**
+	 * While holding a request, make it shared across many users.
+	 */
+	void MakeRequestShared(FHttpRequest* Request, uint8 Users);
+
+private:
+
+	struct FEntry
+	{
+		std::atomic<uint8> Usage;
+		FHttpRequest* Request;
+	};
+
+	TArray<FEntry> Pool;
+	TArray<FHttpRequest> Requests;
+	FAAArrayQueue<FWaiter> Waiters;
+	std::atomic<uint32> ActiveOverflowRequests;
+	struct FInitData
+	{
+		FString ServiceUrl;
+		FString EffectiveServiceUrl;
+		FHttpAccessToken* AccessToken;
+		FHttpSharedData* SharedData;
+
+		FInitData(const TCHAR* InServiceUrl, const TCHAR* InEffectiveServiceUrl, FHttpAccessToken* InAccessToken, FHttpSharedData* InSharedData)
+		: ServiceUrl(InServiceUrl)
+		, EffectiveServiceUrl(InEffectiveServiceUrl)
+		, AccessToken(InAccessToken)
+		, SharedData(InSharedData)
+		{
+		}
+	};
+	TUniquePtr<const FInitData> InitData;
+	const uint32 OverflowLimit;
+
+	FHttpRequestPool() = delete;
+};
+
+/**
+ * Utility class to manage requesting and releasing requests from the \ref FHttpRequestPool.
+ */
+class FScopedHttpPoolRequestPtr
+{
+public:
+	UE_NONCOPYABLE(FScopedHttpPoolRequestPtr)
+
+	explicit FScopedHttpPoolRequestPtr(FHttpRequestPool* InPool)
+		: Request(InPool->WaitForFreeRequest())
+		, Pool(InPool)
+	{
+	}
+
+	~FScopedHttpPoolRequestPtr()
+	{
+		Pool->ReleaseRequestToPool(Request);
+	}
+
+	bool IsValid() const 
+	{
+		return Request != nullptr;
+	}
+
+	FHttpRequest* Get() const
+	{
+		check(IsValid());
+		return Request;
+	}
+
+	FHttpRequest* operator->()
+	{
+		check(IsValid());
+		return Request;
+	}
+
+private:
+	FHttpRequest* Request;
+	FHttpRequestPool* Pool;
+};
+
+} // UE
diff --git a/Engine/Source/Developer/DerivedDataCache/Private/Http/HttpRequest.cpp b/Engine/Source/Developer/DerivedDataCache/Private/Http/HttpRequest.cpp
new file mode 100644
--- /dev/null
+++ b/Engine/Source/Developer/DerivedDataCache/Private/Http/HttpRequest.cpp
@@ -0,0 +1,604 @@
+// Copyright Epic Games, Inc. All Rights Reserved.
+
+#include "HttpClient.h"
+#include "DerivedDataLegacyCacheStore.h"
+#include "Logging/LogMacros.h"
+#include "Serialization/JsonReader.h"
+#include "Serialization/JsonSerializer.h"
+
+#define UE_HTTPDDC_HTTP_REQUEST_TIMEOUT_SECONDS 30L
+#define UE_HTTPDDC_HTTP_REQUEST_TIMEOUT_ENABLED 1
+#define UE_HTTPDDC_HTTP_DEBUG 0
+#define UE_HTTPDDC_MAX_BUFFER_RESERVE 104857600u
+
+namespace UE
+{
+
+static bool bHttpEnableAsync = true;
+static FAutoConsoleVariableRef CVarHttpEnableAsync(
+	TEXT("DDC.Http.EnableAsync"),
+	bHttpEnableAsync,
+	TEXT("If true, async operations are permitted, otherwise all operations are forced to be synchronous."),
+	ECVF_Default);
+
+FHttpRequest::FHttpRequest(const TCHAR* InDomain, const TCHAR* InEffectiveDomain, FHttpAccessToken* InAuthorizationToken, FHttpSharedData* InSharedData, bool bInLogErrors)
+	: SharedData(InSharedData)
+	, AsyncData(nullptr)
+	, bLogErrors(bInLogErrors)
+	, Domain(InDomain)
+	, EffectiveDomain(InEffectiveDomain)
+	, AuthorizationToken(InAuthorizationToken)
+{
+	Curl = curl_easy_init();
+	Reset();
+}
+
+FHttpRequest::~FHttpRequest()
+{
+	curl_easy_cleanup(Curl);
+	check(!AsyncData);
+}
+
+void FHttpRequest::Reset()
+{
+	Headers.Reset();
+	ResponseHeader.Reset();
+	ResponseBuffer.Reset();
+	ResponseCode = 0;
+	ReadDataView = FMemoryView();
+	WriteDataBufferPtr = nullptr;
+	WriteHeaderBufferPtr = nullptr;
+	BytesSent = 0;
+	BytesReceived = 0;
+	Attempts = 0;
+	CurlResult = CURL_LAST;
+
+	curl_easy_reset(Curl);
+	check(!AsyncData);
+
+	// Options that are always set for all connections.
+#if UE_HTTPDDC_HTTP_REQUEST_TIMEOUT_ENABLED
+	curl_easy_setopt(Curl, CURLOPT_CONNECTTIMEOUT, UE_HTTPDDC_HTTP_REQUEST_TIMEOUT_SECONDS);
+#endif
+	curl_easy_setopt(Curl, CURLOPT_FOLLOWLOCATION, 1L);
+	curl_easy_setopt(Curl, CURLOPT_NOSIGNAL, 1L);
+	curl_easy_setopt(Curl, CURLOPT_DNS_CACHE_TIMEOUT, 300L); // Don't re-resolve every minute
+	curl_easy_setopt(Curl, CURLOPT_SHARE, SharedData ? SharedData->GetCurlShare() : nullptr);
+	// SSL options
+	curl_easy_setopt(Curl, CURLOPT_USE_SSL, CURLUSESSL_ALL);
+	curl_easy_setopt(Curl, CURLOPT_SSL_VERIFYPEER, 1);
+	curl_easy_setopt(Curl, CURLOPT_SSL_VERIFYHOST, 1);
+	curl_easy_setopt(Curl, CURLOPT_SSLCERTTYPE, "PEM");
+	// Response functions
+	curl_easy_setopt(Curl, CURLOPT_HEADERDATA, this);
+	curl_easy_setopt(Curl, CURLOPT_HEADERFUNCTION, &FHttpRequest::StaticWriteHeaderFn);
+	curl_easy_setopt(Curl, CURLOPT_WRITEDATA, this);
+	curl_easy_setopt(Curl, CURLOPT_WRITEFUNCTION, StaticWriteBodyFn);
+	// SSL certification verification
+	curl_easy_setopt(Curl, CURLOPT_CAINFO, nullptr);
+	curl_easy_setopt(Curl, CURLOPT_SSL_CTX_FUNCTION, &FHttpRequest::StaticSSLCTXFn);
+	curl_easy_setopt(Curl, CURLOPT_SSL_CTX_DATA, this);
+	// Allow compressed data
+	curl_easy_setopt(Curl, CURLOPT_ACCEPT_ENCODING, "gzip");
+	// Rewind method, handle special error case where request need to rewind data stream
+	curl_easy_setopt(Curl, CURLOPT_SEEKFUNCTION, StaticSeekFn);
+	curl_easy_setopt(Curl, CURLOPT_SEEKDATA, this);
+	// Set minimum speed behavior to allow operations to abort if the transfer speed is poor for the given duration (1kbps over a 30 second span)
+	curl_easy_setopt(Curl, CURLOPT_LOW_SPEED_TIME, 30L);
+	curl_easy_setopt(Curl, CURLOPT_LOW_SPEED_LIMIT, 1024L);
+	// Debug hooks
+#if UE_HTTPDDC_HTTP_DEBUG
+	curl_easy_setopt(Curl, CURLOPT_DEBUGDATA, this);
+	curl_easy_setopt(Curl, CURLOPT_DEBUGFUNCTION, StaticDebugCallback);
+	curl_easy_setopt(Curl, CURLOPT_VERBOSE, 1L);
+#endif
+	curl_easy_setopt(Curl, CURLOPT_PRIVATE, this);
+	curl_easy_setopt(Curl, CURLOPT_HTTP_VERSION, CURL_HTTP_VERSION_2_0);
+}
+
+void FHttpRequest::PrepareToRetry()
+{
+	ResponseHeader.Reset();
+	ResponseBuffer.Reset();
+	ResponseCode = 0;
+	BytesSent = 0;
+	BytesReceived = 0;
+	CurlResult = CURL_LAST;
+	++Attempts;
+}
+
+FHttpRequest::Result FHttpRequest::PerformBlockingDownload(const TCHAR* Uri, TArray<uint8>* Buffer, TConstArrayView<long> ExpectedErrorCodes)
+{
+	curl_easy_setopt(Curl, CURLOPT_HTTPGET, 1L);
+	WriteDataBufferPtr = Buffer;
+
+	return PerformBlocking(Uri, Get, 0u, ExpectedErrorCodes);
+}
+
+void FHttpRequest::EnqueueAsyncDownload(UE::DerivedData::IRequestOwner& Owner, FHttpRequestPool* Pool, const TCHAR* Uri, FOnHttpRequestComplete&& OnComplete, TConstArrayView<long> ExpectedErrorCodes)
+{
+	curl_easy_setopt(Curl, CURLOPT_HTTPGET, 1L);
+
+	return EnqueueAsync(Owner, Pool, Uri, Get, 0u, MoveTemp(OnComplete), ExpectedErrorCodes);
+}
+
+void FHttpRequest::SetHeader(const TCHAR* Header, const TCHAR* Value)
+{
+	check(CurlResult == CURL_LAST); // Cannot set header after request is sent
+	Headers.Add(FString::Printf(TEXT("%s: %s"), Header, Value));
+}
+
+bool FHttpRequest::GetHeader(const ANSICHAR* Header, FString& OutValue) const
+{
+	check(CurlResult != CURL_LAST);  // Cannot query headers before request is sent
+
+	const ANSICHAR* HeadersBuffer = (const ANSICHAR*) ResponseHeader.GetData();
+	size_t HeaderLen = strlen(Header);
+
+	// Find the header key in the (ANSI) response buffer. If not found we can exist immediately
+	if (const ANSICHAR* Found = FCStringAnsi::Stristr(HeadersBuffer, Header))
+	{
+		const ANSICHAR* Linebreak = strchr(Found, '\r');
+		const ANSICHAR* ValueStart = Found + HeaderLen + 2; //colon and space
+		const size_t ValueSize = Linebreak - ValueStart;
+		FUTF8ToTCHAR TCHARData(ValueStart, ValueSize);
+		OutValue = FString(TCHARData.Length(), TCHARData.Get());
+		return true;
+	}
+	return false;
+}
+
+TSharedPtr<FJsonObject> FHttpRequest::GetResponseAsJsonObject() const
+{
+	FString Response = GetAnsiBufferAsString(ResponseBuffer);
+
+	TSharedPtr<FJsonObject> JsonObject;
+	TSharedRef<TJsonReader<> > JsonReader = TJsonReaderFactory<>::Create(Response);
+	if (!FJsonSerializer::Deserialize(JsonReader, JsonObject) || !JsonObject.IsValid())
+	{
+		return TSharedPtr<FJsonObject>(nullptr);
+	}
+
+	return JsonObject;
+}
+
+TArray<TSharedPtr<FJsonValue>> FHttpRequest::GetResponseAsJsonArray() const
+{
+	FString Response = GetAnsiBufferAsString(ResponseBuffer);
+
+	TArray<TSharedPtr<FJsonValue>> JsonArray;
+	TSharedRef<TJsonReader<>> JsonReader = TJsonReaderFactory<>::Create(Response);
+	FJsonSerializer::Deserialize(JsonReader, JsonArray);
+	return JsonArray;
+}
+
+bool FHttpRequest::AllowAsync()
+{
+	if (!FGenericPlatformProcess::SupportsMultithreading() || !bHttpEnableAsync)
+	{
+		return false;
+	}
+
+	return true;
+}
+
+void FHttpRequest::CompleteAsync(CURLcode Result)
+{
+	CurlResult = Result;
+
+	// Get response code
+	curl_easy_getinfo(Curl, CURLINFO_RESPONSE_CODE, &ResponseCode);
+
+	LogResult(Result, *AsyncData->Uri, AsyncData->Verb, AsyncData->ExpectedErrorCodes);
+
+	ECompletionBehavior Behavior;
+	{
+		UE::DerivedData::FRequestBarrier Barrier(*AsyncData->Owner, UE::DerivedData::ERequestBarrierFlags::Priority);
+		FHttpRequest::Result HttpResult;
+		switch (CurlResult)
+		{
+		case CURLE_OK:
+			HttpResult = Success;
+			break;
+		case CURLE_OPERATION_TIMEDOUT:
+			HttpResult = FailedTimeout;
+			break;
+		default:
+			HttpResult = Failed;
+			break;
+		}
+		Behavior = AsyncData->OnComplete(HttpResult, this);
+	}
+
+	if (Behavior == ECompletionBehavior::Retry)
+	{
+		PrepareToRetry();
+		SharedData->AddRequest(Curl);
+	}
+	else
+	{
+		// Clean up
+		curl_slist_free_all(AsyncData->CurlHeaders);
+
+		AsyncData->Owner->End(AsyncData, [this, Behavior]
+			{
+				AsyncData->Event.Trigger();
+				FHttpRequestPool* Pool = AsyncData->Pool;
+				AsyncData = nullptr;
+				if (Pool)
+				{
+					Pool->ReleaseRequestToPool(this);
+				}
+			});
+	}
+}
+
+curl_slist* FHttpRequest::PrepareToIssueRequest(const TCHAR* Uri, uint64 ContentLength)
+{
+	static const char* CommonHeaders[] = {
+		"User-Agent: Unreal Engine",
+		nullptr
+	};
+
+	// Setup request options
+	FString Url = FString::Printf(TEXT("%s/%s"), *EffectiveDomain, Uri);
+	curl_easy_setopt(Curl, CURLOPT_URL, TCHAR_TO_ANSI(*Url));
+
+	// Setup response header buffer. If caller has not setup a response data buffer, use internal.
+	WriteHeaderBufferPtr = &ResponseHeader;
+	if (WriteDataBufferPtr == nullptr)
+	{
+		WriteDataBufferPtr = &ResponseBuffer;
+	}
+
+	// Content-Length should always be set
+	Headers.Add(FString::Printf(TEXT("Content-Length: %d"), ContentLength));
+
+	// And auth token if it's set
+	if (AuthorizationToken)
+	{
+		Headers.Add(AuthorizationToken->GetHeader());
+	}
+
+	// Build headers list
+	curl_slist* CurlHeaders = nullptr;
+	// Add common headers
+	for (uint8 i = 0; CommonHeaders[i] != nullptr; ++i)
+	{
+		CurlHeaders = curl_slist_append(CurlHeaders, CommonHeaders[i]);
+	}
+	// Setup added headers
+	for (const FString& Header : Headers)
+	{
+		CurlHeaders = curl_slist_append(CurlHeaders, TCHAR_TO_ANSI(*Header));
+	}
+	curl_easy_setopt(Curl, CURLOPT_HTTPHEADER, CurlHeaders);
+	return CurlHeaders;
+}
+
+FHttpRequest::Result FHttpRequest::PerformBlocking(const TCHAR* Uri, RequestVerb Verb, uint64 ContentLength, TConstArrayView<long> ExpectedErrorCodes)
+{
+	TRACE_CPUPROFILER_EVENT_SCOPE(HttpDDC_CurlPerform);
+
+	// Build headers list
+	curl_slist* CurlHeaders = PrepareToIssueRequest(Uri, ContentLength);
+
+	// Shots fired!
+	CurlResult = curl_easy_perform(Curl);
+
+	// Get response code
+	curl_easy_getinfo(Curl, CURLINFO_RESPONSE_CODE, &ResponseCode);
+
+	LogResult(CurlResult, Uri, Verb, ExpectedErrorCodes);
+
+	// Clean up
+	curl_slist_free_all(CurlHeaders);
+
+	return CurlResult == CURLE_OK ? Success : Failed;
+}
+
+void FHttpRequest::EnqueueAsync(UE::DerivedData::IRequestOwner& Owner, FHttpRequestPool* Pool, const TCHAR* Uri, RequestVerb Verb, uint64 ContentLength, FOnHttpRequestComplete&& OnComplete, TConstArrayView<long> ExpectedErrorCodes)
+{
+	if (!AllowAsync())
+	{
+		while (OnComplete(PerformBlocking(Uri, Verb, ContentLength, ExpectedErrorCodes), this) == ECompletionBehavior::Retry)
+		{
+			PrepareToRetry();
+		}
+		if (Pool)
+		{
+			Pool->ReleaseRequestToPool(this);
+		}
+		return;
+	}
+
+	TRACE_CPUPROFILER_EVENT_SCOPE(HttpDDC_CurlEnqueueAsync);
+	AsyncData = new FAsyncRequestData;
+	AsyncData->Owner = &Owner;
+	AsyncData->Pool = Pool;
+	AsyncData->CurlHeaders = PrepareToIssueRequest(Uri, ContentLength);
+	AsyncData->Uri = Uri;
+	AsyncData->Verb = Verb;
+	AsyncData->ExpectedErrorCodes = ExpectedErrorCodes;
+	AsyncData->OnComplete = MoveTemp(OnComplete);
+	AsyncData->Owner->Begin(AsyncData);
+
+	SharedData->AddRequest(Curl);
+}
+
+void FHttpRequest::LogResult(CURLcode Result, const TCHAR* Uri, RequestVerb Verb, TConstArrayView<long> ExpectedErrorCodes) const
+{
+	if (Result == CURLE_OK)
+	{
+		bool bSuccess = false;
+		const TCHAR* VerbStr = nullptr;
+		FString AdditionalInfo;
+
+		switch (Verb)
+		{
+		case Head:
+			bSuccess = (ExpectedErrorCodes.Contains(ResponseCode) || IsSuccessResponse(ResponseCode));
+			VerbStr = TEXT("querying");
+			break;
+		case Get:
+			bSuccess = (ExpectedErrorCodes.Contains(ResponseCode) || IsSuccessResponse(ResponseCode));
+			VerbStr = TEXT("fetching");
+			AdditionalInfo = FString::Printf(TEXT("Received: %d bytes."), BytesReceived);
+			break;
+		case Put:
+		case PutCompactBinary:
+		case PutCompressedBlob:
+			bSuccess = (ExpectedErrorCodes.Contains(ResponseCode) || IsSuccessResponse(ResponseCode));
+			VerbStr = TEXT("updating");
+			AdditionalInfo = FString::Printf(TEXT("Sent: %d bytes."), BytesSent);
+			break;
+		case Post:
+		case PostCompactBinary:
+		case PostJson:
+			bSuccess = (ExpectedErrorCodes.Contains(ResponseCode) || IsSuccessResponse(ResponseCode));
+			VerbStr = TEXT("posting");
+			break;
+		case Delete:
+			bSuccess = (ExpectedErrorCodes.Contains(ResponseCode) || IsSuccessResponse(ResponseCode));
+			VerbStr = TEXT("deleting");
+			break;
+		}
+
+		if (bSuccess)
+		{
+			UE_LOG(
+				LogDerivedDataCache, 
+				Verbose, 
+				TEXT("%s: Finished %s HTTP cache entry (response %d) from %s. %s"), 
+				*GetName(),
+				VerbStr,
+				ResponseCode, 
+				Uri,
+				*AdditionalInfo
+			);
+		}
+		else if(bLogErrors)
+		{
+			// Print the response body if we got one, otherwise print header.
+			FString Response = GetAnsiBufferAsString(ResponseBuffer.Num() > 0 ? ResponseBuffer : ResponseHeader);
+			Response.ReplaceCharInline(TEXT('\n'), TEXT(' '));
+			Response.ReplaceCharInline(TEXT('\r'), TEXT(' '));
+			// Dont log access denied as error, since tokens can expire mid session
+			if (ResponseCode == 401)
+			{
+				UE_LOG(
+					LogDerivedDataCache,
+					Verbose,
+					TEXT("%s: Failed %s HTTP cache entry (response %d) from %s. Response: %s"),
+					*GetName(),
+					VerbStr,
+					ResponseCode,
+					Uri,
+					*Response
+				);
+			}
+			else
+			{ 
+				UE_LOG(
+					LogDerivedDataCache,
+					Display,
+					TEXT("%s: Failed %s HTTP cache entry (response %d) from %s. Response: %s"),
+					*GetName(),
+					VerbStr,
+					ResponseCode,
+					Uri,
+					*Response
+				);
+			}
+		}
+	}
+	else if(bLogErrors)
+	{
+		UE_LOG(
+			LogDerivedDataCache, 
+			Display, 
+			TEXT("%s: Error while connecting to %s: %s"), 
+			*GetName(),
+			*EffectiveDomain,
+			ANSI_TO_TCHAR(curl_easy_strerror(Result))
+		);
+	}
+}
+
+#if WITH_SSL
+int FHttpRequest::SslCertVerify(int PreverifyOk, X509_STORE_CTX* Context)
+{
+	if (PreverifyOk == 1)
+	{
+		SSL* Handle = static_cast<SSL*>(X509_STORE_CTX_get_ex_data(Context, SSL_get_ex_data_X509_STORE_CTX_idx()));
+		check(Handle);
+
+		SSL_CTX* SslContext = SSL_get_SSL_CTX(Handle);
+		check(SslContext);
+
+		FHttpRequest* Request = static_cast<FHttpRequest*>(SSL_CTX_get_app_data(SslContext));
+		check(Request);
+
+		const FString& Domain = Request->GetDomain();
+
+		if (!FSslModule::Get().GetCertificateManager().VerifySslCertificates(Context, Domain))
+		{
+			PreverifyOk = 0;
+		}
+	}
+
+	return PreverifyOk;
+}
+
+CURLcode FHttpRequest::StaticSSLCTXFn(CURL * curl, void * sslctx, void * parm)
+{
+	SSL_CTX* Context = static_cast<SSL_CTX*>(sslctx);
+	const ISslCertificateManager& CertificateManager = FSslModule::Get().GetCertificateManager();
+
+	CertificateManager.AddCertificatesToSslContext(Context);
+	SSL_CTX_set_verify(Context, SSL_CTX_get_verify_mode(Context), SslCertVerify);
+	SSL_CTX_set_app_data(Context, parm);
+
+	/* all set to go */
+	return CURLE_OK;
+}
+#endif //#if WITH_SSL
+
+size_t FHttpRequest::StaticDebugCallback(CURL * Handle, curl_infotype DebugInfoType, char * DebugInfo, size_t DebugInfoSize, void* UserData)
+{
+	FHttpRequest* Request = static_cast<FHttpRequest*>(UserData);
+
+	switch (DebugInfoType)
+	{
+	case CURLINFO_TEXT:
+	{
+		// Truncate at 1023 characters. This is just an arbitrary number based on a buffer size seen in
+		// the libcurl code.
+		DebugInfoSize = FMath::Min(DebugInfoSize, (size_t)1023);
+
+		// Calculate the actual length of the string due to incorrect use of snprintf() in lib/vtls/openssl.c.
+		char* FoundNulPtr = (char*)memchr(DebugInfo, 0, DebugInfoSize);
+		int CalculatedSize = FoundNulPtr != nullptr ? FoundNulPtr - DebugInfo : DebugInfoSize;
+
+		auto ConvertedString = StringCast<TCHAR>(static_cast<const ANSICHAR*>(DebugInfo), CalculatedSize);
+		FString DebugText(ConvertedString.Length(), ConvertedString.Get());
+		DebugText.ReplaceInline(TEXT("\n"), TEXT(""), ESearchCase::CaseSensitive);
+		DebugText.ReplaceInline(TEXT("\r"), TEXT(""), ESearchCase::CaseSensitive);
+		UE_LOG(LogDerivedDataCache, VeryVerbose, TEXT("%s: %p: '%s'"), *Request->GetName(), Request, *DebugText);
+	}
+	break;
+
+	case CURLINFO_HEADER_IN:
+		UE_LOG(LogDerivedDataCache, VeryVerbose, TEXT("%s: %p: Received header (%d bytes)"), *Request->GetName(), Request, DebugInfoSize);
+		break;
+
+	case CURLINFO_DATA_IN:
+		UE_LOG(LogDerivedDataCache, VeryVerbose, TEXT("%s: %p: Received data (%d bytes)"), *Request->GetName(), Request, DebugInfoSize);
+		break;
+
+	case CURLINFO_DATA_OUT:
+		UE_LOG(LogDerivedDataCache, VeryVerbose, TEXT("%s: %p: Sent data (%d bytes)"), *Request->GetName(), Request, DebugInfoSize);
+		break;
+
+	case CURLINFO_SSL_DATA_IN:
+		UE_LOG(LogDerivedDataCache, VeryVerbose, TEXT("%s: %p: Received SSL data (%d bytes)"), *Request->GetName(), Request, DebugInfoSize);
+		break;
+
+	case CURLINFO_SSL_DATA_OUT:
+		UE_LOG(LogDerivedDataCache, VeryVerbose, TEXT("%s: %p: Sent SSL data (%d bytes)"), *Request->GetName(), Request, DebugInfoSize);
+		break;
+	}
+
+	return 0;
+}
+
+size_t FHttpRequest::StaticReadFn(void* Ptr, size_t SizeInBlocks, size_t BlockSizeInBytes, void* UserData)
+{
+	FHttpRequest* Request = static_cast<FHttpRequest*>(UserData);
+	const size_t MaxReadSize = SizeInBlocks * BlockSizeInBytes;
+	const FMemoryView SourceView = Request->ReadDataView.Mid(Request->BytesSent, MaxReadSize);
+	MakeMemoryView(Ptr, MaxReadSize).CopyFrom(SourceView);
+	Request->BytesSent += SourceView.GetSize();
+	return SourceView.GetSize();
+}
+
+size_t FHttpRequest::StaticWriteHeaderFn(void* Ptr, size_t SizeInBlocks, size_t BlockSizeInBytes, void* UserData)
+{
+	FHttpRequest* Request = static_cast<FHttpRequest*>(UserData);
+	const size_t WriteSize = SizeInBlocks * BlockSizeInBytes;
+	TArray<uint8>* WriteHeaderBufferPtr = Request->WriteHeaderBufferPtr;
+	if (WriteHeaderBufferPtr && WriteSize > 0)
+	{
+		const size_t CurrentBufferLength = WriteHeaderBufferPtr->Num();
+		if (CurrentBufferLength > 0)
+		{
+			// Remove the previous zero termination
+			(*WriteHeaderBufferPtr)[CurrentBufferLength-1] = ' ';
+		}
+
+		// Write the header
+		WriteHeaderBufferPtr->Append((const uint8*)Ptr, WriteSize + 1);
+		(*WriteHeaderBufferPtr)[WriteHeaderBufferPtr->Num()-1] = 0; // Zero terminate string
+		return WriteSize;
+	}
+	return 0;
+}
+
+size_t FHttpRequest::StaticWriteBodyFn(void* Ptr, size_t SizeInBlocks, size_t BlockSizeInBytes, void* UserData)
+{
+	FHttpRequest* Request = static_cast<FHttpRequest*>(UserData);
+	const size_t WriteSize = SizeInBlocks * BlockSizeInBytes;
+	TArray<uint8>* WriteDataBufferPtr = Request->WriteDataBufferPtr;
+
+	if (WriteDataBufferPtr && WriteSize > 0)
+	{
+		// If this is the first part of the body being received, try to reserve 
+		// memory if content length is defined in the header.
+		if (Request->BytesReceived == 0 && Request->WriteHeaderBufferPtr)
+		{
+			static const ANSICHAR* ContentLengthHeaderStr = "Content-Length: ";
+			const ANSICHAR* Header = (const ANSICHAR*)Request->WriteHeaderBufferPtr->GetData();
+
+			if (const ANSICHAR* ContentLengthHeader = FCStringAnsi::Strstr(Header, ContentLengthHeaderStr))
+			{
+				size_t ContentLength = (size_t)FCStringAnsi::Atoi64(ContentLengthHeader + strlen(ContentLengthHeaderStr));
+				if (ContentLength > 0u && ContentLength < UE_HTTPDDC_MAX_BUFFER_RESERVE)
+				{
+					WriteDataBufferPtr->Reserve(ContentLength);
+				}
+			}
+		}
+
+		// Write to the target buffer
+		WriteDataBufferPtr->Append((const uint8*)Ptr, WriteSize);
+		Request->BytesReceived += WriteSize;
+		return WriteSize;
+	}
+	
+	return 0;
+}
+
+size_t FHttpRequest::StaticSeekFn(void* UserData, curl_off_t Offset, int Origin)
+{
+	FHttpRequest* Request = static_cast<FHttpRequest*>(UserData);
+	size_t NewPosition = 0;
+
+	switch (Origin)
+	{
+	case SEEK_SET: NewPosition = Offset; break;
+	case SEEK_CUR: NewPosition = Request->BytesSent + Offset; break;
+	case SEEK_END: NewPosition = Request->ReadDataView.GetSize() + Offset; break;
+	}
+
+	// Make sure we don't seek outside of the buffer
+	if (NewPosition < 0 || NewPosition >= Request->ReadDataView.GetSize())
+	{
+		return CURL_SEEKFUNC_FAIL;
+	}
+
+	// Update the used offset
+	Request->BytesSent = NewPosition;
+	return CURL_SEEKFUNC_OK;
+}
+
+} // UE
diff --git a/Engine/Source/Developer/DerivedDataCache/Private/HttpCacheStore.cpp b/Engine/Source/Developer/DerivedDataCache/Private/HttpCacheStore.cpp
--- a/Engine/Source/Developer/DerivedDataCache/Private/HttpCacheStore.cpp
+++ b/Engine/Source/Developer/DerivedDataCache/Private/HttpCacheStore.cpp
@@ -56,1294 +56,67 @@
 #include "Serialization/JsonWriter.h"
 #include "Tasks/Task.h"
 
-#if WITH_SSL
-#include "Ssl.h"
-#include <openssl/ssl.h>
-#endif
-
-// Enables data request helpers that internally
-// batch requests to reduce the number of concurrent
-// connections.
-#ifndef WITH_DATAREQUEST_HELPER
-	#define WITH_DATAREQUEST_HELPER 1
-#endif
-
-#define UE_HTTPDDC_BACKEND_WAIT_INTERVAL 0.01f
-#define UE_HTTPDDC_BACKEND_WAIT_INTERVAL_MS ((uint32)(UE_HTTPDDC_BACKEND_WAIT_INTERVAL*1000))
-#define UE_HTTPDDC_HTTP_REQUEST_TIMEOUT_SECONDS 30L
-#define UE_HTTPDDC_HTTP_REQUEST_TIMEOUT_ENABLED 1
-#define UE_HTTPDDC_HTTP_DEBUG 0
-#define UE_HTTPDDC_GET_REQUEST_POOL_SIZE 48
-#define UE_HTTPDDC_PUT_REQUEST_POOL_SIZE 16
-#define UE_HTTPDDC_NONBLOCKING_REQUEST_POOL_SIZE 128
-#define UE_HTTPDDC_MAX_FAILED_LOGIN_ATTEMPTS 16
-#define UE_HTTPDDC_MAX_ATTEMPTS 4
-#define UE_HTTPDDC_MAX_BUFFER_RESERVE 104857600u
-#define UE_HTTPDDC_BATCH_SIZE 12
-#define UE_HTTPDDC_BATCH_NUM 64
-#define UE_HTTPDDC_BATCH_GET_WEIGHT 4
-#define UE_HTTPDDC_BATCH_HEAD_WEIGHT 1
-#define UE_HTTPDDC_BATCH_WEIGHT_HINT 12
-
-namespace UE::DerivedData
-{
-
-TRACE_DECLARE_INT_COUNTER(HttpDDC_Exist, TEXT("HttpDDC Exist"));
-TRACE_DECLARE_INT_COUNTER(HttpDDC_ExistHit, TEXT("HttpDDC Exist Hit"));
-TRACE_DECLARE_INT_COUNTER(HttpDDC_Get, TEXT("HttpDDC Get"));
-TRACE_DECLARE_INT_COUNTER(HttpDDC_GetHit, TEXT("HttpDDC Get Hit"));
-TRACE_DECLARE_INT_COUNTER(HttpDDC_Put, TEXT("HttpDDC Put"));
-TRACE_DECLARE_INT_COUNTER(HttpDDC_PutHit, TEXT("HttpDDC Put Hit"));
-TRACE_DECLARE_INT_COUNTER(HttpDDC_BytesReceived, TEXT("HttpDDC Bytes Received"));
-TRACE_DECLARE_INT_COUNTER(HttpDDC_BytesSent, TEXT("HttpDDC Bytes Sent"));
-
-static bool bHttpEnableAsync = true;
-static FAutoConsoleVariableRef CVarHttpEnableAsync(
-	TEXT("DDC.Http.EnableAsync"),
-	bHttpEnableAsync,
-	TEXT("If true, async operations are permitted, otherwise all operations are forced to be synchronous."),
-	ECVF_Default);
-
-static CURLcode sslctx_function(CURL * curl, void * sslctx, void * parm);
-
-typedef TSharedPtr<class IHttpRequest> FHttpRequestPtr;
-typedef TSharedPtr<class IHttpResponse, ESPMode::ThreadSafe> FHttpResponsePtr;
-
-class FHttpCacheStoreRunnable;
-struct FRequestPool;
-
-/**
- * Encapsulation for access token shared by all requests.
- */
-struct FHttpAccessToken
-{
-public:
-	FHttpAccessToken() = default;
-	FString GetHeader();
-	void SetHeader(const TCHAR*);
-	uint32 GetSerial() const;
-private:
-	FRWLock		Lock;
-	FString		Token;
-	uint32		Serial;
-};
-
-class FHttpSharedData
-{
-public:
-	FHttpSharedData()
-	: PendingRequestEvent(EEventMode::AutoReset)
-	{
-		CurlShare = curl_share_init();
-		curl_share_setopt(CurlShare, CURLSHOPT_USERDATA, this);
-		curl_share_setopt(CurlShare, CURLSHOPT_LOCKFUNC, LockFn);
-		curl_share_setopt(CurlShare, CURLSHOPT_UNLOCKFUNC, UnlockFn);
-		curl_share_setopt(CurlShare, CURLSHOPT_SHARE, CURL_LOCK_DATA_DNS);
-		curl_share_setopt(CurlShare, CURLSHOPT_SHARE, CURL_LOCK_DATA_SSL_SESSION);
-		CurlMulti = curl_multi_init();
-		curl_multi_setopt(CurlMulti, CURLMOPT_PIPELINING, CURLPIPE_MULTIPLEX);
-	}
-
-	~FHttpSharedData()
-	{
-		bAsyncThreadShutdownRequested.store(true, std::memory_order_relaxed);
-		if (AsyncServiceThread.IsJoinable())
-		{
-			AsyncServiceThread.Join();
-		}
-
-		curl_multi_cleanup(CurlMulti);
-		curl_share_cleanup(CurlShare);
-	}
-
-	void AddRequest(CURL* Curl);
-
-	CURLSH* GetCurlShare() const { return CurlShare; }
-
-private:
-	CURLSH* CurlShare;
-	CURLM* CurlMulti;
-	TDepletableMpscQueue<CURL*> PendingRequestAdditions;
-	FLazyEvent PendingRequestEvent;
-	FThread AsyncServiceThread;
-	std::atomic<bool> bAsyncThreadStarting = false;
-	std::atomic<bool> bAsyncThreadShutdownRequested = false;
-
-	FRWLock Locks[CURL_LOCK_DATA_LAST];
-	bool WriteLocked[CURL_LOCK_DATA_LAST]{};
-
-	static void LockFn(CURL* Handle, curl_lock_data Data, curl_lock_access Access, void* User)
-	{
-		FHttpSharedData* SharedData = (FHttpSharedData*)User;
-		if (Access == CURL_LOCK_ACCESS_SHARED)
-		{
-			SharedData->Locks[Data].ReadLock();
-		}
-		else
-		{
-			SharedData->Locks[Data].WriteLock();
-			SharedData->WriteLocked[Data] = true;
-		}
-	}
-
-	static void UnlockFn(CURL* Handle, curl_lock_data Data, void* User)
-	{
-		FHttpSharedData* SharedData = (FHttpSharedData*)User;
-		if (!SharedData->WriteLocked[Data])
-		{
-			SharedData->Locks[Data].ReadUnlock();
-		}
-		else
-		{
-			SharedData->WriteLocked[Data] = false;
-			SharedData->Locks[Data].WriteUnlock();
-		}
-	}
-
-	void ProcessAsyncRequests();
-};
-
-template <typename T>
-class TRefCountedUniqueFunction final : public FThreadSafeRefCountedObject
-{
-public:
-	explicit TRefCountedUniqueFunction(T&& InFunction) : Function(MoveTemp(InFunction))
-	{
-	}
-
-	const T& GetFunction() const { return Function; }
-
-private:
-	T Function;
-};
-
-/**
- * Minimal HTTP request type wrapping CURL without the need for managers. This request
- * is written to allow reuse of request objects, in order to allow connections to be reused.
- *
- * CURL has a global library initialization (curl_global_init). We rely on this happening in 
- * the Online/HTTP library which is a dependency on this module.
- */
-class FHttpRequest
-{
-public:
-	/**
-	 * Supported request verbs
-	 */
-	enum RequestVerb
-	{
-		Get,
-		Put,
-		PutCompactBinary,
-		PutCompressedBlob,
-		Post,
-		PostCompactBinary,
-		PostJson,
-		Delete,
-		Head
-	};
-
-	/**
-	 * Convenience result type interpreted from HTTP response code.
-	 */
-	enum Result
-	{
-		Success,
-		Failed,
-		FailedTimeout
-	};
-
-	enum class ECompletionBehavior : uint8
-	{
-		Done,
-		Retry
-	};
-
-	using FOnHttpRequestComplete = TUniqueFunction<ECompletionBehavior(Result HttpResult, FHttpRequest* Request)>;
-
-	FHttpRequest(const TCHAR* InDomain, const TCHAR* InEffectiveDomain, FHttpAccessToken* InAuthorizationToken, FHttpSharedData* InSharedData, bool bInLogErrors)
-		: SharedData(InSharedData)
-		, AsyncData(nullptr)
-		, bLogErrors(bInLogErrors)
-		, Domain(InDomain)
-		, EffectiveDomain(InEffectiveDomain)
-		, AuthorizationToken(InAuthorizationToken)
-	{
-		Curl = curl_easy_init();
-		Reset();
-	}
-
-	~FHttpRequest()
-	{
-		curl_easy_cleanup(Curl);
-		check(!AsyncData);
-	}
-
-	/**
-	 * Resets all options on the request except those that should always be set.
-	 */
-	void Reset()
-	{
-		Headers.Reset();
-		ResponseHeader.Reset();
-		ResponseBuffer.Reset();
-		ResponseCode = 0;
-		ReadDataView = FMemoryView();
-		WriteDataBufferPtr = nullptr;
-		WriteHeaderBufferPtr = nullptr;
-		BytesSent = 0;
-		BytesReceived = 0;
-		Attempts = 0;
-		CurlResult = CURL_LAST;
-
-		curl_easy_reset(Curl);
-		check(!AsyncData);
-
-		// Options that are always set for all connections.
-#if UE_HTTPDDC_HTTP_REQUEST_TIMEOUT_ENABLED
-		curl_easy_setopt(Curl, CURLOPT_CONNECTTIMEOUT, UE_HTTPDDC_HTTP_REQUEST_TIMEOUT_SECONDS);
-#endif
-		curl_easy_setopt(Curl, CURLOPT_FOLLOWLOCATION, 1L);
-		curl_easy_setopt(Curl, CURLOPT_NOSIGNAL, 1L);
-		curl_easy_setopt(Curl, CURLOPT_DNS_CACHE_TIMEOUT, 300L); // Don't re-resolve every minute
-		curl_easy_setopt(Curl, CURLOPT_SHARE, SharedData ? SharedData->GetCurlShare() : nullptr);
-		// SSL options
-		curl_easy_setopt(Curl, CURLOPT_USE_SSL, CURLUSESSL_ALL);
-		curl_easy_setopt(Curl, CURLOPT_SSL_VERIFYPEER, 1);
-		curl_easy_setopt(Curl, CURLOPT_SSL_VERIFYHOST, 1);
-		curl_easy_setopt(Curl, CURLOPT_SSLCERTTYPE, "PEM");
-		// Response functions
-		curl_easy_setopt(Curl, CURLOPT_HEADERDATA, this);
-		curl_easy_setopt(Curl, CURLOPT_HEADERFUNCTION, &FHttpRequest::StaticWriteHeaderFn);
-		curl_easy_setopt(Curl, CURLOPT_WRITEDATA, this);
-		curl_easy_setopt(Curl, CURLOPT_WRITEFUNCTION, StaticWriteBodyFn);
-		// SSL certification verification
-		curl_easy_setopt(Curl, CURLOPT_CAINFO, nullptr);
-		curl_easy_setopt(Curl, CURLOPT_SSL_CTX_FUNCTION, *sslctx_function);
-		curl_easy_setopt(Curl, CURLOPT_SSL_CTX_DATA, this);
-		// Allow compressed data
-		curl_easy_setopt(Curl, CURLOPT_ACCEPT_ENCODING, "gzip");
-		// Rewind method, handle special error case where request need to rewind data stream
-		curl_easy_setopt(Curl, CURLOPT_SEEKFUNCTION, StaticSeekFn);
-		curl_easy_setopt(Curl, CURLOPT_SEEKDATA, this);
-		// Set minimum speed behavior to allow operations to abort if the transfer speed is poor for the given duration (1kbps over a 30 second span)
-		curl_easy_setopt(Curl, CURLOPT_LOW_SPEED_TIME, 30L);
-		curl_easy_setopt(Curl, CURLOPT_LOW_SPEED_LIMIT, 1024L);
-		// Debug hooks
-#if UE_HTTPDDC_HTTP_DEBUG
-		curl_easy_setopt(Curl, CURLOPT_DEBUGDATA, this);
-		curl_easy_setopt(Curl, CURLOPT_DEBUGFUNCTION, StaticDebugCallback);
-		curl_easy_setopt(Curl, CURLOPT_VERBOSE, 1L);
-#endif
-		curl_easy_setopt(Curl, CURLOPT_PRIVATE, this);
-		curl_easy_setopt(Curl, CURLOPT_HTTP_VERSION, CURL_HTTP_VERSION_2_0);
-	}
-
-	void PrepareToRetry()
-	{
-		ResponseHeader.Reset();
-		ResponseBuffer.Reset();
-		ResponseCode = 0;
-		BytesSent = 0;
-		BytesReceived = 0;
-		CurlResult = CURL_LAST;
-		++Attempts;
-	}
-
-	/** Gets the domain name for this request */
-	const FString& GetName() const
-	{
-		return Domain;
-	}
-
-	/** Gets the domain name for this request */
-	const FString& GetDomain() const
-	{
-		return Domain;
-	}
-
-	/** Gets the effective domain name for this request */
-	const FString& GetEffectiveDomain() const
-	{
-		return EffectiveDomain;
-	}
-
-	/** Returns the HTTP response code.*/
-	const int64 GetResponseCode() const
-	{
-		return ResponseCode;
-	}
-
-	/** Returns the number of bytes received this request (headers withstanding). */
-	const size_t GetBytesReceived() const
-	{
-		return BytesReceived;
-	}
-
-	/** Returns the number of attempts we've made issuing this request (currently tracked for async requests only). */
-	const size_t GetAttempts() const
-	{
-		return Attempts;
-	}
-
-	/** Returns the number of bytes sent during this request (headers withstanding). */
-	const size_t GetBytesSent() const
-	{
-		return BytesSent;
-	}
-
-	/**
-	 * Upload buffer using the request, using either "Put" or "Post" verbs.
-	 * @param Uri Url to use.
-	 * @param Buffer Data to upload
-	 * @return Result of the request
-	 */
-	template<RequestVerb V>
-	Result PerformBlockingUpload(const TCHAR* Uri, TArrayView<const uint8> Buffer, TConstArrayView<long> ExpectedErrorCodes = {})
-	{
-		static_assert(V == Put || V == PutCompactBinary || V == PutCompressedBlob || V == Post || V == PostCompactBinary || V == PostJson, "Upload should use either Put or Post verbs.");
-		
-		uint64 ContentLength = 0u;
-
-		if constexpr (V == Put || V == PutCompactBinary || V == PutCompressedBlob)
-		{
-			curl_easy_setopt(Curl, CURLOPT_UPLOAD, 1L);
-			curl_easy_setopt(Curl, CURLOPT_INFILESIZE, Buffer.Num());
-			curl_easy_setopt(Curl, CURLOPT_READDATA, this);
-			curl_easy_setopt(Curl, CURLOPT_READFUNCTION, StaticReadFn);
-			if constexpr (V == PutCompactBinary)
-			{
-				Headers.Add(FString(TEXT("Content-Type: application/x-ue-cb")));
-			}
-			else if constexpr (V == PutCompressedBlob)
-			{
-				Headers.Add(FString(TEXT("Content-Type: application/x-ue-comp")));
-			}
-			else
-			{
-				Headers.Add(FString(TEXT("Content-Type: application/octet-stream")));
-			}
-			ContentLength = Buffer.Num();
-			ReadDataView = MakeMemoryView(Buffer);
-		}
-		else if constexpr (V == Post || V == PostCompactBinary || V == PostJson)
-		{
-			curl_easy_setopt(Curl, CURLOPT_POST, 1L);
-			curl_easy_setopt(Curl, CURLOPT_INFILESIZE, Buffer.Num());
-			curl_easy_setopt(Curl, CURLOPT_READDATA, this);
-			curl_easy_setopt(Curl, CURLOPT_READFUNCTION, StaticReadFn);
-			if constexpr (V == PostCompactBinary)
-			{
-				Headers.Add(FString(TEXT("Content-Type: application/x-ue-cb")));
-			}
-			else if constexpr (V == PostJson)
-			{
-				Headers.Add(FString(TEXT("Content-Type: application/json")));
-			}
-			else
-			{
-				Headers.Add(FString(TEXT("Content-Type: application/x-www-form-urlencoded")));
-			}
-			ContentLength = Buffer.Num();
-			ReadDataView = MakeMemoryView(Buffer);
-		}
-
-		return PerformBlocking(Uri, V, ContentLength, ExpectedErrorCodes);
-	}
-	
-	template<RequestVerb V>
-	void EnqueueAsyncUpload(IRequestOwner& Owner, FRequestPool* Pool, const TCHAR* Uri, FSharedBuffer Buffer, FOnHttpRequestComplete&& OnComplete, TConstArrayView<long> ExpectedErrorCodes = {})
-	{
-		static_assert(V == Put || V == PutCompactBinary || V == PutCompressedBlob || V == Post || V == PostCompactBinary || V == PostJson, "Upload should use either Put or Post verbs.");
-		
-		uint64 ContentLength = 0u;
-
-		if constexpr (V == Put || V == PutCompactBinary || V == PutCompressedBlob)
-		{
-			curl_easy_setopt(Curl, CURLOPT_UPLOAD, 1L);
-			curl_easy_setopt(Curl, CURLOPT_INFILESIZE, Buffer.GetSize());
-			curl_easy_setopt(Curl, CURLOPT_READDATA, this);
-			curl_easy_setopt(Curl, CURLOPT_READFUNCTION, StaticReadFn);
-			if constexpr (V == PutCompactBinary)
-			{
-				Headers.Add(FString(TEXT("Content-Type: application/x-ue-cb")));
-			}
-			else if constexpr (V == PutCompressedBlob)
-			{
-				Headers.Add(FString(TEXT("Content-Type: application/x-ue-comp")));
-			}
-			else
-			{
-				Headers.Add(FString(TEXT("Content-Type: application/octet-stream")));
-			}
-			ReadSharedBuffer = Buffer;
-			ContentLength = Buffer.GetSize();
-			ReadDataView = Buffer.GetView();
-		}
-		else if constexpr (V == Post || V == PostCompactBinary || V == PostJson)
-		{
-			curl_easy_setopt(Curl, CURLOPT_POST, 1L);
-			curl_easy_setopt(Curl, CURLOPT_INFILESIZE, Buffer.GetSize());
-			curl_easy_setopt(Curl, CURLOPT_READDATA, this);
-			curl_easy_setopt(Curl, CURLOPT_READFUNCTION, StaticReadFn);
-			if constexpr (V == PostCompactBinary)
-			{
-				Headers.Add(FString(TEXT("Content-Type: application/x-ue-cb")));
-			}
-			else if constexpr (V == PostJson)
-			{
-				Headers.Add(FString(TEXT("Content-Type: application/json")));
-			}
-			else
-			{
-				Headers.Add(FString(TEXT("Content-Type: application/x-www-form-urlencoded")));
-			}
-			ReadSharedBuffer = Buffer;
-			ContentLength = Buffer.GetSize();
-			ReadDataView = Buffer.GetView();
-		}
-
-		return EnqueueAsync(Owner, Pool, Uri, V, ContentLength, MoveTemp(OnComplete), ExpectedErrorCodes);
-	}
-
-	/**
-	 * Download an url into a buffer using the request.
-	 * @param Uri Url to use.
-	 * @param Buffer Optional buffer where data should be downloaded to. If empty downloaded data will
-	 * be stored in an internal buffer and accessed GetResponse* methods.
-	 * @return Result of the request
-	 */
-	Result PerformBlockingDownload(const TCHAR* Uri, TArray<uint8>* Buffer, TConstArrayView<long> ExpectedErrorCodes = {400})
-	{
-		curl_easy_setopt(Curl, CURLOPT_HTTPGET, 1L);
-		WriteDataBufferPtr = Buffer;
-
-		return PerformBlocking(Uri, Get, 0u, ExpectedErrorCodes);
-	}
-
-	void EnqueueAsyncDownload(IRequestOwner& Owner, FRequestPool* Pool, const TCHAR* Uri, FOnHttpRequestComplete&& OnComplete, TConstArrayView<long> ExpectedErrorCodes = {400})
-	{
-		curl_easy_setopt(Curl, CURLOPT_HTTPGET, 1L);
-
-		return EnqueueAsync(Owner, Pool, Uri, Get, 0u, MoveTemp(OnComplete), ExpectedErrorCodes);
-	}
-
-	/**
-	 * Query an url using the request. Queries can use either "Head" or "Delete" verbs.
-	 * @param Uri Url to use.
-	 * @return Result of the request
-	 */
-	template<RequestVerb V>
-	Result PerformBlockingQuery(const TCHAR* Uri, TConstArrayView<long> ExpectedErrorCodes = {400})
-	{
-		static_assert(V == Head || V == Delete, "Queries should use either Head or Delete verbs.");
-
-		if (V == Delete)
-		{
-			curl_easy_setopt(Curl, CURLOPT_CUSTOMREQUEST, "DELETE");
-		}
-		else if (V == Head)
-		{
-			curl_easy_setopt(Curl, CURLOPT_NOBODY, 1L);
-		}
-
-		return PerformBlocking(Uri, V, 0u, ExpectedErrorCodes);
-	}
-
-	template<RequestVerb V>
-	void EnqueueAsyncQuery(const TCHAR* Uri, FOnHttpRequestComplete&& OnComplete, TConstArrayView<long> ExpectedErrorCodes = {400})
-	{
-		static_assert(V == Head || V == Delete, "Queries should use either Head or Delete verbs.");
-
-		if (V == Delete)
-		{
-			curl_easy_setopt(Curl, CURLOPT_CUSTOMREQUEST, "DELETE");
-		}
-		else if (V == Head)
-		{
-			curl_easy_setopt(Curl, CURLOPT_NOBODY, 1L);
-		}
-
-		return EnqueueAsync(Uri, V, 0u, MoveTemp(OnComplete), ExpectedErrorCodes);
-	}
-
-	/**
-	 * Set a header to send with the request.
-	 */
-	void SetHeader(const TCHAR* Header, const TCHAR* Value)
-	{
-		check(CurlResult == CURL_LAST); // Cannot set header after request is sent
-		Headers.Add(FString::Printf(TEXT("%s: %s"), Header, Value));
-	}
-
-	/**
-	 * Attempts to find the header from the response. Returns false if header is not present.
-	 */
-	bool GetHeader(const ANSICHAR* Header, FString& OutValue) const
-	{
-		check(CurlResult != CURL_LAST);  // Cannot query headers before request is sent
-
-		const ANSICHAR* HeadersBuffer = (const ANSICHAR*) ResponseHeader.GetData();
-		size_t HeaderLen = strlen(Header);
-
-		// Find the header key in the (ANSI) response buffer. If not found we can exist immediately
-		if (const ANSICHAR* Found = FCStringAnsi::Stristr(HeadersBuffer, Header))
-		{
-			const ANSICHAR* Linebreak = strchr(Found, '\r');
-			const ANSICHAR* ValueStart = Found + HeaderLen + 2; //colon and space
-			const size_t ValueSize = Linebreak - ValueStart;
-			FUTF8ToTCHAR TCHARData(ValueStart, ValueSize);
-			OutValue = FString(TCHARData.Length(), TCHARData.Get());
-			return true;
-		}
-		return false;
-	}
-
-	/**
-	 * Returns the response buffer. Note that is the request is performed
-	 * with an external buffer as target buffer this string will be empty.
-	 */
-	const TArray<uint8>& GetResponseBuffer() const
-	{
-		return ResponseBuffer;
-	}
-
-	FSharedBuffer MoveResponseBufferToShared()
-	{
-		return MakeSharedBufferFromArray(MoveTemp(ResponseBuffer));
-	}
-
-	/**
-	 * Returns the response buffer as a string. Note that is the request is performed
-	 * with an external buffer as target buffer this string will be empty.
-	 */
-	FString GetResponseAsString() const
-	{
-		return GetAnsiBufferAsString(ResponseBuffer);
-	}
-
-	/**
-	 * Returns the response header as a string.
-	 */
-	FString GetResponseHeaderAsString()
-	{
-		return GetAnsiBufferAsString(ResponseHeader);
-	}
-
-	/**
-	 * Tries to parse the response buffer as a JsonObject. Return empty pointer if 
-	 * parse error occurs.
-	 */
-	TSharedPtr<FJsonObject> GetResponseAsJsonObject() const
-	{
-		FString Response = GetAnsiBufferAsString(ResponseBuffer);
-
-		TSharedPtr<FJsonObject> JsonObject;
-		TSharedRef<TJsonReader<> > JsonReader = TJsonReaderFactory<>::Create(Response);
-		if (!FJsonSerializer::Deserialize(JsonReader, JsonObject) || !JsonObject.IsValid())
-		{
-			return TSharedPtr<FJsonObject>(nullptr);
-		}
-
-		return JsonObject;
-	}
-
-	/**
-	 * Tries to parse the response buffer as a JsonArray. Return empty array if
-	 * parse error occurs.
-	 */
-	TArray<TSharedPtr<FJsonValue>> GetResponseAsJsonArray() const
-	{
-		FString Response = GetAnsiBufferAsString(ResponseBuffer);
-
-		TArray<TSharedPtr<FJsonValue>> JsonArray;
-		TSharedRef<TJsonReader<>> JsonReader = TJsonReaderFactory<>::Create(Response);
-		FJsonSerializer::Deserialize(JsonReader, JsonArray);
-		return JsonArray;
-	}
-
-	void CompleteAsync(CURLcode Result);
-
-	/** Will return true if the response code is considered a success */
-	static bool IsSuccessResponse(long ResponseCode)
-	{
-		// We consider anything in the 1XX or 2XX range a success
-		return ResponseCode >= 100 && ResponseCode < 300;
-	}
-
-	static bool AllowAsync();
-
-private:
-
-	struct FAsyncRequestData final : public FRequestBase
-	{
-		IRequestOwner* Owner = nullptr;
-		FRequestPool* Pool = nullptr;
-		curl_slist* CurlHeaders = nullptr;
-		FString Uri;
-		RequestVerb Verb;
-		TArray<long, TInlineAllocator<4>> ExpectedErrorCodes;
-		FOnHttpRequestComplete OnComplete;
-		FLazyEvent Event {EEventMode::ManualReset};
-
-		void Reset()
-		{
-			if (CurlHeaders)
-			{
-				curl_slist_free_all(CurlHeaders);
-				CurlHeaders = nullptr;
-			}
-			Uri.Empty();
-			ExpectedErrorCodes.Empty();
-		}
-
-		void SetPriority(EPriority Priority) final {}
-
-		void Cancel() final
-		{
-			TRACE_CPUPROFILER_EVENT_SCOPE(HttpDDC_Cancel);
-			Event.Wait();
-		}
-
-		void Wait() final
-		{
-			TRACE_CPUPROFILER_EVENT_SCOPE(HttpDDC_Wait);
-			Event.Wait();
-		}
-	};
-
-	CURL*					Curl;
-	CURLcode				CurlResult;
-	FHttpSharedData*		SharedData;
-	FAsyncRequestData*		AsyncData;
-	long					ResponseCode;
-	size_t					BytesSent;
-	size_t					BytesReceived;
-	size_t					Attempts;
-	bool					bLogErrors;
-
-	FSharedBuffer			ReadSharedBuffer;
-	FMemoryView				ReadDataView;
-	TArray<uint8>*			WriteDataBufferPtr;
-	TArray<uint8>*			WriteHeaderBufferPtr;
-	
-	TArray<uint8>			ResponseHeader;
-	TArray<uint8>			ResponseBuffer;
-	TArray<FString>			Headers;
-	FString					Domain;
-	FString					EffectiveDomain;
-	FHttpAccessToken*		AuthorizationToken;
-
-	curl_slist* PrepareToIssueRequest(const TCHAR* Uri, uint64 ContentLength)
-	{
-		static const char* CommonHeaders[] = {
-			"User-Agent: Unreal Engine",
-			nullptr
-		};
-
-		// Setup request options
-		FString Url = FString::Printf(TEXT("%s/%s"), *EffectiveDomain, Uri);
-		curl_easy_setopt(Curl, CURLOPT_URL, TCHAR_TO_ANSI(*Url));
-
-		// Setup response header buffer. If caller has not setup a response data buffer, use internal.
-		WriteHeaderBufferPtr = &ResponseHeader;
-		if (WriteDataBufferPtr == nullptr)
-		{
-			WriteDataBufferPtr = &ResponseBuffer;
-		}
-
-		// Content-Length should always be set
-		Headers.Add(FString::Printf(TEXT("Content-Length: %d"), ContentLength));
-
-		// And auth token if it's set
-		if (AuthorizationToken)
-		{
-			Headers.Add(AuthorizationToken->GetHeader());
-		}
-
-		// Build headers list
-		curl_slist* CurlHeaders = nullptr;
-		// Add common headers
-		for (uint8 i = 0; CommonHeaders[i] != nullptr; ++i)
-		{
-			CurlHeaders = curl_slist_append(CurlHeaders, CommonHeaders[i]);
-		}
-		// Setup added headers
-		for (const FString& Header : Headers)
-		{
-			CurlHeaders = curl_slist_append(CurlHeaders, TCHAR_TO_ANSI(*Header));
-		}
-		curl_easy_setopt(Curl, CURLOPT_HTTPHEADER, CurlHeaders);
-		return CurlHeaders;
-	}
-
-	/**
-	 * Performs the request, blocking until finished.
-	 * @param Uri Address on the domain to query
-	 * @param Verb HTTP verb to use
-	 * @param ContentLength The number of bytes being uploaded for the body of this request.
-	 * @param ExpectedErrorCodes An array of expected return codes outside of the success range that should NOT be logged as an abnormal/exceptional outcome.
-	 * If unset the response body will be stored in the request.
-	 */
-	Result PerformBlocking(const TCHAR* Uri, RequestVerb Verb, uint64 ContentLength, TConstArrayView<long> ExpectedErrorCodes)
-	{
-		TRACE_CPUPROFILER_EVENT_SCOPE(HttpDDC_CurlPerform);
-
-		// Build headers list
-		curl_slist* CurlHeaders = PrepareToIssueRequest(Uri, ContentLength);
-
-		// Shots fired!
-		CurlResult = curl_easy_perform(Curl);
-
-		// Get response code
-		curl_easy_getinfo(Curl, CURLINFO_RESPONSE_CODE, &ResponseCode);
-
-		LogResult(CurlResult, Uri, Verb, ExpectedErrorCodes);
-
-		// Clean up
-		curl_slist_free_all(CurlHeaders);
-
-		return CurlResult == CURLE_OK ? Success : Failed;
-	}
-
-	void EnqueueAsync(IRequestOwner& Owner, FRequestPool* Pool, const TCHAR* Uri, RequestVerb Verb, uint64 ContentLength, FOnHttpRequestComplete&& OnComplete, TConstArrayView<long> ExpectedErrorCodes);
-
-	void LogResult(CURLcode Result, const TCHAR* Uri, RequestVerb Verb, TConstArrayView<long> ExpectedErrorCodes) const
-	{
-		if (Result == CURLE_OK)
-		{
-			bool bSuccess = false;
-			const TCHAR* VerbStr = nullptr;
-			FString AdditionalInfo;
-
-			switch (Verb)
-			{
-			case Head:
-				bSuccess = (ExpectedErrorCodes.Contains(ResponseCode) || IsSuccessResponse(ResponseCode));
-				VerbStr = TEXT("querying");
-				break;
-			case Get:
-				bSuccess = (ExpectedErrorCodes.Contains(ResponseCode) || IsSuccessResponse(ResponseCode));
-				VerbStr = TEXT("fetching");
-				AdditionalInfo = FString::Printf(TEXT("Received: %d bytes."), BytesReceived);
-				break;
-			case Put:
-			case PutCompactBinary:
-			case PutCompressedBlob:
-				bSuccess = (ExpectedErrorCodes.Contains(ResponseCode) || IsSuccessResponse(ResponseCode));
-				VerbStr = TEXT("updating");
-				AdditionalInfo = FString::Printf(TEXT("Sent: %d bytes."), BytesSent);
-				break;
-			case Post:
-			case PostCompactBinary:
-			case PostJson:
-				bSuccess = (ExpectedErrorCodes.Contains(ResponseCode) || IsSuccessResponse(ResponseCode));
-				VerbStr = TEXT("posting");
-				break;
-			case Delete:
-				bSuccess = (ExpectedErrorCodes.Contains(ResponseCode) || IsSuccessResponse(ResponseCode));
-				VerbStr = TEXT("deleting");
-				break;
-			}
-
-			if (bSuccess)
-			{
-				UE_LOG(
-					LogDerivedDataCache, 
-					Verbose, 
-					TEXT("%s: Finished %s HTTP cache entry (response %d) from %s. %s"), 
-					*GetName(),
-					VerbStr,
-					ResponseCode, 
-					Uri,
-					*AdditionalInfo
-				);
-			}
-			else if(bLogErrors)
-			{
-				// Print the response body if we got one, otherwise print header.
-				FString Response = GetAnsiBufferAsString(ResponseBuffer.Num() > 0 ? ResponseBuffer : ResponseHeader);
-				Response.ReplaceCharInline(TEXT('\n'), TEXT(' '));
-				Response.ReplaceCharInline(TEXT('\r'), TEXT(' '));
-				// Dont log access denied as error, since tokens can expire mid session
-				if (ResponseCode == 401)
-				{
-					UE_LOG(
-						LogDerivedDataCache,
-						Verbose,
-						TEXT("%s: Failed %s HTTP cache entry (response %d) from %s. Response: %s"),
-						*GetName(),
-						VerbStr,
-						ResponseCode,
-						Uri,
-						*Response
-					);
-				}
-				else
-				{ 
-					UE_LOG(
-						LogDerivedDataCache,
-						Display,
-						TEXT("%s: Failed %s HTTP cache entry (response %d) from %s. Response: %s"),
-						*GetName(),
-						VerbStr,
-						ResponseCode,
-						Uri,
-						*Response
-					);
-				}
-			}
-		}
-		else if(bLogErrors)
-		{
-			UE_LOG(
-				LogDerivedDataCache, 
-				Display, 
-				TEXT("%s: Error while connecting to %s: %s"), 
-				*GetName(),
-				*EffectiveDomain,
-				ANSI_TO_TCHAR(curl_easy_strerror(Result))
-			);
-		}
-	}
-
-	FString GetAnsiBufferAsString(const TArray<uint8>& Buffer) const
-	{
-		// Content is NOT null-terminated; we need to specify lengths here
-		FUTF8ToTCHAR TCHARData(reinterpret_cast<const ANSICHAR*>(Buffer.GetData()), Buffer.Num());
-		return FString(TCHARData.Length(), TCHARData.Get());
-	}
-
-	static size_t StaticDebugCallback(CURL * Handle, curl_infotype DebugInfoType, char * DebugInfo, size_t DebugInfoSize, void* UserData)
-	{
-		FHttpRequest* Request = static_cast<FHttpRequest*>(UserData);
-
-		switch (DebugInfoType)
-		{
-		case CURLINFO_TEXT:
-		{
-			// Truncate at 1023 characters. This is just an arbitrary number based on a buffer size seen in
-			// the libcurl code.
-			DebugInfoSize = FMath::Min(DebugInfoSize, (size_t)1023);
-
-			// Calculate the actual length of the string due to incorrect use of snprintf() in lib/vtls/openssl.c.
-			char* FoundNulPtr = (char*)memchr(DebugInfo, 0, DebugInfoSize);
-			int CalculatedSize = FoundNulPtr != nullptr ? FoundNulPtr - DebugInfo : DebugInfoSize;
-
-			auto ConvertedString = StringCast<TCHAR>(static_cast<const ANSICHAR*>(DebugInfo), CalculatedSize);
-			FString DebugText(ConvertedString.Length(), ConvertedString.Get());
-			DebugText.ReplaceInline(TEXT("\n"), TEXT(""), ESearchCase::CaseSensitive);
-			DebugText.ReplaceInline(TEXT("\r"), TEXT(""), ESearchCase::CaseSensitive);
-			UE_LOG(LogDerivedDataCache, VeryVerbose, TEXT("%s: %p: '%s'"), *Request->GetName(), Request, *DebugText);
-		}
-		break;
-
-		case CURLINFO_HEADER_IN:
-			UE_LOG(LogDerivedDataCache, VeryVerbose, TEXT("%s: %p: Received header (%d bytes)"), *Request->GetName(), Request, DebugInfoSize);
-			break;
-
-		case CURLINFO_DATA_IN:
-			UE_LOG(LogDerivedDataCache, VeryVerbose, TEXT("%s: %p: Received data (%d bytes)"), *Request->GetName(), Request, DebugInfoSize);
-			break;
-
-		case CURLINFO_DATA_OUT:
-			UE_LOG(LogDerivedDataCache, VeryVerbose, TEXT("%s: %p: Sent data (%d bytes)"), *Request->GetName(), Request, DebugInfoSize);
-			break;
-
-		case CURLINFO_SSL_DATA_IN:
-			UE_LOG(LogDerivedDataCache, VeryVerbose, TEXT("%s: %p: Received SSL data (%d bytes)"), *Request->GetName(), Request, DebugInfoSize);
-			break;
-
-		case CURLINFO_SSL_DATA_OUT:
-			UE_LOG(LogDerivedDataCache, VeryVerbose, TEXT("%s: %p: Sent SSL data (%d bytes)"), *Request->GetName(), Request, DebugInfoSize);
-			break;
-		}
-
-		return 0;
-	}
-
-	static size_t StaticReadFn(void* Ptr, size_t SizeInBlocks, size_t BlockSizeInBytes, void* UserData)
-	{
-		FHttpRequest* Request = static_cast<FHttpRequest*>(UserData);
-		const size_t MaxReadSize = SizeInBlocks * BlockSizeInBytes;
-		const FMemoryView SourceView = Request->ReadDataView.Mid(Request->BytesSent, MaxReadSize);
-		MakeMemoryView(Ptr, MaxReadSize).CopyFrom(SourceView);
-		Request->BytesSent += SourceView.GetSize();
-		return SourceView.GetSize();
-	}
-
-	static size_t StaticWriteHeaderFn(void* Ptr, size_t SizeInBlocks, size_t BlockSizeInBytes, void* UserData)
-	{
-		FHttpRequest* Request = static_cast<FHttpRequest*>(UserData);
-		const size_t WriteSize = SizeInBlocks * BlockSizeInBytes;
-		TArray<uint8>* WriteHeaderBufferPtr = Request->WriteHeaderBufferPtr;
-		if (WriteHeaderBufferPtr && WriteSize > 0)
-		{
-			const size_t CurrentBufferLength = WriteHeaderBufferPtr->Num();
-			if (CurrentBufferLength > 0)
-			{
-				// Remove the previous zero termination
-				(*WriteHeaderBufferPtr)[CurrentBufferLength-1] = ' ';
-			}
-
-			// Write the header
-			WriteHeaderBufferPtr->Append((const uint8*)Ptr, WriteSize + 1);
-			(*WriteHeaderBufferPtr)[WriteHeaderBufferPtr->Num()-1] = 0; // Zero terminate string
-			return WriteSize;
-		}
-		return 0;
-	}
-
-	static size_t StaticWriteBodyFn(void* Ptr, size_t SizeInBlocks, size_t BlockSizeInBytes, void* UserData)
-	{
-		FHttpRequest* Request = static_cast<FHttpRequest*>(UserData);
-		const size_t WriteSize = SizeInBlocks * BlockSizeInBytes;
-		TArray<uint8>* WriteDataBufferPtr = Request->WriteDataBufferPtr;
-
-		if (WriteDataBufferPtr && WriteSize > 0)
-		{
-			// If this is the first part of the body being received, try to reserve 
-			// memory if content length is defined in the header.
-			if (Request->BytesReceived == 0 && Request->WriteHeaderBufferPtr)
-			{
-				static const ANSICHAR* ContentLengthHeaderStr = "Content-Length: ";
-				const ANSICHAR* Header = (const ANSICHAR*)Request->WriteHeaderBufferPtr->GetData();
-
-				if (const ANSICHAR* ContentLengthHeader = FCStringAnsi::Strstr(Header, ContentLengthHeaderStr))
-				{
-					size_t ContentLength = (size_t)FCStringAnsi::Atoi64(ContentLengthHeader + strlen(ContentLengthHeaderStr));
-					if (ContentLength > 0u && ContentLength < UE_HTTPDDC_MAX_BUFFER_RESERVE)
-					{
-						WriteDataBufferPtr->Reserve(ContentLength);
-					}
-				}
-			}
-
-			// Write to the target buffer
-			WriteDataBufferPtr->Append((const uint8*)Ptr, WriteSize);
-			Request->BytesReceived += WriteSize;
-			return WriteSize;
-		}
-		
-		return 0;
-	}
-
-	static size_t StaticSeekFn(void* UserData, curl_off_t Offset, int Origin)
-	{
-		FHttpRequest* Request = static_cast<FHttpRequest*>(UserData);
-		size_t NewPosition = 0;
-
-		switch (Origin)
-		{
-		case SEEK_SET: NewPosition = Offset; break;
-		case SEEK_CUR: NewPosition = Request->BytesSent + Offset; break;
-		case SEEK_END: NewPosition = Request->ReadDataView.GetSize() + Offset; break;
-		}
-
-		// Make sure we don't seek outside of the buffer
-		if (NewPosition < 0 || NewPosition >= Request->ReadDataView.GetSize())
-		{
-			return CURL_SEEKFUNC_FAIL;
-		}
-
-		// Update the used offset
-		Request->BytesSent = NewPosition;
-		return CURL_SEEKFUNC_OK;
-	}
-
-};
-
-
-//----------------------------------------------------------------------------------------------------------
-// Forward declarations
-//----------------------------------------------------------------------------------------------------------
-bool VerifyPayload(const FSHAHash& Hash, const TCHAR* Namespace, const TCHAR* Bucket, const TCHAR* CacheKey, const TArray<uint8>& Payload);
-bool VerifyPayload(const FIoHash& Hash, const TCHAR* Namespace, const TCHAR* Bucket, const TCHAR* CacheKey, const TArray<uint8>& Payload);
-bool VerifyRequest(const class FHttpRequest* Request, const TCHAR* Namespace, const TCHAR* Bucket, const TCHAR* CacheKey, const TArray<uint8>& Payload);
-bool HashPayload(class FHttpRequest* Request, const TArrayView<const uint8> Payload);
-bool ShouldAbortForShutdown();
-
-//----------------------------------------------------------------------------------------------------------
-// Request pool
-//----------------------------------------------------------------------------------------------------------
-
-
-/**
- * Pool that manages a fixed set of requests. Users are required to release requests that have been 
- * acquired. Usable with \ref FScopedRequestPtr which handles this automatically.
- */
-struct FRequestPool
-{
-	FRequestPool(const TCHAR* InServiceUrl, const TCHAR* InEffectiveServiceUrl, FHttpAccessToken* InAuthorizationToken, FHttpSharedData* InSharedData, uint32 PoolSize, uint32 InOverflowLimit = 0)
-	: ActiveOverflowRequests(0)
-	, OverflowLimit(InOverflowLimit)
-	{
-		Pool.AddUninitialized(PoolSize);
-		Requests.AddUninitialized(PoolSize);
-		for (uint8 i = 0; i < Pool.Num(); ++i)
-		{
-			Pool[i].Usage = 0u;
-			Pool[i].Request = new(&Requests[i]) FHttpRequest(InServiceUrl, InEffectiveServiceUrl, InAuthorizationToken, InSharedData, true);
-		}
-
-		InitData = MakeUnique<FInitData>(InServiceUrl, InEffectiveServiceUrl, InAuthorizationToken, InSharedData);
-	}
-
-	~FRequestPool()
-	{
-		check(ActiveOverflowRequests.load() == 0);
-		for (uint8 i = 0; i < Pool.Num(); ++i)
-		{
-			// No requests should be in use by now.
-			check(Pool[i].Usage.load(std::memory_order_acquire) == 0u);
-		}
-	}
-
-	/**
-	 * Attempts to get a request is free. Once a request has been returned it is
-	 * "owned by the caller and need to release it to the pool when work has been completed.
-	 * @return Usable request instance if one is available, otherwise null.
-	 */
-	FHttpRequest* GetFreeRequest(bool bUnboundedOverflow = false)
-	{
-		for (uint8 i = 0; i < Pool.Num(); ++i)
-		{
-			if (!Pool[i].Usage.load(std::memory_order_relaxed))
-			{
-				uint8 Expected = 0u;
-				if (Pool[i].Usage.compare_exchange_strong(Expected, 1u))
-				{
-					Pool[i].Request->Reset();
-					return Pool[i].Request;
-				}
-			}
-		}
-		if (bUnboundedOverflow || (OverflowLimit > 0))
-		{
-			// The use of two operations here (load, then fetch_add) implies that we can exceed the overflow limit because the combined operation
-			// is not atomic.  This is acceptable for our use case.  If we wanted to enforce the hard limit, we could use a loop instead.
-			if (bUnboundedOverflow || (ActiveOverflowRequests.load(std::memory_order_relaxed) < OverflowLimit))
-			{
-				// Create an overflow request (outside of the pre-allocated range of requests)
-				ActiveOverflowRequests.fetch_add(1, std::memory_order_relaxed);
-				return new FHttpRequest(*InitData->ServiceUrl, *InitData->EffectiveServiceUrl, InitData->AccessToken, InitData->SharedData, true);
-			}
-		}
-		return nullptr;
-	}
-
-	class FWaiter : public FThreadSafeRefCountedObject
-	{
-	public:
-		std::atomic<FHttpRequest*> Request{ nullptr };
-
-		FWaiter(FRequestPool* InPool)
-			: Event(FPlatformProcess::GetSynchEventFromPool(true))
-			, Pool(InPool)
-		{
-		}
-		
-		bool Wait(uint32 TimeMS)
-		{
-			return Event->Wait(TimeMS);
-		}
+#include "Http/HttpClient.h"
 
-		void Trigger()
-		{
-			Event->Trigger();
-		}
-	private:
-		~FWaiter()
-		{
-			FPlatformProcess::ReturnSynchEventToPool(Event);
-
-			if (Request)
-			{
-				Pool->ReleaseRequestToPool(Request.exchange(nullptr));
-			}
-		}
-
-		FEvent* Event;
-		FRequestPool* Pool;
-	};
+#if WITH_SSL
+#include "Ssl.h"
+#include <openssl/ssl.h>
+#endif
 
-	/**
-	 * Block until a request is free. Once a request has been returned it is 
-	 * "owned by the caller and need to release it to the pool when work has been completed.
-	 * @return Usable request instance.
-	 */
-	FHttpRequest* WaitForFreeRequest(bool bUnboundedOverflow = false)
-	{
-		TRACE_CPUPROFILER_EVENT_SCOPE(HttpDDC_WaitForConnection);
+// Enables data request helpers that internally
+// batch requests to reduce the number of concurrent
+// connections.
+#ifndef WITH_DATAREQUEST_HELPER
+	#define WITH_DATAREQUEST_HELPER 1
+#endif
 
-		FHttpRequest* Request = GetFreeRequest(bUnboundedOverflow);
-		if (Request == nullptr)
-		{
-			// Make it a fair by allowing each thread to register itself in a FIFO
-			// so that the first thread to start waiting is the first one to get a request.
-			FWaiter* Waiter = new FWaiter(this);
-			Waiter->AddRef(); // One ref for the thread that will dequeue
-			Waiter->AddRef(); // One ref for us
+#define UE_HTTPDDC_GET_REQUEST_POOL_SIZE 48
+#define UE_HTTPDDC_PUT_REQUEST_POOL_SIZE 16
+#define UE_HTTPDDC_NONBLOCKING_REQUEST_POOL_SIZE 128
+#define UE_HTTPDDC_MAX_FAILED_LOGIN_ATTEMPTS 16
+#define UE_HTTPDDC_MAX_ATTEMPTS 4
+#define UE_HTTPDDC_BATCH_SIZE 12
+#define UE_HTTPDDC_BATCH_NUM 64
+#define UE_HTTPDDC_BATCH_GET_WEIGHT 4
+#define UE_HTTPDDC_BATCH_HEAD_WEIGHT 1
+#define UE_HTTPDDC_BATCH_WEIGHT_HINT 12
 
-			Waiters.enqueue(Waiter);
+namespace UE::DerivedData
+{
 
-			while (!Waiter->Wait(UE_HTTPDDC_BACKEND_WAIT_INTERVAL_MS))
-			{
-				// While waiting, allow us to check if a race occurred and a request has been freed
-				// between the time we checked for free requests and the time we queued ourself as a Waiter.
-				if ((Request = GetFreeRequest(bUnboundedOverflow)) != nullptr)
-				{
-					// We abandon the FWaiter, it will be freed by the next dequeue
-					// and if it has a request, it will be queued back to the pool.
-					Waiter->Release();
-					return Request;
-				}
-			}
+TRACE_DECLARE_INT_COUNTER(HttpDDC_Exist, TEXT("HttpDDC Exist"));
+TRACE_DECLARE_INT_COUNTER(HttpDDC_ExistHit, TEXT("HttpDDC Exist Hit"));
+TRACE_DECLARE_INT_COUNTER(HttpDDC_Get, TEXT("HttpDDC Get"));
+TRACE_DECLARE_INT_COUNTER(HttpDDC_GetHit, TEXT("HttpDDC Get Hit"));
+TRACE_DECLARE_INT_COUNTER(HttpDDC_Put, TEXT("HttpDDC Put"));
+TRACE_DECLARE_INT_COUNTER(HttpDDC_PutHit, TEXT("HttpDDC Put Hit"));
+TRACE_DECLARE_INT_COUNTER(HttpDDC_BytesReceived, TEXT("HttpDDC Bytes Received"));
+TRACE_DECLARE_INT_COUNTER(HttpDDC_BytesSent, TEXT("HttpDDC Bytes Sent"));
 
-			Request = Waiter->Request.exchange(nullptr);
-			Request->Reset();
-			Waiter->Release();
-		}
-		check(Request);
-		return Request;
-	}
 
-	/**
-	 * Release request to the pool.
-	 * @param Request Request that should be freed. Note that any buffer owened by the request can now be reset.
-	 */
-	void ReleaseRequestToPool(FHttpRequest* Request)
+template <typename T>
+class TRefCountedUniqueFunction final : public FThreadSafeRefCountedObject
+{
+public:
+	explicit TRefCountedUniqueFunction(T&& InFunction) : Function(MoveTemp(InFunction))
 	{
-		if ((Request < Requests.GetData()) || (Request >= (Requests.GetData() + Requests.Num())))
-		{
-			// For overflow requests (outside of the pre-allocated range of requests), just delete it immediately
-			delete Request;
-			ActiveOverflowRequests.fetch_sub(1, std::memory_order_relaxed);
-			return;
-		}
-
-		for (uint8 i = 0; i < Pool.Num(); ++i)
-		{
-			if (Pool[i].Request == Request)
-			{
-				// If only 1 user is remaining, we can give it to a waiter
-				// instead of releasing it back to the pool.
-				if (Pool[i].Usage == 1u)
-				{
-					if (FWaiter* Waiter = Waiters.dequeue())
-					{
-						Waiter->Request = Request;
-						Waiter->Trigger();
-						Waiter->Release();
-						return;
-					}
-				}
-				
-				Pool[i].Usage--;
-				return;
-			}
-		}
-		check(false);
 	}
 
-	/**
-	 * While holding a request, make it shared across many users.
-	 */
-	void MakeRequestShared(FHttpRequest* Request, uint8 Users)
-	{
-		if ((Request < Requests.GetData()) || (Request >= (Requests.GetData() + Requests.Num())))
-		{
-			// Overflow requests (outside of the pre-allocated range of requests), cannot be made shared
-			check(false);
-		}
-
-		check(Users != 0);
-		for (uint8 i = 0; i < Pool.Num(); ++i)
-		{
-			if (Pool[i].Request == Request)
-			{
-				Pool[i].Usage = Users;
-				return;
-			}
-		}
-		check(false);
-	}
+	const T& GetFunction() const { return Function; }
 
 private:
-
-	struct FEntry
-	{
-		std::atomic<uint8> Usage;
-		FHttpRequest* Request;
-	};
-
-	TArray<FEntry> Pool;
-	TArray<FHttpRequest> Requests;
-	FAAArrayQueue<FWaiter> Waiters;
-	std::atomic<uint32> ActiveOverflowRequests;
-	struct FInitData
-	{
-		FString ServiceUrl;
-		FString EffectiveServiceUrl;
-		FHttpAccessToken* AccessToken;
-		FHttpSharedData* SharedData;
-
-		FInitData(const TCHAR* InServiceUrl, const TCHAR* InEffectiveServiceUrl, FHttpAccessToken* InAccessToken, FHttpSharedData* InSharedData)
-		: ServiceUrl(InServiceUrl)
-		, EffectiveServiceUrl(InEffectiveServiceUrl)
-		, AccessToken(InAccessToken)
-		, SharedData(InSharedData)
-		{
-		}
-	};
-	TUniquePtr<const FInitData> InitData;
-	const uint32 OverflowLimit;
-
-	FRequestPool() = delete;
+	T Function;
 };
 
+
 //----------------------------------------------------------------------------------------------------------
-// FScopedRequestPtr
+// Forward declarations
 //----------------------------------------------------------------------------------------------------------
-
-/**
- * Utility class to manage requesting and releasing requests from the \ref FRequestPool.
- */
-struct FScopedRequestPtr
-{
-public:
-	FScopedRequestPtr(FRequestPool* InPool)
-		: Request(InPool->WaitForFreeRequest())
-		, Pool(InPool)
-	{
-	}
-
-	~FScopedRequestPtr()
-	{
-		Pool->ReleaseRequestToPool(Request);
-	}
-
-	bool IsValid() const 
-	{
-		return Request != nullptr;
-	}
-
-	FHttpRequest* Get() const
-	{
-		check(IsValid());
-		return Request;
-	}
-
-	FHttpRequest* operator->()
-	{
-		check(IsValid());
-		return Request;
-	}
-
-private:
-	FHttpRequest* Request;
-	FRequestPool* Pool;
-};
+bool VerifyPayload(const FSHAHash& Hash, const TCHAR* Namespace, const TCHAR* Bucket, const TCHAR* CacheKey, const TArray<uint8>& Payload);
+bool VerifyPayload(const FIoHash& Hash, const TCHAR* Namespace, const TCHAR* Bucket, const TCHAR* CacheKey, const TArray<uint8>& Payload);
+bool VerifyRequest(const class FHttpRequest* Request, const TCHAR* Namespace, const TCHAR* Bucket, const TCHAR* CacheKey, const TArray<uint8>& Payload);
+bool HashPayload(class FHttpRequest* Request, const TArrayView<const uint8> Payload);
+bool ShouldAbortForShutdown();
 
 
 #if WITH_DATAREQUEST_HELPER
@@ -1356,7 +129,7 @@ private:
  */
 struct FDataRequestHelper
 {
-	FDataRequestHelper(FRequestPool* InPool, const TCHAR* InNamespace, const TCHAR* InBucket, const TCHAR* InCacheKey, TArray<uint8>* OutData)
+	FDataRequestHelper(FHttpRequestPool* InPool, const TCHAR* InNamespace, const TCHAR* InBucket, const TCHAR* InCacheKey, TArray<uint8>* OutData)
 		: Request(nullptr)
 		, Pool(InPool)
 		, bVerified(false, 1)
@@ -1423,7 +196,7 @@ struct FDataRequestHelper
 	}
 
 	// Constructor specifically for batched head queries
-	FDataRequestHelper(FRequestPool* InPool, const TCHAR* InNamespace, const TCHAR* InBucket, TConstArrayView<FString> InCacheKeys)
+	FDataRequestHelper(FHttpRequestPool* InPool, const TCHAR* InNamespace, const TCHAR* InBucket, TConstArrayView<FString> InCacheKeys)
 		: Request(nullptr)
 		, Pool(InPool)
 		, bVerified(false, InCacheKeys.Num())
@@ -1557,7 +330,7 @@ private:
 	};
 
 	FHttpRequest* Request;
-	FRequestPool* Pool;
+	FHttpRequestPool* Pool;
 	TBitArray<> bVerified;
 	static std::atomic<uint32> FirstAvailableBatch;
 	static TStaticArray<FBatch, UE_HTTPDDC_BATCH_NUM> Batches;
@@ -1584,7 +357,7 @@ private:
 	/**
 	 * Queues up a request to be batched. Blocks until the query is made.
 	 */
-	static FHttpRequest* QueueBatchRequest(FRequestPool* InPool, 
+	static FHttpRequest* QueueBatchRequest(FHttpRequestPool* InPool, 
 		const TCHAR* InNamespace, 
 		const TCHAR* InBucket, 
 		TConstArrayView<const TCHAR*> InCacheKeys,
@@ -1987,7 +760,7 @@ std::atomic<uint32> FDataRequestHelper::FirstAvailableBatch;
 //----------------------------------------------------------------------------------------------------------
 struct FDataUploadHelper
 {
-	FDataUploadHelper(FRequestPool* InPool, 
+	FDataUploadHelper(FHttpRequestPool* InPool, 
 		const TCHAR* InNamespace, 
 		const TCHAR* InBucket, 
 		const TCHAR* InCacheKey, 
@@ -2060,7 +833,7 @@ private:
 	bool bSuccess;
 	bool bQueued;
 
-	static void ProcessQueuedPutsAndReleaseRequest(FRequestPool* Pool, FHttpRequest* Request, FDerivedDataCacheUsageStats& UsageStats)
+	static void ProcessQueuedPutsAndReleaseRequest(FHttpRequestPool* Pool, FHttpRequest* Request, FDerivedDataCacheUsageStats& UsageStats)
 	{
 		while (Request)
 		{
@@ -2308,31 +1081,6 @@ static bool IsValueDataReady(FValue& Value, const ECachePolicy Policy)
 	return false;
 };
 
-//----------------------------------------------------------------------------------------------------------
-// FHttpAccessToken
-//----------------------------------------------------------------------------------------------------------
-
-FString FHttpAccessToken::GetHeader()
-{
-	Lock.ReadLock();
-	FString Header = FString::Printf(TEXT("Authorization: Bearer %s"), *Token);
-	Lock.ReadUnlock();
-	return Header;
-}
-
-void FHttpAccessToken::SetHeader(const TCHAR* InToken)
-{
-	Lock.WriteLock();
-	Token = InToken;
-	Serial++;
-	Lock.WriteUnlock();
-}
-
-uint32 FHttpAccessToken::GetSerial() const
-{
-	return Serial;
-}
-
 //----------------------------------------------------------------------------------------------------------
 // FHttpCacheStore
 //----------------------------------------------------------------------------------------------------------
@@ -2447,9 +1195,9 @@ private:
 	FDerivedDataCacheUsageStats UsageStats;
 	FBackendDebugOptions DebugOptions;
 	TUniquePtr<FHttpSharedData> SharedData;
-	TUniquePtr<FRequestPool> GetRequestPools[2];
-	TUniquePtr<FRequestPool> PutRequestPools[2];
-	TUniquePtr<FRequestPool> NonBlockingRequestPools;
+	TUniquePtr<FHttpRequestPool> GetRequestPools[2];
+	TUniquePtr<FHttpRequestPool> PutRequestPools[2];
+	TUniquePtr<FHttpRequestPool> NonBlockingRequestPools;
 	TUniquePtr<FHttpAccessToken> Access;
 	bool bIsUsable;
 	bool bReadOnly;
@@ -2469,7 +1217,7 @@ private:
 		Put
 	};
 	template<OperationCategory Category>
-	FHttpRequest* WaitForHttpRequestForOwner(IRequestOwner& Owner, bool bUnboundedOverflow, FRequestPool*& OutPool)
+	FHttpRequest* WaitForHttpRequestForOwner(IRequestOwner& Owner, bool bUnboundedOverflow, FHttpRequestPool*& OutPool)
 	{
 		if (!FHttpRequest::AllowAsync())
 		{
@@ -2805,7 +1553,7 @@ void FHttpCacheStore::FPutPackageOp::PutRefAsync(
 		RefsUri << "/finalize/" << ObjectHash;
 	}
 
-	FRequestPool* Pool = nullptr;
+	FHttpRequestPool* Pool = nullptr;
 	FHttpRequest* Request = CacheStore.WaitForHttpRequestForOwner<OperationCategory::Put>(Owner, bFinalize /* bUnboundedOverflow */, Pool);
 
 	auto OnHttpRequestComplete = [&CacheStore, &Owner, Name = FSharedString(Name), Key, Object, UserData, bFinalize, OnComplete = MoveTemp(OnComplete)]
@@ -2983,7 +1731,7 @@ void FHttpCacheStore::FPutPackageOp::OnPackagePutRefComplete(
 		TStringBuilder<256> CompressedBlobsUri;
 		CompressedBlobsUri << "api/v1/compressed-blobs/" << CacheStore.StructuredNamespace << "/" << CompressedBlobUpload.Hash;
 
-		FRequestPool* Pool = nullptr;
+		FHttpRequestPool* Pool = nullptr;
 		FHttpRequest* Request = CacheStore.WaitForHttpRequestForOwner<OperationCategory::Put>(Owner, true /* bUnboundedOverflow */, Pool);
 		Request->EnqueueAsyncUpload<FHttpRequest::PutCompressedBlob>(Owner, Pool, *CompressedBlobsUri, CompressedBlobUpload.BlobBuffer,
 			[PutPackageOp](FHttpRequest::Result HttpResult, FHttpRequest* Request)
@@ -3092,7 +1840,7 @@ void FHttpCacheStore::FGetRecordOp::GetDataBatch(
 		const ValueType& Value = Values[ValueIndex];
 		const FIoHash& RawHash = Value.GetRawHash();
 
-		FRequestPool* Pool = nullptr;
+		FHttpRequestPool* Pool = nullptr;
 		FHttpRequest* Request = CacheStore.WaitForHttpRequestForOwner<OperationCategory::Get>(Owner, true /* bUnboundedOverflow */, Pool);
 
 		auto OnHttpRequestComplete = [&CacheStore, &Owner, Name = FSharedString(Name), Key = FCacheKey(Key), ValueIndex, Value = Value.RemoveData(), ValueIdGetter, OnCompletePtr = TRefCountPtr<TRefCountedUniqueFunction<FOnGetCachedDataBatchComplete>>(CompletionFunction)]
@@ -3287,7 +2035,7 @@ void FHttpCacheStore::FGetRecordOp::DataProbablyExistsBatch(
 		return;
 	}
 
-	FRequestPool* Pool = nullptr;
+	FHttpRequestPool* Pool = nullptr;
 	FHttpRequest* Request = CacheStore.WaitForHttpRequestForOwner<OperationCategory::Get>(Owner, true /* bUnboundedOverflow */, Pool);
 
 	TStringBuilder<256> CompressedBlobsUri;
@@ -3508,12 +2256,12 @@ FHttpCacheStore::FHttpCacheStore(
 			EffectiveDomain = Domain;
 		}
 
-		GetRequestPools[0] = MakeUnique<FRequestPool>(*Domain, *EffectiveDomain, Access.Get(), SharedData.Get(), UE_HTTPDDC_GET_REQUEST_POOL_SIZE);
-		GetRequestPools[1] = MakeUnique<FRequestPool>(*Domain, *EffectiveDomain, Access.Get(), SharedData.Get(), UE_HTTPDDC_GET_REQUEST_POOL_SIZE);
-		PutRequestPools[0] = MakeUnique<FRequestPool>(*Domain, *EffectiveDomain, Access.Get(), SharedData.Get(), UE_HTTPDDC_PUT_REQUEST_POOL_SIZE);
-		PutRequestPools[1] = MakeUnique<FRequestPool>(*Domain, *EffectiveDomain, Access.Get(), SharedData.Get(), UE_HTTPDDC_PUT_REQUEST_POOL_SIZE);
+		GetRequestPools[0] = MakeUnique<FHttpRequestPool>(*Domain, *EffectiveDomain, Access.Get(), SharedData.Get(), UE_HTTPDDC_GET_REQUEST_POOL_SIZE);
+		GetRequestPools[1] = MakeUnique<FHttpRequestPool>(*Domain, *EffectiveDomain, Access.Get(), SharedData.Get(), UE_HTTPDDC_GET_REQUEST_POOL_SIZE);
+		PutRequestPools[0] = MakeUnique<FHttpRequestPool>(*Domain, *EffectiveDomain, Access.Get(), SharedData.Get(), UE_HTTPDDC_PUT_REQUEST_POOL_SIZE);
+		PutRequestPools[1] = MakeUnique<FHttpRequestPool>(*Domain, *EffectiveDomain, Access.Get(), SharedData.Get(), UE_HTTPDDC_PUT_REQUEST_POOL_SIZE);
 		// Allowing the non-blocking requests to overflow to double their pre-allocated size before we start waiting for one to free up.
-		NonBlockingRequestPools = MakeUnique<FRequestPool>(*Domain, *EffectiveDomain, Access.Get(), SharedData.Get(), UE_HTTPDDC_NONBLOCKING_REQUEST_POOL_SIZE, UE_HTTPDDC_NONBLOCKING_REQUEST_POOL_SIZE);
+		NonBlockingRequestPools = MakeUnique<FHttpRequestPool>(*Domain, *EffectiveDomain, Access.Get(), SharedData.Get(), UE_HTTPDDC_NONBLOCKING_REQUEST_POOL_SIZE, UE_HTTPDDC_NONBLOCKING_REQUEST_POOL_SIZE);
 		bIsUsable = true;
 	}
 
@@ -3768,7 +2516,7 @@ void FHttpCacheStore::GetCacheRecordOnlyAsync(
 	TStringBuilder<256> RefsUri;
 	RefsUri << "api/v1/refs/" << StructuredNamespace << "/" << Bucket << "/" << Key.Hash;
 
-	FRequestPool* Pool = nullptr;
+	FHttpRequestPool* Pool = nullptr;
 	FHttpRequest* Request = WaitForHttpRequestForOwner<OperationCategory::Get>(Owner, false /* bUnboundedOverflow */, Pool);
 
 	auto OnHttpRequestComplete = [this, &Owner, Name = FSharedString(Name), Key, UserData, OnComplete = MoveTemp(OnComplete)]
@@ -3971,7 +2719,7 @@ void FHttpCacheStore::GetCacheValueAsync(
 	TStringBuilder<256> RefsUri;
 	RefsUri << "api/v1/refs/" << StructuredNamespace << "/" << Bucket << "/" << Key.Hash;
 
-	FRequestPool* Pool = nullptr;
+	FHttpRequestPool* Pool = nullptr;
 	FHttpRequest* Request = WaitForHttpRequestForOwner<OperationCategory::Get>(Owner, false /* bUnboundedOverflow */, Pool);
 	if (bSkipData)
 	{
@@ -4116,7 +2864,7 @@ void FHttpCacheStore::RefCachedDataProbablyExistsBatchAsync(
 	FCbFieldIterator RequestFields = RequestWriter.Save();
 
 
-	FRequestPool* Pool = nullptr;
+	FHttpRequestPool* Pool = nullptr;
 	FHttpRequest* Request = WaitForHttpRequestForOwner<OperationCategory::Get>(Owner, false /* bUnboundedOverflow */, Pool);
 
 	Request->SetHeader(TEXT("Accept"), TEXT("application/x-ue-cb"));
@@ -4257,7 +3005,7 @@ bool FHttpCacheStore::CachedDataProbablyExists(const TCHAR* CacheKey)
 	// Retry request until we get an accepted response or exhaust allowed number of attempts.
 	for (int32 Attempts = 0; Attempts < UE_HTTPDDC_MAX_ATTEMPTS; ++Attempts)
 	{
-		FScopedRequestPtr Request(GetRequestPools[IsInGameThread()].Get());
+		FScopedHttpPoolRequestPtr Request(GetRequestPools[IsInGameThread()].Get());
 		const FHttpRequest::Result Result = Request->PerformBlockingQuery<FHttpRequest::Head>(*Uri);
 		const int64 ResponseCode = Request->GetResponseCode();
 
@@ -4346,7 +3094,7 @@ TBitArray<> FHttpCacheStore::CachedDataProbablyExistsBatch(TConstArrayView<FStri
 	// Retry request until we get an accepted response or exhaust allowed number of attempts.
 	for (int32 Attempts = 0; Attempts < UE_HTTPDDC_MAX_ATTEMPTS; ++Attempts)
 	{
-		FScopedRequestPtr Request(GetRequestPools[IsInGameThread()].Get());
+		FScopedHttpPoolRequestPtr Request(GetRequestPools[IsInGameThread()].Get());
 		const FHttpRequest::Result Result = Request->PerformBlockingUpload<FHttpRequest::PostJson>(Uri, BodyView);
 		const int64 ResponseCode = Request->GetResponseCode();
 
@@ -4428,7 +3176,7 @@ bool FHttpCacheStore::GetCachedData(const TCHAR* CacheKey, TArray<uint8>& OutDat
 	// Retry request until we get an accepted response or exhaust allowed number of attempts.
 	for (int32 Attempts = 0; Attempts < UE_HTTPDDC_MAX_ATTEMPTS; ++Attempts)
 	{
-		FScopedRequestPtr Request(GetRequestPools[IsInGameThread()].Get());
+		FScopedHttpPoolRequestPtr Request(GetRequestPools[IsInGameThread()].Get());
 		if (Request.IsValid())
 		{
 			FHttpRequest::Result Result = Request->PerformBlockingDownload(*Uri, &OutData);
@@ -4506,7 +3254,7 @@ FDerivedDataBackendInterface::EPutStatus FHttpCacheStore::PutCachedData(const TC
 			return EPutStatus::NotCached;
 		}
 
-		FScopedRequestPtr Request(PutRequestPools[IsInGameThread()].Get());
+		FScopedHttpPoolRequestPtr Request(PutRequestPools[IsInGameThread()].Get());
 		if (Request.IsValid())
 		{
 			// Append the content hash to the header
@@ -4548,7 +3296,7 @@ void FHttpCacheStore::RemoveCachedData(const TCHAR* CacheKey, bool bTransient)
 	// Retry request until we get an accepted response or exhaust allowed number of attempts.
 	while (ResponseCode == 0 && ++Attempts < UE_HTTPDDC_MAX_ATTEMPTS)
 	{
-		FScopedRequestPtr Request(PutRequestPools[IsInGameThread()].Get());
+		FScopedHttpPoolRequestPtr Request(PutRequestPools[IsInGameThread()].Get());
 		if (Request.IsValid())
 		{
 			FHttpRequest::Result Result = Request->PerformBlockingQuery<FHttpRequest::Delete>(*Uri, {});
@@ -4913,183 +3661,6 @@ void FHttpCacheStore::GetChunks(
 	}
 }
 
-bool FHttpRequest::AllowAsync()
-{
-	if (!FGenericPlatformProcess::SupportsMultithreading() || !bHttpEnableAsync)
-	{
-		return false;
-	}
-
-	return true;
-}
-
-void FHttpRequest::EnqueueAsync(IRequestOwner& Owner, FRequestPool* Pool, const TCHAR* Uri, RequestVerb Verb, uint64 ContentLength, FOnHttpRequestComplete&& OnComplete, TConstArrayView<long> ExpectedErrorCodes)
-{
-	if (!AllowAsync())
-	{
-		while (OnComplete(PerformBlocking(Uri, Verb, ContentLength, ExpectedErrorCodes), this) == ECompletionBehavior::Retry)
-		{
-			PrepareToRetry();
-		}
-		if (Pool)
-		{
-			Pool->ReleaseRequestToPool(this);
-		}
-		return;
-	}
-
-	TRACE_CPUPROFILER_EVENT_SCOPE(HttpDDC_CurlEnqueueAsync);
-	AsyncData = new FAsyncRequestData;
-	AsyncData->Owner = &Owner;
-	AsyncData->Pool = Pool;
-	AsyncData->CurlHeaders = PrepareToIssueRequest(Uri, ContentLength);
-	AsyncData->Uri = Uri;
-	AsyncData->Verb = Verb;
-	AsyncData->ExpectedErrorCodes = ExpectedErrorCodes;
-	AsyncData->OnComplete = MoveTemp(OnComplete);
-	AsyncData->Owner->Begin(AsyncData);
-
-	SharedData->AddRequest(Curl);
-}
-
-void FHttpRequest::CompleteAsync(CURLcode Result)
-{
-	CurlResult = Result;
-
-	// Get response code
-	curl_easy_getinfo(Curl, CURLINFO_RESPONSE_CODE, &ResponseCode);
-
-	LogResult(Result, *AsyncData->Uri, AsyncData->Verb, AsyncData->ExpectedErrorCodes);
-
-	ECompletionBehavior Behavior;
-	{
-		FRequestBarrier Barrier(*AsyncData->Owner, ERequestBarrierFlags::Priority);
-		FHttpRequest::Result HttpResult;
-		switch (CurlResult)
-		{
-		case CURLE_OK:
-			HttpResult = Success;
-			break;
-		case CURLE_OPERATION_TIMEDOUT:
-			HttpResult = FailedTimeout;
-			break;
-		default:
-			HttpResult = Failed;
-			break;
-		}
-		Behavior = AsyncData->OnComplete(HttpResult, this);
-	}
-
-	if (Behavior == ECompletionBehavior::Retry)
-	{
-		PrepareToRetry();
-		SharedData->AddRequest(Curl);
-	}
-	else
-	{
-		// Clean up
-		curl_slist_free_all(AsyncData->CurlHeaders);
-
-		AsyncData->Owner->End(AsyncData, [this, Behavior]
-			{
-				AsyncData->Event.Trigger();
-				FRequestPool* Pool = AsyncData->Pool;
-				AsyncData = nullptr;
-				if (Pool)
-				{
-					Pool->ReleaseRequestToPool(this);
-				}
-			});
-	}
-}
-
-void FHttpSharedData::AddRequest(CURL* Curl)
-{
-	if (PendingRequestAdditions.EnqueueAndReturnWasEmpty(Curl))
-	{
-		PendingRequestEvent.Trigger();
-	}
-
-	if (!bAsyncThreadStarting.load(std::memory_order_relaxed) && !bAsyncThreadStarting.exchange(true, std::memory_order_relaxed))
-	{
-		AsyncServiceThread = FThread(TEXT("HttpCacheStore"), [this] { ProcessAsyncRequests(); }, 64 * 1024, TPri_Normal);
-	}
-}
-
-void FHttpSharedData::ProcessAsyncRequests()
-{
-	int ActiveTransfers = 0;
-
-	auto ProcessPendingRequests = [this, &ActiveTransfers]()
-	{
-		int CurrentActiveTransfers = -1;
-
-		do
-		{
-			PendingRequestAdditions.Deplete([this, &ActiveTransfers](CURL* Curl)
-				{
-					curl_multi_add_handle(CurlMulti, Curl);
-					++ActiveTransfers;
-				});
-
-			curl_multi_perform(CurlMulti, &CurrentActiveTransfers);
-
-			if (CurrentActiveTransfers == 0 || ActiveTransfers != CurrentActiveTransfers)
-			{
-				for (;;)
-				{
-					int MsgsStillInQueue = 0;	// may use that to impose some upper limit we may spend in that loop
-					CURLMsg* Message = curl_multi_info_read(CurlMulti, &MsgsStillInQueue);
-
-					if (!Message)
-					{
-						break;
-					}
-
-					// find out which requests have completed
-					if (Message->msg == CURLMSG_DONE)
-					{
-						CURL* CompletedHandle = Message->easy_handle;
-						curl_multi_remove_handle(CurlMulti, CompletedHandle);
-
-						void* PrivateData = nullptr;
-						curl_easy_getinfo(CompletedHandle, CURLINFO_PRIVATE, &PrivateData);
-						FHttpRequest* CompletedRequest = (FHttpRequest*)PrivateData;
-
-						if (CompletedRequest)
-						{
-							// It is important that the CompleteAsync call doesn't happen on this thread as it is possible it will block waiting
-							// for a free HTTP request, and if that happens on this thread, we can deadlock as no HTTP requests will become 
-							// available while this thread is blocked.
-							UE::Tasks::Launch(TEXT("FHttpRequest::CompleteAsync"), [CompletedRequest, Result = Message->data.result]() mutable
-							{
-								CompletedRequest->CompleteAsync(Result);
-							});
-						}
-					}
-				}
-				ActiveTransfers = CurrentActiveTransfers;
-			}
-
-			if (CurrentActiveTransfers > 0)
-			{
-				curl_multi_wait(CurlMulti, nullptr, 0, 1, nullptr);
-			}
-		}
-		while (CurrentActiveTransfers > 0);
-	};
-
-	do
-	{
-		ProcessPendingRequests();
-		PendingRequestEvent.Wait(100);
-	}
-	while (!FHttpSharedData::bAsyncThreadShutdownRequested.load(std::memory_order_relaxed));
-
-	// Process last requests before shutdown.  May want these to be aborted instead.
-	ProcessPendingRequests();
-}
-
 } // UE::DerivedData
 
 #endif // WITH_HTTP_DDC_BACKEND
