Fix missing includes in HttpClient.h and HttpRequest.cpp identified by the NoPCH build.

#rnx

diff --git a/Engine/Source/Developer/DerivedDataCache/Private/Http/HttpClient.h b/Engine/Source/Developer/DerivedDataCache/Private/Http/HttpClient.h
--- a/Engine/Source/Developer/DerivedDataCache/Private/Http/HttpClient.h
+++ b/Engine/Source/Developer/DerivedDataCache/Private/Http/HttpClient.h
@@ -16,7 +16,9 @@
 #include "DerivedDataRequestTypes.h"
 #include "Dom/JsonObject.h"
 #include "Experimental/Async/LazyEvent.h"
+#include "Experimental/Containers/FAAArrayQueue.h"
 #include "HAL/CriticalSection.h"
+#include "HAL/PlatformProcess.h"
 #include "HAL/Thread.h"
 #include "Memory/MemoryView.h"
 #include "Memory/SharedBuffer.h"
diff --git a/Engine/Source/Developer/DerivedDataCache/Private/Http/HttpRequest.cpp b/Engine/Source/Developer/DerivedDataCache/Private/Http/HttpRequest.cpp
--- a/Engine/Source/Developer/DerivedDataCache/Private/Http/HttpRequest.cpp
+++ b/Engine/Source/Developer/DerivedDataCache/Private/Http/HttpRequest.cpp
@@ -4,6 +4,7 @@
 
 #if WITH_HTTP_CLIENT
 #include "DerivedDataLegacyCacheStore.h"
+#include "HAL/IConsoleManager.h"
 #include "Logging/LogMacros.h"
 #include "Serialization/JsonReader.h"
 #include "Serialization/JsonSerializer.h"
