Avoid compile error on WinGDK target due to absence of WITH_SSL.

#rnx

diff --git a/Engine/Source/Developer/DerivedDataCache/Private/Http/HttpRequest.cpp b/Engine/Source/Developer/DerivedDataCache/Private/Http/HttpRequest.cpp
--- a/Engine/Source/Developer/DerivedDataCache/Private/Http/HttpRequest.cpp
+++ b/Engine/Source/Developer/DerivedDataCache/Private/Http/HttpRequest.cpp
@@ -67,20 +67,22 @@ void FHttpRequest::Reset()
 	curl_easy_setopt(Curl, CURLOPT_NOSIGNAL, 1L);
 	curl_easy_setopt(Curl, CURLOPT_DNS_CACHE_TIMEOUT, 300L); // Don't re-resolve every minute
 	curl_easy_setopt(Curl, CURLOPT_SHARE, SharedData ? SharedData->GetCurlShare() : nullptr);
-	// SSL options
-	curl_easy_setopt(Curl, CURLOPT_USE_SSL, CURLUSESSL_ALL);
-	curl_easy_setopt(Curl, CURLOPT_SSL_VERIFYPEER, 1);
-	curl_easy_setopt(Curl, CURLOPT_SSL_VERIFYHOST, 1);
-	curl_easy_setopt(Curl, CURLOPT_SSLCERTTYPE, "PEM");
 	// Response functions
 	curl_easy_setopt(Curl, CURLOPT_HEADERDATA, this);
 	curl_easy_setopt(Curl, CURLOPT_HEADERFUNCTION, &FHttpRequest::StaticWriteHeaderFn);
 	curl_easy_setopt(Curl, CURLOPT_WRITEDATA, this);
 	curl_easy_setopt(Curl, CURLOPT_WRITEFUNCTION, StaticWriteBodyFn);
+	#if WITH_SSL
+	// SSL options
+	curl_easy_setopt(Curl, CURLOPT_USE_SSL, CURLUSESSL_ALL);
+	curl_easy_setopt(Curl, CURLOPT_SSL_VERIFYPEER, 1);
+	curl_easy_setopt(Curl, CURLOPT_SSL_VERIFYHOST, 1);
+	curl_easy_setopt(Curl, CURLOPT_SSLCERTTYPE, "PEM");
 	// SSL certification verification
 	curl_easy_setopt(Curl, CURLOPT_CAINFO, nullptr);
 	curl_easy_setopt(Curl, CURLOPT_SSL_CTX_FUNCTION, &FHttpRequest::StaticSSLCTXFn);
 	curl_easy_setopt(Curl, CURLOPT_SSL_CTX_DATA, this);
+	#endif //#if WITH_SSL
 	// Allow compressed data
 	curl_easy_setopt(Curl, CURLOPT_ACCEPT_ENCODING, "gzip");
 	// Rewind method, handle special error case where request need to rewind data stream
