Pass on configured namespace for Zen cache RPC requests.
Current Zen server implementations ignore this field so no functional change until new Zen server with namespace support are deploy/integrated into p4

diff --git a/Engine/Source/Developer/DerivedDataCache/Private/ZenCacheStore.cpp b/Engine/Source/Developer/DerivedDataCache/Private/ZenCacheStore.cpp
--- a/Engine/Source/Developer/DerivedDataCache/Private/ZenCacheStore.cpp
+++ b/Engine/Source/Developer/DerivedDataCache/Private/ZenCacheStore.cpp
@@ -757,6 +757,7 @@ void FZenCacheStore::Get(
 			{
 				ECachePolicy BatchDefaultPolicy = Batch[0].Policy.GetRecordPolicy();
 				BatchRequest << ANSITEXTVIEW("DefaultPolicy") << *WriteToString<128>(BatchDefaultPolicy);
+				BatchRequest.AddString(ANSITEXTVIEW("Namespace"), Namespace);
 
 				BatchRequest.BeginArray(ANSITEXTVIEW("Requests"));
 				for (const FCacheGetRequest& Request : Batch)
@@ -878,6 +879,7 @@ void FZenCacheStore::PutValue(
 			{
 				ECachePolicy BatchDefaultPolicy = Batch[0].Policy;
 				BatchWriter << ANSITEXTVIEW("DefaultPolicy") << *WriteToString<128>(BatchDefaultPolicy);
+				BatchWriter.AddString(ANSITEXTVIEW("Namespace"), Namespace);
 
 				BatchWriter.BeginArray("Requests");
 				for (const FCachePutValueRequest& Request : Batch)
@@ -989,18 +991,19 @@ void FZenCacheStore::GetValue(
 				{
 					ECachePolicy BatchDefaultPolicy = Batch[0].Policy;
 					BatchRequest << ANSITEXTVIEW("DefaultPolicy") << *WriteToString<128>(BatchDefaultPolicy);
+					BatchRequest.AddString(ANSITEXTVIEW("Namespace"), Namespace);
 
 					BatchRequest.BeginArray("Requests");
 					for (const FCacheGetValueRequest& Request : Batch)
 					{
 						BatchRequest.BeginObject();
-					{
+						{
 							BatchRequest << ANSITEXTVIEW("Key") << Request.Key;
 							if (Request.Policy != BatchDefaultPolicy)
 							{
 								BatchRequest << ANSITEXTVIEW("Policy") << WriteToString<128>(Request.Policy);
 							}
-					}
+						}
 						BatchRequest.EndObject();
 					}
 					BatchRequest.EndArray();
@@ -1116,34 +1119,36 @@ void FZenCacheStore::GetChunks(
 			{
 				ECachePolicy DefaultPolicy = Batch[0].Policy;
 				BatchRequest << ANSITEXTVIEW("DefaultPolicy") << WriteToString<128>(DefaultPolicy);
+				BatchRequest.AddString(ANSITEXTVIEW("Namespace"), Namespace);
+
 				BatchRequest.BeginArray(ANSITEXTVIEW("ChunkRequests"));
 				for (const FCacheGetChunkRequest& Request : Batch)
 				{
 					BatchRequest.BeginObject();
-					
-					BatchRequest << ANSITEXTVIEW("Key") << Request.Key;
-
-					if (Request.Id.IsValid())
-					{
-						BatchRequest.AddObjectId(ANSITEXTVIEW("ValueId"), Request.Id);
-					}
-					if (Request.RawOffset != 0)
-					{
-						BatchRequest << ANSITEXTVIEW("RawOffset") << Request.RawOffset;
-					}
-					if (Request.RawSize != MAX_uint64)
-					{
-						BatchRequest << ANSITEXTVIEW("RawSize") << Request.RawSize;
-					}
-					if (!Request.RawHash.IsZero())
-					{
-						BatchRequest << ANSITEXTVIEW("ChunkId") << Request.RawHash;
-					}
-					if (Request.Policy != DefaultPolicy)
 					{
-						BatchRequest << ANSITEXTVIEW("Policy") << WriteToString<128>(Request.Policy);
-					}
+						BatchRequest << ANSITEXTVIEW("Key") << Request.Key;
 
+						if (Request.Id.IsValid())
+						{
+							BatchRequest.AddObjectId(ANSITEXTVIEW("ValueId"), Request.Id);
+						}
+						if (Request.RawOffset != 0)
+						{
+							BatchRequest << ANSITEXTVIEW("RawOffset") << Request.RawOffset;
+						}
+						if (Request.RawSize != MAX_uint64)
+						{
+							BatchRequest << ANSITEXTVIEW("RawSize") << Request.RawSize;
+						}
+						if (!Request.RawHash.IsZero())
+						{
+							BatchRequest << ANSITEXTVIEW("ChunkId") << Request.RawHash;
+						}
+						if (Request.Policy != DefaultPolicy)
+						{
+							BatchRequest << ANSITEXTVIEW("Policy") << WriteToString<128>(Request.Policy);
+						}
+					}
 					BatchRequest.EndObject();
 				}
 				BatchRequest.EndArray();
