Fix incorrect handling of empty "needs" list when performing existence check in Horde Storage.

diff --git a/Engine/Source/Developer/DerivedDataCache/Private/HttpCacheStore.cpp b/Engine/Source/Developer/DerivedDataCache/Private/HttpCacheStore.cpp
--- a/Engine/Source/Developer/DerivedDataCache/Private/HttpCacheStore.cpp
+++ b/Engine/Source/Developer/DerivedDataCache/Private/HttpCacheStore.cpp
@@ -2068,11 +2068,11 @@ void FHttpCacheStore::FGetRecordOp::DataProbablyExistsBatch(
 						for (int32 ValueIndex = 0; ValueIndex < Values.Num(); ++ValueIndex)
 						{
 							const FValueWithId& Value = Values[ValueIndex];
-							UE_LOG(LogDerivedDataCache, Verbose,
-								TEXT("%s: Cache exists miss with missing value %s with hash %s for %s from '%s'"),
+							UE_LOG(LogDerivedDataCache, VeryVerbose,
+								TEXT("%s: Cache exists hit for %s with hash %s for %s from '%s'"),
 								*CacheStore.GetName(), *WriteToString<16>(Value.GetId()), *WriteToString<48>(Value.GetRawHash()), *WriteToString<96>(Key),
 								*Name);
-							InOnComplete({ Name, Key, ValueIndex, EStatus::Error });
+							InOnComplete({ Name, Key, ValueIndex, EStatus::Ok });
 						}
 						return FHttpRequest::ECompletionBehavior::Done;
 					}
@@ -2081,8 +2081,7 @@ void FHttpCacheStore::FGetRecordOp::DataProbablyExistsBatch(
 				TBitArray<> ResultStatus(true, Values.Num());
 				for (const FString& NeedsString : NeedsArrayStrings)
 				{
-					FIoHash NeedHash;
-					LexFromString(NeedHash, *NeedsString);
+					const FIoHash NeedHash(NeedsString);
 					for (int32 ValueIndex = 0; ValueIndex < Values.Num(); ++ValueIndex)
 					{
 						const FValueWithId& Value = Values[ValueIndex];
