FCacheStoreAsync: Fix empty initialization statement warning when COOK_STAT is not defined.
#rnx

diff --git a/Engine/Source/Developer/DerivedDataCache/Private/DerivedDataCacheStoreAsync.cpp b/Engine/Source/Developer/DerivedDataCache/Private/DerivedDataCacheStoreAsync.cpp
--- a/Engine/Source/Developer/DerivedDataCache/Private/DerivedDataCacheStoreAsync.cpp
+++ b/Engine/Source/Developer/DerivedDataCache/Private/DerivedDataCacheStoreAsync.cpp
@@ -186,7 +186,7 @@
 	{
 		Invoke(OnExecute, *InnerCache, Requests, Owner, [this, COOK_STAT(OnAddStats,) OnComplete = MoveTemp(OnComplete)](auto&& Response)
 		{
-			if (COOK_STAT(auto Timer = Invoke(OnAddStats, UsageStats)); Response.Status == EStatus::Ok)
+			if (COOK_STAT(auto Timer = Invoke(OnAddStats, UsageStats);) Response.Status == EStatus::Ok)
 			{
 				COOK_STAT(Timer.AddHit(0));
 			}
