Allow configuration for separate namespace for structured data

diff --git a/Engine/Config/BaseEngine.ini b/Engine/Config/BaseEngine.ini
--- a/Engine/Config/BaseEngine.ini
+++ b/Engine/Config/BaseEngine.ini
@@ -2226,13 +2226,13 @@
 [ZenOnlyDDC]
 Root=(Type=KeyLength, Length=120, Inner=AsyncPut)
 AsyncPut=(Type=AsyncPut, Inner=Zen)
-Zen=(Type=Zen, Namespace="ue4.ddc")
+Zen=(Type=Zen, Namespace="ue4.ddc", StructuredNamespace="ue4.ddc")
 
 [ZenDDC]
 Root=(Type=KeyLength, Length=120, Inner=AsyncPut)
 AsyncPut=(Type=AsyncPut, Inner=Hierarchy)
 Hierarchy=(Type=Hierarchical, Inner=Pak, Inner=EnginePak, Inner=Zen, Inner=Shared)
-Zen=(Type=Zen, Namespace="ue4.ddc")
+Zen=(Type=Zen, Namespace="ue4.ddc", StructuredNamespace="ue4.ddc")
 Shared=(Type=FileSystem, UnusedFileAge=10, FoldersToClean=10, ConsiderSlowAt=70, Path=?EpicDDC, EnvPathOverride=UE-SharedDataCachePath, EditorOverrideSetting=SharedDerivedDataCache, CommandLineOverride=SharedDataCachePath)
 Pak=(Type=ReadPak, Filename="%GAMEDIR%DerivedDataCache/DDC.ddp")
 EnginePak=(Type=ReadPak, Filename="%ENGINEDIR%DerivedDataCache/DDC.ddp")
diff --git a/Engine/Source/Developer/DerivedDataCache/Private/DerivedDataBackends.cpp b/Engine/Source/Developer/DerivedDataCache/Private/DerivedDataBackends.cpp
--- a/Engine/Source/Developer/DerivedDataCache/Private/DerivedDataBackends.cpp
+++ b/Engine/Source/Developer/DerivedDataCache/Private/DerivedDataBackends.cpp
@@ -68,7 +68,7 @@
 IMemoryCacheStore* CreateMemoryCacheStore(const TCHAR* Name, int64 MaxCacheSize, bool bCanBeDisabled);
 IPakFileCacheStore* CreatePakFileCacheStore(const TCHAR* Filename, bool bWriting, bool bCompressed);
 ILegacyCacheStore* CreateS3CacheStore(const TCHAR* RootManifestPath, const TCHAR* BaseUrl, const TCHAR* Region, const TCHAR* CanaryObjectKey, const TCHAR* CachePath);
-ILegacyCacheStore* CreateZenCacheStore(const TCHAR* NodeName, const TCHAR* ServiceUrl, const TCHAR* Namespace);
+ILegacyCacheStore* CreateZenCacheStore(const TCHAR* NodeName, const TCHAR* ServiceUrl, const TCHAR* Namespace, const TCHAR* StructuredNamespace);
 
 /**
  * This class is used to create a singleton that represents the derived data cache hierarchy and all of the wrappers necessary
@@ -1014,8 +1014,14 @@
 			Namespace = FApp::GetProjectName();
 			UE_LOG(LogDerivedDataCache, Warning, TEXT("Node %s does not specify 'Namespace', falling back to '%s'"), NodeName, *Namespace);
 		}
+		FString StructuredNamespace;
+		if (!FParse::Value(Entry, TEXT("StructuredNamespace="), StructuredNamespace))
+		{
+			StructuredNamespace = Namespace;
+			UE_LOG(LogDerivedDataCache, Warning, TEXT("Node %s does not specify 'StructuredNamespace', falling back using Namespace option: '%s'"), NodeName, *StructuredNamespace);
+		}
 
-		if (ILegacyCacheStore* Backend = CreateZenCacheStore(NodeName, *ServiceUrl, *Namespace))
+		if (ILegacyCacheStore* Backend = CreateZenCacheStore(NodeName, *ServiceUrl, *Namespace, *StructuredNamespace))
 		{
 			return Backend;
 		}
diff --git a/Engine/Source/Developer/DerivedDataCache/Private/Tests/HttpCacheStoreTest.cpp b/Engine/Source/Developer/DerivedDataCache/Private/Tests/HttpCacheStoreTest.cpp
--- a/Engine/Source/Developer/DerivedDataCache/Private/Tests/HttpCacheStoreTest.cpp
+++ b/Engine/Source/Developer/DerivedDataCache/Private/Tests/HttpCacheStoreTest.cpp
@@ -47,7 +47,7 @@
 	FString& OutNamespace,
 	FString& OutStructuredNamespace);
 
-ILegacyCacheStore* CreateZenCacheStore(const TCHAR* NodeName, const TCHAR* ServiceUrl, const TCHAR* Namespace);
+ILegacyCacheStore* CreateZenCacheStore(const TCHAR* NodeName, const TCHAR* ServiceUrl, const TCHAR* Namespace, const TCHAR* StructuredNamespace);
 
 class FHttpCacheStoreTestBase : public FAutomationTestBase
 {
@@ -756,10 +756,10 @@
 	ZenTestSiblingAutoLaunchSettings.bLimitProcessLifetime = true;
 
 	FScopeZenService ScopeZenSiblingService(MoveTemp(ZenTestServiceSiblingSettings));
-	TUniquePtr<ILegacyCacheStore> ZenIntermediarySiblingBackend(CreateZenCacheStore(TEXT("TestSibling"), ScopeZenSiblingService.GetInstance().GetURL(), *TestNamespace));
+	TUniquePtr<ILegacyCacheStore> ZenIntermediarySiblingBackend(CreateZenCacheStore(TEXT("TestSibling"), ScopeZenSiblingService.GetInstance().GetURL(), *TestNamespace, *TestNamespace));
 
 	FScopeZenService ScopeZenService(MoveTemp(ZenTestServiceSettings));
-	TUniquePtr<ILegacyCacheStore> ZenIntermediaryBackend(CreateZenCacheStore(TEXT("Test"), ScopeZenService.GetInstance().GetURL(), *TestNamespace));
+	TUniquePtr<ILegacyCacheStore> ZenIntermediaryBackend(CreateZenCacheStore(TEXT("Test"), ScopeZenService.GetInstance().GetURL(), *TestNamespace, *TestNamespace));
 	auto WaitForZenPushToUpstream = [](ILegacyCacheStore* ZenBackend, TConstArrayView<FCacheRecord> Records)
 	{
 		// TODO: Expecting a legitimate means to wait for zen to finish pushing records to its upstream in the future
diff --git a/Engine/Source/Developer/DerivedDataCache/Private/ZenCacheStore.cpp b/Engine/Source/Developer/DerivedDataCache/Private/ZenCacheStore.cpp
--- a/Engine/Source/Developer/DerivedDataCache/Private/ZenCacheStore.cpp
+++ b/Engine/Source/Developer/DerivedDataCache/Private/ZenCacheStore.cpp
@@ -79,7 +79,7 @@
 	 * @param ServiceUrl	Base url to the service including schema.
 	 * @param Namespace		Namespace to use.
 	 */
-	FZenCacheStore(const TCHAR* ServiceUrl, const TCHAR* Namespace);
+	FZenCacheStore(const TCHAR* ServiceUrl, const TCHAR* Namespace, const TCHAR* StructuredNamespace);
 
 	~FZenCacheStore();
 
@@ -154,7 +154,7 @@
 	EGetResult GetZenData(const FCacheKey& Key, ECachePolicy CachePolicy, FCbPackage& OutPackage) const;
 
 	bool IsServiceReady();
-	static FString MakeLegacyZenKey(const TCHAR* CacheKey);
+	FString MakeLegacyZenKey(const TCHAR* CacheKey);
 	static void AppendZenUri(const FCacheKey& CacheKey, FStringBuilderBase& Out);
 	static void AppendZenUri(const FCacheKey& CacheKey, const FValueId& Id, FStringBuilderBase& Out);
 	static void AppendPolicyQueryString(ECachePolicy Policy, FStringBuilderBase& Out);
@@ -163,6 +163,7 @@
 
 private:
 	FString Namespace;
+	FString StructuredNamespace;
 	UE::Zen::FScopeZenService ZenService;
 	mutable FDerivedDataCacheUsageStats UsageStats;
 	TUniquePtr<UE::Zen::FZenHttpRequestPool> RequestPool;
@@ -179,8 +180,10 @@
 
 FZenCacheStore::FZenCacheStore(
 	const TCHAR* InServiceUrl,
-	const TCHAR* InNamespace)
+	const TCHAR* InNamespace,
+	const TCHAR* InStructuredNamespace)
 	: Namespace(InNamespace)
+	, StructuredNamespace(InStructuredNamespace)
 	, ZenService(InServiceUrl)
 {
 	if (IsServiceReady())
@@ -458,7 +461,12 @@
 FString FZenCacheStore::MakeLegacyZenKey(const TCHAR* CacheKey)
 {
 	FIoHash KeyHash = FIoHash::HashBuffer(CacheKey, FCString::Strlen(CacheKey) * sizeof(TCHAR));
-	return FString::Printf(TEXT("/z$/legacy/%s"), *LexToString(KeyHash));
+	if (Namespace == "ue4.ddc")
+	{
+		// TODO: DE:20220525 Remove once we have deployed shared instances with namespace support
+		return FString::Printf(TEXT("/z$/legacy/%s"), *LexToString(KeyHash));
+	}
+	return FString::Printf(TEXT("/z$/%s/legacy/%s"), *Namespace, *LexToString(KeyHash));
 }
 
 void FZenCacheStore::AppendZenUri(const FCacheKey& CacheKey, FStringBuilderBase& Out)
@@ -645,6 +653,7 @@
 			{
 				ECachePolicy BatchDefaultPolicy = Batch[0].Policy.GetRecordPolicy();
 				BatchWriter << ANSITEXTVIEW("DefaultPolicy") << *WriteToString<128>(BatchDefaultPolicy);
+				BatchWriter.AddString(ANSITEXTVIEW("Namespace"), StructuredNamespace);
 
 				BatchWriter.BeginArray(ANSITEXTVIEW("Requests"));
 				for (const FCachePutRequest& Request : Batch)
@@ -757,7 +766,7 @@
 			{
 				ECachePolicy BatchDefaultPolicy = Batch[0].Policy.GetRecordPolicy();
 				BatchRequest << ANSITEXTVIEW("DefaultPolicy") << *WriteToString<128>(BatchDefaultPolicy);
-				BatchRequest.AddString(ANSITEXTVIEW("Namespace"), Namespace);
+				BatchRequest.AddString(ANSITEXTVIEW("Namespace"), StructuredNamespace);
 
 				BatchRequest.BeginArray(ANSITEXTVIEW("Requests"));
 				for (const FCacheGetRequest& Request : Batch)
@@ -879,7 +888,7 @@
 			{
 				ECachePolicy BatchDefaultPolicy = Batch[0].Policy;
 				BatchWriter << ANSITEXTVIEW("DefaultPolicy") << *WriteToString<128>(BatchDefaultPolicy);
-				BatchWriter.AddString(ANSITEXTVIEW("Namespace"), Namespace);
+				BatchWriter.AddString(ANSITEXTVIEW("Namespace"), StructuredNamespace);
 
 				BatchWriter.BeginArray("Requests");
 				for (const FCachePutValueRequest& Request : Batch)
@@ -991,7 +1000,7 @@
 				{
 					ECachePolicy BatchDefaultPolicy = Batch[0].Policy;
 					BatchRequest << ANSITEXTVIEW("DefaultPolicy") << *WriteToString<128>(BatchDefaultPolicy);
-					BatchRequest.AddString(ANSITEXTVIEW("Namespace"), Namespace);
+					BatchRequest.AddString(ANSITEXTVIEW("Namespace"), StructuredNamespace);
 
 					BatchRequest.BeginArray("Requests");
 					for (const FCacheGetValueRequest& Request : Batch)
@@ -1119,7 +1128,7 @@
 			{
 				ECachePolicy DefaultPolicy = Batch[0].Policy;
 				BatchRequest << ANSITEXTVIEW("DefaultPolicy") << WriteToString<128>(DefaultPolicy);
-				BatchRequest.AddString(ANSITEXTVIEW("Namespace"), Namespace);
+				BatchRequest.AddString(ANSITEXTVIEW("Namespace"), StructuredNamespace);
 
 				BatchRequest.BeginArray(ANSITEXTVIEW("ChunkRequests"));
 				for (const FCacheGetChunkRequest& Request : Batch)
@@ -1231,9 +1240,9 @@
 	TRACE_COUNTER_SUBTRACT(ZenDDC_ChunkRequestCount, int64(Requests.Num()));
 }
 
-ILegacyCacheStore* CreateZenCacheStore(const TCHAR* NodeName, const TCHAR* ServiceUrl, const TCHAR* Namespace)
+ILegacyCacheStore* CreateZenCacheStore(const TCHAR* NodeName, const TCHAR* ServiceUrl, const TCHAR* Namespace, const TCHAR* StructuredNamespace)
 {
-	FZenCacheStore* Backend = new FZenCacheStore(ServiceUrl, Namespace);
+	FZenCacheStore* Backend = new FZenCacheStore(ServiceUrl, Namespace, StructuredNamespace);
 	if (Backend->IsUsable())
 	{
 		return Backend;
@@ -1250,7 +1259,7 @@
 namespace UE::DerivedData
 {
 
-ILegacyCacheStore* CreateZenCacheStore(const TCHAR* NodeName, const TCHAR* ServiceUrl, const TCHAR* Namespace)
+ILegacyCacheStore* CreateZenCacheStore(const TCHAR* NodeName, const TCHAR* ServiceUrl, const TCHAR* Namespace, const TCHAR* StructuredNamespace)
 {
 	return nullptr;
 }
