Cleanup and private detail pass on Http Client code within DDC in preparation for re-use from multiple cache stores.

diff --git a/Engine/Source/Developer/DerivedDataCache/Private/Http/HttpClient.cpp b/Engine/Source/Developer/DerivedDataCache/Private/Http/HttpClient.cpp
--- a/Engine/Source/Developer/DerivedDataCache/Private/Http/HttpClient.cpp
+++ b/Engine/Source/Developer/DerivedDataCache/Private/Http/HttpClient.cpp
@@ -2,15 +2,76 @@
 
 #include "HttpClient.h"
 
-#if WITH_HTTP_CLIENT
+#if UE_WITH_HTTP_CLIENT
 #include "Tasks/Task.h"
 
-#define UE_HTTPDDC_BACKEND_WAIT_INTERVAL 0.01f
-#define UE_HTTPDDC_BACKEND_WAIT_INTERVAL_MS ((uint32)(UE_HTTPDDC_BACKEND_WAIT_INTERVAL*1000))
+#define CURL_NO_OLDIES
+#if PLATFORM_MICROSOFT
+#include "Microsoft/AllowMicrosoftPlatformTypes.h"
+#endif
+
+#if defined(PLATFORM_CURL_INCLUDE)
+#include PLATFORM_CURL_INCLUDE
+#else
+#include "curl/curl.h"
+#endif //defined(PLATFORM_CURL_INCLUDE)
+
+#if PLATFORM_MICROSOFT
+#include "Microsoft/HideMicrosoftPlatformTypes.h"
+#endif
+
 
 namespace UE
 {
 
+namespace Http::Private
+{
+
+static constexpr uint32 WaitIntervalMs = ((uint32)(0.01f*1000));
+
+struct FHttpSharedDataInternals
+{
+	CURLSH* CurlShare;
+	CURLM* CurlMulti;
+	UE::TDepletableMpscQueue<CURL*> PendingRequestAdditions;
+	FRWLock Locks[CURL_LOCK_DATA_LAST];
+	bool WriteLocked[CURL_LOCK_DATA_LAST]{};
+};
+
+struct FHttpSharedDataStatics
+{
+	static void LockFn(CURL* Handle, curl_lock_data Data, curl_lock_access Access, void* User)
+	{
+		FHttpSharedDataInternals* SharedDataInternals = ((FHttpSharedData*)User)->Internals.Get();
+		if (Access == CURL_LOCK_ACCESS_SHARED)
+		{
+			SharedDataInternals->Locks[Data].ReadLock();
+		}
+		else
+		{
+			SharedDataInternals->Locks[Data].WriteLock();
+			SharedDataInternals->WriteLocked[Data] = true;
+		}
+	}
+
+	static void UnlockFn(CURL* Handle, curl_lock_data Data, void* User)
+	{
+		FHttpSharedDataInternals* SharedDataInternals = ((FHttpSharedData*)User)->Internals.Get();
+		if (!SharedDataInternals->WriteLocked[Data])
+		{
+			SharedDataInternals->Locks[Data].ReadUnlock();
+		}
+		else
+		{
+			SharedDataInternals->WriteLocked[Data] = false;
+			SharedDataInternals->Locks[Data].WriteUnlock();
+		}
+	}
+};
+
+}
+
+
 FString FHttpAccessToken::GetHeader()
 {
 	Lock.ReadLock();
@@ -35,14 +96,15 @@ uint32 FHttpAccessToken::GetSerial() const
 FHttpSharedData::FHttpSharedData()
 : PendingRequestEvent(EEventMode::AutoReset)
 {
-	CurlShare = curl_share_init();
-	curl_share_setopt(CurlShare, CURLSHOPT_USERDATA, this);
-	curl_share_setopt(CurlShare, CURLSHOPT_LOCKFUNC, LockFn);
-	curl_share_setopt(CurlShare, CURLSHOPT_UNLOCKFUNC, UnlockFn);
-	curl_share_setopt(CurlShare, CURLSHOPT_SHARE, CURL_LOCK_DATA_DNS);
-	curl_share_setopt(CurlShare, CURLSHOPT_SHARE, CURL_LOCK_DATA_SSL_SESSION);
-	CurlMulti = curl_multi_init();
-	curl_multi_setopt(CurlMulti, CURLMOPT_PIPELINING, CURLPIPE_MULTIPLEX);
+	Internals = MakePimpl<Http::Private::FHttpSharedDataInternals>();
+	Internals->CurlShare = curl_share_init();
+	curl_share_setopt(Internals->CurlShare, CURLSHOPT_USERDATA, this);
+	curl_share_setopt(Internals->CurlShare, CURLSHOPT_LOCKFUNC, Http::Private::FHttpSharedDataStatics::LockFn);
+	curl_share_setopt(Internals->CurlShare, CURLSHOPT_UNLOCKFUNC, Http::Private::FHttpSharedDataStatics::UnlockFn);
+	curl_share_setopt(Internals->CurlShare, CURLSHOPT_SHARE, CURL_LOCK_DATA_DNS);
+	curl_share_setopt(Internals->CurlShare, CURLSHOPT_SHARE, CURL_LOCK_DATA_SSL_SESSION);
+	Internals->CurlMulti = curl_multi_init();
+	curl_multi_setopt(Internals->CurlMulti, CURLMOPT_PIPELINING, CURLPIPE_MULTIPLEX);
 }
 
 FHttpSharedData::~FHttpSharedData()
@@ -53,13 +115,13 @@ FHttpSharedData::~FHttpSharedData()
 		AsyncServiceThread.Join();
 	}
 
-	curl_multi_cleanup(CurlMulti);
-	curl_share_cleanup(CurlShare);
+	curl_multi_cleanup(Internals->CurlMulti);
+	curl_share_cleanup(Internals->CurlShare);
 }
 
-void FHttpSharedData::AddRequest(CURL* Curl)
+void FHttpSharedData::AddRequest(void* Curl)
 {
-	if (PendingRequestAdditions.EnqueueAndReturnWasEmpty(Curl))
+	if (Internals->PendingRequestAdditions.EnqueueAndReturnWasEmpty(static_cast<CURL*>(Curl)))
 	{
 		PendingRequestEvent.Trigger();
 	}
@@ -70,32 +132,9 @@ void FHttpSharedData::AddRequest(CURL* Curl)
 	}
 }
 
-void FHttpSharedData::LockFn(CURL* Handle, curl_lock_data Data, curl_lock_access Access, void* User)
-{
-	FHttpSharedData* SharedData = (FHttpSharedData*)User;
-	if (Access == CURL_LOCK_ACCESS_SHARED)
-	{
-		SharedData->Locks[Data].ReadLock();
-	}
-	else
-	{
-		SharedData->Locks[Data].WriteLock();
-		SharedData->WriteLocked[Data] = true;
-	}
-}
-
-void FHttpSharedData::UnlockFn(CURL* Handle, curl_lock_data Data, void* User)
+void* FHttpSharedData::GetCurlShare() const
 {
-	FHttpSharedData* SharedData = (FHttpSharedData*)User;
-	if (!SharedData->WriteLocked[Data])
-	{
-		SharedData->Locks[Data].ReadUnlock();
-	}
-	else
-	{
-		SharedData->WriteLocked[Data] = false;
-		SharedData->Locks[Data].WriteUnlock();
-	}
+	return Internals->CurlShare;
 }
 
 void FHttpSharedData::ProcessAsyncRequests()
@@ -108,20 +147,20 @@ void FHttpSharedData::ProcessAsyncRequests()
 
 		do
 		{
-			PendingRequestAdditions.Deplete([this, &ActiveTransfers](CURL* Curl)
+			Internals->PendingRequestAdditions.Deplete([this, &ActiveTransfers](CURL* Curl)
 				{
-					curl_multi_add_handle(CurlMulti, Curl);
+					curl_multi_add_handle(Internals->CurlMulti, Curl);
 					++ActiveTransfers;
 				});
 
-			curl_multi_perform(CurlMulti, &CurrentActiveTransfers);
+			curl_multi_perform(Internals->CurlMulti, &CurrentActiveTransfers);
 
 			if (CurrentActiveTransfers == 0 || ActiveTransfers != CurrentActiveTransfers)
 			{
 				for (;;)
 				{
 					int MsgsStillInQueue = 0;	// may use that to impose some upper limit we may spend in that loop
-					CURLMsg* Message = curl_multi_info_read(CurlMulti, &MsgsStillInQueue);
+					CURLMsg* Message = curl_multi_info_read(Internals->CurlMulti, &MsgsStillInQueue);
 
 					if (!Message)
 					{
@@ -132,7 +171,7 @@ void FHttpSharedData::ProcessAsyncRequests()
 					if (Message->msg == CURLMSG_DONE)
 					{
 						CURL* CompletedHandle = Message->easy_handle;
-						curl_multi_remove_handle(CurlMulti, CompletedHandle);
+						curl_multi_remove_handle(Internals->CurlMulti, CompletedHandle);
 
 						void* PrivateData = nullptr;
 						curl_easy_getinfo(CompletedHandle, CURLINFO_PRIVATE, &PrivateData);
@@ -155,7 +194,7 @@ void FHttpSharedData::ProcessAsyncRequests()
 
 			if (CurrentActiveTransfers > 0)
 			{
-				curl_multi_wait(CurlMulti, nullptr, 0, 1, nullptr);
+				curl_multi_wait(Internals->CurlMulti, nullptr, 0, 1, nullptr);
 			}
 		}
 		while (CurrentActiveTransfers > 0);
@@ -240,7 +279,7 @@ FHttpRequest* FHttpRequestPool::WaitForFreeRequest(bool bUnboundedOverflow)
 
 		Waiters.enqueue(Waiter);
 
-		while (!Waiter->Wait(UE_HTTPDDC_BACKEND_WAIT_INTERVAL_MS))
+		while (!Waiter->Wait(Http::Private::WaitIntervalMs))
 		{
 			// While waiting, allow us to check if a race occurred and a request has been freed
 			// between the time we checked for free requests and the time we queued ourself as a Waiter.
@@ -317,4 +356,4 @@ void FHttpRequestPool::MakeRequestShared(FHttpRequest* Request, uint8 Users)
 
 } // UE
 
-#endif // WITH_HTTP_CLIENT
+#endif // UE_WITH_HTTP_CLIENT
diff --git a/Engine/Source/Developer/DerivedDataCache/Private/Http/HttpClient.h b/Engine/Source/Developer/DerivedDataCache/Private/Http/HttpClient.h
--- a/Engine/Source/Developer/DerivedDataCache/Private/Http/HttpClient.h
+++ b/Engine/Source/Developer/DerivedDataCache/Private/Http/HttpClient.h
@@ -3,9 +3,9 @@
 #pragma once
 
 #include "HAL/Platform.h"
-#define WITH_HTTP_CLIENT PLATFORM_DESKTOP
+#define UE_WITH_HTTP_CLIENT PLATFORM_DESKTOP
 
-#if WITH_HTTP_CLIENT
+#if UE_WITH_HTTP_CLIENT
 #include "CoreTypes.h"
 #include "Containers/Array.h"
 #include "Containers/ArrayView.h"
@@ -22,28 +22,16 @@
 #include "HAL/Thread.h"
 #include "Memory/MemoryView.h"
 #include "Memory/SharedBuffer.h"
+#include "Templates/PimplPtr.h"
 #include "Templates/SharedPointer.h"
 #include <atomic>
 
-#if WITH_SSL
-#include "Ssl.h"
-#include <openssl/ssl.h>
-#endif
 
-#if PLATFORM_MICROSOFT
-#include "Microsoft/AllowMicrosoftPlatformTypes.h"
-#endif
-
-// if there is a platform-specific include then it must be used in place of forward declation in case it defines CURL_STRICTER
-#if defined(PLATFORM_CURL_INCLUDE)
-#include PLATFORM_CURL_INCLUDE
-#else
-#include "curl/curl.h"
-#endif //defined(PLATFORM_CURL_INCLUDE)
-
-#if PLATFORM_MICROSOFT
-#include "Microsoft/HideMicrosoftPlatformTypes.h"
-#endif
+struct curl_slist;
+namespace UE::Http::Private { struct FAsyncRequestData; }
+namespace UE::Http::Private { struct FHttpRequestStatics; }
+namespace UE::Http::Private { struct FHttpSharedDataInternals; }
+namespace UE::Http::Private { struct FHttpSharedDataStatics; }
 
 namespace UE
 {
@@ -61,9 +49,9 @@ public:
 	void SetHeader(const TCHAR*);
 	uint32 GetSerial() const;
 private:
-	FRWLock		Lock;
-	FString		Token;
-	uint32		Serial;
+	FRWLock Lock;
+	FString Token;
+	uint32 Serial;
 };
 
 class FHttpSharedData
@@ -71,25 +59,20 @@ class FHttpSharedData
 public:
 	FHttpSharedData();
 	~FHttpSharedData();
-
-	void AddRequest(CURL* Curl);
-	CURLSH* GetCurlShare() const { return CurlShare; }
-
 private:
-	CURLSH* CurlShare;
-	CURLM* CurlMulti;
-	UE::TDepletableMpscQueue<CURL*> PendingRequestAdditions;
+	friend class FHttpRequest;
+	friend struct Http::Private::FHttpSharedDataStatics;
+
 	UE::FLazyEvent PendingRequestEvent;
 	FThread AsyncServiceThread;
 	std::atomic<bool> bAsyncThreadStarting = false;
 	std::atomic<bool> bAsyncThreadShutdownRequested = false;
 
-	FRWLock Locks[CURL_LOCK_DATA_LAST];
-	bool WriteLocked[CURL_LOCK_DATA_LAST]{};
+	TPimplPtr<Http::Private::FHttpSharedDataInternals> Internals;
 
-	static void LockFn(CURL* Handle, curl_lock_data Data, curl_lock_access Access, void* User);
-	static void UnlockFn(CURL* Handle, curl_lock_data Data, void* User);
 	void ProcessAsyncRequests();
+	void AddRequest(void* Curl);
+	void* GetCurlShare() const;
 };
 
 
@@ -135,6 +118,8 @@ public:
 		Retry
 	};
 
+	using FResultCode = int;
+
 	using FOnHttpRequestComplete = TUniqueFunction<ECompletionBehavior(Result HttpResult, FHttpRequest* Request)>;
 
 	FHttpRequest(const TCHAR* InDomain, const TCHAR* InEffectiveDomain, FHttpAccessToken* InAuthorizationToken, FHttpSharedData* InSharedData, bool bInLogErrors);
@@ -196,112 +181,10 @@ public:
 	 * @return Result of the request
 	 */
 	template<RequestVerb V>
-	Result PerformBlockingUpload(const TCHAR* Uri, TArrayView<const uint8> Buffer, TConstArrayView<long> ExpectedErrorCodes = {})
-	{
-		static_assert(V == Put || V == PutCompactBinary || V == PutCompressedBlob || V == Post || V == PostCompactBinary || V == PostJson, "Upload should use either Put or Post verbs.");
-		
-		uint64 ContentLength = 0u;
-
-		if constexpr (V == Put || V == PutCompactBinary || V == PutCompressedBlob)
-		{
-			curl_easy_setopt(Curl, CURLOPT_UPLOAD, 1L);
-			curl_easy_setopt(Curl, CURLOPT_INFILESIZE, Buffer.Num());
-			curl_easy_setopt(Curl, CURLOPT_READDATA, this);
-			curl_easy_setopt(Curl, CURLOPT_READFUNCTION, StaticReadFn);
-			if constexpr (V == PutCompactBinary)
-			{
-				Headers.Add(FString(TEXT("Content-Type: application/x-ue-cb")));
-			}
-			else if constexpr (V == PutCompressedBlob)
-			{
-				Headers.Add(FString(TEXT("Content-Type: application/x-ue-comp")));
-			}
-			else
-			{
-				Headers.Add(FString(TEXT("Content-Type: application/octet-stream")));
-			}
-			ContentLength = Buffer.Num();
-			ReadDataView = MakeMemoryView(Buffer);
-		}
-		else if constexpr (V == Post || V == PostCompactBinary || V == PostJson)
-		{
-			curl_easy_setopt(Curl, CURLOPT_POST, 1L);
-			curl_easy_setopt(Curl, CURLOPT_INFILESIZE, Buffer.Num());
-			curl_easy_setopt(Curl, CURLOPT_READDATA, this);
-			curl_easy_setopt(Curl, CURLOPT_READFUNCTION, StaticReadFn);
-			if constexpr (V == PostCompactBinary)
-			{
-				Headers.Add(FString(TEXT("Content-Type: application/x-ue-cb")));
-			}
-			else if constexpr (V == PostJson)
-			{
-				Headers.Add(FString(TEXT("Content-Type: application/json")));
-			}
-			else
-			{
-				Headers.Add(FString(TEXT("Content-Type: application/x-www-form-urlencoded")));
-			}
-			ContentLength = Buffer.Num();
-			ReadDataView = MakeMemoryView(Buffer);
-		}
-
-		return PerformBlocking(Uri, V, ContentLength, ExpectedErrorCodes);
-	}
+	Result PerformBlockingUpload(const TCHAR* Uri, TArrayView<const uint8> Buffer, TConstArrayView<long> ExpectedErrorCodes = {});
 	
 	template<RequestVerb V>
-	void EnqueueAsyncUpload(UE::DerivedData::IRequestOwner& Owner, FHttpRequestPool* Pool, const TCHAR* Uri, FSharedBuffer Buffer, FOnHttpRequestComplete&& OnComplete, TConstArrayView<long> ExpectedErrorCodes = {})
-	{
-		static_assert(V == Put || V == PutCompactBinary || V == PutCompressedBlob || V == Post || V == PostCompactBinary || V == PostJson, "Upload should use either Put or Post verbs.");
-		
-		uint64 ContentLength = 0u;
-
-		if constexpr (V == Put || V == PutCompactBinary || V == PutCompressedBlob)
-		{
-			curl_easy_setopt(Curl, CURLOPT_UPLOAD, 1L);
-			curl_easy_setopt(Curl, CURLOPT_INFILESIZE, Buffer.GetSize());
-			curl_easy_setopt(Curl, CURLOPT_READDATA, this);
-			curl_easy_setopt(Curl, CURLOPT_READFUNCTION, StaticReadFn);
-			if constexpr (V == PutCompactBinary)
-			{
-				Headers.Add(FString(TEXT("Content-Type: application/x-ue-cb")));
-			}
-			else if constexpr (V == PutCompressedBlob)
-			{
-				Headers.Add(FString(TEXT("Content-Type: application/x-ue-comp")));
-			}
-			else
-			{
-				Headers.Add(FString(TEXT("Content-Type: application/octet-stream")));
-			}
-			ReadSharedBuffer = Buffer;
-			ContentLength = Buffer.GetSize();
-			ReadDataView = Buffer.GetView();
-		}
-		else if constexpr (V == Post || V == PostCompactBinary || V == PostJson)
-		{
-			curl_easy_setopt(Curl, CURLOPT_POST, 1L);
-			curl_easy_setopt(Curl, CURLOPT_INFILESIZE, Buffer.GetSize());
-			curl_easy_setopt(Curl, CURLOPT_READDATA, this);
-			curl_easy_setopt(Curl, CURLOPT_READFUNCTION, StaticReadFn);
-			if constexpr (V == PostCompactBinary)
-			{
-				Headers.Add(FString(TEXT("Content-Type: application/x-ue-cb")));
-			}
-			else if constexpr (V == PostJson)
-			{
-				Headers.Add(FString(TEXT("Content-Type: application/json")));
-			}
-			else
-			{
-				Headers.Add(FString(TEXT("Content-Type: application/x-www-form-urlencoded")));
-			}
-			ReadSharedBuffer = Buffer;
-			ContentLength = Buffer.GetSize();
-			ReadDataView = Buffer.GetView();
-		}
-
-		return EnqueueAsync(Owner, Pool, Uri, V, ContentLength, MoveTemp(OnComplete), ExpectedErrorCodes);
-	}
+	void EnqueueAsyncUpload(UE::DerivedData::IRequestOwner& Owner, FHttpRequestPool* Pool, const TCHAR* Uri, FSharedBuffer Buffer, FOnHttpRequestComplete&& OnComplete, TConstArrayView<long> ExpectedErrorCodes = {});
 
 	/**
 	 * Download an url into a buffer using the request.
@@ -320,38 +203,10 @@ public:
 	 * @return Result of the request
 	 */
 	template<RequestVerb V>
-	Result PerformBlockingQuery(const TCHAR* Uri, TConstArrayView<long> ExpectedErrorCodes = {400})
-	{
-		static_assert(V == Head || V == Delete, "Queries should use either Head or Delete verbs.");
-
-		if (V == Delete)
-		{
-			curl_easy_setopt(Curl, CURLOPT_CUSTOMREQUEST, "DELETE");
-		}
-		else if (V == Head)
-		{
-			curl_easy_setopt(Curl, CURLOPT_NOBODY, 1L);
-		}
-
-		return PerformBlocking(Uri, V, 0u, ExpectedErrorCodes);
-	}
+	Result PerformBlockingQuery(const TCHAR* Uri, TConstArrayView<long> ExpectedErrorCodes = {400});
 
 	template<RequestVerb V>
-	void EnqueueAsyncQuery(UE::DerivedData::IRequestOwner& Owner, FHttpRequestPool* Pool, const TCHAR* Uri, FOnHttpRequestComplete&& OnComplete, TConstArrayView<long> ExpectedErrorCodes = {400})
-	{
-		static_assert(V == Head || V == Delete, "Queries should use either Head or Delete verbs.");
-
-		if (V == Delete)
-		{
-			curl_easy_setopt(Curl, CURLOPT_CUSTOMREQUEST, "DELETE");
-		}
-		else if (V == Head)
-		{
-			curl_easy_setopt(Curl, CURLOPT_NOBODY, 1L);
-		}
-
-		return EnqueueAsync(Owner, Pool, Uri, V, 0u, MoveTemp(OnComplete), ExpectedErrorCodes);
-	}
+	void EnqueueAsyncQuery(UE::DerivedData::IRequestOwner& Owner, FHttpRequestPool* Pool, const TCHAR* Uri, FOnHttpRequestComplete&& OnComplete, TConstArrayView<long> ExpectedErrorCodes = {400});
 
 	/**
 	 * Set a header to send with the request.
@@ -406,7 +261,7 @@ public:
 	 */
 	TArray<TSharedPtr<FJsonValue>> GetResponseAsJsonArray() const;
 
-	void CompleteAsync(CURLcode Result);
+	void CompleteAsync(FResultCode Result);
 
 	/** Will return true if the response code is considered a success */
 	static bool IsSuccessResponse(long ResponseCode)
@@ -418,65 +273,30 @@ public:
 	static bool AllowAsync();
 
 private:
-
-	struct FAsyncRequestData final : public UE::DerivedData::FRequestBase
-	{
-		UE::DerivedData::IRequestOwner* Owner = nullptr;
-		FHttpRequestPool* Pool = nullptr;
-		curl_slist* CurlHeaders = nullptr;
-		FString Uri;
-		RequestVerb Verb;
-		TArray<long, TInlineAllocator<4>> ExpectedErrorCodes;
-		FOnHttpRequestComplete OnComplete;
-		UE::FLazyEvent Event {EEventMode::ManualReset};
-
-		void Reset()
-		{
-			if (CurlHeaders)
-			{
-				curl_slist_free_all(CurlHeaders);
-				CurlHeaders = nullptr;
-			}
-			Uri.Empty();
-			ExpectedErrorCodes.Empty();
-		}
-
-		void SetPriority(UE::DerivedData::EPriority Priority) final {}
-
-		void Cancel() final
-		{
-			TRACE_CPUPROFILER_EVENT_SCOPE(HttpDDC_Cancel);
-			Event.Wait();
-		}
-
-		void Wait() final
-		{
-			TRACE_CPUPROFILER_EVENT_SCOPE(HttpDDC_Wait);
-			Event.Wait();
-		}
-	};
-
-	CURL*					Curl;
-	CURLcode				CurlResult;
-	FHttpSharedData*		SharedData;
-	FAsyncRequestData*		AsyncData;
-	long					ResponseCode;
-	size_t					BytesSent;
-	size_t					BytesReceived;
-	size_t					Attempts;
-	bool					bLogErrors;
-
-	FSharedBuffer			ReadSharedBuffer;
-	FMemoryView				ReadDataView;
-	TArray<uint8>*			WriteDataBufferPtr;
-	TArray<uint8>*			WriteHeaderBufferPtr;
+	friend struct Http::Private::FAsyncRequestData;
+	friend struct Http::Private::FHttpRequestStatics;
+
+	void* Curl;
+	FResultCode ResultCode;
+	FHttpSharedData* SharedData;
+	Http::Private::FAsyncRequestData* AsyncData;
+	long ResponseCode;
+	size_t BytesSent;
+	size_t BytesReceived;
+	size_t Attempts;
+	bool bLogErrors;
+
+	FSharedBuffer ReadSharedBuffer;
+	FMemoryView ReadDataView;
+	TArray<uint8>* WriteDataBufferPtr;
+	TArray<uint8>* WriteHeaderBufferPtr;
 	
-	TArray<uint8>			ResponseHeader;
-	TArray<uint8>			ResponseBuffer;
-	TArray<FString>			Headers;
-	FString					Domain;
-	FString					EffectiveDomain;
-	FHttpAccessToken*		AuthorizationToken;
+	TArray<uint8> ResponseHeader;
+	TArray<uint8> ResponseBuffer;
+	TArray<FString> Headers;
+	FString Domain;
+	FString EffectiveDomain;
+	FHttpAccessToken* AuthorizationToken;
 
 	curl_slist* PrepareToIssueRequest(const TCHAR* Uri, uint64 ContentLength);
 
@@ -492,7 +312,7 @@ private:
 
 	void EnqueueAsync(UE::DerivedData::IRequestOwner& Owner, FHttpRequestPool* Pool, const TCHAR* Uri, RequestVerb Verb, uint64 ContentLength, FOnHttpRequestComplete&& OnComplete, TConstArrayView<long> ExpectedErrorCodes);
 
-	void LogResult(CURLcode Result, const TCHAR* Uri, RequestVerb Verb, TConstArrayView<long> ExpectedErrorCodes) const;
+	void LogResult(FResultCode Result, const TCHAR* Uri, RequestVerb Verb, TConstArrayView<long> ExpectedErrorCodes) const;
 
 	FString GetAnsiBufferAsString(const TArray<uint8>& Buffer) const
 	{
@@ -501,21 +321,37 @@ private:
 		return FString(TCHARData.Length(), TCHARData.Get());
 	}
 
-	//----------------------------------------------------------------------------------------------------------
-	// Certificate checking
-	//----------------------------------------------------------------------------------------------------------
-	#if WITH_SSL
-	static int SslCertVerify(int PreverifyOk, X509_STORE_CTX* Context);
-	static CURLcode StaticSSLCTXFn(CURL * curl, void * sslctx, void * parm);
-	#endif //#if WITH_SSL
+};
 
-	static size_t StaticDebugCallback(CURL * Handle, curl_infotype DebugInfoType, char * DebugInfo, size_t DebugInfoSize, void* UserData);
-	static size_t StaticReadFn(void* Ptr, size_t SizeInBlocks, size_t BlockSizeInBytes, void* UserData);
-	static size_t StaticWriteHeaderFn(void* Ptr, size_t SizeInBlocks, size_t BlockSizeInBytes, void* UserData);
-	static size_t StaticWriteBodyFn(void* Ptr, size_t SizeInBlocks, size_t BlockSizeInBytes, void* UserData);
-	static size_t StaticSeekFn(void* UserData, curl_off_t Offset, int Origin);
+#define UE_HTTP_FOREACH_UPLOAD_VERB(op) \
+	op(FHttpRequest::Put); \
+	op(FHttpRequest::PutCompactBinary); \
+	op(FHttpRequest::PutCompressedBlob); \
+	op(FHttpRequest::Post); \
+	op(FHttpRequest::PostCompactBinary); \
+	op(FHttpRequest::PostJson);
+
+#define UE_HTTP_FOREACH_QUERY_VERB(op) \
+	op(FHttpRequest::Head) \
+	op(FHttpRequest::Delete)
+
+#define UE_HTTP_DECLARE_EXTERN_TEMPLATE(Verb) extern template FHttpRequest::Result FHttpRequest::PerformBlockingUpload<Verb>(const TCHAR* Uri, TArrayView<const uint8> Buffer, TConstArrayView<long> ExpectedErrorCodes);
+UE_HTTP_FOREACH_UPLOAD_VERB(UE_HTTP_DECLARE_EXTERN_TEMPLATE)
+#undef UE_HTTP_DECLARE_EXTERN_TEMPLATE
+
+#define UE_HTTP_DECLARE_EXTERN_TEMPLATE(Verb) extern template void FHttpRequest::EnqueueAsyncUpload<Verb>(UE::DerivedData::IRequestOwner& Owner, FHttpRequestPool* Pool, const TCHAR* Uri, FSharedBuffer Buffer, FOnHttpRequestComplete&& OnComplete, TConstArrayView<long> ExpectedErrorCodes);
+UE_HTTP_FOREACH_UPLOAD_VERB(UE_HTTP_DECLARE_EXTERN_TEMPLATE)
+#undef UE_HTTP_DECLARE_EXTERN_TEMPLATE
+
+
+#define UE_HTTP_DECLARE_EXTERN_TEMPLATE(Verb) extern template FHttpRequest::Result FHttpRequest::PerformBlockingQuery<Verb>(const TCHAR* Uri, TConstArrayView<long> ExpectedErrorCodes);
+UE_HTTP_FOREACH_QUERY_VERB(UE_HTTP_DECLARE_EXTERN_TEMPLATE)
+#undef UE_HTTP_DECLARE_EXTERN_TEMPLATE
+
+#define UE_HTTP_DECLARE_EXTERN_TEMPLATE(Verb) extern template void FHttpRequest::EnqueueAsyncQuery<Verb>(UE::DerivedData::IRequestOwner& Owner, FHttpRequestPool* Pool, const TCHAR* Uri, FOnHttpRequestComplete&& OnComplete, TConstArrayView<long> ExpectedErrorCodes);
+UE_HTTP_FOREACH_QUERY_VERB(UE_HTTP_DECLARE_EXTERN_TEMPLATE)
+#undef UE_HTTP_DECLARE_EXTERN_TEMPLATE
 
-};
 
 /**
  * Pool that manages a fixed set of requests. Users are required to release requests that have been 
@@ -663,4 +499,4 @@ private:
 
 } // UE
 
-#endif // WITH_HTTP_CLIENT
+#endif // UE_WITH_HTTP_CLIENT
diff --git a/Engine/Source/Developer/DerivedDataCache/Private/Http/HttpRequest.cpp b/Engine/Source/Developer/DerivedDataCache/Private/Http/HttpRequest.cpp
--- a/Engine/Source/Developer/DerivedDataCache/Private/Http/HttpRequest.cpp
+++ b/Engine/Source/Developer/DerivedDataCache/Private/Http/HttpRequest.cpp
@@ -2,17 +2,32 @@
 
 #include "HttpClient.h"
 
-#if WITH_HTTP_CLIENT
+#if UE_WITH_HTTP_CLIENT
 #include "DerivedDataLegacyCacheStore.h"
 #include "HAL/IConsoleManager.h"
 #include "Logging/LogMacros.h"
 #include "Serialization/JsonReader.h"
 #include "Serialization/JsonSerializer.h"
 
-#define UE_HTTPDDC_HTTP_REQUEST_TIMEOUT_SECONDS 30L
-#define UE_HTTPDDC_HTTP_REQUEST_TIMEOUT_ENABLED 1
-#define UE_HTTPDDC_HTTP_DEBUG 0
-#define UE_HTTPDDC_MAX_BUFFER_RESERVE 104857600u
+#if WITH_SSL
+#include "Ssl.h"
+#include <openssl/ssl.h>
+#endif
+
+#define CURL_NO_OLDIES
+#if PLATFORM_MICROSOFT
+#include "Microsoft/AllowMicrosoftPlatformTypes.h"
+#endif
+
+#if defined(PLATFORM_CURL_INCLUDE)
+#include PLATFORM_CURL_INCLUDE
+#else
+#include "curl/curl.h"
+#endif //defined(PLATFORM_CURL_INCLUDE)
+
+#if PLATFORM_MICROSOFT
+#include "Microsoft/HideMicrosoftPlatformTypes.h"
+#endif
 
 namespace UE
 {
@@ -24,6 +39,232 @@ static FAutoConsoleVariableRef CVarHttpEnableAsync(
 	TEXT("If true, async operations are permitted, otherwise all operations are forced to be synchronous."),
 	ECVF_Default);
 
+namespace Http::Private
+{
+
+static constexpr bool RequestDebugEnabled = false;
+static constexpr bool RequestTimeoutEnabled = true;
+static constexpr long RequestTimeoutSeconds = 30L;
+static constexpr size_t RequestMaxBufferReserve = 104857600u;
+
+struct FAsyncRequestData final : public DerivedData::FRequestBase
+{
+	UE::DerivedData::IRequestOwner* Owner = nullptr;
+	UE::FHttpRequestPool* Pool = nullptr;
+	curl_slist* CurlHeaders = nullptr;
+	FString Uri;
+	UE::FHttpRequest::RequestVerb Verb;
+	TArray<long, TInlineAllocator<4>> ExpectedErrorCodes;
+	UE::FHttpRequest::FOnHttpRequestComplete OnComplete;
+	UE::FLazyEvent Event {EEventMode::ManualReset};
+
+	void Reset()
+	{
+		if (CurlHeaders)
+		{
+			curl_slist_free_all(CurlHeaders);
+			CurlHeaders = nullptr;
+		}
+		Uri.Empty();
+		ExpectedErrorCodes.Empty();
+	}
+
+	void SetPriority(UE::DerivedData::EPriority Priority) final {}
+
+	void Cancel() final
+	{
+		TRACE_CPUPROFILER_EVENT_SCOPE(HttpDDC_Cancel);
+		Event.Wait();
+	}
+
+	void Wait() final
+	{
+		TRACE_CPUPROFILER_EVENT_SCOPE(HttpDDC_Wait);
+		Event.Wait();
+	}
+};
+
+struct FHttpRequestStatics
+{
+	#if WITH_SSL
+	static int SslCertVerify(int PreverifyOk, X509_STORE_CTX* Context)
+	{
+		if (PreverifyOk == 1)
+		{
+			SSL* Handle = static_cast<SSL*>(X509_STORE_CTX_get_ex_data(Context, SSL_get_ex_data_X509_STORE_CTX_idx()));
+			check(Handle);
+
+			SSL_CTX* SslContext = SSL_get_SSL_CTX(Handle);
+			check(SslContext);
+
+			FHttpRequest* Request = static_cast<FHttpRequest*>(SSL_CTX_get_app_data(SslContext));
+			check(Request);
+
+			const FString& Domain = Request->GetDomain();
+
+			if (!FSslModule::Get().GetCertificateManager().VerifySslCertificates(Context, Domain))
+			{
+				PreverifyOk = 0;
+			}
+		}
+
+		return PreverifyOk;
+	}
+
+	static UE::FHttpRequest::FResultCode StaticSSLCTXFn(CURL* curl, void* sslctx, void* parm)
+	{
+		SSL_CTX* Context = static_cast<SSL_CTX*>(sslctx);
+		const ISslCertificateManager& CertificateManager = FSslModule::Get().GetCertificateManager();
+
+		CertificateManager.AddCertificatesToSslContext(Context);
+		SSL_CTX_set_verify(Context, SSL_CTX_get_verify_mode(Context), SslCertVerify);
+		SSL_CTX_set_app_data(Context, parm);
+
+		/* all set to go */
+		return CURLE_OK;
+	}
+	#endif //#if WITH_SSL
+
+	static size_t StaticDebugCallback(CURL * Handle, curl_infotype DebugInfoType, char * DebugInfo, size_t DebugInfoSize, void* UserData)
+	{
+		FHttpRequest* Request = static_cast<FHttpRequest*>(UserData);
+
+		switch (DebugInfoType)
+		{
+		case CURLINFO_TEXT:
+		{
+			// Truncate at 1023 characters. This is just an arbitrary number based on a buffer size seen in
+			// the libcurl code.
+			DebugInfoSize = FMath::Min(DebugInfoSize, (size_t)1023);
+
+			// Calculate the actual length of the string due to incorrect use of snprintf() in lib/vtls/openssl.c.
+			char* FoundNulPtr = (char*)memchr(DebugInfo, 0, DebugInfoSize);
+			int CalculatedSize = FoundNulPtr != nullptr ? FoundNulPtr - DebugInfo : DebugInfoSize;
+
+			auto ConvertedString = StringCast<TCHAR>(static_cast<const ANSICHAR*>(DebugInfo), CalculatedSize);
+			FString DebugText(ConvertedString.Length(), ConvertedString.Get());
+			DebugText.ReplaceInline(TEXT("\n"), TEXT(""), ESearchCase::CaseSensitive);
+			DebugText.ReplaceInline(TEXT("\r"), TEXT(""), ESearchCase::CaseSensitive);
+			UE_LOG(LogDerivedDataCache, VeryVerbose, TEXT("%s: %p: '%s'"), *Request->GetName(), Request, *DebugText);
+		}
+		break;
+
+		case CURLINFO_HEADER_IN:
+			UE_LOG(LogDerivedDataCache, VeryVerbose, TEXT("%s: %p: Received header (%d bytes)"), *Request->GetName(), Request, DebugInfoSize);
+			break;
+
+		case CURLINFO_DATA_IN:
+			UE_LOG(LogDerivedDataCache, VeryVerbose, TEXT("%s: %p: Received data (%d bytes)"), *Request->GetName(), Request, DebugInfoSize);
+			break;
+
+		case CURLINFO_DATA_OUT:
+			UE_LOG(LogDerivedDataCache, VeryVerbose, TEXT("%s: %p: Sent data (%d bytes)"), *Request->GetName(), Request, DebugInfoSize);
+			break;
+
+		case CURLINFO_SSL_DATA_IN:
+			UE_LOG(LogDerivedDataCache, VeryVerbose, TEXT("%s: %p: Received SSL data (%d bytes)"), *Request->GetName(), Request, DebugInfoSize);
+			break;
+
+		case CURLINFO_SSL_DATA_OUT:
+			UE_LOG(LogDerivedDataCache, VeryVerbose, TEXT("%s: %p: Sent SSL data (%d bytes)"), *Request->GetName(), Request, DebugInfoSize);
+			break;
+		}
+
+		return 0;
+	}
+
+	static size_t StaticReadFn(void* Ptr, size_t SizeInBlocks, size_t BlockSizeInBytes, void* UserData)
+	{
+		FHttpRequest* Request = static_cast<FHttpRequest*>(UserData);
+		const size_t MaxReadSize = SizeInBlocks * BlockSizeInBytes;
+		const FMemoryView SourceView = Request->ReadDataView.Mid(Request->BytesSent, MaxReadSize);
+		MakeMemoryView(Ptr, MaxReadSize).CopyFrom(SourceView);
+		Request->BytesSent += SourceView.GetSize();
+		return SourceView.GetSize();
+	}
+
+	static size_t StaticWriteHeaderFn(void* Ptr, size_t SizeInBlocks, size_t BlockSizeInBytes, void* UserData)
+	{
+		FHttpRequest* Request = static_cast<FHttpRequest*>(UserData);
+		const size_t WriteSize = SizeInBlocks * BlockSizeInBytes;
+		TArray<uint8>* WriteHeaderBufferPtr = Request->WriteHeaderBufferPtr;
+		if (WriteHeaderBufferPtr && WriteSize > 0)
+		{
+			const size_t CurrentBufferLength = WriteHeaderBufferPtr->Num();
+			if (CurrentBufferLength > 0)
+			{
+				// Remove the previous zero termination
+				(*WriteHeaderBufferPtr)[CurrentBufferLength-1] = ' ';
+			}
+
+			// Write the header
+			WriteHeaderBufferPtr->Append((const uint8*)Ptr, WriteSize + 1);
+			(*WriteHeaderBufferPtr)[WriteHeaderBufferPtr->Num()-1] = 0; // Zero terminate string
+			return WriteSize;
+		}
+		return 0;
+	}
+
+	static size_t StaticWriteBodyFn(void* Ptr, size_t SizeInBlocks, size_t BlockSizeInBytes, void* UserData)
+	{
+		FHttpRequest* Request = static_cast<FHttpRequest*>(UserData);
+		const size_t WriteSize = SizeInBlocks * BlockSizeInBytes;
+		TArray<uint8>* WriteDataBufferPtr = Request->WriteDataBufferPtr;
+
+		if (WriteDataBufferPtr && WriteSize > 0)
+		{
+			// If this is the first part of the body being received, try to reserve 
+			// memory if content length is defined in the header.
+			if (Request->BytesReceived == 0 && Request->WriteHeaderBufferPtr)
+			{
+				static const ANSICHAR* ContentLengthHeaderStr = "Content-Length: ";
+				const ANSICHAR* Header = (const ANSICHAR*)Request->WriteHeaderBufferPtr->GetData();
+
+				if (const ANSICHAR* ContentLengthHeader = FCStringAnsi::Strstr(Header, ContentLengthHeaderStr))
+				{
+					size_t ContentLength = (size_t)FCStringAnsi::Atoi64(ContentLengthHeader + strlen(ContentLengthHeaderStr));
+					if (ContentLength > 0u && ContentLength < Http::Private::RequestMaxBufferReserve)
+					{
+						WriteDataBufferPtr->Reserve(ContentLength);
+					}
+				}
+			}
+
+			// Write to the target buffer
+			WriteDataBufferPtr->Append((const uint8*)Ptr, WriteSize);
+			Request->BytesReceived += WriteSize;
+			return WriteSize;
+		}
+		
+		return 0;
+	}
+
+	static size_t StaticSeekFn(void* UserData, curl_off_t Offset, int Origin)
+	{
+		FHttpRequest* Request = static_cast<FHttpRequest*>(UserData);
+		size_t NewPosition = 0;
+
+		switch (Origin)
+		{
+		case SEEK_SET: NewPosition = Offset; break;
+		case SEEK_CUR: NewPosition = Request->BytesSent + Offset; break;
+		case SEEK_END: NewPosition = Request->ReadDataView.GetSize() + Offset; break;
+		}
+
+		// Make sure we don't seek outside of the buffer
+		if (NewPosition < 0 || NewPosition >= Request->ReadDataView.GetSize())
+		{
+			return CURL_SEEKFUNC_FAIL;
+		}
+
+		// Update the used offset
+		Request->BytesSent = NewPosition;
+		return CURL_SEEKFUNC_OK;
+	}
+};
+
+}
+
 FHttpRequest::FHttpRequest(const TCHAR* InDomain, const TCHAR* InEffectiveDomain, FHttpAccessToken* InAuthorizationToken, FHttpSharedData* InSharedData, bool bInLogErrors)
 	: SharedData(InSharedData)
 	, AsyncData(nullptr)
@@ -38,7 +279,7 @@ FHttpRequest::FHttpRequest(const TCHAR* InDomain, const TCHAR* InEffectiveDomain
 
 FHttpRequest::~FHttpRequest()
 {
-	curl_easy_cleanup(Curl);
+	curl_easy_cleanup(static_cast<CURL*>(Curl));
 	check(!AsyncData);
 }
 
@@ -54,51 +295,54 @@ void FHttpRequest::Reset()
 	BytesSent = 0;
 	BytesReceived = 0;
 	Attempts = 0;
-	CurlResult = CURL_LAST;
+	ResultCode = CURL_LAST;
 
-	curl_easy_reset(Curl);
+	CURL* LocalCurl = static_cast<CURL*>(Curl);
+	curl_easy_reset(LocalCurl);
 	check(!AsyncData);
 
 	// Options that are always set for all connections.
-#if UE_HTTPDDC_HTTP_REQUEST_TIMEOUT_ENABLED
-	curl_easy_setopt(Curl, CURLOPT_CONNECTTIMEOUT, UE_HTTPDDC_HTTP_REQUEST_TIMEOUT_SECONDS);
-#endif
-	curl_easy_setopt(Curl, CURLOPT_FOLLOWLOCATION, 1L);
-	curl_easy_setopt(Curl, CURLOPT_NOSIGNAL, 1L);
-	curl_easy_setopt(Curl, CURLOPT_DNS_CACHE_TIMEOUT, 300L); // Don't re-resolve every minute
-	curl_easy_setopt(Curl, CURLOPT_SHARE, SharedData ? SharedData->GetCurlShare() : nullptr);
+	if constexpr(Http::Private::RequestTimeoutEnabled)
+	{
+		curl_easy_setopt(LocalCurl, CURLOPT_CONNECTTIMEOUT, Http::Private::RequestTimeoutSeconds);
+	}
+	curl_easy_setopt(LocalCurl, CURLOPT_FOLLOWLOCATION, 1L);
+	curl_easy_setopt(LocalCurl, CURLOPT_NOSIGNAL, 1L);
+	curl_easy_setopt(LocalCurl, CURLOPT_DNS_CACHE_TIMEOUT, 300L); // Don't re-resolve every minute
+	curl_easy_setopt(LocalCurl, CURLOPT_SHARE, SharedData ? SharedData->GetCurlShare() : nullptr);
 	// Response functions
-	curl_easy_setopt(Curl, CURLOPT_HEADERDATA, this);
-	curl_easy_setopt(Curl, CURLOPT_HEADERFUNCTION, &FHttpRequest::StaticWriteHeaderFn);
-	curl_easy_setopt(Curl, CURLOPT_WRITEDATA, this);
-	curl_easy_setopt(Curl, CURLOPT_WRITEFUNCTION, StaticWriteBodyFn);
+	curl_easy_setopt(LocalCurl, CURLOPT_HEADERDATA, this);
+	curl_easy_setopt(LocalCurl, CURLOPT_HEADERFUNCTION, Http::Private::FHttpRequestStatics::StaticWriteHeaderFn);
+	curl_easy_setopt(LocalCurl, CURLOPT_WRITEDATA, this);
+	curl_easy_setopt(LocalCurl, CURLOPT_WRITEFUNCTION, Http::Private::FHttpRequestStatics::StaticWriteBodyFn);
 	#if WITH_SSL
 	// SSL options
-	curl_easy_setopt(Curl, CURLOPT_USE_SSL, CURLUSESSL_ALL);
-	curl_easy_setopt(Curl, CURLOPT_SSL_VERIFYPEER, 1);
-	curl_easy_setopt(Curl, CURLOPT_SSL_VERIFYHOST, 1);
-	curl_easy_setopt(Curl, CURLOPT_SSLCERTTYPE, "PEM");
+	curl_easy_setopt(LocalCurl, CURLOPT_USE_SSL, CURLUSESSL_ALL);
+	curl_easy_setopt(LocalCurl, CURLOPT_SSL_VERIFYPEER, 1);
+	curl_easy_setopt(LocalCurl, CURLOPT_SSL_VERIFYHOST, 1);
+	curl_easy_setopt(LocalCurl, CURLOPT_SSLCERTTYPE, "PEM");
 	// SSL certification verification
-	curl_easy_setopt(Curl, CURLOPT_CAINFO, nullptr);
-	curl_easy_setopt(Curl, CURLOPT_SSL_CTX_FUNCTION, &FHttpRequest::StaticSSLCTXFn);
-	curl_easy_setopt(Curl, CURLOPT_SSL_CTX_DATA, this);
+	curl_easy_setopt(LocalCurl, CURLOPT_CAINFO, nullptr);
+	curl_easy_setopt(LocalCurl, CURLOPT_SSL_CTX_FUNCTION, Http::Private::FHttpRequestStatics::StaticSSLCTXFn);
+	curl_easy_setopt(LocalCurl, CURLOPT_SSL_CTX_DATA, this);
 	#endif //#if WITH_SSL
 	// Allow compressed data
-	curl_easy_setopt(Curl, CURLOPT_ACCEPT_ENCODING, "gzip");
+	curl_easy_setopt(LocalCurl, CURLOPT_ACCEPT_ENCODING, "gzip");
 	// Rewind method, handle special error case where request need to rewind data stream
-	curl_easy_setopt(Curl, CURLOPT_SEEKFUNCTION, StaticSeekFn);
-	curl_easy_setopt(Curl, CURLOPT_SEEKDATA, this);
+	curl_easy_setopt(LocalCurl, CURLOPT_SEEKFUNCTION, Http::Private::FHttpRequestStatics::StaticSeekFn);
+	curl_easy_setopt(LocalCurl, CURLOPT_SEEKDATA, this);
 	// Set minimum speed behavior to allow operations to abort if the transfer speed is poor for the given duration (1kbps over a 30 second span)
-	curl_easy_setopt(Curl, CURLOPT_LOW_SPEED_TIME, 30L);
-	curl_easy_setopt(Curl, CURLOPT_LOW_SPEED_LIMIT, 1024L);
+	curl_easy_setopt(LocalCurl, CURLOPT_LOW_SPEED_TIME, 30L);
+	curl_easy_setopt(LocalCurl, CURLOPT_LOW_SPEED_LIMIT, 1024L);
 	// Debug hooks
-#if UE_HTTPDDC_HTTP_DEBUG
-	curl_easy_setopt(Curl, CURLOPT_DEBUGDATA, this);
-	curl_easy_setopt(Curl, CURLOPT_DEBUGFUNCTION, StaticDebugCallback);
-	curl_easy_setopt(Curl, CURLOPT_VERBOSE, 1L);
-#endif
-	curl_easy_setopt(Curl, CURLOPT_PRIVATE, this);
-	curl_easy_setopt(Curl, CURLOPT_HTTP_VERSION, CURL_HTTP_VERSION_2_0);
+	if constexpr(Http::Private::RequestDebugEnabled)
+	{
+		curl_easy_setopt(LocalCurl, CURLOPT_DEBUGDATA, this);
+		curl_easy_setopt(LocalCurl, CURLOPT_DEBUGFUNCTION, Http::Private::FHttpRequestStatics::StaticDebugCallback);
+		curl_easy_setopt(LocalCurl, CURLOPT_VERBOSE, 1L);
+	}
+	curl_easy_setopt(LocalCurl, CURLOPT_PRIVATE, this);
+	curl_easy_setopt(LocalCurl, CURLOPT_HTTP_VERSION, CURL_HTTP_VERSION_2_0);
 }
 
 void FHttpRequest::PrepareToRetry()
@@ -108,13 +352,13 @@ void FHttpRequest::PrepareToRetry()
 	ResponseCode = 0;
 	BytesSent = 0;
 	BytesReceived = 0;
-	CurlResult = CURL_LAST;
+	ResultCode = CURL_LAST;
 	++Attempts;
 }
 
 FHttpRequest::Result FHttpRequest::PerformBlockingDownload(const TCHAR* Uri, TArray<uint8>* Buffer, TConstArrayView<long> ExpectedErrorCodes)
 {
-	curl_easy_setopt(Curl, CURLOPT_HTTPGET, 1L);
+	curl_easy_setopt(static_cast<CURL*>(Curl), CURLOPT_HTTPGET, 1L);
 	WriteDataBufferPtr = Buffer;
 
 	return PerformBlocking(Uri, Get, 0u, ExpectedErrorCodes);
@@ -122,20 +366,20 @@ FHttpRequest::Result FHttpRequest::PerformBlockingDownload(const TCHAR* Uri, TAr
 
 void FHttpRequest::EnqueueAsyncDownload(UE::DerivedData::IRequestOwner& Owner, FHttpRequestPool* Pool, const TCHAR* Uri, FOnHttpRequestComplete&& OnComplete, TConstArrayView<long> ExpectedErrorCodes)
 {
-	curl_easy_setopt(Curl, CURLOPT_HTTPGET, 1L);
+	curl_easy_setopt(static_cast<CURL*>(Curl), CURLOPT_HTTPGET, 1L);
 
 	return EnqueueAsync(Owner, Pool, Uri, Get, 0u, MoveTemp(OnComplete), ExpectedErrorCodes);
 }
 
 void FHttpRequest::SetHeader(const TCHAR* Header, const TCHAR* Value)
 {
-	check(CurlResult == CURL_LAST); // Cannot set header after request is sent
+	check(ResultCode == CURL_LAST); // Cannot set header after request is sent
 	Headers.Add(FString::Printf(TEXT("%s: %s"), Header, Value));
 }
 
 bool FHttpRequest::GetHeader(const ANSICHAR* Header, FString& OutValue) const
 {
-	check(CurlResult != CURL_LAST);  // Cannot query headers before request is sent
+	check(ResultCode != CURL_LAST);  // Cannot query headers before request is sent
 
 	const ANSICHAR* HeadersBuffer = (const ANSICHAR*) ResponseHeader.GetData();
 	size_t HeaderLen = strlen(Header);
@@ -187,12 +431,12 @@ bool FHttpRequest::AllowAsync()
 	return true;
 }
 
-void FHttpRequest::CompleteAsync(CURLcode Result)
+void FHttpRequest::CompleteAsync(FResultCode Result)
 {
-	CurlResult = Result;
+	ResultCode = Result;
 
 	// Get response code
-	curl_easy_getinfo(Curl, CURLINFO_RESPONSE_CODE, &ResponseCode);
+	curl_easy_getinfo(static_cast<CURL*>(Curl), CURLINFO_RESPONSE_CODE, &ResponseCode);
 
 	LogResult(Result, *AsyncData->Uri, AsyncData->Verb, AsyncData->ExpectedErrorCodes);
 
@@ -200,7 +444,7 @@ void FHttpRequest::CompleteAsync(CURLcode Result)
 	{
 		UE::DerivedData::FRequestBarrier Barrier(*AsyncData->Owner, UE::DerivedData::ERequestBarrierFlags::Priority);
 		FHttpRequest::Result HttpResult;
-		switch (CurlResult)
+		switch (ResultCode)
 		{
 		case CURLE_OK:
 			HttpResult = Success;
@@ -247,7 +491,7 @@ curl_slist* FHttpRequest::PrepareToIssueRequest(const TCHAR* Uri, uint64 Content
 
 	// Setup request options
 	FString Url = FString::Printf(TEXT("%s/%s"), *EffectiveDomain, Uri);
-	curl_easy_setopt(Curl, CURLOPT_URL, TCHAR_TO_ANSI(*Url));
+	curl_easy_setopt(static_cast<CURL*>(Curl), CURLOPT_URL, TCHAR_TO_ANSI(*Url));
 
 	// Setup response header buffer. If caller has not setup a response data buffer, use internal.
 	WriteHeaderBufferPtr = &ResponseHeader;
@@ -277,7 +521,7 @@ curl_slist* FHttpRequest::PrepareToIssueRequest(const TCHAR* Uri, uint64 Content
 	{
 		CurlHeaders = curl_slist_append(CurlHeaders, TCHAR_TO_ANSI(*Header));
 	}
-	curl_easy_setopt(Curl, CURLOPT_HTTPHEADER, CurlHeaders);
+	curl_easy_setopt(static_cast<CURL*>(Curl), CURLOPT_HTTPHEADER, CurlHeaders);
 	return CurlHeaders;
 }
 
@@ -289,17 +533,17 @@ FHttpRequest::Result FHttpRequest::PerformBlocking(const TCHAR* Uri, RequestVerb
 	curl_slist* CurlHeaders = PrepareToIssueRequest(Uri, ContentLength);
 
 	// Shots fired!
-	CurlResult = curl_easy_perform(Curl);
+	ResultCode = curl_easy_perform(static_cast<CURL*>(Curl));
 
 	// Get response code
-	curl_easy_getinfo(Curl, CURLINFO_RESPONSE_CODE, &ResponseCode);
+	curl_easy_getinfo(static_cast<CURL*>(Curl), CURLINFO_RESPONSE_CODE, &ResponseCode);
 
-	LogResult(CurlResult, Uri, Verb, ExpectedErrorCodes);
+	LogResult(ResultCode, Uri, Verb, ExpectedErrorCodes);
 
 	// Clean up
 	curl_slist_free_all(CurlHeaders);
 
-	return CurlResult == CURLE_OK ? Success : Failed;
+	return ResultCode == CURLE_OK ? Success : Failed;
 }
 
 void FHttpRequest::EnqueueAsync(UE::DerivedData::IRequestOwner& Owner, FHttpRequestPool* Pool, const TCHAR* Uri, RequestVerb Verb, uint64 ContentLength, FOnHttpRequestComplete&& OnComplete, TConstArrayView<long> ExpectedErrorCodes)
@@ -318,7 +562,7 @@ void FHttpRequest::EnqueueAsync(UE::DerivedData::IRequestOwner& Owner, FHttpRequ
 	}
 
 	TRACE_CPUPROFILER_EVENT_SCOPE(HttpDDC_CurlEnqueueAsync);
-	AsyncData = new FAsyncRequestData;
+	AsyncData = new Http::Private::FAsyncRequestData;
 	AsyncData->Owner = &Owner;
 	AsyncData->Pool = Pool;
 	AsyncData->CurlHeaders = PrepareToIssueRequest(Uri, ContentLength);
@@ -331,7 +575,7 @@ void FHttpRequest::EnqueueAsync(UE::DerivedData::IRequestOwner& Owner, FHttpRequ
 	SharedData->AddRequest(Curl);
 }
 
-void FHttpRequest::LogResult(CURLcode Result, const TCHAR* Uri, RequestVerb Verb, TConstArrayView<long> ExpectedErrorCodes) const
+void FHttpRequest::LogResult(FResultCode Result, const TCHAR* Uri, RequestVerb Verb, TConstArrayView<long> ExpectedErrorCodes) const
 {
 	if (Result == CURLE_OK)
 	{
@@ -425,187 +669,172 @@ void FHttpRequest::LogResult(CURLcode Result, const TCHAR* Uri, RequestVerb Verb
 			TEXT("%s: Error while connecting to %s: %s"), 
 			*GetName(),
 			*EffectiveDomain,
-			ANSI_TO_TCHAR(curl_easy_strerror(Result))
+			ANSI_TO_TCHAR(curl_easy_strerror(static_cast<CURLcode>(Result)))
 		);
 	}
 }
 
-#if WITH_SSL
-int FHttpRequest::SslCertVerify(int PreverifyOk, X509_STORE_CTX* Context)
+template<FHttpRequest::RequestVerb V>
+FHttpRequest::Result FHttpRequest::PerformBlockingUpload(const TCHAR* Uri, TArrayView<const uint8> Buffer, TConstArrayView<long> ExpectedErrorCodes)
 {
-	if (PreverifyOk == 1)
-	{
-		SSL* Handle = static_cast<SSL*>(X509_STORE_CTX_get_ex_data(Context, SSL_get_ex_data_X509_STORE_CTX_idx()));
-		check(Handle);
-
-		SSL_CTX* SslContext = SSL_get_SSL_CTX(Handle);
-		check(SslContext);
-
-		FHttpRequest* Request = static_cast<FHttpRequest*>(SSL_CTX_get_app_data(SslContext));
-		check(Request);
-
-		const FString& Domain = Request->GetDomain();
+	static_assert(V == Put || V == PutCompactBinary || V == PutCompressedBlob || V == Post || V == PostCompactBinary || V == PostJson, "Upload should use either Put or Post verbs.");
+	
+	uint64 ContentLength = 0u;
 
-		if (!FSslModule::Get().GetCertificateManager().VerifySslCertificates(Context, Domain))
+	CURL* LocalCurl = static_cast<CURL*>(Curl);
+	if constexpr (V == Put || V == PutCompactBinary || V == PutCompressedBlob)
+	{
+		curl_easy_setopt(LocalCurl, CURLOPT_UPLOAD, 1L);
+		curl_easy_setopt(LocalCurl, CURLOPT_INFILESIZE, Buffer.Num());
+		curl_easy_setopt(LocalCurl, CURLOPT_READDATA, this);
+		curl_easy_setopt(LocalCurl, CURLOPT_READFUNCTION, Http::Private::FHttpRequestStatics::StaticReadFn);
+		if constexpr (V == PutCompactBinary)
+		{
+			Headers.Add(FString(TEXT("Content-Type: application/x-ue-cb")));
+		}
+		else if constexpr (V == PutCompressedBlob)
+		{
+			Headers.Add(FString(TEXT("Content-Type: application/x-ue-comp")));
+		}
+		else
 		{
-			PreverifyOk = 0;
+			Headers.Add(FString(TEXT("Content-Type: application/octet-stream")));
 		}
+		ContentLength = Buffer.Num();
+		ReadDataView = MakeMemoryView(Buffer);
 	}
-
-	return PreverifyOk;
-}
-
-CURLcode FHttpRequest::StaticSSLCTXFn(CURL * curl, void * sslctx, void * parm)
-{
-	SSL_CTX* Context = static_cast<SSL_CTX*>(sslctx);
-	const ISslCertificateManager& CertificateManager = FSslModule::Get().GetCertificateManager();
-
-	CertificateManager.AddCertificatesToSslContext(Context);
-	SSL_CTX_set_verify(Context, SSL_CTX_get_verify_mode(Context), SslCertVerify);
-	SSL_CTX_set_app_data(Context, parm);
-
-	/* all set to go */
-	return CURLE_OK;
-}
-#endif //#if WITH_SSL
-
-size_t FHttpRequest::StaticDebugCallback(CURL * Handle, curl_infotype DebugInfoType, char * DebugInfo, size_t DebugInfoSize, void* UserData)
-{
-	FHttpRequest* Request = static_cast<FHttpRequest*>(UserData);
-
-	switch (DebugInfoType)
-	{
-	case CURLINFO_TEXT:
+	else if constexpr (V == Post || V == PostCompactBinary || V == PostJson)
 	{
-		// Truncate at 1023 characters. This is just an arbitrary number based on a buffer size seen in
-		// the libcurl code.
-		DebugInfoSize = FMath::Min(DebugInfoSize, (size_t)1023);
-
-		// Calculate the actual length of the string due to incorrect use of snprintf() in lib/vtls/openssl.c.
-		char* FoundNulPtr = (char*)memchr(DebugInfo, 0, DebugInfoSize);
-		int CalculatedSize = FoundNulPtr != nullptr ? FoundNulPtr - DebugInfo : DebugInfoSize;
-
-		auto ConvertedString = StringCast<TCHAR>(static_cast<const ANSICHAR*>(DebugInfo), CalculatedSize);
-		FString DebugText(ConvertedString.Length(), ConvertedString.Get());
-		DebugText.ReplaceInline(TEXT("\n"), TEXT(""), ESearchCase::CaseSensitive);
-		DebugText.ReplaceInline(TEXT("\r"), TEXT(""), ESearchCase::CaseSensitive);
-		UE_LOG(LogDerivedDataCache, VeryVerbose, TEXT("%s: %p: '%s'"), *Request->GetName(), Request, *DebugText);
-	}
-	break;
-
-	case CURLINFO_HEADER_IN:
-		UE_LOG(LogDerivedDataCache, VeryVerbose, TEXT("%s: %p: Received header (%d bytes)"), *Request->GetName(), Request, DebugInfoSize);
-		break;
-
-	case CURLINFO_DATA_IN:
-		UE_LOG(LogDerivedDataCache, VeryVerbose, TEXT("%s: %p: Received data (%d bytes)"), *Request->GetName(), Request, DebugInfoSize);
-		break;
-
-	case CURLINFO_DATA_OUT:
-		UE_LOG(LogDerivedDataCache, VeryVerbose, TEXT("%s: %p: Sent data (%d bytes)"), *Request->GetName(), Request, DebugInfoSize);
-		break;
-
-	case CURLINFO_SSL_DATA_IN:
-		UE_LOG(LogDerivedDataCache, VeryVerbose, TEXT("%s: %p: Received SSL data (%d bytes)"), *Request->GetName(), Request, DebugInfoSize);
-		break;
-
-	case CURLINFO_SSL_DATA_OUT:
-		UE_LOG(LogDerivedDataCache, VeryVerbose, TEXT("%s: %p: Sent SSL data (%d bytes)"), *Request->GetName(), Request, DebugInfoSize);
-		break;
+		curl_easy_setopt(LocalCurl, CURLOPT_POST, 1L);
+		curl_easy_setopt(LocalCurl, CURLOPT_INFILESIZE, Buffer.Num());
+		curl_easy_setopt(LocalCurl, CURLOPT_READDATA, this);
+		curl_easy_setopt(LocalCurl, CURLOPT_READFUNCTION, Http::Private::FHttpRequestStatics::StaticReadFn);
+		if constexpr (V == PostCompactBinary)
+		{
+			Headers.Add(FString(TEXT("Content-Type: application/x-ue-cb")));
+		}
+		else if constexpr (V == PostJson)
+		{
+			Headers.Add(FString(TEXT("Content-Type: application/json")));
+		}
+		else
+		{
+			Headers.Add(FString(TEXT("Content-Type: application/x-www-form-urlencoded")));
+		}
+		ContentLength = Buffer.Num();
+		ReadDataView = MakeMemoryView(Buffer);
 	}
 
-	return 0;
+	return PerformBlocking(Uri, V, ContentLength, ExpectedErrorCodes);
 }
 
-size_t FHttpRequest::StaticReadFn(void* Ptr, size_t SizeInBlocks, size_t BlockSizeInBytes, void* UserData)
+template<FHttpRequest::RequestVerb V>
+void FHttpRequest::EnqueueAsyncUpload(UE::DerivedData::IRequestOwner& Owner, FHttpRequestPool* Pool, const TCHAR* Uri, FSharedBuffer Buffer, FOnHttpRequestComplete&& OnComplete, TConstArrayView<long> ExpectedErrorCodes)
 {
-	FHttpRequest* Request = static_cast<FHttpRequest*>(UserData);
-	const size_t MaxReadSize = SizeInBlocks * BlockSizeInBytes;
-	const FMemoryView SourceView = Request->ReadDataView.Mid(Request->BytesSent, MaxReadSize);
-	MakeMemoryView(Ptr, MaxReadSize).CopyFrom(SourceView);
-	Request->BytesSent += SourceView.GetSize();
-	return SourceView.GetSize();
-}
+	static_assert(V == Put || V == PutCompactBinary || V == PutCompressedBlob || V == Post || V == PostCompactBinary || V == PostJson, "Upload should use either Put or Post verbs.");
+	
+	uint64 ContentLength = 0u;
 
-size_t FHttpRequest::StaticWriteHeaderFn(void* Ptr, size_t SizeInBlocks, size_t BlockSizeInBytes, void* UserData)
-{
-	FHttpRequest* Request = static_cast<FHttpRequest*>(UserData);
-	const size_t WriteSize = SizeInBlocks * BlockSizeInBytes;
-	TArray<uint8>* WriteHeaderBufferPtr = Request->WriteHeaderBufferPtr;
-	if (WriteHeaderBufferPtr && WriteSize > 0)
+	CURL* LocalCurl = static_cast<CURL*>(Curl);
+	if constexpr (V == Put || V == PutCompactBinary || V == PutCompressedBlob)
 	{
-		const size_t CurrentBufferLength = WriteHeaderBufferPtr->Num();
-		if (CurrentBufferLength > 0)
+		curl_easy_setopt(LocalCurl, CURLOPT_UPLOAD, 1L);
+		curl_easy_setopt(LocalCurl, CURLOPT_INFILESIZE, Buffer.GetSize());
+		curl_easy_setopt(LocalCurl, CURLOPT_READDATA, this);
+		curl_easy_setopt(LocalCurl, CURLOPT_READFUNCTION, Http::Private::FHttpRequestStatics::StaticReadFn);
+		if constexpr (V == PutCompactBinary)
 		{
-			// Remove the previous zero termination
-			(*WriteHeaderBufferPtr)[CurrentBufferLength-1] = ' ';
+			Headers.Add(FString(TEXT("Content-Type: application/x-ue-cb")));
 		}
-
-		// Write the header
-		WriteHeaderBufferPtr->Append((const uint8*)Ptr, WriteSize + 1);
-		(*WriteHeaderBufferPtr)[WriteHeaderBufferPtr->Num()-1] = 0; // Zero terminate string
-		return WriteSize;
+		else if constexpr (V == PutCompressedBlob)
+		{
+			Headers.Add(FString(TEXT("Content-Type: application/x-ue-comp")));
+		}
+		else
+		{
+			Headers.Add(FString(TEXT("Content-Type: application/octet-stream")));
+		}
+		ReadSharedBuffer = Buffer;
+		ContentLength = Buffer.GetSize();
+		ReadDataView = Buffer.GetView();
 	}
-	return 0;
+	else if constexpr (V == Post || V == PostCompactBinary || V == PostJson)
+	{
+		curl_easy_setopt(LocalCurl, CURLOPT_POST, 1L);
+		curl_easy_setopt(LocalCurl, CURLOPT_INFILESIZE, Buffer.GetSize());
+		curl_easy_setopt(LocalCurl, CURLOPT_READDATA, this);
+		curl_easy_setopt(LocalCurl, CURLOPT_READFUNCTION, Http::Private::FHttpRequestStatics::StaticReadFn);
+		if constexpr (V == PostCompactBinary)
+		{
+			Headers.Add(FString(TEXT("Content-Type: application/x-ue-cb")));
+		}
+		else if constexpr (V == PostJson)
+		{
+			Headers.Add(FString(TEXT("Content-Type: application/json")));
+		}
+		else
+		{
+			Headers.Add(FString(TEXT("Content-Type: application/x-www-form-urlencoded")));
+		}
+		ReadSharedBuffer = Buffer;
+		ContentLength = Buffer.GetSize();
+		ReadDataView = Buffer.GetView();
+	}
+
+	return EnqueueAsync(Owner, Pool, Uri, V, ContentLength, MoveTemp(OnComplete), ExpectedErrorCodes);
 }
 
-size_t FHttpRequest::StaticWriteBodyFn(void* Ptr, size_t SizeInBlocks, size_t BlockSizeInBytes, void* UserData)
+template<FHttpRequest::RequestVerb V>
+FHttpRequest::Result FHttpRequest::PerformBlockingQuery(const TCHAR* Uri, TConstArrayView<long> ExpectedErrorCodes)
 {
-	FHttpRequest* Request = static_cast<FHttpRequest*>(UserData);
-	const size_t WriteSize = SizeInBlocks * BlockSizeInBytes;
-	TArray<uint8>* WriteDataBufferPtr = Request->WriteDataBufferPtr;
+	static_assert(V == Head || V == Delete, "Queries should use either Head or Delete verbs.");
 
-	if (WriteDataBufferPtr && WriteSize > 0)
+	if (V == Delete)
 	{
-		// If this is the first part of the body being received, try to reserve 
-		// memory if content length is defined in the header.
-		if (Request->BytesReceived == 0 && Request->WriteHeaderBufferPtr)
-		{
-			static const ANSICHAR* ContentLengthHeaderStr = "Content-Length: ";
-			const ANSICHAR* Header = (const ANSICHAR*)Request->WriteHeaderBufferPtr->GetData();
-
-			if (const ANSICHAR* ContentLengthHeader = FCStringAnsi::Strstr(Header, ContentLengthHeaderStr))
-			{
-				size_t ContentLength = (size_t)FCStringAnsi::Atoi64(ContentLengthHeader + strlen(ContentLengthHeaderStr));
-				if (ContentLength > 0u && ContentLength < UE_HTTPDDC_MAX_BUFFER_RESERVE)
-				{
-					WriteDataBufferPtr->Reserve(ContentLength);
-				}
-			}
-		}
-
-		// Write to the target buffer
-		WriteDataBufferPtr->Append((const uint8*)Ptr, WriteSize);
-		Request->BytesReceived += WriteSize;
-		return WriteSize;
+		curl_easy_setopt(static_cast<CURL*>(Curl), CURLOPT_CUSTOMREQUEST, "DELETE");
 	}
-	
-	return 0;
+	else if (V == Head)
+	{
+		curl_easy_setopt(static_cast<CURL*>(Curl), CURLOPT_NOBODY, 1L);
+	}
+
+	return PerformBlocking(Uri, V, 0u, ExpectedErrorCodes);
 }
 
-size_t FHttpRequest::StaticSeekFn(void* UserData, curl_off_t Offset, int Origin)
+template<FHttpRequest::RequestVerb V>
+void FHttpRequest::EnqueueAsyncQuery(UE::DerivedData::IRequestOwner& Owner, FHttpRequestPool* Pool, const TCHAR* Uri, FOnHttpRequestComplete&& OnComplete, TConstArrayView<long> ExpectedErrorCodes)
 {
-	FHttpRequest* Request = static_cast<FHttpRequest*>(UserData);
-	size_t NewPosition = 0;
+	static_assert(V == Head || V == Delete, "Queries should use either Head or Delete verbs.");
 
-	switch (Origin)
+	if (V == Delete)
 	{
-	case SEEK_SET: NewPosition = Offset; break;
-	case SEEK_CUR: NewPosition = Request->BytesSent + Offset; break;
-	case SEEK_END: NewPosition = Request->ReadDataView.GetSize() + Offset; break;
+		curl_easy_setopt(static_cast<CURL*>(Curl), CURLOPT_CUSTOMREQUEST, "DELETE");
 	}
-
-	// Make sure we don't seek outside of the buffer
-	if (NewPosition < 0 || NewPosition >= Request->ReadDataView.GetSize())
+	else if (V == Head)
 	{
-		return CURL_SEEKFUNC_FAIL;
+		curl_easy_setopt(static_cast<CURL*>(Curl), CURLOPT_NOBODY, 1L);
 	}
 
-	// Update the used offset
-	Request->BytesSent = NewPosition;
-	return CURL_SEEKFUNC_OK;
+	return EnqueueAsync(Owner, Pool, Uri, V, 0u, MoveTemp(OnComplete), ExpectedErrorCodes);
 }
 
+#define UE_HTTP_DEFINE_EXTERN_TEMPLATE(Verb) template FHttpRequest::Result FHttpRequest::PerformBlockingUpload<Verb>(const TCHAR* Uri, TArrayView<const uint8> Buffer, TConstArrayView<long> ExpectedErrorCodes);
+UE_HTTP_FOREACH_UPLOAD_VERB(UE_HTTP_DEFINE_EXTERN_TEMPLATE)
+#undef UE_HTTP_DEFINE_EXTERN_TEMPLATE
+
+#define UE_HTTP_DEFINE_EXTERN_TEMPLATE(Verb) template void FHttpRequest::EnqueueAsyncUpload<Verb>(UE::DerivedData::IRequestOwner& Owner, FHttpRequestPool* Pool, const TCHAR* Uri, FSharedBuffer Buffer, FOnHttpRequestComplete&& OnComplete, TConstArrayView<long> ExpectedErrorCodes);
+UE_HTTP_FOREACH_UPLOAD_VERB(UE_HTTP_DEFINE_EXTERN_TEMPLATE)
+#undef UE_HTTP_DEFINE_EXTERN_TEMPLATE
+
+
+#define UE_HTTP_DEFINE_EXTERN_TEMPLATE(Verb) template FHttpRequest::Result FHttpRequest::PerformBlockingQuery<Verb>(const TCHAR* Uri, TConstArrayView<long> ExpectedErrorCodes);
+UE_HTTP_FOREACH_QUERY_VERB(UE_HTTP_DEFINE_EXTERN_TEMPLATE)
+#undef UE_HTTP_DEFINE_EXTERN_TEMPLATE
+
+#define UE_HTTP_DEFINE_EXTERN_TEMPLATE(Verb) template void FHttpRequest::EnqueueAsyncQuery<Verb>(UE::DerivedData::IRequestOwner& Owner, FHttpRequestPool* Pool, const TCHAR* Uri, FOnHttpRequestComplete&& OnComplete, TConstArrayView<long> ExpectedErrorCodes);
+UE_HTTP_FOREACH_QUERY_VERB(UE_HTTP_DEFINE_EXTERN_TEMPLATE)
+#undef UE_HTTP_DEFINE_EXTERN_TEMPLATE
+
 } // UE
 
-#endif // WITH_HTTP_CLIENT
+#endif // UE_WITH_HTTP_CLIENT
diff --git a/Engine/Source/Developer/DerivedDataCache/Private/HttpCacheStore.cpp b/Engine/Source/Developer/DerivedDataCache/Private/HttpCacheStore.cpp
--- a/Engine/Source/Developer/DerivedDataCache/Private/HttpCacheStore.cpp
+++ b/Engine/Source/Developer/DerivedDataCache/Private/HttpCacheStore.cpp
@@ -4,14 +4,6 @@
 
 #if WITH_HTTP_DDC_BACKEND
 
-#if PLATFORM_WINDOWS || PLATFORM_HOLOLENS
-#include "Windows/WindowsHWrapper.h"
-#include "Windows/AllowWindowsPlatformTypes.h"
-#endif
-#include "curl/curl.h"
-#if PLATFORM_WINDOWS || PLATFORM_HOLOLENS
-#include "Windows/HideWindowsPlatformTypes.h"
-#endif
 #include "Algo/Accumulate.h"
 #include "Algo/Find.h"
 #include "Algo/Transform.h"
@@ -58,9 +50,16 @@
 
 #include "Http/HttpClient.h"
 
-#if WITH_SSL
-#include "Ssl.h"
-#include <openssl/ssl.h>
+#if PLATFORM_WINDOWS || PLATFORM_HOLOLENS
+#include "Windows/WindowsHWrapper.h"
+#include "Windows/AllowWindowsPlatformTypes.h"
+#include <winsock2.h>
+#include <ws2tcpip.h>
+#include "Windows/HideWindowsPlatformTypes.h"
+#else
+#include <sys/types.h>
+#include <sys/socket.h>
+#include <netdb.h>
 #endif
 
 // Enables data request helpers that internally
@@ -900,51 +899,6 @@ TLockFreePointerListUnordered<FDataUploadHelper::FQueuedEntry, PLATFORM_CACHE_LI
 
 #endif // WITH_DATAREQUEST_HELPER
 
-//----------------------------------------------------------------------------------------------------------
-// Certificate checking
-//----------------------------------------------------------------------------------------------------------
-
-#if WITH_SSL
-
-static int SslCertVerify(int PreverifyOk, X509_STORE_CTX* Context)
-{
-	if (PreverifyOk == 1)
-	{
-		SSL* Handle = static_cast<SSL*>(X509_STORE_CTX_get_ex_data(Context, SSL_get_ex_data_X509_STORE_CTX_idx()));
-		check(Handle);
-
-		SSL_CTX* SslContext = SSL_get_SSL_CTX(Handle);
-		check(SslContext);
-
-		FHttpRequest* Request = static_cast<FHttpRequest*>(SSL_CTX_get_app_data(SslContext));
-		check(Request);
-
-		const FString& Domain = Request->GetDomain();
-
-		if (!FSslModule::Get().GetCertificateManager().VerifySslCertificates(Context, Domain))
-		{
-			PreverifyOk = 0;
-		}
-	}
-
-	return PreverifyOk;
-}
-
-static CURLcode sslctx_function(CURL * curl, void * sslctx, void * parm)
-{
-	SSL_CTX* Context = static_cast<SSL_CTX*>(sslctx);
-	const ISslCertificateManager& CertificateManager = FSslModule::Get().GetCertificateManager();
-
-	CertificateManager.AddCertificatesToSslContext(Context);
-	SSL_CTX_set_verify(Context, SSL_CTX_get_verify_mode(Context), SslCertVerify);
-	SSL_CTX_set_app_data(Context, parm);
-
-	/* all set to go */
-	return CURLE_OK;
-}
-
-#endif //#if WITH_SSL
-
 //----------------------------------------------------------------------------------------------------------
 // Content parsing and checking
 //----------------------------------------------------------------------------------------------------------
