Switch Zen cache store to use the HTTP request codepath in the Jupiter cache store while taking some of the improvements that had been made when that code was branched off to make the Zen HTTP codepath.  This is a precursor for getting Zen cache store using non-blocking operations like Jupiter cache store - this has started to be stubbed in with the EnqueueAsyncRpc methods that are stubbed in but not used from the Zen cache store.  Eventually I'd like this HTTP code to move out of the DerivedDataCache module and be usable elsewhere and supplant the HTTP code currently in the Zen module, but this will require some decisions about whether IRequestOwner should move out of the DerivedDataCache module or not.

Tested Zen operation using ShooterGame editor.
Tested Jupiter operation using CitySample editor (empty local cache).

diff --git a/Engine/Source/Developer/DerivedDataCache/Private/Http/HttpClient.cpp b/Engine/Source/Developer/DerivedDataCache/Private/Http/HttpClient.cpp
--- a/Engine/Source/Developer/DerivedDataCache/Private/Http/HttpClient.cpp
+++ b/Engine/Source/Developer/DerivedDataCache/Private/Http/HttpClient.cpp
@@ -211,10 +211,16 @@ void FHttpSharedData::ProcessAsyncRequests()
 	ProcessPendingRequests();
 }
 
-FHttpRequestPool::FHttpRequestPool(const TCHAR* InServiceUrl, const TCHAR* InEffectiveServiceUrl, FHttpAccessToken* InAuthorizationToken, FHttpSharedData* InSharedData, uint32 PoolSize, uint32 InOverflowLimit)
+FHttpRequestPool::FHttpRequestPool(FStringView InServiceUrl, FStringView InEffectiveServiceUrl, FHttpAccessToken* InAuthorizationToken, FHttpSharedData* InSharedData, uint32 PoolSize, uint32 InOverflowLimit)
 : ActiveOverflowRequests(0)
 , OverflowLimit(InOverflowLimit)
 {
+	int32 LastSlash = 0;
+	while (InEffectiveServiceUrl.FindLastChar(TEXT('/'), /* out */ LastSlash) && (LastSlash == (InEffectiveServiceUrl.Len() - 1)))
+	{
+		InEffectiveServiceUrl.LeftChopInline(1);
+	}
+	
 	Pool.AddUninitialized(PoolSize);
 	Requests.AddUninitialized(PoolSize);
 	for (uint8 i = 0; i < Pool.Num(); ++i)
diff --git a/Engine/Source/Developer/DerivedDataCache/Private/Http/HttpClient.h b/Engine/Source/Developer/DerivedDataCache/Private/Http/HttpClient.h
--- a/Engine/Source/Developer/DerivedDataCache/Private/Http/HttpClient.h
+++ b/Engine/Source/Developer/DerivedDataCache/Private/Http/HttpClient.h
@@ -20,8 +20,8 @@
 #include "HAL/CriticalSection.h"
 #include "HAL/PlatformProcess.h"
 #include "HAL/Thread.h"
+#include "Memory/CompositeBuffer.h"
 #include "Memory/MemoryView.h"
-#include "Memory/SharedBuffer.h"
 #include "Templates/PimplPtr.h"
 #include "Templates/SharedPointer.h"
 #include <atomic>
@@ -38,6 +38,55 @@ namespace UE
 
 class FHttpRequestPool;
 
+enum class EHttpContentType : int16
+{
+	// Negative integers reserved for future custom string content types
+	UnspecifiedContentType = 0,
+	AnyContentType,
+	Binary,
+	Text,
+	JSON,
+	CbObject,
+	CbPackage,
+	YAML,
+	CbPackageOffer,
+	CompressedBinary,
+	FormUrlEncoded,
+	Count
+};
+
+inline FStringView GetHttpMimeType(EHttpContentType Type)
+{
+	switch (Type)
+	{
+		case EHttpContentType::UnspecifiedContentType:
+			checkNoEntry();
+			return FStringView();
+		case EHttpContentType::AnyContentType:
+			return TEXTVIEW("*/*");
+		case EHttpContentType::Binary:
+			return TEXTVIEW("application/octet-stream");
+		case EHttpContentType::Text:
+			return TEXTVIEW("text/plain");
+		case EHttpContentType::JSON:
+			return TEXTVIEW("application/json");
+		case EHttpContentType::CbObject:
+			return TEXTVIEW("application/x-ue-cb");
+		case EHttpContentType::CbPackage:
+			return TEXTVIEW("application/x-ue-cbpkg");
+		case EHttpContentType::YAML:
+			return TEXTVIEW("text/yaml");
+		case EHttpContentType::CbPackageOffer:
+			return TEXTVIEW("application/x-ue-offer");
+		case EHttpContentType::CompressedBinary:
+			return TEXTVIEW("application/x-ue-comp");
+		case EHttpContentType::FormUrlEncoded:
+			return TEXTVIEW("application/x-www-form-urlencoded");
+		default:
+			return TEXTVIEW("unknown");
+	}
+}
+
 /**
  * Encapsulation for access token shared by all requests.
  */
@@ -86,26 +135,11 @@ private:
 class FHttpRequest
 {
 public:
-	/**
-	 * Supported request verbs
-	 */
-	enum RequestVerb
-	{
-		Get,
-		Put,
-		PutCompactBinary,
-		PutCompressedBlob,
-		Post,
-		PostCompactBinary,
-		PostJson,
-		Delete,
-		Head
-	};
 
 	/**
 	 * Convenience result type interpreted from HTTP response code.
 	 */
-	enum Result
+	enum class EResult
 	{
 		Success,
 		Failed,
@@ -118,11 +152,23 @@ public:
 		Retry
 	};
 
+	/**
+	 * Supported request verbs
+	 */
+	enum class ERequestVerb
+	{
+		Get,
+		Put,
+		Post,
+		Delete,
+		Head
+	};
+
 	using FResultCode = int;
 
-	using FOnHttpRequestComplete = TUniqueFunction<ECompletionBehavior(Result HttpResult, FHttpRequest* Request)>;
+	using FOnHttpRequestComplete = TUniqueFunction<ECompletionBehavior(EResult HttpResult, FHttpRequest* Request)>;
 
-	FHttpRequest(const TCHAR* InDomain, const TCHAR* InEffectiveDomain, FHttpAccessToken* InAuthorizationToken, FHttpSharedData* InSharedData, bool bInLogErrors);
+	FHttpRequest(FStringView InDomain, FStringView InEffectiveDomain, FHttpAccessToken* InAuthorizationToken, FHttpSharedData* InSharedData, bool bInLogErrors);
 	~FHttpRequest();
 
 	/**
@@ -180,11 +226,30 @@ public:
 	 * @param Buffer Data to upload
 	 * @return Result of the request
 	 */
-	template<RequestVerb V>
-	Result PerformBlockingUpload(const TCHAR* Uri, TArrayView<const uint8> Buffer, TConstArrayView<long> ExpectedErrorCodes = {});
-	
-	template<RequestVerb V>
-	void EnqueueAsyncUpload(UE::DerivedData::IRequestOwner& Owner, FHttpRequestPool* Pool, const TCHAR* Uri, FSharedBuffer Buffer, FOnHttpRequestComplete&& OnComplete, TConstArrayView<long> ExpectedErrorCodes = {});
+	EResult PerformBlockingPut(FStringView Uri,
+		FCompositeBuffer Buffer,
+		EHttpContentType ContentType = EHttpContentType::UnspecifiedContentType,
+		TConstArrayView<long> ExpectedErrorCodes = {});
+	EResult PerformBlockingPost(FStringView Uri,
+		FCompositeBuffer Buffer,
+		EHttpContentType ContentType = EHttpContentType::UnspecifiedContentType,
+		EHttpContentType AcceptType = EHttpContentType::UnspecifiedContentType,
+		TConstArrayView<long> ExpectedErrorCodes = {});
+	void EnqueueAsyncPut(UE::DerivedData::IRequestOwner& Owner,
+		FHttpRequestPool* Pool,
+		FStringView Uri,
+		FCompositeBuffer Buffer,
+		FOnHttpRequestComplete&& OnComplete,
+		EHttpContentType ContentType = EHttpContentType::UnspecifiedContentType,
+		TConstArrayView<long> ExpectedErrorCodes = {});
+	void EnqueueAsyncPost(UE::DerivedData::IRequestOwner& Owner,
+		FHttpRequestPool* Pool,
+		FStringView Uri,
+		FCompositeBuffer Buffer,
+		FOnHttpRequestComplete&& OnComplete,
+		EHttpContentType ContentType = EHttpContentType::UnspecifiedContentType,
+		EHttpContentType AcceptType = EHttpContentType::UnspecifiedContentType,
+		TConstArrayView<long> ExpectedErrorCodes = {});
 
 	/**
 	 * Download an url into a buffer using the request.
@@ -193,25 +258,43 @@ public:
 	 * be stored in an internal buffer and accessed GetResponse* methods.
 	 * @return Result of the request
 	 */
-	Result PerformBlockingDownload(const TCHAR* Uri, TArray<uint8>* Buffer, TConstArrayView<long> ExpectedErrorCodes = {400});
-
-	void EnqueueAsyncDownload(UE::DerivedData::IRequestOwner& Owner, FHttpRequestPool* Pool, const TCHAR* Uri, FOnHttpRequestComplete&& OnComplete, TConstArrayView<long> ExpectedErrorCodes = {400});
+	EResult PerformBlockingDownload(FStringView Uri,
+		TArray64<uint8>* Buffer,
+		EHttpContentType AcceptType = EHttpContentType::UnspecifiedContentType,
+		TConstArrayView<long> ExpectedErrorCodes = {400});
+	void EnqueueAsyncDownload(UE::DerivedData::IRequestOwner& Owner,
+		FHttpRequestPool* Pool,
+		FStringView Uri,
+		FOnHttpRequestComplete&& OnComplete,
+		EHttpContentType AcceptType = EHttpContentType::UnspecifiedContentType,
+		TConstArrayView<long> ExpectedErrorCodes = {400});
 
 	/**
 	 * Query an url using the request. Queries can use either "Head" or "Delete" verbs.
 	 * @param Uri Url to use.
 	 * @return Result of the request
 	 */
-	template<RequestVerb V>
-	Result PerformBlockingQuery(const TCHAR* Uri, TConstArrayView<long> ExpectedErrorCodes = {400});
-
-	template<RequestVerb V>
-	void EnqueueAsyncQuery(UE::DerivedData::IRequestOwner& Owner, FHttpRequestPool* Pool, const TCHAR* Uri, FOnHttpRequestComplete&& OnComplete, TConstArrayView<long> ExpectedErrorCodes = {400});
+	EResult PerformBlockingHead(FStringView Uri,
+		EHttpContentType AcceptType = EHttpContentType::UnspecifiedContentType,
+		TConstArrayView<long> ExpectedErrorCodes = {400});
+	EResult PerformBlockingDelete(FStringView Uri,
+		TConstArrayView<long> ExpectedErrorCodes = {400});
+	void EnqueueAsyncHead(UE::DerivedData::IRequestOwner& Owner,
+		FHttpRequestPool* Pool,
+		FStringView Uri,
+		FOnHttpRequestComplete&& OnComplete,
+		EHttpContentType AcceptType = EHttpContentType::UnspecifiedContentType,
+		TConstArrayView<long> ExpectedErrorCodes = {400});
+	void EnqueueAsyncDelete(UE::DerivedData::IRequestOwner& Owner,
+		FHttpRequestPool* Pool,
+		FStringView Uri,
+		FOnHttpRequestComplete&& OnComplete,
+		TConstArrayView<long> ExpectedErrorCodes = {400});
 
 	/**
 	 * Set a header to send with the request.
 	 */
-	void SetHeader(const TCHAR* Header, const TCHAR* Value);
+	void AddHeader(FStringView Header, FStringView Value);
 
 	/**
 	 * Attempts to find the header from the response. Returns false if header is not present.
@@ -222,7 +305,7 @@ public:
 	 * Returns the response buffer. Note that is the request is performed
 	 * with an external buffer as target buffer this string will be empty.
 	 */
-	const TArray<uint8>& GetResponseBuffer() const
+	const TArray64<uint8>& GetResponseBuffer() const
 	{
 		return ResponseBuffer;
 	}
@@ -286,19 +369,19 @@ private:
 	size_t Attempts;
 	bool bLogErrors;
 
-	FSharedBuffer ReadSharedBuffer;
-	FMemoryView ReadDataView;
-	TArray<uint8>* WriteDataBufferPtr;
-	TArray<uint8>* WriteHeaderBufferPtr;
+	FCompositeBuffer ReadCompositeBuffer;
+	TArray64<uint8>* WriteDataBufferPtr;
+	TArray64<uint8>* WriteHeaderBufferPtr;
 	
-	TArray<uint8> ResponseHeader;
-	TArray<uint8> ResponseBuffer;
+	TArray64<uint8> ResponseHeader;
+	TArray64<uint8> ResponseBuffer;
 	TArray<FString> Headers;
 	FString Domain;
 	FString EffectiveDomain;
 	FHttpAccessToken* AuthorizationToken;
 
-	curl_slist* PrepareToIssueRequest(const TCHAR* Uri, uint64 ContentLength);
+	void AddContentTypeHeader(FStringView Header, EHttpContentType Type);
+	curl_slist* PrepareToIssueRequest(FStringView Uri, ERequestVerb Verb, uint64 ContentLength);
 
 	/**
 	 * Performs the request, blocking until finished.
@@ -308,13 +391,13 @@ private:
 	 * @param ExpectedErrorCodes An array of expected return codes outside of the success range that should NOT be logged as an abnormal/exceptional outcome.
 	 * If unset the response body will be stored in the request.
 	 */
-	Result PerformBlocking(const TCHAR* Uri, RequestVerb Verb, uint64 ContentLength, TConstArrayView<long> ExpectedErrorCodes);
+	EResult PerformBlocking(FStringView Uri, ERequestVerb Verb, uint64 ContentLength, TConstArrayView<long> ExpectedErrorCodes);
 
-	void EnqueueAsync(UE::DerivedData::IRequestOwner& Owner, FHttpRequestPool* Pool, const TCHAR* Uri, RequestVerb Verb, uint64 ContentLength, FOnHttpRequestComplete&& OnComplete, TConstArrayView<long> ExpectedErrorCodes);
+	void EnqueueAsync(UE::DerivedData::IRequestOwner& Owner, FHttpRequestPool* Pool, FStringView Uri, ERequestVerb Verb, uint64 ContentLength, FOnHttpRequestComplete&& OnComplete, TConstArrayView<long> ExpectedErrorCodes);
 
-	void LogResult(FResultCode Result, const TCHAR* Uri, RequestVerb Verb, TConstArrayView<long> ExpectedErrorCodes) const;
+	void LogResult(FResultCode EResult, FStringView Uri, ERequestVerb Verb, TConstArrayView<long> ExpectedErrorCodes) const;
 
-	FString GetAnsiBufferAsString(const TArray<uint8>& Buffer) const
+	FString GetAnsiBufferAsString(const TArray64<uint8>& Buffer) const
 	{
 		// Content is NOT null-terminated; we need to specify lengths here
 		FUTF8ToTCHAR TCHARData(reinterpret_cast<const ANSICHAR*>(Buffer.GetData()), Buffer.Num());
@@ -323,35 +406,6 @@ private:
 
 };
 
-#define UE_HTTP_FOREACH_UPLOAD_VERB(op) \
-	op(FHttpRequest::Put); \
-	op(FHttpRequest::PutCompactBinary); \
-	op(FHttpRequest::PutCompressedBlob); \
-	op(FHttpRequest::Post); \
-	op(FHttpRequest::PostCompactBinary); \
-	op(FHttpRequest::PostJson);
-
-#define UE_HTTP_FOREACH_QUERY_VERB(op) \
-	op(FHttpRequest::Head) \
-	op(FHttpRequest::Delete)
-
-#define UE_HTTP_DECLARE_EXTERN_TEMPLATE(Verb) extern template FHttpRequest::Result FHttpRequest::PerformBlockingUpload<Verb>(const TCHAR* Uri, TArrayView<const uint8> Buffer, TConstArrayView<long> ExpectedErrorCodes);
-UE_HTTP_FOREACH_UPLOAD_VERB(UE_HTTP_DECLARE_EXTERN_TEMPLATE)
-#undef UE_HTTP_DECLARE_EXTERN_TEMPLATE
-
-#define UE_HTTP_DECLARE_EXTERN_TEMPLATE(Verb) extern template void FHttpRequest::EnqueueAsyncUpload<Verb>(UE::DerivedData::IRequestOwner& Owner, FHttpRequestPool* Pool, const TCHAR* Uri, FSharedBuffer Buffer, FOnHttpRequestComplete&& OnComplete, TConstArrayView<long> ExpectedErrorCodes);
-UE_HTTP_FOREACH_UPLOAD_VERB(UE_HTTP_DECLARE_EXTERN_TEMPLATE)
-#undef UE_HTTP_DECLARE_EXTERN_TEMPLATE
-
-
-#define UE_HTTP_DECLARE_EXTERN_TEMPLATE(Verb) extern template FHttpRequest::Result FHttpRequest::PerformBlockingQuery<Verb>(const TCHAR* Uri, TConstArrayView<long> ExpectedErrorCodes);
-UE_HTTP_FOREACH_QUERY_VERB(UE_HTTP_DECLARE_EXTERN_TEMPLATE)
-#undef UE_HTTP_DECLARE_EXTERN_TEMPLATE
-
-#define UE_HTTP_DECLARE_EXTERN_TEMPLATE(Verb) extern template void FHttpRequest::EnqueueAsyncQuery<Verb>(UE::DerivedData::IRequestOwner& Owner, FHttpRequestPool* Pool, const TCHAR* Uri, FOnHttpRequestComplete&& OnComplete, TConstArrayView<long> ExpectedErrorCodes);
-UE_HTTP_FOREACH_QUERY_VERB(UE_HTTP_DECLARE_EXTERN_TEMPLATE)
-#undef UE_HTTP_DECLARE_EXTERN_TEMPLATE
-
 
 /**
  * Pool that manages a fixed set of requests. Users are required to release requests that have been 
@@ -360,7 +414,7 @@ UE_HTTP_FOREACH_QUERY_VERB(UE_HTTP_DECLARE_EXTERN_TEMPLATE)
 class FHttpRequestPool
 {
 public:
-	FHttpRequestPool(const TCHAR* InServiceUrl, const TCHAR* InEffectiveServiceUrl, FHttpAccessToken* InAuthorizationToken, FHttpSharedData* InSharedData, uint32 PoolSize, uint32 InOverflowLimit = 0);
+	FHttpRequestPool(FStringView InServiceUrl, FStringView InEffectiveServiceUrl, FHttpAccessToken* InAuthorizationToken, FHttpSharedData* InSharedData, uint32 PoolSize, uint32 InOverflowLimit = 0);
 	~FHttpRequestPool();
 
 	/**
@@ -442,7 +496,7 @@ private:
 		FHttpAccessToken* AccessToken;
 		FHttpSharedData* SharedData;
 
-		FInitData(const TCHAR* InServiceUrl, const TCHAR* InEffectiveServiceUrl, FHttpAccessToken* InAccessToken, FHttpSharedData* InSharedData)
+		FInitData(FStringView InServiceUrl, FStringView InEffectiveServiceUrl, FHttpAccessToken* InAccessToken, FHttpSharedData* InSharedData)
 		: ServiceUrl(InServiceUrl)
 		, EffectiveServiceUrl(InEffectiveServiceUrl)
 		, AccessToken(InAccessToken)
@@ -480,6 +534,8 @@ public:
 		return Request != nullptr;
 	}
 
+	inline explicit operator bool() const { return IsValid(); }
+
 	FHttpRequest* Get() const
 	{
 		check(IsValid());
diff --git a/Engine/Source/Developer/DerivedDataCache/Private/Http/HttpRequest.cpp b/Engine/Source/Developer/DerivedDataCache/Private/Http/HttpRequest.cpp
--- a/Engine/Source/Developer/DerivedDataCache/Private/Http/HttpRequest.cpp
+++ b/Engine/Source/Developer/DerivedDataCache/Private/Http/HttpRequest.cpp
@@ -6,6 +6,7 @@
 #include "DerivedDataLegacyCacheStore.h"
 #include "HAL/IConsoleManager.h"
 #include "Logging/LogMacros.h"
+#include "Serialization/CompactBinary.h"
 #include "Serialization/JsonReader.h"
 #include "Serialization/JsonSerializer.h"
 
@@ -53,7 +54,7 @@ struct FAsyncRequestData final : public DerivedData::FRequestBase
 	UE::FHttpRequestPool* Pool = nullptr;
 	curl_slist* CurlHeaders = nullptr;
 	FString Uri;
-	UE::FHttpRequest::RequestVerb Verb;
+	UE::FHttpRequest::ERequestVerb Verb;
 	TArray<long, TInlineAllocator<4>> ExpectedErrorCodes;
 	UE::FHttpRequest::FOnHttpRequestComplete OnComplete;
 	UE::FLazyEvent Event {EEventMode::ManualReset};
@@ -176,18 +177,18 @@ struct FHttpRequestStatics
 	static size_t StaticReadFn(void* Ptr, size_t SizeInBlocks, size_t BlockSizeInBytes, void* UserData)
 	{
 		FHttpRequest* Request = static_cast<FHttpRequest*>(UserData);
-		const size_t MaxReadSize = SizeInBlocks * BlockSizeInBytes;
-		const FMemoryView SourceView = Request->ReadDataView.Mid(Request->BytesSent, MaxReadSize);
-		MakeMemoryView(Ptr, MaxReadSize).CopyFrom(SourceView);
-		Request->BytesSent += SourceView.GetSize();
-		return SourceView.GetSize();
+		const size_t Offset = Request->BytesSent;
+		const size_t ReadSize = FMath::Min((size_t)Request->ReadCompositeBuffer.GetSize() - Offset, SizeInBlocks * BlockSizeInBytes);
+		Request->ReadCompositeBuffer.CopyTo(MakeMemoryView(Ptr, ReadSize), Request->BytesSent);
+		Request->BytesSent += ReadSize;
+		return ReadSize;
 	}
 
 	static size_t StaticWriteHeaderFn(void* Ptr, size_t SizeInBlocks, size_t BlockSizeInBytes, void* UserData)
 	{
 		FHttpRequest* Request = static_cast<FHttpRequest*>(UserData);
 		const size_t WriteSize = SizeInBlocks * BlockSizeInBytes;
-		TArray<uint8>* WriteHeaderBufferPtr = Request->WriteHeaderBufferPtr;
+		TArray64<uint8>* WriteHeaderBufferPtr = Request->WriteHeaderBufferPtr;
 		if (WriteHeaderBufferPtr && WriteSize > 0)
 		{
 			const size_t CurrentBufferLength = WriteHeaderBufferPtr->Num();
@@ -209,7 +210,7 @@ struct FHttpRequestStatics
 	{
 		FHttpRequest* Request = static_cast<FHttpRequest*>(UserData);
 		const size_t WriteSize = SizeInBlocks * BlockSizeInBytes;
-		TArray<uint8>* WriteDataBufferPtr = Request->WriteDataBufferPtr;
+		TArray64<uint8>* WriteDataBufferPtr = Request->WriteDataBufferPtr;
 
 		if (WriteDataBufferPtr && WriteSize > 0)
 		{
@@ -248,11 +249,11 @@ struct FHttpRequestStatics
 		{
 		case SEEK_SET: NewPosition = Offset; break;
 		case SEEK_CUR: NewPosition = Request->BytesSent + Offset; break;
-		case SEEK_END: NewPosition = Request->ReadDataView.GetSize() + Offset; break;
+		case SEEK_END: NewPosition = Request->ReadCompositeBuffer.GetSize() + Offset; break;
 		}
 
 		// Make sure we don't seek outside of the buffer
-		if (NewPosition < 0 || NewPosition >= Request->ReadDataView.GetSize())
+		if (NewPosition < 0 || NewPosition >= Request->ReadCompositeBuffer.GetSize())
 		{
 			return CURL_SEEKFUNC_FAIL;
 		}
@@ -265,7 +266,7 @@ struct FHttpRequestStatics
 
 }
 
-FHttpRequest::FHttpRequest(const TCHAR* InDomain, const TCHAR* InEffectiveDomain, FHttpAccessToken* InAuthorizationToken, FHttpSharedData* InSharedData, bool bInLogErrors)
+FHttpRequest::FHttpRequest(FStringView InDomain, FStringView InEffectiveDomain, FHttpAccessToken* InAuthorizationToken, FHttpSharedData* InSharedData, bool bInLogErrors)
 	: SharedData(InSharedData)
 	, AsyncData(nullptr)
 	, bLogErrors(bInLogErrors)
@@ -289,7 +290,7 @@ void FHttpRequest::Reset()
 	ResponseHeader.Reset();
 	ResponseBuffer.Reset();
 	ResponseCode = 0;
-	ReadDataView = FMemoryView();
+	ReadCompositeBuffer.Reset();
 	WriteDataBufferPtr = nullptr;
 	WriteHeaderBufferPtr = nullptr;
 	BytesSent = 0;
@@ -326,8 +327,6 @@ void FHttpRequest::Reset()
 	curl_easy_setopt(LocalCurl, CURLOPT_SSL_CTX_FUNCTION, Http::Private::FHttpRequestStatics::StaticSSLCTXFn);
 	curl_easy_setopt(LocalCurl, CURLOPT_SSL_CTX_DATA, this);
 	#endif //#if WITH_SSL
-	// Allow compressed data
-	curl_easy_setopt(LocalCurl, CURLOPT_ACCEPT_ENCODING, "gzip");
 	// Rewind method, handle special error case where request need to rewind data stream
 	curl_easy_setopt(LocalCurl, CURLOPT_SEEKFUNCTION, Http::Private::FHttpRequestStatics::StaticSeekFn);
 	curl_easy_setopt(LocalCurl, CURLOPT_SEEKDATA, this);
@@ -356,25 +355,37 @@ void FHttpRequest::PrepareToRetry()
 	++Attempts;
 }
 
-FHttpRequest::Result FHttpRequest::PerformBlockingDownload(const TCHAR* Uri, TArray<uint8>* Buffer, TConstArrayView<long> ExpectedErrorCodes)
+FHttpRequest::EResult FHttpRequest::PerformBlockingDownload(FStringView Uri,
+	TArray64<uint8>* Buffer,
+	EHttpContentType AcceptType,
+	TConstArrayView<long> ExpectedErrorCodes)
 {
 	curl_easy_setopt(static_cast<CURL*>(Curl), CURLOPT_HTTPGET, 1L);
 	WriteDataBufferPtr = Buffer;
+		
+	AddContentTypeHeader(TEXTVIEW("Accept"), AcceptType);
 
-	return PerformBlocking(Uri, Get, 0u, ExpectedErrorCodes);
+	return PerformBlocking(Uri, ERequestVerb::Get, 0u, ExpectedErrorCodes);
 }
 
-void FHttpRequest::EnqueueAsyncDownload(UE::DerivedData::IRequestOwner& Owner, FHttpRequestPool* Pool, const TCHAR* Uri, FOnHttpRequestComplete&& OnComplete, TConstArrayView<long> ExpectedErrorCodes)
+void FHttpRequest::EnqueueAsyncDownload(UE::DerivedData::IRequestOwner& Owner,
+	FHttpRequestPool* Pool,
+	FStringView Uri,
+	FOnHttpRequestComplete&& OnComplete,
+	EHttpContentType AcceptType,
+	TConstArrayView<long> ExpectedErrorCodes)
 {
 	curl_easy_setopt(static_cast<CURL*>(Curl), CURLOPT_HTTPGET, 1L);
 
-	return EnqueueAsync(Owner, Pool, Uri, Get, 0u, MoveTemp(OnComplete), ExpectedErrorCodes);
+	return EnqueueAsync(Owner, Pool, Uri, ERequestVerb::Get, 0u, MoveTemp(OnComplete), ExpectedErrorCodes);
 }
 
-void FHttpRequest::SetHeader(const TCHAR* Header, const TCHAR* Value)
+void FHttpRequest::AddHeader(FStringView Header, FStringView Value)
 {
 	check(ResultCode == CURL_LAST); // Cannot set header after request is sent
-	Headers.Add(FString::Printf(TEXT("%s: %s"), Header, Value));
+	TStringBuilder<128> Sb;
+	Sb << Header << TEXTVIEW(": ") << Value;
+	Headers.Emplace(Sb);
 }
 
 bool FHttpRequest::GetHeader(const ANSICHAR* Header, FString& OutValue) const
@@ -443,17 +454,17 @@ void FHttpRequest::CompleteAsync(FResultCode Result)
 	ECompletionBehavior Behavior;
 	{
 		UE::DerivedData::FRequestBarrier Barrier(*AsyncData->Owner, UE::DerivedData::ERequestBarrierFlags::Priority);
-		FHttpRequest::Result HttpResult;
+		FHttpRequest::EResult HttpResult;
 		switch (ResultCode)
 		{
 		case CURLE_OK:
-			HttpResult = Success;
+			HttpResult = EResult::Success;
 			break;
 		case CURLE_OPERATION_TIMEDOUT:
-			HttpResult = FailedTimeout;
+			HttpResult = EResult::FailedTimeout;
 			break;
 		default:
-			HttpResult = Failed;
+			HttpResult = EResult::Failed;
 			break;
 		}
 		Behavior = AsyncData->OnComplete(HttpResult, this);
@@ -482,16 +493,53 @@ void FHttpRequest::CompleteAsync(FResultCode Result)
 	}
 }
 
-curl_slist* FHttpRequest::PrepareToIssueRequest(const TCHAR* Uri, uint64 ContentLength)
+void FHttpRequest::AddContentTypeHeader(FStringView Header, EHttpContentType Type)
+{
+	if (Type != EHttpContentType::UnspecifiedContentType)
+	{
+		AddHeader(Header, GetHttpMimeType(Type));
+	}
+}
+
+static const char* GetSessionIdHeader()
+{
+	static FCbObjectId SessionId = FCbObjectId::NewObjectId();
+
+	static const char* HeaderString = []
+	{
+		static TAnsiStringBuilder<64> SessionIdHeader;
+		SessionIdHeader << "UE-Session: " << SessionId;
+		return SessionIdHeader.GetData(); 
+	}();
+
+	return HeaderString;
+}
+
+static std::atomic<int> GRequestId{1};
+
+curl_slist* FHttpRequest::PrepareToIssueRequest(FStringView Uri, ERequestVerb Verb, uint64 ContentLength)
 {
-	static const char* CommonHeaders[] = {
+	// Strip any leading slashes because we compose the prefix and the suffix with a separating slash below
+	while (Uri.StartsWith(TEXT('/')))
+	{
+		Uri.RightChopInline(1);
+	}
+
+	TAnsiStringBuilder<32> RequestIdHeader;
+	RequestIdHeader << "UE-Request: " << GRequestId.fetch_add(1, std::memory_order_relaxed);
+
+	const char* CommonHeaders[] =
+	{
+		GetSessionIdHeader(),
+		RequestIdHeader.GetData(),
 		"User-Agent: Unreal Engine",
+		// Strip any Expect: 100-Continue header since this just introduces latency
+		"Expect:",	
 		nullptr
 	};
 
 	// Setup request options
-	FString Url = FString::Printf(TEXT("%s/%s"), *EffectiveDomain, Uri);
-	curl_easy_setopt(static_cast<CURL*>(Curl), CURLOPT_URL, TCHAR_TO_ANSI(*Url));
+	curl_easy_setopt(static_cast<CURL*>(Curl), CURLOPT_URL, *WriteToUtf8String<512>(EffectiveDomain, '/', Uri));
 
 	// Setup response header buffer. If caller has not setup a response data buffer, use internal.
 	WriteHeaderBufferPtr = &ResponseHeader;
@@ -500,8 +548,10 @@ curl_slist* FHttpRequest::PrepareToIssueRequest(const TCHAR* Uri, uint64 Content
 		WriteDataBufferPtr = &ResponseBuffer;
 	}
 
-	// Content-Length should always be set
-	Headers.Add(FString::Printf(TEXT("Content-Length: %d"), ContentLength));
+	if ((Verb != ERequestVerb::Delete) && (Verb != ERequestVerb::Get))
+	{
+		Headers.Add(FString::Printf(TEXT("Content-Length: %d"), ContentLength));
+	}
 
 	// And auth token if it's set
 	if (AuthorizationToken)
@@ -525,12 +575,12 @@ curl_slist* FHttpRequest::PrepareToIssueRequest(const TCHAR* Uri, uint64 Content
 	return CurlHeaders;
 }
 
-FHttpRequest::Result FHttpRequest::PerformBlocking(const TCHAR* Uri, RequestVerb Verb, uint64 ContentLength, TConstArrayView<long> ExpectedErrorCodes)
+FHttpRequest::EResult FHttpRequest::PerformBlocking(FStringView Uri, ERequestVerb Verb, uint64 ContentLength, TConstArrayView<long> ExpectedErrorCodes)
 {
 	TRACE_CPUPROFILER_EVENT_SCOPE(HttpDDC_CurlPerform);
 
 	// Build headers list
-	curl_slist* CurlHeaders = PrepareToIssueRequest(Uri, ContentLength);
+	curl_slist* CurlHeaders = PrepareToIssueRequest(Uri, Verb, ContentLength);
 
 	// Shots fired!
 	ResultCode = curl_easy_perform(static_cast<CURL*>(Curl));
@@ -543,10 +593,10 @@ FHttpRequest::Result FHttpRequest::PerformBlocking(const TCHAR* Uri, RequestVerb
 	// Clean up
 	curl_slist_free_all(CurlHeaders);
 
-	return ResultCode == CURLE_OK ? Success : Failed;
+	return ResultCode == CURLE_OK ? EResult::Success : EResult::Failed;
 }
 
-void FHttpRequest::EnqueueAsync(UE::DerivedData::IRequestOwner& Owner, FHttpRequestPool* Pool, const TCHAR* Uri, RequestVerb Verb, uint64 ContentLength, FOnHttpRequestComplete&& OnComplete, TConstArrayView<long> ExpectedErrorCodes)
+void FHttpRequest::EnqueueAsync(UE::DerivedData::IRequestOwner& Owner, FHttpRequestPool* Pool, FStringView Uri, ERequestVerb Verb, uint64 ContentLength, FOnHttpRequestComplete&& OnComplete, TConstArrayView<long> ExpectedErrorCodes)
 {
 	if (!AllowAsync())
 	{
@@ -565,7 +615,7 @@ void FHttpRequest::EnqueueAsync(UE::DerivedData::IRequestOwner& Owner, FHttpRequ
 	AsyncData = new Http::Private::FAsyncRequestData;
 	AsyncData->Owner = &Owner;
 	AsyncData->Pool = Pool;
-	AsyncData->CurlHeaders = PrepareToIssueRequest(Uri, ContentLength);
+	AsyncData->CurlHeaders = PrepareToIssueRequest(Uri, Verb, ContentLength);
 	AsyncData->Uri = Uri;
 	AsyncData->Verb = Verb;
 	AsyncData->ExpectedErrorCodes = ExpectedErrorCodes;
@@ -575,7 +625,7 @@ void FHttpRequest::EnqueueAsync(UE::DerivedData::IRequestOwner& Owner, FHttpRequ
 	SharedData->AddRequest(Curl);
 }
 
-void FHttpRequest::LogResult(FResultCode Result, const TCHAR* Uri, RequestVerb Verb, TConstArrayView<long> ExpectedErrorCodes) const
+void FHttpRequest::LogResult(FResultCode Result, FStringView Uri, ERequestVerb Verb, TConstArrayView<long> ExpectedErrorCodes) const
 {
 	if (Result == CURLE_OK)
 	{
@@ -585,29 +635,25 @@ void FHttpRequest::LogResult(FResultCode Result, const TCHAR* Uri, RequestVerb V
 
 		switch (Verb)
 		{
-		case Head:
+		case ERequestVerb::Head:
 			bSuccess = (ExpectedErrorCodes.Contains(ResponseCode) || IsSuccessResponse(ResponseCode));
 			VerbStr = TEXT("querying");
 			break;
-		case Get:
+		case ERequestVerb::Get:
 			bSuccess = (ExpectedErrorCodes.Contains(ResponseCode) || IsSuccessResponse(ResponseCode));
 			VerbStr = TEXT("fetching");
 			AdditionalInfo = FString::Printf(TEXT("Received: %d bytes."), BytesReceived);
 			break;
-		case Put:
-		case PutCompactBinary:
-		case PutCompressedBlob:
+		case ERequestVerb::Put:
 			bSuccess = (ExpectedErrorCodes.Contains(ResponseCode) || IsSuccessResponse(ResponseCode));
 			VerbStr = TEXT("updating");
 			AdditionalInfo = FString::Printf(TEXT("Sent: %d bytes."), BytesSent);
 			break;
-		case Post:
-		case PostCompactBinary:
-		case PostJson:
+		case ERequestVerb::Post:
 			bSuccess = (ExpectedErrorCodes.Contains(ResponseCode) || IsSuccessResponse(ResponseCode));
 			VerbStr = TEXT("posting");
 			break;
-		case Delete:
+		case ERequestVerb::Delete:
 			bSuccess = (ExpectedErrorCodes.Contains(ResponseCode) || IsSuccessResponse(ResponseCode));
 			VerbStr = TEXT("deleting");
 			break;
@@ -618,11 +664,11 @@ void FHttpRequest::LogResult(FResultCode Result, const TCHAR* Uri, RequestVerb V
 			UE_LOG(
 				LogDerivedDataCache, 
 				Verbose, 
-				TEXT("%s: Finished %s HTTP cache entry (response %d) from %s. %s"), 
+				TEXT("%s: Finished %s HTTP cache entry (response %d) from %.*s. %s"), 
 				*GetName(),
 				VerbStr,
 				ResponseCode, 
-				Uri,
+				Uri.Len(), Uri.GetData(),
 				*AdditionalInfo
 			);
 		}
@@ -638,11 +684,11 @@ void FHttpRequest::LogResult(FResultCode Result, const TCHAR* Uri, RequestVerb V
 				UE_LOG(
 					LogDerivedDataCache,
 					Verbose,
-					TEXT("%s: Failed %s HTTP cache entry (response %d) from %s. Response: %s"),
+					TEXT("%s: Failed %s HTTP cache entry (response %d) from %.*s. Response: %s"),
 					*GetName(),
 					VerbStr,
 					ResponseCode,
-					Uri,
+					Uri.Len(), Uri.GetData(),
 					*Response
 				);
 			}
@@ -651,11 +697,11 @@ void FHttpRequest::LogResult(FResultCode Result, const TCHAR* Uri, RequestVerb V
 				UE_LOG(
 					LogDerivedDataCache,
 					Display,
-					TEXT("%s: Failed %s HTTP cache entry (response %d) from %s. Response: %s"),
+					TEXT("%s: Failed %s HTTP cache entry (response %d) from %.*s. Response: %s"),
 					*GetName(),
 					VerbStr,
 					ResponseCode,
-					Uri,
+					Uri.Len(), Uri.GetData(),
 					*Response
 				);
 			}
@@ -674,166 +720,129 @@ void FHttpRequest::LogResult(FResultCode Result, const TCHAR* Uri, RequestVerb V
 	}
 }
 
-template<FHttpRequest::RequestVerb V>
-FHttpRequest::Result FHttpRequest::PerformBlockingUpload(const TCHAR* Uri, TArrayView<const uint8> Buffer, TConstArrayView<long> ExpectedErrorCodes)
+FHttpRequest::EResult FHttpRequest::PerformBlockingPut(FStringView Uri,
+	FCompositeBuffer Buffer,
+	EHttpContentType ContentType,
+	TConstArrayView<long> ExpectedErrorCodes)
 {
-	static_assert(V == Put || V == PutCompactBinary || V == PutCompressedBlob || V == Post || V == PostCompactBinary || V == PostJson, "Upload should use either Put or Post verbs.");
-	
 	uint64 ContentLength = 0u;
 
 	CURL* LocalCurl = static_cast<CURL*>(Curl);
-	if constexpr (V == Put || V == PutCompactBinary || V == PutCompressedBlob)
-	{
-		curl_easy_setopt(LocalCurl, CURLOPT_UPLOAD, 1L);
-		curl_easy_setopt(LocalCurl, CURLOPT_INFILESIZE, Buffer.Num());
-		curl_easy_setopt(LocalCurl, CURLOPT_READDATA, this);
-		curl_easy_setopt(LocalCurl, CURLOPT_READFUNCTION, Http::Private::FHttpRequestStatics::StaticReadFn);
-		if constexpr (V == PutCompactBinary)
-		{
-			Headers.Add(FString(TEXT("Content-Type: application/x-ue-cb")));
-		}
-		else if constexpr (V == PutCompressedBlob)
-		{
-			Headers.Add(FString(TEXT("Content-Type: application/x-ue-comp")));
-		}
-		else
-		{
-			Headers.Add(FString(TEXT("Content-Type: application/octet-stream")));
-		}
-		ContentLength = Buffer.Num();
-		ReadDataView = MakeMemoryView(Buffer);
-	}
-	else if constexpr (V == Post || V == PostCompactBinary || V == PostJson)
-	{
-		curl_easy_setopt(LocalCurl, CURLOPT_POST, 1L);
-		curl_easy_setopt(LocalCurl, CURLOPT_INFILESIZE, Buffer.Num());
-		curl_easy_setopt(LocalCurl, CURLOPT_READDATA, this);
-		curl_easy_setopt(LocalCurl, CURLOPT_READFUNCTION, Http::Private::FHttpRequestStatics::StaticReadFn);
-		if constexpr (V == PostCompactBinary)
-		{
-			Headers.Add(FString(TEXT("Content-Type: application/x-ue-cb")));
-		}
-		else if constexpr (V == PostJson)
-		{
-			Headers.Add(FString(TEXT("Content-Type: application/json")));
-		}
-		else
-		{
-			Headers.Add(FString(TEXT("Content-Type: application/x-www-form-urlencoded")));
-		}
-		ContentLength = Buffer.Num();
-		ReadDataView = MakeMemoryView(Buffer);
-	}
-
-	return PerformBlocking(Uri, V, ContentLength, ExpectedErrorCodes);
+	curl_easy_setopt(LocalCurl, CURLOPT_UPLOAD, 1L);
+	curl_easy_setopt(LocalCurl, CURLOPT_INFILESIZE, Buffer.GetSize());
+	curl_easy_setopt(LocalCurl, CURLOPT_READDATA, this);
+	curl_easy_setopt(LocalCurl, CURLOPT_READFUNCTION, Http::Private::FHttpRequestStatics::StaticReadFn);
+	AddContentTypeHeader(TEXTVIEW("Content-Type"), ContentType);
+	ContentLength = Buffer.GetSize();
+	ReadCompositeBuffer = Buffer;
+
+	return PerformBlocking(Uri, ERequestVerb::Put, ContentLength, ExpectedErrorCodes);
 }
 
-template<FHttpRequest::RequestVerb V>
-void FHttpRequest::EnqueueAsyncUpload(UE::DerivedData::IRequestOwner& Owner, FHttpRequestPool* Pool, const TCHAR* Uri, FSharedBuffer Buffer, FOnHttpRequestComplete&& OnComplete, TConstArrayView<long> ExpectedErrorCodes)
+FHttpRequest::EResult FHttpRequest::PerformBlockingPost(FStringView Uri,
+	FCompositeBuffer Buffer,
+	EHttpContentType ContentType,
+	EHttpContentType AcceptType,
+	TConstArrayView<long> ExpectedErrorCodes)
 {
-	static_assert(V == Put || V == PutCompactBinary || V == PutCompressedBlob || V == Post || V == PostCompactBinary || V == PostJson, "Upload should use either Put or Post verbs.");
-	
 	uint64 ContentLength = 0u;
 
 	CURL* LocalCurl = static_cast<CURL*>(Curl);
-	if constexpr (V == Put || V == PutCompactBinary || V == PutCompressedBlob)
-	{
-		curl_easy_setopt(LocalCurl, CURLOPT_UPLOAD, 1L);
-		curl_easy_setopt(LocalCurl, CURLOPT_INFILESIZE, Buffer.GetSize());
-		curl_easy_setopt(LocalCurl, CURLOPT_READDATA, this);
-		curl_easy_setopt(LocalCurl, CURLOPT_READFUNCTION, Http::Private::FHttpRequestStatics::StaticReadFn);
-		if constexpr (V == PutCompactBinary)
-		{
-			Headers.Add(FString(TEXT("Content-Type: application/x-ue-cb")));
-		}
-		else if constexpr (V == PutCompressedBlob)
-		{
-			Headers.Add(FString(TEXT("Content-Type: application/x-ue-comp")));
-		}
-		else
-		{
-			Headers.Add(FString(TEXT("Content-Type: application/octet-stream")));
-		}
-		ReadSharedBuffer = Buffer;
-		ContentLength = Buffer.GetSize();
-		ReadDataView = Buffer.GetView();
-	}
-	else if constexpr (V == Post || V == PostCompactBinary || V == PostJson)
-	{
-		curl_easy_setopt(LocalCurl, CURLOPT_POST, 1L);
-		curl_easy_setopt(LocalCurl, CURLOPT_INFILESIZE, Buffer.GetSize());
-		curl_easy_setopt(LocalCurl, CURLOPT_READDATA, this);
-		curl_easy_setopt(LocalCurl, CURLOPT_READFUNCTION, Http::Private::FHttpRequestStatics::StaticReadFn);
-		if constexpr (V == PostCompactBinary)
-		{
-			Headers.Add(FString(TEXT("Content-Type: application/x-ue-cb")));
-		}
-		else if constexpr (V == PostJson)
-		{
-			Headers.Add(FString(TEXT("Content-Type: application/json")));
-		}
-		else
-		{
-			Headers.Add(FString(TEXT("Content-Type: application/x-www-form-urlencoded")));
-		}
-		ReadSharedBuffer = Buffer;
-		ContentLength = Buffer.GetSize();
-		ReadDataView = Buffer.GetView();
-	}
-
-	return EnqueueAsync(Owner, Pool, Uri, V, ContentLength, MoveTemp(OnComplete), ExpectedErrorCodes);
+	curl_easy_setopt(LocalCurl, CURLOPT_POST, 1L);
+	curl_easy_setopt(LocalCurl, CURLOPT_POSTFIELDSIZE, Buffer.GetSize());
+	curl_easy_setopt(LocalCurl, CURLOPT_READDATA, this);
+	curl_easy_setopt(LocalCurl, CURLOPT_READFUNCTION, Http::Private::FHttpRequestStatics::StaticReadFn);
+	AddContentTypeHeader(TEXTVIEW("Content-Type"), ContentType);
+	AddContentTypeHeader(TEXTVIEW("Accept"), AcceptType);
+	ContentLength = Buffer.GetSize();
+	ReadCompositeBuffer = Buffer;
+
+	return PerformBlocking(Uri, ERequestVerb::Post, ContentLength, ExpectedErrorCodes);
 }
 
-template<FHttpRequest::RequestVerb V>
-FHttpRequest::Result FHttpRequest::PerformBlockingQuery(const TCHAR* Uri, TConstArrayView<long> ExpectedErrorCodes)
+void FHttpRequest::EnqueueAsyncPut(UE::DerivedData::IRequestOwner& Owner,
+	FHttpRequestPool* Pool,
+	FStringView Uri,
+	FCompositeBuffer Buffer,
+	FOnHttpRequestComplete&& OnComplete,
+	EHttpContentType ContentType,
+	TConstArrayView<long> ExpectedErrorCodes)
 {
-	static_assert(V == Head || V == Delete, "Queries should use either Head or Delete verbs.");
-
-	if (V == Delete)
-	{
-		curl_easy_setopt(static_cast<CURL*>(Curl), CURLOPT_CUSTOMREQUEST, "DELETE");
-	}
-	else if (V == Head)
-	{
-		curl_easy_setopt(static_cast<CURL*>(Curl), CURLOPT_NOBODY, 1L);
-	}
+	uint64 ContentLength = 0u;
 
-	return PerformBlocking(Uri, V, 0u, ExpectedErrorCodes);
+	CURL* LocalCurl = static_cast<CURL*>(Curl);
+	curl_easy_setopt(LocalCurl, CURLOPT_UPLOAD, 1L);
+	curl_easy_setopt(LocalCurl, CURLOPT_INFILESIZE, Buffer.GetSize());
+	curl_easy_setopt(LocalCurl, CURLOPT_READDATA, this);
+	curl_easy_setopt(LocalCurl, CURLOPT_READFUNCTION, Http::Private::FHttpRequestStatics::StaticReadFn);
+	AddContentTypeHeader(TEXTVIEW("Content-Type"), ContentType);
+	ReadCompositeBuffer = Buffer;
+	ContentLength = Buffer.GetSize();
+
+	return EnqueueAsync(Owner, Pool, Uri, ERequestVerb::Put, ContentLength, MoveTemp(OnComplete), ExpectedErrorCodes);
 }
 
-template<FHttpRequest::RequestVerb V>
-void FHttpRequest::EnqueueAsyncQuery(UE::DerivedData::IRequestOwner& Owner, FHttpRequestPool* Pool, const TCHAR* Uri, FOnHttpRequestComplete&& OnComplete, TConstArrayView<long> ExpectedErrorCodes)
+void FHttpRequest::EnqueueAsyncPost(UE::DerivedData::IRequestOwner& Owner,
+	FHttpRequestPool* Pool,
+	FStringView Uri,
+	FCompositeBuffer Buffer,
+	FOnHttpRequestComplete&& OnComplete,
+	EHttpContentType ContentType,
+	EHttpContentType AcceptType,
+	TConstArrayView<long> ExpectedErrorCodes)
 {
-	static_assert(V == Head || V == Delete, "Queries should use either Head or Delete verbs.");
-
-	if (V == Delete)
-	{
-		curl_easy_setopt(static_cast<CURL*>(Curl), CURLOPT_CUSTOMREQUEST, "DELETE");
-	}
-	else if (V == Head)
-	{
-		curl_easy_setopt(static_cast<CURL*>(Curl), CURLOPT_NOBODY, 1L);
-	}
+	uint64 ContentLength = 0u;
 
-	return EnqueueAsync(Owner, Pool, Uri, V, 0u, MoveTemp(OnComplete), ExpectedErrorCodes);
+	CURL* LocalCurl = static_cast<CURL*>(Curl);
+	curl_easy_setopt(LocalCurl, CURLOPT_POST, 1L);
+	curl_easy_setopt(LocalCurl, CURLOPT_INFILESIZE, Buffer.GetSize());
+	curl_easy_setopt(LocalCurl, CURLOPT_READDATA, this);
+	curl_easy_setopt(LocalCurl, CURLOPT_READFUNCTION, Http::Private::FHttpRequestStatics::StaticReadFn);
+	AddContentTypeHeader(TEXTVIEW("Content-Type"), ContentType);
+	AddContentTypeHeader(TEXTVIEW("Accept"), AcceptType);
+	ReadCompositeBuffer = Buffer;
+	ContentLength = Buffer.GetSize();
+
+	return EnqueueAsync(Owner, Pool, Uri, ERequestVerb::Post, ContentLength, MoveTemp(OnComplete), ExpectedErrorCodes);
 }
 
-#define UE_HTTP_DEFINE_EXTERN_TEMPLATE(Verb) template FHttpRequest::Result FHttpRequest::PerformBlockingUpload<Verb>(const TCHAR* Uri, TArrayView<const uint8> Buffer, TConstArrayView<long> ExpectedErrorCodes);
-UE_HTTP_FOREACH_UPLOAD_VERB(UE_HTTP_DEFINE_EXTERN_TEMPLATE)
-#undef UE_HTTP_DEFINE_EXTERN_TEMPLATE
-
-#define UE_HTTP_DEFINE_EXTERN_TEMPLATE(Verb) template void FHttpRequest::EnqueueAsyncUpload<Verb>(UE::DerivedData::IRequestOwner& Owner, FHttpRequestPool* Pool, const TCHAR* Uri, FSharedBuffer Buffer, FOnHttpRequestComplete&& OnComplete, TConstArrayView<long> ExpectedErrorCodes);
-UE_HTTP_FOREACH_UPLOAD_VERB(UE_HTTP_DEFINE_EXTERN_TEMPLATE)
-#undef UE_HTTP_DEFINE_EXTERN_TEMPLATE
+FHttpRequest::EResult FHttpRequest::PerformBlockingHead(FStringView Uri,
+	EHttpContentType AcceptType,
+	TConstArrayView<long> ExpectedErrorCodes)
+{
+	curl_easy_setopt(static_cast<CURL*>(Curl), CURLOPT_NOBODY, 1L);
+	AddContentTypeHeader(TEXTVIEW("Accept"), AcceptType);
+	return PerformBlocking(Uri, ERequestVerb::Head, 0u, ExpectedErrorCodes);
+}
 
+FHttpRequest::EResult FHttpRequest::PerformBlockingDelete(FStringView Uri,
+															TConstArrayView<long> ExpectedErrorCodes)
+{
+	curl_easy_setopt(static_cast<CURL*>(Curl), CURLOPT_CUSTOMREQUEST, "DELETE");
+	return PerformBlocking(Uri, ERequestVerb::Delete, 0u, ExpectedErrorCodes);
+}
 
-#define UE_HTTP_DEFINE_EXTERN_TEMPLATE(Verb) template FHttpRequest::Result FHttpRequest::PerformBlockingQuery<Verb>(const TCHAR* Uri, TConstArrayView<long> ExpectedErrorCodes);
-UE_HTTP_FOREACH_QUERY_VERB(UE_HTTP_DEFINE_EXTERN_TEMPLATE)
-#undef UE_HTTP_DEFINE_EXTERN_TEMPLATE
+void FHttpRequest::EnqueueAsyncHead(UE::DerivedData::IRequestOwner& Owner,
+	FHttpRequestPool* Pool,
+	FStringView Uri,
+	FOnHttpRequestComplete&& OnComplete,
+	EHttpContentType AcceptType,
+	TConstArrayView<long> ExpectedErrorCodes)
+{
+	curl_easy_setopt(static_cast<CURL*>(Curl), CURLOPT_NOBODY, 1L);
+	AddContentTypeHeader(TEXTVIEW("Accept"), AcceptType);
+	return EnqueueAsync(Owner, Pool, Uri, ERequestVerb::Head, 0u, MoveTemp(OnComplete), ExpectedErrorCodes);
+}
 
-#define UE_HTTP_DEFINE_EXTERN_TEMPLATE(Verb) template void FHttpRequest::EnqueueAsyncQuery<Verb>(UE::DerivedData::IRequestOwner& Owner, FHttpRequestPool* Pool, const TCHAR* Uri, FOnHttpRequestComplete&& OnComplete, TConstArrayView<long> ExpectedErrorCodes);
-UE_HTTP_FOREACH_QUERY_VERB(UE_HTTP_DEFINE_EXTERN_TEMPLATE)
-#undef UE_HTTP_DEFINE_EXTERN_TEMPLATE
+void FHttpRequest::EnqueueAsyncDelete(UE::DerivedData::IRequestOwner& Owner,
+	FHttpRequestPool* Pool,
+	FStringView Uri,
+	FOnHttpRequestComplete&& OnComplete,
+	TConstArrayView<long> ExpectedErrorCodes)
+{
+	curl_easy_setopt(static_cast<CURL*>(Curl), CURLOPT_CUSTOMREQUEST, "DELETE");
+	return EnqueueAsync(Owner, Pool, Uri, ERequestVerb::Delete, 0u, MoveTemp(OnComplete), ExpectedErrorCodes);
+}
 
 } // UE
 
diff --git a/Engine/Source/Developer/DerivedDataCache/Private/HttpCacheStore.cpp b/Engine/Source/Developer/DerivedDataCache/Private/HttpCacheStore.cpp
--- a/Engine/Source/Developer/DerivedDataCache/Private/HttpCacheStore.cpp
+++ b/Engine/Source/Developer/DerivedDataCache/Private/HttpCacheStore.cpp
@@ -111,9 +111,9 @@ private:
 //----------------------------------------------------------------------------------------------------------
 // Forward declarations
 //----------------------------------------------------------------------------------------------------------
-bool VerifyPayload(const FSHAHash& Hash, const TCHAR* Namespace, const TCHAR* Bucket, const TCHAR* CacheKey, const TArray<uint8>& Payload);
-bool VerifyPayload(const FIoHash& Hash, const TCHAR* Namespace, const TCHAR* Bucket, const TCHAR* CacheKey, const TArray<uint8>& Payload);
-bool VerifyRequest(const class FHttpRequest* Request, const TCHAR* Namespace, const TCHAR* Bucket, const TCHAR* CacheKey, const TArray<uint8>& Payload);
+bool VerifyPayload(const FSHAHash& Hash, const TCHAR* Namespace, const TCHAR* Bucket, const TCHAR* CacheKey, const TArray64<uint8>& Payload);
+bool VerifyPayload(const FIoHash& Hash, const TCHAR* Namespace, const TCHAR* Bucket, const TCHAR* CacheKey, const TArray64<uint8>& Payload);
+bool VerifyRequest(const class FHttpRequest* Request, const TCHAR* Namespace, const TCHAR* Bucket, const TCHAR* CacheKey, const TArray64<uint8>& Payload);
 bool HashPayload(class FHttpRequest* Request, const TArrayView<const uint8> Payload);
 bool ShouldAbortForShutdown();
 
@@ -128,7 +128,7 @@ bool ShouldAbortForShutdown();
  */
 struct FDataRequestHelper
 {
-	FDataRequestHelper(FHttpRequestPool* InPool, const TCHAR* InNamespace, const TCHAR* InBucket, const TCHAR* InCacheKey, TArray<uint8>* OutData)
+	FDataRequestHelper(FHttpRequestPool* InPool, const TCHAR* InNamespace, const TCHAR* InBucket, const TCHAR* InCacheKey, TArray64<uint8>* OutData)
 		: Request(nullptr)
 		, Pool(InPool)
 		, bVerified(false, 1)
@@ -138,7 +138,7 @@ struct FDataRequestHelper
 		{
 			// We are below the threshold, make the connection immediately. OutData is set so this is a get.
 			FString Uri = FString::Printf(TEXT("api/v1/c/ddc/%s/%s/%s.raw"), InNamespace, InBucket, InCacheKey);
-			const FHttpRequest::Result Result = Request->PerformBlockingDownload(*Uri, OutData);
+			const FHttpRequest::EResult Result = Request->PerformBlockingDownload(*Uri, OutData);
 			if (FHttpRequest::IsSuccessResponse(Request->GetResponseCode()))
 			{
 				if (VerifyRequest(Request, InNamespace, InBucket, InCacheKey, *OutData))
@@ -153,7 +153,7 @@ struct FDataRequestHelper
 		{
 			// We are below the threshold, make the connection immediately. OutData is missing so this is a head.
 			FString Uri = FString::Printf(TEXT("api/v1/c/ddc/%s/%s/%s"), InNamespace, InBucket, InCacheKey);
-			const FHttpRequest::Result Result = Request->PerformBlockingQuery<FHttpRequest::Head>(*Uri);
+			const FHttpRequest::EResult Result = Request->PerformBlockingHead(*Uri);
 			if (FHttpRequest::IsSuccessResponse(Request->GetResponseCode()))
 			{
 				TRACE_COUNTER_ADD(HttpDDC_ExistHit, int64(1));
@@ -171,7 +171,7 @@ struct FDataRequestHelper
 					InNamespace,
 					InBucket,
 					TConstArrayView<const TCHAR*>({InCacheKey}),
-					OutData ? TConstArrayView<TArray<uint8>*>({OutData}) : TConstArrayView<TArray<uint8>*>(),
+					OutData ? TConstArrayView<TArray64<uint8>*>({OutData}) : TConstArrayView<TArray64<uint8>*>(),
 					bVerified
 				);
 			}
@@ -184,8 +184,8 @@ struct FDataRequestHelper
 					InNamespace,
 					InBucket,
 					TConstArrayView<const TCHAR*>({InCacheKey}),
-					OutData ? TConstArrayView<TArray<uint8>*>({OutData}) : TConstArrayView<TArray<uint8>*>(),
-					OutData && !OutData->IsEmpty() ? FHttpRequest::RequestVerb::Get : FHttpRequest::RequestVerb::Head,
+					OutData ? TConstArrayView<TArray64<uint8>*>({OutData}) : TConstArrayView<TArray64<uint8>*>(),
+					OutData && !OutData->IsEmpty() ? FHttpRequest::ERequestVerb::Get : FHttpRequest::ERequestVerb::Head,
 					&bVerified
 				};
 
@@ -218,8 +218,8 @@ struct FDataRequestHelper
 				InNamespace, 
 				InBucket,
 				CacheKeys,
-				TConstArrayView<TArray<uint8>*>(),
-				FHttpRequest::RequestVerb::Head,
+				TConstArrayView<TArray64<uint8>*>(),
+				FHttpRequest::ERequestVerb::Head,
 				&bVerified
 			};
 
@@ -232,7 +232,7 @@ struct FDataRequestHelper
 				InNamespace, 
 				InBucket, 
 				CacheKeys, 
-				TConstArrayView<TArray<uint8>*>(), 
+				TConstArrayView<TArray64<uint8>*>(), 
 				bVerified
 			);
 
@@ -244,8 +244,8 @@ struct FDataRequestHelper
 					InNamespace,
 					InBucket,
 					CacheKeys,
-					TConstArrayView<TArray<uint8>*>(),
-					FHttpRequest::RequestVerb::Head,
+					TConstArrayView<TArray64<uint8>*>(),
+					FHttpRequest::ERequestVerb::Head,
 					&bVerified
 				};
 
@@ -305,8 +305,8 @@ private:
 		const TCHAR* Namespace;
 		const TCHAR* Bucket;
 		TConstArrayView<const TCHAR*> CacheKeys;
-		TConstArrayView<TArray<uint8>*> OutDatas;
-		FHttpRequest::RequestVerb Verb;
+		TConstArrayView<TArray64<uint8>*> OutDatas;
+		FHttpRequest::ERequestVerb Verb;
 		TBitArray<>* bSuccess;
 	};
 
@@ -360,7 +360,7 @@ private:
 		const TCHAR* InNamespace, 
 		const TCHAR* InBucket, 
 		TConstArrayView<const TCHAR*> InCacheKeys,
-		TConstArrayView<TArray<uint8>*> OutDatas, 
+		TConstArrayView<TArray64<uint8>*> OutDatas, 
 		TBitArray<>& bOutVerified)
 	{
 		TRACE_CPUPROFILER_EVENT_SCOPE(HttpDDC_BatchQuery);
@@ -404,7 +404,7 @@ private:
 				InBucket,
 				InCacheKeys,
 				OutDatas,
-				OutDatas.Num() ? FHttpRequest::RequestVerb::Get : FHttpRequest::RequestVerb::Head,
+				OutDatas.Num() ? FHttpRequest::ERequestVerb::Get : FHttpRequest::ERequestVerb::Head,
 				&bOutVerified
 			};
 
@@ -485,7 +485,7 @@ private:
 				TSharedPtr<FJsonObject> Object = MakeShared<FJsonObject>();
 				Object->SetField(TEXT("bucket"), MakeShared<FJsonValueString>(Entry.Bucket));
 				Object->SetField(TEXT("key"), MakeShared<FJsonValueString>(Entry.CacheKeys[KeyIdx]));
-				if (Entry.Verb == FHttpRequest::RequestVerb::Head)
+				if (Entry.Verb == FHttpRequest::ERequestVerb::Head)
 				{
 					Object->SetField(TEXT("verb"), MakeShared<FJsonValueString>(TEXT("HEAD")));
 				}
@@ -500,12 +500,12 @@ private:
 		FBufferArchive RequestData;
 		if (FJsonSerializer::Serialize(RequestObject.ToSharedRef(), TJsonWriterFactory<ANSICHAR, TCondensedJsonPrintPolicy<ANSICHAR>>::Create(&RequestData)))
 		{
-			Request->PerformBlockingUpload<FHttpRequest::PostJson>(Uri, MakeArrayView(RequestData));
+			Request->PerformBlockingPost(Uri, FCompositeBuffer(FSharedBuffer::MakeView(RequestData.GetData(), RequestData.Num())), EHttpContentType::JSON);
 			ResponseCode = Request->GetResponseCode();
 
 			if (ResponseCode == 200)
 			{
-				const TArray<uint8>& ResponseBuffer = Request->GetResponseBuffer();
+				const TArray64<uint8>& ResponseBuffer = Request->GetResponseBuffer();
 				const uint8* Response = ResponseBuffer.GetData();
 				const int32 ResponseSize = ResponseBuffer.Num();
 
@@ -562,8 +562,8 @@ private:
 				// or a head request with a get response code (i.e. Ok)
 				// if the response code is an error or not found they can be matched to both head or get request it doesn't matter
 				const FQueuedBatchEntry& CurrentRequest = Requests[CurrentEntryIdx];
-				bool bRequestTypeMatch = !((CurrentRequest.Verb == FHttpRequest::Get) && (RequestResult == OpResult::Exists))
-					&& !((CurrentRequest.Verb == FHttpRequest::Head) && (RequestResult == OpResult::Ok));
+				bool bRequestTypeMatch = !((CurrentRequest.Verb == FHttpRequest::ERequestVerb::Get) && (RequestResult == OpResult::Exists))
+					&& !((CurrentRequest.Verb == FHttpRequest::ERequestVerb::Head) && (RequestResult == OpResult::Ok));
 				if (bRequestTypeMatch && FCString::Stricmp(CurrentRequest.CacheKeys[CurrentKeyIdx], CacheKey.Get()) == 0)
 				{
 					EntryIdx = CurrentEntryIdx;
@@ -684,7 +684,7 @@ private:
 							}
 							else
 							{
-								TArray<uint8>* OutData = RequestOp.OutDatas[KeyIdx];
+								TArray64<uint8>* OutData = RequestOp.OutDatas[KeyIdx];
 
 								OutData->Append(Response, PayloadSize);
 								Response += PayloadSize;
@@ -815,7 +815,7 @@ private:
 		FString Namespace;
 		FString Bucket;
 		FString CacheKey;
-		TArray<uint8> Data;
+		TArray64<uint8> Data;
 
 		FQueuedEntry(const TCHAR* InNamespace, const TCHAR* InBucket, const TCHAR* InCacheKey, const TArrayView<const uint8> InData)
 			: Namespace(InNamespace)
@@ -882,7 +882,7 @@ private:
 		TStringBuilder<256> Uri;
 		Uri.Appendf(TEXT("api/v1/c/ddc/%s/%s/%s"), Namespace, Bucket, CacheKey);
 
-		Request->PerformBlockingUpload<FHttpRequest::Put>(*Uri, Data);
+		Request->PerformBlockingPut(*Uri, FCompositeBuffer(FSharedBuffer::MakeView(Data.GetData(), Data.Num())), EHttpContentType::Binary);
 
 		const int64 ResponseCode = Request->GetResponseCode();
 		if (FHttpRequest::IsSuccessResponse(ResponseCode))
@@ -912,7 +912,7 @@ TLockFreePointerListUnordered<FDataUploadHelper::FQueuedEntry, PLATFORM_CACHE_LI
  * @param Payload Payload received.
  * @return True if the data is correct, false if checksums doesn't match.
  */
-bool VerifyPayload(const FSHAHash& Hash, const TCHAR* Namespace, const TCHAR* Bucket, const TCHAR* CacheKey, const TArray<uint8>& Payload)
+bool VerifyPayload(const FSHAHash& Hash, const TCHAR* Namespace, const TCHAR* Bucket, const TCHAR* CacheKey, const TArray64<uint8>& Payload)
 {
 	FSHAHash PayloadHash;
 	FSHA1::HashBuffer(Payload.GetData(), Payload.Num(), PayloadHash.Hash);
@@ -943,7 +943,7 @@ bool VerifyPayload(const FSHAHash& Hash, const TCHAR* Namespace, const TCHAR* Bu
  * @param Payload Payload received.
  * @return True if the data is correct, false if checksums doesn't match.
  */
-bool VerifyPayload(const FIoHash& Hash, const TCHAR* Namespace, const TCHAR* Bucket, const TCHAR* CacheKey, const TArray<uint8>& Payload)
+bool VerifyPayload(const FIoHash& Hash, const TCHAR* Namespace, const TCHAR* Bucket, const TCHAR* CacheKey, const TArray64<uint8>& Payload)
 {
 	FIoHash PayloadHash = FIoHash::HashBuffer(Payload.GetData(), Payload.Num());
 
@@ -974,7 +974,7 @@ bool VerifyPayload(const FIoHash& Hash, const TCHAR* Namespace, const TCHAR* Buc
  * @param Payload Payload received.
  * @return True if the data is correct, false if checksums doesn't match.
  */
-bool VerifyRequest(const FHttpRequest* Request, const TCHAR* Namespace, const TCHAR* Bucket, const TCHAR* CacheKey, const TArray<uint8>& Payload)
+bool VerifyRequest(const FHttpRequest* Request, const TCHAR* Namespace, const TCHAR* Bucket, const TCHAR* CacheKey, const TArray64<uint8>& Payload)
 {
 	FString ReceivedHashStr;
 	if (Request->GetHeader("X-Jupiter-Sha1", ReceivedHashStr))
@@ -1002,7 +1002,7 @@ bool VerifyRequest(const FHttpRequest* Request, const TCHAR* Namespace, const TC
 bool HashPayload(FHttpRequest* Request, const TArrayView<const uint8> Payload)
 {
 	FIoHash PayloadHash = FIoHash::HashBuffer(Payload.GetData(), Payload.Num());
-	Request->SetHeader(TEXT("X-Jupiter-IoHash"), *WriteToString<48>(PayloadHash));
+	Request->AddHeader(TEXTVIEW("X-Jupiter-IoHash"), WriteToString<48>(PayloadHash));
 	return true;
 }
 
@@ -1162,8 +1162,8 @@ private:
 
 	bool IsServiceReady();
 	bool AcquireAccessToken();
-	bool ShouldRetryOnError(FHttpRequest::Result Result, int64 ResponseCode);
-	bool ShouldRetryOnError(int64 ResponseCode) { return ShouldRetryOnError(FHttpRequest::Success, ResponseCode); }
+	bool ShouldRetryOnError(FHttpRequest::EResult Result, int64 ResponseCode);
+	bool ShouldRetryOnError(int64 ResponseCode) { return ShouldRetryOnError(FHttpRequest::EResult::Success, ResponseCode); }
 
 	enum class OperationCategory
 	{
@@ -1340,7 +1340,7 @@ private:
 		FCachePutRefResponse&& Response);
 
 	FHttpRequest::ECompletionBehavior OnCompressedBlobUploadComplete(
-		FHttpRequest::Result HttpResult,
+		FHttpRequest::EResult HttpResult,
 		FHttpRequest* Request);
 
 	void OnPutRefFinalizationComplete(
@@ -1511,7 +1511,7 @@ void FHttpCacheStore::FPutPackageOp::PutRefAsync(
 	FHttpRequest* Request = CacheStore.WaitForHttpRequestForOwner<OperationCategory::Put>(Owner, bFinalize /* bUnboundedOverflow */, Pool);
 
 	auto OnHttpRequestComplete = [&CacheStore, &Owner, Name = FSharedString(Name), Key, Object, UserData, bFinalize, OnComplete = MoveTemp(OnComplete)]
-	(FHttpRequest::Result HttpResult, FHttpRequest* Request)
+	(FHttpRequest::EResult HttpResult, FHttpRequest* Request)
 	{
 		TRACE_CPUPROFILER_EVENT_SCOPE(HttpDDC_PutRefAsync_OnHttpRequestComplete);
 
@@ -1572,12 +1572,12 @@ void FHttpCacheStore::FPutPackageOp::PutRefAsync(
 
 	if (bFinalize)
 	{
-		Request->EnqueueAsyncUpload<FHttpRequest::Post>(Owner, Pool, *RefsUri, FSharedBuffer(), MoveTemp(OnHttpRequestComplete));
+		Request->EnqueueAsyncPost(Owner, Pool, *RefsUri, FCompositeBuffer(), MoveTemp(OnHttpRequestComplete), EHttpContentType::FormUrlEncoded);
 	}
 	else
 	{
-		Request->SetHeader(TEXT("X-Jupiter-IoHash"), *WriteToString<48>(ObjectHash));
-		Request->EnqueueAsyncUpload<FHttpRequest::PutCompactBinary>(Owner, Pool, *RefsUri, Object.GetBuffer().ToShared(), MoveTemp(OnHttpRequestComplete));
+		Request->AddHeader(TEXTVIEW("X-Jupiter-IoHash"), WriteToString<48>(ObjectHash));
+		Request->EnqueueAsyncPut(Owner, Pool, *RefsUri, Object.GetBuffer(), MoveTemp(OnHttpRequestComplete), EHttpContentType::CbObject);
 	}
 }
 
@@ -1687,16 +1687,17 @@ void FHttpCacheStore::FPutPackageOp::OnPackagePutRefComplete(
 
 		FHttpRequestPool* Pool = nullptr;
 		FHttpRequest* Request = CacheStore.WaitForHttpRequestForOwner<OperationCategory::Put>(Owner, true /* bUnboundedOverflow */, Pool);
-		Request->EnqueueAsyncUpload<FHttpRequest::PutCompressedBlob>(Owner, Pool, *CompressedBlobsUri, CompressedBlobUpload.BlobBuffer,
-			[PutPackageOp](FHttpRequest::Result HttpResult, FHttpRequest* Request)
+		Request->EnqueueAsyncPut(Owner, Pool, *CompressedBlobsUri, FCompositeBuffer(CompressedBlobUpload.BlobBuffer),
+			[PutPackageOp](FHttpRequest::EResult HttpResult, FHttpRequest* Request)
 			{
 				return PutPackageOp->OnCompressedBlobUploadComplete(HttpResult, Request);
-			});
+			},
+			EHttpContentType::CompressedBinary);
 	}
 }
 
 FHttpRequest::ECompletionBehavior FHttpCacheStore::FPutPackageOp::OnCompressedBlobUploadComplete(
-	FHttpRequest::Result HttpResult,
+	FHttpRequest::EResult HttpResult,
 	FHttpRequest* Request)
 {
 	int64 ResponseCode = Request->GetResponseCode();
@@ -1798,7 +1799,7 @@ void FHttpCacheStore::FGetRecordOp::GetDataBatch(
 		FHttpRequest* Request = CacheStore.WaitForHttpRequestForOwner<OperationCategory::Get>(Owner, true /* bUnboundedOverflow */, Pool);
 
 		auto OnHttpRequestComplete = [&CacheStore, &Owner, Name = FSharedString(Name), Key = FCacheKey(Key), ValueIndex, Value = Value.RemoveData(), ValueIdGetter, OnCompletePtr = TRefCountPtr<TRefCountedUniqueFunction<FOnGetCachedDataBatchComplete>>(CompletionFunction)]
-		(FHttpRequest::Result HttpResult, FHttpRequest* Request)
+		(FHttpRequest::EResult HttpResult, FHttpRequest* Request)
 		{
 			TRACE_CPUPROFILER_EVENT_SCOPE(HttpDDC_GetDataBatch_OnHttpRequestComplete);
 
@@ -1863,8 +1864,7 @@ void FHttpCacheStore::FGetRecordOp::GetDataBatch(
 
 		TStringBuilder<256> CompressedBlobsUri;
 		CompressedBlobsUri << "api/v1/compressed-blobs/" << CacheStore.StructuredNamespace << "/" << RawHash;
-		Request->SetHeader(TEXT("Accept"), TEXT("*/*"));
-		Request->EnqueueAsyncDownload(Owner, Pool, *CompressedBlobsUri, MoveTemp(OnHttpRequestComplete), { 404 });
+		Request->EnqueueAsyncDownload(Owner, Pool, *CompressedBlobsUri, MoveTemp(OnHttpRequestComplete), EHttpContentType::AnyContentType, { 404 });
 	}
 }
 
@@ -2005,7 +2005,7 @@ void FHttpCacheStore::FGetRecordOp::DataProbablyExistsBatch(
 		bFirstItem = false;
 	}
 
-	auto OnHttpRequestComplete = [this, Values = TArray<FValueWithId>(Values), InOnComplete = MoveTemp(InOnComplete)](FHttpRequest::Result HttpResult, FHttpRequest* Request)
+	auto OnHttpRequestComplete = [this, Values = TArray<FValueWithId>(Values), InOnComplete = MoveTemp(InOnComplete)](FHttpRequest::EResult HttpResult, FHttpRequest* Request)
 	{
 		TRACE_CPUPROFILER_EVENT_SCOPE(HttpDDC_DataProbablyExistsBatch_OnHttpRequestComplete);
 
@@ -2097,8 +2097,8 @@ void FHttpCacheStore::FGetRecordOp::DataProbablyExistsBatch(
 		}
 		return FHttpRequest::ECompletionBehavior::Done;
 	};
-	FSharedBuffer DummyBuffer;
-	Request->EnqueueAsyncUpload<FHttpRequest::Post>(Owner, Pool, *CompressedBlobsUri, DummyBuffer, MoveTemp(OnHttpRequestComplete));
+	FCompositeBuffer DummyBuffer;
+	Request->EnqueueAsyncPost(Owner, Pool, *CompressedBlobsUri, DummyBuffer, MoveTemp(OnHttpRequestComplete), EHttpContentType::FormUrlEncoded);
 }
 
 void FHttpCacheStore::FGetRecordOp::FinishDataStep(bool bSuccess, uint64 InBytesReceived)
@@ -2262,9 +2262,9 @@ bool FHttpCacheStore::ApplyDebugOptions(FBackendDebugOptions& InOptions)
 bool FHttpCacheStore::IsServiceReady()
 {
 	FHttpRequest Request(*Domain, *Domain, nullptr, SharedData.Get(), false);
-	FHttpRequest::Result Result = Request.PerformBlockingDownload(TEXT("health/ready"), nullptr);
+	FHttpRequest::EResult Result = Request.PerformBlockingDownload(TEXT("health/ready"), nullptr);
 	
-	if (Result == FHttpRequest::Success && Request.GetResponseCode() == 200)
+	if (Result == FHttpRequest::EResult::Success && Request.GetResponseCode() == 200)
 	{
 		UE_LOG(LogDerivedDataCache, Display, TEXT("%s: HTTP DDC service status: %s."), *Request.GetName(), *Request.GetResponseAsString());
 		return true;
@@ -2316,7 +2316,7 @@ bool FHttpCacheStore::AcquireAccessToken()
 		FString Uri(*OAuthProvider + DomainEnd + 1);
 
 		FHttpRequest Request(*AuthDomain, *AuthDomain, nullptr, SharedData.Get(), false);
-		FHttpRequest::Result Result = FHttpRequest::Success;
+		FHttpRequest::EResult Result = FHttpRequest::EResult::Success;
 		if (OAuthProvider.StartsWith(TEXT("http://localhost")))
 		{
 			// Simple unauthenticated call to a local endpoint that mimics
@@ -2351,14 +2351,14 @@ bool FHttpCacheStore::AcquireAccessToken()
 				*OAuthSecret
 			);
 
-			TArray<uint8> FormData;
+			TArray64<uint8> FormData;
 			auto OAuthFormDataUTF8 = FTCHARToUTF8(*OAuthFormData);
 			FormData.Append((uint8*)OAuthFormDataUTF8.Get(), OAuthFormDataUTF8.Length());
 
-			Result = Request.PerformBlockingUpload<FHttpRequest::Post>(*Uri, MakeArrayView(FormData));
+			Result = Request.PerformBlockingPost(*Uri, FCompositeBuffer(FSharedBuffer::MakeView(FormData.GetData(), FormData.Num())), EHttpContentType::FormUrlEncoded);
 		}
 
-		if (Result == FHttpRequest::Success && Request.GetResponseCode() == 200)
+		if (Result == FHttpRequest::EResult::Success && Request.GetResponseCode() == 200)
 		{
 			TSharedPtr<FJsonObject> ResponseObject = Request.GetResponseAsJsonObject();
 			if (ResponseObject)
@@ -2403,9 +2403,9 @@ bool FHttpCacheStore::AcquireAccessToken()
 	return false;
 }
 
-bool FHttpCacheStore::ShouldRetryOnError(FHttpRequest::Result Result, int64 ResponseCode)
+bool FHttpCacheStore::ShouldRetryOnError(FHttpRequest::EResult Result, int64 ResponseCode)
 {
-	if (Result == FHttpRequest::FailedTimeout)
+	if (Result == FHttpRequest::EResult::FailedTimeout)
 	{
 		return true;
 	}
@@ -2473,7 +2473,7 @@ void FHttpCacheStore::GetCacheRecordOnlyAsync(
 	FHttpRequest* Request = WaitForHttpRequestForOwner<OperationCategory::Get>(Owner, false /* bUnboundedOverflow */, Pool);
 
 	auto OnHttpRequestComplete = [this, &Owner, Name = FSharedString(Name), Key, UserData, OnComplete = MoveTemp(OnComplete)]
-	(FHttpRequest::Result HttpResult, FHttpRequest* Request)
+	(FHttpRequest::EResult HttpResult, FHttpRequest* Request)
 	{
 		TRACE_CPUPROFILER_EVENT_SCOPE(HttpDDC_GetCacheRecordOnlyAsync_OnHttpRequestComplete);
 
@@ -2513,8 +2513,7 @@ void FHttpCacheStore::GetCacheRecordOnlyAsync(
 		OnComplete({ Name, Key, UserData, Request->GetBytesReceived(), {}, EStatus::Error });
 		return FHttpRequest::ECompletionBehavior::Done;
 	};
-	Request->SetHeader(TEXT("Accept"), TEXT("application/x-ue-cb"));
-	Request->EnqueueAsyncDownload(Owner, Pool, *RefsUri, MoveTemp(OnHttpRequestComplete), { 401, 404 });
+	Request->EnqueueAsyncDownload(Owner, Pool, *RefsUri, MoveTemp(OnHttpRequestComplete), EHttpContentType::CbObject, { 401, 404 });
 }
 
 void FHttpCacheStore::PutCacheRecordAsync(
@@ -2676,14 +2675,14 @@ void FHttpCacheStore::GetCacheValueAsync(
 	FHttpRequest* Request = WaitForHttpRequestForOwner<OperationCategory::Get>(Owner, false /* bUnboundedOverflow */, Pool);
 	if (bSkipData)
 	{
-		Request->SetHeader(TEXT("Accept"), TEXT("application/x-ue-cb"));
+		Request->AddHeader(TEXT("Accept"), TEXT("application/x-ue-cb"));
 	}
 	else
 	{
-		Request->SetHeader(TEXT("Accept"), TEXT("application/x-jupiter-inline"));
+		Request->AddHeader(TEXT("Accept"), TEXT("application/x-jupiter-inline"));
 	}
 
-	auto OnHttpRequestComplete = [this, &Owner, Name, Key, UserData, bSkipData, OnComplete = MoveTemp(OnComplete)] (FHttpRequest::Result HttpResult, FHttpRequest* Request)
+	auto OnHttpRequestComplete = [this, &Owner, Name, Key, UserData, bSkipData, OnComplete = MoveTemp(OnComplete)] (FHttpRequest::EResult HttpResult, FHttpRequest* Request)
 	{
 		TRACE_CPUPROFILER_EVENT_SCOPE(HttpDDC_GetCacheValueAsync_OnHttpRequestComplete);
 
@@ -2756,7 +2755,7 @@ void FHttpCacheStore::GetCacheValueAsync(
 		return FHttpRequest::ECompletionBehavior::Done;
 	};
 
-	Request->EnqueueAsyncDownload(Owner, Pool, *RefsUri, MoveTemp(OnHttpRequestComplete), { 401, 404 });
+	Request->EnqueueAsyncDownload(Owner, Pool, *RefsUri, MoveTemp(OnHttpRequestComplete), EHttpContentType::UnspecifiedContentType, { 401, 404 });
 }
 
 void FHttpCacheStore::GetCacheRecordAsync(
@@ -2820,9 +2819,9 @@ void FHttpCacheStore::RefCachedDataProbablyExistsBatchAsync(
 	FHttpRequestPool* Pool = nullptr;
 	FHttpRequest* Request = WaitForHttpRequestForOwner<OperationCategory::Get>(Owner, false /* bUnboundedOverflow */, Pool);
 
-	Request->SetHeader(TEXT("Accept"), TEXT("application/x-ue-cb"));
+	Request->AddHeader(TEXT("Accept"), TEXT("application/x-ue-cb"));
 
-	auto OnHttpRequestComplete = [this, &Owner, ValueRefs = TArray<FCacheGetValueRequest>(ValueRefs), OnComplete = MoveTemp(OnComplete)](FHttpRequest::Result HttpResult, FHttpRequest* Request)
+	auto OnHttpRequestComplete = [this, &Owner, ValueRefs = TArray<FCacheGetValueRequest>(ValueRefs), OnComplete = MoveTemp(OnComplete)](FHttpRequest::EResult HttpResult, FHttpRequest* Request)
 	{
 		TRACE_CPUPROFILER_EVENT_SCOPE(HttpDDC_RefCachedDataProbablyExistsBatchAsync_OnHttpRequestComplete);
 
@@ -2920,7 +2919,7 @@ void FHttpCacheStore::RefCachedDataProbablyExistsBatchAsync(
 		return FHttpRequest::ECompletionBehavior::Done;
 	};
 
-	Request->EnqueueAsyncUpload<FHttpRequest::PostCompactBinary>(Owner, Pool, *RefsUri, RequestFields.GetOuterBuffer(), MoveTemp(OnHttpRequestComplete));
+	Request->EnqueueAsyncPost(Owner, Pool, *RefsUri, FCompositeBuffer(RequestFields.GetOuterBuffer()), MoveTemp(OnHttpRequestComplete), EHttpContentType::CbObject);
 }
 
 bool FHttpCacheStore::CachedDataProbablyExists(const TCHAR* CacheKey)
@@ -2959,7 +2958,7 @@ bool FHttpCacheStore::CachedDataProbablyExists(const TCHAR* CacheKey)
 	for (int32 Attempts = 0; Attempts < UE_HTTPDDC_MAX_ATTEMPTS; ++Attempts)
 	{
 		FScopedHttpPoolRequestPtr Request(GetRequestPools[IsInGameThread()].Get());
-		const FHttpRequest::Result Result = Request->PerformBlockingQuery<FHttpRequest::Head>(*Uri);
+		const FHttpRequest::EResult Result = Request->PerformBlockingQuery<FHttpRequest::Head>(*Uri);
 		const int64 ResponseCode = Request->GetResponseCode();
 
 		if (FHttpRequest::IsSuccessResponse(ResponseCode) || ResponseCode == 400)
@@ -3048,7 +3047,7 @@ TBitArray<> FHttpCacheStore::CachedDataProbablyExistsBatch(TConstArrayView<FStri
 	for (int32 Attempts = 0; Attempts < UE_HTTPDDC_MAX_ATTEMPTS; ++Attempts)
 	{
 		FScopedHttpPoolRequestPtr Request(GetRequestPools[IsInGameThread()].Get());
-		const FHttpRequest::Result Result = Request->PerformBlockingUpload<FHttpRequest::PostJson>(Uri, BodyView);
+		const FHttpRequest::EResult Result = Request->PerformBlockingUpload<FHttpRequest::PostJson>(Uri, BodyView);
 		const int64 ResponseCode = Request->GetResponseCode();
 
 		if (Result == FHttpRequest::Success && ResponseCode == 200)
@@ -3104,11 +3103,14 @@ bool FHttpCacheStore::GetCachedData(const TCHAR* CacheKey, TArray<uint8>& OutDat
 		return false;
 	}
 
+	TArray64<uint8> ArrayBuffer;
+
 #if WITH_DATAREQUEST_HELPER
 	// Retry request until we get an accepted response or exhaust allowed number of attempts.
 	for (int32 Attempts = 0; Attempts < UE_HTTPDDC_MAX_ATTEMPTS; ++Attempts)
 	{
-		FDataRequestHelper RequestHelper(GetRequestPools[IsInGameThread()].Get(), *Namespace, *DefaultBucket, CacheKey, &OutData);
+		FDataRequestHelper RequestHelper(GetRequestPools[IsInGameThread()].Get(), *Namespace, *DefaultBucket, CacheKey, &ArrayBuffer);
+		OutData = TArray<uint8>(MoveTemp(ArrayBuffer));
 		const int64 ResponseCode = RequestHelper.GetResponseCode();
 
 		if (FHttpRequest::IsSuccessResponse(ResponseCode) && RequestHelper.IsSuccess())
@@ -3132,17 +3134,19 @@ bool FHttpCacheStore::GetCachedData(const TCHAR* CacheKey, TArray<uint8>& OutDat
 		FScopedHttpPoolRequestPtr Request(GetRequestPools[IsInGameThread()].Get());
 		if (Request.IsValid())
 		{
-			FHttpRequest::Result Result = Request->PerformBlockingDownload(*Uri, &OutData);
+			FHttpRequest::EResult Result = Request->PerformBlockingDownload(*Uri, &ArrayBuffer);
 			const uint64 ResponseCode = Request->GetResponseCode();
 
 			// Request was successful, make sure we got all the expected data.
-			if (FHttpRequest::IsSuccessResponse(ResponseCode) && VerifyRequest(Request.Get(), *Namespace, *DefaultBucket, CacheKey, OutData))
+			if (FHttpRequest::IsSuccessResponse(ResponseCode) && VerifyRequest(Request.Get(), *Namespace, *DefaultBucket, CacheKey, ArrayBuffer))
 			{
+				OutData = TArray<uint8>(MoveTemp(ArrayBuffer));
 				TRACE_COUNTER_ADD(HttpDDC_GetHit, int64(1));
 				TRACE_COUNTER_ADD(HttpDDC_BytesReceived, int64(Request->GetBytesReceived()));
 				COOK_STAT(Timer.AddHit(Request->GetBytesReceived()));
 				return true;
 			}
+			OutData = TArray<uint8>(MoveTemp(ArrayBuffer));
 
 			if (!ShouldRetryOnError(ResponseCode))
 			{
@@ -3213,7 +3217,7 @@ FDerivedDataBackendInterface::EPutStatus FHttpCacheStore::PutCachedData(const TC
 			// Append the content hash to the header
 			HashPayload(Request.Get(), InData);
 
-			Request->PerformBlockingUpload<FHttpRequest::Put>(*Uri, InData);
+			Request->PerformBlockingPut(*Uri, FCompositeBuffer(FSharedBuffer::MakeView(InData.GetData(), InData.Num())), EHttpContentType::Binary);
 			ResponseCode = Request->GetResponseCode();
 
 			if (FHttpRequest::IsSuccessResponse(ResponseCode))
@@ -3252,7 +3256,7 @@ void FHttpCacheStore::RemoveCachedData(const TCHAR* CacheKey, bool bTransient)
 		FScopedHttpPoolRequestPtr Request(PutRequestPools[IsInGameThread()].Get());
 		if (Request.IsValid())
 		{
-			FHttpRequest::Result Result = Request->PerformBlockingQuery<FHttpRequest::Delete>(*Uri, {});
+			FHttpRequest::EResult Result = Request->PerformBlockingDelete(*Uri, {});
 			ResponseCode = Request->GetResponseCode();
 
 			if (ResponseCode == 200)
diff --git a/Engine/Source/Developer/DerivedDataCache/Private/ZenCacheStore.cpp b/Engine/Source/Developer/DerivedDataCache/Private/ZenCacheStore.cpp
--- a/Engine/Source/Developer/DerivedDataCache/Private/ZenCacheStore.cpp
+++ b/Engine/Source/Developer/DerivedDataCache/Private/ZenCacheStore.cpp
@@ -13,6 +13,7 @@
 #include "DerivedDataCacheRecord.h"
 #include "DerivedDataCacheUsageStats.h"
 #include "DerivedDataChunk.h"
+#include "Http/HttpClient.h"
 #include "Misc/ConfigCacheIni.h"
 #include "Misc/Optional.h"
 #include "ProfilingDebugging/CountersTrace.h"
@@ -23,9 +24,11 @@
 #include "Serialization/CompactBinarySerialization.h"
 #include "Serialization/CompactBinaryValidation.h"
 #include "Serialization/CompactBinaryWriter.h"
+#include "Serialization/LargeMemoryReader.h"
+#include "Serialization/LargeMemoryWriter.h"
+#include "Templates/Function.h"
 #include "ZenBackendUtils.h"
 #include "ZenSerialization.h"
-#include "ZenServerHttp.h"
 #include "ZenStatistics.h"
 
 TRACE_DECLARE_INT_COUNTER(ZenDDC_Exist,			TEXT("ZenDDC Exist"));
@@ -147,10 +150,10 @@
 		NotFound,
 		Corrupted
 	};
-	EGetResult GetZenData(FStringView Uri, TArray64<uint8>* OutData, Zen::EContentType ContentType) const;
+	EGetResult GetZenData(FStringView Uri, TArray64<uint8>* OutData, EHttpContentType ContentType) const;
 
 	// TODO: need ability to specify content type
-	FDerivedDataBackendInterface::EPutStatus PutZenData(const TCHAR* Uri, const FCompositeBuffer& InData, Zen::EContentType ContentType);
+	FDerivedDataBackendInterface::EPutStatus PutZenData(const TCHAR* Uri, const FCompositeBuffer& InData, EHttpContentType ContentType);
 	EGetResult GetZenData(const FCacheKey& Key, ECachePolicy CachePolicy, FCbPackage& OutPackage) const;
 
 	bool IsServiceReady();
@@ -161,12 +164,35 @@
 
 	static bool ShouldRetryOnError(int64 ResponseCode);
 
+	using FOnRpcComplete = TUniqueFunction<void(FHttpRequest::EResult HttpResult, FHttpRequest* Request, FCbPackage& OutResponse)>;
+	static FHttpRequest::EResult ParseRpcResponse(FHttpRequest::EResult ResultFromPost, FHttpRequest& Request, FCbPackage& OutResponse);
+	static FHttpRequest::EResult PerformBlockingRpc(FHttpRequest& Request,
+		FStringView Uri,
+		FCbObject RequestObject,
+		FCbPackage &OutResponse);
+	static FHttpRequest::EResult PerformBlockingRpc(FHttpRequest& Request,
+		FStringView Uri,
+		const FCbPackage& RequestPackage,
+		FCbPackage& OutResponse);
+	static void EnqueueAsyncRpc(FHttpRequest& Request,
+		UE::DerivedData::IRequestOwner& Owner,
+		FHttpRequestPool* Pool,
+		FStringView Uri,
+		FCbObject RequestObject,
+		FOnRpcComplete&& OnComplete);
+	static void EnqueueAsyncRpc(FHttpRequest& Request,
+		UE::DerivedData::IRequestOwner& Owner,
+		FHttpRequestPool* Pool,
+		FStringView Uri,
+		const FCbPackage& RequestPackage,
+		FOnRpcComplete&& OnComplete);
+
 private:
 	FString Namespace;
 	FString StructuredNamespace;
 	UE::Zen::FScopeZenService ZenService;
 	mutable FDerivedDataCacheUsageStats UsageStats;
-	TUniquePtr<UE::Zen::FZenHttpRequestPool> RequestPool;
+	TUniquePtr<FHttpRequestPool> RequestPool;
 	bool bIsUsable = false;
 	uint32 FailedLoginAttempts = 0;
 	uint32 MaxAttempts = 4;
@@ -188,7 +214,7 @@
 {
 	if (IsServiceReady())
 	{
-		RequestPool = MakeUnique<Zen::FZenHttpRequestPool>(ZenService.GetInstance().GetURL(), 32);
+		RequestPool = MakeUnique<FHttpRequestPool>(ZenService.GetInstance().GetURL(), ZenService.GetInstance().GetURL(), nullptr, nullptr, 32);
 		bIsUsable = true;
 
 		// Issue a request for stats as it will be fetched asynchronously and issuing now makes them available sooner for future callers.
@@ -232,6 +258,87 @@
 	return false;
 }
 
+FHttpRequest::EResult FZenCacheStore::ParseRpcResponse(FHttpRequest::EResult ResultFromPost, FHttpRequest& Request, FCbPackage& OutResponse)
+{
+	if (ResultFromPost != FHttpRequest::EResult::Success || !FHttpRequest::IsSuccessResponse(Request.GetResponseCode()))
+	{
+		return FHttpRequest::EResult::Failed;
+	}
+
+	const TArray64<uint8>& ResponseBuffer = Request.GetResponseBuffer();
+	if (ResponseBuffer.Num())
+	{
+		FLargeMemoryReader Ar(ResponseBuffer.GetData(), ResponseBuffer.Num());
+		if (!OutResponse.TryLoad(Ar))
+		{
+			return FHttpRequest::EResult::Failed;
+		}
+	}
+
+	return FHttpRequest::EResult::Success;
+}
+
+FHttpRequest::EResult FZenCacheStore::PerformBlockingRpc(FHttpRequest& Request,
+	FStringView Uri,
+	FCbObject RequestObject,
+	FCbPackage& OutResponse)
+{
+	return ParseRpcResponse(
+		Request.PerformBlockingPost(Uri, RequestObject.GetBuffer(), EHttpContentType::CbObject, EHttpContentType::CbPackage), Request, OutResponse);
+}
+
+FHttpRequest::EResult FZenCacheStore::PerformBlockingRpc(FHttpRequest& Request,
+	FStringView Uri,
+	const FCbPackage& RequestPackage,
+	FCbPackage& OutResponse)
+{
+	FLargeMemoryWriter PackageMemory;
+	UE::Zen::Http::SaveCbPackage(RequestPackage, PackageMemory);
+
+	return ParseRpcResponse(
+		Request.PerformBlockingPost(Uri, FCompositeBuffer(FSharedBuffer::MakeView(PackageMemory.GetView())), EHttpContentType::CbPackage, EHttpContentType::CbPackage), Request, OutResponse);
+}
+
+void FZenCacheStore::EnqueueAsyncRpc(FHttpRequest& Request,
+	UE::DerivedData::IRequestOwner& Owner,
+	FHttpRequestPool* Pool,
+	FStringView Uri,
+	FCbObject RequestObject,
+	FOnRpcComplete&& OnComplete)
+{
+	auto OnHttpRequestComplete = [OnComplete = MoveTemp(OnComplete)]
+	(FHttpRequest::EResult HttpResult, FHttpRequest* Request)
+	{
+		FCbPackage Response;
+		FHttpRequest::EResult ParsedResponseResult = ParseRpcResponse(HttpResult, *Request, Response);
+		OnComplete(ParsedResponseResult, Request, Response);
+		return FHttpRequest::ECompletionBehavior::Done;
+	};
+	Request.EnqueueAsyncPost(Owner, Pool, Uri, RequestObject.GetBuffer(), MoveTemp(OnHttpRequestComplete), EHttpContentType::CbObject, EHttpContentType::CbPackage);
+}
+
+void FZenCacheStore::EnqueueAsyncRpc(FHttpRequest& Request,
+	UE::DerivedData::IRequestOwner& Owner,
+	FHttpRequestPool* Pool,
+	FStringView Uri,
+	const FCbPackage& RequestPackage,
+	FOnRpcComplete&& OnComplete)
+{
+	auto OnHttpRequestComplete = [OnComplete = MoveTemp(OnComplete)]
+	(FHttpRequest::EResult HttpResult, FHttpRequest* Request)
+	{
+		FCbPackage Response;
+		FHttpRequest::EResult ParsedResponseResult = ParseRpcResponse(HttpResult, *Request, Response);
+		OnComplete(ParsedResponseResult, Request, Response);
+		return FHttpRequest::ECompletionBehavior::Done;
+	};
+	FLargeMemoryWriter PackageMemory;
+	UE::Zen::Http::SaveCbPackage(RequestPackage, PackageMemory);
+	uint64 PackageMemorySize = PackageMemory.TotalSize();
+	FSharedBuffer PackageSharedBuffer = FSharedBuffer::TakeOwnership(PackageMemory.ReleaseOwnership(), PackageMemorySize, FMemory::Free);
+	Request.EnqueueAsyncPost(Owner, Pool, Uri, FCompositeBuffer(PackageSharedBuffer), MoveTemp(OnHttpRequestComplete), EHttpContentType::CbPackage, EHttpContentType::CbPackage);
+}
+
 bool FZenCacheStore::CachedDataProbablyExists(const TCHAR* CacheKey)
 {
 	TRACE_CPUPROFILER_EVENT_SCOPE(ZenDDC::Exist);
@@ -250,13 +357,13 @@
 	// Retry request until we get an accepted response or exhaust allowed number of attempts.
 	while (ResponseCode == 0 && ++Attempts < MaxAttempts)
 	{
-		Zen::FZenScopedRequestPtr Request(RequestPool.Get());
-		Zen::FZenHttpRequest::Result Result = Request->PerformBlockingHead(*Uri, Zen::EContentType::Binary);
+		FScopedHttpPoolRequestPtr Request(RequestPool.Get());
+		FHttpRequest::EResult Result = Request->PerformBlockingHead(*Uri, EHttpContentType::Binary);
 		ResponseCode = Request->GetResponseCode();
 
-		if (Zen::IsSuccessCode(ResponseCode) || ResponseCode == 404)
+		if (FHttpRequest::IsSuccessResponse(ResponseCode) || ResponseCode == 404)
 		{
-			const bool bIsHit = (Result == Zen::FZenHttpRequest::Result::Success && Zen::IsSuccessCode(ResponseCode));
+			const bool bIsHit = (Result == FHttpRequest::EResult::Success && FHttpRequest::IsSuccessResponse(ResponseCode));
 
 			if (bIsHit)
 			{
@@ -291,7 +398,7 @@
 	double StartTime = FPlatformTime::Seconds();
 
 	TArray64<uint8> ArrayBuffer;
-	EGetResult Result = GetZenData(MakeLegacyZenKey(CacheKey), &ArrayBuffer, Zen::EContentType::Binary);
+	EGetResult Result = GetZenData(MakeLegacyZenKey(CacheKey), &ArrayBuffer, EHttpContentType::Binary);
 	check(ArrayBuffer.Num() <= UINT32_MAX);
 	OutData = TArray<uint8>(MoveTemp(ArrayBuffer));
 	if (Result != EGetResult::Success)
@@ -320,16 +427,16 @@
 }
 
 FZenCacheStore::EGetResult
-FZenCacheStore::GetZenData(FStringView Uri, TArray64<uint8>* OutData, Zen::EContentType ContentType) const
+FZenCacheStore::GetZenData(FStringView Uri, TArray64<uint8>* OutData, EHttpContentType ContentType) const
 {
 	// Retry request until we get an accepted response or exhaust allowed number of attempts.
 	EGetResult GetResult = EGetResult::NotFound;
 	for (uint32 Attempts = 0; Attempts < MaxAttempts; ++Attempts)
 	{
-		Zen::FZenScopedRequestPtr Request(RequestPool.Get());
+		FScopedHttpPoolRequestPtr Request(RequestPool.Get());
 		if (Request.IsValid())
 		{
-			Zen::FZenHttpRequest::Result Result;
+			FHttpRequest::EResult Result;
 			if (OutData)
 			{
 				Result = Request->PerformBlockingDownload(Uri, OutData, ContentType);
@@ -341,7 +448,7 @@
 			int64 ResponseCode = Request->GetResponseCode();
 
 			// Request was successful, make sure we got all the expected data.
-			if (Zen::IsSuccessCode(ResponseCode))
+			if (FHttpRequest::IsSuccessResponse(ResponseCode))
 			{
 				return EGetResult::Success;
 			}
@@ -372,17 +479,18 @@
 	EGetResult GetResult = EGetResult::NotFound;
 	for (uint32 Attempts = 0; Attempts < MaxAttempts; ++Attempts)
 	{
-		Zen::FZenScopedRequestPtr Request(RequestPool.Get());
+		FScopedHttpPoolRequestPtr Request(RequestPool.Get());
 		if (Request.IsValid())
 		{
-			Zen::FZenHttpRequest::Result Result = Request->PerformBlockingDownload(QueryUri.ToString(), OutPackage);
+			TArray64<uint8> PackageBytes;
+			FHttpRequest::EResult Result = Request->PerformBlockingDownload(QueryUri.ToString(), &PackageBytes, EHttpContentType::CbPackage);
 			int64 ResponseCode = Request->GetResponseCode();
-			bool bPackageValid = Request->GetResponseFormatValid();
 
 			// Request was successful, make sure we got all the expected data.
-			if (Zen::IsSuccessCode(ResponseCode))
+			if (FHttpRequest::IsSuccessResponse(ResponseCode))
 			{
-				if (bPackageValid)
+				FLargeMemoryReader Ar(PackageBytes.GetData(), PackageBytes.Num());
+				if (OutPackage.TryLoad(Ar))
 				{
 					GetResult = EGetResult::Success;
 				}
@@ -418,11 +526,11 @@
 	}
 
 	FSharedBuffer DataBuffer = FSharedBuffer::MakeView(InData.GetData(), InData.Num());
-	return PutZenData(*MakeLegacyZenKey(CacheKey), FCompositeBuffer(DataBuffer), Zen::EContentType::Binary);
+	return PutZenData(*MakeLegacyZenKey(CacheKey), FCompositeBuffer(DataBuffer), EHttpContentType::Binary);
 }
 
 FDerivedDataBackendInterface::EPutStatus
-FZenCacheStore::PutZenData(const TCHAR* Uri, const FCompositeBuffer& InData, Zen::EContentType ContentType)
+FZenCacheStore::PutZenData(const TCHAR* Uri, const FCompositeBuffer& InData, EHttpContentType ContentType)
 {
 	TRACE_CPUPROFILER_EVENT_SCOPE(ZenDDC_Put);
 	COOK_STAT(auto Timer = UsageStats.TimePut());
@@ -433,13 +541,13 @@
 	// Retry request until we get an accepted response or exhaust allowed number of attempts.
 	while (ResponseCode == 0 && ++Attempts < MaxAttempts)
 	{
-		Zen::FZenScopedRequestPtr Request(RequestPool.Get());
+		FScopedHttpPoolRequestPtr Request(RequestPool.Get());
 		if (Request.IsValid())
 		{
-			Zen::FZenHttpRequest::Result Result = Request->PerformBlockingPut(Uri, InData, ContentType);
+			FHttpRequest::EResult Result = Request->PerformBlockingPut(Uri, InData, ContentType);
 			ResponseCode = Request->GetResponseCode();
 
-			if (Zen::IsSuccessCode(ResponseCode))
+			if (FHttpRequest::IsSuccessResponse(ResponseCode))
 			{
 				TRACE_COUNTER_ADD(ZenDDC_BytesSent, int64(Request->GetBytesSent()));
 				COOK_STAT(Timer.AddHit(Request->GetBytesSent()));
@@ -499,13 +607,13 @@
 	// Retry request until we get an accepted response or exhaust allowed number of attempts.
 	while (ResponseCode == 0 && ++Attempts < MaxAttempts)
 	{
-		Zen::FZenScopedRequestPtr Request(RequestPool.Get());
+		FScopedHttpPoolRequestPtr Request(RequestPool.Get());
 		if (Request)
 		{
-			Zen::FZenHttpRequest::Result Result = Request->PerformBlockingDelete(*Uri);
+			FHttpRequest::EResult Result = Request->PerformBlockingDelete(*Uri);
 			ResponseCode = Request->GetResponseCode();
 
-			if (Zen::IsSuccessCode(ResponseCode))
+			if (FHttpRequest::IsSuccessResponse(ResponseCode))
 			{
 				return;
 			}
@@ -680,15 +788,15 @@
 		BatchPackage.SetObject(BatchWriter.Save().AsObject());
 
 		FCbPackage BatchResponse;
-		Zen::FZenHttpRequest::Result HttpResult = Zen::FZenHttpRequest::Result::Failed;
+		FHttpRequest::EResult HttpResult = FHttpRequest::EResult::Failed;
 
 		{
-			Zen::FZenScopedRequestPtr Request(RequestPool.Get());
-			HttpResult = Request->PerformRpc(TEXTVIEW("/z$/$rpc"), BatchPackage, BatchResponse);
+			FScopedHttpPoolRequestPtr Request(RequestPool.Get());
+			HttpResult = PerformBlockingRpc(*Request.Get(), TEXTVIEW("/z$/$rpc"), BatchPackage, BatchResponse);
 		}
 
 		int32 RequestIndex = 0;
-		if (HttpResult == Zen::FZenHttpRequest::Result::Success)
+		if (HttpResult == FHttpRequest::EResult::Success)
 		{
 			const FCbObject& ResponseObj = BatchResponse.GetObject();
 			for (FCbField ResponseField : ResponseObj[ANSITEXTVIEW("Result")])
@@ -789,16 +897,16 @@
 		BatchRequest.EndObject();
 
 		FCbPackage BatchResponse;
-		Zen::FZenHttpRequest::Result HttpResult = Zen::FZenHttpRequest::Result::Failed;
+		FHttpRequest::EResult HttpResult = FHttpRequest::EResult::Failed;
 
 		{
 			LLM_SCOPE_BYTAG(UntaggedDDCResult);
-			Zen::FZenScopedRequestPtr Request(RequestPool.Get());
-			HttpResult = Request->PerformRpc(TEXTVIEW("/z$/$rpc"), BatchRequest.Save().AsObject(), BatchResponse);
+			FScopedHttpPoolRequestPtr Request(RequestPool.Get());
+			HttpResult = PerformBlockingRpc(*Request.Get(), TEXTVIEW("/z$/$rpc"), BatchRequest.Save().AsObject(), BatchResponse);
 		}
 
 		int32 RequestIndex = 0;
-		if (HttpResult == Zen::FZenHttpRequest::Result::Success)
+		if (HttpResult == FHttpRequest::EResult::Success)
 		{
 			const FCbObject& ResponseObj = BatchResponse.GetObject();
 			
@@ -917,14 +1025,14 @@
 		BatchPackage.SetObject(BatchWriter.Save().AsObject());
 
 		FCbPackage BatchResponse;
-		Zen::FZenHttpRequest::Result HttpResult = Zen::FZenHttpRequest::Result::Failed;
+		FHttpRequest::EResult HttpResult = FHttpRequest::EResult::Failed;
 		{
-			Zen::FZenScopedRequestPtr Request(RequestPool.Get());
-			HttpResult = Request->PerformRpc(TEXTVIEW("/z$/$rpc"), BatchPackage, BatchResponse);
+			FScopedHttpPoolRequestPtr Request(RequestPool.Get());
+			HttpResult = PerformBlockingRpc(*Request.Get(), TEXTVIEW("/z$/$rpc"), BatchPackage, BatchResponse);
 		}
 
 		int32 RequestIndex = 0;
-		if (HttpResult == Zen::FZenHttpRequest::Result::Success)
+		if (HttpResult == FHttpRequest::EResult::Success)
 		{
 			const FCbObject& ResponseObj = BatchResponse.GetObject();
 			for (FCbField ResponseField : ResponseObj[ANSITEXTVIEW("Result")])
@@ -1022,16 +1130,16 @@
 			BatchRequest.EndObject();
 
 			FCbPackage BatchResponse;
-			Zen::FZenHttpRequest::Result HttpResult = Zen::FZenHttpRequest::Result::Failed;
+			FHttpRequest::EResult HttpResult = FHttpRequest::EResult::Failed;
 
 			{
 				LLM_SCOPE_BYTAG(UntaggedDDCResult);
-				Zen::FZenScopedRequestPtr Request(RequestPool.Get());
-				HttpResult = Request->PerformRpc(TEXTVIEW("/z$/$rpc"), BatchRequest.Save().AsObject(), BatchResponse);
+				FScopedHttpPoolRequestPtr Request(RequestPool.Get());
+				HttpResult = PerformBlockingRpc(*Request.Get(), TEXTVIEW("/z$/$rpc"), BatchRequest.Save().AsObject(), BatchResponse);
 			}
 
 			int32 RequestIndex = 0;
-			if (HttpResult == Zen::FZenHttpRequest::Result::Success)
+			if (HttpResult == FHttpRequest::EResult::Success)
 			{
 				const FCbObject& ResponseObj = BatchResponse.GetObject();
 
@@ -1167,16 +1275,16 @@
 		BatchRequest.EndObject();
 
 		FCbPackage BatchResponse;
-		Zen::FZenHttpRequest::Result HttpResult = Zen::FZenHttpRequest::Result::Failed;
+		FHttpRequest::EResult HttpResult = FHttpRequest::EResult::Failed;
 
 		{
 			LLM_SCOPE_BYTAG(UntaggedDDCResult);
-			Zen::FZenScopedRequestPtr Request(RequestPool.Get());
-			HttpResult = Request->PerformRpc(TEXTVIEW("/z$/$rpc"), BatchRequest.Save().AsObject(), BatchResponse);
+			FScopedHttpPoolRequestPtr Request(RequestPool.Get());
+			HttpResult = PerformBlockingRpc(*Request.Get(), TEXTVIEW("/z$/$rpc"), BatchRequest.Save().AsObject(), BatchResponse);
 		}
 
 		int32 RequestIndex = 0;
-		if (HttpResult == Zen::FZenHttpRequest::Result::Success)
+		if (HttpResult == FHttpRequest::EResult::Success)
 		{
 			const FCbObject& ResponseObj = BatchResponse.GetObject();
 
diff --git a/Engine/Source/Developer/IoStoreUtilities/Private/ZenStoreHttpClient.cpp b/Engine/Source/Developer/IoStoreUtilities/Private/ZenStoreHttpClient.cpp
--- a/Engine/Source/Developer/IoStoreUtilities/Private/ZenStoreHttpClient.cpp
+++ b/Engine/Source/Developer/IoStoreUtilities/Private/ZenStoreHttpClient.cpp
@@ -261,7 +261,7 @@
 		{
 			// Old-style with all attachments by value
 
-			UE::Zen::SaveCbPackage(OpEntry, SerializedPackage);
+			UE::Zen::OpLog::SaveCbPackage(OpEntry, SerializedPackage);
 		}
 		else
 		{
@@ -368,7 +368,7 @@
 
 					if (!bIsSerialized)
 					{
-						UE::Zen::SaveCbAttachment(Attachment, Writer);
+						UE::Zen::OpLog::SaveCbAttachment(Attachment, Writer);
 					}
 				}
 				else
@@ -585,7 +585,7 @@
 	check(bAllowEdit);
 
 	FLargeMemoryWriter SerializedPackage;
-	UE::Zen::SaveCbPackage(OpEntry, SerializedPackage);
+	UE::Zen::OpLog::SaveCbPackage(OpEntry, SerializedPackage);
 
 	UE_LOG(LogZenStore, Verbose, TEXT("Package size: %lld"), SerializedPackage.TotalSize());
 
diff --git a/Engine/Source/Developer/Zen/Internal/ZenSerialization.cpp b/Engine/Source/Developer/Zen/Internal/ZenSerialization.cpp
--- a/Engine/Source/Developer/Zen/Internal/ZenSerialization.cpp
+++ b/Engine/Source/Developer/Zen/Internal/ZenSerialization.cpp
@@ -1,6 +1,7 @@
 // Copyright Epic Games, Inc. All Rights Reserved.
 
 #include "ZenSerialization.h"
+#include "Compression/OodleDataCompression.h"
 #include "Serialization/CompactBinary.h"
 #include "Serialization/CompactBinaryPackage.h"
 #include "Serialization/CompactBinarySerialization.h"
@@ -9,8 +10,11 @@
 
 #if UE_WITH_ZEN
 
-namespace UE::Zen {
+namespace UE::Zen
+{
 
+namespace OpLog
+{
 void SaveCbAttachment(const FCbAttachment& Attachment, FCbWriter& Writer)
 {
 	if (Attachment.IsCompressedBinary())
@@ -117,6 +121,186 @@ bool TryLoadCbPackage(FCbPackage& Package, FArchive& Ar, FCbBufferAllocator Allo
 		}
 	}
 }
+}
+
+namespace Http
+{
+struct CbPackageHeader
+{
+	uint32	HeaderMagic;
+	uint32	AttachmentCount;
+	uint32	Reserved1;
+	uint32	Reserved2;
+};
+
+static const uint32 kMagic = 0xaa77aacc;
+
+struct CbAttachmentEntry
+{
+	uint64	AttachmentSize;
+	uint32	Flags;
+	FIoHash	AttachmentHash;
+
+	enum
+	{
+		IsCompressed = (1u << 0),	// Is marshaled using compressed buffer storage format
+		IsObject = (1u << 1),		// Is compact binary object
+	};
+
+};
+
+void SaveCbPackage(const FCbPackage& Package, FArchive& Ar)
+{
+	TConstArrayView<FCbAttachment> Attachments = Package.GetAttachments();
+	const FCbObject& Object = Package.GetObject();
+	FCompressedBuffer ObjectBuffer = FCompressedBuffer::Compress(Object.GetBuffer(), FOodleDataCompression::ECompressor::NotSet, FOodleDataCompression::ECompressionLevel::None);
+
+	CbPackageHeader Hdr;
+	Hdr.HeaderMagic = kMagic;
+	Hdr.AttachmentCount = Attachments.Num();
+	Hdr.Reserved1 = 0;
+	Hdr.Reserved2 = 0;
+
+	Ar.Serialize(&Hdr, sizeof Hdr);
+
+	// Root object metadata
+
+	{
+		CbAttachmentEntry Entry;
+		Entry.AttachmentHash = ObjectBuffer.GetRawHash();
+		Entry.AttachmentSize = ObjectBuffer.GetCompressedSize();
+		Entry.Flags = CbAttachmentEntry::IsObject | CbAttachmentEntry::IsCompressed;
+
+		Ar.Serialize(&Entry, sizeof Entry);
+	}
+
+	// Attachment metadata
+
+	for (const FCbAttachment& Attachment : Attachments)
+	{
+		CbAttachmentEntry Entry;
+		Entry.AttachmentHash = Attachment.GetHash();
+		Entry.Flags = 0;
+
+		if (Attachment.IsCompressedBinary())
+		{
+			Entry.AttachmentSize = Attachment.AsCompressedBinary().GetCompressedSize();
+			Entry.Flags |= CbAttachmentEntry::IsCompressed;
+		}
+		else if (Attachment.IsBinary())
+		{
+			Entry.AttachmentSize = Attachment.AsCompositeBinary().GetSize();
+		}
+		else if (Attachment.IsNull())
+		{
+			checkNoEntry();
+		}
+		else if (Attachment.IsObject())
+		{
+			Entry.AttachmentSize = Attachment.AsObject().GetSize();
+			Entry.Flags |= CbAttachmentEntry::IsObject;
+		}
+		else
+		{
+			checkNoEntry();
+		}
+
+		Ar.Serialize(&Entry, sizeof Entry);
+	}
+
+	// Root object
+
+	Ar << ObjectBuffer;
+
+	// Payloads back-to-back
+
+	for (const FCbAttachment& Attachment : Attachments)
+	{
+		if (Attachment.IsCompressedBinary())
+		{
+			FCompressedBuffer Payload = Attachment.AsCompressedBinary();
+			Ar << Payload;
+		}
+		else if (Attachment.IsBinary())
+		{
+			const FCompositeBuffer& Buffer = Attachment.AsCompositeBinary();
+			FSharedBuffer SharedBuffer = Buffer.ToShared();
+
+			Ar.Serialize((void*)SharedBuffer.GetData(), SharedBuffer.GetSize());
+		}
+		else if (Attachment.IsNull())
+		{
+			checkNoEntry();
+		}
+		else if (Attachment.IsObject())
+		{
+			checkNoEntry();
+		}
+		else
+		{
+			checkNoEntry();
+		}
+	}
+}
+
+bool TryLoadCbPackage(FCbPackage& Package, FArchive& Ar, FCbBufferAllocator Allocator)
+{
+	CbPackageHeader Hdr;
+	Ar.Serialize(&Hdr, sizeof Hdr);
+
+	if (Hdr.HeaderMagic != kMagic)
+	{
+		return false;
+	}
+
+	TArray<CbAttachmentEntry> AttachmentEntries;
+	AttachmentEntries.SetNum(Hdr.AttachmentCount + 1);
+
+	Ar.Serialize(AttachmentEntries.GetData(), (Hdr.AttachmentCount + 1) * sizeof(CbAttachmentEntry));
+
+	int Index = 0;
+
+	for (const CbAttachmentEntry& Entry : AttachmentEntries)
+	{
+		FUniqueBuffer AttachmentData = FUniqueBuffer::Alloc(Entry.AttachmentSize);
+		Ar.Serialize(AttachmentData.GetData(), AttachmentData.GetSize());
+
+		if (Entry.Flags & CbAttachmentEntry::IsCompressed)
+		{
+			FCompressedBuffer CompBuf(FCompressedBuffer::FromCompressed(AttachmentData.MoveToShared()));
+
+			if (Entry.Flags & CbAttachmentEntry::IsObject)
+			{
+				checkf(Index == 0, TEXT("Object attachments are not currently supported"));
+
+				Package.SetObject(FCbObject(CompBuf.Decompress()));
+			}
+			else
+			{
+				FCbAttachment Attachment(MoveTemp(CompBuf));
+				Package.AddAttachment(Attachment);
+			}
+		}
+		else /* not compressed */
+		{
+			if (Entry.Flags & CbAttachmentEntry::IsObject)
+			{
+				checkf(Index == 0, TEXT("Object attachments are not currently supported"));
+
+				Package.SetObject(FCbObject(AttachmentData.MoveToShared()));
+			}
+			else
+			{
+				FCbAttachment Attachment(AttachmentData.MoveToShared());
+				Package.AddAttachment(Attachment);
+			}
+		}
+
+		++Index;
+	}
+	return true;
+}
+}
 
 }
 #endif // UE_WITH_ZEN
diff --git a/Engine/Source/Developer/Zen/Internal/ZenSerialization.h b/Engine/Source/Developer/Zen/Internal/ZenSerialization.h
--- a/Engine/Source/Developer/Zen/Internal/ZenSerialization.h
+++ b/Engine/Source/Developer/Zen/Internal/ZenSerialization.h
@@ -16,10 +16,19 @@ class FCbWriter;
 
 namespace UE::Zen {
 
+namespace OpLog
+{
 ZEN_API void SaveCbAttachment(const FCbAttachment& Attachment, FCbWriter& Writer);
 ZEN_API void SaveCbPackage(const FCbPackage& Package, FCbWriter& Writer);
 ZEN_API void SaveCbPackage(const FCbPackage& Package, FArchive& Ar);
 ZEN_API bool TryLoadCbPackage(FCbPackage& Package, FArchive& Ar, FCbBufferAllocator Allocator = FUniqueBuffer::Alloc);
+}
+
+namespace Http
+{
+ZEN_API void SaveCbPackage(const FCbPackage& Package, FArchive& Ar);
+ZEN_API bool TryLoadCbPackage(FCbPackage& Package, FArchive& Ar, FCbBufferAllocator Allocator = FUniqueBuffer::Alloc);
+}
 
 } // namespace UE::Zen
 
diff --git a/Engine/Source/Developer/Zen/Internal/ZenServerHttp.cpp b/Engine/Source/Developer/Zen/Internal/ZenServerHttp.cpp
--- a/Engine/Source/Developer/Zen/Internal/ZenServerHttp.cpp
+++ b/Engine/Source/Developer/Zen/Internal/ZenServerHttp.cpp
@@ -26,7 +26,6 @@
 
 #include "Logging/LogMacros.h"
 #include "Compression/CompressedBuffer.h"
-#include "Compression/OodleDataCompression.h"
 #include "Containers/StringFwd.h"
 #include "HAL/LowLevelMemTracker.h"
 #include "Memory/CompositeBuffer.h"
@@ -38,6 +37,7 @@
 #include "Serialization/LargeMemoryWriter.h"
 #include "HAL/PlatformProcess.h"
 #include "ProfilingDebugging/CpuProfilerTrace.h"
+#include "ZenSerialization.h"
 
 LLM_DEFINE_TAG(ZenDDC, NAME_None, TEXT("DDCBackend"));
 
@@ -153,123 +153,10 @@ namespace UE::Zen {
 		return PerformBlockingPost(Uri, Out.GetView(), EContentType::CbObject, AcceptType);
 	}
 
-	struct CbPackageHeader
-	{
-		uint32	HeaderMagic;
-		uint32	AttachmentCount;
-		uint32	Reserved1;
-		uint32	Reserved2;
-	};
-
-	static const uint32 kMagic = 0xaa77aacc;
-
-	struct CbAttachmentEntry
-	{
-		uint64	AttachmentSize;
-		uint32	Flags;
-		FIoHash	AttachmentHash;
-
-		enum
-		{
-			IsCompressed = (1u << 0),	// Is marshaled using compressed buffer storage format
-			IsObject = (1u << 1),		// Is compact binary object
-		};
-
-	};
-
 	FZenHttpRequest::Result FZenHttpRequest::PerformBlockingPostPackage(FStringView Uri, const FCbPackage& Package, EContentType AcceptType)
 	{
-		TConstArrayView<FCbAttachment> Attachments = Package.GetAttachments();
-		const FCbObject& Object = Package.GetObject();
-		FCompressedBuffer ObjectBuffer = FCompressedBuffer::Compress(Object.GetBuffer(), FOodleDataCompression::ECompressor::NotSet, FOodleDataCompression::ECompressionLevel::None);
-
-		CbPackageHeader Hdr;
-		Hdr.HeaderMagic = kMagic;
-		Hdr.AttachmentCount = Attachments.Num();
-		Hdr.Reserved1 = 0;
-		Hdr.Reserved2 = 0;
-
 		FLargeMemoryWriter Out;
-		Out.Serialize(&Hdr, sizeof Hdr);
-
-		// Root object metadata
-
-		{
-			CbAttachmentEntry Entry;
-			Entry.AttachmentHash = ObjectBuffer.GetRawHash();
-			Entry.AttachmentSize = ObjectBuffer.GetCompressedSize();
-			Entry.Flags = CbAttachmentEntry::IsObject | CbAttachmentEntry::IsCompressed;
-
-			Out.Serialize(&Entry, sizeof Entry);
-		}
-
-		// Attachment metadata
-
-		for (const FCbAttachment& Attachment : Attachments)
-		{
-			CbAttachmentEntry Entry;
-			Entry.AttachmentHash = Attachment.GetHash();
-			Entry.Flags = 0;
-
-			if (Attachment.IsCompressedBinary())
-			{
-				Entry.AttachmentSize = Attachment.AsCompressedBinary().GetCompressedSize();
-				Entry.Flags |= CbAttachmentEntry::IsCompressed;
-			}
-			else if (Attachment.IsBinary())
-			{
-				Entry.AttachmentSize = Attachment.AsCompositeBinary().GetSize();
-			}
-			else if (Attachment.IsNull())
-			{
-				checkNoEntry();
-			}
-			else if (Attachment.IsObject())
-			{
-				Entry.AttachmentSize = Attachment.AsObject().GetSize();
-				Entry.Flags |= CbAttachmentEntry::IsObject;
-			}
-			else
-			{
-				checkNoEntry();
-			}
-
-			Out.Serialize(&Entry, sizeof Entry);
-		}
-
-		// Root object
-
-		Out << ObjectBuffer;
-
-		// Payloads back-to-back
-
-		for (const FCbAttachment& Attachment : Attachments)
-		{
-			if (Attachment.IsCompressedBinary())
-			{
-				FCompressedBuffer Payload = Attachment.AsCompressedBinary();
-				Out << Payload;
-			}
-			else if (Attachment.IsBinary())
-			{
-				const FCompositeBuffer& Buffer = Attachment.AsCompositeBinary();
-				FSharedBuffer SharedBuffer = Buffer.ToShared();
-
-				Out.Serialize((void*)SharedBuffer.GetData(), SharedBuffer.GetSize());
-			}
-			else if (Attachment.IsNull())
-			{
-				checkNoEntry();
-			}
-			else if (Attachment.IsObject())
-			{
-				checkNoEntry();
-			}
-			else
-			{
-				checkNoEntry();
-			}
-		}
+		Http::SaveCbPackage(Package, Out);
 
 		return PerformBlockingPost(Uri, Out.GetView(), EContentType::CbPackage, AcceptType);
 	}
@@ -311,61 +198,11 @@ namespace UE::Zen {
 		FLargeMemoryReader Reader(Response.GetData(), Response.Num());
 
 		FCbPackage Package;
-
-		CbPackageHeader Hdr;
-		Reader.Serialize(&Hdr, sizeof Hdr);
-
-		if (Hdr.HeaderMagic != kMagic)
+		if (!Http::TryLoadCbPackage(Package, Reader))
 		{
 			return {};
 		}
 
-		TArray<CbAttachmentEntry> AttachmentEntries;
-		AttachmentEntries.SetNum(Hdr.AttachmentCount + 1);
-
-		Reader.Serialize(AttachmentEntries.GetData(), (Hdr.AttachmentCount + 1) * sizeof(CbAttachmentEntry));
-
-		int Index = 0;
-
-		for (const CbAttachmentEntry& Entry : AttachmentEntries)
-		{
-			FUniqueBuffer AttachmentData = FUniqueBuffer::Alloc(Entry.AttachmentSize);
-			Reader.Serialize(AttachmentData.GetData(), AttachmentData.GetSize());
-
-			if (Entry.Flags & CbAttachmentEntry::IsCompressed)
-			{
-				FCompressedBuffer CompBuf(FCompressedBuffer::FromCompressed(AttachmentData.MoveToShared()));
-
-				if (Entry.Flags & CbAttachmentEntry::IsObject)
-				{
-					checkf(Index == 0, TEXT("Object attachments are not currently supported"));
-
-					Package.SetObject(FCbObject(CompBuf.Decompress()));
-				}
-				else
-				{
-					FCbAttachment Attachment(MoveTemp(CompBuf));
-					Package.AddAttachment(Attachment);
-				}
-			}
-			else /* not compressed */
-			{
-				if (Entry.Flags & CbAttachmentEntry::IsObject)
-				{
-					checkf(Index == 0, TEXT("Object attachments are not currently supported"));
-
-					Package.SetObject(FCbObject(AttachmentData.MoveToShared()));
-				}
-				else
-				{
-					FCbAttachment Attachment(AttachmentData.MoveToShared());
-					Package.AddAttachment(Attachment);
-				}
-			}
-
-			++Index;
-		}
-
 		return Package;
 	}
 
