DDC: Exposed IsCanceled and LaunchTask through FRequestOwner

#rnx

diff --git a/Engine/Source/Developer/DerivedDataCache/Public/DerivedDataRequestOwner.h b/Engine/Source/Developer/DerivedDataCache/Public/DerivedDataRequestOwner.h
--- a/Engine/Source/Developer/DerivedDataCache/Public/DerivedDataRequestOwner.h
+++ b/Engine/Source/Developer/DerivedDataCache/Public/DerivedDataRequestOwner.h
@@ -128,6 +128,15 @@ public:
 	/** Poll whether the owner has any active requests or barriers. */
 	UE_API bool Poll() const;
 
+	/** Returns whether the owner has been canceled, which new requests are expected to check. */
+	inline bool IsCanceled() const { return Owner->IsCanceled(); }
+
+	/** Launches a task that executes the task body when scheduled. See IRequestOwner::LaunchTask. */
+	inline void LaunchTask(const TCHAR* DebugName, TUniqueFunction<void ()>&& TaskBody)
+	{
+		Owner->LaunchTask(DebugName, MoveTemp(TaskBody));
+	}
+
 	/** Access as a request owner. */
 	inline operator IRequestOwner&() { return *Owner; }
 	/** Access as a request. */
