DDC: Removed the legacy cache implementation from the clients for Horde Storage and Zen

diff --git a/Engine/Source/Developer/DerivedDataCache/Private/DerivedDataBackends.cpp b/Engine/Source/Developer/DerivedDataCache/Private/DerivedDataBackends.cpp
--- a/Engine/Source/Developer/DerivedDataCache/Private/DerivedDataBackends.cpp
+++ b/Engine/Source/Developer/DerivedDataCache/Private/DerivedDataBackends.cpp
@@ -52,23 +52,11 @@
 ILegacyCacheStore* CreateCacheStoreThrottle(ILegacyCacheStore* InnerCache, uint32 LatencyMS, uint32 MaxBytesPerSecond);
 ILegacyCacheStore* CreateCacheStoreVerify(ILegacyCacheStore* InnerCache, bool bPutOnError);
 ILegacyCacheStore* CreateFileSystemCacheStore(const TCHAR* CacheDirectory, const TCHAR* Params, const TCHAR* AccessLogFileName, ECacheStoreFlags& OutFlags);
-ILegacyCacheStore* CreateHttpCacheStore(
-	const TCHAR* NodeName,
-	const TCHAR* ServiceUrl,
-	bool bResolveHostCanonicalName,
-	const TCHAR* Namespace,
-	const TCHAR* StructuredNamespace,
-	const TCHAR* OAuthProvider,
-	const TCHAR* OAuthClientId,
-	const TCHAR* OAuthData, 
-	const TCHAR* OAuthScope,
-	const FDerivedDataBackendInterface::ESpeedClass* ForceSpeedClass,
-	EBackendLegacyMode LegacyMode,
-	bool bReadOnly);
+TTuple<ILegacyCacheStore*, ECacheStoreFlags> CreateHttpCacheStore(const TCHAR* NodeName, const TCHAR* Config);
 IMemoryCacheStore* CreateMemoryCacheStore(const TCHAR* Name, int64 MaxCacheSize, bool bCanBeDisabled);
 IPakFileCacheStore* CreatePakFileCacheStore(const TCHAR* Filename, bool bWriting, bool bCompressed);
 ILegacyCacheStore* CreateS3CacheStore(const TCHAR* RootManifestPath, const TCHAR* BaseUrl, const TCHAR* Region, const TCHAR* CanaryObjectKey, const TCHAR* CachePath);
-ILegacyCacheStore* CreateZenCacheStore(const TCHAR* NodeName, const TCHAR* ServiceUrl, const TCHAR* Namespace, const TCHAR* StructuredNamespace);
+TTuple<ILegacyCacheStore*, ECacheStoreFlags> CreateZenCacheStore(const TCHAR* NodeName, const TCHAR* Config);
 
 /**
  * This class is used to create a singleton that represents the derived data cache hierarchy and all of the wrappers necessary
@@ -138,8 +126,8 @@
 					ParsedNodes.Empty();
 					DestroyCreatedBackends();
 					UE_LOG(LogDerivedDataCache, Warning,
-						TEXT("Unable to create cache graph using the requested graph settings (%s). "),
-						TEXT("Reverting to the default graph."), *GraphName);
+						TEXT("Unable to create cache graph using the requested graph settings (%s). "
+						"Reverting to the default graph."), *GraphName);
 				}
 			}
 
@@ -318,11 +306,11 @@
 				}
 				else if (NodeType == TEXT("Http"))
 				{
-					ParsedNode = ParseHttpCache(*NodeName, *Entry, IniFilename, IniSection);
+					ParsedNode = CreateHttpCacheStore(*NodeName, *Entry);
 				}
 				else if (NodeType == TEXT("Zen"))
 				{
-					ParsedNode = MakeTuple(ParseZenCache(*NodeName, *Entry), ECacheStoreFlags::Local | ECacheStoreFlags::Remote | ECacheStoreFlags::Query | ECacheStoreFlags::Store);
+					ParsedNode = CreateZenCacheStore(*NodeName, *Entry);
 				}
 				
 				if (ParsedNode.Key)
@@ -799,238 +787,7 @@
 		return nullptr;
 	}
 
-	void ParseHttpCacheParams(
-		const TCHAR* NodeName,
-		const TCHAR* Entry,
-		const FString& IniFilename,
-		const TCHAR* IniSection,
-		FString& Host,
-		bool& bResolveHostCanonicalName,
-		FString& Namespace,
-		FString& StructuredNamespace,
-		FString& OAuthProvider,
-		FString& OAuthClientId,
-		FString& OAuthSecret,
-		FString& OAuthScope,
-		EBackendLegacyMode& LegacyMode,
-		bool& bReadOnly)
-	{
-		FString ServerId;
-		if (FParse::Value(Entry, TEXT("ServerID="), ServerId))
-		{
-			FString ServerEntry;
-			const TCHAR* ServerSection = TEXT("HordeStorageServers");
-			if (GConfig->GetString(ServerSection, *ServerId, ServerEntry, IniFilename))
-			{
-				ParseHttpCacheParams(NodeName, *ServerEntry, IniFilename, IniSection, Host, bResolveHostCanonicalName, Namespace, StructuredNamespace, OAuthProvider, OAuthClientId, OAuthSecret, OAuthScope, LegacyMode, bReadOnly);
-			}
-			else
-			{
-				UE_LOG(LogDerivedDataCache, Warning, TEXT("Node %s is using ServerID=%s which was not found in [%s]"), NodeName, *ServerId, ServerSection);
-			}
-		}
-
-		FParse::Value(Entry, TEXT("Host="), Host);
-
-		FString EnvHostOverride;
-		if (FParse::Value(Entry, TEXT("EnvHostOverride="), EnvHostOverride))
-		{
-			FString HostEnv = FPlatformMisc::GetEnvironmentVariable(*EnvHostOverride);
-			if (!HostEnv.IsEmpty())
-			{
-				Host = HostEnv;
-				UE_LOG(LogDerivedDataCache, Log, TEXT("Node %s found environment variable for Host %s=%s"), NodeName, *EnvHostOverride, *Host);
-			}
-		}
-
-		FString CommandLineOverride;
-		if (FParse::Value(Entry, TEXT("CommandLineHostOverride="), CommandLineOverride))
-		{
-			if (FParse::Value(FCommandLine::Get(), *(CommandLineOverride + TEXT("=")), Host))
-			{
-				UE_LOG(LogDerivedDataCache, Log, TEXT("Node %s found command line override for Host %s=%s"), NodeName, *CommandLineOverride, *Host);
-			}
-		}
-
-		FParse::Bool(Entry, TEXT("ResolveHostCanonicalName="), bResolveHostCanonicalName);
-		FParse::Value(Entry, TEXT("Namespace="), Namespace);
-		FParse::Value(Entry, TEXT("StructuredNamespace="), StructuredNamespace);
-		FParse::Value(Entry, TEXT("OAuthProvider="), OAuthProvider);
-		
-		FString CommandLineOAuthProvidorOverride;
-		if (FParse::Value(Entry, TEXT("CommandLineOAuthProviderOverride="), CommandLineOAuthProvidorOverride))
-		{
-			if (FParse::Value(FCommandLine::Get(), *(CommandLineOAuthProvidorOverride + TEXT("=")), OAuthProvider))
-			{
-				UE_LOG(LogDerivedDataCache, Log, TEXT("Node %s found command line override for OAuthProvider %s=%s"), NodeName, *CommandLineOAuthProvidorOverride, *OAuthProvider);
-			}
-		}
-		
-		FParse::Value(Entry, TEXT("OAuthClientId="), OAuthClientId);
-		FParse::Value(Entry, TEXT("OAuthSecret="), OAuthSecret);
-		
-		FString CommandLineOAuthSecretOverride;
-		if (FParse::Value(Entry, TEXT("CommandLineOAuthSecretOverride="), CommandLineOAuthSecretOverride))
-		{
-			if (FParse::Value(FCommandLine::Get(), *(CommandLineOAuthSecretOverride + TEXT("=")), OAuthSecret))
-			{
-				UE_LOG(LogDerivedDataCache, Log, TEXT("Node %s found command line override for OAuthSecret %s=%s"), NodeName, *CommandLineOAuthSecretOverride, *OAuthSecret);
-			}
-		}
-		
-		FParse::Value(Entry, TEXT("OAuthScope="), OAuthScope);
-		
-		FParse::Bool(Entry, TEXT("ReadOnly="), bReadOnly);
-
-		if (FString LegacyModeString; FParse::Value(Entry, TEXT("LegacyMode="), LegacyModeString))
-		{
-			if (!TryLexFromString(LegacyMode, LegacyModeString))
-			{
-				UE_LOG(LogDerivedDataCache, Warning, TEXT("%s: Ignoring unrecognized legacy mode '%s'"), NodeName, *LegacyModeString);
-			}
-		}
-	}
-
 	/**
-	 * Creates a HTTP data cache interface.
-	 */
-	FParsedNode ParseHttpCache(
-		const TCHAR* NodeName,
-		const TCHAR* Entry,
-		const FString& IniFilename,
-		const TCHAR* IniSection)
-	{
-		FString Host;
-		bool bResolveHostCanonicalName = true;
-		FString Namespace;
-		FString StructuredNamespace;
-		FString OAuthProvider;
-		FString OAuthClientId;
-		FString OAuthSecret;
-		FString OAuthScope;
-		EBackendLegacyMode LegacyMode = EBackendLegacyMode::ValueOnly;
-		bool bReadOnly = false;
-
-		ParseHttpCacheParams(NodeName, Entry, IniFilename, IniSection, Host, bResolveHostCanonicalName, Namespace, StructuredNamespace, OAuthProvider, OAuthClientId, OAuthSecret, OAuthScope, LegacyMode, bReadOnly);
-
-		if (Host.IsEmpty())
-		{
-			UE_LOG(LogDerivedDataCache, Error, TEXT("Node %s does not specify 'Host'"), NodeName);
-			return MakeTuple(nullptr, ECacheStoreFlags::None);
-		}
-
-		if (Host == TEXT("None"))
-		{
-			UE_LOG(LogDerivedDataCache, Log, TEXT("Node %s is disabled because Host is set to 'None'"), NodeName);
-			return MakeTuple(nullptr, ECacheStoreFlags::None);
-		}
-
-		if (Namespace.IsEmpty())
-		{
-			Namespace = FApp::GetProjectName();
-			UE_LOG(LogDerivedDataCache, Warning, TEXT("Node %s does not specify 'Namespace', falling back to '%s'"), NodeName, *Namespace);
-		}
-
-		if (StructuredNamespace.IsEmpty())
-		{
-			StructuredNamespace = Namespace;
-		}
-
-		if (OAuthProvider.IsEmpty())
-		{
-			UE_LOG(LogDerivedDataCache, Error, TEXT("Node %s does not specify 'OAuthProvider'"), NodeName);
-			return MakeTuple(nullptr, ECacheStoreFlags::None);
-		}
-
-		// No need for OAuth client id and secret if using a local provider.
-		if (!OAuthProvider.StartsWith(TEXT("http://localhost")))
-		{
-			if (OAuthClientId.IsEmpty())
-			{
-				UE_LOG(LogDerivedDataCache, Error, TEXT("Node %s does not specify 'OAuthClientId'"), NodeName);
-				return MakeTuple(nullptr, ECacheStoreFlags::None);
-			}
-
-			if (OAuthSecret.IsEmpty())
-			{
-				UE_LOG(LogDerivedDataCache, Error, TEXT("Node %s does not specify 'OAuthSecret'"), NodeName);
-				return MakeTuple(nullptr, ECacheStoreFlags::None);
-			}
-		}
-
-		if (OAuthScope.IsEmpty())
-		{
-			OAuthScope = TEXT("cache_access");
-		}
-		
-		FDerivedDataBackendInterface::ESpeedClass ForceSpeedClass = FDerivedDataBackendInterface::ESpeedClass::Unknown;
-		FString ForceSpeedClassValue;
-		if (FParse::Value(FCommandLine::Get(), TEXT("HttpForceSpeedClass="), ForceSpeedClassValue))
-		{
-			if (ForceSpeedClassValue == TEXT("Slow"))
-			{
-				ForceSpeedClass = FDerivedDataBackendInterface::ESpeedClass::Slow;
-			}
-			else if (ForceSpeedClassValue == TEXT("Ok"))
-			{
-				ForceSpeedClass = FDerivedDataBackendInterface::ESpeedClass::Ok;
-			}
-			else if (ForceSpeedClassValue == TEXT("Fast"))
-			{
-				ForceSpeedClass = FDerivedDataBackendInterface::ESpeedClass::Fast;
-			}
-			else if (ForceSpeedClassValue == TEXT("Local"))
-			{
-				ForceSpeedClass = FDerivedDataBackendInterface::ESpeedClass::Local;
-			}
-			else
-			{
-				UE_LOG(LogDerivedDataCache, Warning, TEXT("Node %s found unknown speed class override HttpForceSpeedClass=%s"), NodeName, *ForceSpeedClassValue);
-			}
-		}
-
-		if (ForceSpeedClass != FDerivedDataBackendInterface::ESpeedClass::Unknown)
-		{
-			UE_LOG(LogDerivedDataCache, Log, TEXT("Node %s found speed class override ForceSpeedClass=%s"), NodeName, *ForceSpeedClassValue);
-		}
-
-		return MakeTuple(CreateHttpCacheStore(
-			NodeName, *Host, bResolveHostCanonicalName, *Namespace, *StructuredNamespace, *OAuthProvider, *OAuthClientId, *OAuthSecret, *OAuthScope,
-			ForceSpeedClass == FDerivedDataBackendInterface::ESpeedClass::Unknown ? nullptr : &ForceSpeedClass, LegacyMode, bReadOnly),
-			ECacheStoreFlags::Remote | ECacheStoreFlags::Query | (bReadOnly ? ECacheStoreFlags::None : ECacheStoreFlags::Store));
-	}
-
-	/**
-	 * Creates a Zen structured data cache interface
-	 */
-	ILegacyCacheStore* ParseZenCache(const TCHAR* NodeName, const TCHAR* Entry)
-	{
-		FString ServiceUrl;
-		FParse::Value(Entry, TEXT("Host="), ServiceUrl);
-
-		FString Namespace;
-		if (!FParse::Value(Entry, TEXT("Namespace="), Namespace))
-		{
-			Namespace = FApp::GetProjectName();
-			UE_LOG(LogDerivedDataCache, Warning, TEXT("Node %s does not specify 'Namespace', falling back to '%s'"), NodeName, *Namespace);
-		}
-		FString StructuredNamespace;
-		if (!FParse::Value(Entry, TEXT("StructuredNamespace="), StructuredNamespace))
-		{
-			StructuredNamespace = Namespace;
-			UE_LOG(LogDerivedDataCache, Warning, TEXT("Node %s does not specify 'StructuredNamespace', falling back using Namespace option: '%s'"), NodeName, *StructuredNamespace);
-		}
-
-		if (ILegacyCacheStore* Backend = CreateZenCacheStore(NodeName, *ServiceUrl, *Namespace, *StructuredNamespace))
-		{
-			return Backend;
-		}
-		
-		UE_LOG(LogDerivedDataCache, Warning, TEXT("Zen backend is not yet supported in the current build configuration."));
-		return nullptr;
-	}
-
-	/**
 	 * Creates Boot data cache interface from ini settings.
 	 *
 	 * @param NodeName Node name.
diff --git a/Engine/Source/Developer/DerivedDataCache/Private/HttpCacheStore.cpp b/Engine/Source/Developer/DerivedDataCache/Private/HttpCacheStore.cpp
--- a/Engine/Source/Developer/DerivedDataCache/Private/HttpCacheStore.cpp
+++ b/Engine/Source/Developer/DerivedDataCache/Private/HttpCacheStore.cpp
@@ -1,54 +1,35 @@
 // Copyright Epic Games, Inc. All Rights Reserved.
 
-#include "DerivedDataBackendInterface.h"
-
 #if WITH_HTTP_DDC_BACKEND
 
-#include "Algo/Accumulate.h"
-#include "Algo/Find.h"
-#include "Algo/Transform.h"
 #include "Compression/CompressedBuffer.h"
-#include "Containers/DepletableMpscQueue.h"
-#include "Containers/StaticArray.h"
 #include "Containers/StringView.h"
 #include "Containers/Ticker.h"
+#include "DerivedDataBackendInterface.h"
 #include "DerivedDataCacheKey.h"
-#include "DerivedDataCachePrivate.h"
 #include "DerivedDataCacheRecord.h"
 #include "DerivedDataCacheUsageStats.h"
 #include "DerivedDataChunk.h"
-#include "DerivedDataRequest.h"
+#include "DerivedDataLegacyCacheStore.h"
 #include "DerivedDataRequestOwner.h"
 #include "DerivedDataValue.h"
 #include "Dom/JsonObject.h"
-#include "Experimental/Async/LazyEvent.h"
-#include "Experimental/Containers/FAAArrayQueue.h"
-#include "GenericPlatform/GenericPlatformFile.h"
 #include "HAL/IConsoleManager.h"
-#include "HAL/PlatformFileManager.h"
-#include "HAL/Thread.h"
+#include "Http/HttpClient.h"
 #include "IO/IoHash.h"
 #include "Memory/SharedBuffer.h"
-#include "Misc/CString.h"
+#include "Misc/App.h"
+#include "Misc/CommandLine.h"
+#include "Misc/ConfigCacheIni.h"
 #include "Misc/FileHelper.h"
-#include "Misc/Optional.h"
 #include "Misc/ScopeLock.h"
-#include "Misc/SecureHash.h"
 #include "Misc/StringBuilder.h"
-#include "Policies/CondensedJsonPrintPolicy.h"
 #include "ProfilingDebugging/CountersTrace.h"
 #include "ProfilingDebugging/CpuProfilerTrace.h"
-#include "Serialization/BufferArchive.h"
 #include "Serialization/CompactBinary.h"
 #include "Serialization/CompactBinaryPackage.h"
 #include "Serialization/CompactBinaryValidation.h"
 #include "Serialization/CompactBinaryWriter.h"
-#include "Serialization/JsonReader.h"
-#include "Serialization/JsonSerializer.h"
-#include "Serialization/JsonWriter.h"
-#include "Tasks/Task.h"
-
-#include "Http/HttpClient.h"
 
 #if PLATFORM_MICROSOFT
 #include "Microsoft/WindowsHWrapper.h"
@@ -62,29 +43,15 @@
 #include <netdb.h>
 #endif
 
-// Enables data request helpers that internally
-// batch requests to reduce the number of concurrent
-// connections.
-#ifndef WITH_DATAREQUEST_HELPER
-	#define WITH_DATAREQUEST_HELPER 1
-#endif
-
 #define UE_HTTPDDC_GET_REQUEST_POOL_SIZE 48
 #define UE_HTTPDDC_PUT_REQUEST_POOL_SIZE 16
 #define UE_HTTPDDC_NONBLOCKING_REQUEST_POOL_SIZE 128
 #define UE_HTTPDDC_MAX_FAILED_LOGIN_ATTEMPTS 16
 #define UE_HTTPDDC_MAX_ATTEMPTS 4
-#define UE_HTTPDDC_BATCH_SIZE 12
-#define UE_HTTPDDC_BATCH_NUM 64
-#define UE_HTTPDDC_BATCH_GET_WEIGHT 4
-#define UE_HTTPDDC_BATCH_HEAD_WEIGHT 1
-#define UE_HTTPDDC_BATCH_WEIGHT_HINT 12
 
 namespace UE::DerivedData
 {
 
-TRACE_DECLARE_INT_COUNTER(HttpDDC_Exist, TEXT("HttpDDC Exist"));
-TRACE_DECLARE_INT_COUNTER(HttpDDC_ExistHit, TEXT("HttpDDC Exist Hit"));
 TRACE_DECLARE_INT_COUNTER(HttpDDC_Get, TEXT("HttpDDC Get"));
 TRACE_DECLARE_INT_COUNTER(HttpDDC_GetHit, TEXT("HttpDDC Get Hit"));
 TRACE_DECLARE_INT_COUNTER(HttpDDC_Put, TEXT("HttpDDC Put"));
@@ -108,914 +75,11 @@ private:
 };
 
 
-//----------------------------------------------------------------------------------------------------------
-// Forward declarations
-//----------------------------------------------------------------------------------------------------------
-bool VerifyPayload(const FSHAHash& Hash, const TCHAR* Namespace, const TCHAR* Bucket, const TCHAR* CacheKey, const TArray64<uint8>& Payload);
-bool VerifyPayload(const FIoHash& Hash, const TCHAR* Namespace, const TCHAR* Bucket, const TCHAR* CacheKey, const TArray64<uint8>& Payload);
-bool VerifyRequest(const class FHttpRequest* Request, const TCHAR* Namespace, const TCHAR* Bucket, const TCHAR* CacheKey, const TArray64<uint8>& Payload);
-bool HashPayload(class FHttpRequest* Request, const TArrayView<const uint8> Payload);
-bool ShouldAbortForShutdown();
-
-
-#if WITH_DATAREQUEST_HELPER
-
-//----------------------------------------------------------------------------------------------------------
-// FDataRequestHelper
-//----------------------------------------------------------------------------------------------------------
-/**
- * Helper class for requesting data. Will batch requests once the number of concurrent requests reach a threshold.
- */
-struct FDataRequestHelper
-{
-	FDataRequestHelper(FHttpRequestPool* InPool, const TCHAR* InNamespace, const TCHAR* InBucket, const TCHAR* InCacheKey, TArray64<uint8>* OutData)
-		: Request(nullptr)
-		, Pool(InPool)
-		, bVerified(false, 1)
-	{
-		Request = Pool->GetFreeRequest();
-		if (Request && OutData != nullptr)
-		{
-			// We are below the threshold, make the connection immediately. OutData is set so this is a get.
-			FString Uri = FString::Printf(TEXT("api/v1/c/ddc/%s/%s/%s.raw"), InNamespace, InBucket, InCacheKey);
-			const FHttpRequest::EResult Result = Request->PerformBlockingDownload(*Uri, OutData);
-			if (FHttpRequest::IsSuccessResponse(Request->GetResponseCode()))
-			{
-				if (VerifyRequest(Request, InNamespace, InBucket, InCacheKey, *OutData))
-				{
-					TRACE_COUNTER_ADD(HttpDDC_GetHit, int64(1));
-					TRACE_COUNTER_ADD(HttpDDC_BytesReceived, int64(Request->GetBytesReceived()));
-					bVerified[0] = true;
-				}
-			}
-		}
-		else if (Request)
-		{
-			// We are below the threshold, make the connection immediately. OutData is missing so this is a head.
-			FString Uri = FString::Printf(TEXT("api/v1/c/ddc/%s/%s/%s"), InNamespace, InBucket, InCacheKey);
-			const FHttpRequest::EResult Result = Request->PerformBlockingHead(*Uri);
-			if (FHttpRequest::IsSuccessResponse(Request->GetResponseCode()))
-			{
-				TRACE_COUNTER_ADD(HttpDDC_ExistHit, int64(1));
-				bVerified[0] = true;
-			}
-		}
-		else
-		{
-			// We have exceeded the threshold for concurrent connections, start or add this request
-			// to a batched request.
-			if (IsQueueCandidate(1, OutData && !OutData->IsEmpty()))
-			{
-				Request = QueueBatchRequest(
-					InPool,
-					InNamespace,
-					InBucket,
-					TConstArrayView<const TCHAR*>({InCacheKey}),
-					OutData ? TConstArrayView<TArray64<uint8>*>({OutData}) : TConstArrayView<TArray64<uint8>*>(),
-					bVerified
-				);
-			}
-
-			if (!Request)
-			{
-				Request = Pool->WaitForFreeRequest();
-
-				FQueuedBatchEntry Entry{
-					InNamespace,
-					InBucket,
-					TConstArrayView<const TCHAR*>({InCacheKey}),
-					OutData ? TConstArrayView<TArray64<uint8>*>({OutData}) : TConstArrayView<TArray64<uint8>*>(),
-					OutData && !OutData->IsEmpty() ? FHttpRequest::ERequestVerb::Get : FHttpRequest::ERequestVerb::Head,
-					&bVerified
-				};
-
-				PerformBatchQuery(Request, TArrayView<FQueuedBatchEntry>(&Entry, 1));
-			}
-		}
-	}
-
-	// Constructor specifically for batched head queries
-	FDataRequestHelper(FHttpRequestPool* InPool, const TCHAR* InNamespace, const TCHAR* InBucket, TConstArrayView<FString> InCacheKeys)
-		: Request(nullptr)
-		, Pool(InPool)
-		, bVerified(false, InCacheKeys.Num())
-	{
-		// Transform the FString array to char pointers
-		TArray<const TCHAR*> CacheKeys;
-		Algo::Transform(InCacheKeys, CacheKeys, [](const FString& Key) { return *Key; });
-		
-		Request = Pool->GetFreeRequest();
-
-		if (Request || !IsQueueCandidate(InCacheKeys.Num(), false))
-		{
-			// If the request is too big for existing batches, wait for a free connection and create our own.
-			if (!Request)
-			{
-				Request = Pool->WaitForFreeRequest();
-			}
-
-			FQueuedBatchEntry Entry{
-				InNamespace, 
-				InBucket,
-				CacheKeys,
-				TConstArrayView<TArray64<uint8>*>(),
-				FHttpRequest::ERequestVerb::Head,
-				&bVerified
-			};
-
-			PerformBatchQuery(Request, TArrayView<FQueuedBatchEntry>(&Entry, 1));
-		}
-		else
-		{
-			Request = QueueBatchRequest(
-				InPool, 
-				InNamespace, 
-				InBucket, 
-				CacheKeys, 
-				TConstArrayView<TArray64<uint8>*>(), 
-				bVerified
-			);
-
-			if (!Request)
-			{
-				Request = Pool->WaitForFreeRequest();
-
-				FQueuedBatchEntry Entry{
-					InNamespace,
-					InBucket,
-					CacheKeys,
-					TConstArrayView<TArray64<uint8>*>(),
-					FHttpRequest::ERequestVerb::Head,
-					&bVerified
-				};
-
-				PerformBatchQuery(Request, TArrayView<FQueuedBatchEntry>(&Entry, 1));
-			}
-		}
-	}
-
-	~FDataRequestHelper()
-	{
-		if (Request)
-		{
-			Pool->ReleaseRequestToPool(Request);
-		}
-	}
-
-	static void StaticInitialize()
-	{
-		static bool bInitialized = false;
-		check(!bInitialized);
-		for (FBatch& Batch : Batches)
-		{
-			Batch.Reserved = 0;
-			Batch.Ready = 0;
-			Batch.Complete = TUniquePtr<FEvent, FBatch::FEventDeleter>(FPlatformProcess::GetSynchEventFromPool(true));
-		}
-		bInitialized = true;
-	}
-
-	static void StaticShutdown()
-	{
-		for (FBatch& Batch : Batches)
-		{
-			Batch.Complete.Reset();
-		}
-	}
-
-	bool IsSuccess() const
-	{
-		return bVerified[0];
-	}
-
-	const TBitArray<>& IsBatchSuccess() const
-	{
-		return bVerified;
-	}
-
-	int64 GetResponseCode() const
-	{
-		return Request ? Request->GetResponseCode() : 0;
-	}
-
-private:
-
-	struct FQueuedBatchEntry
-	{
-		const TCHAR* Namespace;
-		const TCHAR* Bucket;
-		TConstArrayView<const TCHAR*> CacheKeys;
-		TConstArrayView<TArray64<uint8>*> OutDatas;
-		FHttpRequest::ERequestVerb Verb;
-		TBitArray<>* bSuccess;
-	};
-
-	struct FBatch
-	{
-		struct FEventDeleter
-		{
-			void operator()(FEvent* Event)
-			{
-				FPlatformProcess::ReturnSynchEventToPool(Event);
-			}
-		};
-
-		FQueuedBatchEntry Entries[UE_HTTPDDC_BATCH_SIZE];
-		std::atomic<uint32> Reserved;
-		std::atomic<uint32> Ready;
-		std::atomic<uint32> WeightHint;
-		FHttpRequest* Request;
-		TUniquePtr<FEvent, FEventDeleter> Complete;
-	};
-
-	FHttpRequest* Request;
-	FHttpRequestPool* Pool;
-	TBitArray<> bVerified;
-	static std::atomic<uint32> FirstAvailableBatch;
-	static TStaticArray<FBatch, UE_HTTPDDC_BATCH_NUM> Batches;
-
-	static uint32 ComputeWeight(int32 NumKeys, bool bHasDatas)
-	{
-		return NumKeys * (bHasDatas ? UE_HTTPDDC_BATCH_GET_WEIGHT : UE_HTTPDDC_BATCH_HEAD_WEIGHT);
-	}
-
-	static bool IsQueueCandidate(int32 NumKeys, bool bHasDatas)
-	{
-		if (NumKeys > UE_HTTPDDC_BATCH_SIZE)
-		{
-			return false;
-		}
-		const uint32 Weight = ComputeWeight(NumKeys, bHasDatas);
-		if (Weight > UE_HTTPDDC_BATCH_WEIGHT_HINT)
-		{
-			return false;
-		}
-		return true;
-	}
-
-	/**
-	 * Queues up a request to be batched. Blocks until the query is made.
-	 */
-	static FHttpRequest* QueueBatchRequest(FHttpRequestPool* InPool, 
-		const TCHAR* InNamespace, 
-		const TCHAR* InBucket, 
-		TConstArrayView<const TCHAR*> InCacheKeys,
-		TConstArrayView<TArray64<uint8>*> OutDatas, 
-		TBitArray<>& bOutVerified)
-	{
-		TRACE_CPUPROFILER_EVENT_SCOPE(HttpDDC_BatchQuery);
-		check(InCacheKeys.Num() == OutDatas.Num() || OutDatas.Num() == 0);
-		const uint32 RequestNum = InCacheKeys.Num();
-		const uint32 RequestWeight = ComputeWeight(InCacheKeys.Num(), !OutDatas.IsEmpty());
-
-		for (int32 i = 0; i < Batches.Num(); i++)
-		{
-			uint32 Index = (FirstAvailableBatch.load(std::memory_order_relaxed) + i) % Batches.Num();
-			FBatch& Batch = Batches[Index];
-
-			//Assign different weights to head vs. get queries
-			if (Batch.WeightHint.load(std::memory_order_acquire) + RequestWeight > UE_HTTPDDC_BATCH_WEIGHT_HINT)
-			{
-				continue;
-			}
-
-			// Attempt to reserve a spot in the batch
-			const uint32 Reserve = Batch.Reserved.fetch_add(1, std::memory_order_acquire);
-			if (Reserve >= UE_HTTPDDC_BATCH_SIZE)
-			{
-				// We didn't manage to snag a valid reserve index try next batch
-				continue;
-			}
-
-			// Add our weight to the batch. Note we are treating it as a hint, so don't syncronize.
-			const uint32 ActualWeight = Batch.WeightHint.fetch_add(RequestWeight, std::memory_order_release);
-
-			TAnsiStringBuilder<64> BatchString;
-			BatchString << "HttpDDC_Batch" << Index;
-			TRACE_CPUPROFILER_EVENT_SCOPE_TEXT(*BatchString);
-
-			if (Reserve == (UE_HTTPDDC_BATCH_SIZE - 1))
-			{
-				FirstAvailableBatch++;
-			}
-
-			Batch.Entries[Reserve] = FQueuedBatchEntry{
-				InNamespace,
-				InBucket,
-				InCacheKeys,
-				OutDatas,
-				OutDatas.Num() ? FHttpRequest::ERequestVerb::Get : FHttpRequest::ERequestVerb::Head,
-				&bOutVerified
-			};
-
-			// Signal we are ready for batch to be submitted
-			Batch.Ready.fetch_add(1u, std::memory_order_release);
-
-			FHttpRequest* Request = nullptr;
-
-			// The first to reserve a slot is the "driver" of the batch
-			if (Reserve == 0)
-			{
-				Batch.Request = InPool->WaitForFreeRequest();
-
-				// Make sure no new requests are added
-				const uint32 Reserved = FMath::Min((uint32)UE_HTTPDDC_BATCH_SIZE, Batch.Reserved.fetch_add(UE_HTTPDDC_BATCH_SIZE, std::memory_order_acquire));
-
-				// Give other threads time to copy their data to batch
-				while (Batch.Ready.load(std::memory_order_acquire) < Reserved)
-				{
-				}
-
-				// Increment request ref count to reflect all waiting threads
-				InPool->MakeRequestShared(Batch.Request, Reserved);
-
-				// Do the actual query and write response to respective target arrays
-				PerformBatchQuery(Batch.Request, TArrayView<FQueuedBatchEntry>(Batch.Entries, Batch.Ready));
-
-				// Signal to waiting threads the batch is complete
-				Batch.Complete->Trigger();
-
-				// Store away the request and wait until other threads have too
-				Request = Batch.Request;
-				while (Batch.Ready.load(std::memory_order_acquire) > 1)
-				{
-				}
-
-				//Reset batch for next use
-				Batch.Complete->Reset();
-				Batch.WeightHint.store(0, std::memory_order_release);
-				Batch.Ready.store(0, std::memory_order_release);
-				Batch.Reserved.store(0, std::memory_order_release);
-			}
-			else
-			{
-				// Wait until "driver" has done query
-				{
-					TRACE_CPUPROFILER_EVENT_SCOPE(HttpDDC_WaitForMasterOfBatch);
-					Batch.Complete->Wait(~0);
-				}
-
-				// Store away request and signal we are done
-				Request = Batch.Request;
-				Batch.Ready.fetch_sub(1u, std::memory_order_release);
-			}
-
-			return Request;
-		}
-
-		return nullptr;
-	}
-
-
-	/**
-	 * Creates request uri and headers and submits the request
-	 */
-	static void PerformBatchQuery(FHttpRequest* Request, TArrayView<FQueuedBatchEntry> Entries)
-	{
-		TRACE_CPUPROFILER_EVENT_SCOPE(HttpDDC_BatchGet);
-		const TCHAR* Uri(TEXT("api/v1/c/ddc-rpc/batchget"));
-		int64 ResponseCode = 0; uint32 Attempts = 0;
-
-		//Prepare request object
-		TArray<TSharedPtr<FJsonValue>> Operations;
-		for (const FQueuedBatchEntry& Entry : Entries)
-		{
-			for (int32 KeyIdx = 0; KeyIdx < Entry.CacheKeys.Num(); KeyIdx++)
-			{
-				TSharedPtr<FJsonObject> Object = MakeShared<FJsonObject>();
-				Object->SetField(TEXT("bucket"), MakeShared<FJsonValueString>(Entry.Bucket));
-				Object->SetField(TEXT("key"), MakeShared<FJsonValueString>(Entry.CacheKeys[KeyIdx]));
-				if (Entry.Verb == FHttpRequest::ERequestVerb::Head)
-				{
-					Object->SetField(TEXT("verb"), MakeShared<FJsonValueString>(TEXT("HEAD")));
-				}
-				Operations.Add(MakeShared<FJsonValueObject>(Object));
-			}
-		}
-		TSharedPtr<FJsonObject> RequestObject = MakeShared<FJsonObject>();
-		RequestObject->SetField(TEXT("namespace"), MakeShared<FJsonValueString>(Entries[0].Namespace));
-		RequestObject->SetField(TEXT("operations"), MakeShared<FJsonValueArray>(Operations));
-
-		//Serialize to a buffer
-		FBufferArchive RequestData;
-		if (FJsonSerializer::Serialize(RequestObject.ToSharedRef(), TJsonWriterFactory<ANSICHAR, TCondensedJsonPrintPolicy<ANSICHAR>>::Create(&RequestData)))
-		{
-			Request->PerformBlockingPost(Uri, FCompositeBuffer(FSharedBuffer::MakeView(RequestData.GetData(), RequestData.Num())), EHttpContentType::JSON);
-			ResponseCode = Request->GetResponseCode();
-
-			if (ResponseCode == 200)
-			{
-				const TArray64<uint8>& ResponseBuffer = Request->GetResponseBuffer();
-				const uint8* Response = ResponseBuffer.GetData();
-				const int32 ResponseSize = ResponseBuffer.Num();
-
-				// Parse the response and move the data to the target requests.
-				if (ParseBatchedResponse(Response, ResponseSize, Entries))
-				{
-					UE_LOG(LogDerivedDataCache, VeryVerbose, TEXT("%s: Batch query with %d operations completed."), *Request->GetName(), Entries.Num());
-					return;
-				}
-			}
-		}
-		
-		// If we get here the request failed.
-		UE_LOG(LogDerivedDataCache, Display, TEXT("%s: Batch query failed. Query: %s"), *Request->GetName(), ANSI_TO_TCHAR((ANSICHAR*)RequestData.GetData()));
-
-		// Set all batch operations to failures
-		for (FQueuedBatchEntry Entry : Entries)
-		{
-			Entry.bSuccess->SetRange(0, Entry.CacheKeys.Num(), false);
-		}
-	}
-
-	// Above result value
-	enum class OpResult : uint8
-	{
-		Ok = 0,			// Op finished succesfully
-		Error = 1,		// Error during op
-		NotFound = 2,	// Key was not found
-		Exists = 3		// Used to indicate head op success
-	};
-
-	// Searches for potentially multiple key requests that are satisfied the given cache key result
-	// Search strategy is exhaustive forward search from the last found entry.  If the results come in ordered the same as the requests,
-	//  and there are no duplicates, the search will be somewhat efficient (still has to do exhaustive searching looking for duplicates).
-	//  If the results are unordered or there are duplicates, search will become more inefficient.
-	struct FRequestSearchHelper
-	{
-		FRequestSearchHelper(TArrayView<FQueuedBatchEntry> InRequests, const FUTF8ToTCHAR& InCacheKey, int32 InEntryIdx, int32 InKeyIdx, OpResult InRequestResult)
-			: Requests(InRequests)
-			, CacheKey(InCacheKey)
-			, StartEntryIdx(InEntryIdx)
-			, StartKeyIdx(InKeyIdx)
-			, RequestResult(InRequestResult)
-		{
-		}
-
-		bool FindNext(int32& EntryIdx, int32& KeyIdx)
-		{
-			int32 CurrentEntryIdx = EntryIdx;
-			int32 CurrentKeyIdx = KeyIdx;
-			do
-			{
-				// Do not match a get request with a head response code (i.e. Exists) 
-				// or a head request with a get response code (i.e. Ok)
-				// if the response code is an error or not found they can be matched to both head or get request it doesn't matter
-				const FQueuedBatchEntry& CurrentRequest = Requests[CurrentEntryIdx];
-				bool bRequestTypeMatch = !((CurrentRequest.Verb == FHttpRequest::ERequestVerb::Get) && (RequestResult == OpResult::Exists))
-					&& !((CurrentRequest.Verb == FHttpRequest::ERequestVerb::Head) && (RequestResult == OpResult::Ok));
-				if (bRequestTypeMatch && FCString::Stricmp(CurrentRequest.CacheKeys[CurrentKeyIdx], CacheKey.Get()) == 0)
-				{
-					EntryIdx = CurrentEntryIdx;
-					KeyIdx = CurrentKeyIdx;
-					return true;
-				}
-			}
-			while (AdvanceIndices(CurrentEntryIdx, CurrentKeyIdx));
-
-			return false;
-		}
-
-		bool AdvanceIndices(int32& EntryIdx, int32& KeyIdx)
-		{
-			if (++KeyIdx >= Requests[EntryIdx].CacheKeys.Num())
-			{
-				EntryIdx = (EntryIdx + 1) % Requests.Num();
-				KeyIdx = 0;
-			}
-
-			return !((EntryIdx == StartEntryIdx) && (KeyIdx == StartKeyIdx));
-		}
-
-		TArrayView<FQueuedBatchEntry> Requests;
-		const FUTF8ToTCHAR& CacheKey;
-		int32 StartEntryIdx;
-		int32 StartKeyIdx;
-		OpResult RequestResult;
-	};
-
-	/**
-	 * Parses a batched response stream, moves the data to target requests and marks them with result.
-	 * @param Response Pointer to Response buffer
-	 * @param ResponseSize Size of response buffer
-	 * @param Requests Requests that will be filled with data.
-	 * @return True if response was successfully parsed, false otherwise.
-	 */
-	static bool ParseBatchedResponse(const uint8* ResponseStart, const int32 ResponseSize, TArrayView<FQueuedBatchEntry> Requests)
-	{
-		// The expected data stream is structured accordingly
-		// {"JPTR"} {PayloadCount:uint32} {{"JPEE"} {Name:cstr} {Result:uint8} {Hash:IoHash} {Size:uint64} {Payload...}} ...
-
-		const auto& ResponseErrorMessage = TEXT("Malformed response from server.");
-		const ANSICHAR* ProtocolMagic = "JPTR";
-		const ANSICHAR* PayloadMagic = "JPEE";
-		const uint32 MagicSize = 4;
-		const uint8* Response = ResponseStart;
-		const uint8* ResponseEnd = Response + ResponseSize;
-
-		// Check that the stream starts with the protocol magic
-		if (FMemory::Memcmp(ProtocolMagic, Response, MagicSize) != 0)
-		{
-			UE_LOG(LogDerivedDataCache, Display, ResponseErrorMessage);
-			return false;
-		}
-		Response += MagicSize;
-
-		// Number of payloads recieved
-		uint32 PayloadCount = *(uint32*)Response;
-		Response += sizeof(uint32);
-
-		uint32 PayloadIdx = 0; 	// Current processed result
-		int32 EntryIdx = 0; 	// Current Entry index
-		int32 KeyIdx = 0; 		// Current Key index for current Entry
-
-		while (Response < ResponseEnd && FMemory::Memcmp(PayloadMagic, Response, MagicSize) == 0)
-		{
-			PayloadIdx++;
-			Response += MagicSize;
-
-			const ANSICHAR* PayloadNameA = (const ANSICHAR*)Response;
-			Response += FCStringAnsi::Strlen(PayloadNameA) + 1; //String and zero termination
-			const ANSICHAR* CacheKeyA = FCStringAnsi::Strrchr(PayloadNameA, '.') + 1; // "namespace.bucket.cachekey"
-
-			// Result of the operation is used to match to the appropriate request (i.e. get or head)
-			OpResult PayloadResult = static_cast<OpResult>(*Response);
-			Response += sizeof(uint8);
-
-			const uint8* ResponseRewindMark = Response;
-
-			// Find the payload among the requests.  Payloads may be returned in any order and if the same cache key was part of two requests,
-			// a single payload may satisfy multiple cache keys in multiple requests.
-			FUTF8ToTCHAR CacheKey(CacheKeyA);
-			FRequestSearchHelper RequestSearch(Requests, CacheKey, EntryIdx, KeyIdx, PayloadResult);
-			bool bFoundAny = false;
-
-			while (RequestSearch.FindNext(EntryIdx, KeyIdx))
-			{
-				Response = ResponseRewindMark;
-				bFoundAny = true;
-
-				FQueuedBatchEntry& RequestOp = Requests[EntryIdx];
-				TBitArray<>& bSuccess = *RequestOp.bSuccess;
-
-				switch (PayloadResult)
-				{
-				case OpResult::Ok:
-					{
-						// Payload hash of the following payload data
-						FIoHash PayloadHash = *(FIoHash*)Response;
-						Response += sizeof(FIoHash);
-
-						// Size of the following payload data
-						const uint64 PayloadSize = *(uint64*)Response;
-						Response += sizeof(uint64);
-
-						if (PayloadSize > 0)
-						{
-							if (Response + PayloadSize > ResponseEnd)
-							{
-								UE_LOG(LogDerivedDataCache, Display, ResponseErrorMessage);
-								return false;
-							}
-
-							if (bSuccess[KeyIdx])
-							{
-								Response += PayloadSize;
-							}
-							else
-							{
-								TArray64<uint8>* OutData = RequestOp.OutDatas[KeyIdx];
-
-								OutData->Append(Response, PayloadSize);
-								Response += PayloadSize;
-								// Verify the received and parsed payload
-								if (VerifyPayload(PayloadHash, RequestOp.Namespace, RequestOp.Bucket, RequestOp.CacheKeys[KeyIdx], *OutData))
-								{
-									TRACE_COUNTER_ADD(HttpDDC_GetHit, int64(1));
-									TRACE_COUNTER_ADD(HttpDDC_BytesReceived, int64(PayloadSize));
-									
-									bSuccess[KeyIdx] = true;
-								}
-								else
-								{
-									OutData->Empty();
-									bSuccess[KeyIdx] = false;
-								}
-							}
-						}
-						else
-						{
-							bSuccess[KeyIdx] = false;
-						}
-					}
-					break;
-
-				case OpResult::Exists:
-					{
-						TRACE_COUNTER_ADD(HttpDDC_ExistHit, int64(1));
-						bSuccess[KeyIdx] = true;
-					}
-					break;
-
-				default:
-				case OpResult::Error:
-					UE_LOG(LogDerivedDataCache, Display, TEXT("Server error while getting %s"), CacheKey.Get());
-					// intentional falltrough
-
-				case OpResult::NotFound:
-					bSuccess[KeyIdx] = false;
-					break;
-
-				}
-
-				if (!RequestSearch.AdvanceIndices(EntryIdx, KeyIdx))
-				{
-					break;
-				}
-			}
-
-			if (!bFoundAny)
-			{
-				UE_LOG(LogDerivedDataCache, Error, ResponseErrorMessage);
-				return false;
-			}
-		}
-
-		// Have we parsed all the payloads from the message?
-		if (PayloadIdx != PayloadCount)
-		{
-			UE_LOG(LogDerivedDataCache, Display, TEXT("%s: Found %d payloads but %d was reported."), ResponseErrorMessage, PayloadIdx, PayloadCount);
-		}
-
-		return true;
-	}
-};
-
-TStaticArray<FDataRequestHelper::FBatch, UE_HTTPDDC_BATCH_NUM> FDataRequestHelper::Batches;
-std::atomic<uint32> FDataRequestHelper::FirstAvailableBatch;
-
-//----------------------------------------------------------------------------------------------------------
-// FDataUploadHelper
-//----------------------------------------------------------------------------------------------------------
-struct FDataUploadHelper
-{
-	FDataUploadHelper(FHttpRequestPool* InPool, 
-		const TCHAR* InNamespace, 
-		const TCHAR* InBucket, 
-		const TCHAR* InCacheKey, 
-		const TArrayView<const uint8>& InData,
-		FDerivedDataCacheUsageStats& InUsageStats)
-		: ResponseCode(0)
-		, bSuccess(false)
-		, bQueued(false)
-	{
-		FHttpRequest* Request = InPool->GetFreeRequest();
-		if (Request)
-		{
-			ResponseCode = PerformPut(Request, InNamespace, InBucket, InCacheKey, InData, InUsageStats);
-			bSuccess = FHttpRequest::IsSuccessResponse(Request->GetResponseCode());
-
-			ProcessQueuedPutsAndReleaseRequest(InPool, Request, InUsageStats);
-		}
-		else
-		{
-			FQueuedEntry* Entry = new FQueuedEntry(InNamespace, InBucket, InCacheKey, InData);
-			QueuedPuts.Push(Entry);
-			bSuccess = true;
-			bQueued = true;
-			
-			// A request may have been released while the entry was being queued.
-			Request = InPool->GetFreeRequest();
-			if (Request)
-			{
-				ProcessQueuedPutsAndReleaseRequest(InPool, Request, InUsageStats);
-			}
-		}
-	}
-
-	bool IsSuccess() const
-	{
-		return bSuccess;
-	}
-
-	int64 GetResponseCode() const
-	{
-		return ResponseCode;
-	}
-
-	bool IsQueued() const
-	{
-		return bQueued;
-	}
-
-private:
-
-	struct FQueuedEntry
-	{
-		FString Namespace;
-		FString Bucket;
-		FString CacheKey;
-		TArray64<uint8> Data;
-
-		FQueuedEntry(const TCHAR* InNamespace, const TCHAR* InBucket, const TCHAR* InCacheKey, const TArrayView<const uint8> InData)
-			: Namespace(InNamespace)
-			, Bucket(InBucket)
-			, CacheKey(InCacheKey)
-			, Data(InData) // Copies the data!
-		{
-		}
-	};
-
-	static TLockFreePointerListUnordered<FQueuedEntry, PLATFORM_CACHE_LINE_SIZE> QueuedPuts;
-
-	int64 ResponseCode;
-	bool bSuccess;
-	bool bQueued;
-
-	static void ProcessQueuedPutsAndReleaseRequest(FHttpRequestPool* Pool, FHttpRequest* Request, FDerivedDataCacheUsageStats& UsageStats)
-	{
-		while (Request)
-		{
-			// Make sure that whether we early exit or execute past the end of this scope that
-			// the request is released back to the pool.
-			{
-				ON_SCOPE_EXIT
-				{
-					Pool->ReleaseRequestToPool(Request);
-				};
-
-				if (ShouldAbortForShutdown())
-				{
-					return;
-				}
-
-				while (FQueuedEntry* Entry = QueuedPuts.Pop())
-				{
-					Request->Reset();
-					PerformPut(Request, *Entry->Namespace, *Entry->Bucket, *Entry->CacheKey, Entry->Data, UsageStats);
-					delete Entry;
-
-					if (ShouldAbortForShutdown())
-					{
-						return;
-					}
-				}
-			}
-
-			// An entry may have been queued while the request was being released.
-			if (QueuedPuts.IsEmpty())
-			{
-				break;
-			}
-
-			// Process the queue again if a request is free, otherwise the thread that got the request will process it.
-			Request = Pool->GetFreeRequest();
-		}
-	}
-
-	static int64 PerformPut(FHttpRequest* Request, const TCHAR* Namespace, const TCHAR* Bucket, const TCHAR* CacheKey, const TArrayView<const uint8> Data, FDerivedDataCacheUsageStats& UsageStats)
-	{
-		COOK_STAT(auto Timer = UsageStats.TimePut());
-
-		HashPayload(Request, Data);
-
-		TStringBuilder<256> Uri;
-		Uri.Appendf(TEXT("api/v1/c/ddc/%s/%s/%s"), Namespace, Bucket, CacheKey);
-
-		Request->PerformBlockingPut(*Uri, FCompositeBuffer(FSharedBuffer::MakeView(Data.GetData(), Data.Num())), EHttpContentType::Binary);
-
-		const int64 ResponseCode = Request->GetResponseCode();
-		if (FHttpRequest::IsSuccessResponse(ResponseCode))
-		{
-			TRACE_COUNTER_ADD(HttpDDC_BytesSent, int64(Request->GetBytesSent()));
-			COOK_STAT(Timer.AddHit(Request->GetBytesSent()));
-		}
-
-		return Request->GetResponseCode();
-	}
-};
-
-TLockFreePointerListUnordered<FDataUploadHelper::FQueuedEntry, PLATFORM_CACHE_LINE_SIZE> FDataUploadHelper::QueuedPuts;
-
-#endif // WITH_DATAREQUEST_HELPER
-
-//----------------------------------------------------------------------------------------------------------
-// Content parsing and checking
-//----------------------------------------------------------------------------------------------------------
-
-/**
- * Verifies the integrity of the received data using supplied checksum.
- * @param Hash received hash value.
- * @param Namespace The namespace string used when originally fetching the request.
- * @param Bucket The bucket string used when originally fetching the request.
- * @param CacheKey The cache key string used when originally fetching the request.
- * @param Payload Payload received.
- * @return True if the data is correct, false if checksums doesn't match.
- */
-bool VerifyPayload(const FSHAHash& Hash, const TCHAR* Namespace, const TCHAR* Bucket, const TCHAR* CacheKey, const TArray64<uint8>& Payload)
-{
-	FSHAHash PayloadHash;
-	FSHA1::HashBuffer(Payload.GetData(), Payload.Num(), PayloadHash.Hash);
-
-	if (Hash != PayloadHash)
-	{
-		UE_LOG(LogDerivedDataCache,
-			Display,
-			TEXT("Checksum from server did not match received data (%s vs %s). Discarding cached result. Namespace: %s, Bucket: %s, Key: %s."),
-			*WriteToString<48>(Hash),
-			*WriteToString<48>(PayloadHash),
-			Namespace,
-			Bucket,
-			CacheKey
-		);
-		return false;
-	}
-
-	return true;
-}
-
-/**
- * Verifies the integrity of the received data using supplied checksum.
- * @param Hash received hash value.
- * @param Namespace The namespace string used when originally fetching the request.
- * @param Bucket The bucket string used when originally fetching the request.
- * @param CacheKey The cache key string used when originally fetching the request.
- * @param Payload Payload received.
- * @return True if the data is correct, false if checksums doesn't match.
- */
-bool VerifyPayload(const FIoHash& Hash, const TCHAR* Namespace, const TCHAR* Bucket, const TCHAR* CacheKey, const TArray64<uint8>& Payload)
-{
-	FIoHash PayloadHash = FIoHash::HashBuffer(Payload.GetData(), Payload.Num());
-
-	if (Hash != PayloadHash)
-	{
-		UE_LOG(LogDerivedDataCache,
-			Display,
-			TEXT("Checksum from server did not match received data (%s vs %s). Discarding cached result. Namespace: %s, Bucket: %s, Key: %s."),
-			*WriteToString<48>(Hash),
-			*WriteToString<48>(PayloadHash),
-			Namespace,
-			Bucket,
-			CacheKey
-		);
-		return false;
-	}
-
-	return true;
-}
-
-
-/**
- * Verifies the integrity of the received data using supplied checksum.
- * @param Request Request that the data was be received with.
- * @param Namespace The namespace string used when originally fetching the request.
- * @param Bucket The bucket string used when originally fetching the request.
- * @param CacheKey The cache key string used when originally fetching the request.
- * @param Payload Payload received.
- * @return True if the data is correct, false if checksums doesn't match.
- */
-bool VerifyRequest(const FHttpRequest* Request, const TCHAR* Namespace, const TCHAR* Bucket, const TCHAR* CacheKey, const TArray64<uint8>& Payload)
-{
-	FString ReceivedHashStr;
-	if (Request->GetHeader("X-Jupiter-Sha1", ReceivedHashStr))
-	{
-		FSHAHash ReceivedHash;
-		ReceivedHash.FromString(ReceivedHashStr);
-		return VerifyPayload(ReceivedHash, Namespace, Bucket, CacheKey, Payload);
-	}
-	if (Request->GetHeader("X-Jupiter-IoHash", ReceivedHashStr))
-	{
-		FIoHash ReceivedHash(ReceivedHashStr);
-		return VerifyPayload(ReceivedHash, Namespace, Bucket, CacheKey, Payload);
-	}
-	UE_LOG(LogDerivedDataCache, Warning, TEXT("%s: HTTP server did not send a content hash. Wrong server version?"), *Request->GetName());
-	return true;
-}
-
-/**
- * Adds a checksum (as request header) for a given payload. Jupiter will use this to verify the integrity
- * of the received data.
- * @param Request Request that the data will be sent with.
- * @param Payload Payload that will be sent.
- * @return True on success, false on failure.
- */
-bool HashPayload(FHttpRequest* Request, const TArrayView<const uint8> Payload)
-{
-	FIoHash PayloadHash = FIoHash::HashBuffer(Payload.GetData(), Payload.Num());
-	Request->AddHeader(TEXTVIEW("X-Jupiter-IoHash"), WriteToString<48>(PayloadHash));
-	return true;
-}
-
-bool ShouldAbortForShutdown()
+static bool ShouldAbortForShutdown()
 {
 	return !GIsBuildMachine && FDerivedDataBackend::Get().IsShuttingDown();
 }
 
-TConstArrayView<uint8> MakeConstArrayView(FSharedBuffer Buffer)
-{
-	return TConstArrayView<uint8>(reinterpret_cast<const uint8*>(Buffer.GetData()), Buffer.GetSize());
-}
-
 static bool IsValueDataReady(FValue& Value, const ECachePolicy Policy)
 {
 	if (!EnumHasAnyFlags(Policy, ECachePolicy::Query))
@@ -1035,92 +99,68 @@ static bool IsValueDataReady(FValue& Value, const ECachePolicy Policy)
 	return false;
 };
 
+struct FHttpCacheStoreParams
+{
+	FString Host;
+	FString Namespace;
+	FString StructuredNamespace;
+	FString OAuthProvider;
+	FString OAuthClientId;
+	FString OAuthSecret;
+	FString OAuthScope;
+	bool bResolveHostCanonicalName = false;
+	bool bReadOnly = false;
+
+	void Parse(const TCHAR* NodeName, const TCHAR* Config);
+};
+
 //----------------------------------------------------------------------------------------------------------
 // FHttpCacheStore
 //----------------------------------------------------------------------------------------------------------
 
 /**
  * Backend for a HTTP based caching service (Jupiter).
- **/
-class FHttpCacheStore final : public FDerivedDataBackendInterface
+ */
+class FHttpCacheStore final : public ILegacyCacheStore
 {
 public:
 	
 	/**
-	 * Creates the backend, checks health status and attempts to acquire an access token.
-	 *
-	 * @param ServiceUrl			Base url to the service including schema.
-	 * @param Namespace				Namespace to use.
-	 * @param StructuredNamespace	Namespace to use for structured cache operations.
-	 * @param OAuthProvider			Url to OAuth provider, for example "https://myprovider.com/oauth2/v1/token".
-	 * @param OAuthClientId			OAuth client identifier.
-	 * @param OAuthData				OAuth form data to send to login service. Can either be the raw form data or a Windows network file address (starting with "\\").
-	 * @param OAuthScope			OAuth scope identifier
+	 * Creates the cache store client, checks health status and attempts to acquire an access token.
 	 */
-	FHttpCacheStore(
-		const TCHAR* ServiceUrl, 
-		bool bResolveHostCanonicalName,
-		const TCHAR* Namespace, 
-		const TCHAR* StructuredNamespace, 
-		const TCHAR* OAuthProvider, 
-		const TCHAR* OAuthClientId, 
-		const TCHAR* OAuthData,
-		const TCHAR* OAuthScope,
-		EBackendLegacyMode LegacyMode,
-		bool bReadOnly);
+	explicit FHttpCacheStore(const FHttpCacheStoreParams& Params);
 
 	~FHttpCacheStore();
 
 	/**
-	 * Checks is backend is usable (reachable and accessible).
+	 * Checks is cache service is usable (reachable and accessible).
 	 * @return true if usable
 	 */
-	bool IsUsable() const { return bIsUsable; }
-
-	/** return true if this cache is writable **/
-	virtual bool IsWritable() const override
-	{
-		return !bReadOnly && bIsUsable;
-	}
-
-	virtual bool CachedDataProbablyExists(const TCHAR* CacheKey) override;
-	virtual TBitArray<> CachedDataProbablyExistsBatch(TConstArrayView<FString> CacheKeys) override;
-	virtual bool GetCachedData(const TCHAR* CacheKey, TArray<uint8>& OutData) override;
-	virtual EPutStatus PutCachedData(const TCHAR* CacheKey, TArrayView<const uint8> InData, bool bPutEvenIfExists) override;
-	virtual void RemoveCachedData(const TCHAR* CacheKey, bool bTransient) override;
-	virtual TSharedRef<FDerivedDataCacheStatsNode> GatherUsageStats() const override;
-	
-	virtual FString GetName() const override;
-	virtual TBitArray<> TryToPrefetch(TConstArrayView<FString> CacheKeys) override;
-	virtual bool WouldCache(const TCHAR* CacheKey, TArrayView<const uint8> InData) override;
-	virtual ESpeedClass GetSpeedClass() const override;
-	virtual bool ApplyDebugOptions(FBackendDebugOptions& InOptions) override;
-
-	void SetSpeedClass(ESpeedClass InSpeedClass) { SpeedClass = InSpeedClass; }
+	inline bool IsUsable() const { return bIsUsable; }
 
-	EBackendLegacyMode GetLegacyMode() const final { return LegacyMode; }
-
-	virtual void Put(
+	void Put(
 		TConstArrayView<FCachePutRequest> Requests,
 		IRequestOwner& Owner,
-		FOnCachePutComplete&& OnComplete) override;
-
-	virtual void Get(
+		FOnCachePutComplete&& OnComplete) final;
+	void Get(
 		TConstArrayView<FCacheGetRequest> Requests,
 		IRequestOwner& Owner,
-		FOnCacheGetComplete&& OnComplete) override;
-	virtual void PutValue(
+		FOnCacheGetComplete&& OnComplete) final;
+	void PutValue(
 		TConstArrayView<FCachePutValueRequest> Requests,
 		IRequestOwner& Owner,
-		FOnCachePutValueComplete&& OnComplete) override;
-	virtual void GetValue(
+		FOnCachePutValueComplete&& OnComplete) final;
+	void GetValue(
 		TConstArrayView<FCacheGetValueRequest> Requests,
 		IRequestOwner& Owner,
-		FOnCacheGetValueComplete&& OnComplete) override;
-	virtual void GetChunks(
+		FOnCacheGetValueComplete&& OnComplete) final;
+	void GetChunks(
 		TConstArrayView<FCacheGetChunkRequest> Requests,
 		IRequestOwner& Owner,
-		FOnCacheGetChunkComplete&& OnComplete) override;
+		FOnCacheGetChunkComplete&& OnComplete) final;
+
+	void LegacyStats(FDerivedDataCacheStatsNode& OutNode) final;
+	bool LegacyDebugOptions(FBackendDebugOptions& Options) final;
 
 	static FHttpCacheStore* GetAny()
 	{
@@ -1156,8 +196,6 @@ private:
 	bool bIsUsable;
 	bool bReadOnly;
 	uint32 FailedLoginAttempts;
-	ESpeedClass SpeedClass;
-	EBackendLegacyMode LegacyMode;
 	static inline FHttpCacheStore* AnyInstance = nullptr;
 
 	bool IsServiceReady();
@@ -1597,7 +635,7 @@ void FHttpCacheStore::FPutPackageOp::OnPackagePutRefComplete(
 		if (Response.Status == EStatus::Error)
 		{
 			UE_LOG(LogDerivedDataCache, Log, TEXT("%s: Failed to put reference object for put of %s from '%s'"),
-				*CacheStore.GetName(), *WriteToString<96>(Response.Key), *Response.Name);
+				*CacheStore.Domain, *WriteToString<96>(Response.Key), *Response.Name);
 		}
 		return OnComplete(FCachePutPackageResponse{ Name, Key, UserData, Response.BytesSent, Response.Status });
 	}
@@ -1655,7 +693,7 @@ void FHttpCacheStore::FPutPackageOp::OnPackagePutRefComplete(
 				bExpectedHashesSerialized = true;
 			}
 			UE_LOG(LogDerivedDataCache, Log, TEXT("%s: Server reported needed hash '%s' that is outside the set of expected hashes (%s) for put of %s from '%s'"),
-				*CacheStore.GetName(), *WriteToString<96>(NeededBlobHash), ExpectedHashes.ToString(), *WriteToString<96>(Response.Key), *Response.Name);
+				*CacheStore.Domain, *WriteToString<96>(NeededBlobHash), ExpectedHashes.ToString(), *WriteToString<96>(Response.Key), *Response.Name);
 		}
 	}
 
@@ -1736,7 +774,7 @@ FHttpRequest::ECompletionBehavior FHttpCacheStore::FPutPackageOp::OnCompressedBl
 		{
 			uint32 FailedBlobUploads = (uint32)(TotalBlobUploads - LocalSuccessfulBlobUploads);
 			UE_LOG(LogDerivedDataCache, Log, TEXT("%s: Failed to put %d/%d blobs for put of %s from '%s'"),
-				*CacheStore.GetName(), FailedBlobUploads, TotalBlobUploads, *WriteToString<96>(Key), *Name);
+				*CacheStore.Domain, FailedBlobUploads, TotalBlobUploads, *WriteToString<96>(Key), *Name);
 			OnComplete(MakeResponse(BytesSent.load(std::memory_order_relaxed), EStatus::Error));
 		}
 	}
@@ -1751,7 +789,7 @@ void FHttpCacheStore::FPutPackageOp::OnPutRefFinalizationComplete(
 	if (Response.Status == EStatus::Error)
 	{
 		UE_LOG(LogDerivedDataCache, Log, TEXT("%s: Failed to finalize reference object for put of %s from '%s'"),
-			*CacheStore.GetName(), *WriteToString<96>(Key), *Name);
+			*CacheStore.Domain, *WriteToString<96>(Key), *Name);
 	}
 
 	return OnComplete(MakeResponse(BytesSent.load(std::memory_order_relaxed), Response.Status));
@@ -1842,7 +880,7 @@ void FHttpCacheStore::FGetRecordOp::GetDataBatch(
 			{
 				UE_LOG(LogDerivedDataCache, Verbose,
 					TEXT("%s: Cache miss with missing value %s with hash %s for %s from '%s'"),
-					*CacheStore.GetName(), *ValueIdGetter(Value), *WriteToString<48>(Value.GetRawHash()), *WriteToString<96>(Key),
+					*CacheStore.Domain, *ValueIdGetter(Value), *WriteToString<48>(Value.GetRawHash()), *WriteToString<96>(Key),
 					*Name);
 				OnCompletePtr->GetFunction()({ Name, Key, ValueIndex, Request->GetBytesReceived(), {}, EStatus::Error });
 			}
@@ -1850,7 +888,7 @@ void FHttpCacheStore::FGetRecordOp::GetDataBatch(
 			{
 				UE_LOG(LogDerivedDataCache, Display,
 					TEXT("%s: Cache miss with corrupted value %s with hash %s for %s from '%s'"),
-					*CacheStore.GetName(), *ValueIdGetter(Value), *WriteToString<48>(Value.GetRawHash()),
+					*CacheStore.Domain, *ValueIdGetter(Value), *WriteToString<48>(Value.GetRawHash()),
 					*WriteToString<96>(Key), *Name);
 				OnCompletePtr->GetFunction()({ Name, Key, ValueIndex, Request->GetBytesReceived(), {}, EStatus::Error });
 			}
@@ -2024,7 +1062,7 @@ void FHttpCacheStore::FGetRecordOp::DataProbablyExistsBatch(
 							const FValueWithId& Value = Values[ValueIndex];
 							UE_LOG(LogDerivedDataCache, VeryVerbose,
 								TEXT("%s: Cache exists hit for %s with hash %s for %s from '%s'"),
-								*CacheStore.GetName(), *WriteToString<16>(Value.GetId()), *WriteToString<48>(Value.GetRawHash()), *WriteToString<96>(Key),
+								*CacheStore.Domain, *WriteToString<16>(Value.GetId()), *WriteToString<48>(Value.GetRawHash()), *WriteToString<96>(Key),
 								*Name);
 							InOnComplete({ Name, Key, ValueIndex, EStatus::Ok });
 						}
@@ -2055,7 +1093,7 @@ void FHttpCacheStore::FGetRecordOp::DataProbablyExistsBatch(
 					{
 						UE_LOG(LogDerivedDataCache, VeryVerbose,
 							TEXT("%s: Cache exists hit for %s with hash %s for %s from '%s'"),
-							*CacheStore.GetName(), *WriteToString<16>(Value.GetId()), *WriteToString<48>(Value.GetRawHash()), *WriteToString<96>(Key),
+							*CacheStore.Domain, *WriteToString<16>(Value.GetId()), *WriteToString<48>(Value.GetRawHash()), *WriteToString<96>(Key),
 							*Name);
 						InOnComplete({ Name, Key, ValueIndex, EStatus::Ok });
 					}
@@ -2063,7 +1101,7 @@ void FHttpCacheStore::FGetRecordOp::DataProbablyExistsBatch(
 					{
 						UE_LOG(LogDerivedDataCache, Verbose,
 							TEXT("%s: Cache exists miss with missing value %s with hash %s for %s from '%s'"),
-							*CacheStore.GetName(), *WriteToString<16>(Value.GetId()), *WriteToString<48>(Value.GetRawHash()), *WriteToString<96>(Key),
+							*CacheStore.Domain, *WriteToString<16>(Value.GetId()), *WriteToString<48>(Value.GetRawHash()), *WriteToString<96>(Key),
 							*Name);
 						InOnComplete({ Name, Key, ValueIndex, EStatus::Error });
 					}
@@ -2075,7 +1113,7 @@ void FHttpCacheStore::FGetRecordOp::DataProbablyExistsBatch(
 				{
 					UE_LOG(LogDerivedDataCache, Log,
 						TEXT("%s: Cache exists returned invalid results."),
-						*CacheStore.GetName());
+						*CacheStore.Domain);
 					InOnComplete({ Name, Key, ValueIndex, EStatus::Error });
 				}
 			}
@@ -2092,7 +1130,7 @@ void FHttpCacheStore::FGetRecordOp::DataProbablyExistsBatch(
 		{
 			const FValueWithId& Value = Values[ValueIndex];
 			UE_LOG(LogDerivedDataCache, Verbose, TEXT("%s: Cache miss with failed HTTP request for %s from '%s'"),
-				*CacheStore.GetName(), *WriteToString<96>(Key), *Name);
+				*CacheStore.Domain, *WriteToString<96>(Key), *Name);
 			InOnComplete({Name, Key, ValueIndex, EStatus::Error});
 		}
 		return FHttpRequest::ECompletionBehavior::Done;
@@ -2130,36 +1168,21 @@ void FHttpCacheStore::FGetRecordOp::FinishDataStep(bool bSuccess, uint64 InBytes
 	}
 }
 
-FHttpCacheStore::FHttpCacheStore(
-	const TCHAR* InServiceUrl, 
-	bool bResolveHostCanonicalName,
-	const TCHAR* InNamespace, 
-	const TCHAR* InStructuredNamespace, 
-	const TCHAR* InOAuthProvider,
-	const TCHAR* InOAuthClientId,
-	const TCHAR* InOAuthSecret,
-	const TCHAR* InOAuthScope,
-	const EBackendLegacyMode InLegacyMode,
-	const bool bInReadOnly)
-	: Domain(InServiceUrl)
-	, EffectiveDomain(InServiceUrl)
-	, Namespace(InNamespace)
-	, StructuredNamespace(InStructuredNamespace)
+FHttpCacheStore::FHttpCacheStore(const FHttpCacheStoreParams& Params)
+	: Domain(Params.Host)
+	, EffectiveDomain(Params.Host)
+	, Namespace(Params.Namespace)
+	, StructuredNamespace(Params.StructuredNamespace)
 	, DefaultBucket(TEXT("default"))
-	, OAuthProvider(InOAuthProvider)
-	, OAuthClientId(InOAuthClientId)
-	, OAuthSecret(InOAuthSecret)
-	, OAuthScope(InOAuthScope)
+	, OAuthProvider(Params.OAuthProvider)
+	, OAuthClientId(Params.OAuthClientId)
+	, OAuthSecret(Params.OAuthSecret)
+	, OAuthScope(Params.OAuthScope)
 	, Access(nullptr)
 	, bIsUsable(false)
-	, bReadOnly(bInReadOnly)
+	, bReadOnly(Params.bReadOnly)
 	, FailedLoginAttempts(0)
-	, SpeedClass(ESpeedClass::Slow)
-	, LegacyMode(InLegacyMode)
 {
-#if WITH_DATAREQUEST_HELPER
-	FDataRequestHelper::StaticInitialize();
-#endif
 	SharedData = MakeUnique<FHttpSharedData>();
 	if (IsServiceReady() && AcquireAccessToken())
 	{
@@ -2186,7 +1209,7 @@ FHttpCacheStore::FHttpCacheStore(
 		FMemory::Memset(&AddrHints, 0, sizeof(AddrHints));
 		AddrHints.ai_flags = AI_CANONNAME;
 		AddrHints.ai_family = AF_UNSPEC;
-		if (bResolveHostCanonicalName && !::getaddrinfo(*DomainResolveName, nullptr, &AddrHints, &AddrResult))
+		if (Params.bResolveHostCanonicalName && !::getaddrinfo(*DomainResolveName, nullptr, &AddrHints, &AddrResult))
 		{
 			if (AddrResult->ai_canonname)
 			{
@@ -2214,51 +1237,27 @@ FHttpCacheStore::FHttpCacheStore(
 		PutRequestPools[0] = MakeUnique<FHttpRequestPool>(*Domain, *EffectiveDomain, Access.Get(), SharedData.Get(), UE_HTTPDDC_PUT_REQUEST_POOL_SIZE);
 		PutRequestPools[1] = MakeUnique<FHttpRequestPool>(*Domain, *EffectiveDomain, Access.Get(), SharedData.Get(), UE_HTTPDDC_PUT_REQUEST_POOL_SIZE);
 		// Allowing the non-blocking requests to overflow to double their pre-allocated size before we start waiting for one to free up.
-		NonBlockingRequestPools = MakeUnique<FHttpRequestPool>(*Domain, *EffectiveDomain, Access.Get(), SharedData.Get(), UE_HTTPDDC_NONBLOCKING_REQUEST_POOL_SIZE, UE_HTTPDDC_NONBLOCKING_REQUEST_POOL_SIZE);
-		bIsUsable = true;
-	}
-
-	AnyInstance = this;
-}
-
-FHttpCacheStore::~FHttpCacheStore()
-{
-	if (AnyInstance == this)
-	{
-		AnyInstance = nullptr;
-	}
-#if WITH_DATAREQUEST_HELPER
-	FDataRequestHelper::StaticShutdown();
-#endif
-}
-
-FString FHttpCacheStore::GetName() const
-{
-	return Domain;
-}
-
-TBitArray<> FHttpCacheStore::TryToPrefetch(TConstArrayView<FString> CacheKeys)
-{
-	return CachedDataProbablyExistsBatch(CacheKeys);
-}
+		NonBlockingRequestPools = MakeUnique<FHttpRequestPool>(*Domain, *EffectiveDomain, Access.Get(), SharedData.Get(), UE_HTTPDDC_NONBLOCKING_REQUEST_POOL_SIZE, UE_HTTPDDC_NONBLOCKING_REQUEST_POOL_SIZE);
+		bIsUsable = true;
+	}
 
-bool FHttpCacheStore::WouldCache(const TCHAR* CacheKey, TArrayView<const uint8> InData)
-{
-	return IsWritable();
+	AnyInstance = this;
 }
 
-FHttpCacheStore::ESpeedClass FHttpCacheStore::GetSpeedClass() const
+FHttpCacheStore::~FHttpCacheStore()
 {
-	return SpeedClass;
+	if (AnyInstance == this)
+	{
+		AnyInstance = nullptr;
+	}
 }
 
-bool FHttpCacheStore::ApplyDebugOptions(FBackendDebugOptions& InOptions)
+bool FHttpCacheStore::LegacyDebugOptions(FBackendDebugOptions& InOptions)
 {
 	DebugOptions = InOptions;
 	return true;
 }
 
-
 bool FHttpCacheStore::IsServiceReady()
 {
 	FHttpRequest Request(*Domain, *Domain, nullptr, SharedData.Get(), false);
@@ -2442,24 +1441,22 @@ void FHttpCacheStore::GetCacheRecordOnlyAsync(
 	{
 		UE_LOG(LogDerivedDataCache, VeryVerbose,
 			TEXT("%s: Skipped get of %s from '%s' because this cache store is not available"),
-			*GetName(), *WriteToString<96>(Key), *Name);
+			*Domain, *WriteToString<96>(Key), *Name);
 		return OnComplete(MakeResponse(0, EStatus::Error));
 	}
 
 	// Skip the request if querying the cache is disabled.
-	const ECachePolicy QueryPolicy = SpeedClass == ESpeedClass::Local
-		? ECachePolicy::QueryLocal : ECachePolicy::QueryRemote;
-	if (!EnumHasAnyFlags(Policy.GetRecordPolicy(), QueryPolicy))
+	if (!EnumHasAnyFlags(Policy.GetRecordPolicy(), ECachePolicy::QueryRemote))
 	{
 		UE_LOG(LogDerivedDataCache, VeryVerbose, TEXT("%s: Skipped get of %s from '%s' due to cache policy"),
-			*GetName(), *WriteToString<96>(Key), *Name);
+			*Domain, *WriteToString<96>(Key), *Name);
 		return OnComplete(MakeResponse(0, EStatus::Error));
 	}
 
 	if (DebugOptions.ShouldSimulateGetMiss(Key))
 	{
 		UE_LOG(LogDerivedDataCache, Verbose, TEXT("%s: Simulated miss for get of %s from '%s'"),
-			*GetName(), *WriteToString<96>(Key), *Name);
+			*Domain, *WriteToString<96>(Key), *Name);
 		return OnComplete(MakeResponse(0, EStatus::Error));
 	}
 
@@ -2485,7 +1482,7 @@ void FHttpCacheStore::GetCacheRecordOnlyAsync(
 			if (ValidateCompactBinary(ResponseBuffer, ECbValidateMode::Default) != ECbValidateError::None)
 			{
 				UE_LOG(LogDerivedDataCache, Log, TEXT("%s: Cache miss with invalid package for %s from '%s'"),
-					*GetName(), *WriteToString<96>(Key), *Name);
+					*Domain, *WriteToString<96>(Key), *Name);
 				OnComplete({ Name, Key, UserData, Request->GetBytesReceived(), {}, EStatus::Error });
 				return FHttpRequest::ECompletionBehavior::Done;
 			}
@@ -2494,7 +1491,7 @@ void FHttpCacheStore::GetCacheRecordOnlyAsync(
 			if (Record.IsNull())
 			{
 				UE_LOG(LogDerivedDataCache, Log, TEXT("%s: Cache miss with record load failure for %s from '%s'"),
-					*GetName(), *WriteToString<96>(Key), *Name);
+					*Domain, *WriteToString<96>(Key), *Name);
 				OnComplete({ Name, Key, UserData, Request->GetBytesReceived(), {}, EStatus::Error });
 				return FHttpRequest::ECompletionBehavior::Done;
 			}
@@ -2509,7 +1506,7 @@ void FHttpCacheStore::GetCacheRecordOnlyAsync(
 		}
 
 		UE_LOG(LogDerivedDataCache, Verbose, TEXT("%s: Cache miss with missing package for %s from '%s'"),
-			*GetName(), *WriteToString<96>(Key), *Name);
+			*Domain, *WriteToString<96>(Key), *Name);
 		OnComplete({ Name, Key, UserData, Request->GetBytesReceived(), {}, EStatus::Error });
 		return FHttpRequest::ECompletionBehavior::Done;
 	};
@@ -2530,28 +1527,27 @@ void FHttpCacheStore::PutCacheRecordAsync(
 		return FCachePutResponse{ Name, Key, UserData, Status };
 	};
 
-	if (!IsWritable())
+	if (bReadOnly)
 	{
 		UE_LOG(LogDerivedDataCache, VeryVerbose,
 			TEXT("%s: Skipped put of %s from '%s' because this cache store is read-only"),
-			*GetName(), *WriteToString<96>(Key), *Name);
+			*Domain, *WriteToString<96>(Key), *Name);
 		return OnComplete(MakeResponse(EStatus::Error), 0);
 	}
 
 	// Skip the request if storing to the cache is disabled.
 	const ECachePolicy RecordPolicy = Policy.GetRecordPolicy();
-	const ECachePolicy StoreFlag = SpeedClass == ESpeedClass::Local ? ECachePolicy::StoreLocal : ECachePolicy::StoreRemote;
-	if (!EnumHasAnyFlags(RecordPolicy, StoreFlag))
+	if (!EnumHasAnyFlags(RecordPolicy, ECachePolicy::StoreRemote))
 	{
 		UE_LOG(LogDerivedDataCache, VeryVerbose, TEXT("%s: Skipped put of %s from '%s' due to cache policy"),
-			*GetName(), *WriteToString<96>(Key), *Name);
+			*Domain, *WriteToString<96>(Key), *Name);
 		return OnComplete(MakeResponse(EStatus::Error), 0);
 	}
 
 	if (DebugOptions.ShouldSimulatePutMiss(Key))
 	{
 		UE_LOG(LogDerivedDataCache, Verbose, TEXT("%s: Simulated miss for put of %s from '%s'"),
-			*GetName(), *WriteToString<96>(Key), *Name);
+			*Domain, *WriteToString<96>(Key), *Name);
 		return OnComplete(MakeResponse(EStatus::Error), 0);
 	}
 
@@ -2583,27 +1579,26 @@ void FHttpCacheStore::PutCacheValueAsync(
 		return FCachePutValueResponse{ Name, Key, UserData, Status };
 	};
 
-	if (!IsWritable())
+	if (bReadOnly)
 	{
 		UE_LOG(LogDerivedDataCache, VeryVerbose,
 			TEXT("%s: Skipped put of %s from '%s' because this cache store is read-only"),
-			*GetName(), *WriteToString<96>(Key), *Name);
+			*Domain, *WriteToString<96>(Key), *Name);
 		return OnComplete(MakeResponse(EStatus::Error), 0);
 	}
 
 	// Skip the request if storing to the cache is disabled.
-	const ECachePolicy StoreFlag = SpeedClass == ESpeedClass::Local ? ECachePolicy::StoreLocal : ECachePolicy::StoreRemote;
-	if (!EnumHasAnyFlags(Policy, StoreFlag))
+	if (!EnumHasAnyFlags(Policy, ECachePolicy::StoreRemote))
 	{
 		UE_LOG(LogDerivedDataCache, VeryVerbose, TEXT("%s: Skipped put of %s from '%s' due to cache policy"),
-			*GetName(), *WriteToString<96>(Key), *Name);
+			*Domain, *WriteToString<96>(Key), *Name);
 		return OnComplete(MakeResponse(EStatus::Error), 0);
 	}
 
 	if (DebugOptions.ShouldSimulatePutMiss(Key))
 	{
 		UE_LOG(LogDerivedDataCache, Verbose, TEXT("%s: Simulated miss for put of %s from '%s'"),
-			*GetName(), *WriteToString<96>(Key), *Name);
+			*Domain, *WriteToString<96>(Key), *Name);
 		return OnComplete(MakeResponse(EStatus::Error), 0);
 	}
 
@@ -2640,17 +1635,16 @@ void FHttpCacheStore::GetCacheValueAsync(
 	{
 		UE_LOG(LogDerivedDataCache, VeryVerbose,
 			TEXT("%s: Skipped get of %s from '%s' because this cache store is not available"),
-			*GetName(), *WriteToString<96>(Key), *Name);
+			*Domain, *WriteToString<96>(Key), *Name);
 		OnComplete({Name, Key, {}, UserData, EStatus::Error});
 		return;
 	}
 
 	// Skip the request if querying the cache is disabled.
-	const ECachePolicy QueryFlag = SpeedClass == ESpeedClass::Local ? ECachePolicy::QueryLocal : ECachePolicy::QueryRemote;
-	if (!EnumHasAnyFlags(Policy, QueryFlag))
+	if (!EnumHasAnyFlags(Policy, ECachePolicy::QueryRemote))
 	{
 		UE_LOG(LogDerivedDataCache, VeryVerbose, TEXT("%s: Skipped get of %s from '%s' due to cache policy"),
-			*GetName(), *WriteToString<96>(Key), *Name);
+			*Domain, *WriteToString<96>(Key), *Name);
 		OnComplete({Name, Key, {}, UserData, EStatus::Error});
 		return;
 	}
@@ -2658,7 +1652,7 @@ void FHttpCacheStore::GetCacheValueAsync(
 	if (DebugOptions.ShouldSimulateGetMiss(Key))
 	{
 		UE_LOG(LogDerivedDataCache, Verbose, TEXT("%s: Simulated miss for get of %s from '%s'"),
-			*GetName(), *WriteToString<96>(Key), *Name);
+			*Domain, *WriteToString<96>(Key), *Name);
 		OnComplete({Name, Key, {}, UserData, EStatus::Error});
 		return;
 	}
@@ -2697,7 +1691,7 @@ void FHttpCacheStore::GetCacheValueAsync(
 				if (ValidateCompactBinary(ResponseBuffer, ECbValidateMode::Default) != ECbValidateError::None)
 				{
 					UE_LOG(LogDerivedDataCache, Display, TEXT("%s: Cache miss with invalid package for %s from '%s'"),
-						*GetName(), *WriteToString<96>(Key), *Name);
+						*Domain, *WriteToString<96>(Key), *Name);
 					OnComplete({Name, Key, {}, UserData, EStatus::Error});
 					return FHttpRequest::ECompletionBehavior::Done;
 				}
@@ -2708,7 +1702,7 @@ void FHttpCacheStore::GetCacheValueAsync(
 				if (RawHash.IsZero() || RawSize == MAX_uint64)
 				{
 					UE_LOG(LogDerivedDataCache, Display, TEXT("%s: Cache miss with invalid value for %s from '%'"),
-						*GetName(), *WriteToString<96>(Key), *Name);
+						*Domain, *WriteToString<96>(Key), *Name);
 					OnComplete({Name, Key, {}, UserData, EStatus::Error});
 					return FHttpRequest::ECompletionBehavior::Done;
 				}
@@ -2734,7 +1728,7 @@ void FHttpCacheStore::GetCacheValueAsync(
 				if (!CompressedBuffer)
 				{
 					UE_LOG(LogDerivedDataCache, Display, TEXT("%s: Cache miss with invalid package for %s from '%s'"),
-						*GetName(), *WriteToString<96>(Key), *Name);
+						*Domain, *WriteToString<96>(Key), *Name);
 					OnComplete({Name, Key, {}, UserData, EStatus::Error});
 					return FHttpRequest::ECompletionBehavior::Done;
 				}
@@ -2750,7 +1744,7 @@ void FHttpCacheStore::GetCacheValueAsync(
 		 }
 
 		UE_LOG(LogDerivedDataCache, Verbose, TEXT("%s: Cache miss with failed HTTP request for %s from '%s'"),
-			*GetName(), *WriteToString<96>(Key), *Name);
+			*Domain, *WriteToString<96>(Key), *Name);
 		OnComplete({Name, Key, {}, UserData, EStatus::Error});
 		return FHttpRequest::ECompletionBehavior::Done;
 	};
@@ -2785,7 +1779,7 @@ void FHttpCacheStore::RefCachedDataProbablyExistsBatchAsync(
 	{
 			UE_LOG(LogDerivedDataCache, VeryVerbose,
 				TEXT("%s: Skipped exists check of %s from '%s' because this cache store is not available"),
-				*GetName(), *WriteToString<96>(ValueRef.Key), *ValueRef.Name);
+				*Domain, *WriteToString<96>(ValueRef.Key), *ValueRef.Name);
 			OnComplete(ValueRef.MakeResponse(EStatus::Error));
 	}
 		return;
@@ -2835,7 +1829,7 @@ void FHttpCacheStore::RefCachedDataProbablyExistsBatchAsync(
 				{
 				UE_LOG(LogDerivedDataCache, Log,
 					TEXT("%s: Cache exists returned invalid results."),
-					*GetName());
+					*Domain);
 					OnComplete(ValueRef.MakeResponse(EStatus::Error));
 				}
 				return FHttpRequest::ECompletionBehavior::Done;
@@ -2851,7 +1845,7 @@ void FHttpCacheStore::RefCachedDataProbablyExistsBatchAsync(
 				{
 					UE_LOG(LogDerivedDataCache, Log,
 						TEXT("%s: Cache exists returned unexpected quantity of results (expected %d, got %d)."),
-						*GetName(), ValueRefs.Num(), ResultsArrayView.Num());
+						*Domain, ValueRefs.Num(), ResultsArrayView.Num());
 						OnComplete(ValueRef.MakeResponse(EStatus::Error));
 				}
 				return FHttpRequest::ECompletionBehavior::Done;
@@ -2867,7 +1861,7 @@ void FHttpCacheStore::RefCachedDataProbablyExistsBatchAsync(
 				if (OpId >= (uint32)ValueRefs.Num())
 				{
 					UE_LOG(LogDerivedDataCache, Display, TEXT("%s: Encountered invalid opId %d while querying %d values"),
-						*GetName(), OpId, ValueRefs.Num());
+						*Domain, OpId, ValueRefs.Num());
 					continue;
 				}
 
@@ -2876,16 +1870,15 @@ void FHttpCacheStore::RefCachedDataProbablyExistsBatchAsync(
 				if (!FHttpRequest::IsSuccessResponse(StatusCode))
 				{
 					UE_LOG(LogDerivedDataCache, Verbose, TEXT("%s: Cache miss with unsuccessful response code %d for %s from '%s'"),
-						*GetName(), StatusCode, *WriteToString<96>(ValueRef.Key), *ValueRef.Name);
+						*Domain, StatusCode, *WriteToString<96>(ValueRef.Key), *ValueRef.Name);
 					OnComplete(ValueRef.MakeResponse(EStatus::Error));
 					continue;
 				}
 
-				const ECachePolicy QueryFlag = SpeedClass == ESpeedClass::Local ? ECachePolicy::QueryLocal : ECachePolicy::QueryRemote;
-				if (!EnumHasAnyFlags(ValueRef.Policy, QueryFlag))
+				if (!EnumHasAnyFlags(ValueRef.Policy, ECachePolicy::QueryRemote))
 				{
 					UE_LOG(LogDerivedDataCache, VeryVerbose, TEXT("%s: Skipped exists check of %s from '%s' due to cache policy"),
-						*GetName(), *WriteToString<96>(ValueRef.Key), *ValueRef.Name);
+						*Domain, *WriteToString<96>(ValueRef.Key), *ValueRef.Name);
 					OnComplete(ValueRef.MakeResponse(EStatus::Error));
 					continue;
 				}
@@ -2895,7 +1888,7 @@ void FHttpCacheStore::RefCachedDataProbablyExistsBatchAsync(
 				if (RawHash.IsZero() || RawSize == MAX_uint64)
 				{
 					UE_LOG(LogDerivedDataCache, Display, TEXT("%s: Cache miss with invalid value for %s from '%s'"),
-						*GetName(), *WriteToString<96>(ValueRef.Key), *ValueRef.Name);
+						*Domain, *WriteToString<96>(ValueRef.Key), *ValueRef.Name);
 					OnComplete(ValueRef.MakeResponse(EStatus::Error));
 					continue;
 				}
@@ -2913,7 +1906,7 @@ void FHttpCacheStore::RefCachedDataProbablyExistsBatchAsync(
 		for (const FCacheGetValueRequest& ValueRef : ValueRefs)
 		{
 			UE_LOG(LogDerivedDataCache, Verbose, TEXT("%s: Cache miss with failed HTTP request for %s from '%s'"),
-				*GetName(), *WriteToString<96>(ValueRef.Key), *ValueRef.Name);
+				*Domain, *WriteToString<96>(ValueRef.Key), *ValueRef.Name);
 			OnComplete(ValueRef.MakeResponse(EStatus::Error));
 		}
 		return FHttpRequest::ECompletionBehavior::Done;
@@ -2922,364 +1915,10 @@ void FHttpCacheStore::RefCachedDataProbablyExistsBatchAsync(
 	Request->EnqueueAsyncPost(Owner, Pool, *RefsUri, FCompositeBuffer(RequestFields.GetOuterBuffer()), MoveTemp(OnHttpRequestComplete), EHttpContentType::CbObject);
 }
 
-bool FHttpCacheStore::CachedDataProbablyExists(const TCHAR* CacheKey)
-{
-	TRACE_CPUPROFILER_EVENT_SCOPE(HttpDDC_Exist);
-	TRACE_COUNTER_ADD(HttpDDC_Exist, int64(1));
-	COOK_STAT(auto Timer = UsageStats.TimeProbablyExists());
-
-	if (DebugOptions.ShouldSimulateGetMiss(CacheKey))
-	{
-		return false;
-	}
-
-#if WITH_DATAREQUEST_HELPER
-	// Retry request until we get an accepted response or exhaust allowed number of attempts.
-	for (int32 Attempts = 0; Attempts < UE_HTTPDDC_MAX_ATTEMPTS; ++Attempts)
-	{
-		FDataRequestHelper RequestHelper(GetRequestPools[IsInGameThread()].Get(), *Namespace, *DefaultBucket, CacheKey, nullptr);
-		const int64 ResponseCode = RequestHelper.GetResponseCode();
-
-		if (FHttpRequest::IsSuccessResponse(ResponseCode) && RequestHelper.IsSuccess())
-		{
-			COOK_STAT(Timer.AddHit(0));
-			return true;
-		}
-
-		if (!ShouldRetryOnError(ResponseCode))
-		{
-			return false;
-		}
-	}
-#else
-	FString Uri = FString::Printf(TEXT("api/v1/c/ddc/%s/%s/%s"), *Namespace, *DefaultBucket, CacheKey);
-
-	// Retry request until we get an accepted response or exhaust allowed number of attempts.
-	for (int32 Attempts = 0; Attempts < UE_HTTPDDC_MAX_ATTEMPTS; ++Attempts)
-	{
-		FScopedHttpPoolRequestPtr Request(GetRequestPools[IsInGameThread()].Get());
-		const FHttpRequest::EResult Result = Request->PerformBlockingQuery<FHttpRequest::Head>(*Uri);
-		const int64 ResponseCode = Request->GetResponseCode();
-
-		if (FHttpRequest::IsSuccessResponse(ResponseCode) || ResponseCode == 400)
-		{
-			const bool bIsHit = (Result == FHttpRequest::Success && FHttpRequest::IsSuccessResponse(ResponseCode));
-			if (bIsHit)
-			{
-				TRACE_COUNTER_ADD(HttpDDC_ExistHit, int64(1));
-				COOK_STAT(Timer.AddHit(0));
-			}
-			return bIsHit;
-		}
-
-		if (!ShouldRetryOnError(ResponseCode))
-		{
-			break;
-		}
-	}
-#endif
-
-	return false;
-}
-
-TBitArray<> FHttpCacheStore::CachedDataProbablyExistsBatch(TConstArrayView<FString> CacheKeys)
-{
-	TRACE_CPUPROFILER_EVENT_SCOPE(HttpDDC_Exist);
-	TRACE_COUNTER_ADD(HttpDDC_Exist, int64(1));
-	COOK_STAT(auto Timer = UsageStats.TimeProbablyExists());
-#if WITH_DATAREQUEST_HELPER
-	for (int32 Attempts = 0; Attempts < UE_HTTPDDC_MAX_ATTEMPTS; ++Attempts)
-	{
-		FDataRequestHelper RequestHelper(GetRequestPools[IsInGameThread()].Get(), *Namespace, *DefaultBucket, CacheKeys);
-		const int64 ResponseCode = RequestHelper.GetResponseCode();
-
-		if (FHttpRequest::IsSuccessResponse(ResponseCode) && RequestHelper.IsSuccess())
-		{
-			COOK_STAT(Timer.AddHit(0));
-			TBitArray<> Results = RequestHelper.IsBatchSuccess();
-			int32 ResultIndex = 0;
-			for (const FString& CacheKey : CacheKeys)
-			{
-				if (DebugOptions.ShouldSimulateGetMiss(*CacheKey))
-				{
-					Results[ResultIndex] = false;
-				}
-				ResultIndex++;
-			}
-
-			return Results;
-		}
-
-		if (!ShouldRetryOnError(ResponseCode))
-		{
-			TBitArray<> Results = RequestHelper.IsBatchSuccess();
-			int32 ResultIndex = 0;
-			for (const FString& CacheKey : CacheKeys)
-			{
-				if (DebugOptions.ShouldSimulateGetMiss(*CacheKey))
-				{
-					Results[ResultIndex] = false;
-				}
-				ResultIndex++;
-			}
-
-			return Results;
-		}
-	}
-#else
-	const TCHAR* const Uri = TEXT("api/v1/c/ddc-rpc");
-
-	TAnsiStringBuilder<512> Body;
-	const FTCHARToUTF8 AnsiNamespace(*Namespace);
-	const FTCHARToUTF8 AnsiBucket(*DefaultBucket);
-	Body << "{\"Operations\":[";
-	for (const FString& CacheKey : CacheKeys)
-	{
-		Body << "{\"Namespace\":\"" << AnsiNamespace.Get() << "\",\"Bucket\":\"" << AnsiBucket.Get() << "\",";
-		Body << "\"Id\":\"" << FTCHARToUTF8(*CacheKey).Get() << "\",\"Op\":\"HEAD\"},";
-	}
-	Body.RemoveSuffix(1);
-	Body << "]}";
-
-	TConstArrayView<uint8> BodyView(reinterpret_cast<const uint8*>(Body.ToString()), Body.Len());
-
-	// Retry request until we get an accepted response or exhaust allowed number of attempts.
-	for (int32 Attempts = 0; Attempts < UE_HTTPDDC_MAX_ATTEMPTS; ++Attempts)
-	{
-		FScopedHttpPoolRequestPtr Request(GetRequestPools[IsInGameThread()].Get());
-		const FHttpRequest::EResult Result = Request->PerformBlockingUpload<FHttpRequest::PostJson>(Uri, BodyView);
-		const int64 ResponseCode = Request->GetResponseCode();
-
-		if (Result == FHttpRequest::Success && ResponseCode == 200)
-		{
-			TArray<TSharedPtr<FJsonValue>> ResponseArray = Request->GetResponseAsJsonArray();
-
-			TBitArray<> Exists;
-			Exists.Reserve(CacheKeys.Num());
-			for (const FString& CacheKey : CacheKeys)
-			{
-				if (DebugOptions.ShouldSimulateGetMiss(*CacheKey))
-				{
-					Exists.Add(false);
-				}
-				else
-				{
-					const TSharedPtr<FJsonValue>* FoundResponse = Algo::FindByPredicate(ResponseArray, [&CacheKey](const TSharedPtr<FJsonValue>& Response) {
-						FString Key;
-						Response->TryGetString(Key);
-						return Key == CacheKey;
-					});
-
-					Exists.Add(FoundResponse != nullptr);
-				}
-			}
-
-			if (Exists.CountSetBits() == CacheKeys.Num())
-			{
-				TRACE_COUNTER_ADD(HttpDDC_ExistHit, int64(1));
-				COOK_STAT(Timer.AddHit(0));
-			}
-			return Exists;
-		}
-
-		if (!ShouldRetryOnError(ResponseCode))
-		{
-			break;
-		}
-	}
-#endif
-
-	return TBitArray<>(false, CacheKeys.Num());
-}
-
-bool FHttpCacheStore::GetCachedData(const TCHAR* CacheKey, TArray<uint8>& OutData)
-{
-	TRACE_CPUPROFILER_EVENT_SCOPE(HttpDDC_GetCachedData);
-	TRACE_COUNTER_ADD(HttpDDC_Get, int64(1));
-	COOK_STAT(auto Timer = UsageStats.TimeGet());
-
-	if (DebugOptions.ShouldSimulateGetMiss(CacheKey))
-	{
-		return false;
-	}
-
-	TArray64<uint8> ArrayBuffer;
-
-#if WITH_DATAREQUEST_HELPER
-	// Retry request until we get an accepted response or exhaust allowed number of attempts.
-	for (int32 Attempts = 0; Attempts < UE_HTTPDDC_MAX_ATTEMPTS; ++Attempts)
-	{
-		FDataRequestHelper RequestHelper(GetRequestPools[IsInGameThread()].Get(), *Namespace, *DefaultBucket, CacheKey, &ArrayBuffer);
-		OutData = TArray<uint8>(MoveTemp(ArrayBuffer));
-		const int64 ResponseCode = RequestHelper.GetResponseCode();
-
-		if (FHttpRequest::IsSuccessResponse(ResponseCode) && RequestHelper.IsSuccess())
-		{
-			COOK_STAT(Timer.AddHit(OutData.Num()));
-			check(OutData.Num() > 0);
-			return true;
-		}
-
-		if (!ShouldRetryOnError(ResponseCode))
-		{
-			return false;
-		}
-	}
-#else 
-	FString Uri = FString::Printf(TEXT("api/v1/c/ddc/%s/%s/%s.raw"), *Namespace, *DefaultBucket, CacheKey);
-
-	// Retry request until we get an accepted response or exhaust allowed number of attempts.
-	for (int32 Attempts = 0; Attempts < UE_HTTPDDC_MAX_ATTEMPTS; ++Attempts)
-	{
-		FScopedHttpPoolRequestPtr Request(GetRequestPools[IsInGameThread()].Get());
-		if (Request.IsValid())
-		{
-			FHttpRequest::EResult Result = Request->PerformBlockingDownload(*Uri, &ArrayBuffer);
-			const uint64 ResponseCode = Request->GetResponseCode();
-
-			// Request was successful, make sure we got all the expected data.
-			if (FHttpRequest::IsSuccessResponse(ResponseCode) && VerifyRequest(Request.Get(), *Namespace, *DefaultBucket, CacheKey, ArrayBuffer))
-			{
-				OutData = TArray<uint8>(MoveTemp(ArrayBuffer));
-				TRACE_COUNTER_ADD(HttpDDC_GetHit, int64(1));
-				TRACE_COUNTER_ADD(HttpDDC_BytesReceived, int64(Request->GetBytesReceived()));
-				COOK_STAT(Timer.AddHit(Request->GetBytesReceived()));
-				return true;
-			}
-			OutData = TArray<uint8>(MoveTemp(ArrayBuffer));
-
-			if (!ShouldRetryOnError(ResponseCode))
-			{
-				return false;
-			}
-		}
-	}
-#endif
-
-	return false;
-}
-
-FDerivedDataBackendInterface::EPutStatus FHttpCacheStore::PutCachedData(const TCHAR* CacheKey, TArrayView<const uint8> InData, bool bPutEvenIfExists)
-{
-	TRACE_CPUPROFILER_EVENT_SCOPE(HttpDDC_PutCachedData);
-
-	if (!IsWritable())
-	{
-		UE_LOG(LogDerivedDataCache, Verbose, TEXT("%s is read only. Skipping put of %s"), *GetName(), CacheKey);
-		return EPutStatus::NotCached;
-	}
-
-	// don't put anything we pretended didn't exist
-	if (DebugOptions.ShouldSimulatePutMiss(CacheKey))
-	{
-		return EPutStatus::Skipped;
-	}
-
-#if 0 // No longer WITH_DATAREQUEST_HELPER as async puts are unsupported except through the AsyncPutWrapper which expects the inner backend to perform the put synchronously
-	for (int32 Attempts = 0; Attempts < UE_HTTPDDC_MAX_ATTEMPTS; ++Attempts)
-	{
-		FDataUploadHelper Request(PutRequestPools[IsInGameThread()].Get(), *Namespace, *DefaultBucket,	CacheKey, InData, UsageStats);
-
-		if (ShouldAbortForShutdown())
-		{
-			return EPutStatus::NotCached;
-		}
-
-		const int64 ResponseCode = Request.GetResponseCode();
-
-		if (Request.IsSuccess() && (Request.IsQueued() || FHttpRequest::IsSuccessResponse(ResponseCode)))
-		{
-			return Request.IsQueued() ? EPutStatus::Executing : EPutStatus::Cached;
-		}
-
-		if (!ShouldRetryOnError(ResponseCode))
-		{
-			return EPutStatus::NotCached;
-		}
-	}
-#else
-	COOK_STAT(auto Timer = UsageStats.TimePut());
-
-	FString Uri = FString::Printf(TEXT("api/v1/c/ddc/%s/%s/%s"), *Namespace, *DefaultBucket, CacheKey);
-	int64 ResponseCode = 0; uint32 Attempts = 0;
-
-	// Retry request until we get an accepted response or exhaust allowed number of attempts.
-	while (ResponseCode == 0 && ++Attempts < UE_HTTPDDC_MAX_ATTEMPTS)
-	{
-		if (ShouldAbortForShutdown())
-		{
-			return EPutStatus::NotCached;
-		}
-
-		FScopedHttpPoolRequestPtr Request(PutRequestPools[IsInGameThread()].Get());
-		if (Request.IsValid())
-		{
-			// Append the content hash to the header
-			HashPayload(Request.Get(), InData);
-
-			Request->PerformBlockingPut(*Uri, FCompositeBuffer(FSharedBuffer::MakeView(InData.GetData(), InData.Num())), EHttpContentType::Binary);
-			ResponseCode = Request->GetResponseCode();
-
-			if (FHttpRequest::IsSuccessResponse(ResponseCode))
-			{
-				TRACE_COUNTER_ADD(HttpDDC_BytesSent, int64(Request->GetBytesSent()));
-				COOK_STAT(Timer.AddHit(Request->GetBytesSent()));
-				return EPutStatus::Cached;
-			}
-
-			if (!ShouldRetryOnError(ResponseCode))
-			{
-				return EPutStatus::NotCached;
-			}
-
-			ResponseCode = 0;
-		}
-	}
-#endif // WITH_DATAREQUEST_HELPER
-
-	return EPutStatus::NotCached;
-}
-
-void FHttpCacheStore::RemoveCachedData(const TCHAR* CacheKey, bool bTransient)
-{
-	// do not remove transient data as Jupiter does its own verification of the content and cleans itself up
-	if (!IsWritable() || bTransient)
-		return;
-	
-	TRACE_CPUPROFILER_EVENT_SCOPE(HttpDDC_Remove);
-	FString Uri = FString::Printf(TEXT("api/v1/c/ddc/%s/%s/%s"), *Namespace, *DefaultBucket, CacheKey);
-	int64 ResponseCode = 0; uint32 Attempts = 0;
-
-	// Retry request until we get an accepted response or exhaust allowed number of attempts.
-	while (ResponseCode == 0 && ++Attempts < UE_HTTPDDC_MAX_ATTEMPTS)
-	{
-		FScopedHttpPoolRequestPtr Request(PutRequestPools[IsInGameThread()].Get());
-		if (Request.IsValid())
-		{
-			FHttpRequest::EResult Result = Request->PerformBlockingDelete(*Uri, {});
-			ResponseCode = Request->GetResponseCode();
-
-			if (ResponseCode == 200)
-			{
-				return;
-			}
-
-			if (!ShouldRetryOnError(ResponseCode))
-			{
-				return;
-			}
-
-			ResponseCode = 0;
-		}
-	}
-}
-
-TSharedRef<FDerivedDataCacheStatsNode> FHttpCacheStore::GatherUsageStats() const
+void FHttpCacheStore::LegacyStats(FDerivedDataCacheStatsNode& OutNode)
 {
-	TSharedRef<FDerivedDataCacheStatsNode> StatsNode =
-		MakeShared<FDerivedDataCacheStatsNode>(TEXT("Horde Storage"), FString::Printf(TEXT("%s (%s)"), *Domain, *Namespace), /*bIsLocal*/ false);
-	StatsNode->UsageStats.Add(TEXT(""), UsageStats);
-	return StatsNode;
+	OutNode = {TEXT("Horde Storage"), FString::Printf(TEXT("%s (%s)"), *Domain, *Namespace), /*bIsLocal*/ false};
+	OutNode.UsageStats.Add(TEXT(""), UsageStats);
 }
 
 void FHttpCacheStore::Put(
@@ -3382,7 +2021,7 @@ void FHttpCacheStore::GetValue(
 			else
 			{
 				UE_LOG(LogDerivedDataCache, Verbose, TEXT("%s: Cache hit for %s from '%s'"),
-					*GetName(), *WriteToString<96>(Response.Key), *Response.Name);
+					*Domain, *WriteToString<96>(Response.Key), *Response.Name);
 				COOK_STAT(UsageStats.GetStats.Accumulate(FCookStats::CallStats::EHitOrMiss::Hit, FCookStats::CallStats::EStatType::Counter, 1l, bIsInGameThread));
 				OnComplete(MoveTemp(Response));
 			}
@@ -3415,14 +2054,14 @@ void FHttpCacheStore::GetValue(
 					{
 						// With inline fetching, expect we will always have a value we can use.  Even SkipData/Exists can rely on the blob existing if the ref is reported to exist.
 						UE_LOG(LogDerivedDataCache, Log, TEXT("%s: Cache miss due to inlining failure for %s from '%s'"),
-									*GetName(), *WriteToString<96>(Response.Key), *Response.Name);
+									*Domain, *WriteToString<96>(Response.Key), *Response.Name);
 						COOK_STAT(UsageStats.GetStats.Accumulate(FCookStats::CallStats::EHitOrMiss::Miss, FCookStats::CallStats::EStatType::Counter, 1l, bIsInGameThread));
 						OnComplete(MoveTemp(Response));
 					}
 					else
 					{
 						UE_LOG(LogDerivedDataCache, Verbose, TEXT("%s: Cache hit for %s from '%s'"),
-									*GetName(), *WriteToString<96>(Response.Key), *Response.Name);
+									*Domain, *WriteToString<96>(Response.Key), *Response.Name);
 						uint64 ValueSize = Response.Value.GetData().GetCompressedSize();
 						TRACE_COUNTER_ADD(HttpDDC_BytesReceived, ValueSize);
 						COOK_STAT(UsageStats.GetStats.Accumulate(FCookStats::CallStats::EHitOrMiss::Hit, FCookStats::CallStats::EStatType::Counter, 1l, bIsInGameThread));
@@ -3601,7 +2240,7 @@ void FHttpCacheStore::GetChunks(
 			const uint64 RawOffset = FMath::Min(Value.GetRawSize(), Request.RawOffset);
 			const uint64 RawSize = FMath::Min(Value.GetRawSize() - RawOffset, Request.RawSize);
 			UE_LOG(LogDerivedDataCache, Verbose, TEXT("%s: Cache hit for %s from '%s'"),
-				*GetName(), *WriteToString<96>(Request.Key, '/', Request.Id), *Request.Name);
+				*Domain, *WriteToString<96>(Request.Key, '/', Request.Id), *Request.Name);
 			COOK_STAT(Timer.AddHit(!bExistsOnly ? RawSize : 0));
 			FSharedBuffer Buffer;
 			if (!bExistsOnly)
@@ -3618,6 +2257,82 @@ void FHttpCacheStore::GetChunks(
 	}
 }
 
+void FHttpCacheStoreParams::Parse(const TCHAR* NodeName, const TCHAR* Config)
+{
+	FString ServerId;
+	if (FParse::Value(Config, TEXT("ServerID="), ServerId))
+	{
+		FString ServerEntry;
+		const TCHAR* ServerSection = TEXT("HordeStorageServers");
+		if (GConfig->GetString(ServerSection, *ServerId, ServerEntry, GEngineIni))
+		{
+			Parse(NodeName, *ServerEntry);
+		}
+		else
+		{
+			UE_LOG(LogDerivedDataCache, Warning, TEXT("%s: Using ServerID=%s which was not found in [%s]"), NodeName, *ServerId, ServerSection);
+		}
+	}
+
+	FString OverrideName;
+
+	// Host Params
+
+	FParse::Value(Config, TEXT("Host="), Host);
+	if (FParse::Value(Config, TEXT("EnvHostOverride="), OverrideName))
+	{
+		FString HostEnv = FPlatformMisc::GetEnvironmentVariable(*OverrideName);
+		if (!HostEnv.IsEmpty())
+		{
+			Host = HostEnv;
+			UE_LOG(LogDerivedDataCache, Log, TEXT("%s: Found environment override for Host %s=%s"), NodeName, *OverrideName, *Host);
+		}
+	}
+	if (FParse::Value(Config, TEXT("CommandLineHostOverride="), OverrideName))
+	{
+		if (FParse::Value(FCommandLine::Get(), *(OverrideName + TEXT("=")), Host))
+		{
+			UE_LOG(LogDerivedDataCache, Log, TEXT("%s: Found command line override for Host %s=%s"), NodeName, *OverrideName, *Host);
+		}
+	}
+
+	FParse::Bool(Config, TEXT("ResolveHostCanonicalName="), bResolveHostCanonicalName);
+
+	// Namespace Params
+
+	FParse::Value(Config, TEXT("Namespace="), Namespace);
+	FParse::Value(Config, TEXT("StructuredNamespace="), StructuredNamespace);
+
+	// OAuth Params
+
+	FParse::Value(Config, TEXT("OAuthProvider="), OAuthProvider);
+
+	if (FParse::Value(Config, TEXT("CommandLineOAuthProviderOverride="), OverrideName))
+	{
+		if (FParse::Value(FCommandLine::Get(), *(OverrideName + TEXT("=")), OAuthProvider))
+		{
+			UE_LOG(LogDerivedDataCache, Log, TEXT("%s: Found command line override for OAuthProvider %s=%s"), NodeName, *OverrideName, *OAuthProvider);
+		}
+	}
+
+	FParse::Value(Config, TEXT("OAuthClientId="), OAuthClientId);
+	FParse::Value(Config, TEXT("OAuthSecret="), OAuthSecret);
+
+	if (FParse::Value(Config, TEXT("CommandLineOAuthSecretOverride="), OverrideName))
+	{
+		if (FParse::Value(FCommandLine::Get(), *(OverrideName + TEXT("=")), OAuthSecret))
+		{
+			UE_LOG(LogDerivedDataCache, Log, TEXT("%s: Found command line override for OAuthSecret %s=%s"), NodeName, *OverrideName, *OAuthSecret);
+		}
+	}
+
+	FParse::Value(Config, TEXT("OAuthScope="), OAuthScope);
+
+	// Cache Params
+
+	FParse::Bool(Config, TEXT("ReadOnly="), bReadOnly);
+}
+
 } // UE::DerivedData
 
 #endif // WITH_HTTP_DDC_BACKEND
@@ -3625,36 +2340,78 @@ void FHttpCacheStore::GetChunks(
 namespace UE::DerivedData
 {
 
-ILegacyCacheStore* CreateHttpCacheStore(
-	const TCHAR* NodeName,
-	const TCHAR* ServiceUrl,
-	bool bResolveHostCanonicalName,
-	const TCHAR* Namespace,
-	const TCHAR* StructuredNamespace,
-	const TCHAR* OAuthProvider,
-	const TCHAR* OAuthClientId,
-	const TCHAR* OAuthData,
-	const TCHAR* OAuthScope,
-	const FDerivedDataBackendInterface::ESpeedClass* ForceSpeedClass,
-	EBackendLegacyMode LegacyMode,
-	bool bReadOnly)
+TTuple<ILegacyCacheStore*, ECacheStoreFlags> CreateHttpCacheStore(const TCHAR* NodeName, const TCHAR* Config)
 {
 #if WITH_HTTP_DDC_BACKEND
-	FHttpCacheStore* Backend = new FHttpCacheStore(ServiceUrl, bResolveHostCanonicalName, Namespace, StructuredNamespace, OAuthProvider, OAuthClientId, OAuthData, OAuthScope, LegacyMode, bReadOnly);
-	if (Backend->IsUsable())
+	FHttpCacheStoreParams Params;
+	Params.Parse(NodeName, Config);
+
+	if (Params.Host.IsEmpty())
 	{
-		return Backend;
+		UE_LOG(LogDerivedDataCache, Error, TEXT("%s: Missing required parameter 'Host'"), NodeName);
+		return MakeTuple(nullptr, ECacheStoreFlags::None);
 	}
-	UE_LOG(LogDerivedDataCache, Warning, TEXT("Node %s could not contact the service (%s), will not use it"), NodeName, ServiceUrl);
-	delete Backend;
-	return nullptr;
+
+	if (Params.Host == TEXT("None"))
+	{
+		UE_LOG(LogDerivedDataCache, Log, TEXT("%s: Disabled because Host is set to 'None'"), NodeName);
+		return MakeTuple(nullptr, ECacheStoreFlags::None);
+	}
+
+	if (Params.Namespace.IsEmpty())
+	{
+		Params.Namespace = FApp::GetProjectName();
+		UE_LOG(LogDerivedDataCache, Warning, TEXT("%s: Missing required parameter 'Namespace', falling back to '%s'"), NodeName, *Params.Namespace);
+	}
+
+	if (Params.StructuredNamespace.IsEmpty())
+	{
+		Params.StructuredNamespace = Params.Namespace;
+		UE_LOG(LogDerivedDataCache, Warning, TEXT("%s: Missing required parameter 'StructuredNamespace', falling back to '%s'"), NodeName, *Params.StructuredNamespace);
+	}
+
+	if (Params.OAuthProvider.IsEmpty())
+	{
+		UE_LOG(LogDerivedDataCache, Error, TEXT("%s: Missing required parameter 'OAuthProvider'"), NodeName);
+		return MakeTuple(nullptr, ECacheStoreFlags::None);
+	}
+
+	// No need for OAuth client id and secret if using a local provider.
+	if (!Params.OAuthProvider.StartsWith(TEXT("http://localhost")))
+	{
+		if (Params.OAuthClientId.IsEmpty())
+		{
+			UE_LOG(LogDerivedDataCache, Error, TEXT("%s: Missing required parameter 'OAuthClientId'"), NodeName);
+			return MakeTuple(nullptr, ECacheStoreFlags::None);
+		}
+
+		if (Params.OAuthSecret.IsEmpty())
+		{
+			UE_LOG(LogDerivedDataCache, Error, TEXT("%s: Missing required parameter 'OAuthSecret'"), NodeName);
+			return MakeTuple(nullptr, ECacheStoreFlags::None);
+		}
+	}
+
+	if (Params.OAuthScope.IsEmpty())
+	{
+		Params.OAuthScope = TEXTVIEW("cache_access");
+	}
+
+	TUniquePtr<FHttpCacheStore> Backend = MakeUnique<FHttpCacheStore>(Params);
+	if (!Backend->IsUsable())
+	{
+		UE_LOG(LogDerivedDataCache, Warning, TEXT("%s: Failed to contact the service (%s), will not use it."), NodeName, *Params.Host);
+		Backend.Reset();
+	}
+
+	return MakeTuple(Backend.Release(), ECacheStoreFlags::Remote | ECacheStoreFlags::Query | (Params.bReadOnly ? ECacheStoreFlags::None : ECacheStoreFlags::Store));
 #else
-	UE_LOG(LogDerivedDataCache, Warning, TEXT("HTTP backend is not yet supported in the current build configuration."));
-	return nullptr;
+	UE_LOG(LogDerivedDataCache, Warning, TEXT("%s: HTTP cache is not yet supported in the current build configuration."), NodeName);
+	return MakeTuple(nullptr, ECacheStoreFlags::None);
 #endif
 }
 
-FDerivedDataBackendInterface* GetAnyHttpCacheStore(
+ILegacyCacheStore* GetAnyHttpCacheStore(
 	FString& OutDomain,
 	FString& OutOAuthProvider,
 	FString& OutOAuthClientId,
diff --git a/Engine/Source/Developer/DerivedDataCache/Private/Tests/HttpCacheStoreTest.cpp b/Engine/Source/Developer/DerivedDataCache/Private/Tests/HttpCacheStoreTest.cpp
--- a/Engine/Source/Developer/DerivedDataCache/Private/Tests/HttpCacheStoreTest.cpp
+++ b/Engine/Source/Developer/DerivedDataCache/Private/Tests/HttpCacheStoreTest.cpp
@@ -38,7 +38,7 @@
 namespace UE::DerivedData
 {
 
-FDerivedDataBackendInterface* GetAnyHttpCacheStore(
+ILegacyCacheStore* GetAnyHttpCacheStore(
 	FString& OutDomain,
 	FString& OutOAuthProvider,
 	FString& OutOAuthClientId,
@@ -47,7 +47,7 @@
 	FString& OutNamespace,
 	FString& OutStructuredNamespace);
 
-ILegacyCacheStore* CreateZenCacheStore(const TCHAR* NodeName, const TCHAR* ServiceUrl, const TCHAR* Namespace, const TCHAR* StructuredNamespace);
+TTuple<ILegacyCacheStore*, ECacheStoreFlags> CreateZenCacheStore(const TCHAR* NodeName, const TCHAR* Config);
 
 class FHttpCacheStoreTestBase : public FAutomationTestBase
 {
@@ -59,7 +59,7 @@
 
 	bool CheckPrequisites() const
 	{
-		if (FDerivedDataBackendInterface* Backend = GetTestBackend())
+		if (ILegacyCacheStore* Backend = GetTestBackend())
 		{
 			return true;
 		}
@@ -155,9 +155,9 @@
 		FPlatformProcess::ReturnSynchEventToPool(LastEvent);
 	}
 
-	static FDerivedDataBackendInterface* GetTestBackend()
+	static ILegacyCacheStore* GetTestBackend()
 	{
-		static FDerivedDataBackendInterface* CachedBackend = GetAnyHttpCacheStore(
+		static ILegacyCacheStore* CachedBackend = GetAnyHttpCacheStore(
 			TestDomain, TestOAuthProvider, TestOAuthClientId, TestOAuthSecret, TestOAuthScope, TestNamespace, TestStructuredNamespace);
 		return CachedBackend;
 	}
@@ -492,29 +492,6 @@
 	static inline FString TestStructuredNamespace;
 };
 
-// Helper function to create a number of dummy cache keys for testing
-TArray<FString> CreateTestCacheKeys(FDerivedDataBackendInterface* InTestBackend, uint32 InNumKeys)
-{
-	TArray<FString> Keys;
-	TArray<uint8> KeyContents;
-	KeyContents.Add(42);
-
-	FSHA1 HashState;
-	HashState.Update(KeyContents.GetData(), KeyContents.Num());
-	HashState.Final();
-	uint8 Hash[FSHA1::DigestSize];
-	HashState.GetHash(Hash);
-	const FString HashString = BytesToHex(Hash, FSHA1::DigestSize);
-
-	for (uint32 KeyIndex = 0; KeyIndex < InNumKeys; ++KeyIndex)
-	{
-		FString NewKey = FString::Printf(TEXT("__AutoTest_Dummy_%u__%s"), KeyIndex, *HashString);
-		Keys.Add(NewKey);
-		InTestBackend->PutCachedData(*NewKey, KeyContents, false);
-	}
-	return Keys;
-}
-
 TArray<FCacheRecord> CreateTestCacheRecords(ICacheStore* InTestBackend, uint32 InNumKeys, uint32 InNumValues, FCbObject MetaContents = FCbObject(), const char* BucketName = nullptr)
 {
 	using namespace UE::DerivedData;
@@ -620,84 +597,12 @@
 	return Values;
 }
 
-IMPLEMENT_HTTPDERIVEDDATA_AUTOMATION_TEST(FConcurrentCachedDataProbablyExistsBatch, TEXT(".FConcurrentCachedDataProbablyExistsBatch"), EAutomationTestFlags::EditorContext | EAutomationTestFlags::ProductFilter)
-bool FConcurrentCachedDataProbablyExistsBatch::RunTest(const FString& Parameters)
-{
-	FDerivedDataBackendInterface* TestBackend = GetTestBackend();
-
-	const int32 ThreadCount = 64;
-	const double Duration = 10;
-	const uint32 KeysInBatch = 4;
-
-	TArray<FString> Keys = CreateTestCacheKeys(TestBackend, KeysInBatch);
-
-	std::atomic<uint32> MismatchedResults = 0;
-
-	ConcurrentTestWithStats(
-		[&]()
-		{
-			TConstArrayView<FString> BatchView = MakeArrayView(Keys.GetData(), KeysInBatch);
-			TBitArray<> Result = TestBackend->CachedDataProbablyExistsBatch(BatchView);
-			if (Result.CountSetBits() != BatchView.Num())
-			{
-				MismatchedResults.fetch_add(BatchView.Num() - Result.CountSetBits(), std::memory_order_relaxed);
-			}
-		},
-		ThreadCount,
-		Duration
-	);
-
-	TestEqual(TEXT("Concurrent calls to CachedDataProbablyExistsBatch for a batch of keys that were put are not reliably found"), MismatchedResults, 0);
-
-	return true;
-}
-
-// This test validate that batch requests wont mismatch head and get request for the same keys in the same batch
-IMPLEMENT_HTTPDERIVEDDATA_AUTOMATION_TEST(FConcurrentExistsAndGetForSameKeyBatch, TEXT(".FConcurrentExistsAndGetForSameKeyBatch"), EAutomationTestFlags::EditorContext | EAutomationTestFlags::ProductFilter)
-bool FConcurrentExistsAndGetForSameKeyBatch::RunTest(const FString& Parameters)
-{
-	FDerivedDataBackendInterface* TestBackend = GetTestBackend();
-
-	const int32 ParallelTasks = 32;
-	const uint32 Iterations = 20;
-	const uint32 KeysInBatch = 4;
-
-	TArray<FString> Keys = CreateTestCacheKeys(TestBackend, KeysInBatch);
-	// Add some non valid keys by just using guids
-	for (int32 Index = 0; Index < KeysInBatch; ++Index)
-	{
-		Keys.Add(FGuid::NewGuid().ToString());
-	}
-	
-	ParallelFor(ParallelTasks,
-		[&](int32 TaskIndex)
-		{
-			for (uint32 Iteration = 0; Iteration < Iterations; ++Iteration)
-			{
-				for (int32 KeyIndex = 0; KeyIndex < Keys.Num(); ++KeyIndex)
-				{
-					if ((Iteration % 2) ^ (KeyIndex % 2))
-					{
-						TestBackend->CachedDataProbablyExists(*Keys[KeyIndex]);
-					}
-					else
-					{
-						TArray<uint8> OutData;
-						TestBackend->GetCachedData(*Keys[KeyIndex], OutData);
-					}
-				}
-			}
-		}
-	);
-	return true;
-}
-
 // Tests basic functionality for structured cache operations
 IMPLEMENT_HTTPDERIVEDDATA_AUTOMATION_TEST(CacheStore, TEXT(".CacheStore"), EAutomationTestFlags::EditorContext | EAutomationTestFlags::ProductFilter)
 bool CacheStore::RunTest(const FString& Parameters)
 {
 	using namespace UE::DerivedData;
-	FDerivedDataBackendInterface* TestBackend = GetTestBackend();
+	ILegacyCacheStore* TestBackend = GetTestBackend();
 
 #if UE_WITH_ZEN
 	using namespace UE::Zen;
@@ -756,10 +661,12 @@
 	ZenTestSiblingAutoLaunchSettings.bLimitProcessLifetime = true;
 
 	FScopeZenService ScopeZenSiblingService(MoveTemp(ZenTestServiceSiblingSettings));
-	TUniquePtr<ILegacyCacheStore> ZenIntermediarySiblingBackend(CreateZenCacheStore(TEXT("TestSibling"), ScopeZenSiblingService.GetInstance().GetURL(), *TestNamespace, *TestNamespace));
+	TUniquePtr<ILegacyCacheStore> ZenIntermediarySiblingBackend(CreateZenCacheStore(TEXT("TestSibling"),
+		*FString::Printf(TEXT("Host=%s, Namespace=%s, StructuredNamespace=%s"), ScopeZenSiblingService.GetInstance().GetURL(), *TestNamespace, *TestNamespace)).Key);
 
 	FScopeZenService ScopeZenService(MoveTemp(ZenTestServiceSettings));
-	TUniquePtr<ILegacyCacheStore> ZenIntermediaryBackend(CreateZenCacheStore(TEXT("Test"), ScopeZenService.GetInstance().GetURL(), *TestNamespace, *TestNamespace));
+	TUniquePtr<ILegacyCacheStore> ZenIntermediaryBackend(CreateZenCacheStore(TEXT("Test"),
+		*FString::Printf(TEXT("Host=%s, Namespace=%s, StructuredNamespace=%s"), ScopeZenService.GetInstance().GetURL(), *TestNamespace, *TestNamespace)).Key);
 	auto WaitForZenPushToUpstream = [](ILegacyCacheStore* ZenBackend, TConstArrayView<FCacheRecord> Records)
 	{
 		// TODO: Expecting a legitimate means to wait for zen to finish pushing records to its upstream in the future
diff --git a/Engine/Source/Developer/DerivedDataCache/Private/ZenCacheStore.cpp b/Engine/Source/Developer/DerivedDataCache/Private/ZenCacheStore.cpp
--- a/Engine/Source/Developer/DerivedDataCache/Private/ZenCacheStore.cpp
+++ b/Engine/Source/Developer/DerivedDataCache/Private/ZenCacheStore.cpp
@@ -1,24 +1,22 @@
 // Copyright Epic Games, Inc. All Rights Reserved.
 
-#include "DerivedDataBackendInterface.h"
 #include "ZenServerInterface.h"
 
 #if UE_WITH_ZEN
 
 #include "BatchView.h"
-#include "Containers/Set.h"
-#include "Containers/StringFwd.h"
-#include "Containers/StringFwd.h"
+#include "DerivedDataBackendInterface.h"
 #include "DerivedDataCachePrivate.h"
 #include "DerivedDataCacheRecord.h"
 #include "DerivedDataCacheUsageStats.h"
 #include "DerivedDataChunk.h"
+#include "DerivedDataLegacyCacheStore.h"
 #include "Http/HttpClient.h"
+#include "Misc/App.h"
 #include "Misc/ConfigCacheIni.h"
 #include "Misc/Optional.h"
 #include "ProfilingDebugging/CountersTrace.h"
 #include "ProfilingDebugging/CpuProfilerTrace.h"
-#include "Serialization/BufferArchive.h"
 #include "Serialization/CompactBinary.h"
 #include "Serialization/CompactBinaryPackage.h"
 #include "Serialization/CompactBinarySerialization.h"
@@ -31,8 +29,6 @@
 #include "ZenSerialization.h"
 #include "ZenStatistics.h"
 
-TRACE_DECLARE_INT_COUNTER(ZenDDC_Exist,			TEXT("ZenDDC Exist"));
-TRACE_DECLARE_INT_COUNTER(ZenDDC_ExistHit,		TEXT("ZenDDC Exist Hit"));
 TRACE_DECLARE_INT_COUNTER(ZenDDC_Get,			TEXT("ZenDDC Get"));
 TRACE_DECLARE_INT_COUNTER(ZenDDC_GetHit,		TEXT("ZenDDC Get Hit"));
 TRACE_DECLARE_INT_COUNTER(ZenDDC_Put,			TEXT("ZenDDC Put"));
@@ -72,97 +68,60 @@ void ForEachBatch(const int32 BatchSize, const int32 TotalCount, T&& Fn)
 /**
  * Backend for a HTTP based caching service (Zen)
  */
-class FZenCacheStore final : public FDerivedDataBackendInterface
+class FZenCacheStore final : public ILegacyCacheStore
 {
 public:
 	
 	/**
-	 * Creates the backend, checks health status and attempts to acquire an access token.
+	 * Creates the cache store client, checks health status and attempts to acquire an access token.
 	 *
 	 * @param ServiceUrl	Base url to the service including schema.
 	 * @param Namespace		Namespace to use.
 	 */
 	FZenCacheStore(const TCHAR* ServiceUrl, const TCHAR* Namespace, const TCHAR* StructuredNamespace);
 
-	~FZenCacheStore();
+	inline FString GetName() const { return ZenService.GetInstance().GetURL(); }
 
 	/**
-	 * Checks is backend is usable (reachable and accessible).
+	 * Checks if cache service is usable (reachable and accessible).
 	 * @return true if usable
 	 */
-	bool IsUsable() const { return bIsUsable; }
-
-	virtual bool IsWritable() const override;
-	virtual ESpeedClass GetSpeedClass() const override;
-	virtual TSharedRef<FDerivedDataCacheStatsNode> GatherUsageStats() const override;
-
-	/**
-	 * Synchronous attempt to make sure the cached data will be available as optimally as possible.
-	 *
-	 * @param	CacheKeys	Alphanumeric+underscore keys of the cache items
-	 * @return				true if the data will probably be found in a fast backend on a future request.
-	 */
-	virtual TBitArray<> TryToPrefetch(TConstArrayView<FString> CacheKeys) override;
-
-	virtual bool CachedDataProbablyExists(const TCHAR* CacheKey) override;
-	virtual bool GetCachedData(const TCHAR* CacheKey, TArray<uint8>& OutData) override;
-	virtual EPutStatus PutCachedData(const TCHAR* CacheKey, TArrayView<const uint8> InData, bool bPutEvenIfExists) override;
-
-	virtual void RemoveCachedData(const TCHAR* CacheKey, bool bTransient) override;
-	
-	virtual FString GetName() const override;
-	virtual bool WouldCache(const TCHAR* CacheKey, TArrayView<const uint8> InData) override;
-	virtual bool ApplyDebugOptions(FBackendDebugOptions& InOptions) override;
-
-	virtual EBackendLegacyMode GetLegacyMode() const override { return EBackendLegacyMode::ValueOnly; }
+	inline bool IsUsable() const { return bIsUsable; }
 
 	// ICacheStore
 
-	virtual void Put(
+	void Put(
 		TConstArrayView<FCachePutRequest> Requests,
 		IRequestOwner& Owner,
-		FOnCachePutComplete&& OnComplete = FOnCachePutComplete()) override;
+		FOnCachePutComplete&& OnComplete = FOnCachePutComplete()) final;
 
-	virtual void Get(
+	void Get(
 		TConstArrayView<FCacheGetRequest> Requests,
 		IRequestOwner& Owner,
-		FOnCacheGetComplete&& OnComplete) override;
+		FOnCacheGetComplete&& OnComplete) final;
 
-	virtual void PutValue(
+	void PutValue(
 		TConstArrayView<FCachePutValueRequest> Requests,
 		IRequestOwner& Owner,
-		FOnCachePutValueComplete&& OnComplete = FOnCachePutValueComplete()) override;
+		FOnCachePutValueComplete&& OnComplete = FOnCachePutValueComplete()) final;
 
-	virtual void GetValue(
+	void GetValue(
 		TConstArrayView<FCacheGetValueRequest> Requests,
 		IRequestOwner& Owner,
-		FOnCacheGetValueComplete&& OnComplete = FOnCacheGetValueComplete()) override;
+		FOnCacheGetValueComplete&& OnComplete = FOnCacheGetValueComplete()) final;
 
-	virtual void GetChunks(
+	void GetChunks(
 		TConstArrayView<FCacheGetChunkRequest> Requests,
 		IRequestOwner& Owner,
-		FOnCacheGetChunkComplete&& OnComplete) override;
+		FOnCacheGetChunkComplete&& OnComplete) final;
 
-private:
-	enum class EGetResult
-	{
-		Success,
-		NotFound,
-		Corrupted
-	};
-	EGetResult GetZenData(FStringView Uri, TArray64<uint8>* OutData, EHttpContentType ContentType) const;
+	// ILegacyCacheStore
 
-	// TODO: need ability to specify content type
-	FDerivedDataBackendInterface::EPutStatus PutZenData(const TCHAR* Uri, const FCompositeBuffer& InData, EHttpContentType ContentType);
-	EGetResult GetZenData(const FCacheKey& Key, ECachePolicy CachePolicy, FCbPackage& OutPackage) const;
+	void LegacyStats(FDerivedDataCacheStatsNode& OutNode) final;
+	bool LegacyDebugOptions(FBackendDebugOptions& Options) final;
 
+private:
 	bool IsServiceReady();
-	FString MakeLegacyZenKey(const TCHAR* CacheKey);
-	static void AppendZenUri(const FCacheKey& CacheKey, FStringBuilderBase& Out);
-	static void AppendZenUri(const FCacheKey& CacheKey, const FValueId& Id, FStringBuilderBase& Out);
-	static void AppendPolicyQueryString(ECachePolicy Policy, FStringBuilderBase& Out);
-
-	static bool ShouldRetryOnError(int64 ResponseCode);
 
 	using FOnRpcComplete = TUniqueFunction<void(FHttpRequest::EResult HttpResult, FHttpRequest* Request, FCbPackage& OutResponse)>;
 	static FHttpRequest::EResult ParseRpcResponse(FHttpRequest::EResult ResultFromPost, FHttpRequest& Request, FCbPackage& OutResponse);
@@ -194,8 +153,6 @@ private:
 	mutable FDerivedDataCacheUsageStats UsageStats;
 	TUniquePtr<FHttpRequestPool> RequestPool;
 	bool bIsUsable = false;
-	uint32 FailedLoginAttempts = 0;
-	uint32 MaxAttempts = 4;
 	int32 BatchPutMaxBytes = 1024*1024;
 	int32 CacheRecordBatchSize = 8;
 	int32 CacheChunksBatchSize = 8;
@@ -227,37 +184,11 @@ FZenCacheStore::FZenCacheStore(
 	GConfig->GetInt(TEXT("Zen"), TEXT("CacheChunksBatchSize"), CacheChunksBatchSize, GEngineIni);
 }
 
-FZenCacheStore::~FZenCacheStore()
-{
-}
-
-FString FZenCacheStore::GetName() const
-{
-	return ZenService.GetInstance().GetURL();
-}
-
 bool FZenCacheStore::IsServiceReady()
 {
 	return ZenService.GetInstance().IsServiceReady();
 }
 
-bool FZenCacheStore::ShouldRetryOnError(int64 ResponseCode)
-{
-	// Access token might have expired, request a new token and try again.
-	if (ResponseCode == 401)
-	{
-		return true;
-	}
-
-	// Too many requests, make a new attempt
-	if (ResponseCode == 429)
-	{
-		return true;
-	}
-
-	return false;
-}
-
 FHttpRequest::EResult FZenCacheStore::ParseRpcResponse(FHttpRequest::EResult ResultFromPost, FHttpRequest& Request, FCbPackage& OutResponse)
 {
 	if (ResultFromPost != FHttpRequest::EResult::Success || !FHttpRequest::IsSuccessResponse(Request.GetResponseCode()))
@@ -339,307 +270,7 @@ void FZenCacheStore::EnqueueAsyncRpc(FHttpRequest& Request,
 	Request.EnqueueAsyncPost(Owner, Pool, Uri, FCompositeBuffer(PackageSharedBuffer), MoveTemp(OnHttpRequestComplete), EHttpContentType::CbPackage, EHttpContentType::CbPackage);
 }
 
-bool FZenCacheStore::CachedDataProbablyExists(const TCHAR* CacheKey)
-{
-	TRACE_CPUPROFILER_EVENT_SCOPE(ZenDDC::Exist);
-	TRACE_COUNTER_ADD(ZenDDC_Exist, int64(1));
-	COOK_STAT(auto Timer = UsageStats.TimeProbablyExists());
-
-	if (DebugOptions.ShouldSimulateGetMiss(CacheKey))
-	{
-		return false;
-	}
-	
-	FString Uri = MakeLegacyZenKey(CacheKey);
-	int64 ResponseCode = 0; 
-	uint32 Attempts = 0;
-
-	// Retry request until we get an accepted response or exhaust allowed number of attempts.
-	while (ResponseCode == 0 && ++Attempts < MaxAttempts)
-	{
-		FScopedHttpPoolRequestPtr Request(RequestPool.Get());
-		FHttpRequest::EResult Result = Request->PerformBlockingHead(*Uri, EHttpContentType::Binary);
-		ResponseCode = Request->GetResponseCode();
-
-		if (FHttpRequest::IsSuccessResponse(ResponseCode) || ResponseCode == 404)
-		{
-			const bool bIsHit = (Result == FHttpRequest::EResult::Success && FHttpRequest::IsSuccessResponse(ResponseCode));
-
-			if (bIsHit)
-			{
-				TRACE_COUNTER_ADD(ZenDDC_ExistHit, int64(1));
-				COOK_STAT(Timer.AddHit(0));
-			}
-			return bIsHit;
-		}
-
-		if (!ShouldRetryOnError(ResponseCode))
-		{
-			return false;
-		}
-
-		ResponseCode = 0;
-	}
-
-	return false;
-}
-
-bool FZenCacheStore::GetCachedData(const TCHAR* CacheKey, TArray<uint8>& OutData)
-{
-	TRACE_CPUPROFILER_EVENT_SCOPE(ZenDDC::GetCachedData);
-	TRACE_COUNTER_ADD(ZenDDC_Get, int64(1));
-	COOK_STAT(auto Timer = UsageStats.TimeGet());
-
-	if (DebugOptions.ShouldSimulateGetMiss(CacheKey))
-	{
-		return false;
-	}
-
-	double StartTime = FPlatformTime::Seconds();
-
-	TArray64<uint8> ArrayBuffer;
-	EGetResult Result = GetZenData(MakeLegacyZenKey(CacheKey), &ArrayBuffer, EHttpContentType::Binary);
-	check(ArrayBuffer.Num() <= UINT32_MAX);
-	OutData = TArray<uint8>(MoveTemp(ArrayBuffer));
-	if (Result != EGetResult::Success)
-	{
-		switch (Result)
-		{
-		default:
-			UE_LOG(LogDerivedDataCache, Verbose, TEXT("%s: Cache miss on %s"), *GetName(), CacheKey);
-			break;
-		case EGetResult::Corrupted:
-			UE_LOG(LogDerivedDataCache, Warning,
-				TEXT("Checksum from server on %s did not match recieved data. Discarding cached result."), CacheKey);
-			break;
-		}
-		return false;
-	}
-
-	TRACE_COUNTER_ADD(ZenDDC_GetHit, int64(1));
-	TRACE_COUNTER_ADD(ZenDDC_BytesReceived, int64(OutData.Num()));
-	COOK_STAT(Timer.AddHit(OutData.Num()));
-	double ReadDuration = FPlatformTime::Seconds() - StartTime;
-	double ReadSpeed = (OutData.Num() / ReadDuration) / (1024.0 * 1024.0);
-	UE_LOG(LogDerivedDataCache, Verbose, TEXT("%s: Cache hit on %s (%d bytes, %.02f secs, %.2fMB/s)"),
-		*GetName(), CacheKey, OutData.Num(), ReadDuration, ReadSpeed);
-	return true;
-}
-
-FZenCacheStore::EGetResult
-FZenCacheStore::GetZenData(FStringView Uri, TArray64<uint8>* OutData, EHttpContentType ContentType) const
-{
-	// Retry request until we get an accepted response or exhaust allowed number of attempts.
-	EGetResult GetResult = EGetResult::NotFound;
-	for (uint32 Attempts = 0; Attempts < MaxAttempts; ++Attempts)
-	{
-		FScopedHttpPoolRequestPtr Request(RequestPool.Get());
-		if (Request.IsValid())
-		{
-			FHttpRequest::EResult Result;
-			if (OutData)
-			{
-				Result = Request->PerformBlockingDownload(Uri, OutData, ContentType);
-			}
-			else
-			{
-				Result = Request->PerformBlockingHead(Uri, ContentType);
-			}
-			int64 ResponseCode = Request->GetResponseCode();
-
-			// Request was successful, make sure we got all the expected data.
-			if (FHttpRequest::IsSuccessResponse(ResponseCode))
-			{
-				return EGetResult::Success;
-			}
-
-			if (!ShouldRetryOnError(ResponseCode))
-			{
-				break;
-			}
-		}
-	}
-
-	if (OutData)
-	{
-		OutData->Reset();
-	}
-
-	return GetResult;
-}
-
-FZenCacheStore::EGetResult
-FZenCacheStore::GetZenData(const FCacheKey& CacheKey, ECachePolicy CachePolicy, FCbPackage& OutPackage) const
-{
-	TStringBuilder<256> QueryUri;
-	AppendZenUri(CacheKey, QueryUri);
-	AppendPolicyQueryString(CachePolicy, QueryUri);
-
-	// Retry request until we get an accepted response or exhaust allowed number of attempts.
-	EGetResult GetResult = EGetResult::NotFound;
-	for (uint32 Attempts = 0; Attempts < MaxAttempts; ++Attempts)
-	{
-		FScopedHttpPoolRequestPtr Request(RequestPool.Get());
-		if (Request.IsValid())
-		{
-			TArray64<uint8> PackageBytes;
-			FHttpRequest::EResult Result = Request->PerformBlockingDownload(QueryUri.ToString(), &PackageBytes, EHttpContentType::CbPackage);
-			int64 ResponseCode = Request->GetResponseCode();
-
-			// Request was successful, make sure we got all the expected data.
-			if (FHttpRequest::IsSuccessResponse(ResponseCode))
-			{
-				FLargeMemoryReader Ar(PackageBytes.GetData(), PackageBytes.Num());
-				if (OutPackage.TryLoad(Ar))
-				{
-					GetResult = EGetResult::Success;
-				}
-				else
-				{
-					GetResult = EGetResult::Corrupted;
-				}
-				break;
-			}
-
-			if (!ShouldRetryOnError(ResponseCode))
-			{
-				break;
-			}
-		}
-	}
-
-	if (GetResult != EGetResult::Success)
-	{
-		OutPackage.Reset();
-	}
-	return GetResult;
-}
-
-FDerivedDataBackendInterface::EPutStatus
-FZenCacheStore::PutCachedData(const TCHAR* CacheKey, TArrayView<const uint8> InData, bool bPutEvenIfExists)
-{
-	TRACE_CPUPROFILER_EVENT_SCOPE(ZenDDC::PutCachedData);
-
-	if (DebugOptions.ShouldSimulatePutMiss(CacheKey))
-	{
-		return EPutStatus::NotCached;
-	}
-
-	FSharedBuffer DataBuffer = FSharedBuffer::MakeView(InData.GetData(), InData.Num());
-	return PutZenData(*MakeLegacyZenKey(CacheKey), FCompositeBuffer(DataBuffer), EHttpContentType::Binary);
-}
-
-FDerivedDataBackendInterface::EPutStatus
-FZenCacheStore::PutZenData(const TCHAR* Uri, const FCompositeBuffer& InData, EHttpContentType ContentType)
-{
-	TRACE_CPUPROFILER_EVENT_SCOPE(ZenDDC_Put);
-	COOK_STAT(auto Timer = UsageStats.TimePut());
-
-	int64 ResponseCode = 0; 
-	uint32 Attempts = 0;
-
-	// Retry request until we get an accepted response or exhaust allowed number of attempts.
-	while (ResponseCode == 0 && ++Attempts < MaxAttempts)
-	{
-		FScopedHttpPoolRequestPtr Request(RequestPool.Get());
-		if (Request.IsValid())
-		{
-			FHttpRequest::EResult Result = Request->PerformBlockingPut(Uri, InData, ContentType);
-			ResponseCode = Request->GetResponseCode();
-
-			if (FHttpRequest::IsSuccessResponse(ResponseCode))
-			{
-				TRACE_COUNTER_ADD(ZenDDC_BytesSent, int64(Request->GetBytesSent()));
-				COOK_STAT(Timer.AddHit(Request->GetBytesSent()));
-				return EPutStatus::Cached;
-			}
-
-			if (!ShouldRetryOnError(ResponseCode))
-			{
-				return EPutStatus::NotCached;
-			}
-
-			ResponseCode = 0;
-		}
-	}
-
-	return EPutStatus::NotCached;
-}
-
-FString FZenCacheStore::MakeLegacyZenKey(const TCHAR* CacheKey)
-{
-	FIoHash KeyHash = FIoHash::HashBuffer(CacheKey, FCString::Strlen(CacheKey) * sizeof(TCHAR));
-	if (Namespace == "ue4.ddc")
-	{
-		// TODO: DE:20220525 Remove once we have deployed shared instances with namespace support
-		return FString::Printf(TEXT("/z$/legacy/%s"), *LexToString(KeyHash));
-	}
-	return FString::Printf(TEXT("/z$/%s/legacy/%s"), *Namespace, *LexToString(KeyHash));
-}
-
-void FZenCacheStore::AppendZenUri(const FCacheKey& CacheKey, FStringBuilderBase& Out)
-{
-	Out << TEXT("/z$/") << CacheKey.Bucket << TEXT('/') << CacheKey.Hash;
-}
-
-void FZenCacheStore::AppendZenUri(const FCacheKey& CacheKey, const FValueId& Id, FStringBuilderBase& Out)
-{
-	AppendZenUri(CacheKey, Out);
-	Out << TEXT('/') << Id;
-}
-
-void FZenCacheStore::AppendPolicyQueryString(ECachePolicy Policy, FStringBuilderBase& Uri)
-{
-	if (Policy != ECachePolicy::Default)
-	{
-		Uri << TEXT("?Policy=") << Policy;
-	}
-}
-
-void FZenCacheStore::RemoveCachedData(const TCHAR* CacheKey, bool bTransient)
-{
-	TRACE_CPUPROFILER_EVENT_SCOPE(ZenDDC_Remove);
-	FString Uri = MakeLegacyZenKey(CacheKey);
-
-	int64 ResponseCode = 0; 
-	uint32 Attempts = 0;
-
-	// Retry request until we get an accepted response or exhaust allowed number of attempts.
-	while (ResponseCode == 0 && ++Attempts < MaxAttempts)
-	{
-		FScopedHttpPoolRequestPtr Request(RequestPool.Get());
-		if (Request)
-		{
-			FHttpRequest::EResult Result = Request->PerformBlockingDelete(*Uri);
-			ResponseCode = Request->GetResponseCode();
-
-			if (FHttpRequest::IsSuccessResponse(ResponseCode))
-			{
-				return;
-			}
-
-			if (!ShouldRetryOnError(ResponseCode))
-			{
-				return;
-			}
-
-			ResponseCode = 0;
-		}
-	}
-}
-
-bool FZenCacheStore::IsWritable() const 
-{
-	return true;
-}
-
-FDerivedDataBackendInterface::ESpeedClass 
-FZenCacheStore::GetSpeedClass() const
-{
-	return ESpeedClass::Fast;
-}
-
-TSharedRef<FDerivedDataCacheStatsNode> FZenCacheStore::GatherUsageStats() const
+void FZenCacheStore::LegacyStats(FDerivedDataCacheStatsNode& OutNode)
 {
 	Zen::FZenStats ZenStats;
 
@@ -670,37 +301,27 @@ TSharedRef<FDerivedDataCacheStatsNode> FZenCacheStore::GatherUsageStats() const
 	RemoteStats.PutStats.Accumulate(EHitOrMiss::Hit, ECacheStatType::Bytes, RemotePutSize, /*bIsInGameThread*/ false);
 #endif
 
-	TSharedRef<FDerivedDataCacheStatsNode> LocalNode =
-		MakeShared<FDerivedDataCacheStatsNode>(TEXT("Zen"), ZenService.GetInstance().GetURL(), /*bIsLocal*/ true);
-	LocalNode->UsageStats.Add(TEXT(""), LocalStats);
-
 	if (ZenStats.UpstreamStats.EndPointStats.IsEmpty())
 	{
-		return LocalNode;
+		OutNode = {TEXT("Zen"), ZenService.GetInstance().GetURL(), /*bIsLocal*/ true};
+		OutNode.UsageStats.Add(TEXT(""), LocalStats);
+		return;
 	}
 
+	TSharedRef<FDerivedDataCacheStatsNode> LocalNode =
+		MakeShared<FDerivedDataCacheStatsNode>(TEXT("Zen"), ZenService.GetInstance().GetURL(), /*bIsLocal*/ true);
+	LocalNode->UsageStats.Add(TEXT(""), LocalStats);
+
 	TSharedRef<FDerivedDataCacheStatsNode> RemoteNode =
 		MakeShared<FDerivedDataCacheStatsNode>(ZenStats.UpstreamStats.EndPointStats[0].Name, ZenStats.UpstreamStats.EndPointStats[0].Url, /*bIsLocal*/ false);
 	RemoteNode->UsageStats.Add(TEXT(""), RemoteStats);
 
-	TSharedRef<FDerivedDataCacheStatsNode> GroupNode =
-		MakeShared<FDerivedDataCacheStatsNode>(TEXT("Zen Group"), TEXT(""), /*bIsLocal*/ true);
-	GroupNode->Children.Add(LocalNode);
-	GroupNode->Children.Add(RemoteNode);
-	return GroupNode;
-}
-
-TBitArray<> FZenCacheStore::TryToPrefetch(TConstArrayView<FString> CacheKeys)
-{
-	return CachedDataProbablyExistsBatch(CacheKeys);
-}
-
-bool FZenCacheStore::WouldCache(const TCHAR* CacheKey, TArrayView<const uint8> InData)
-{
-	return true;
+	OutNode = {TEXT("Zen Group"), TEXT(""), /*bIsLocal*/ true};
+	OutNode.Children.Add(LocalNode);
+	OutNode.Children.Add(RemoteNode);
 }
 
-bool FZenCacheStore::ApplyDebugOptions(FBackendDebugOptions& InOptions)
+bool FZenCacheStore::LegacyDebugOptions(FBackendDebugOptions& InOptions)
 {
 	DebugOptions = InOptions;
 	return true;
@@ -1348,16 +969,33 @@ void FZenCacheStore::GetChunks(
 	TRACE_COUNTER_SUBTRACT(ZenDDC_ChunkRequestCount, int64(Requests.Num()));
 }
 
-ILegacyCacheStore* CreateZenCacheStore(const TCHAR* NodeName, const TCHAR* ServiceUrl, const TCHAR* Namespace, const TCHAR* StructuredNamespace)
+TTuple<ILegacyCacheStore*, ECacheStoreFlags> CreateZenCacheStore(const TCHAR* NodeName, const TCHAR* Config)
 {
-	FZenCacheStore* Backend = new FZenCacheStore(ServiceUrl, Namespace, StructuredNamespace);
-	if (Backend->IsUsable())
+	FString ServiceUrl;
+	FParse::Value(Config, TEXT("Host="), ServiceUrl);
+
+	FString Namespace;
+	if (!FParse::Value(Config, TEXT("Namespace="), Namespace))
+	{
+		Namespace = FApp::GetProjectName();
+		UE_LOG(LogDerivedDataCache, Warning, TEXT("%s: Missing required parameter 'Namespace', falling back to '%s'"), NodeName, *Namespace);
+	}
+
+	FString StructuredNamespace;
+	if (!FParse::Value(Config, TEXT("StructuredNamespace="), StructuredNamespace))
+	{
+		StructuredNamespace = Namespace;
+		UE_LOG(LogDerivedDataCache, Warning, TEXT("%s: Missing required parameter 'StructuredNamespace', falling back to Namespace: '%s'"), NodeName, *StructuredNamespace);
+	}
+
+	TUniquePtr<FZenCacheStore> Backend = MakeUnique<FZenCacheStore>(*ServiceUrl, *Namespace, *StructuredNamespace);
+	if (!Backend->IsUsable())
 	{
-		return Backend;
+		UE_LOG(LogDerivedDataCache, Warning, TEXT("%s: Failed to contact the service (%s), will not use it."), NodeName, *Backend->GetName());
+		Backend.Reset();
 	}
-	UE_LOG(LogDerivedDataCache, Warning, TEXT("%s could not contact the service (%s), will not use it."), NodeName, *Backend->GetName());
-	delete Backend;
-	return nullptr;
+
+	return MakeTuple(Backend.Release(), ECacheStoreFlags::Local | ECacheStoreFlags::Remote | ECacheStoreFlags::Query | ECacheStoreFlags::Store);
 }
 
 } // namespace UE::DerivedData
@@ -1367,9 +1005,10 @@ ILegacyCacheStore* CreateZenCacheStore(const TCHAR* NodeName, const TCHAR* Servi
 namespace UE::DerivedData
 {
 
-ILegacyCacheStore* CreateZenCacheStore(const TCHAR* NodeName, const TCHAR* ServiceUrl, const TCHAR* Namespace, const TCHAR* StructuredNamespace)
+TTuple<ILegacyCacheStore*, ECacheStoreFlags> CreateZenCacheStore(const TCHAR* NodeName, const TCHAR* Config)
 {
-	return nullptr;
+	UE_LOG(LogDerivedDataCache, Warning, TEXT("%s: Zen cache is not yet supported in the current build configuration."), NodeName);
+	return MakeTuple(nullptr, ECacheStoreFlags::None);
 }
 
 } // UE::DerivedData
