DDC: Iteration on HttpClient

- Cleaned up some string allocations and conversions.
- Changed code formatting for greater consistency with DDC.

#rnx

diff --git a/Engine/Source/Developer/DerivedDataCache/Private/Http/HttpClient.cpp b/Engine/Source/Developer/DerivedDataCache/Private/Http/HttpClient.cpp
--- a/Engine/Source/Developer/DerivedDataCache/Private/Http/HttpClient.cpp
+++ b/Engine/Source/Developer/DerivedDataCache/Private/Http/HttpClient.cpp
@@ -2,7 +2,8 @@
 
 #include "HttpClient.h"
 
-#if UE_WITH_HTTP_CLIENT
+#include "Misc/AsciiSet.h"
+#include "Misc/ScopeRWLock.h"
 #include "Tasks/Task.h"
 
 #define CURL_NO_OLDIES
@@ -14,26 +15,25 @@
 #include PLATFORM_CURL_INCLUDE
 #else
 #include "curl/curl.h"
-#endif //defined(PLATFORM_CURL_INCLUDE)
+#endif // defined(PLATFORM_CURL_INCLUDE)
 
 #if PLATFORM_MICROSOFT
 #include "Microsoft/HideMicrosoftPlatformTypes.h"
 #endif
 
-
 namespace UE
 {
 
 namespace Http::Private
 {
 
-static constexpr uint32 WaitIntervalMs = ((uint32)(0.01f*1000));
+static constexpr uint32 WaitIntervalMs = 10;
 
 struct FHttpSharedDataInternals
 {
 	CURLSH* CurlShare;
 	CURLM* CurlMulti;
-	UE::TDepletableMpscQueue<CURL*> PendingRequestAdditions;
+	TDepletableMpscQueue<CURL*> PendingRequestAdditions;
 	FRWLock Locks[CURL_LOCK_DATA_LAST];
 	bool WriteLocked[CURL_LOCK_DATA_LAST]{};
 };
@@ -69,32 +69,28 @@ struct FHttpSharedDataStatics
 	}
 };
 
-}
-
+} // Http::Private
 
-FString FHttpAccessToken::GetHeader()
+void FHttpAccessToken::SetToken(const FStringView Token)
 {
-	Lock.ReadLock();
-	FString Header = FString::Printf(TEXT("Authorization: Bearer %s"), *Token);
-	Lock.ReadUnlock();
-	return Header;
+	FWriteScopeLock WriteLock(Lock);
+	const FAnsiStringView Prefix = ANSITEXTVIEW("Authorization: Bearer ");
+	const int32 TokenLen = FPlatformString::ConvertedLength<ANSICHAR>(Token.GetData(), Token.Len());
+	Header.Empty(Prefix.Len() + TokenLen);
+	Header.Append(Prefix.GetData(), Prefix.Len());
+	const int32 TokenIndex = Header.AddUninitialized(TokenLen);
+	FPlatformString::Convert(Header.GetData() + TokenIndex, TokenLen, Token.GetData(), Token.Len());
+	Serial.fetch_add(1, std::memory_order_relaxed);
 }
 
-void FHttpAccessToken::SetHeader(const TCHAR* InToken)
+FAnsiStringBuilderBase& operator<<(FAnsiStringBuilderBase& Builder, const FHttpAccessToken& Token)
 {
-	Lock.WriteLock();
-	Token = InToken;
-	Serial++;
-	Lock.WriteUnlock();
-}
-
-uint32 FHttpAccessToken::GetSerial() const
-{
-	return Serial;
+	FReadScopeLock ReadLock(Token.Lock);
+	return Builder.Append(Token.Header);
 }
 
 FHttpSharedData::FHttpSharedData()
-: PendingRequestEvent(EEventMode::AutoReset)
+	: PendingRequestEvent(EEventMode::AutoReset)
 {
 	Internals = MakePimpl<Http::Private::FHttpSharedDataInternals>();
 	Internals->CurlShare = curl_share_init();
@@ -141,17 +137,17 @@ void FHttpSharedData::ProcessAsyncRequests()
 {
 	int ActiveTransfers = 0;
 
-	auto ProcessPendingRequests = [this, &ActiveTransfers]()
+	auto ProcessPendingRequests = [this, &ActiveTransfers]
 	{
 		int CurrentActiveTransfers = -1;
 
 		do
 		{
 			Internals->PendingRequestAdditions.Deplete([this, &ActiveTransfers](CURL* Curl)
-				{
-					curl_multi_add_handle(Internals->CurlMulti, Curl);
-					++ActiveTransfers;
-				});
+			{
+				curl_multi_add_handle(Internals->CurlMulti, Curl);
+				++ActiveTransfers;
+			});
 
 			curl_multi_perform(Internals->CurlMulti, &CurrentActiveTransfers);
 
@@ -159,7 +155,7 @@ void FHttpSharedData::ProcessAsyncRequests()
 			{
 				for (;;)
 				{
-					int MsgsStillInQueue = 0;	// may use that to impose some upper limit we may spend in that loop
+					int MsgsStillInQueue = 0; // may use that to impose some upper limit we may spend in that loop
 					CURLMsg* Message = curl_multi_info_read(Internals->CurlMulti, &MsgsStillInQueue);
 
 					if (!Message)
@@ -182,7 +178,7 @@ void FHttpSharedData::ProcessAsyncRequests()
 							// It is important that the CompleteAsync call doesn't happen on this thread as it is possible it will block waiting
 							// for a free HTTP request, and if that happens on this thread, we can deadlock as no HTTP requests will become 
 							// available while this thread is blocked.
-							UE::Tasks::Launch(TEXT("FHttpRequest::CompleteAsync"), [CompletedRequest, Result = Message->data.result]() mutable
+							Tasks::Launch(TEXT("FHttpRequest::CompleteAsync"), [CompletedRequest, Result = Message->data.result]() mutable
 							{
 								CompletedRequest->CompleteAsync(Result);
 							});
@@ -211,22 +207,24 @@ void FHttpSharedData::ProcessAsyncRequests()
 	ProcessPendingRequests();
 }
 
-FHttpRequestPool::FHttpRequestPool(FStringView InServiceUrl, FStringView InEffectiveServiceUrl, FHttpAccessToken* InAuthorizationToken, FHttpSharedData* InSharedData, uint32 PoolSize, uint32 InOverflowLimit)
-: ActiveOverflowRequests(0)
-, OverflowLimit(InOverflowLimit)
+FHttpRequestPool::FHttpRequestPool(
+	const FStringView InServiceUrl,
+	FStringView InEffectiveServiceUrl,
+	const FHttpAccessToken* const InAuthorizationToken,
+	FHttpSharedData* const InSharedData,
+	const uint32 PoolSize,
+	const uint32 InOverflowLimit)
+	: ActiveOverflowRequests(0)
+	, OverflowLimit(InOverflowLimit)
 {
-	int32 LastSlash = 0;
-	while (InEffectiveServiceUrl.FindLastChar(TEXT('/'), /* out */ LastSlash) && (LastSlash == (InEffectiveServiceUrl.Len() - 1)))
-	{
-		InEffectiveServiceUrl.LeftChopInline(1);
-	}
-	
+	InEffectiveServiceUrl = FAsciiSet::TrimSuffixWith(InEffectiveServiceUrl, "/");
+
 	Pool.AddUninitialized(PoolSize);
 	Requests.AddUninitialized(PoolSize);
-	for (uint8 i = 0; i < Pool.Num(); ++i)
+	for (int32 Index = 0; Index < Pool.Num(); ++Index)
 	{
-		Pool[i].Usage = 0u;
-		Pool[i].Request = new(&Requests[i]) FHttpRequest(InServiceUrl, InEffectiveServiceUrl, InAuthorizationToken, InSharedData, true);
+		Pool[Index].Usage = 0u;
+		Pool[Index].Request = new(&Requests[Index]) FHttpRequest(InServiceUrl, InEffectiveServiceUrl, InAuthorizationToken, InSharedData, /*bLogErrors*/ true);
 	}
 
 	InitData = MakeUnique<FInitData>(InServiceUrl, InEffectiveServiceUrl, InAuthorizationToken, InSharedData);
@@ -235,24 +233,24 @@ FHttpRequestPool::FHttpRequestPool(FStringView InServiceUrl, FStringView InEffec
 FHttpRequestPool::~FHttpRequestPool()
 {
 	check(ActiveOverflowRequests.load() == 0);
-	for (uint8 i = 0; i < Pool.Num(); ++i)
+	for (const FEntry& Entry : Pool)
 	{
 		// No requests should be in use by now.
-		check(Pool[i].Usage.load(std::memory_order_acquire) == 0u);
+		check(Entry.Usage.load(std::memory_order_acquire) == 0u);
 	}
 }
 
 FHttpRequest* FHttpRequestPool::GetFreeRequest(bool bUnboundedOverflow)
 {
-	for (uint8 i = 0; i < Pool.Num(); ++i)
+	for (FEntry& Entry : Pool)
 	{
-		if (!Pool[i].Usage.load(std::memory_order_relaxed))
+		if (!Entry.Usage.load(std::memory_order_relaxed))
 		{
 			uint8 Expected = 0u;
-			if (Pool[i].Usage.compare_exchange_strong(Expected, 1u))
+			if (Entry.Usage.compare_exchange_strong(Expected, 1u))
 			{
-				Pool[i].Request->Reset();
-				return Pool[i].Request;
+				Entry.Request->Reset();
+				return Entry.Request;
 			}
 		}
 	}
@@ -264,7 +262,7 @@ FHttpRequest* FHttpRequestPool::GetFreeRequest(bool bUnboundedOverflow)
 		{
 			// Create an overflow request (outside of the pre-allocated range of requests)
 			ActiveOverflowRequests.fetch_add(1, std::memory_order_relaxed);
-			return new FHttpRequest(*InitData->ServiceUrl, *InitData->EffectiveServiceUrl, InitData->AccessToken, InitData->SharedData, true);
+			return new FHttpRequest(InitData->ServiceUrl, InitData->EffectiveServiceUrl, InitData->AccessToken, InitData->SharedData, true);
 		}
 	}
 	return nullptr;
@@ -316,13 +314,13 @@ void FHttpRequestPool::ReleaseRequestToPool(FHttpRequest* Request)
 		return;
 	}
 
-	for (uint8 i = 0; i < Pool.Num(); ++i)
+	for (FEntry& Entry : Pool)
 	{
-		if (Pool[i].Request == Request)
+		if (Entry.Request == Request)
 		{
 			// If only 1 user is remaining, we can give it to a waiter
 			// instead of releasing it back to the pool.
-			if (Pool[i].Usage == 1u)
+			if (Entry.Usage == 1u)
 			{
 				if (FWaiter* Waiter = Waiters.dequeue())
 				{
@@ -333,33 +331,30 @@ void FHttpRequestPool::ReleaseRequestToPool(FHttpRequest* Request)
 				}
 			}
 			
-			Pool[i].Usage--;
+			Entry.Usage--;
 			return;
 		}
 	}
-	check(false);
+
+	checkNoEntry();
 }
 
 void FHttpRequestPool::MakeRequestShared(FHttpRequest* Request, uint8 Users)
 {
-	if ((Request < Requests.GetData()) || (Request >= (Requests.GetData() + Requests.Num())))
-	{
-		// Overflow requests (outside of the pre-allocated range of requests), cannot be made shared
-		check(false);
-	}
+	// Overflow requests (outside of the pre-allocated range of requests), cannot be made shared
+	check((Request >= Requests.GetData()) && (Request < (Requests.GetData() + Requests.Num())));
 
 	check(Users != 0);
-	for (uint8 i = 0; i < Pool.Num(); ++i)
+	for (FEntry& Entry : Pool)
 	{
-		if (Pool[i].Request == Request)
+		if (Entry.Request == Request)
 		{
-			Pool[i].Usage = Users;
+			Entry.Usage = Users;
 			return;
 		}
 	}
-	check(false);
+
+	checkNoEntry();
 }
 
 } // UE
-
-#endif // UE_WITH_HTTP_CLIENT
diff --git a/Engine/Source/Developer/DerivedDataCache/Private/Http/HttpClient.h b/Engine/Source/Developer/DerivedDataCache/Private/Http/HttpClient.h
--- a/Engine/Source/Developer/DerivedDataCache/Private/Http/HttpClient.h
+++ b/Engine/Source/Developer/DerivedDataCache/Private/Http/HttpClient.h
@@ -2,18 +2,12 @@
 
 #pragma once
 
-#include "HAL/Platform.h"
-#define UE_WITH_HTTP_CLIENT PLATFORM_DESKTOP
-
-#if UE_WITH_HTTP_CLIENT
 #include "CoreTypes.h"
 #include "Containers/Array.h"
 #include "Containers/ArrayView.h"
 #include "Containers/DepletableMpscQueue.h"
+#include "Containers/StringFwd.h"
 #include "Containers/UnrealString.h"
-#include "DerivedDataRequest.h"
-#include "DerivedDataRequestOwner.h"
-#include "DerivedDataRequestTypes.h"
 #include "Dom/JsonObject.h"
 #include "Experimental/Async/LazyEvent.h"
 #include "Experimental/Containers/FAAArrayQueue.h"
@@ -23,15 +17,20 @@
 #include "Memory/CompositeBuffer.h"
 #include "Memory/MemoryView.h"
 #include "Templates/PimplPtr.h"
+#include "Templates/RefCounting.h"
 #include "Templates/SharedPointer.h"
+#include "Templates/UniquePtr.h"
 #include <atomic>
 
-
 struct curl_slist;
+namespace UE { class FHttpRequest; }
+namespace UE::DerivedData { class IRequestOwner; }
 namespace UE::Http::Private { struct FAsyncRequestData; }
+namespace UE::Http::Private { struct FCurlStringListDeleter; }
 namespace UE::Http::Private { struct FHttpRequestStatics; }
 namespace UE::Http::Private { struct FHttpSharedDataInternals; }
 namespace UE::Http::Private { struct FHttpSharedDataStatics; }
+namespace UE::Http::Private { using FCurlStringList = TUniquePtr<curl_slist, FCurlStringListDeleter>; }
 
 namespace UE
 {
@@ -93,14 +92,14 @@ inline FStringView GetHttpMimeType(EHttpContentType Type)
 class FHttpAccessToken
 {
 public:
-	FHttpAccessToken() = default;
-	FString GetHeader();
-	void SetHeader(const TCHAR*);
-	uint32 GetSerial() const;
+	void SetToken(FStringView Token);
+	inline uint32 GetSerial() const { return Serial.load(std::memory_order_relaxed); }
+	friend FAnsiStringBuilderBase& operator<<(FAnsiStringBuilderBase& Builder, const FHttpAccessToken& Token);
+
 private:
-	FRWLock Lock;
-	FString Token;
-	uint32 Serial;
+	mutable FRWLock Lock;
+	TArray<ANSICHAR> Header;
+	std::atomic<uint32> Serial;
 };
 
 class FHttpSharedData
@@ -108,20 +107,21 @@ class FHttpSharedData
 public:
 	FHttpSharedData();
 	~FHttpSharedData();
+
 private:
-	friend class FHttpRequest;
-	friend struct Http::Private::FHttpSharedDataStatics;
+	void ProcessAsyncRequests();
+	void AddRequest(void* Curl);
+	void* GetCurlShare() const;
 
-	UE::FLazyEvent PendingRequestEvent;
+	friend FHttpRequest;
+	friend Http::Private::FHttpSharedDataStatics;
+
+	FLazyEvent PendingRequestEvent;
 	FThread AsyncServiceThread;
 	std::atomic<bool> bAsyncThreadStarting = false;
 	std::atomic<bool> bAsyncThreadShutdownRequested = false;
 
 	TPimplPtr<Http::Private::FHttpSharedDataInternals> Internals;
-
-	void ProcessAsyncRequests();
-	void AddRequest(void* Curl);
-	void* GetCurlShare() const;
 };
 
 
@@ -168,7 +168,12 @@ public:
 
 	using FOnHttpRequestComplete = TUniqueFunction<ECompletionBehavior(EResult HttpResult, FHttpRequest* Request)>;
 
-	FHttpRequest(FStringView InDomain, FStringView InEffectiveDomain, FHttpAccessToken* InAuthorizationToken, FHttpSharedData* InSharedData, bool bInLogErrors);
+	FHttpRequest(
+		FStringView InDomain,
+		FStringView InEffectiveDomain,
+		const FHttpAccessToken* InAuthorizationToken,
+		FHttpSharedData* InSharedData,
+		bool bInLogErrors);
 	~FHttpRequest();
 
 	/**
@@ -226,26 +231,29 @@ public:
 	 * @param Buffer Data to upload
 	 * @return Result of the request
 	 */
-	EResult PerformBlockingPut(FStringView Uri,
-		FCompositeBuffer Buffer,
+	EResult PerformBlockingPut(
+		FStringView Uri,
+		const FCompositeBuffer& Buffer,
 		EHttpContentType ContentType = EHttpContentType::UnspecifiedContentType,
 		TConstArrayView<long> ExpectedErrorCodes = {});
 	EResult PerformBlockingPost(FStringView Uri,
-		FCompositeBuffer Buffer,
+		const FCompositeBuffer& Buffer,
 		EHttpContentType ContentType = EHttpContentType::UnspecifiedContentType,
 		EHttpContentType AcceptType = EHttpContentType::UnspecifiedContentType,
 		TConstArrayView<long> ExpectedErrorCodes = {});
-	void EnqueueAsyncPut(UE::DerivedData::IRequestOwner& Owner,
+	void EnqueueAsyncPut(
+		DerivedData::IRequestOwner& Owner,
 		FHttpRequestPool* Pool,
 		FStringView Uri,
-		FCompositeBuffer Buffer,
+		const FCompositeBuffer& Buffer,
 		FOnHttpRequestComplete&& OnComplete,
 		EHttpContentType ContentType = EHttpContentType::UnspecifiedContentType,
 		TConstArrayView<long> ExpectedErrorCodes = {});
-	void EnqueueAsyncPost(UE::DerivedData::IRequestOwner& Owner,
+	void EnqueueAsyncPost(
+		DerivedData::IRequestOwner& Owner,
 		FHttpRequestPool* Pool,
 		FStringView Uri,
-		FCompositeBuffer Buffer,
+		const FCompositeBuffer& Buffer,
 		FOnHttpRequestComplete&& OnComplete,
 		EHttpContentType ContentType = EHttpContentType::UnspecifiedContentType,
 		EHttpContentType AcceptType = EHttpContentType::UnspecifiedContentType,
@@ -258,11 +266,13 @@ public:
 	 * be stored in an internal buffer and accessed GetResponse* methods.
 	 * @return Result of the request
 	 */
-	EResult PerformBlockingDownload(FStringView Uri,
+	EResult PerformBlockingDownload(
+		FStringView Uri,
 		TArray64<uint8>* Buffer,
 		EHttpContentType AcceptType = EHttpContentType::UnspecifiedContentType,
 		TConstArrayView<long> ExpectedErrorCodes = {400});
-	void EnqueueAsyncDownload(UE::DerivedData::IRequestOwner& Owner,
+	void EnqueueAsyncDownload(
+		DerivedData::IRequestOwner& Owner,
 		FHttpRequestPool* Pool,
 		FStringView Uri,
 		FOnHttpRequestComplete&& OnComplete,
@@ -274,18 +284,22 @@ public:
 	 * @param Uri Url to use.
 	 * @return Result of the request
 	 */
-	EResult PerformBlockingHead(FStringView Uri,
+	EResult PerformBlockingHead(
+		FStringView Uri,
 		EHttpContentType AcceptType = EHttpContentType::UnspecifiedContentType,
 		TConstArrayView<long> ExpectedErrorCodes = {400});
-	EResult PerformBlockingDelete(FStringView Uri,
+	EResult PerformBlockingDelete(
+		FStringView Uri,
 		TConstArrayView<long> ExpectedErrorCodes = {400});
-	void EnqueueAsyncHead(UE::DerivedData::IRequestOwner& Owner,
+	void EnqueueAsyncHead(
+		DerivedData::IRequestOwner& Owner,
 		FHttpRequestPool* Pool,
 		FStringView Uri,
 		FOnHttpRequestComplete&& OnComplete,
 		EHttpContentType AcceptType = EHttpContentType::UnspecifiedContentType,
 		TConstArrayView<long> ExpectedErrorCodes = {400});
-	void EnqueueAsyncDelete(UE::DerivedData::IRequestOwner& Owner,
+	void EnqueueAsyncDelete(
+		DerivedData::IRequestOwner& Owner,
 		FHttpRequestPool* Pool,
 		FStringView Uri,
 		FOnHttpRequestComplete&& OnComplete,
@@ -299,7 +313,7 @@ public:
 	/**
 	 * Attempts to find the header from the response. Returns false if header is not present.
 	 */
-	bool GetHeader(const ANSICHAR* Header, FString& OutValue) const;
+	bool GetHeader(FAnsiStringView Header, FString& OutValue) const;
 
 	/**
 	 * Returns the response buffer. Note that is the request is performed
@@ -378,10 +392,10 @@ private:
 	TArray<FString> Headers;
 	FString Domain;
 	FString EffectiveDomain;
-	FHttpAccessToken* AuthorizationToken;
+	const FHttpAccessToken* AuthorizationToken;
 
 	void AddContentTypeHeader(FStringView Header, EHttpContentType Type);
-	curl_slist* PrepareToIssueRequest(FStringView Uri, ERequestVerb Verb, uint64 ContentLength);
+	Http::Private::FCurlStringList PrepareToIssueRequest(FStringView Uri, ERequestVerb Verb, uint64 ContentLength);
 
 	/**
 	 * Performs the request, blocking until finished.
@@ -393,20 +407,25 @@ private:
 	 */
 	EResult PerformBlocking(FStringView Uri, ERequestVerb Verb, uint64 ContentLength, TConstArrayView<long> ExpectedErrorCodes);
 
-	void EnqueueAsync(UE::DerivedData::IRequestOwner& Owner, FHttpRequestPool* Pool, FStringView Uri, ERequestVerb Verb, uint64 ContentLength, FOnHttpRequestComplete&& OnComplete, TConstArrayView<long> ExpectedErrorCodes);
+	void EnqueueAsync(
+		DerivedData::IRequestOwner& Owner,
+		FHttpRequestPool* Pool,
+		FStringView Uri,
+		ERequestVerb Verb,
+		uint64 ContentLength,
+		FOnHttpRequestComplete&& OnComplete,
+		TConstArrayView<long> ExpectedErrorCodes);
 
 	void LogResult(FResultCode EResult, FStringView Uri, ERequestVerb Verb, TConstArrayView<long> ExpectedErrorCodes) const;
 
-	FString GetAnsiBufferAsString(const TArray64<uint8>& Buffer) const
+	static FString GetAnsiBufferAsString(TConstArrayView64<uint8> Buffer)
 	{
 		// Content is NOT null-terminated; we need to specify lengths here
-		FUTF8ToTCHAR TCHARData(reinterpret_cast<const ANSICHAR*>(Buffer.GetData()), Buffer.Num());
-		return FString(TCHARData.Length(), TCHARData.Get());
+		static_assert(sizeof(UTF8CHAR) == sizeof(uint8));
+		return FString(Buffer.Num(), reinterpret_cast<const UTF8CHAR*>(Buffer.GetData()));
 	}
-
 };
 
-
 /**
  * Pool that manages a fixed set of requests. Users are required to release requests that have been 
  * acquired. Usable with \ref FScopedRequestPtr which handles this automatically.
@@ -414,7 +433,13 @@ private:
 class FHttpRequestPool
 {
 public:
-	FHttpRequestPool(FStringView InServiceUrl, FStringView InEffectiveServiceUrl, FHttpAccessToken* InAuthorizationToken, FHttpSharedData* InSharedData, uint32 PoolSize, uint32 InOverflowLimit = 0);
+	FHttpRequestPool(
+		FStringView InServiceUrl,
+		FStringView InEffectiveServiceUrl,
+		const FHttpAccessToken* InAuthorizationToken,
+		FHttpSharedData* InSharedData,
+		uint32 PoolSize,
+		uint32 InOverflowLimit = 0);
 	~FHttpRequestPool();
 
 	/**
@@ -427,7 +452,7 @@ public:
 	class FWaiter : public FThreadSafeRefCountedObject
 	{
 	public:
-		std::atomic<FHttpRequest*> Request{ nullptr };
+		std::atomic<FHttpRequest*> Request = nullptr;
 
 		FWaiter(FHttpRequestPool* InPool)
 			: Event(FPlatformProcess::GetSynchEventFromPool(true))
@@ -444,6 +469,7 @@ public:
 		{
 			Event->Trigger();
 		}
+
 	private:
 		~FWaiter()
 		{
@@ -485,29 +511,32 @@ private:
 		FHttpRequest* Request;
 	};
 
-	TArray<FEntry> Pool;
-	TArray<FHttpRequest> Requests;
-	FAAArrayQueue<FWaiter> Waiters;
-	std::atomic<uint32> ActiveOverflowRequests;
 	struct FInitData
 	{
 		FString ServiceUrl;
 		FString EffectiveServiceUrl;
-		FHttpAccessToken* AccessToken;
+		const FHttpAccessToken* AccessToken;
 		FHttpSharedData* SharedData;
 
-		FInitData(FStringView InServiceUrl, FStringView InEffectiveServiceUrl, FHttpAccessToken* InAccessToken, FHttpSharedData* InSharedData)
-		: ServiceUrl(InServiceUrl)
-		, EffectiveServiceUrl(InEffectiveServiceUrl)
-		, AccessToken(InAccessToken)
-		, SharedData(InSharedData)
+		FInitData(
+			FStringView InServiceUrl,
+			FStringView InEffectiveServiceUrl,
+			const FHttpAccessToken* InAccessToken,
+			FHttpSharedData* InSharedData)
+			: ServiceUrl(InServiceUrl)
+			, EffectiveServiceUrl(InEffectiveServiceUrl)
+			, AccessToken(InAccessToken)
+			, SharedData(InSharedData)
 		{
 		}
 	};
+
+	TArray<FEntry> Pool;
+	TArray<FHttpRequest> Requests;
+	FAAArrayQueue<FWaiter> Waiters;
+	std::atomic<uint32> ActiveOverflowRequests;
 	TUniquePtr<const FInitData> InitData;
 	const uint32 OverflowLimit;
-
-	FHttpRequestPool() = delete;
 };
 
 /**
@@ -529,7 +558,7 @@ public:
 		Pool->ReleaseRequestToPool(Request);
 	}
 
-	bool IsValid() const 
+	bool IsValid() const
 	{
 		return Request != nullptr;
 	}
@@ -554,5 +583,3 @@ private:
 };
 
 } // UE
-
-#endif // UE_WITH_HTTP_CLIENT
diff --git a/Engine/Source/Developer/DerivedDataCache/Private/Http/HttpRequest.cpp b/Engine/Source/Developer/DerivedDataCache/Private/Http/HttpRequest.cpp
--- a/Engine/Source/Developer/DerivedDataCache/Private/Http/HttpRequest.cpp
+++ b/Engine/Source/Developer/DerivedDataCache/Private/Http/HttpRequest.cpp
@@ -2,10 +2,12 @@
 
 #include "HttpClient.h"
 
-#if UE_WITH_HTTP_CLIENT
 #include "DerivedDataLegacyCacheStore.h"
+#include "DerivedDataRequest.h"
+#include "DerivedDataRequestOwner.h"
 #include "HAL/IConsoleManager.h"
 #include "Logging/LogMacros.h"
+#include "Misc/AsciiSet.h"
 #include "Serialization/CompactBinary.h"
 #include "Serialization/JsonReader.h"
 #include "Serialization/JsonSerializer.h"
@@ -24,7 +26,7 @@
 #include PLATFORM_CURL_INCLUDE
 #else
 #include "curl/curl.h"
-#endif //defined(PLATFORM_CURL_INCLUDE)
+#endif // defined(PLATFORM_CURL_INCLUDE)
 
 #if PLATFORM_MICROSOFT
 #include "Microsoft/HideMicrosoftPlatformTypes.h"
@@ -37,7 +39,7 @@ static bool bHttpEnableAsync = true;
 static FAutoConsoleVariableRef CVarHttpEnableAsync(
 	TEXT("DDC.Http.EnableAsync"),
 	bHttpEnableAsync,
-	TEXT("If true, async operations are permitted, otherwise all operations are forced to be synchronous."),
+	TEXT("If true, asynchronous operations are permitted, otherwise all operations are forced to be synchronous."),
 	ECVF_Default);
 
 namespace Http::Private
@@ -45,32 +47,39 @@ namespace Http::Private
 
 static constexpr bool RequestDebugEnabled = false;
 static constexpr bool RequestTimeoutEnabled = true;
-static constexpr long RequestTimeoutSeconds = 30L;
-static constexpr size_t RequestMaxBufferReserve = 104857600u;
+static constexpr long RequestTimeoutSeconds = 30;
+static constexpr size_t RequestMaxBufferReserve = 100 * 1024 * 1024; // 100 MiB
+
+struct FCurlStringListDeleter
+{
+	void operator()(curl_slist* List) const
+	{
+		if (List)
+		{
+			curl_slist_free_all(List);
+		}
+	}
+};
 
 struct FAsyncRequestData final : public DerivedData::FRequestBase
 {
-	UE::DerivedData::IRequestOwner* Owner = nullptr;
-	UE::FHttpRequestPool* Pool = nullptr;
-	curl_slist* CurlHeaders = nullptr;
+	DerivedData::IRequestOwner* Owner = nullptr;
+	FHttpRequestPool* Pool = nullptr;
+	FCurlStringList CurlHeaders;
 	FString Uri;
-	UE::FHttpRequest::ERequestVerb Verb;
+	FHttpRequest::ERequestVerb Verb;
 	TArray<long, TInlineAllocator<4>> ExpectedErrorCodes;
-	UE::FHttpRequest::FOnHttpRequestComplete OnComplete;
-	UE::FLazyEvent Event {EEventMode::ManualReset};
+	FHttpRequest::FOnHttpRequestComplete OnComplete;
+	FLazyEvent Event {EEventMode::ManualReset};
 
 	void Reset()
 	{
-		if (CurlHeaders)
-		{
-			curl_slist_free_all(CurlHeaders);
-			CurlHeaders = nullptr;
-		}
+		CurlHeaders.Reset();
 		Uri.Empty();
 		ExpectedErrorCodes.Empty();
 	}
 
-	void SetPriority(UE::DerivedData::EPriority Priority) final {}
+	void SetPriority(DerivedData::EPriority Priority) final {}
 
 	void Cancel() final
 	{
@@ -112,7 +121,7 @@ struct FHttpRequestStatics
 		return PreverifyOk;
 	}
 
-	static UE::FHttpRequest::FResultCode StaticSSLCTXFn(CURL* curl, void* sslctx, void* parm)
+	static FHttpRequest::FResultCode StaticSSLCTXFn(CURL* Handle, void* sslctx, void* parm)
 	{
 		SSL_CTX* Context = static_cast<SSL_CTX*>(sslctx);
 		const ISslCertificateManager& CertificateManager = FSslModule::Get().GetCertificateManager();
@@ -126,7 +135,7 @@ struct FHttpRequestStatics
 	}
 	#endif //#if WITH_SSL
 
-	static size_t StaticDebugCallback(CURL * Handle, curl_infotype DebugInfoType, char * DebugInfo, size_t DebugInfoSize, void* UserData)
+	static size_t StaticDebugCallback(CURL* Handle, curl_infotype DebugInfoType, char* DebugInfo, size_t DebugInfoSize, void* UserData)
 	{
 		FHttpRequest* Request = static_cast<FHttpRequest*>(UserData);
 
@@ -142,13 +151,12 @@ struct FHttpRequestStatics
 			char* FoundNulPtr = (char*)memchr(DebugInfo, 0, DebugInfoSize);
 			int CalculatedSize = FoundNulPtr != nullptr ? FoundNulPtr - DebugInfo : DebugInfoSize;
 
-			auto ConvertedString = StringCast<TCHAR>(static_cast<const ANSICHAR*>(DebugInfo), CalculatedSize);
-			FString DebugText(ConvertedString.Length(), ConvertedString.Get());
+			FString DebugText(CalculatedSize, static_cast<const ANSICHAR*>(DebugInfo));
 			DebugText.ReplaceInline(TEXT("\n"), TEXT(""), ESearchCase::CaseSensitive);
 			DebugText.ReplaceInline(TEXT("\r"), TEXT(""), ESearchCase::CaseSensitive);
 			UE_LOG(LogDerivedDataCache, VeryVerbose, TEXT("%s: %p: '%s'"), *Request->GetName(), Request, *DebugText);
+			break;
 		}
-		break;
 
 		case CURLINFO_HEADER_IN:
 			UE_LOG(LogDerivedDataCache, VeryVerbose, TEXT("%s: %p: Received header (%d bytes)"), *Request->GetName(), Request, DebugInfoSize);
@@ -264,9 +272,14 @@ struct FHttpRequestStatics
 	}
 };
 
-}
+} // Http::Private
 
-FHttpRequest::FHttpRequest(FStringView InDomain, FStringView InEffectiveDomain, FHttpAccessToken* InAuthorizationToken, FHttpSharedData* InSharedData, bool bInLogErrors)
+FHttpRequest::FHttpRequest(
+	const FStringView InDomain,
+	const FStringView InEffectiveDomain,
+	const FHttpAccessToken* const InAuthorizationToken,
+	FHttpSharedData* const InSharedData,
+	const bool bInLogErrors)
 	: SharedData(InSharedData)
 	, AsyncData(nullptr)
 	, bLogErrors(bInLogErrors)
@@ -298,12 +311,12 @@ void FHttpRequest::Reset()
 	Attempts = 0;
 	ResultCode = CURL_LAST;
 
-	CURL* LocalCurl = static_cast<CURL*>(Curl);
+	CURL* const LocalCurl = static_cast<CURL*>(Curl);
 	curl_easy_reset(LocalCurl);
 	check(!AsyncData);
 
 	// Options that are always set for all connections.
-	if constexpr(Http::Private::RequestTimeoutEnabled)
+	if constexpr (Http::Private::RequestTimeoutEnabled)
 	{
 		curl_easy_setopt(LocalCurl, CURLOPT_CONNECTTIMEOUT, Http::Private::RequestTimeoutSeconds);
 	}
@@ -355,54 +368,53 @@ void FHttpRequest::PrepareToRetry()
 	++Attempts;
 }
 
-FHttpRequest::EResult FHttpRequest::PerformBlockingDownload(FStringView Uri,
-	TArray64<uint8>* Buffer,
-	EHttpContentType AcceptType,
-	TConstArrayView<long> ExpectedErrorCodes)
+FHttpRequest::EResult FHttpRequest::PerformBlockingDownload(
+	const FStringView Uri,
+	TArray64<uint8>* const Buffer,
+	const EHttpContentType AcceptType,
+	const TConstArrayView<long> ExpectedErrorCodes)
 {
 	curl_easy_setopt(static_cast<CURL*>(Curl), CURLOPT_HTTPGET, 1L);
 	WriteDataBufferPtr = Buffer;
-		
+
 	AddContentTypeHeader(TEXTVIEW("Accept"), AcceptType);
 
 	return PerformBlocking(Uri, ERequestVerb::Get, 0u, ExpectedErrorCodes);
 }
 
-void FHttpRequest::EnqueueAsyncDownload(UE::DerivedData::IRequestOwner& Owner,
-	FHttpRequestPool* Pool,
-	FStringView Uri,
+void FHttpRequest::EnqueueAsyncDownload(
+	DerivedData::IRequestOwner& Owner,
+	FHttpRequestPool* const Pool,
+	const FStringView Uri,
 	FOnHttpRequestComplete&& OnComplete,
-	EHttpContentType AcceptType,
-	TConstArrayView<long> ExpectedErrorCodes)
+	const EHttpContentType AcceptType,
+	const TConstArrayView<long> ExpectedErrorCodes)
 {
 	curl_easy_setopt(static_cast<CURL*>(Curl), CURLOPT_HTTPGET, 1L);
 
 	return EnqueueAsync(Owner, Pool, Uri, ERequestVerb::Get, 0u, MoveTemp(OnComplete), ExpectedErrorCodes);
 }
 
-void FHttpRequest::AddHeader(FStringView Header, FStringView Value)
+void FHttpRequest::AddHeader(const FStringView Header, const FStringView Value)
 {
 	check(ResultCode == CURL_LAST); // Cannot set header after request is sent
-	TStringBuilder<128> Sb;
-	Sb << Header << TEXTVIEW(": ") << Value;
-	Headers.Emplace(Sb);
+
+	Headers.Emplace(WriteToString<128>(Header, TEXTVIEW(": "), Value));
 }
 
-bool FHttpRequest::GetHeader(const ANSICHAR* Header, FString& OutValue) const
+bool FHttpRequest::GetHeader(const FAnsiStringView Header, FString& OutValue) const
 {
-	check(ResultCode != CURL_LAST);  // Cannot query headers before request is sent
-
-	const ANSICHAR* HeadersBuffer = (const ANSICHAR*) ResponseHeader.GetData();
-	size_t HeaderLen = strlen(Header);
+	check(ResultCode != CURL_LAST); // Cannot query headers before request is sent
 
-	// Find the header key in the (ANSI) response buffer. If not found we can exist immediately
-	if (const ANSICHAR* Found = FCStringAnsi::Stristr(HeadersBuffer, Header))
+	// Find the header key in the (UTF-8) response buffer. If not found we can exit immediately.
+	const FUtf8StringView LocalHeaders((const UTF8CHAR*)ResponseHeader.GetData(), ResponseHeader.Num());
+	if (const int32 HeaderIndex = String::FindFirst(LocalHeaders, Header, ESearchCase::IgnoreCase); HeaderIndex != INDEX_NONE)
 	{
-		const ANSICHAR* Linebreak = strchr(Found, '\r');
-		const ANSICHAR* ValueStart = Found + HeaderLen + 2; //colon and space
-		const size_t ValueSize = Linebreak - ValueStart;
-		FUTF8ToTCHAR TCHARData(ValueStart, ValueSize);
-		OutValue = FString(TCHARData.Length(), TCHARData.Get());
+		const FUtf8StringView HeaderToEnd = LocalHeaders.RightChop(HeaderIndex);
+		const int32 LineBreakIndex = String::FindFirstChar(HeaderToEnd, '\r');
+		const int32 ValueIndex = Header.Len() + 2; // colon and space
+		const FUtf8StringView Value = HeaderToEnd.Mid(ValueIndex, LineBreakIndex - ValueIndex);
+		OutValue = FString(Value);
 		return true;
 	}
 	return false;
@@ -414,11 +426,7 @@ TSharedPtr<FJsonObject> FHttpRequest::GetResponseAsJsonObject() const
 
 	TSharedPtr<FJsonObject> JsonObject;
 	TSharedRef<TJsonReader<> > JsonReader = TJsonReaderFactory<>::Create(Response);
-	if (!FJsonSerializer::Deserialize(JsonReader, JsonObject) || !JsonObject.IsValid())
-	{
-		return TSharedPtr<FJsonObject>(nullptr);
-	}
-
+	FJsonSerializer::Deserialize(JsonReader, JsonObject);
 	return JsonObject;
 }
 
@@ -446,14 +454,12 @@ void FHttpRequest::CompleteAsync(FResultCode Result)
 {
 	ResultCode = Result;
 
-	// Get response code
 	curl_easy_getinfo(static_cast<CURL*>(Curl), CURLINFO_RESPONSE_CODE, &ResponseCode);
 
 	LogResult(Result, *AsyncData->Uri, AsyncData->Verb, AsyncData->ExpectedErrorCodes);
 
 	ECompletionBehavior Behavior;
 	{
-		UE::DerivedData::FRequestBarrier Barrier(*AsyncData->Owner, UE::DerivedData::ERequestBarrierFlags::Priority);
 		FHttpRequest::EResult HttpResult;
 		switch (ResultCode)
 		{
@@ -467,29 +473,27 @@ void FHttpRequest::CompleteAsync(FResultCode Result)
 			HttpResult = EResult::Failed;
 			break;
 		}
+		DerivedData::FRequestBarrier Barrier(*AsyncData->Owner, DerivedData::ERequestBarrierFlags::Priority);
 		Behavior = AsyncData->OnComplete(HttpResult, this);
 	}
 
-	if (Behavior == ECompletionBehavior::Retry)
+	if (Behavior == ECompletionBehavior::Retry && !AsyncData->Owner->IsCanceled())
 	{
 		PrepareToRetry();
 		SharedData->AddRequest(Curl);
 	}
 	else
 	{
-		// Clean up
-		curl_slist_free_all(AsyncData->CurlHeaders);
-
 		AsyncData->Owner->End(AsyncData, [this, Behavior]
+		{
+			AsyncData->Event.Trigger();
+			FHttpRequestPool* Pool = AsyncData->Pool;
+			AsyncData = nullptr;
+			if (Pool)
 			{
-				AsyncData->Event.Trigger();
-				FHttpRequestPool* Pool = AsyncData->Pool;
-				AsyncData = nullptr;
-				if (Pool)
-				{
-					Pool->ReleaseRequestToPool(this);
-				}
-			});
+				Pool->ReleaseRequestToPool(this);
+			}
+		});
 	}
 }
 
@@ -509,94 +513,89 @@ static const char* GetSessionIdHeader()
 	{
 		static TAnsiStringBuilder<64> SessionIdHeader;
 		SessionIdHeader << "UE-Session: " << SessionId;
-		return SessionIdHeader.GetData(); 
+		return SessionIdHeader.GetData();
 	}();
 
 	return HeaderString;
 }
 
-static std::atomic<int> GRequestId{1};
+static std::atomic<int> GRequestId = 1;
 
-curl_slist* FHttpRequest::PrepareToIssueRequest(FStringView Uri, ERequestVerb Verb, uint64 ContentLength)
+Http::Private::FCurlStringList FHttpRequest::PrepareToIssueRequest(FStringView Uri, ERequestVerb Verb, uint64 ContentLength)
 {
-	// Strip any leading slashes because we compose the prefix and the suffix with a separating slash below
-	while (Uri.StartsWith(TEXT('/')))
-	{
-		Uri.RightChopInline(1);
-	}
-
-	TAnsiStringBuilder<32> RequestIdHeader;
-	RequestIdHeader << "UE-Request: " << GRequestId.fetch_add(1, std::memory_order_relaxed);
-
-	const char* CommonHeaders[] =
-	{
-		GetSessionIdHeader(),
-		RequestIdHeader.GetData(),
-		"User-Agent: Unreal Engine",
-		// Strip any Expect: 100-Continue header since this just introduces latency
-		"Expect:",	
-		nullptr
-	};
-
-	// Setup request options
+	// Strip leading slashes because the separating slash is added below.
+	Uri = FAsciiSet::TrimPrefixWith(Uri, "/");
 	curl_easy_setopt(static_cast<CURL*>(Curl), CURLOPT_URL, *WriteToUtf8String<512>(EffectiveDomain, '/', Uri));
 
-	// Setup response header buffer. If caller has not setup a response data buffer, use internal.
+	// Set up response header buffer.
 	WriteHeaderBufferPtr = &ResponseHeader;
+
+	// Use internal response data buffer if not provided by the caller.
 	if (WriteDataBufferPtr == nullptr)
 	{
 		WriteDataBufferPtr = &ResponseBuffer;
 	}
 
-	if ((Verb != ERequestVerb::Delete) && (Verb != ERequestVerb::Get))
+	// Set up headers.
+	curl_slist* CurlHeaders = nullptr;
+	const auto AddHeader = [&CurlHeaders](const ANSICHAR* Header)
 	{
-		Headers.Add(FString::Printf(TEXT("Content-Length: %d"), ContentLength));
-	}
+		CurlHeaders = curl_slist_append(CurlHeaders, Header);
+	};
 
-	// And auth token if it's set
-	if (AuthorizationToken)
+	AddHeader(GetSessionIdHeader());
+	AddHeader(*WriteToAnsiString<32>(ANSITEXTVIEW("UE-Request: "), GRequestId.fetch_add(1, std::memory_order_relaxed)));
+	AddHeader("User-Agent: Unreal Engine");
+
+	// Strip any Expect: 100-Continue header since this just introduces latency.
+	AddHeader("Expect:");
+
+	for (const FString& Header : Headers)
 	{
-		Headers.Add(AuthorizationToken->GetHeader());
+		AddHeader(TCHAR_TO_ANSI(*Header));
 	}
 
-	// Build headers list
-	curl_slist* CurlHeaders = nullptr;
-	// Add common headers
-	for (uint8 i = 0; CommonHeaders[i] != nullptr; ++i)
+	if ((Verb != ERequestVerb::Get) && (Verb != ERequestVerb::Delete))
 	{
-		CurlHeaders = curl_slist_append(CurlHeaders, CommonHeaders[i]);
+		AddHeader(*WriteToAnsiString<32>(ANSITEXTVIEW("Content-Length: "), ContentLength));
 	}
-	// Setup added headers
-	for (const FString& Header : Headers)
+
+	if (AuthorizationToken)
 	{
-		CurlHeaders = curl_slist_append(CurlHeaders, TCHAR_TO_ANSI(*Header));
+		AddHeader(*WriteToAnsiString<1024>(*AuthorizationToken));
 	}
+
 	curl_easy_setopt(static_cast<CURL*>(Curl), CURLOPT_HTTPHEADER, CurlHeaders);
-	return CurlHeaders;
+	return Http::Private::FCurlStringList(CurlHeaders);
 }
 
-FHttpRequest::EResult FHttpRequest::PerformBlocking(FStringView Uri, ERequestVerb Verb, uint64 ContentLength, TConstArrayView<long> ExpectedErrorCodes)
+FHttpRequest::EResult FHttpRequest::PerformBlocking(
+	const FStringView Uri,
+	const ERequestVerb Verb,
+	const uint64 ContentLength,
+	const TConstArrayView<long> ExpectedErrorCodes)
 {
 	TRACE_CPUPROFILER_EVENT_SCOPE(HttpDDC_CurlPerform);
 
-	// Build headers list
-	curl_slist* CurlHeaders = PrepareToIssueRequest(Uri, Verb, ContentLength);
+	Http::Private::FCurlStringList CurlHeaders = PrepareToIssueRequest(Uri, Verb, ContentLength);
 
-	// Shots fired!
 	ResultCode = curl_easy_perform(static_cast<CURL*>(Curl));
 
-	// Get response code
 	curl_easy_getinfo(static_cast<CURL*>(Curl), CURLINFO_RESPONSE_CODE, &ResponseCode);
 
 	LogResult(ResultCode, Uri, Verb, ExpectedErrorCodes);
 
-	// Clean up
-	curl_slist_free_all(CurlHeaders);
-
 	return ResultCode == CURLE_OK ? EResult::Success : EResult::Failed;
 }
 
-void FHttpRequest::EnqueueAsync(UE::DerivedData::IRequestOwner& Owner, FHttpRequestPool* Pool, FStringView Uri, ERequestVerb Verb, uint64 ContentLength, FOnHttpRequestComplete&& OnComplete, TConstArrayView<long> ExpectedErrorCodes)
+void FHttpRequest::EnqueueAsync(
+	DerivedData::IRequestOwner& Owner,
+	FHttpRequestPool* const Pool,
+	const FStringView Uri,
+	const ERequestVerb Verb,
+	const uint64 ContentLength,
+	FOnHttpRequestComplete&& OnComplete,
+	const TConstArrayView<long> ExpectedErrorCodes)
 {
 	if (!AllowAsync())
 	{
@@ -625,109 +624,86 @@ void FHttpRequest::EnqueueAsync(UE::DerivedData::IRequestOwner& Owner, FHttpRequ
 	SharedData->AddRequest(Curl);
 }
 
-void FHttpRequest::LogResult(FResultCode Result, FStringView Uri, ERequestVerb Verb, TConstArrayView<long> ExpectedErrorCodes) const
+void FHttpRequest::LogResult(
+	const FResultCode Result,
+	const FStringView Uri,
+	const ERequestVerb Verb,
+	const TConstArrayView<long> ExpectedErrorCodes) const
 {
 	if (Result == CURLE_OK)
 	{
-		bool bSuccess = false;
+		const bool bSuccess = (IsSuccessResponse(ResponseCode) || ExpectedErrorCodes.Contains(ResponseCode));
 		const TCHAR* VerbStr = nullptr;
-		FString AdditionalInfo;
+		TStringBuilder<32> AdditionalInfo;
 
 		switch (Verb)
 		{
 		case ERequestVerb::Head:
-			bSuccess = (ExpectedErrorCodes.Contains(ResponseCode) || IsSuccessResponse(ResponseCode));
 			VerbStr = TEXT("querying");
 			break;
 		case ERequestVerb::Get:
-			bSuccess = (ExpectedErrorCodes.Contains(ResponseCode) || IsSuccessResponse(ResponseCode));
 			VerbStr = TEXT("fetching");
-			AdditionalInfo = FString::Printf(TEXT("Received: %d bytes."), BytesReceived);
+			AdditionalInfo << TEXTVIEW("Received: ") << uint64(BytesReceived) << TEXTVIEW(" bytes.");
 			break;
 		case ERequestVerb::Put:
-			bSuccess = (ExpectedErrorCodes.Contains(ResponseCode) || IsSuccessResponse(ResponseCode));
 			VerbStr = TEXT("updating");
-			AdditionalInfo = FString::Printf(TEXT("Sent: %d bytes."), BytesSent);
+			AdditionalInfo << TEXTVIEW("Sent: ") << uint64(BytesSent) << TEXTVIEW(" bytes.");
 			break;
 		case ERequestVerb::Post:
-			bSuccess = (ExpectedErrorCodes.Contains(ResponseCode) || IsSuccessResponse(ResponseCode));
 			VerbStr = TEXT("posting");
 			break;
 		case ERequestVerb::Delete:
-			bSuccess = (ExpectedErrorCodes.Contains(ResponseCode) || IsSuccessResponse(ResponseCode));
 			VerbStr = TEXT("deleting");
 			break;
+		default:
+			checkNoEntry();
+			break;
 		}
 
 		if (bSuccess)
 		{
-			UE_LOG(
-				LogDerivedDataCache, 
-				Verbose, 
-				TEXT("%s: Finished %s HTTP cache entry (response %d) from %.*s. %s"), 
-				*GetName(),
-				VerbStr,
-				ResponseCode, 
-				Uri.Len(), Uri.GetData(),
-				*AdditionalInfo
-			);
+			UE_LOG(LogDerivedDataCache, Verbose,
+				TEXT("%s: Finished %s HTTP cache entry (response %d) from %.*s. %s"),
+				*GetName(), VerbStr, ResponseCode, Uri.Len(), Uri.GetData(), *AdditionalInfo);
 		}
-		else if(bLogErrors)
+		else if (bLogErrors)
 		{
 			// Print the response body if we got one, otherwise print header.
 			FString Response = GetAnsiBufferAsString(ResponseBuffer.Num() > 0 ? ResponseBuffer : ResponseHeader);
 			Response.ReplaceCharInline(TEXT('\n'), TEXT(' '));
 			Response.ReplaceCharInline(TEXT('\r'), TEXT(' '));
-			// Dont log access denied as error, since tokens can expire mid session
+			// Don't log access denied as error, since tokens can expire mid-session.
 			if (ResponseCode == 401)
 			{
-				UE_LOG(
-					LogDerivedDataCache,
-					Verbose,
+				UE_LOG(LogDerivedDataCache, Verbose,
 					TEXT("%s: Failed %s HTTP cache entry (response %d) from %.*s. Response: %s"),
-					*GetName(),
-					VerbStr,
-					ResponseCode,
-					Uri.Len(), Uri.GetData(),
-					*Response
-				);
+					*GetName(), VerbStr, ResponseCode, Uri.Len(), Uri.GetData(), *Response);
 			}
 			else
-			{ 
-				UE_LOG(
-					LogDerivedDataCache,
-					Display,
+			{
+				UE_LOG(LogDerivedDataCache, Display,
 					TEXT("%s: Failed %s HTTP cache entry (response %d) from %.*s. Response: %s"),
-					*GetName(),
-					VerbStr,
-					ResponseCode,
-					Uri.Len(), Uri.GetData(),
-					*Response
-				);
+					*GetName(), VerbStr, ResponseCode, Uri.Len(), Uri.GetData(), *Response);
 			}
 		}
 	}
-	else if(bLogErrors)
+	else if (bLogErrors)
 	{
-		UE_LOG(
-			LogDerivedDataCache, 
-			Display, 
-			TEXT("%s: Error while connecting to %s: %s"), 
-			*GetName(),
-			*EffectiveDomain,
-			ANSI_TO_TCHAR(curl_easy_strerror(static_cast<CURLcode>(Result)))
-		);
+		UE_LOG(LogDerivedDataCache, Display,
+			TEXT("%s: Error while connecting to %s: %s"),
+			*GetName(), *EffectiveDomain, ANSI_TO_TCHAR(curl_easy_strerror(static_cast<CURLcode>(Result))));
 	}
 }
 
-FHttpRequest::EResult FHttpRequest::PerformBlockingPut(FStringView Uri,
-	FCompositeBuffer Buffer,
-	EHttpContentType ContentType,
-	TConstArrayView<long> ExpectedErrorCodes)
+FHttpRequest::EResult FHttpRequest::PerformBlockingPut(
+	const FStringView Uri,
+	const FCompositeBuffer& Buffer,
+	const EHttpContentType ContentType,
+	const TConstArrayView<long> ExpectedErrorCodes)
 {
 	uint64 ContentLength = 0u;
 
-	CURL* LocalCurl = static_cast<CURL*>(Curl);
+	CURL* const LocalCurl = static_cast<CURL*>(Curl);
 	curl_easy_setopt(LocalCurl, CURLOPT_UPLOAD, 1L);
 	curl_easy_setopt(LocalCurl, CURLOPT_INFILESIZE, Buffer.GetSize());
 	curl_easy_setopt(LocalCurl, CURLOPT_READDATA, this);
@@ -739,15 +715,16 @@ FHttpRequest::EResult FHttpRequest::PerformBlockingPut(FStringView Uri,
 	return PerformBlocking(Uri, ERequestVerb::Put, ContentLength, ExpectedErrorCodes);
 }
 
-FHttpRequest::EResult FHttpRequest::PerformBlockingPost(FStringView Uri,
-	FCompositeBuffer Buffer,
-	EHttpContentType ContentType,
-	EHttpContentType AcceptType,
-	TConstArrayView<long> ExpectedErrorCodes)
+FHttpRequest::EResult FHttpRequest::PerformBlockingPost(
+	const FStringView Uri,
+	const FCompositeBuffer& Buffer,
+	const EHttpContentType ContentType,
+	const EHttpContentType AcceptType,
+	const TConstArrayView<long> ExpectedErrorCodes)
 {
-	uint64 ContentLength = 0u;
+	uint64 ContentLength = 0;
 
-	CURL* LocalCurl = static_cast<CURL*>(Curl);
+	CURL* const LocalCurl = static_cast<CURL*>(Curl);
 	curl_easy_setopt(LocalCurl, CURLOPT_POST, 1L);
 	curl_easy_setopt(LocalCurl, CURLOPT_POSTFIELDSIZE, Buffer.GetSize());
 	curl_easy_setopt(LocalCurl, CURLOPT_READDATA, this);
@@ -760,17 +737,18 @@ FHttpRequest::EResult FHttpRequest::PerformBlockingPost(FStringView Uri,
 	return PerformBlocking(Uri, ERequestVerb::Post, ContentLength, ExpectedErrorCodes);
 }
 
-void FHttpRequest::EnqueueAsyncPut(UE::DerivedData::IRequestOwner& Owner,
-	FHttpRequestPool* Pool,
-	FStringView Uri,
-	FCompositeBuffer Buffer,
+void FHttpRequest::EnqueueAsyncPut(
+	DerivedData::IRequestOwner& Owner,
+	FHttpRequestPool* const Pool,
+	const FStringView Uri,
+	const FCompositeBuffer& Buffer,
 	FOnHttpRequestComplete&& OnComplete,
-	EHttpContentType ContentType,
-	TConstArrayView<long> ExpectedErrorCodes)
+	const EHttpContentType ContentType,
+	const TConstArrayView<long> ExpectedErrorCodes)
 {
-	uint64 ContentLength = 0u;
+	uint64 ContentLength = 0;
 
-	CURL* LocalCurl = static_cast<CURL*>(Curl);
+	CURL* const LocalCurl = static_cast<CURL*>(Curl);
 	curl_easy_setopt(LocalCurl, CURLOPT_UPLOAD, 1L);
 	curl_easy_setopt(LocalCurl, CURLOPT_INFILESIZE, Buffer.GetSize());
 	curl_easy_setopt(LocalCurl, CURLOPT_READDATA, this);
@@ -782,18 +760,19 @@ void FHttpRequest::EnqueueAsyncPut(UE::DerivedData::IRequestOwner& Owner,
 	return EnqueueAsync(Owner, Pool, Uri, ERequestVerb::Put, ContentLength, MoveTemp(OnComplete), ExpectedErrorCodes);
 }
 
-void FHttpRequest::EnqueueAsyncPost(UE::DerivedData::IRequestOwner& Owner,
-	FHttpRequestPool* Pool,
-	FStringView Uri,
-	FCompositeBuffer Buffer,
+void FHttpRequest::EnqueueAsyncPost(
+	DerivedData::IRequestOwner& Owner,
+	FHttpRequestPool* const Pool,
+	const FStringView Uri,
+	const FCompositeBuffer& Buffer,
 	FOnHttpRequestComplete&& OnComplete,
-	EHttpContentType ContentType,
-	EHttpContentType AcceptType,
-	TConstArrayView<long> ExpectedErrorCodes)
+	const EHttpContentType ContentType,
+	const EHttpContentType AcceptType,
+	const TConstArrayView<long> ExpectedErrorCodes)
 {
-	uint64 ContentLength = 0u;
+	uint64 ContentLength = 0;
 
-	CURL* LocalCurl = static_cast<CURL*>(Curl);
+	CURL* const LocalCurl = static_cast<CURL*>(Curl);
 	curl_easy_setopt(LocalCurl, CURLOPT_POST, 1L);
 	curl_easy_setopt(LocalCurl, CURLOPT_INFILESIZE, Buffer.GetSize());
 	curl_easy_setopt(LocalCurl, CURLOPT_READDATA, this);
@@ -806,44 +785,46 @@ void FHttpRequest::EnqueueAsyncPost(UE::DerivedData::IRequestOwner& Owner,
 	return EnqueueAsync(Owner, Pool, Uri, ERequestVerb::Post, ContentLength, MoveTemp(OnComplete), ExpectedErrorCodes);
 }
 
-FHttpRequest::EResult FHttpRequest::PerformBlockingHead(FStringView Uri,
-	EHttpContentType AcceptType,
-	TConstArrayView<long> ExpectedErrorCodes)
+FHttpRequest::EResult FHttpRequest::PerformBlockingHead(
+	const FStringView Uri,
+	const EHttpContentType AcceptType,
+	const TConstArrayView<long> ExpectedErrorCodes)
 {
 	curl_easy_setopt(static_cast<CURL*>(Curl), CURLOPT_NOBODY, 1L);
 	AddContentTypeHeader(TEXTVIEW("Accept"), AcceptType);
-	return PerformBlocking(Uri, ERequestVerb::Head, 0u, ExpectedErrorCodes);
+	return PerformBlocking(Uri, ERequestVerb::Head, 0, ExpectedErrorCodes);
 }
 
-FHttpRequest::EResult FHttpRequest::PerformBlockingDelete(FStringView Uri,
-															TConstArrayView<long> ExpectedErrorCodes)
+FHttpRequest::EResult FHttpRequest::PerformBlockingDelete(
+	const FStringView Uri,
+	const TConstArrayView<long> ExpectedErrorCodes)
 {
 	curl_easy_setopt(static_cast<CURL*>(Curl), CURLOPT_CUSTOMREQUEST, "DELETE");
-	return PerformBlocking(Uri, ERequestVerb::Delete, 0u, ExpectedErrorCodes);
+	return PerformBlocking(Uri, ERequestVerb::Delete, 0, ExpectedErrorCodes);
 }
 
-void FHttpRequest::EnqueueAsyncHead(UE::DerivedData::IRequestOwner& Owner,
-	FHttpRequestPool* Pool,
-	FStringView Uri,
+void FHttpRequest::EnqueueAsyncHead(
+	DerivedData::IRequestOwner& Owner,
+	FHttpRequestPool* const Pool,
+	const FStringView Uri,
 	FOnHttpRequestComplete&& OnComplete,
-	EHttpContentType AcceptType,
-	TConstArrayView<long> ExpectedErrorCodes)
+	const EHttpContentType AcceptType,
+	const TConstArrayView<long> ExpectedErrorCodes)
 {
 	curl_easy_setopt(static_cast<CURL*>(Curl), CURLOPT_NOBODY, 1L);
 	AddContentTypeHeader(TEXTVIEW("Accept"), AcceptType);
-	return EnqueueAsync(Owner, Pool, Uri, ERequestVerb::Head, 0u, MoveTemp(OnComplete), ExpectedErrorCodes);
+	return EnqueueAsync(Owner, Pool, Uri, ERequestVerb::Head, 0, MoveTemp(OnComplete), ExpectedErrorCodes);
 }
 
-void FHttpRequest::EnqueueAsyncDelete(UE::DerivedData::IRequestOwner& Owner,
-	FHttpRequestPool* Pool,
-	FStringView Uri,
+void FHttpRequest::EnqueueAsyncDelete(
+	DerivedData::IRequestOwner& Owner,
+	FHttpRequestPool* const Pool,
+	const FStringView Uri,
 	FOnHttpRequestComplete&& OnComplete,
-	TConstArrayView<long> ExpectedErrorCodes)
+	const TConstArrayView<long> ExpectedErrorCodes)
 {
 	curl_easy_setopt(static_cast<CURL*>(Curl), CURLOPT_CUSTOMREQUEST, "DELETE");
-	return EnqueueAsync(Owner, Pool, Uri, ERequestVerb::Delete, 0u, MoveTemp(OnComplete), ExpectedErrorCodes);
+	return EnqueueAsync(Owner, Pool, Uri, ERequestVerb::Delete, 0, MoveTemp(OnComplete), ExpectedErrorCodes);
 }
 
 } // UE
-
-#endif // UE_WITH_HTTP_CLIENT
diff --git a/Engine/Source/Developer/DerivedDataCache/Private/HttpCacheStore.cpp b/Engine/Source/Developer/DerivedDataCache/Private/HttpCacheStore.cpp
--- a/Engine/Source/Developer/DerivedDataCache/Private/HttpCacheStore.cpp
+++ b/Engine/Source/Developer/DerivedDataCache/Private/HttpCacheStore.cpp
@@ -1375,7 +1375,7 @@ bool FHttpCacheStore::AcquireAccessToken()
 					{
 						Access = MakeUnique<FHttpAccessToken>();
 					}
-					Access->SetHeader(*AccessTokenString);
+					Access->SetToken(AccessTokenString);
 					UE_LOG(LogDerivedDataCache, Display, TEXT("%s: Logged in to HTTP DDC services. Expires in %d seconds."), *Request.GetName(), ExpiryTimeSeconds);
 
 					//Schedule a refresh of the token ahead of expiry time (this will not work in commandlets)
