DDC: Restored the correct default of true for bResolveHostCanonicalName in HttpCacheStore

#rnx

diff --git a/Engine/Source/Developer/DerivedDataCache/Private/HttpCacheStore.cpp b/Engine/Source/Developer/DerivedDataCache/Private/HttpCacheStore.cpp
--- a/Engine/Source/Developer/DerivedDataCache/Private/HttpCacheStore.cpp
+++ b/Engine/Source/Developer/DerivedDataCache/Private/HttpCacheStore.cpp
@@ -110,7 +110,7 @@ struct FHttpCacheStoreParams
 	FString OAuthClientId;
 	FString OAuthSecret;
 	FString OAuthScope;
-	bool bResolveHostCanonicalName = false;
+	bool bResolveHostCanonicalName = true;
 	bool bReadOnly = false;
 
 	void Parse(const TCHAR* NodeName, const TCHAR* Config);
