Limit the maximum open HTTP connections used by Jupiter and Zen DDC backends to 32 HTTP2 connections (w/ multiplexing).  I want to limit it further.  I've seen equally fast cold local editor load times in CitySample with only 8 maximum connections, but until I've had a chance to test that with others in differing network environments, I'm going to stick with 32 and consider reducing later.

diff --git a/Engine/Source/Developer/DerivedDataCache/Private/Http/HttpClient.cpp b/Engine/Source/Developer/DerivedDataCache/Private/Http/HttpClient.cpp
--- a/Engine/Source/Developer/DerivedDataCache/Private/Http/HttpClient.cpp
+++ b/Engine/Source/Developer/DerivedDataCache/Private/Http/HttpClient.cpp
@@ -27,6 +27,7 @@ namespace UE
 namespace Http::Private
 {
 
+static constexpr uint32 MaxTotalConnections = 32;
 static constexpr uint32 WaitIntervalMs = 10;
 
 struct FHttpSharedDataInternals
@@ -100,6 +101,7 @@ FHttpSharedData::FHttpSharedData()
 	curl_share_setopt(Internals->CurlShare, CURLSHOPT_SHARE, CURL_LOCK_DATA_DNS);
 	curl_share_setopt(Internals->CurlShare, CURLSHOPT_SHARE, CURL_LOCK_DATA_SSL_SESSION);
 	Internals->CurlMulti = curl_multi_init();
+	curl_multi_setopt(Internals->CurlMulti, CURLMOPT_MAX_TOTAL_CONNECTIONS, Http::Private::MaxTotalConnections);
 	curl_multi_setopt(Internals->CurlMulti, CURLMOPT_PIPELINING, CURLPIPE_MULTIPLEX);
 }
 
