DDC: Removed the legacy cache store implementations and limit cache bucket names to 63 characters

The cache bucket name was already limited to 63 characters in practice by [REDACTED_PERSON]. This makes it consistent across every cache store.

diff --git a/Engine/Source/Developer/DerivedDataCache/Private/DerivedDataBackendCorruptionWrapper.h b/Engine/Source/Developer/DerivedDataCache/Private/DerivedDataBackendCorruptionWrapper.h
deleted file mode 100644
diff --git a/Engine/Source/Developer/DerivedDataCache/Private/DerivedDataBackends.cpp b/Engine/Source/Developer/DerivedDataCache/Private/DerivedDataBackends.cpp
--- a/Engine/Source/Developer/DerivedDataCache/Private/DerivedDataBackends.cpp
+++ b/Engine/Source/Developer/DerivedDataCache/Private/DerivedDataBackends.cpp
@@ -339,8 +339,8 @@ public:
 			CreatedNodes.AddUnique(ParsedNode.Key);
 
 			// Parse any debug options for this node. E.g. -DDC-<Name>-MissRate
-			FDerivedDataBackendInterface::FBackendDebugOptions DebugOptions;
-			if (FDerivedDataBackendInterface::FBackendDebugOptions::ParseFromTokens(DebugOptions, *NodeName, FCommandLine::Get()))
+			FBackendDebugOptions DebugOptions;
+			if (FBackendDebugOptions::ParseFromTokens(DebugOptions, *NodeName, FCommandLine::Get()))
 			{
 				if (!ParsedNode.Key->LegacyDebugOptions(DebugOptions))
 				{
@@ -1127,185 +1127,6 @@ private:
 namespace UE::DerivedData
 {
 
-void FDerivedDataBackendInterface::LegacyPut(
-	const TConstArrayView<FLegacyCachePutRequest> Requests,
-	IRequestOwner& Owner,
-	FOnLegacyCachePutComplete&& OnComplete)
-{
-	if (GetLegacyMode() != EBackendLegacyMode::LegacyOnly)
-	{
-		return ILegacyCacheStore::LegacyPut(Requests, Owner, MoveTemp(OnComplete));
-	}
-
-	for (const FLegacyCachePutRequest& Request : Requests)
-	{
-		FCompositeBuffer CompositeValue = Request.Value.GetRawData();
-		Request.Key.WriteValueTrailer(CompositeValue);
-
-		checkf(CompositeValue.GetSize() < MAX_int32,
-			TEXT("Value is 2 GiB or greater, which is not supported for put of '%s' from '%s'"),
-			*Request.Key.GetShortKey(), *Request.Name);
-
-		UE_CLOG(Request.Key.HasShortKey(), LogDerivedDataCache, VeryVerbose,
-			TEXT("ShortenKey %s -> %s"), *Request.Key.GetFullKey(), *Request.Key.GetShortKey());
-
-		FSharedBuffer Value = MoveTemp(CompositeValue).ToShared();
-		const TArrayView<const uint8> Data(MakeArrayView(static_cast<const uint8*>(Value.GetData()), int32(Value.GetSize())));
-		const EPutStatus Status = PutCachedData(*Request.Key.GetShortKey(), Data, /*bPutEvenIfExists*/ false);
-		OnComplete({Request.Name, Request.Key, Request.UserData, Status == EPutStatus::Cached ? EStatus::Ok : EStatus::Error});
-	}
-}
-
-void FDerivedDataBackendInterface::LegacyGet(
-	TConstArrayView<FLegacyCacheGetRequest> Requests,
-	IRequestOwner& Owner,
-	FOnLegacyCacheGetComplete&& OnComplete)
-{
-	const EBackendLegacyMode LegacyMode = GetLegacyMode();
-	if (LegacyMode == EBackendLegacyMode::ValueOnly)
-	{
-		return ILegacyCacheStore::LegacyGet(Requests, Owner, MoveTemp(OnComplete));
-	}
-
-	// Make a blocking query to the value cache and fall back to the legacy cache for requests with errors.
-
-	TArray<FLegacyCacheGetRequest> LegacyRequests;
-	if (LegacyMode == EBackendLegacyMode::ValueWithLegacyFallback)
-	{
-		TArray<FLegacyCacheGetRequest, TInlineAllocator<8>> ValueRequests;
-		ValueRequests.Reserve(Requests.Num());
-		uint64 RequestIndex = 0;
-		for (const FLegacyCacheGetRequest& Request : Requests)
-		{
-			ValueRequests.Add_GetRef(Request).UserData = RequestIndex++;
-		}
-
-		FRequestOwner BlockingOwner(EPriority::Blocking);
-		ILegacyCacheStore::LegacyGet(ValueRequests, BlockingOwner, [this, &OnComplete, &Requests, &ValueRequests](FLegacyCacheGetResponse&& Response)
-		{
-			if (Response.Status != EStatus::Error)
-			{
-				const int32 Index = int32(Response.UserData);
-				Response.UserData = Requests[Index].UserData;
-				ValueRequests[Index].UserData = MAX_uint64;
-				OnComplete(MoveTemp(Response));
-			}
-		});
-		BlockingOwner.Wait();
-
-		for (const FLegacyCacheGetRequest& Request : ValueRequests)
-		{
-			if (Request.UserData != MAX_uint64)
-			{
-				LegacyRequests.Add(Requests[int32(Request.UserData)]);
-			}
-		}
-		if (LegacyRequests.IsEmpty())
-		{
-			return;
-		}
-		Requests = LegacyRequests;
-	}
-
-	// Query the legacy cache by translating the requests to legacy cache functions.
-
-	FRequestOwner AsyncOwner(FPlatformMath::Min(Owner.GetPriority(), EPriority::Highest));
-	FRequestBarrier Barrier(AsyncOwner);
-	AsyncOwner.KeepAlive();
-
-	TArray<FString, TInlineAllocator<8>> ExistsKeys;
-	TArray<const FLegacyCacheGetRequest*, TInlineAllocator<8>> ExistsRequests;
-
-	TArray<FString, TInlineAllocator<8>> PrefetchKeys;
-	TArray<const FLegacyCacheGetRequest*, TInlineAllocator<8>> PrefetchRequests;
-
-	for (const FLegacyCacheGetRequest& Request : Requests)
-	{
-		if (!EnumHasAnyFlags(Request.Policy, ECachePolicy::Query))
-		{
-			OnComplete({Request.Name, Request.Key, {}, Request.UserData, EStatus::Error});
-		}
-		else if (EnumHasAnyFlags(Request.Policy, ECachePolicy::SkipData))
-		{
-			const bool bExists = !EnumHasAnyFlags(Request.Policy, ECachePolicy::Store);
-			(bExists ? ExistsKeys : PrefetchKeys).Emplace(Request.Key.GetShortKey());
-			(bExists ? ExistsRequests : PrefetchRequests).Add(&Request);
-		}
-		else
-		{
-			FSharedBuffer Value;
-			TArray<uint8> Data;
-			if (const bool bGetOk = GetCachedData(*Request.Key.GetShortKey(), Data))
-			{
-				FCompositeBuffer CompositeValue(MakeSharedBufferFromArray(MoveTemp(Data)));
-				if (const bool bKeyOk = Request.Key.ReadValueTrailer(CompositeValue))
-				{
-					Value = MoveTemp(CompositeValue).ToShared();
-				}
-			}
-			const EStatus Status = Value ? EStatus::Ok : EStatus::Error;
-			FLegacyCacheValue LegacyValue(FCompositeBuffer(MoveTemp(Value)));
-			if (LegacyValue.HasData() && LegacyMode == EBackendLegacyMode::ValueWithLegacyFallback)
-			{
-				Private::LaunchTaskInCacheThreadPool(AsyncOwner, [this, &AsyncOwner, Request, LegacyValue]
-				{
-					ILegacyCacheStore::LegacyPut({{Request.Name, Request.Key, LegacyValue}}, AsyncOwner, [](auto&&){});
-				});
-			}
-			OnComplete({Request.Name, Request.Key, MoveTemp(LegacyValue), Request.UserData, Status});
-		}
-	}
-
-	if (!PrefetchKeys.IsEmpty())
-	{
-		int32 Index = 0;
-		const TBitArray<> Exists = TryToPrefetch(PrefetchKeys);
-		for (const FLegacyCacheGetRequest* const Request : PrefetchRequests)
-		{
-			OnComplete({Request->Name, Request->Key, {}, Request->UserData, Exists[Index] ? EStatus::Ok : EStatus::Error});
-			++Index;
-		}
-	}
-
-	if (!ExistsKeys.IsEmpty())
-	{
-		int32 Index = 0;
-		const TBitArray<> Exists = CachedDataProbablyExistsBatch(ExistsKeys);
-		for (const FLegacyCacheGetRequest* const Request : ExistsRequests)
-		{
-			OnComplete({Request->Name, Request->Key, {}, Request->UserData, Exists[Index] ? EStatus::Ok : EStatus::Error});
-			++Index;
-		}
-	}
-}
-
-void FDerivedDataBackendInterface::LegacyDelete(
-	const TConstArrayView<FLegacyCacheDeleteRequest> Requests,
-	IRequestOwner& Owner,
-	FOnLegacyCacheDeleteComplete&& OnComplete)
-{
-	if (GetLegacyMode() != EBackendLegacyMode::LegacyOnly)
-	{
-		return ILegacyCacheStore::LegacyDelete(Requests, Owner, MoveTemp(OnComplete));
-	}
-
-	for (const FLegacyCacheDeleteRequest& Request : Requests)
-	{
-		RemoveCachedData(*Request.Key.GetShortKey(), Request.bTransient);
-		OnComplete({Request.Name, Request.Key, Request.UserData, EStatus::Ok});
-	}
-}
-
-void FDerivedDataBackendInterface::LegacyStats(FDerivedDataCacheStatsNode& OutNode)
-{
-	OutNode = MoveTemp(GatherUsageStats().Get());
-}
-
-bool FDerivedDataBackendInterface::LegacyDebugOptions(FBackendDebugOptions& Options)
-{
-	return ApplyDebugOptions(Options);
-}
-
 FDerivedDataBackend* FDerivedDataBackend::Create()
 {
 	return new FDerivedDataBackendGraph();
@@ -1341,7 +1162,7 @@ FBackendDebugOptions::FBackendDebugOptions()
 /**
  * Parse debug options for the provided node name. Returns true if any options were specified
  */
-bool FBackendDebugOptions::ParseFromTokens(FDerivedDataBackendInterface::FBackendDebugOptions& OutOptions, const TCHAR* InNodeName, const TCHAR* InInputTokens)
+bool FBackendDebugOptions::ParseFromTokens(FBackendDebugOptions& OutOptions, const TCHAR* InNodeName, const TCHAR* InInputTokens)
 {
 	// Check if the input stream has any DDC options for this node.
 	TStringBuilder<64> Prefix;
diff --git a/Engine/Source/Developer/DerivedDataCache/Private/DerivedDataCacheKey.cpp b/Engine/Source/Developer/DerivedDataCache/Private/DerivedDataCacheKey.cpp
--- a/Engine/Source/Developer/DerivedDataCache/Private/DerivedDataCacheKey.cpp
+++ b/Engine/Source/Developer/DerivedDataCache/Private/DerivedDataCacheKey.cpp
@@ -18,9 +18,9 @@ namespace UE::DerivedData::Private
 template <typename CharType>
 static inline void AssertValidCacheBucketName(TStringView<CharType> Name)
 {
-	checkf(IsValidCacheBucketName(Name),
-		TEXT("A cache bucket name must be alphanumeric, non-empty, and contain fewer than 256 code units. ")
-		TEXT("Name: '%s'"), *WriteToString<256>(Name));
+	checkf(FCacheBucket::IsValidName(Name),
+		TEXT("A cache bucket name must be alphanumeric, non-empty, and contain at most %d code units. Name: '%s'"),
+		*WriteToString<256>(Name), FCacheBucket::MaxNameLen);
 }
 
 class FCacheBucketOwner : public FCacheBucket
diff --git a/Engine/Source/Developer/DerivedDataCache/Private/DerivedDataCachePrivate.h b/Engine/Source/Developer/DerivedDataCache/Private/DerivedDataCachePrivate.h
--- a/Engine/Source/Developer/DerivedDataCache/Private/DerivedDataCachePrivate.h
+++ b/Engine/Source/Developer/DerivedDataCache/Private/DerivedDataCachePrivate.h
@@ -3,8 +3,6 @@
 #pragma once
 
 #include "CoreTypes.h"
-#include "Containers/StringFwd.h"
-#include "Misc/AsciiSet.h"
 
 class FDerivedDataCacheInterface;
 
@@ -26,13 +24,4 @@ uint64 GetCacheRecordCompressedSize(const FCacheRecord& Record);
 uint64 GetCacheRecordTotalRawSize(const FCacheRecord& Record);
 uint64 GetCacheRecordRawSize(const FCacheRecord& Record);
 
-///////////////////////////////////////////////////////////////////////////////////////////////////////////////////////
-
-template <typename CharType>
-inline bool IsValidCacheBucketName(TStringView<CharType> Name)
-{
-	constexpr FAsciiSet Valid("ABCDEFGHIJKLMNOPQRSTUVWXYZabcdefghijklmnopqrstuvwxyz0123456789");
-	return !Name.IsEmpty() && Name.Len() < 256 && FAsciiSet::HasOnly(Name, Valid);
-}
-
 } // UE::DerivedData::Private
diff --git a/Engine/Source/Developer/DerivedDataCache/Private/DerivedDataCacheRecord.cpp b/Engine/Source/Developer/DerivedDataCache/Private/DerivedDataCacheRecord.cpp
--- a/Engine/Source/Developer/DerivedDataCache/Private/DerivedDataCacheRecord.cpp
+++ b/Engine/Source/Developer/DerivedDataCache/Private/DerivedDataCacheRecord.cpp
@@ -210,7 +210,7 @@ FOptionalCacheRecord FCacheRecord::Load(const FCbPackage& Attachments, const FCb
 	FCbObjectView KeyObject = ObjectView[ANSITEXTVIEW("Key")].AsObjectView();
 	auto TrySetBucketName = [](FUtf8StringView Name, FCacheKey& Key)
 	{
-		if (Private::IsValidCacheBucketName(Name))
+		if (FCacheBucket::IsValidName(Name))
 		{
 			Key.Bucket = FCacheBucket(Name);
 			return true;
diff --git a/Engine/Source/Developer/DerivedDataCache/Private/FileSystemCacheStore.cpp b/Engine/Source/Developer/DerivedDataCache/Private/FileSystemCacheStore.cpp
--- a/Engine/Source/Developer/DerivedDataCache/Private/FileSystemCacheStore.cpp
+++ b/Engine/Source/Developer/DerivedDataCache/Private/FileSystemCacheStore.cpp
@@ -6,9 +6,7 @@
 #include "Algo/StableSort.h"
 #include "Algo/Transform.h"
 #include "Async/Async.h"
-#include "Async/TaskGraphInterfaces.h"
 #include "Containers/StaticBitArray.h"
-#include "DerivedDataBackendCorruptionWrapper.h"
 #include "DerivedDataBackendInterface.h"
 #include "DerivedDataCacheInterface.h"
 #include "DerivedDataCacheMaintainer.h"
@@ -25,7 +23,6 @@
 #include "HAL/Thread.h"
 #include "Hash/xxhash.h"
 #include "HashingArchiveProxy.h"
-#include "Math/NumericLimits.h"
 #include "Misc/CommandLine.h"
 #include "Misc/ConfigCacheIni.h"
 #include "Misc/CoreMisc.h"
@@ -46,21 +43,9 @@
 #include "Serialization/CompactBinaryWriter.h"
 #include "Templates/Greater.h"
 
-#define MAX_BACKEND_KEY_LENGTH (120)
-#define MAX_BACKEND_NUMBERED_SUBFOLDER_LENGTH (9)
-#if PLATFORM_LINUX	// PATH_MAX on Linux is 4096 (getconf PATH_MAX /, also see limits.h), so this value can be larger (note that it is still arbitrary).
-                    // This should not affect sharing the cache between platforms as the absolute paths will be different anyway.
-	#define MAX_CACHE_DIR_LEN (3119)
-#else
-	#define MAX_CACHE_DIR_LEN (119)
-#endif // PLATFORM_LINUX
-#define MAX_CACHE_EXTENTION_LEN (4)
-
 namespace UE::DerivedData
 {
 
-TRACE_DECLARE_INT_COUNTER(FileSystemDDC_Exist, TEXT("FileSystemDDC Exist"));
-TRACE_DECLARE_INT_COUNTER(FileSystemDDC_ExistHit, TEXT("FileSystemDDC Exist Hit"));
 TRACE_DECLARE_INT_COUNTER(FileSystemDDC_Get, TEXT("FileSystemDDC Get"));
 TRACE_DECLARE_INT_COUNTER(FileSystemDDC_GetHit, TEXT("FileSystemDDC Get Hit"));
 TRACE_DECLARE_INT_COUNTER(FileSystemDDC_Put, TEXT("FileSystemDDC Put"));
@@ -70,26 +55,25 @@ TRACE_DECLARE_INT_COUNTER(FileSystemDDC_BytesWritten, TEXT("FileSystemDDC Bytes
 
 ///////////////////////////////////////////////////////////////////////////////////////////////////////////////////////
 
+#if PLATFORM_LINUX
+// PATH_MAX on Linux is 4096 (getconf PATH_MAX /, also see limits.h), so this value can be larger (note that it is still arbitrary).
+// This should not affect sharing the cache between platforms as the absolute paths will be different anyway.
+static constexpr int32 GMaxCacheRootLen = 3119;
+#else
+static constexpr int32 GMaxCacheRootLen = 119;
+#endif // PLATFORM_LINUX
+
+static constexpr int32 GMaxCacheKeyLen =
+	FCacheBucket::MaxNameLen + // Name
+	sizeof(FIoHash) * 2 +      // Hash
+	4 +                        // Separators /<Name>/<Hash01>/<Hash23>/<Hash4-40>
+	4;                         // Extension (.udd)
+
 static const TCHAR* GBucketsDirectoryName = TEXT("Buckets");
 static const TCHAR* GContentDirectoryName = TEXT("Content");
 
 ///////////////////////////////////////////////////////////////////////////////////////////////////////////////////////
 
-void BuildPathForLegacyCache(const TCHAR* const CacheKey, FStringBuilderBase& Path)
-{
-	TStringBuilder<256> Key;
-	Key << CacheKey;
-	for (TCHAR& Char : MakeArrayView(Key))
-	{
-		Char = FChar::ToUpper(Char);
-	}
-	checkf(Algo::AllOf(MakeArrayView(Key),
-		[](TCHAR C) { return FChar::IsAlnum(C) || FChar::IsUnderscore(C) || C == TEXT('$'); }),
-		TEXT("Invalid characters in cache key %s"), CacheKey);
-	const uint32 Hash = FCrc::StrCrc_DEPRECATED(Key.Len(), Key.GetData());
-	Path.Appendf(TEXT("%1d/%1d/%1d/%s.udd"), (Hash / 100) % 10, (Hash / 10) % 10, Hash % 10, *Key);
-}
-
 void BuildPathForCachePackage(const FCacheKey& CacheKey, FStringBuilderBase& Path)
 {
 	const FIoHash::ByteArray& Bytes = CacheKey.Hash.GetBytes();
@@ -696,7 +680,6 @@ class FAccessLogWriter
 public:
 	FAccessLogWriter(const TCHAR* FileName, const FString& CachePath);
 
-	void Append(const TCHAR* CacheKey, FStringView Path);
 	void Append(const FIoHash& RawHash, FStringView Path);
 	void Append(const FCacheKey& CacheKey, FStringView Path);
 
@@ -706,7 +689,6 @@ private:
 	TUniquePtr<FArchive> Archive;
 	FString BasePath;
 	FCriticalSection CriticalSection;
-	TSet<FString> CacheKeys;
 	TSet<FIoHash> ContentKeys;
 	TSet<FCacheKey> RecordKeys;
 };
@@ -717,18 +699,6 @@ FAccessLogWriter::FAccessLogWriter(const TCHAR* const FileName, const FString& C
 {
 }
 
-void FAccessLogWriter::Append(const TCHAR* const CacheKey, const FStringView Path)
-{
-	FScopeLock Lock(&CriticalSection);
-
-	bool bIsAlreadyInSet = false;
-	CacheKeys.FindOrAdd(FString(CacheKey), &bIsAlreadyInSet);
-	if (!bIsAlreadyInSet)
-	{
-		AppendPath(Path);
-	}
-}
-
 void FAccessLogWriter::Append(const FIoHash& RawHash, const FStringView Path)
 {
 	FScopeLock Lock(&CriticalSection);
@@ -765,12 +735,14 @@ void FAccessLogWriter::AppendPath(const FStringView Path)
 
 ///////////////////////////////////////////////////////////////////////////////////////////////////////////////////////
 
-class FFileSystemCacheStore final : public FDerivedDataBackendInterface
+class FFileSystemCacheStore final : public ILegacyCacheStore
 {
 public:
-	FFileSystemCacheStore(const TCHAR* CachePath, const TCHAR* Params, const TCHAR* AccessLogPath);
-
-	inline bool IsUsable() const { return !bDisabled; }
+	FFileSystemCacheStore(
+		const TCHAR* CachePath,
+		const TCHAR* Params,
+		const TCHAR* AccessLogPath,
+		ECacheStoreFlags& OutFlags);
 
 	bool RunSpeedTest(
 		double InSkipTestsIfSeeksExceedMS,
@@ -778,23 +750,6 @@ public:
 		double& OutReadSpeedMBs,
 		double& OutWriteSpeedMBs) const;
 
-	// FDerivedDataBackendInterface Interface
-
-	FString GetName() const final { return CachePath; }
-	bool IsWritable() const final { return !bReadOnly && !bDisabled; }
-	ESpeedClass GetSpeedClass() const final { return SpeedClass; }
-	
-	bool CachedDataProbablyExists(const TCHAR* CacheKey) final;
-	bool GetCachedData(const TCHAR* CacheKey, TArray<uint8>& Data) final;
-	EPutStatus PutCachedData(const TCHAR* CacheKey, TConstArrayView<uint8> Data, bool bPutEvenIfExists) final;
-	void RemoveCachedData(const TCHAR* CacheKey, bool bTransient) final;
-	TSharedRef<FDerivedDataCacheStatsNode> GatherUsageStats() const final;
-	TBitArray<> TryToPrefetch(TConstArrayView<FString> CacheKeys) final;
-	bool WouldCache(const TCHAR* CacheKey, TConstArrayView<uint8> InData) final;
-	bool ApplyDebugOptions(FBackendDebugOptions& InOptions) final;
-
-	EBackendLegacyMode GetLegacyMode() const final { return LegacyMode; }
-
 	// ICacheStore Interface
 
 	void Put(
@@ -818,6 +773,11 @@ public:
 		IRequestOwner& Owner,
 		FOnCacheGetChunkComplete&& OnComplete) final;
 
+	// ILegacyCacheStore
+
+	void LegacyStats(FDerivedDataCacheStatsNode& OutNode) final;
+	bool LegacyDebugOptions(FBackendDebugOptions& Options) final;
+
 private:
 	[[nodiscard]] bool PutCacheRecord(FStringView Name, const FCacheRecord& Record, const FCacheRecordPolicy& Policy, uint64& OutWriteSize);
 
@@ -868,26 +828,14 @@ private:
 	[[nodiscard]] bool FileExists(FStringBuilderBase& Path) const;
 
 private:
-	void BuildCacheLegacyPath(const TCHAR* const CacheKey, FStringBuilderBase& Path) const
-	{
-		Path << CachePath << TEXT('/');
-		BuildPathForLegacyCache(CacheKey, Path);
-	}
-
 	/** Base path we are storing the cache files in. */
 	FString	CachePath;
-	/** How to access the legacy cache. */
-	EBackendLegacyMode LegacyMode = EBackendLegacyMode::ValueWithLegacyFallback;
-	/** Class of this cache */
-	ESpeedClass SpeedClass;
-	/** If true, we failed to write to this directory and it did not contain anything so we should not be used. */
-	bool		bDisabled;
+	/** Class of this cache. */
+	EBackendSpeedClass SpeedClass;
 	/** If true, do not attempt to write to this cache. */
 	bool		bReadOnly;
 	/** If true, CachedDataProbablyExists will update the file timestamps. */
 	bool		bTouch;
-	/** If true, allow transient data to be removed from the cache. */
-	bool		bPurgeTransient;
 	/** Age of file when it should be deleted from DDC cache. */
 	double		DaysToDeleteUnusedFiles;
 
@@ -896,17 +844,6 @@ private:
 	/** Maximum total size of compressed data stored within a value package, or a record package with one attachment. */
 	uint64		MaxValueSizeKB = 1024;
 
-	/** Object used for synchronization via a scoped lock. */
-	FCriticalSection SynchronizationObject;
-
-	// DDCNotification metrics
-
-	/** Map of cache keys to miss times for generating timing deltas. */
-	TMap<FString, double> DDCNotificationCacheTimes;
-
-	/** The total estimated build time accumulated from cache miss/put deltas. */
-	double TotalEstimatedBuildTime;
-
 	/** Access log to write to */
 	TUniquePtr<FAccessLogWriter> AccessLogWriter;
 
@@ -924,18 +861,16 @@ private:
 FFileSystemCacheStore::FFileSystemCacheStore(
 	const TCHAR* const InCachePath,
 	const TCHAR* const InParams,
-	const TCHAR* const InAccessLogPath)
+	const TCHAR* const InAccessLogPath,
+	ECacheStoreFlags& OutFlags)
 	: CachePath(InCachePath)
-	, SpeedClass(ESpeedClass::Unknown)
-	, bDisabled(false)
+	, SpeedClass(EBackendSpeedClass::Unknown)
 	, bReadOnly(false)
 	, bTouch(false)
-	, bPurgeTransient(false)
 	, DaysToDeleteUnusedFiles(15.0)
-	, TotalEstimatedBuildTime(0)
 {
 	// If we find a platform that has more stringent limits, this needs to be rethought.
-	checkf(MAX_BACKEND_KEY_LENGTH + MAX_CACHE_DIR_LEN + MAX_BACKEND_NUMBERED_SUBFOLDER_LENGTH + MAX_CACHE_EXTENTION_LEN < FPlatformMisc::GetMaxPathLength(),
+	checkf(GMaxCacheRootLen + GMaxCacheKeyLen <= FPlatformMisc::GetMaxPathLength(),
 		TEXT("Not enough room left for cache keys in max path."));
 
 	check(CachePath.Len());
@@ -944,19 +879,10 @@ FFileSystemCacheStore::FFileSystemCacheStore(
 	// Params that override our instance defaults
 	FParse::Bool(InParams, TEXT("ReadOnly="), bReadOnly);
 	FParse::Bool(InParams, TEXT("Touch="), bTouch);
-	FParse::Bool(InParams, TEXT("PurgeTransient="), bPurgeTransient);
 	FParse::Value(InParams, TEXT("UnusedFileAge="), DaysToDeleteUnusedFiles);
 	FParse::Value(InParams, TEXT("MaxRecordSizeKB="), MaxRecordSizeKB);
 	FParse::Value(InParams, TEXT("MaxValueSizeKB="), MaxValueSizeKB);
 
-	if (FString LegacyModeString; FParse::Value(InParams, TEXT("LegacyMode="), LegacyModeString))
-	{
-		if (!TryLexFromString(LegacyMode, LegacyModeString))
-		{
-			UE_LOG(LogDerivedDataCache, Warning, TEXT("%s: Ignoring unrecognized legacy mode '%s'"), *CachePath, *LegacyModeString);
-		}
-	}
-
 	// Flush the cache if requested.
 	bool bFlush = false;
 	if (!bReadOnly && FParse::Bool(InParams, TEXT("Flush="), bFlush) && bFlush)
@@ -985,7 +911,7 @@ FFileSystemCacheStore::FFileSystemCacheStore(
 
 	if (!SkipSpeedTest && !RunSpeedTest(ConsiderSlowAtMS * 2, SpeedStats.LatencyMS, SpeedStats.ReadSpeedMBs, SpeedStats.WriteSpeedMBs))
 	{
-		bDisabled = true;
+		OutFlags = ECacheStoreFlags::None;
 		UE_LOG(LogDerivedDataCache, Warning, TEXT("No read or write access to %s"), *CachePath);
 	}
 	else
@@ -999,19 +925,19 @@ FFileSystemCacheStore::FFileSystemCacheStore(
 		// classify and report on these times
 		if (SpeedStats.LatencyMS < 1)
 		{
-			SpeedClass = ESpeedClass::Local;
+			SpeedClass = EBackendSpeedClass::Local;
 		}
 		else if (SpeedStats.LatencyMS <= ConsiderFastAtMS)
 		{
-			SpeedClass = ESpeedClass::Fast;
+			SpeedClass = EBackendSpeedClass::Fast;
 		}
 		else if (SpeedStats.LatencyMS >= ConsiderSlowAtMS)
 		{
-			SpeedClass = ESpeedClass::Slow;
+			SpeedClass = EBackendSpeedClass::Slow;
 		}
 		else
 		{
-			SpeedClass = ESpeedClass::Ok;
+			SpeedClass = EBackendSpeedClass::Ok;
 		}
 
 		UE_LOG(LogDerivedDataCache, Display,
@@ -1019,7 +945,7 @@ FFileSystemCacheStore::FFileSystemCacheStore(
 			TEXT("Assigned SpeedClass '%s'"),
 			*CachePath, SpeedStats.LatencyMS, SpeedStats.ReadSpeedMBs, SpeedStats.WriteSpeedMBs, LexToString(SpeedClass));
 
-		if (SpeedClass <= FDerivedDataBackendInterface::ESpeedClass::Slow && !bReadOnly)
+		if (SpeedClass <= EBackendSpeedClass::Slow && !bReadOnly)
 		{
 			if (GIsBuildMachine)
 			{
@@ -1096,10 +1022,15 @@ FFileSystemCacheStore::FFileSystemCacheStore(
 			}
 		}
 
-		if (IsUsable() && InAccessLogPath != nullptr && *InAccessLogPath != 0)
+		if (InAccessLogPath && *InAccessLogPath)
 		{
 			AccessLogWriter.Reset(new FAccessLogWriter(InAccessLogPath, CachePath));
 		}
+
+		ECacheStoreFlags Flags = ECacheStoreFlags::Query;
+		Flags |= bReadOnly ? ECacheStoreFlags::None : ECacheStoreFlags::Store;
+		Flags |= SpeedClass == EBackendSpeedClass::Local ? ECacheStoreFlags::Local : ECacheStoreFlags::Remote;
+		OutFlags = Flags;
 	}
 }
 
@@ -1128,10 +1059,10 @@ bool FFileSystemCacheStore::RunSpeedTest(
 	int TotalDataWritten = 0;
 
 	const FString AbsoluteCachePath = IFileManager::Get().ConvertToAbsolutePathForExternalAppForRead(*CachePath);
-	if (AbsoluteCachePath.Len() > MAX_CACHE_DIR_LEN)
+	if (AbsoluteCachePath.Len() >= GMaxCacheRootLen)
 	{
 		const FText ErrorMessage = FText::Format(NSLOCTEXT("DerivedDataCache", "PathTooLong", "Cache path {0} is longer than {1} characters... please adjust [DerivedDataBackendGraph] paths to be shorter (this leaves more room for cache keys)."),
-			FText::FromString(AbsoluteCachePath), FText::AsNumber(MAX_CACHE_DIR_LEN));
+			FText::FromString(AbsoluteCachePath), FText::AsNumber(GMaxCacheRootLen));
 		FMessageDialog::Open(EAppMsgType::Ok, ErrorMessage);
 		UE_LOG(LogDerivedDataCache, Fatal, TEXT("%s"), *ErrorMessage.ToString());
 	}
@@ -1315,274 +1246,14 @@ bool FFileSystemCacheStore::RunSpeedTest(
 	return bWriteTestPassed || bReadTestPassed;
 }
 
-bool FFileSystemCacheStore::CachedDataProbablyExists(const TCHAR* const CacheKey)
-{
-	TRACE_CPUPROFILER_EVENT_SCOPE(FileSystemDDC_Exist);
-	TRACE_COUNTER_INCREMENT(FileSystemDDC_Exist);
-	COOK_STAT(auto Timer = UsageStats.TimeProbablyExists());
-
-	if (!IsUsable())
-	{
-		UE_LOG(LogDerivedDataCache, VeryVerbose,
-			TEXT("%s: Skipped exists check of %s because this cache store is not available"), *CachePath, CacheKey);
-		return false;
-	}
-
-	if (DebugOptions.ShouldSimulateGetMiss(CacheKey))
-	{
-		UE_LOG(LogDerivedDataCache, Verbose, TEXT("%s: Simulated miss for exists check of %s"), *CachePath, CacheKey);
-		return false;
-	}
-
-	TStringBuilder<256> LegacyPath;
-	BuildCacheLegacyPath(CacheKey, LegacyPath);
-	const bool bExists = FileExists(LegacyPath);
-	if (bExists)
-	{
-		if (AccessLogWriter)
-		{
-			AccessLogWriter->Append(CacheKey, LegacyPath);
-		}
-
-		TRACE_COUNTER_INCREMENT(FileSystemDDC_ExistHit);
-		COOK_STAT(Timer.AddHit(0));
-	}
-	// If not using a shared cache, record a (probable) miss
-	else if (!GetDerivedDataCacheRef().GetUsingSharedDDC())
-	{
-		// store a cache miss
-		FScopeLock ScopeLock(&SynchronizationObject);
-		if (!DDCNotificationCacheTimes.Contains(CacheKey))
-		{
-			DDCNotificationCacheTimes.Add(CacheKey, FPlatformTime::Seconds());
-		}
-	}
-
-	UE_LOG(LogDerivedDataCache, Verbose,
-		TEXT("%s: CachedDataProbablyExists=%d for %s"), *CachePath, int32(bExists), CacheKey);
-
-	return bExists;
-}
-
-bool FFileSystemCacheStore::GetCachedData(const TCHAR* const CacheKey, TArray<uint8>& Data)
-{
-	TRACE_CPUPROFILER_EVENT_SCOPE(FileSystemDDC_Get);
-	TRACE_COUNTER_INCREMENT(FileSystemDDC_Get);
-	COOK_STAT(auto Timer = UsageStats.TimeGet());
-
-	const double StartTime = FPlatformTime::Seconds();
-
-	if (!IsUsable())
-	{
-		UE_LOG(LogDerivedDataCache, VeryVerbose,
-			TEXT("%s: Skipped get of %s because this cache store is not available"), *CachePath, CacheKey);
-		return false;
-	}
-
-	if (DebugOptions.ShouldSimulateGetMiss(CacheKey))
-	{
-		UE_LOG(LogDerivedDataCache, Verbose, TEXT("%s: Simulated miss for get of %s"), *CachePath, CacheKey);
-		return false;
-	}
-
-	const auto ReadData = [this, CacheKey, &Data](FArchive& Ar)
-	{
-		const int64 Size = Ar.TotalSize();
-		if (Size > MAX_int32)
-		{
-			Ar.SetError();
-			return;
-		}
-
-		Data.Reset(int32(Size));
-		Data.AddUninitialized(int32(Size));
-		Ar.Serialize(Data.GetData(), Size);
-
-		if (!FCorruptionWrapper::ReadTrailer(Data, *CachePath, CacheKey))
-		{
-			Ar.SetError();
-		}
-	};
-
-	TStringBuilder<256> LegacyPath;
-	BuildCacheLegacyPath(CacheKey, LegacyPath);
-	if (LoadFile(LegacyPath, CacheKey, ReadData))
-	{
-		if (AccessLogWriter)
-		{
-			AccessLogWriter->Append(CacheKey, LegacyPath);
-		}
-
-		const double ReadDuration = FPlatformTime::Seconds() - StartTime;
-		const double ReadSpeed = ReadDuration > 0.001 ? (Data.Num() / ReadDuration) / (1024.0 * 1024.0) : 0.0;
-		UE_LOG(LogDerivedDataCache, Verbose,
-			TEXT("%s: Cache hit on %s (%d bytes, %.02f secs, %.2fMB/s)"),
-			*CachePath, CacheKey, Data.Num(), ReadDuration, ReadSpeed);
-		TRACE_COUNTER_INCREMENT(FileSystemDDC_GetHit);
-		TRACE_COUNTER_ADD(FileSystemDDC_BytesRead, int64(Data.Num()));
-		COOK_STAT(Timer.AddHit(Data.Num()));
-		return true;
-	}
-
-	UE_LOG(LogDerivedDataCache, Verbose, TEXT("%s: Cache miss on %s"), *CachePath, CacheKey);
-	Data.Empty();
-
-	// If not using a shared cache, record a miss
-	if (!GetDerivedDataCacheRef().GetUsingSharedDDC())
-	{
-		// store a cache miss
-		FScopeLock ScopeLock(&SynchronizationObject);
-		if (!DDCNotificationCacheTimes.Contains(CacheKey))
-		{
-			DDCNotificationCacheTimes.Add(CacheKey, FPlatformTime::Seconds());
-		}
-	}
-
-	return false;
-}
-
-bool FFileSystemCacheStore::WouldCache(const TCHAR* const CacheKey, const TConstArrayView<uint8> InData)
+void FFileSystemCacheStore::LegacyStats(FDerivedDataCacheStatsNode& OutNode)
 {
-	return IsWritable() && !CachedDataProbablyExists(CacheKey);
+	OutNode = {TEXT("File System"), *CachePath, SpeedClass == EBackendSpeedClass::Local};
+	OutNode.UsageStats.Add(TEXT(""), UsageStats);
+	OutNode.SpeedStats = SpeedStats;
 }
 
-FFileSystemCacheStore::EPutStatus FFileSystemCacheStore::PutCachedData(
-	const TCHAR* const CacheKey,
-	const TConstArrayView<uint8> Data,
-	const bool bPutEvenIfExists)
-{
-	TRACE_CPUPROFILER_EVENT_SCOPE(FileSystemDDC_Put);
-	TRACE_COUNTER_INCREMENT(FileSystemDDC_Put);
-	COOK_STAT(auto Timer = UsageStats.TimePut());
-
-	checkf(Data.Num(), TEXT("%s: Failed to put %s due to empty data"), *CachePath, CacheKey);
-
-	if (!IsUsable())
-	{
-		UE_LOG(LogDerivedDataCache, VeryVerbose,
-			TEXT("%s: Skipped put of %s because this cache store is not available"), *CachePath, CacheKey);
-		return EPutStatus::NotCached;
-	}
-
-	if (!IsWritable())
-	{
-		UE_LOG(LogDerivedDataCache, VeryVerbose,
-			TEXT("%s: Skipped put of %s because this cache store is read-only"), *CachePath, CacheKey);
-		return EPutStatus::NotCached;
-	}
-
-	if (DebugOptions.ShouldSimulatePutMiss(CacheKey))
-	{
-		UE_LOG(LogDerivedDataCache, Verbose, TEXT("%s: Simulated miss for put of %s"), *CachePath, CacheKey);
-		return EPutStatus::NotCached;
-	}
-
-	TStringBuilder<256> LegacyPath;
-	BuildCacheLegacyPath(CacheKey, LegacyPath);
-
-	if (AccessLogWriter)
-	{
-		AccessLogWriter->Append(CacheKey, LegacyPath);
-	}
-
-	EPutStatus Status = EPutStatus::NotCached;
-
-	if (bPutEvenIfExists || !FileExists(LegacyPath))
-	{
-		const auto WriteData = [this, CacheKey, &Data](FArchive& Ar)
-		{
-			Ar.Serialize(const_cast<uint8*>(Data.GetData()), Data.Num());
-			FCorruptionWrapper::WriteTrailer(Ar, Data, *CachePath, CacheKey);
-		};
-
-		if (SaveFile(LegacyPath, CacheKey, WriteData))
-		{
-			UE_LOG(LogDerivedDataCache, Verbose,
-				TEXT("%s: Successful cache put of %s to %s"), *CachePath, CacheKey, *LegacyPath);
-			TRACE_COUNTER_INCREMENT(FileSystemDDC_PutHit);
-			TRACE_COUNTER_ADD(FileSystemDDC_BytesWritten, int64(Data.Num()));
-			COOK_STAT(Timer.AddHit(Data.Num()));
-			Status = EPutStatus::Cached;
-		}
-	}
-	else
-	{
-		UE_LOG(LogDerivedDataCache, Verbose, TEXT("%s: Skipping put of %s to existing file"), *CachePath, CacheKey);
-		Status = EPutStatus::Cached;
-	}
-
-	// If not using a shared cache, update estimated build time
-	if (!GetDerivedDataCacheRef().GetUsingSharedDDC())
-	{
-		FScopeLock ScopeLock(&SynchronizationObject);
-
-		if (DDCNotificationCacheTimes.Contains(CacheKey))
-		{
-			// There isn't any way to get exact build times in the DDC code as custom asset processing and async are factors.
-			// So, estimate the asset build time based on the delta between the cache miss and the put
-			TotalEstimatedBuildTime += (FPlatformTime::Seconds() - DDCNotificationCacheTimes[CacheKey]);
-			DDCNotificationCacheTimes.Remove(CacheKey);
-
-			// If more than 20 seconds has been spent building assets, send out a notification
-			if (TotalEstimatedBuildTime > 20.0f)
-			{
-				// Send out a DDC put notification if we have any subscribers
-				FDerivedDataCacheInterface::FOnDDCNotification& DDCNotificationEvent =
-					GetDerivedDataCacheRef().GetDDCNotificationEvent();
-
-				if (DDCNotificationEvent.IsBound())
-				{
-					TotalEstimatedBuildTime = 0.0f;
-
-					DECLARE_CYCLE_STAT(TEXT("FSimpleDelegateGraphTask.PutCachedData"),
-					STAT_FSimpleDelegateGraphTask_DDCNotification, STATGROUP_TaskGraphTasks);
-
-					FSimpleDelegateGraphTask::CreateAndDispatchWhenReady(
-						FSimpleDelegateGraphTask::FDelegate::CreateLambda([DDCNotificationEvent]() {
-						DDCNotificationEvent.Broadcast(FDerivedDataCacheInterface::SharedDDCPerformanceNotification);
-					}),
-						GET_STATID(STAT_FSimpleDelegateGraphTask_DDCNotification),
-						nullptr,
-						ENamedThreads::GameThread);
-				}
-			}
-		}
-	}
-
-	return Status;
-}
-
-void FFileSystemCacheStore::RemoveCachedData(const TCHAR* const CacheKey, const bool bTransient)
-{
-	check(IsUsable());
-	if (IsWritable() && (!bTransient || bPurgeTransient))
-	{
-		TStringBuilder<256> LegacyPath;
-		BuildCacheLegacyPath(CacheKey, LegacyPath);
-		if (bTransient)
-		{
-			UE_LOG(LogDerivedDataCache, Verbose,
-				TEXT("%s: Deleting cached data for %s from %s"), *CachePath, CacheKey, *LegacyPath);
-		}
-		IFileManager::Get().Delete(*LegacyPath, /*bRequireExists*/ false, /*bEvenReadOnly*/ false, /*bQuiet*/ true);
-	}
-}
-
-TSharedRef<FDerivedDataCacheStatsNode> FFileSystemCacheStore::GatherUsageStats() const
-{
-	TSharedRef<FDerivedDataCacheStatsNode> Usage =
-		MakeShared<FDerivedDataCacheStatsNode>(TEXT("File System"), *CachePath, SpeedClass == ESpeedClass::Local);
-	Usage->UsageStats.Add(TEXT(""), UsageStats);
-	Usage->SpeedStats = SpeedStats;
-	return Usage;
-}
-
-TBitArray<> FFileSystemCacheStore::TryToPrefetch(const TConstArrayView<FString> CacheKeys)
-{
-	return CachedDataProbablyExistsBatch(CacheKeys);
-}
-
-bool FFileSystemCacheStore::ApplyDebugOptions(FBackendDebugOptions& InOptions)
+bool FFileSystemCacheStore::LegacyDebugOptions(FBackendDebugOptions& InOptions)
 {
 	DebugOptions = InOptions;
 	return true;
@@ -1797,7 +1468,7 @@ bool FFileSystemCacheStore::PutCacheRecord(
 	const FCacheRecordPolicy& Policy,
 	uint64& OutWriteSize)
 {
-	if (!IsWritable())
+	if (bReadOnly)
 	{
 		UE_LOG(LogDerivedDataCache, VeryVerbose,
 			TEXT("%s: Skipped put of %s from '%.*s' because this cache store is read-only"),
@@ -1809,7 +1480,7 @@ bool FFileSystemCacheStore::PutCacheRecord(
 	const ECachePolicy RecordPolicy = Policy.GetRecordPolicy();
 
 	// Skip the request if storing to the cache is disabled.
-	const ECachePolicy StoreFlag = SpeedClass == ESpeedClass::Local ? ECachePolicy::StoreLocal : ECachePolicy::StoreRemote;
+	const ECachePolicy StoreFlag = SpeedClass == EBackendSpeedClass::Local ? ECachePolicy::StoreLocal : ECachePolicy::StoreRemote;
 	if (!EnumHasAnyFlags(RecordPolicy, StoreFlag))
 	{
 		UE_LOG(LogDerivedDataCache, VeryVerbose, TEXT("%s: Skipped put of %s from '%.*s' due to cache policy"),
@@ -1829,7 +1500,7 @@ bool FFileSystemCacheStore::PutCacheRecord(
 
 	// Check if there is an existing record package.
 	FCbPackage ExistingPackage;
-	const ECachePolicy QueryFlag = SpeedClass == ESpeedClass::Local ? ECachePolicy::QueryLocal : ECachePolicy::QueryRemote;
+	const ECachePolicy QueryFlag = SpeedClass == EBackendSpeedClass::Local ? ECachePolicy::QueryLocal : ECachePolicy::QueryRemote;
 	bool bReplaceExisting = !EnumHasAnyFlags(RecordPolicy, QueryFlag);
 	bool bSavePackage = bReplaceExisting;
 	if (const bool bLoadPackage = !bReplaceExisting || !Algo::AllOf(Record.GetValues(), &FValue::HasData))
@@ -1963,16 +1634,8 @@ FOptionalCacheRecord FFileSystemCacheStore::GetCacheRecordOnly(
 	const FCacheKey& Key,
 	const FCacheRecordPolicy& Policy)
 {
-	if (!IsUsable())
-	{
-		UE_LOG(LogDerivedDataCache, VeryVerbose,
-			TEXT("%s: Skipped get of %s from '%.*s' because this cache store is not available"),
-			*CachePath, *WriteToString<96>(Key), Name.Len(), Name.GetData());
-		return FOptionalCacheRecord();
-	}
-
 	// Skip the request if querying the cache is disabled.
-	const ECachePolicy QueryFlag = SpeedClass == ESpeedClass::Local ? ECachePolicy::QueryLocal : ECachePolicy::QueryRemote;
+	const ECachePolicy QueryFlag = SpeedClass == EBackendSpeedClass::Local ? ECachePolicy::QueryLocal : ECachePolicy::QueryRemote;
 	if (!EnumHasAnyFlags(Policy.GetRecordPolicy(), QueryFlag))
 	{
 		UE_LOG(LogDerivedDataCache, VeryVerbose, TEXT("%s: Skipped get of %s from '%.*s' due to cache policy"),
@@ -2088,7 +1751,7 @@ bool FFileSystemCacheStore::PutCacheValue(
 	const ECachePolicy Policy,
 	uint64& OutWriteSize)
 {
-	if (!IsWritable())
+	if (bReadOnly)
 	{
 		UE_LOG(LogDerivedDataCache, VeryVerbose,
 			TEXT("%s: Skipped put of %s from '%.*s' because this cache store is read-only"),
@@ -2097,7 +1760,7 @@ bool FFileSystemCacheStore::PutCacheValue(
 	}
 
 	// Skip the request if storing to the cache is disabled.
-	const ECachePolicy StoreFlag = SpeedClass == ESpeedClass::Local ? ECachePolicy::StoreLocal : ECachePolicy::StoreRemote;
+	const ECachePolicy StoreFlag = SpeedClass == EBackendSpeedClass::Local ? ECachePolicy::StoreLocal : ECachePolicy::StoreRemote;
 	if (!EnumHasAnyFlags(Policy, StoreFlag))
 	{
 		UE_LOG(LogDerivedDataCache, VeryVerbose, TEXT("%s: Skipped put of %s from '%.*s' due to cache policy"),
@@ -2116,7 +1779,7 @@ bool FFileSystemCacheStore::PutCacheValue(
 	FCbPackage ExistingPackage;
 	TStringBuilder<256> Path;
 	BuildCachePackagePath(Key, Path);
-	const ECachePolicy QueryFlag = SpeedClass == ESpeedClass::Local ? ECachePolicy::QueryLocal : ECachePolicy::QueryRemote;
+	const ECachePolicy QueryFlag = SpeedClass == EBackendSpeedClass::Local ? ECachePolicy::QueryLocal : ECachePolicy::QueryRemote;
 	bool bReplaceExisting = !EnumHasAnyFlags(Policy, QueryFlag);
 	bool bSavePackage = bReplaceExisting;
 	if (const bool bLoadPackage = !bReplaceExisting || !Value.HasData())
@@ -2233,16 +1896,8 @@ bool FFileSystemCacheStore::GetCacheValueOnly(
 	const ECachePolicy Policy,
 	FValue& OutValue)
 {
-	if (!IsUsable())
-	{
-		UE_LOG(LogDerivedDataCache, VeryVerbose,
-			TEXT("%s: Skipped get of %s from '%.*s' because this cache store is not available"),
-			*CachePath, *WriteToString<96>(Key), Name.Len(), Name.GetData());
-		return false;
-	}
-
 	// Skip the request if querying the cache is disabled.
-	const ECachePolicy QueryFlag = SpeedClass == ESpeedClass::Local ? ECachePolicy::QueryLocal : ECachePolicy::QueryRemote;
+	const ECachePolicy QueryFlag = SpeedClass == EBackendSpeedClass::Local ? ECachePolicy::QueryLocal : ECachePolicy::QueryRemote;
 	if (!EnumHasAnyFlags(Policy, QueryFlag))
 	{
 		UE_LOG(LogDerivedDataCache, VeryVerbose, TEXT("%s: Skipped get of %s from '%.*s' due to cache policy"),
@@ -2597,7 +2252,6 @@ bool FFileSystemCacheStore::LoadFile(
 	const FStringView DebugName,
 	const TFunctionRef<void (FArchive& Ar)> ReadFunction) const
 {
-	check(IsUsable());
 	const double StartTime = FPlatformTime::Seconds();
 
 	TUniquePtr<FArchive> Ar = OpenFileRead(Path, DebugName);
@@ -2700,16 +2354,12 @@ ILegacyCacheStore* CreateFileSystemCacheStore(
 	const TCHAR* AccessLogPath,
 	ECacheStoreFlags& OutFlags)
 {
-	TUniquePtr<FFileSystemCacheStore> Store(new FFileSystemCacheStore(CachePath, Params, AccessLogPath));
-	if (Store->IsUsable())
+	TUniquePtr<FFileSystemCacheStore> Store = MakeUnique<FFileSystemCacheStore>(CachePath, Params, AccessLogPath, OutFlags);
+	if (OutFlags == ECacheStoreFlags::None)
 	{
-		ECacheStoreFlags Flags = ECacheStoreFlags::Query;
-		Flags |= Store->IsWritable() ? ECacheStoreFlags::Store : ECacheStoreFlags::None;
-		Flags |= Store->GetSpeedClass() == EBackendSpeedClass::Local ? ECacheStoreFlags::Local : ECacheStoreFlags::Remote;
-		OutFlags = Flags;
-		return Store.Release();
+		Store.Reset();
 	}
-	return nullptr;
+	return Store.Release();
 }
 
 } // UE::DerivedData
diff --git a/Engine/Source/Developer/DerivedDataCache/Private/PakFileCacheStore.cpp b/Engine/Source/Developer/DerivedDataCache/Private/PakFileCacheStore.cpp
--- a/Engine/Source/Developer/DerivedDataCache/Private/PakFileCacheStore.cpp
+++ b/Engine/Source/Developer/DerivedDataCache/Private/PakFileCacheStore.cpp
@@ -6,6 +6,7 @@
 #include "Algo/StableSort.h"
 #include "Algo/Transform.h"
 #include "Compression/OodleDataCompression.h"
+#include "DerivedDataBackendInterface.h"
 #include "DerivedDataCachePrivate.h"
 #include "DerivedDataCacheRecord.h"
 #include "DerivedDataCacheUsageStats.h"
@@ -46,44 +47,11 @@ public:
 
 	void Close() final;
 
-	/** Return a name for this interface */
-	FString GetName() const final;
+	bool IsWritable() const final { return bWriting; }
 
-	/** return true if this cache is writable **/
-	bool IsWritable() const final;
-
-	/** Returns a class of speed for this interface **/
-	ESpeedClass GetSpeedClass() const final;
-
-	bool BackfillLowerCacheLevels() const final;
-
-	/**
-	 * Synchronous test for the existence of a cache item
-	 *
-	 * @param	CacheKey	Alphanumeric+underscore key of this cache item
-	 * @return				true if the data probably will be found, this can't be guaranteed because of concurrency in the backends, corruption, etc
-	 */
-	bool CachedDataProbablyExists(const TCHAR* CacheKey) final;
-
-	/**
-	 * Synchronous retrieve of a cache item
-	 *
-	 * @param	CacheKey	Alphanumeric+underscore key of this cache item
-	 * @param	OutData		Buffer to receive the results, if any were found
-	 * @return				true if any data was found, and in this case OutData is non-empty
-	 */
-	bool GetCachedData(const TCHAR* CacheKey, TArray<uint8>& OutData) override;
-
-	/**
-	 * Asynchronous, fire-and-forget placement of a cache item
-	 *
-	 * @param	CacheKey	Alphanumeric+underscore key of this cache item
-	 * @param	InData		Buffer containing the data to cache, can be destroyed after the call returns, immediately
-	 * @param	bPutEvenIfExists	If true, then do not attempt skip the put even if CachedDataProbablyExists returns true
-	 */
-	EPutStatus PutCachedData(const TCHAR* CacheKey, TArrayView<const uint8> InData, bool bPutEvenIfExists) override;
-
-	void RemoveCachedData(const TCHAR* CacheKey, bool bTransient) final;
+	bool CachedDataProbablyExists(const TCHAR* CacheKey);
+	bool GetCachedData(const TCHAR* CacheKey, TArray<uint8>& OutData);
+	void PutCachedData(const TCHAR* CacheKey, TArrayView<const uint8> InData, bool bPutEvenIfExists);
 
 	/**
 	 * Save the cache to disk
@@ -108,18 +76,7 @@ public:
 		return CachePath;
 	}
 
-	TSharedRef<FDerivedDataCacheStatsNode> GatherUsageStats() const final;
-
-	TBitArray<> TryToPrefetch(TConstArrayView<FString> CacheKeys) final
-	{
-		return CachedDataProbablyExistsBatch(CacheKeys);
-	}
-
-	bool WouldCache(const TCHAR* CacheKey, TArrayView<const uint8> InData) final { return true; }
-
-	bool ApplyDebugOptions(FBackendDebugOptions& InOptions) final;
-
-	EBackendLegacyMode GetLegacyMode() const final { return EBackendLegacyMode::ValueOnly; }
+	// ICacheStore
 
 	void Put(
 		TConstArrayView<FCachePutRequest> Requests,
@@ -146,6 +103,11 @@ public:
 		IRequestOwner& Owner,
 		FOnCacheGetChunkComplete&& OnComplete) final;
 
+	// ILegacyCacheStore
+
+	void LegacyStats(FDerivedDataCacheStatsNode& OutNode) final;
+	bool LegacyDebugOptions(FBackendDebugOptions& Options) final;
+
 private:
 	[[nodiscard]] bool PutCacheRecord(FStringView Name, const FCacheRecord& Record, const FCacheRecordPolicy& Policy, uint64& OutWriteSize);
 
@@ -278,11 +240,6 @@ FPakFileCacheStore::~FPakFileCacheStore()
 	Close();
 }
 
-FString FPakFileCacheStore::GetName() const
-{
-	return CachePath;
-}
-
 void FPakFileCacheStore::Close()
 {
 	FDerivedDataBackend::Get().WaitForQuiescence();
@@ -299,21 +256,6 @@ void FPakFileCacheStore::Close()
 	}
 }
 
-bool FPakFileCacheStore::IsWritable() const
-{
-	return bWriting && !bClosed;
-}
-
-FDerivedDataBackendInterface::ESpeedClass FPakFileCacheStore::GetSpeedClass() const
-{
-	return ESpeedClass::Local;
-}
-
-bool FPakFileCacheStore::BackfillLowerCacheLevels() const
-{
-	return false;
-}
-
 bool FPakFileCacheStore::CachedDataProbablyExists(const TCHAR* CacheKey)
 {
 	COOK_STAT(auto Timer = UsageStats.TimeProbablyExists());
@@ -386,13 +328,10 @@ bool FPakFileCacheStore::GetCachedData(const TCHAR* CacheKey, TArray<uint8>& Out
 	return false;
 }
 
-FDerivedDataBackendInterface::EPutStatus FPakFileCacheStore::PutCachedData(const TCHAR* CacheKey, TArrayView<const uint8> InData, bool bPutEvenIfExists)
+void FPakFileCacheStore::PutCachedData(const TCHAR* CacheKey, TArrayView<const uint8> InData, bool bPutEvenIfExists)
 {
 	COOK_STAT(auto Timer = UsageStats.TimePut());
-	if (!IsWritable() || DebugOptions.ShouldSimulateGetMiss(CacheKey))
-	{
-		return EPutStatus::NotCached;
-	}
+	if (bWriting && !bClosed && !DebugOptions.ShouldSimulateGetMiss(CacheKey))
 	{
 		FWriteScopeLock ScopeLock(SynchronizationObject);
 		FString Key(CacheKey);
@@ -421,7 +360,7 @@ FDerivedDataBackendInterface::EPutStatus FPakFileCacheStore::PutCachedData(const
 						Item->Crc = Crc.GetValue();
 						FileHandle->Seek(Offset);
 					}
-					return EPutStatus::Cached;
+					return;
 				}
 
 				UE_LOG(LogDerivedDataCache, Warning,
@@ -438,7 +377,6 @@ FDerivedDataBackendInterface::EPutStatus FPakFileCacheStore::PutCachedData(const
 			CacheItems.Empty();
 			FileHandle.Reset();
 			UE_LOG(LogDerivedDataCache, Fatal, TEXT("%s: Could not write pak file... out of disk space?"), *CachePath);
-			return EPutStatus::NotCached;
 		}
 		else
 		{
@@ -450,25 +388,10 @@ FDerivedDataBackendInterface::EPutStatus FPakFileCacheStore::PutCachedData(const
 			FileHandle->Write(InData.GetData(), InData.Num());
 			UE_LOG(LogDerivedDataCache, Verbose, TEXT("%s: Put %s"), *CachePath, CacheKey);
 			CacheItems.Add(Key, FCacheValue(Offset, InData.Num(), Crc.GetValue()));
-			return EPutStatus::Cached;
 		}
 	}
 }
 
-void FPakFileCacheStore::RemoveCachedData(const TCHAR* CacheKey, bool bTransient)
-{
-	if (bClosed || bTransient)
-	{
-		return;
-	}
-	// strangish. We can delete from a pak, but it only deletes the index 
-	// if this is a read cache, it will read it next time
-	// if this is a write cache, we wasted space
-	FWriteScopeLock ScopeLock(SynchronizationObject);
-	FString Key(CacheKey);
-	CacheItems.Remove(Key);
-}
-
 bool FPakFileCacheStore::SaveCache()
 {
 	FWriteScopeLock ScopeLock(SynchronizationObject);
@@ -681,15 +604,13 @@ bool IPakFileCacheStore::SortAndCopy(const FString &InputFilename, const FString
 	return true;
 }
 
-TSharedRef<FDerivedDataCacheStatsNode> FPakFileCacheStore::GatherUsageStats() const
+void FPakFileCacheStore::LegacyStats(FDerivedDataCacheStatsNode& OutNode)
 {
-	TSharedRef<FDerivedDataCacheStatsNode> StatsNode =
-		MakeShared<FDerivedDataCacheStatsNode>(TEXT("PakFile"), CachePath, /*bIsLocal*/ true);
-	StatsNode->UsageStats.Add(TEXT(""), UsageStats);
-	return StatsNode;
+	OutNode = {TEXT("PakFile"), CachePath, /*bIsLocal*/ true};
+	OutNode.UsageStats.Add(TEXT(""), UsageStats);
 }
 
-bool FPakFileCacheStore::ApplyDebugOptions(FBackendDebugOptions& InOptions)
+bool FPakFileCacheStore::LegacyDebugOptions(FBackendDebugOptions& InOptions)
 {
 	DebugOptions = InOptions;
 	return true;
@@ -889,7 +810,7 @@ bool FPakFileCacheStore::PutCacheRecord(
 	const FCacheRecordPolicy& Policy,
 	uint64& OutWriteSize)
 {
-	if (!IsWritable())
+	if (!bWriting || bClosed)
 	{
 		return false;
 	}
@@ -1074,7 +995,7 @@ bool FPakFileCacheStore::PutCacheValue(
 	const ECachePolicy Policy,
 	uint64& OutWriteSize)
 {
-	if (!IsWritable())
+	if (!bWriting || bClosed)
 	{
 		return false;
 	}
@@ -1543,9 +1464,6 @@ class FCompressedPakFileCacheStore final : public FPakFileCacheStore
 public:
 	FCompressedPakFileCacheStore(const TCHAR* InFilename, bool bInWriting);
 
-	EPutStatus PutCachedData(const TCHAR* CacheKey, TArrayView<const uint8> InData, bool bPutEvenIfExists) final;
-	bool GetCachedData(const TCHAR* CacheKey, TArray<uint8>& OutData) final;
-
 	void Put(
 		TConstArrayView<FCachePutRequest> Requests,
 		IRequestOwner& Owner,
@@ -1557,8 +1475,6 @@ public:
 		FOnCachePutValueComplete&& OnComplete) final;
 
 private:
-	static const EName CompressionFormat = NAME_Zlib;
-	static const ECompressionFlags CompressionFlags = COMPRESS_BiasMemory;
 	static const ECompressedBufferCompressor RequiredCompressor = ECompressedBufferCompressor::Kraken;
 	static const ECompressedBufferCompressionLevel MinRequiredCompressionLevel = ECompressedBufferCompressionLevel::Optimal2;
 
@@ -1570,37 +1486,6 @@ FCompressedPakFileCacheStore::FCompressedPakFileCacheStore(const TCHAR* InFilena
 {
 }
 
-FDerivedDataBackendInterface::EPutStatus FCompressedPakFileCacheStore::PutCachedData(const TCHAR* CacheKey, TArrayView<const uint8> InData, bool bPutEvenIfExists)
-{
-	int32 UncompressedSize = InData.Num();
-	int32 CompressedSize = FCompression::CompressMemoryBound(CompressionFormat, UncompressedSize, CompressionFlags);
-
-	TArray<uint8> CompressedData;
-	CompressedData.AddUninitialized(CompressedSize + sizeof(UncompressedSize));
-
-	FMemory::Memcpy(&CompressedData[0], &UncompressedSize, sizeof(UncompressedSize));
-	verify(FCompression::CompressMemory(CompressionFormat, CompressedData.GetData() + sizeof(UncompressedSize), CompressedSize, InData.GetData(), InData.Num(), CompressionFlags));
-	CompressedData.SetNum(CompressedSize + sizeof(UncompressedSize), false);
-
-	return FPakFileCacheStore::PutCachedData(CacheKey, CompressedData, bPutEvenIfExists);
-}
-
-bool FCompressedPakFileCacheStore::GetCachedData(const TCHAR* CacheKey, TArray<uint8>& OutData)
-{
-	TArray<uint8> CompressedData;
-	if(!FPakFileCacheStore::GetCachedData(CacheKey, CompressedData))
-	{
-		return false;
-	}
-
-	int32 UncompressedSize;
-	FMemory::Memcpy(&UncompressedSize, &CompressedData[0], sizeof(UncompressedSize));
-	OutData.SetNum(UncompressedSize);
-	verify(FCompression::UncompressMemory(CompressionFormat, OutData.GetData(), UncompressedSize, CompressedData.GetData() + sizeof(UncompressedSize), CompressedData.Num() - sizeof(UncompressedSize), CompressionFlags));
-
-	return true;
-}
-
 void FCompressedPakFileCacheStore::Put(
 	const TConstArrayView<FCachePutRequest> Requests,
 	IRequestOwner& Owner,
diff --git a/Engine/Source/Developer/DerivedDataCache/Private/PakFileCacheStore.h b/Engine/Source/Developer/DerivedDataCache/Private/PakFileCacheStore.h
--- a/Engine/Source/Developer/DerivedDataCache/Private/PakFileCacheStore.h
+++ b/Engine/Source/Developer/DerivedDataCache/Private/PakFileCacheStore.h
@@ -3,16 +3,16 @@
 #pragma once
 
 #include "CoreTypes.h"
-#include "Containers/StringFwd.h"
-#include "DerivedDataBackendInterface.h"
+#include "DerivedDataLegacyCacheStore.h"
 
 namespace UE::DerivedData
 {
 
-class IPakFileCacheStore : public FDerivedDataBackendInterface
+class IPakFileCacheStore : public ILegacyCacheStore
 {
 public:
 	virtual void Close() = 0;
+	virtual bool IsWritable() const = 0;
 	virtual bool SaveCache() = 0;
 	virtual bool LoadCache(const TCHAR* InFilename) = 0;
 	virtual void MergeCache(IPakFileCacheStore* OtherPak) = 0;
diff --git a/Engine/Source/Developer/DerivedDataCache/Private/S3CacheStore.cpp b/Engine/Source/Developer/DerivedDataCache/Private/S3CacheStore.cpp
--- a/Engine/Source/Developer/DerivedDataCache/Private/S3CacheStore.cpp
+++ b/Engine/Source/Developer/DerivedDataCache/Private/S3CacheStore.cpp
@@ -1,6 +1,6 @@
 // Copyright Epic Games, Inc. All Rights Reserved.
 
-#include "DerivedDataBackendInterface.h"
+#include "DerivedDataLegacyCacheStore.h"
 
 #if WITH_S3_DDC_BACKEND
 
@@ -12,9 +12,8 @@
 	#include "Microsoft/HideMicrosoftPlatformTypes.h"
 #endif
 
-#include "Algo/AllOf.h"
 #include "Async/ParallelFor.h"
-#include "DerivedDataBackendCorruptionWrapper.h"
+#include "DerivedDataBackendInterface.h"
 #include "DerivedDataCacheRecord.h"
 #include "DerivedDataCachePrivate.h"
 #include "DerivedDataCacheUsageStats.h"
@@ -23,14 +22,12 @@
 #include "Dom/JsonObject.h"
 #include "HAL/PlatformFile.h"
 #include "HAL/PlatformFileManager.h"
-#include "HAL/Runnable.h"
 #include "HashingArchiveProxy.h"
 #include "Memory/SharedBuffer.h"
 #include "Misc/Base64.h"
 #include "Misc/ConfigCacheIni.h"
 #include "Misc/FeedbackContext.h"
 #include "Misc/FileHelper.h"
-#include "Misc/OutputDeviceRedirector.h"
 #include "Misc/Paths.h"
 #include "Misc/SecureHash.h"
 #include "ProfilingDebugging/CountersTrace.h"
@@ -40,7 +37,6 @@
 #include "Serialization/CompactBinaryValidation.h"
 #include "Serialization/JsonReader.h"
 #include "Serialization/JsonSerializer.h"
-#include "Serialization/MemoryReader.h"
 
 #include "curl/curl.h"
 
@@ -62,25 +58,14 @@
 namespace UE::DerivedData
 {
 
-TRACE_DECLARE_INT_COUNTER(S3DDC_Exist, TEXT("S3DDC Exist"));
-TRACE_DECLARE_INT_COUNTER(S3DDC_ExistHit, TEXT("S3DDC Exist Hit"));
 TRACE_DECLARE_INT_COUNTER(S3DDC_Get, TEXT("S3DDC Get"));
 TRACE_DECLARE_INT_COUNTER(S3DDC_GetHit, TEXT("S3DDC Get Hit"));
-TRACE_DECLARE_INT_COUNTER(S3DDC_BytesRecieved, TEXT("S3DDC Bytes Recieved"));
 
 ///////////////////////////////////////////////////////////////////////////////////////////////////////////////////////
 
-void BuildPathForLegacyCache(const TCHAR* const CacheKey, FStringBuilderBase& Path);
 void BuildPathForCachePackage(const FCacheKey& CacheKey, FStringBuilderBase& Path);
 void BuildPathForCacheContent(const FIoHash& RawHash, FStringBuilderBase& Path);
 
-static FString BuildPathForLegacyCache(const TCHAR* CacheKey)
-{
-	TStringBuilder<256> Path;
-	BuildPathForLegacyCache(CacheKey, Path);
-	return FString(Path);
-}
-
 class FStringAnsi
 {
 public:
@@ -192,18 +177,18 @@ bool IsSuccessfulHttpResponse(long ResponseCode)
 
 struct IRequestCallback
 {
-	virtual ~IRequestCallback() { }
+	virtual ~IRequestCallback() = default;
 	virtual bool Update(int NumBytes, int TotalBytes) = 0;
 };
 
 /**
  * Backend for a read-only AWS S3 based caching service.
  **/
-class FS3CacheStore final : public FDerivedDataBackendInterface
+class FS3CacheStore final : public ILegacyCacheStore
 {
 public:
 	/**
-	 * Creates the backend, checks health status and attempts to acquire an access token.
+	 * Creates the cache store, checks health status and attempts to acquire an access token.
 	 *
 	 * @param  InRootManifestPath   Local path to the JSON manifest in the workspace containing a list of files to download
 	 * @param  InBaseUrl            Base URL for the bucket, with trailing slash (eg. https://foo.s3.us-east-1.amazonaws.com/)
@@ -212,34 +197,16 @@ public:
 	 * @param  InCachePath          Path to cache the DDC files
 	 */
 	FS3CacheStore(const TCHAR* InRootManifestPath, const TCHAR* InBaseUrl, const TCHAR* InRegion, const TCHAR* InCanaryObjectKey, const TCHAR* InCachePath);
-	~FS3CacheStore() final;
+
+	inline const FString& GetName() const { return BaseUrl; }
 
 	/**
-	 * Checks is backend is usable (reachable and accessible).
+	 * Checks if cache store is usable (reachable and accessible).
 	 * @return true if usable
 	 */
-	bool IsUsable() const;
-
-	/* S3 Cache cannot be written to*/
-	bool IsWritable() const final { return false; }
-
-	/* S3 Cache does not try to write back to lower caches (e.g. Shared DDC) */
-	bool BackfillLowerCacheLevels() const final { return false; }
-
-	bool CachedDataProbablyExists(const TCHAR* CacheKey) final;
-	bool GetCachedData(const TCHAR* CacheKey, TArray<uint8>& OutData) final;
-	EPutStatus PutCachedData(const TCHAR* CacheKey, TArrayView<const uint8> InData, bool bPutEvenIfExists) final;
-	void RemoveCachedData(const TCHAR* CacheKey, bool bTransient) final;
-	TSharedRef<FDerivedDataCacheStatsNode> GatherUsageStats() const final;
-
-	FString GetName() const final;
-	ESpeedClass GetSpeedClass() const final;
-	TBitArray<> TryToPrefetch(TConstArrayView<FString> CacheKeys) final;
-	bool WouldCache(const TCHAR* CacheKey, TArrayView<const uint8> InData) final;
+	inline bool IsUsable() const { return bEnabled; }
 
-	bool ApplyDebugOptions(FBackendDebugOptions& InOptions) final;
-
-	EBackendLegacyMode GetLegacyMode() const final { return EBackendLegacyMode::ValueWithLegacyFallback; }
+	// ICacheStore
 
 	void Put(
 		TConstArrayView<FCachePutRequest> Requests,
@@ -262,6 +229,11 @@ public:
 		IRequestOwner& Owner,
 		FOnCacheGetChunkComplete&& OnComplete) final;
 
+	// ILegacyCacheStore
+
+	void LegacyStats(FDerivedDataCacheStatsNode& OutNode) final;
+	bool LegacyDebugOptions(FBackendDebugOptions& Options) final;
+
 private:
 	struct FBundle;
 	struct FBundleEntry;
@@ -324,7 +296,6 @@ private:
 	bool DownloadManifest(const FRootManifest& RootManifest, FFeedbackContext* Context);
 	void RemoveUnusedBundles();
 	void ReadBundle(FBundle& Bundle);
-	bool FindBundleEntry(const TCHAR* CacheKey, const FBundle*& OutBundle, const FBundleEntry*& OutBundleEntry) const;
 
 	FBackendDebugOptions DebugOptions;
 };
@@ -1044,112 +1015,10 @@ FS3CacheStore::FS3CacheStore(const TCHAR* InRootManifestPath, const TCHAR* InBas
 	}
 }
 
-FS3CacheStore::~FS3CacheStore()
-{
-}
-
-bool FS3CacheStore::IsUsable() const
+void FS3CacheStore::LegacyStats(FDerivedDataCacheStatsNode& OutNode)
 {
-	return bEnabled;
-}
-
-bool FS3CacheStore::CachedDataProbablyExists(const TCHAR* CacheKey)
-{
-	TRACE_CPUPROFILER_EVENT_SCOPE(S3DDC_Exist);
-	TRACE_COUNTER_ADD(S3DDC_Exist, int64(1));
-	COOK_STAT(auto Timer = UsageStats.TimeProbablyExists());
-
-	const FBundle* Bundle;
-	const FBundleEntry* BundleEntry;
-
-	if (DebugOptions.ShouldSimulateGetMiss(CacheKey))
-	{
-		return false;
-	}
-
-	if (!FindBundleEntry(CacheKey, Bundle, BundleEntry))
-	{
-		UE_LOG(LogDerivedDataCache, Verbose, TEXT("S3DerivedDataBackend: Cache miss on %s (probably)"), CacheKey);
-		return false;
-	}
-
-	TRACE_COUNTER_ADD(S3DDC_ExistHit, int64(1));
-	COOK_STAT(Timer.AddHit(BundleEntry->Length));
-	return true;
-}
-
-bool FS3CacheStore::GetCachedData(const TCHAR* CacheKey, TArray<uint8>& OutData)
-{
-	TRACE_CPUPROFILER_EVENT_SCOPE(S3DDC_Get);
-	TRACE_COUNTER_ADD(S3DDC_Get, int64(1));
-	COOK_STAT(auto Timer = UsageStats.TimeGet());
-
-	if (DebugOptions.ShouldSimulateGetMiss(CacheKey))
-	{
-		return false;
-	}
-
-	const FBundle* Bundle;
-	const FBundleEntry* BundleEntry;
-	if (FindBundleEntry(CacheKey, Bundle, BundleEntry))
-	{
-		TUniquePtr<FArchive> Reader(IFileManager::Get().CreateFileReader(*Bundle->LocalFile));
-		if (Reader.IsValid() && !Reader->IsError())
-		{
-			TRACE_COUNTER_ADD(S3DDC_GetHit, int64(1));
-			COOK_STAT(Timer.AddHit(BundleEntry->Length));
-			UE_LOG(LogDerivedDataCache, Verbose, TEXT("S3DerivedDataBackend: Cache hit on %s"), CacheKey);
-			OutData.SetNum(BundleEntry->Length);
-			Reader->Seek(BundleEntry->Offset);
-			Reader->Serialize(OutData.GetData(), BundleEntry->Length);
-			if (FCorruptionWrapper::ReadTrailer(OutData, *BaseUrl, CacheKey))
-			{
-				return true;
-			}
-		}
-	}
-
-	UE_LOG(LogDerivedDataCache, Verbose, TEXT("S3DerivedDataBackend: Cache miss on %s"), CacheKey);
-	return false;
-}
-
-FDerivedDataBackendInterface::EPutStatus FS3CacheStore::PutCachedData(const TCHAR* CacheKey, TArrayView<const uint8> InData, bool bPutEvenIfExists)
-{
-	// Not implemented
-	return EPutStatus::NotCached;
-}
-
-void FS3CacheStore::RemoveCachedData(const TCHAR* CacheKey, bool bTransient)
-{
-	// Not implemented
-}
-
-TSharedRef<FDerivedDataCacheStatsNode> FS3CacheStore::GatherUsageStats() const
-{
-	TSharedRef<FDerivedDataCacheStatsNode> StatsNode =
-		MakeShared<FDerivedDataCacheStatsNode>(TEXT("S3"), BaseUrl, /*bIsLocal*/ false);
-	StatsNode->UsageStats.Add(TEXT(""), UsageStats);
-	return StatsNode;
-}
-
-FString FS3CacheStore::GetName() const
-{
-	return BaseUrl;
-}
-
-FDerivedDataBackendInterface::ESpeedClass FS3CacheStore::GetSpeedClass() const
-{
-	return ESpeedClass::Local;
-}
-
-TBitArray<> FS3CacheStore::TryToPrefetch(TConstArrayView<FString> CacheKeys)
-{
-	return CachedDataProbablyExistsBatch(CacheKeys);
-}
-
-bool FS3CacheStore::WouldCache(const TCHAR* CacheKey, TArrayView<const uint8> InData)
-{
-	return false;
+	OutNode = {TEXT("S3"), BaseUrl, /*bIsLocal*/ false};
+	OutNode.UsageStats.Add(TEXT(""), UsageStats);
 }
 
 bool FS3CacheStore::DownloadManifest(const FRootManifest& RootManifest, FFeedbackContext* Context)
@@ -1304,28 +1173,7 @@ void FS3CacheStore::ReadBundle(FBundle& Bundle)
 	}
 }
 
-bool FS3CacheStore::FindBundleEntry(const TCHAR* CacheKey, const FBundle*& OutBundle, const FBundleEntry*& OutBundleEntry) const
-{
-	FSHAHash Hash;
-
-	auto AnsiString = StringCast<ANSICHAR>(*BuildPathForLegacyCache(CacheKey).ToUpper());
-	FSHA1::HashBuffer(AnsiString.Get(), AnsiString.Length(), Hash.Hash);
-
-	for (const FBundle& Bundle : Bundles)
-	{
-		const FBundleEntry* Entry = Bundle.Entries.Find(Hash);
-		if (Entry != nullptr)
-		{
-			OutBundle = &Bundle;
-			OutBundleEntry = Entry;
-			return true;
-		}
-	}
-
-	return false;
-}
-
-bool FS3CacheStore::ApplyDebugOptions(FBackendDebugOptions& InOptions)
+bool FS3CacheStore::LegacyDebugOptions(FBackendDebugOptions& InOptions)
 {
 	DebugOptions = InOptions;
 	return true;
diff --git a/Engine/Source/Developer/DerivedDataCache/Public/DerivedDataCacheKey.h b/Engine/Source/Developer/DerivedDataCache/Public/DerivedDataCacheKey.h
--- a/Engine/Source/Developer/DerivedDataCache/Public/DerivedDataCacheKey.h
+++ b/Engine/Source/Developer/DerivedDataCache/Public/DerivedDataCacheKey.h
@@ -6,6 +6,7 @@
 #include "Containers/StringFwd.h"
 #include "Containers/StringView.h"
 #include "IO/IoHash.h"
+#include "Misc/AsciiSet.h"
 #include "Templates/TypeHash.h"
 
 #define UE_API DERIVEDDATACACHE_API
@@ -20,11 +21,22 @@ namespace UE::DerivedData
 /**
  * An alphanumeric identifier that groups related cache records.
  *
- * A cache bucket name must be alphanumeric, non-empty, and contain fewer than 256 code units.
+ * A cache bucket name must be alphanumeric, non-empty, and contain fewer than 64 code units.
  */
 class FCacheBucket
 {
 public:
+	/** Maximum number of code units in a valid cache bucket name. */
+	static constexpr int32 MaxNameLen = 63;
+
+	/** Returns true if the name is a valid cache bucket name. */
+	template <typename CharType>
+	static inline bool IsValidName(TStringView<CharType> Name)
+	{
+		constexpr FAsciiSet Valid("ABCDEFGHIJKLMNOPQRSTUVWXYZabcdefghijklmnopqrstuvwxyz0123456789");
+		return !Name.IsEmpty() && Name.Len() <= MaxNameLen && FAsciiSet::HasOnly(Name, Valid);
+	}
+
 	/** Construct a null cache bucket. */
 	FCacheBucket() = default;
 
