only define CURL_NO_OLDIES if it isn't defined already.

#rnx

diff --git a/Engine/Source/Developer/DerivedDataCache/Private/Http/HttpClient.cpp b/Engine/Source/Developer/DerivedDataCache/Private/Http/HttpClient.cpp
--- a/Engine/Source/Developer/DerivedDataCache/Private/Http/HttpClient.cpp
+++ b/Engine/Source/Developer/DerivedDataCache/Private/Http/HttpClient.cpp
@@ -6,7 +6,10 @@
 #include "Misc/ScopeRWLock.h"
 #include "Tasks/Task.h"
 
-#define CURL_NO_OLDIES
+#if !defined(CURL_NO_OLDIES)
+	#define CURL_NO_OLDIES
+#endif
+
 #if PLATFORM_MICROSOFT
 #include "Microsoft/AllowMicrosoftPlatformTypes.h"
 #endif
diff --git a/Engine/Source/Developer/DerivedDataCache/Private/Http/HttpRequest.cpp b/Engine/Source/Developer/DerivedDataCache/Private/Http/HttpRequest.cpp
--- a/Engine/Source/Developer/DerivedDataCache/Private/Http/HttpRequest.cpp
+++ b/Engine/Source/Developer/DerivedDataCache/Private/Http/HttpRequest.cpp
@@ -17,7 +17,10 @@
 #include <openssl/ssl.h>
 #endif
 
-#define CURL_NO_OLDIES
+#if !defined(CURL_NO_OLDIES)
+	#define CURL_NO_OLDIES
+#endif
+
 #if PLATFORM_MICROSOFT
 #include "Microsoft/AllowMicrosoftPlatformTypes.h"
 #endif
