DDC: Fixed the pak file cache IsWritable function to take into account whether the file is closed

#rnx

diff --git a/Engine/Source/Developer/DerivedDataCache/Private/PakFileCacheStore.cpp b/Engine/Source/Developer/DerivedDataCache/Private/PakFileCacheStore.cpp
--- a/Engine/Source/Developer/DerivedDataCache/Private/PakFileCacheStore.cpp
+++ b/Engine/Source/Developer/DerivedDataCache/Private/PakFileCacheStore.cpp
@@ -47,7 +47,7 @@ public:
 
 	void Close() final;
 
-	bool IsWritable() const final { return bWriting; }
+	bool IsWritable() const final { return bWriting && !bClosed; }
 
 	bool CachedDataProbablyExists(const TCHAR* CacheKey);
 	bool GetCachedData(const TCHAR* CacheKey, TArray<uint8>& OutData);
