Reduce HTTP connection limit from 32 to 8 for the DDC HTTP client.

#rnx

diff --git a/Engine/Source/Developer/DerivedDataCache/Private/Http/HttpClient.cpp b/Engine/Source/Developer/DerivedDataCache/Private/Http/HttpClient.cpp
--- a/Engine/Source/Developer/DerivedDataCache/Private/Http/HttpClient.cpp
+++ b/Engine/Source/Developer/DerivedDataCache/Private/Http/HttpClient.cpp
@@ -30,7 +30,7 @@ namespace UE
 namespace Http::Private
 {
 
-static constexpr uint32 MaxTotalConnections = 32;
+static constexpr uint32 MaxTotalConnections = 8;
 static constexpr uint32 WaitIntervalMs = 10;
 
 struct FHttpSharedDataInternals
