Change Zen cache store to issue non-blocking HTTP requests.  GetChunks remains blocking because of the sequential nature of the way it consumes results.  This can be changed in the future.  Fixed a bug in the code for async POST operations to ensure it sets the post size appropriately.

This change requires a fix (in progress) to EditorDomain to avoid a deadlock.

diff --git a/Engine/Source/Developer/DerivedDataCache/Private/Http/HttpClient.cpp b/Engine/Source/Developer/DerivedDataCache/Private/Http/HttpClient.cpp
--- a/Engine/Source/Developer/DerivedDataCache/Private/Http/HttpClient.cpp
+++ b/Engine/Source/Developer/DerivedDataCache/Private/Http/HttpClient.cpp
@@ -2,6 +2,7 @@
 
 #include "HttpClient.h"
 
+#include "Algo/RemoveIf.h"
 #include "Misc/AsciiSet.h"
 #include "Misc/ScopeRWLock.h"
 #include "Tasks/Task.h"
@@ -93,7 +94,7 @@ FAnsiStringBuilderBase& operator<<(FAnsiStringBuilderBase& Builder, const FHttpA
 	return Builder.Append(Token.Header);
 }
 
-FHttpSharedData::FHttpSharedData()
+FHttpSharedData::FHttpSharedData(uint32 OverrideMaxConnections)
 	: PendingRequestEvent(EEventMode::AutoReset)
 {
 	Internals = MakePimpl<Http::Private::FHttpSharedDataInternals>();
@@ -104,7 +105,7 @@ FHttpSharedData::FHttpSharedData()
 	curl_share_setopt(Internals->CurlShare, CURLSHOPT_SHARE, CURL_LOCK_DATA_DNS);
 	curl_share_setopt(Internals->CurlShare, CURLSHOPT_SHARE, CURL_LOCK_DATA_SSL_SESSION);
 	Internals->CurlMulti = curl_multi_init();
-	curl_multi_setopt(Internals->CurlMulti, CURLMOPT_MAX_TOTAL_CONNECTIONS, Http::Private::MaxTotalConnections);
+	curl_multi_setopt(Internals->CurlMulti, CURLMOPT_MAX_TOTAL_CONNECTIONS, OverrideMaxConnections ? OverrideMaxConnections : Http::Private::MaxTotalConnections);
 	curl_multi_setopt(Internals->CurlMulti, CURLMOPT_PIPELINING, CURLPIPE_MULTIPLEX);
 }
 
@@ -129,7 +130,7 @@ void FHttpSharedData::AddRequest(void* Curl)
 
 	if (!bAsyncThreadStarting.load(std::memory_order_relaxed) && !bAsyncThreadStarting.exchange(true, std::memory_order_relaxed))
 	{
-		AsyncServiceThread = FThread(TEXT("HttpCacheStore"), [this] { ProcessAsyncRequests(); }, 64 * 1024, TPri_Normal);
+		AsyncServiceThread = FThread(TEXT("HttpService"), [this] { ProcessAsyncRequests(); }, 64 * 1024, TPri_Normal);
 	}
 }
 
@@ -141,8 +142,10 @@ void* FHttpSharedData::GetCurlShare() const
 void FHttpSharedData::ProcessAsyncRequests()
 {
 	int ActiveTransfers = 0;
+	TArray<Tasks::TTask<void>> PendingCompletionTasks;
+	PendingCompletionTasks.Reserve(128);
 
-	auto ProcessPendingRequests = [this, &ActiveTransfers]
+	auto ProcessPendingRequests = [this, &ActiveTransfers, &PendingCompletionTasks]
 	{
 		int CurrentActiveTransfers = -1;
 
@@ -183,10 +186,10 @@ void FHttpSharedData::ProcessAsyncRequests()
 							// It is important that the CompleteAsync call doesn't happen on this thread as it is possible it will block waiting
 							// for a free HTTP request, and if that happens on this thread, we can deadlock as no HTTP requests will become 
 							// available while this thread is blocked.
-							Tasks::Launch(TEXT("FHttpRequest::CompleteAsync"), [CompletedRequest, Result = Message->data.result]() mutable
+							PendingCompletionTasks.Emplace(Tasks::Launch(TEXT("FHttpRequest::CompleteAsync"), [CompletedRequest, Result = Message->data.result]() mutable
 							{
 								CompletedRequest->CompleteAsync(Result);
-							});
+							}));
 						}
 					}
 				}
@@ -204,7 +207,16 @@ void FHttpSharedData::ProcessAsyncRequests()
 	do
 	{
 		ProcessPendingRequests();
-		PendingRequestEvent.Wait(100);
+
+		PendingCompletionTasks.SetNum(Algo::StableRemoveIf(PendingCompletionTasks, [](const Tasks::TTask<void>& Task) { return Task.IsCompleted(); }), false);
+		if (!PendingCompletionTasks.IsEmpty())
+		{
+			Tasks::Wait(PendingCompletionTasks, FTimespan::FromMilliseconds(100));
+		}
+		else
+		{
+			PendingRequestEvent.Wait(100);
+		}
 	}
 	while (!FHttpSharedData::bAsyncThreadShutdownRequested.load(std::memory_order_relaxed));
 
@@ -275,7 +287,7 @@ FHttpRequest* FHttpRequestPool::GetFreeRequest(bool bUnboundedOverflow)
 
 FHttpRequest* FHttpRequestPool::WaitForFreeRequest(bool bUnboundedOverflow)
 {
-	TRACE_CPUPROFILER_EVENT_SCOPE(HttpDDC_WaitForConnection);
+	TRACE_CPUPROFILER_EVENT_SCOPE(HttpDDC_WaitForFreeRequest);
 
 	FHttpRequest* Request = GetFreeRequest(bUnboundedOverflow);
 	if (Request == nullptr)
diff --git a/Engine/Source/Developer/DerivedDataCache/Private/Http/HttpClient.h b/Engine/Source/Developer/DerivedDataCache/Private/Http/HttpClient.h
--- a/Engine/Source/Developer/DerivedDataCache/Private/Http/HttpClient.h
+++ b/Engine/Source/Developer/DerivedDataCache/Private/Http/HttpClient.h
@@ -105,7 +105,7 @@ private:
 class FHttpSharedData
 {
 public:
-	FHttpSharedData();
+	FHttpSharedData(uint32 OverrideMaxConnections = 0);
 	~FHttpSharedData();
 
 private:
diff --git a/Engine/Source/Developer/DerivedDataCache/Private/Http/HttpRequest.cpp b/Engine/Source/Developer/DerivedDataCache/Private/Http/HttpRequest.cpp
--- a/Engine/Source/Developer/DerivedDataCache/Private/Http/HttpRequest.cpp
+++ b/Engine/Source/Developer/DerivedDataCache/Private/Http/HttpRequest.cpp
@@ -777,7 +777,7 @@ void FHttpRequest::EnqueueAsyncPost(
 
 	CURL* const LocalCurl = static_cast<CURL*>(Curl);
 	curl_easy_setopt(LocalCurl, CURLOPT_POST, 1L);
-	curl_easy_setopt(LocalCurl, CURLOPT_INFILESIZE, Buffer.GetSize());
+	curl_easy_setopt(LocalCurl, CURLOPT_POSTFIELDSIZE, Buffer.GetSize());
 	curl_easy_setopt(LocalCurl, CURLOPT_READDATA, this);
 	curl_easy_setopt(LocalCurl, CURLOPT_READFUNCTION, Http::Private::FHttpRequestStatics::StaticReadFn);
 	AddContentTypeHeader(TEXTVIEW("Content-Type"), ContentType);
diff --git a/Engine/Source/Developer/DerivedDataCache/Private/Tests/HttpCacheStoreTest.cpp b/Engine/Source/Developer/DerivedDataCache/Private/Tests/HttpCacheStoreTest.cpp
--- a/Engine/Source/Developer/DerivedDataCache/Private/Tests/HttpCacheStoreTest.cpp
+++ b/Engine/Source/Developer/DerivedDataCache/Private/Tests/HttpCacheStoreTest.cpp
@@ -598,8 +598,8 @@ TArray<FValue> CreateTestCacheValues(ICacheStore* InTestBackend, uint32 InNumVal
 }
 
 // Tests basic functionality for structured cache operations
-IMPLEMENT_HTTPDERIVEDDATA_AUTOMATION_TEST(CacheStore, TEXT(".CacheStore"), EAutomationTestFlags::EditorContext | EAutomationTestFlags::ProductFilter)
-bool CacheStore::RunTest(const FString& Parameters)
+IMPLEMENT_HTTPDERIVEDDATA_AUTOMATION_TEST(FHttpCacheStoreTest, TEXT(".CacheStore"), EAutomationTestFlags::EditorContext | EAutomationTestFlags::ProductFilter)
+bool FHttpCacheStoreTest::RunTest(const FString& Parameters)
 {
 	using namespace UE::DerivedData;
 	ILegacyCacheStore* TestBackend = GetTestBackend();
diff --git a/Engine/Source/Developer/DerivedDataCache/Private/ZenCacheStore.cpp b/Engine/Source/Developer/DerivedDataCache/Private/ZenCacheStore.cpp
--- a/Engine/Source/Developer/DerivedDataCache/Private/ZenCacheStore.cpp
+++ b/Engine/Source/Developer/DerivedDataCache/Private/ZenCacheStore.cpp
@@ -11,6 +11,7 @@
 #include "DerivedDataCacheRecord.h"
 #include "DerivedDataCacheUsageStats.h"
 #include "DerivedDataChunk.h"
+#include "DerivedDataRequestOwner.h"
 #include "Http/HttpClient.h"
 #include "Misc/App.h"
 #include "Misc/ConfigCacheIni.h"
@@ -123,7 +124,7 @@ public:
 private:
 	bool IsServiceReady();
 
-	using FOnRpcComplete = TUniqueFunction<void(FHttpRequest::EResult HttpResult, FHttpRequest* Request, FCbPackage& OutResponse)>;
+	using FOnRpcComplete = TUniqueFunction<void(FHttpRequest::EResult HttpResult, FHttpRequest* Request, FCbPackage& Response)>;
 	static FHttpRequest::EResult ParseRpcResponse(FHttpRequest::EResult ResultFromPost, FHttpRequest& Request, FCbPackage& OutResponse);
 	static FHttpRequest::EResult PerformBlockingRpc(FHttpRequest& Request,
 		FStringView Uri,
@@ -146,11 +147,25 @@ private:
 		const FCbPackage& RequestPackage,
 		FOnRpcComplete&& OnComplete);
 
+
+	template <typename T, typename... ArgTypes>
+	static TRefCountPtr<T> MakeAsyncOp(ArgTypes&&... Args)
+	{
+		// TODO: This should in-place construct from a pre-allocated memory pool
+		return TRefCountPtr<T>(new T(Forward<ArgTypes>(Args)...));
+	}
 private:
+	friend class FPutOp;
+	friend class FGetOp;
+	friend class FPutValueOp;
+	friend class FGetValueOp;
+	friend class FGetChunksOp;
+
 	FString Namespace;
 	FString StructuredNamespace;
 	UE::Zen::FScopeZenService ZenService;
 	mutable FDerivedDataCacheUsageStats UsageStats;
+	TUniquePtr<FHttpSharedData> SharedData;
 	TUniquePtr<FHttpRequestPool> RequestPool;
 	bool bIsUsable = false;
 	int32 BatchPutMaxBytes = 1024*1024;
@@ -161,546 +176,539 @@ private:
 	FBackendDebugOptions DebugOptions;
 };
 
-FZenCacheStore::FZenCacheStore(
-	const TCHAR* InServiceUrl,
-	const TCHAR* InNamespace,
-	const TCHAR* InStructuredNamespace)
-	: Namespace(InNamespace)
-	, StructuredNamespace(InStructuredNamespace)
-	, ZenService(InServiceUrl)
+class FPutOp final : public FThreadSafeRefCountedObject
 {
-	if (IsServiceReady())
+public:
+	FPutOp(FZenCacheStore& InCacheStore,
+		IRequestOwner& InOwner,
+		const TConstArrayView<FCachePutRequest> InRequests,
+		FOnCachePutComplete&& InOnComplete)
+		: CacheStore(InCacheStore)
+		, Owner(InOwner)
+		, Requests(InRequests)
+		, Batches(Requests, [this](const FCachePutRequest& NextRequest) {return BatchGroupingFilter(NextRequest);})
+		, OnComplete(MoveTemp(InOnComplete))
 	{
-		RequestPool = MakeUnique<FHttpRequestPool>(ZenService.GetInstance().GetURL(), ZenService.GetInstance().GetURL(), nullptr, nullptr, 32);
-		bIsUsable = true;
-
-		// Issue a request for stats as it will be fetched asynchronously and issuing now makes them available sooner for future callers.
-		Zen::FZenStats ZenStats;
-		ZenService.GetInstance().GetStats(ZenStats);
 	}
 
-	GConfig->GetInt(TEXT("Zen"), TEXT("BatchPutMaxBytes"), BatchPutMaxBytes, GEngineIni);
-	GConfig->GetInt(TEXT("Zen"), TEXT("CacheRecordBatchSize"), CacheRecordBatchSize, GEngineIni);
-	GConfig->GetInt(TEXT("Zen"), TEXT("CacheChunksBatchSize"), CacheChunksBatchSize, GEngineIni);
-}
+	void IssueRequests()
+	{
+		FRequestBarrier Barrier(Owner);
+		for (TArrayView<const FCachePutRequest> Batch : Batches)
+		{
+			FCbPackage BatchPackage;
+			FCbWriter BatchWriter;
+			BatchWriter.BeginObject();
+			{
+				BatchWriter << ANSITEXTVIEW("Method") << "PutCacheRecords";
+				BatchWriter.BeginObject(ANSITEXTVIEW("Params"));
+				{
+					ECachePolicy BatchDefaultPolicy = Batch[0].Policy.GetRecordPolicy();
+					BatchWriter << ANSITEXTVIEW("DefaultPolicy") << *WriteToString<128>(BatchDefaultPolicy);
+					BatchWriter.AddString(ANSITEXTVIEW("Namespace"), CacheStore.StructuredNamespace);
 
-bool FZenCacheStore::IsServiceReady()
-{
-	return ZenService.GetInstance().IsServiceReady();
-}
+					BatchWriter.BeginArray(ANSITEXTVIEW("Requests"));
+					for (const FCachePutRequest& Request : Batch)
+					{
+						const FCacheRecord& Record = Request.Record;
 
-FHttpRequest::EResult FZenCacheStore::ParseRpcResponse(FHttpRequest::EResult ResultFromPost, FHttpRequest& Request, FCbPackage& OutResponse)
-{
-	if (ResultFromPost != FHttpRequest::EResult::Success || !FHttpRequest::IsSuccessResponse(Request.GetResponseCode()))
-	{
-		return FHttpRequest::EResult::Failed;
+						BatchWriter.BeginObject();
+						{
+							BatchWriter.SetName(ANSITEXTVIEW("Record"));
+							Record.Save(BatchPackage, BatchWriter);
+							if (!Request.Policy.IsUniform() || Request.Policy.GetRecordPolicy() != BatchDefaultPolicy)
+							{
+								BatchWriter.SetName(ANSITEXTVIEW("Policy"));
+								Request.Policy.Save(BatchWriter);
+							}
+						}
+						BatchWriter.EndObject();
+					}
+					BatchWriter.EndArray();
+				}
+				BatchWriter.EndObject();
+			}
+			BatchWriter.EndObject();
+			BatchPackage.SetObject(BatchWriter.Save().AsObject());
+
+			auto OnRpcComplete = [this, OpRef = TRefCountPtr<FPutOp>(this), Batch](FHttpRequest::EResult HttpResult, FHttpRequest* HttpRequest, FCbPackage& Response)
+			{
+				int32 RequestIndex = 0;
+				if (HttpResult == FHttpRequest::EResult::Success)
+				{
+					const FCbObject& ResponseObj = Response.GetObject();
+					for (FCbField ResponseField : ResponseObj[ANSITEXTVIEW("Result")])
+					{
+						if (RequestIndex >= Batch.Num())
+						{
+							++RequestIndex;
+							continue;
+						}
+						const FCachePutRequest& Request = Batch[RequestIndex++];
+
+						const FCacheKey& Key = Request.Record.GetKey();
+						bool bPutSucceeded = ResponseField.AsBool();
+						if (CacheStore.DebugOptions.ShouldSimulatePutMiss(Key))
+						{
+							UE_LOG(LogDerivedDataCache, Verbose, TEXT("%s: Simulated miss for put of %s from '%s'"),
+								*CacheStore.GetName(), *WriteToString<96>(Key), *Request.Name);
+							bPutSucceeded = false;
+						}
+						bPutSucceeded ? OnHit(Request) : OnMiss(Request);
+					}
+					if (RequestIndex != Batch.Num())
+					{
+						UE_LOG(LogDerivedDataCache, Warning,
+							TEXT("%s: Invalid response received from PutCacheRecords rpc: %d results expected, received %d."), *CacheStore.GetName(), Batch.Num(), RequestIndex);
+					}
+				}
+				for (const FCachePutRequest& Request : Batch.RightChop(RequestIndex))
+				{
+					OnMiss(Request);
+				}
+			};
+			FHttpRequest* Request = CacheStore.RequestPool->WaitForFreeRequest();
+			CacheStore.EnqueueAsyncRpc(*Request, Owner, CacheStore.RequestPool.Get(), TEXTVIEW("/z$/$rpc"), BatchPackage, MoveTemp(OnRpcComplete));
+		}
 	}
 
-	const TArray64<uint8>& ResponseBuffer = Request.GetResponseBuffer();
-	if (ResponseBuffer.Num())
+private:
+	EBatchView BatchGroupingFilter(const FCachePutRequest& NextRequest)
 	{
-		FLargeMemoryReader Ar(ResponseBuffer.GetData(), ResponseBuffer.Num());
-		if (!OutResponse.TryLoad(Ar))
+		const FCacheRecord& Record = NextRequest.Record;
+		uint64 RecordSize = sizeof(FCacheKey) + Record.GetMeta().GetSize();
+		for (const FValueWithId& Value : Record.GetValues())
 		{
-			return FHttpRequest::EResult::Failed;
+			RecordSize += Value.GetData().GetCompressedSize();
+		}
+		BatchSize += RecordSize;
+		if (BatchSize > CacheStore.BatchPutMaxBytes)
+		{
+			BatchSize = RecordSize;
+			return EBatchView::NewBatch;
 		}
+		return EBatchView::Continue;
 	}
 
-	return FHttpRequest::EResult::Success;
-}
-
-FHttpRequest::EResult FZenCacheStore::PerformBlockingRpc(FHttpRequest& Request,
-	FStringView Uri,
-	FCbObject RequestObject,
-	FCbPackage& OutResponse)
-{
-	return ParseRpcResponse(
-		Request.PerformBlockingPost(Uri, RequestObject.GetBuffer(), EHttpContentType::CbObject, EHttpContentType::CbPackage), Request, OutResponse);
-}
-
-FHttpRequest::EResult FZenCacheStore::PerformBlockingRpc(FHttpRequest& Request,
-	FStringView Uri,
-	const FCbPackage& RequestPackage,
-	FCbPackage& OutResponse)
-{
-	FLargeMemoryWriter PackageMemory;
-	UE::Zen::Http::SaveCbPackage(RequestPackage, PackageMemory);
-
-	return ParseRpcResponse(
-		Request.PerformBlockingPost(Uri, FCompositeBuffer(FSharedBuffer::MakeView(PackageMemory.GetView())), EHttpContentType::CbPackage, EHttpContentType::CbPackage), Request, OutResponse);
-}
-
-void FZenCacheStore::EnqueueAsyncRpc(FHttpRequest& Request,
-	UE::DerivedData::IRequestOwner& Owner,
-	FHttpRequestPool* Pool,
-	FStringView Uri,
-	FCbObject RequestObject,
-	FOnRpcComplete&& OnComplete)
-{
-	auto OnHttpRequestComplete = [OnComplete = MoveTemp(OnComplete)]
-	(FHttpRequest::EResult HttpResult, FHttpRequest* Request)
+	void OnHit(const FCachePutRequest& Request)
 	{
-		FCbPackage Response;
-		FHttpRequest::EResult ParsedResponseResult = ParseRpcResponse(HttpResult, *Request, Response);
-		OnComplete(ParsedResponseResult, Request, Response);
-		return FHttpRequest::ECompletionBehavior::Done;
+		const FCacheKey& Key = Request.Record.GetKey();
+		UE_LOG(LogDerivedDataCache, Verbose, TEXT("%s: Cache Put complete for %s from '%s'"),
+			*CacheStore.GetName(), *WriteToString<96>(Key), *Request.Name);
+		COOK_STAT(Timer.AddHit(Private::GetCacheRecordCompressedSize(Request.Record)));
+		OnComplete(Request.MakeResponse(EStatus::Ok));
 	};
-	Request.EnqueueAsyncPost(Owner, Pool, Uri, RequestObject.GetBuffer(), MoveTemp(OnHttpRequestComplete), EHttpContentType::CbObject, EHttpContentType::CbPackage);
-}
 
-void FZenCacheStore::EnqueueAsyncRpc(FHttpRequest& Request,
-	UE::DerivedData::IRequestOwner& Owner,
-	FHttpRequestPool* Pool,
-	FStringView Uri,
-	const FCbPackage& RequestPackage,
-	FOnRpcComplete&& OnComplete)
-{
-	auto OnHttpRequestComplete = [OnComplete = MoveTemp(OnComplete)]
-	(FHttpRequest::EResult HttpResult, FHttpRequest* Request)
+	void OnMiss(const FCachePutRequest& Request)
 	{
-		FCbPackage Response;
-		FHttpRequest::EResult ParsedResponseResult = ParseRpcResponse(HttpResult, *Request, Response);
-		OnComplete(ParsedResponseResult, Request, Response);
-		return FHttpRequest::ECompletionBehavior::Done;
+		const FCacheKey& Key = Request.Record.GetKey();
+		UE_LOG(LogDerivedDataCache, Verbose, TEXT("%s: Cache Put miss for '%s' from '%s'"),
+			*CacheStore.GetName(), *WriteToString<96>(Key), *Request.Name);
+		COOK_STAT(Timer.AddMiss());
+		OnComplete(Request.MakeResponse(EStatus::Error));
 	};
-	FLargeMemoryWriter PackageMemory;
-	UE::Zen::Http::SaveCbPackage(RequestPackage, PackageMemory);
-	uint64 PackageMemorySize = PackageMemory.TotalSize();
-	FSharedBuffer PackageSharedBuffer = FSharedBuffer::TakeOwnership(PackageMemory.ReleaseOwnership(), PackageMemorySize, FMemory::Free);
-	Request.EnqueueAsyncPost(Owner, Pool, Uri, FCompositeBuffer(PackageSharedBuffer), MoveTemp(OnHttpRequestComplete), EHttpContentType::CbPackage, EHttpContentType::CbPackage);
-}
 
-void FZenCacheStore::LegacyStats(FDerivedDataCacheStatsNode& OutNode)
+	FZenCacheStore& CacheStore;
+	IRequestOwner& Owner;
+	const TArray<FCachePutRequest, TInlineAllocator<1>> Requests;
+	uint64 BatchSize = 0;
+	TBatchView<const FCachePutRequest> Batches;
+	FOnCachePutComplete OnComplete;
+	COOK_STAT(FCookStats::FScopedStatsCounter Timer = CacheStore.UsageStats.TimePut());
+};
+
+class FGetOp final : public FThreadSafeRefCountedObject
 {
-	Zen::FZenStats ZenStats;
+public:
+	FGetOp(FZenCacheStore& InCacheStore,
+		IRequestOwner& InOwner,
+		const TConstArrayView<FCacheGetRequest> InRequests,
+		FOnCacheGetComplete&& InOnComplete)
+		: CacheStore(InCacheStore)
+		, Owner(InOwner)
+		, Requests(InRequests)
+		, OnComplete(MoveTemp(InOnComplete))
+	{
+		TRACE_COUNTER_ADD(ZenDDC_Get, (int64)Requests.Num());
+		TRACE_COUNTER_ADD(ZenDDC_CacheRecordRequestCount, int64(Requests.Num()));
+	}
 
-	FDerivedDataCacheUsageStats LocalStats;
-	FDerivedDataCacheUsageStats RemoteStats;
+	virtual ~FGetOp()
+	{
+		TRACE_COUNTER_SUBTRACT(ZenDDC_CacheRecordRequestCount, int64(Requests.Num()));
+	}
 
-#if ENABLE_COOK_STATS
-	using EHitOrMiss = FCookStats::CallStats::EHitOrMiss;
-	using ECacheStatType = FCookStats::CallStats::EStatType;
+	void IssueRequests()
+	{
+		FRequestBarrier Barrier(Owner);
+		ForEachBatch(CacheStore.CacheRecordBatchSize, Requests.Num(),
+			[this](int32 BatchFirst, int32 BatchLast)
+			{
+				TArrayView<const FCacheGetRequest> Batch(Requests.GetData() + BatchFirst, BatchLast - BatchFirst + 1);
 
-	Zen::GetDefaultServiceInstance().GetStats(ZenStats);
+				FCbWriter BatchRequest;
+				BatchRequest.BeginObject();
+				{
+					BatchRequest << ANSITEXTVIEW("Method") << ANSITEXTVIEW("GetCacheRecords");
+					BatchRequest.BeginObject(ANSITEXTVIEW("Params"));
+					{
+						ECachePolicy BatchDefaultPolicy = Batch[0].Policy.GetRecordPolicy();
+						BatchRequest << ANSITEXTVIEW("DefaultPolicy") << *WriteToString<128>(BatchDefaultPolicy);
+						BatchRequest.AddString(ANSITEXTVIEW("Namespace"), CacheStore.StructuredNamespace);
 
-	const int64 RemotePutSize = int64(ZenStats.UpstreamStats.TotalUploadedMB * 1024 * 1024);
-	const int64 RemoteGetSize = int64(ZenStats.UpstreamStats.TotalDownloadedMB * 1024 * 1024);
-	const int64 LocalGetSize = FMath::Max<int64>(0, UsageStats.GetStats.GetAccumulatedValueAnyThread(EHitOrMiss::Hit, ECacheStatType::Bytes) - RemoteGetSize);
+						BatchRequest.BeginArray(ANSITEXTVIEW("Requests"));
+						for (const FCacheGetRequest& Request : Batch)
+						{
+							BatchRequest.BeginObject();
+							{
+								BatchRequest << ANSITEXTVIEW("Key") << Request.Key;
+								if (!Request.Policy.IsUniform() || Request.Policy.GetRecordPolicy() != BatchDefaultPolicy)
+								{
+									BatchRequest.SetName(ANSITEXTVIEW("Policy"));
+									Request.Policy.Save(BatchRequest);
+								}
+							}
+							BatchRequest.EndObject();
+						}
+						BatchRequest.EndArray();
+					}
+					BatchRequest.EndObject();
+				}
+				BatchRequest.EndObject();
 
-	LocalStats.PutStats = UsageStats.PutStats;
-	LocalStats.ExistsStats = UsageStats.ExistsStats;
-	LocalStats.PrefetchStats = UsageStats.PrefetchStats;
+				FGetOp* OriginalOp = this;
+				auto OnRpcComplete = [this, OpRef = TRefCountPtr<FGetOp>(OriginalOp), Batch](FHttpRequest::EResult HttpResult, FHttpRequest* HttpRequest, FCbPackage& Response)
+				{
+					int32 RequestIndex = 0;
+					if (HttpResult == FHttpRequest::EResult::Success)
+					{
+						const FCbObject& ResponseObj = Response.GetObject();
+						
+						for (FCbField RecordField : ResponseObj[ANSITEXTVIEW("Result")])
+						{
+							if (RequestIndex >= Batch.Num())
+							{
+								++RequestIndex;
+								continue;
+							}
+							const FCacheGetRequest& Request = Batch[RequestIndex++];
 
-	LocalStats.GetStats.Accumulate(EHitOrMiss::Hit, ECacheStatType::Counter, ZenStats.CacheStats.Hits - ZenStats.CacheStats.UpstreamHits, /*bIsInGameThread*/ false);
-	LocalStats.GetStats.Accumulate(EHitOrMiss::Miss, ECacheStatType::Counter, ZenStats.CacheStats.Misses + ZenStats.CacheStats.UpstreamHits, /*bIsInGameThread*/ false);
-	RemoteStats.GetStats.Accumulate(EHitOrMiss::Hit, ECacheStatType::Counter, ZenStats.CacheStats.UpstreamHits, /*bIsInGameThread*/ false);
-	RemoteStats.GetStats.Accumulate(EHitOrMiss::Miss, ECacheStatType::Counter, ZenStats.CacheStats.Misses, /*bIsInGameThread*/ false);
+							const FCacheKey& Key = Request.Key;
+							FOptionalCacheRecord Record;
 
-	LocalStats.GetStats.Accumulate(EHitOrMiss::Hit, ECacheStatType::Bytes, LocalGetSize, /*bIsInGameThread*/ false);
-	RemoteStats.GetStats.Accumulate(EHitOrMiss::Hit, ECacheStatType::Bytes, RemoteGetSize, /*bIsInGameThread*/ false);
-	RemoteStats.PutStats.Accumulate(EHitOrMiss::Hit, ECacheStatType::Bytes, RemotePutSize, /*bIsInGameThread*/ false);
-#endif
+							if (CacheStore.DebugOptions.ShouldSimulateGetMiss(Key))
+							{
+								UE_LOG(LogDerivedDataCache, Verbose, TEXT("%s: Simulated miss for put of '%s' from '%s'"),
+									*CacheStore.GetName(), *WriteToString<96>(Key), *Request.Name);
+							}
+							else if (!RecordField.IsNull())
+							{
+								Record = FCacheRecord::Load(Response, RecordField.AsObject());
+							}
+							Record ? OnHit(Request, MoveTemp(Record).Get()) : OnMiss(Request);
+						}
+						if (RequestIndex != Batch.Num())
+						{
+							UE_LOG(LogDerivedDataCache, Warning,
+								TEXT("%s: Invalid response received from GetCacheRecords rpc: %d results expected, received %d."), *CacheStore.GetName(), Batch.Num(), RequestIndex);
+						}
+					}
+					for (const FCacheGetRequest& Request : Batch.RightChop(RequestIndex))
+					{
+						OnMiss(Request);
+					}
+				};
 
-	if (ZenStats.UpstreamStats.EndPointStats.IsEmpty())
-	{
-		OutNode = {TEXT("Zen"), ZenService.GetInstance().GetURL(), /*bIsLocal*/ true};
-		OutNode.UsageStats.Add(TEXT(""), LocalStats);
-		return;
+				FHttpRequest* Request = CacheStore.RequestPool->WaitForFreeRequest();
+				CacheStore.EnqueueAsyncRpc(*Request, Owner, CacheStore.RequestPool.Get(), TEXTVIEW("/z$/$rpc"), BatchRequest.Save().AsObject(), MoveTemp(OnRpcComplete));
+			});
 	}
+private:
+	void OnHit(const FCacheGetRequest& Request, FCacheRecord&& Record)
+	{
+		UE_LOG(LogDerivedDataCache, Verbose, TEXT("%s: Cache hit for '%s' from '%s'"),
+			*CacheStore.GetName(), *WriteToString<96>(Request.Key), *Request.Name);
 
-	TSharedRef<FDerivedDataCacheStatsNode> LocalNode =
-		MakeShared<FDerivedDataCacheStatsNode>(TEXT("Zen"), ZenService.GetInstance().GetURL(), /*bIsLocal*/ true);
-	LocalNode->UsageStats.Add(TEXT(""), LocalStats);
-
-	TSharedRef<FDerivedDataCacheStatsNode> RemoteNode =
-		MakeShared<FDerivedDataCacheStatsNode>(ZenStats.UpstreamStats.EndPointStats[0].Name, ZenStats.UpstreamStats.EndPointStats[0].Url, /*bIsLocal*/ false);
-	RemoteNode->UsageStats.Add(TEXT(""), RemoteStats);
+		TRACE_COUNTER_ADD(ZenDDC_GetHit, int64(1));
+		int64 ReceivedSize = Private::GetCacheRecordCompressedSize(Record);
+		TRACE_COUNTER_ADD(ZenDDC_BytesReceived, ReceivedSize);
+		COOK_STAT(Timer.AddHit(ReceivedSize));
 
-	OutNode = {TEXT("Zen Group"), TEXT(""), /*bIsLocal*/ true};
-	OutNode.Children.Add(LocalNode);
-	OutNode.Children.Add(RemoteNode);
-}
+		OnComplete({Request.Name, MoveTemp(Record), Request.UserData, EStatus::Ok});
+	};
 
-bool FZenCacheStore::LegacyDebugOptions(FBackendDebugOptions& InOptions)
-{
-	DebugOptions = InOptions;
-	return true;
-}
+	void OnMiss(const FCacheGetRequest& Request)
+	{
+		UE_LOG(LogDerivedDataCache, Verbose, TEXT("%s: Cache miss for '%s' from '%s'"),
+			*CacheStore.GetName(), *WriteToString<96>(Request.Key), *Request.Name);
 
-void FZenCacheStore::Put(
-	const TConstArrayView<FCachePutRequest> Requests,
-	IRequestOwner& Owner,
-	FOnCachePutComplete&& OnComplete)
-{
-	TRACE_CPUPROFILER_EVENT_SCOPE(ZenDDC::PutCachedRecord);
-	COOK_STAT(auto Timer = UsageStats.TimePut());
+		OnComplete(Request.MakeResponse(EStatus::Error));
+	};
 
-	uint64 BatchSize = 0;
-	TBatchView<const FCachePutRequest> Batches(Requests,
-		[&BatchSize, this](const FCachePutRequest& NextRequest) -> EBatchView
-		{
-			const FCacheRecord& Record = NextRequest.Record;
-			uint64 RecordSize = sizeof(FCacheKey) + Record.GetMeta().GetSize();
-			for (const FValueWithId& Value : Record.GetValues())
-			{
-				RecordSize += Value.GetData().GetCompressedSize();
-			}
-			BatchSize += RecordSize;
-			if (BatchSize > BatchPutMaxBytes)
-			{
-				BatchSize = RecordSize;
-				return EBatchView::NewBatch;
-			}
-			return EBatchView::Continue;
-		});
+	FZenCacheStore& CacheStore;
+	IRequestOwner& Owner;
+	const TArray<FCacheGetRequest, TInlineAllocator<1>> Requests;
+	FOnCacheGetComplete OnComplete;
+	COOK_STAT(FCookStats::FScopedStatsCounter Timer = CacheStore.UsageStats.TimeGet());
+};
 
-	auto OnHit = [this, &OnComplete COOK_STAT(, &Timer)](const FCachePutRequest& Request)
-	{
-		const FCacheKey& Key = Request.Record.GetKey();
-		UE_LOG(LogDerivedDataCache, Verbose, TEXT("%s: Cache Put complete for %s from '%s'"),
-			*GetName(), *WriteToString<96>(Key), *Request.Name);
-		COOK_STAT(Timer.AddHit(Private::GetCacheRecordCompressedSize(Request.Record)));
-		OnComplete(Request.MakeResponse(EStatus::Ok));
-	};
-	auto OnMiss = [this, &OnComplete COOK_STAT(, &Timer)](const FCachePutRequest& Request)
+class FPutValueOp final : public FThreadSafeRefCountedObject
+{
+public:
+	FPutValueOp(FZenCacheStore& InCacheStore,
+		IRequestOwner& InOwner,
+		const TConstArrayView<FCachePutValueRequest> InRequests,
+		FOnCachePutValueComplete&& InOnComplete)
+		: CacheStore(InCacheStore)
+		, Owner(InOwner)
+		, Requests(InRequests)
+		, Batches(Requests, [this](const FCachePutValueRequest& NextRequest) {return BatchGroupingFilter(NextRequest);})
+		, OnComplete(MoveTemp(InOnComplete))
 	{
-		const FCacheKey& Key = Request.Record.GetKey();
-		UE_LOG(LogDerivedDataCache, Verbose, TEXT("%s: Cache Put miss for '%s' from '%s'"),
-			*GetName(), *WriteToString<96>(Key), *Request.Name);
-		COOK_STAT(Timer.AddMiss());
-		OnComplete(Request.MakeResponse(EStatus::Error));
-	};
+	}
 
-	for (TArrayView<const FCachePutRequest> Batch : Batches)
+	void IssueRequests()
 	{
-		FCbPackage BatchPackage;
-		FCbWriter BatchWriter;
-		BatchWriter.BeginObject();
+		FRequestBarrier Barrier(Owner);
+		for (TArrayView<const FCachePutValueRequest> Batch : Batches)
 		{
-			BatchWriter << ANSITEXTVIEW("Method") << "PutCacheRecords";
-			BatchWriter.BeginObject(ANSITEXTVIEW("Params"));
+			FCbPackage BatchPackage;
+			FCbWriter BatchWriter;
+			BatchWriter.BeginObject();
 			{
-				ECachePolicy BatchDefaultPolicy = Batch[0].Policy.GetRecordPolicy();
-				BatchWriter << ANSITEXTVIEW("DefaultPolicy") << *WriteToString<128>(BatchDefaultPolicy);
-				BatchWriter.AddString(ANSITEXTVIEW("Namespace"), StructuredNamespace);
-
-				BatchWriter.BeginArray(ANSITEXTVIEW("Requests"));
-				for (const FCachePutRequest& Request : Batch)
+				BatchWriter << ANSITEXTVIEW("Method") << ANSITEXTVIEW("PutCacheValues");
+				BatchWriter.BeginObject(ANSITEXTVIEW("Params"));
 				{
-					const FCacheRecord& Record = Request.Record;
+					ECachePolicy BatchDefaultPolicy = Batch[0].Policy;
+					BatchWriter << ANSITEXTVIEW("DefaultPolicy") << *WriteToString<128>(BatchDefaultPolicy);
+					BatchWriter.AddString(ANSITEXTVIEW("Namespace"), CacheStore.StructuredNamespace);
 
-					BatchWriter.BeginObject();
+					BatchWriter.BeginArray("Requests");
+					for (const FCachePutValueRequest& Request : Batch)
 					{
-						BatchWriter.SetName(ANSITEXTVIEW("Record"));
-						Record.Save(BatchPackage, BatchWriter);
-						if (!Request.Policy.IsUniform() || Request.Policy.GetRecordPolicy() != BatchDefaultPolicy)
+						BatchWriter.BeginObject();
 						{
-							BatchWriter.SetName(ANSITEXTVIEW("Policy"));
-							Request.Policy.Save(BatchWriter);
-						}
+							BatchWriter << ANSITEXTVIEW("Key") << Request.Key;
+							const FValue& Value = Request.Value;
+							BatchWriter.AddBinaryAttachment("RawHash", Value.GetRawHash());
+							if (Value.HasData())
+							{
+								BatchPackage.AddAttachment(FCbAttachment(Value.GetData()));
+							}
+							if (Request.Policy != BatchDefaultPolicy)
+							{
+								BatchWriter << ANSITEXTVIEW("Policy") << WriteToString<128>(Request.Policy);
+							}
+						}
+						BatchWriter.EndObject();
 					}
-					BatchWriter.EndObject();
+					BatchWriter.EndArray();
 				}
-				BatchWriter.EndArray();
+				BatchWriter.EndObject();
 			}
 			BatchWriter.EndObject();
-		}
-		BatchWriter.EndObject();
-		BatchPackage.SetObject(BatchWriter.Save().AsObject());
-
-		FCbPackage BatchResponse;
-		FHttpRequest::EResult HttpResult = FHttpRequest::EResult::Failed;
+			BatchPackage.SetObject(BatchWriter.Save().AsObject());
 
-		{
-			FScopedHttpPoolRequestPtr Request(RequestPool.Get());
-			HttpResult = PerformBlockingRpc(*Request.Get(), TEXTVIEW("/z$/$rpc"), BatchPackage, BatchResponse);
-		}
-
-		int32 RequestIndex = 0;
-		if (HttpResult == FHttpRequest::EResult::Success)
-		{
-			const FCbObject& ResponseObj = BatchResponse.GetObject();
-			for (FCbField ResponseField : ResponseObj[ANSITEXTVIEW("Result")])
+			auto OnRpcComplete = [this, OpRef = TRefCountPtr<FPutValueOp>(this), Batch](FHttpRequest::EResult HttpResult, FHttpRequest* HttpRequest, FCbPackage& Response)
 			{
-				if (RequestIndex >= Batch.Num())
+				int32 RequestIndex = 0;
+				if (HttpResult == FHttpRequest::EResult::Success)
 				{
-					++RequestIndex;
-					continue;
-				}
-				const FCachePutRequest& Request = Batch[RequestIndex++];
+					const FCbObject& ResponseObj = Response.GetObject();
+					for (FCbField ResponseField : ResponseObj[ANSITEXTVIEW("Result")])
+					{
+						if (RequestIndex >= Batch.Num())
+						{
+							++RequestIndex;
+							continue;
+						}
+						const FCachePutValueRequest& Request = Batch[RequestIndex++];
 
-				const FCacheKey& Key = Request.Record.GetKey();
-				bool bPutSucceeded = ResponseField.AsBool();
-				if (DebugOptions.ShouldSimulatePutMiss(Key))
+						bool bPutSucceeded = ResponseField.AsBool();
+						if (CacheStore.DebugOptions.ShouldSimulatePutMiss(Request.Key))
+						{
+							UE_LOG(LogDerivedDataCache, Verbose, TEXT("%s: Simulated miss for PutValue of %s from '%s'"),
+								*CacheStore.GetName(), *WriteToString<96>(Request.Key), *Request.Name);
+							bPutSucceeded = false;
+						}
+						bPutSucceeded ? OnHit(Request) : OnMiss(Request);
+					}
+					if (RequestIndex != Batch.Num())
+					{
+						UE_LOG(LogDerivedDataCache, Warning,
+							TEXT("%s: Invalid response received from PutCacheValues rpc: %d results expected, received %d."), *CacheStore.GetName(), Batch.Num(), RequestIndex);
+					}
+				}
+				for (const FCachePutValueRequest& Request : Batch.RightChop(RequestIndex))
 				{
-					UE_LOG(LogDerivedDataCache, Verbose, TEXT("%s: Simulated miss for put of %s from '%s'"),
-						*GetName(), *WriteToString<96>(Key), *Request.Name);
-					bPutSucceeded = false;
+					OnMiss(Request);
 				}
-				bPutSucceeded ? OnHit(Request) : OnMiss(Request);
-			}
-			if (RequestIndex != Batch.Num())
-			{
-				UE_LOG(LogDerivedDataCache, Warning,
-					TEXT("Invalid response received from PutCacheRecords rpc: %d results expected, received %d."), Batch.Num(), RequestIndex);
-			}
+			};
+			FHttpRequest* Request = CacheStore.RequestPool->WaitForFreeRequest();
+			CacheStore.EnqueueAsyncRpc(*Request, Owner, CacheStore.RequestPool.Get(), TEXTVIEW("/z$/$rpc"), BatchPackage, MoveTemp(OnRpcComplete));
 		}
-		for (const FCachePutRequest& Request : Batch.RightChop(RequestIndex))
+	}
+
+private:
+	EBatchView BatchGroupingFilter(const FCachePutValueRequest& NextRequest)
+	{
+		uint64 ValueSize = sizeof(FCacheKey) + NextRequest.Value.GetData().GetCompressedSize();
+		BatchSize += ValueSize;
+		if (BatchSize > CacheStore.BatchPutMaxBytes)
 		{
-			OnMiss(Request);
+			BatchSize = ValueSize;
+			return EBatchView::NewBatch;
 		}
+		return EBatchView::Continue;
 	}
-}
-
-void FZenCacheStore::Get(
-	const TConstArrayView<FCacheGetRequest> Requests,
-	IRequestOwner& Owner,
-	FOnCacheGetComplete&& OnComplete)
-{
-	TRACE_CPUPROFILER_EVENT_SCOPE(ZenDDC::GetCacheRecord);
-	TRACE_COUNTER_ADD(ZenDDC_CacheRecordRequestCount, int64(Requests.Num()));
-	TRACE_COUNTER_ADD(ZenDDC_Get, int64(Requests.Num()));
-	COOK_STAT(auto Timer = UsageStats.TimeGet());
 
-	auto OnHit = [this, &OnComplete COOK_STAT(, &Timer)](const FCacheGetRequest& Request, FCacheRecord&& Record)
+	void OnHit(const FCachePutValueRequest& Request)
 	{
-		UE_LOG(LogDerivedDataCache, Verbose, TEXT("%s: Cache hit for '%s' from '%s'"),
-			*GetName(), *WriteToString<96>(Request.Key), *Request.Name);
-
-		TRACE_COUNTER_ADD(ZenDDC_GetHit, int64(1));
-		int64 ReceivedSize = Private::GetCacheRecordCompressedSize(Record);
-		TRACE_COUNTER_ADD(ZenDDC_BytesReceived, ReceivedSize);
-		COOK_STAT(Timer.AddHit(ReceivedSize));
-
-		OnComplete({Request.Name, MoveTemp(Record), Request.UserData, EStatus::Ok});
+		UE_LOG(LogDerivedDataCache, Verbose, TEXT("%s: Cache PutValue complete for %s from '%s'"),
+			*CacheStore.GetName(), *WriteToString<96>(Request.Key), *Request.Name);
+		COOK_STAT(Timer.AddHit(Request.Value.GetData().GetCompressedSize()));
+		OnComplete(Request.MakeResponse(EStatus::Ok));
 	};
-	auto OnMiss = [this, &OnComplete](const FCacheGetRequest& Request)
-	{
-		UE_LOG(LogDerivedDataCache, Verbose, TEXT("%s: Cache miss for '%s' from '%s'"),
-			*GetName(), *WriteToString<96>(Request.Key), *Request.Name);
 
+	void OnMiss(const FCachePutValueRequest& Request)
+	{
+		UE_LOG(LogDerivedDataCache, Verbose, TEXT("%s: Cache PutValue miss for '%s' from '%s'"),
+			*CacheStore.GetName(), *WriteToString<96>(Request.Key), *Request.Name);
+		COOK_STAT(Timer.AddMiss());
 		OnComplete(Request.MakeResponse(EStatus::Error));
 	};
 
-	ForEachBatch(CacheRecordBatchSize, Requests.Num(),
-		[this, &Requests, &Owner, &OnComplete, &OnHit, &OnMiss](int32 BatchFirst, int32 BatchLast)
+	FZenCacheStore& CacheStore;
+	IRequestOwner& Owner;
+	const TArray<FCachePutValueRequest, TInlineAllocator<1>> Requests;
+	uint64 BatchSize = 0;
+	TBatchView<const FCachePutValueRequest> Batches;
+	FOnCachePutValueComplete OnComplete;
+	COOK_STAT(FCookStats::FScopedStatsCounter Timer = CacheStore.UsageStats.TimePut());
+};
+
+class FGetValueOp final : public FThreadSafeRefCountedObject
+{
+public:
+	FGetValueOp(FZenCacheStore& InCacheStore,
+		IRequestOwner& InOwner,
+		const TConstArrayView<FCacheGetValueRequest> InRequests,
+		FOnCacheGetValueComplete&& InOnComplete)
+		: CacheStore(InCacheStore)
+		, Owner(InOwner)
+		, Requests(InRequests)
+		, OnComplete(MoveTemp(InOnComplete))
 	{
-		TArrayView<const FCacheGetRequest> Batch(Requests.GetData() + BatchFirst, BatchLast - BatchFirst + 1);
+		TRACE_COUNTER_ADD(ZenDDC_Get, (int64)Requests.Num());
+		TRACE_COUNTER_ADD(ZenDDC_CacheRecordRequestCount, int64(Requests.Num()));
+	}
+
+	virtual ~FGetValueOp()
+	{
+		TRACE_COUNTER_SUBTRACT(ZenDDC_CacheRecordRequestCount, int64(Requests.Num()));
+	}
 
-		FCbWriter BatchRequest;
-		BatchRequest.BeginObject();
+	void IssueRequests()
+	{
+		FRequestBarrier Barrier(Owner);
+		ForEachBatch(CacheStore.CacheRecordBatchSize, Requests.Num(),
+			[this](int32 BatchFirst, int32 BatchLast)
 		{
-			BatchRequest << ANSITEXTVIEW("Method") << ANSITEXTVIEW("GetCacheRecords");
-			BatchRequest.BeginObject(ANSITEXTVIEW("Params"));
-			{
-				ECachePolicy BatchDefaultPolicy = Batch[0].Policy.GetRecordPolicy();
-				BatchRequest << ANSITEXTVIEW("DefaultPolicy") << *WriteToString<128>(BatchDefaultPolicy);
-				BatchRequest.AddString(ANSITEXTVIEW("Namespace"), StructuredNamespace);
+			TArrayView<const FCacheGetValueRequest> Batch(Requests.GetData() + BatchFirst, BatchLast - BatchFirst + 1);
 
-				BatchRequest.BeginArray(ANSITEXTVIEW("Requests"));
-				for (const FCacheGetRequest& Request : Batch)
+			FCbWriter BatchRequest;
+			BatchRequest.BeginObject();
+			{
+				BatchRequest << ANSITEXTVIEW("Method") << ANSITEXTVIEW("GetCacheValues");
+				BatchRequest.BeginObject(ANSITEXTVIEW("Params"));
 				{
-					BatchRequest.BeginObject();
+					ECachePolicy BatchDefaultPolicy = Batch[0].Policy;
+					BatchRequest << ANSITEXTVIEW("DefaultPolicy") << *WriteToString<128>(BatchDefaultPolicy);
+					BatchRequest.AddString(ANSITEXTVIEW("Namespace"), CacheStore.StructuredNamespace);
+
+					BatchRequest.BeginArray("Requests");
+					for (const FCacheGetValueRequest& Request : Batch)
 					{
-						BatchRequest << ANSITEXTVIEW("Key") << Request.Key;
-						if (!Request.Policy.IsUniform() || Request.Policy.GetRecordPolicy() != BatchDefaultPolicy)
+						BatchRequest.BeginObject();
 						{
-							BatchRequest.SetName(ANSITEXTVIEW("Policy"));
-							Request.Policy.Save(BatchRequest);
+							BatchRequest << ANSITEXTVIEW("Key") << Request.Key;
+							if (Request.Policy != BatchDefaultPolicy)
+							{
+								BatchRequest << ANSITEXTVIEW("Policy") << WriteToString<128>(Request.Policy);
+							}
 						}
+						BatchRequest.EndObject();
 					}
-					BatchRequest.EndObject();
+					BatchRequest.EndArray();
 				}
-				BatchRequest.EndArray();
+				BatchRequest.EndObject();
 			}
 			BatchRequest.EndObject();
-		}
-		BatchRequest.EndObject();
-
-		FCbPackage BatchResponse;
-		FHttpRequest::EResult HttpResult = FHttpRequest::EResult::Failed;
 
-		{
-			LLM_SCOPE_BYTAG(UntaggedDDCResult);
-			FScopedHttpPoolRequestPtr Request(RequestPool.Get());
-			HttpResult = PerformBlockingRpc(*Request.Get(), TEXTVIEW("/z$/$rpc"), BatchRequest.Save().AsObject(), BatchResponse);
-		}
-
-		int32 RequestIndex = 0;
-		if (HttpResult == FHttpRequest::EResult::Success)
-		{
-			const FCbObject& ResponseObj = BatchResponse.GetObject();
-			
-			for (FCbField RecordField : ResponseObj[ANSITEXTVIEW("Result")])
+			FGetValueOp* OriginalOp = this;
+			auto OnRpcComplete = [this, OpRef = TRefCountPtr<FGetValueOp>(OriginalOp), Batch](FHttpRequest::EResult HttpResult, FHttpRequest* HttpRequest, FCbPackage& Response)
 			{
-				if (RequestIndex >= Batch.Num())
-				{
-					++RequestIndex;
-					continue;
-				}
-				const FCacheGetRequest& Request = Batch[RequestIndex++];
-
-				const FCacheKey& Key = Request.Key;
-				FOptionalCacheRecord Record;
-
-				if (DebugOptions.ShouldSimulateGetMiss(Key))
-				{
-					UE_LOG(LogDerivedDataCache, Verbose, TEXT("%s: Simulated miss for put of '%s' from '%s'"),
-						*GetName(), *WriteToString<96>(Key), *Request.Name);
-				}
-				else if (!RecordField.IsNull())
+				int32 RequestIndex = 0;
+				if (HttpResult == FHttpRequest::EResult::Success)
 				{
-					Record = FCacheRecord::Load(BatchResponse, RecordField.AsObject());
-				}
-				Record ? OnHit(Request, MoveTemp(Record).Get()) : OnMiss(Request);
-			}
-			if (RequestIndex != Batch.Num())
-			{
-				UE_LOG(LogDerivedDataCache, Warning,
-					TEXT("Invalid response received from GetCacheRecords rpc: %d results expected, received %d."), Batch.Num(), RequestIndex);
-			}
-		}
-		for (const FCacheGetRequest& Request : Batch.RightChop(RequestIndex))
-		{
-			OnMiss(Request);
-		}
-	});
-	
-	TRACE_COUNTER_SUBTRACT(ZenDDC_CacheRecordRequestCount, int64(Requests.Num()));
-}
-
-void FZenCacheStore::PutValue(
-	TConstArrayView<FCachePutValueRequest> Requests,
-	IRequestOwner& Owner,
-	FOnCachePutValueComplete&& OnComplete)
-{
-	TRACE_CPUPROFILER_EVENT_SCOPE(ZenDDC::PutValue);
-	COOK_STAT(auto Timer = UsageStats.TimePut());
+					const FCbObject& ResponseObj = Response.GetObject();
 
-	uint64 BatchSize = 0;
-	TBatchView<const FCachePutValueRequest> Batches(Requests,
-		[&BatchSize, this](const FCachePutValueRequest& NextValue) -> EBatchView
-		{
-			uint64 ValueSize = sizeof(FCacheKey) + NextValue.Value.GetData().GetCompressedSize();
-			BatchSize += ValueSize;
-			if (BatchSize > BatchPutMaxBytes)
-			{
-				BatchSize = ValueSize;
-				return EBatchView::NewBatch;
-			}
-			return EBatchView::Continue;
-		});
-
-	auto OnHit = [this, &OnComplete COOK_STAT(, &Timer)](const FCachePutValueRequest& Request)
-	{
-		UE_LOG(LogDerivedDataCache, Verbose, TEXT("%s: Cache PutValue complete for %s from '%s'"),
-			*GetName(), *WriteToString<96>(Request.Key), *Request.Name);
-		COOK_STAT(Timer.AddHit(Request.Value.GetData().GetCompressedSize()));
-		OnComplete(Request.MakeResponse(EStatus::Ok));
-	};
-	auto OnMiss = [this, &OnComplete COOK_STAT(, &Timer)](const FCachePutValueRequest& Request)
-	{
-		UE_LOG(LogDerivedDataCache, Verbose, TEXT("%s: Cache PutValue miss for '%s' from '%s'"),
-			*GetName(), *WriteToString<96>(Request.Key), *Request.Name);
-		COOK_STAT(Timer.AddMiss());
-		OnComplete(Request.MakeResponse(EStatus::Error));
-	};
-
-	for (TArrayView<const FCachePutValueRequest> Batch : Batches)
-	{
-		FCbPackage BatchPackage;
-		FCbWriter BatchWriter;
-		BatchWriter.BeginObject();
-		{
-			BatchWriter << ANSITEXTVIEW("Method") << ANSITEXTVIEW("PutCacheValues");
-			BatchWriter.BeginObject(ANSITEXTVIEW("Params"));
-			{
-				ECachePolicy BatchDefaultPolicy = Batch[0].Policy;
-				BatchWriter << ANSITEXTVIEW("DefaultPolicy") << *WriteToString<128>(BatchDefaultPolicy);
-				BatchWriter.AddString(ANSITEXTVIEW("Namespace"), StructuredNamespace);
-
-				BatchWriter.BeginArray("Requests");
-				for (const FCachePutValueRequest& Request : Batch)
-				{
-					BatchWriter.BeginObject();
+					for (FCbFieldView ResultField : ResponseObj[ANSITEXTVIEW("Result")])
 					{
-						BatchWriter << ANSITEXTVIEW("Key") << Request.Key;
-						const FValue& Value = Request.Value;
-						BatchWriter.AddBinaryAttachment("RawHash", Value.GetRawHash());
-						if (Value.HasData())
+						if (RequestIndex >= Batch.Num())
+						{
+							++RequestIndex;
+							continue;
+						}
+						const FCacheGetValueRequest& Request = Batch[RequestIndex++];
+
+						FCbObjectView ResultObj = ResultField.AsObjectView();
+						TOptional<FValue> Value;
+						if (CacheStore.DebugOptions.ShouldSimulateGetMiss(Request.Key))
 						{
-							BatchPackage.AddAttachment(FCbAttachment(Value.GetData()));
+							UE_LOG(LogDerivedDataCache, Verbose, TEXT("%s: Simulated miss for GetValue of '%s' from '%s'"),
+								*CacheStore.GetName(), *WriteToString<96>(Request.Key), *Request.Name);
 						}
-						if (Request.Policy != BatchDefaultPolicy)
+						else
 						{
-							BatchWriter << ANSITEXTVIEW("Policy") << WriteToString<128>(Request.Policy);
+							FCbFieldView RawHashField = ResultObj["RawHash"];
+							FIoHash RawHash = RawHashField.AsHash();
+							if (const FCbAttachment* Attachment = Response.FindAttachment(RawHash))
+							{
+								Value.Emplace(Attachment->AsCompressedBinary());
+							}
+							else
+							{
+								FCbFieldView RawSizeField = ResultObj["RawSize"];
+								uint64 RawSize = RawSizeField.AsUInt64();
+								if (!RawSizeField.HasError() && !RawHashField.HasError())
+								{
+									Value.Emplace(RawHash, RawSize);
+								}
+							}
 						}
+						(bool)Value ? OnHit(Request, MoveTemp(*Value)) : OnMiss(Request);
+					}
+					if (RequestIndex != Batch.Num())
+					{
+						UE_LOG(LogDerivedDataCache, Warning,
+							TEXT("Invalid response received from GetCacheValues rpc: %d results expected, received %d."), Batch.Num(), RequestIndex);
 					}
-					BatchWriter.EndObject();
-				}
-				BatchWriter.EndArray();
-			}
-			BatchWriter.EndObject();
-		}
-		BatchWriter.EndObject();
-		BatchPackage.SetObject(BatchWriter.Save().AsObject());
-
-		FCbPackage BatchResponse;
-		FHttpRequest::EResult HttpResult = FHttpRequest::EResult::Failed;
-		{
-			FScopedHttpPoolRequestPtr Request(RequestPool.Get());
-			HttpResult = PerformBlockingRpc(*Request.Get(), TEXTVIEW("/z$/$rpc"), BatchPackage, BatchResponse);
-		}
-
-		int32 RequestIndex = 0;
-		if (HttpResult == FHttpRequest::EResult::Success)
-		{
-			const FCbObject& ResponseObj = BatchResponse.GetObject();
-			for (FCbField ResponseField : ResponseObj[ANSITEXTVIEW("Result")])
-			{
-				if (RequestIndex >= Batch.Num())
-				{
-					++RequestIndex;
-					continue;
 				}
-				const FCachePutValueRequest& Request = Batch[RequestIndex++];
-
-				bool bPutSucceeded = ResponseField.AsBool();
-				if (DebugOptions.ShouldSimulatePutMiss(Request.Key))
+				for (const FCacheGetValueRequest& Request : Batch.RightChop(RequestIndex))
 				{
-					UE_LOG(LogDerivedDataCache, Verbose, TEXT("%s: Simulated miss for PutValue of %s from '%s'"),
-						*GetName(), *WriteToString<96>(Request.Key), *Request.Name);
-					bPutSucceeded = false;
+					OnMiss(Request);
 				}
-				bPutSucceeded ? OnHit(Request) : OnMiss(Request);
-			}
-			if (RequestIndex != Batch.Num())
-			{
-				UE_LOG(LogDerivedDataCache, Warning,
-					TEXT("Invalid response received from PutCacheValues rpc: %d results expected, received %d."), Batch.Num(), RequestIndex);
-			}
-		}
-		for (const FCachePutValueRequest& Request : Batch.RightChop(RequestIndex))
-		{
-			OnMiss(Request);
-		}
+			};
+			FHttpRequest* Request = CacheStore.RequestPool->WaitForFreeRequest();
+			CacheStore.EnqueueAsyncRpc(*Request, Owner, CacheStore.RequestPool.Get(), TEXTVIEW("/z$/$rpc"), BatchRequest.Save().AsObject(), MoveTemp(OnRpcComplete));
+		});
 	}
-}
-
-void FZenCacheStore::GetValue(
-	TConstArrayView<FCacheGetValueRequest> Requests,
-	IRequestOwner& Owner,
-	FOnCacheGetValueComplete&& OnComplete)
-				{
-	TRACE_CPUPROFILER_EVENT_SCOPE(ZenDDC::GetValue);
-	TRACE_COUNTER_ADD(ZenDDC_CacheRecordRequestCount, int64(Requests.Num()));
-	COOK_STAT(auto Timer = UsageStats.TimeGet());
-	TRACE_COUNTER_ADD(ZenDDC_Get, (int64)Requests.Num());
-					
-	auto OnHit = [this, &OnComplete COOK_STAT(, &Timer)](const FCacheGetValueRequest& Request, FValue&& Value)
+private:
+	void OnHit(const FCacheGetValueRequest& Request, FValue&& Value)
 	{
 		UE_LOG(LogDerivedDataCache, Verbose, TEXT("%s: Cache hit for '%s' from '%s'"),
-			*GetName(), *WriteToString<96>(Request.Key), *Request.Name);
+			*CacheStore.GetName(), *WriteToString<96>(Request.Key), *Request.Name);
 					TRACE_COUNTER_ADD(ZenDDC_GetHit, int64(1));
 		int64 ReceivedSize = Value.GetData().GetCompressedSize();
 					TRACE_COUNTER_ADD(ZenDDC_BytesReceived, ReceivedSize);
@@ -708,36 +716,87 @@ void FZenCacheStore::GetValue(
 
 		OnComplete({Request.Name, Request.Key, MoveTemp(Value), Request.UserData, EStatus::Ok});
 	};
-	auto OnMiss = [this, &OnComplete](const FCacheGetValueRequest& Request)
+
+	void OnMiss(const FCacheGetValueRequest& Request)
 	{
 		UE_LOG(LogDerivedDataCache, Verbose, TEXT("%s: Cache miss for '%s' from '%s'"),
-			*GetName(), *WriteToString<96>(Request.Key), *Request.Name);
+			*CacheStore.GetName(), *WriteToString<96>(Request.Key), *Request.Name);
 
 		OnComplete(Request.MakeResponse(EStatus::Error));
 	};
 
-	ForEachBatch(CacheRecordBatchSize, Requests.Num(),
-		[this, &Requests, &Owner, &OnComplete, &OnHit, &OnMiss](int32 BatchFirst, int32 BatchLast)
+	FZenCacheStore& CacheStore;
+	IRequestOwner& Owner;
+	const TArray<FCacheGetValueRequest, TInlineAllocator<1>> Requests;
+	FOnCacheGetValueComplete OnComplete;
+	COOK_STAT(FCookStats::FScopedStatsCounter Timer = CacheStore.UsageStats.TimeGet());
+};
+
+class FGetChunksOp final : public FThreadSafeRefCountedObject
+{
+public:
+	FGetChunksOp(FZenCacheStore& InCacheStore,
+		IRequestOwner& InOwner,
+		const TConstArrayView<FCacheGetChunkRequest> InRequests,
+		FOnCacheGetChunkComplete&& InOnComplete)
+		: CacheStore(InCacheStore)
+		, Owner(InOwner)
+		, Requests(InRequests)
+		, OnComplete(MoveTemp(InOnComplete))
+	{
+		TRACE_COUNTER_ADD(ZenDDC_ChunkRequestCount, int64(Requests.Num()));
+		TRACE_COUNTER_ADD(ZenDDC_Get, int64(Requests.Num()));
+		Requests.StableSort(TChunkLess());
+	}
+
+	virtual ~FGetChunksOp()
+	{
+		TRACE_COUNTER_SUBTRACT(ZenDDC_ChunkRequestCount, int64(Requests.Num()));
+	}
+
+	void IssueRequests()
+	{
+		FRequestBarrier Barrier(Owner);
+		ForEachBatch(CacheStore.CacheChunksBatchSize, Requests.Num(),
+			[this](int32 BatchFirst, int32 BatchLast)
 		{
-			TArrayView<const FCacheGetValueRequest> Batch(Requests.GetData() + BatchFirst, BatchLast - BatchFirst + 1);
+			TArrayView<const FCacheGetChunkRequest> Batch(Requests.GetData() + BatchFirst, BatchLast - BatchFirst + 1);
 
 			FCbWriter BatchRequest;
 			BatchRequest.BeginObject();
 			{
-				BatchRequest << ANSITEXTVIEW("Method") << ANSITEXTVIEW("GetCacheValues");
+				BatchRequest << ANSITEXTVIEW("Method") << "GetCacheChunks";
+				BatchRequest.AddInteger(ANSITEXTVIEW("MethodVersion"), 1);
 				BatchRequest.BeginObject(ANSITEXTVIEW("Params"));
 				{
-					ECachePolicy BatchDefaultPolicy = Batch[0].Policy;
-					BatchRequest << ANSITEXTVIEW("DefaultPolicy") << *WriteToString<128>(BatchDefaultPolicy);
-					BatchRequest.AddString(ANSITEXTVIEW("Namespace"), StructuredNamespace);
+					ECachePolicy DefaultPolicy = Batch[0].Policy;
+					BatchRequest << ANSITEXTVIEW("DefaultPolicy") << WriteToString<128>(DefaultPolicy);
+					BatchRequest.AddString(ANSITEXTVIEW("Namespace"), CacheStore.StructuredNamespace);
 
-					BatchRequest.BeginArray("Requests");
-					for (const FCacheGetValueRequest& Request : Batch)
+					BatchRequest.BeginArray(ANSITEXTVIEW("ChunkRequests"));
+					for (const FCacheGetChunkRequest& Request : Batch)
 					{
 						BatchRequest.BeginObject();
 						{
 							BatchRequest << ANSITEXTVIEW("Key") << Request.Key;
-							if (Request.Policy != BatchDefaultPolicy)
+
+							if (Request.Id.IsValid())
+							{
+								BatchRequest.AddObjectId(ANSITEXTVIEW("ValueId"), Request.Id);
+							}
+							if (Request.RawOffset != 0)
+							{
+								BatchRequest << ANSITEXTVIEW("RawOffset") << Request.RawOffset;
+							}
+							if (Request.RawSize != MAX_uint64)
+							{
+								BatchRequest << ANSITEXTVIEW("RawSize") << Request.RawSize;
+							}
+							if (!Request.RawHash.IsZero())
+							{
+								BatchRequest << ANSITEXTVIEW("ChunkId") << Request.RawHash;
+							}
+							if (Request.Policy != DefaultPolicy)
 							{
 								BatchRequest << ANSITEXTVIEW("Policy") << WriteToString<128>(Request.Policy);
 							}
@@ -755,8 +814,8 @@ void FZenCacheStore::GetValue(
 
 			{
 				LLM_SCOPE_BYTAG(UntaggedDDCResult);
-				FScopedHttpPoolRequestPtr Request(RequestPool.Get());
-				HttpResult = PerformBlockingRpc(*Request.Get(), TEXTVIEW("/z$/$rpc"), BatchRequest.Save().AsObject(), BatchResponse);
+				FScopedHttpPoolRequestPtr Request(CacheStore.RequestPool.Get());
+				HttpResult = CacheStore.PerformBlockingRpc(*Request.Get(), TEXTVIEW("/z$/$rpc"), BatchRequest.Save().AsObject(), BatchResponse);
 			}
 
 			int32 RequestIndex = 0;
@@ -764,209 +823,307 @@ void FZenCacheStore::GetValue(
 			{
 				const FCbObject& ResponseObj = BatchResponse.GetObject();
 
-				for (FCbFieldView ResultField : ResponseObj[ANSITEXTVIEW("Result")])
+				for (FCbFieldView ResultView : ResponseObj[ANSITEXTVIEW("Result")])
 				{
 					if (RequestIndex >= Batch.Num())
 					{
 						++RequestIndex;
 						continue;
 					}
-					const FCacheGetValueRequest& Request = Batch[RequestIndex++];
-
-					FCbObjectView ResultObj = ResultField.AsObjectView();
-					TOptional<FValue> Value;
-					if (DebugOptions.ShouldSimulateGetMiss(Request.Key))
+					const FCacheGetChunkRequest& Request = Batch[RequestIndex++];
+						
+					FIoHash RawHash;
+					bool Succeeded = false;
+					uint64 RawSize = 0;
+					FCbObjectView ResultObject = ResultView.AsObjectView();
+					FSharedBuffer RequestedBytes;
+					if (CacheStore.DebugOptions.ShouldSimulateGetMiss(Request.Key))
 					{
-						UE_LOG(LogDerivedDataCache, Verbose, TEXT("%s: Simulated miss for GetValue of '%s' from '%s'"),
-							*GetName(), *WriteToString<96>(Request.Key), *Request.Name);
+						UE_LOG(LogDerivedDataCache, Verbose, TEXT("%s: Simulated miss for get of '%s' from '%s'"),
+							*CacheStore.GetName(), *WriteToString<96>(Request.Key, '/', Request.Id), *Request.Name);
 					}
 					else
 					{
-						FCbFieldView RawHashField = ResultObj["RawHash"];
-						FIoHash RawHash = RawHashField.AsHash();
-						if (const FCbAttachment* Attachment = BatchResponse.FindAttachment(RawHash))
-						{
-							Value.Emplace(Attachment->AsCompressedBinary());
-						}
-						else
+						FCbFieldView HashView = ResultObject[ANSITEXTVIEW("RawHash")];
+						RawHash = HashView.AsHash();
+						if (!HashView.HasError())
 						{
-							FCbFieldView RawSizeField = ResultObj["RawSize"];
-							uint64 RawSize = RawSizeField.AsUInt64();
-							if (!RawSizeField.HasError() && !RawHashField.HasError())
+							if (const FCbAttachment* Attachment = BatchResponse.FindAttachment(HashView.AsHash()))
 							{
-								Value.Emplace(RawHash, RawSize);
+								FCompressedBuffer CompressedBuffer = Attachment->AsCompressedBinary();
+								if (CompressedBuffer)
+								{
+									TRACE_COUNTER_ADD(ZenDDC_BytesReceived, CompressedBuffer.GetCompressedSize());
+									RequestedBytes = FCompressedBufferReader(CompressedBuffer).Decompress(Request.RawOffset, Request.RawSize);
+									RawSize = RequestedBytes.GetSize();
+									Succeeded = true;
+								}
+							}
+							else
+							{
+								FCbFieldView RawSizeField = ResultObject[ANSITEXTVIEW("RawSize")];
+								uint64 TotalSize = RawSizeField.AsUInt64();
+								Succeeded = !RawSizeField.HasError();
+								if (Succeeded)
+								{
+									RawSize = FMath::Min(Request.RawSize, TotalSize - FMath::Min(Request.RawOffset, TotalSize));
+								}
 							}
 						}
 					}
-					(bool)Value ? OnHit(Request, MoveTemp(*Value)) : OnMiss(Request);
-				}
-				if (RequestIndex != Batch.Num())
-				{
-					UE_LOG(LogDerivedDataCache, Warning,
-						TEXT("Invalid response received from GetCacheValues rpc: %d results expected, received %d."), Batch.Num(), RequestIndex);
+					Succeeded ? OnHit(Request, MoveTemp(RawHash), RawSize, MoveTemp(RequestedBytes)) : OnMiss(Request);
 				}
 			}
-			for (const FCacheGetValueRequest& Request : Batch.RightChop(RequestIndex))
+			for (const FCacheGetChunkRequest& Request : Batch.RightChop(RequestIndex))
 			{
 				OnMiss(Request);
 			}
 		});
-	TRACE_COUNTER_SUBTRACT(ZenDDC_CacheRecordRequestCount, int64(Requests.Num()));
-}
-
-void FZenCacheStore::GetChunks(
-	TConstArrayView<FCacheGetChunkRequest> Requests,
-	IRequestOwner& Owner,
-	FOnCacheGetChunkComplete&& OnComplete)
-{
-	TRACE_CPUPROFILER_EVENT_SCOPE(ZenDDC::GetChunks);
-	TRACE_COUNTER_ADD(ZenDDC_ChunkRequestCount, int64(Requests.Num()));
-	TRACE_COUNTER_ADD(ZenDDC_Get, int64(Requests.Num()));
-	COOK_STAT(auto Timer = UsageStats.TimeGet());
-
-	TArray<FCacheGetChunkRequest, TInlineAllocator<16>> SortedRequests(Requests);
-	SortedRequests.StableSort(TChunkLess());
-	auto OnHit = [this, &OnComplete COOK_STAT(, &Timer)](const FCacheGetChunkRequest& Request, FIoHash&& RawHash, uint64 RawSize, FSharedBuffer&& RequestedBytes)
+	}
+private:
+	void OnHit(const FCacheGetChunkRequest& Request, FIoHash&& RawHash, uint64 RawSize, FSharedBuffer&& RequestedBytes)
 	{
 		UE_LOG(LogDerivedDataCache, Verbose, TEXT("%s: CacheChunk hit for '%s' from '%s'"),
-			*GetName(), *WriteToString<96>(Request.Key, '/', Request.Id), *Request.Name);
+			*CacheStore.GetName(), *WriteToString<96>(Request.Key, '/', Request.Id), *Request.Name);
 		TRACE_COUNTER_ADD(ZenDDC_GetHit, int64(1));
 		COOK_STAT(Timer.AddHit(RawSize));
 		OnComplete({Request.Name, Request.Key, Request.Id, Request.RawOffset,
 			RawSize, MoveTemp(RawHash), MoveTemp(RequestedBytes), Request.UserData, EStatus::Ok});
 	};
 
-	auto OnMiss = [this, &OnComplete COOK_STAT(, &Timer)](const FCacheGetChunkRequest& Request)
+	void OnMiss(const FCacheGetChunkRequest& Request)
 	{
 		UE_LOG(LogDerivedDataCache, Display, TEXT("%s: CacheChunk miss with missing value '%s' for '%s' from '%s'"),
-			*GetName(), *WriteToString<16>(Request.Id), *WriteToString<96>(Request.Key), *Request.Name);
+			*CacheStore.GetName(), *WriteToString<16>(Request.Id), *WriteToString<96>(Request.Key), *Request.Name);
 		OnComplete(Request.MakeResponse(EStatus::Error));
 	};
 
-	ForEachBatch(CacheChunksBatchSize, SortedRequests.Num(),
-		[this, &SortedRequests, &Owner, &OnHit, &OnMiss](int32 BatchFirst, int32 BatchLast)
+	FZenCacheStore& CacheStore;
+	IRequestOwner& Owner;
+	TArray<FCacheGetChunkRequest, TInlineAllocator<1>> Requests;
+	FOnCacheGetChunkComplete OnComplete;
+	COOK_STAT(FCookStats::FScopedStatsCounter Timer = CacheStore.UsageStats.TimeGet());
+};
+
+FZenCacheStore::FZenCacheStore(
+	const TCHAR* InServiceUrl,
+	const TCHAR* InNamespace,
+	const TCHAR* InStructuredNamespace)
+	: Namespace(InNamespace)
+	, StructuredNamespace(InStructuredNamespace)
+	, ZenService(InServiceUrl)
+{
+	constexpr uint32 MaxConnections = 64;
+	constexpr uint32 RequestPoolSize = 128;
+	constexpr uint32 RequestPoolOverflowSize = 128;
+	
+	SharedData = MakeUnique<FHttpSharedData>(MaxConnections);
+	if (IsServiceReady())
 	{
-		TArrayView<const FCacheGetChunkRequest> Batch(SortedRequests.GetData() + BatchFirst, BatchLast - BatchFirst + 1);
+		RequestPool = MakeUnique<FHttpRequestPool>(ZenService.GetInstance().GetURL(), ZenService.GetInstance().GetURL(), nullptr, SharedData.Get(), RequestPoolSize, RequestPoolOverflowSize);
+		bIsUsable = true;
 
-		FCbWriter BatchRequest;
-		BatchRequest.BeginObject();
-		{
-			BatchRequest << ANSITEXTVIEW("Method") << "GetCacheChunks";
-			BatchRequest.AddInteger(ANSITEXTVIEW("MethodVersion"), 1);
-			BatchRequest.BeginObject(ANSITEXTVIEW("Params"));
-			{
-				ECachePolicy DefaultPolicy = Batch[0].Policy;
-				BatchRequest << ANSITEXTVIEW("DefaultPolicy") << WriteToString<128>(DefaultPolicy);
-				BatchRequest.AddString(ANSITEXTVIEW("Namespace"), StructuredNamespace);
+		// Issue a request for stats as it will be fetched asynchronously and issuing now makes them available sooner for future callers.
+		Zen::FZenStats ZenStats;
+		ZenService.GetInstance().GetStats(ZenStats);
+	}
 
-				BatchRequest.BeginArray(ANSITEXTVIEW("ChunkRequests"));
-				for (const FCacheGetChunkRequest& Request : Batch)
-				{
-					BatchRequest.BeginObject();
-					{
-						BatchRequest << ANSITEXTVIEW("Key") << Request.Key;
+	GConfig->GetInt(TEXT("Zen"), TEXT("BatchPutMaxBytes"), BatchPutMaxBytes, GEngineIni);
+	GConfig->GetInt(TEXT("Zen"), TEXT("CacheRecordBatchSize"), CacheRecordBatchSize, GEngineIni);
+	GConfig->GetInt(TEXT("Zen"), TEXT("CacheChunksBatchSize"), CacheChunksBatchSize, GEngineIni);
+}
 
-						if (Request.Id.IsValid())
-						{
-							BatchRequest.AddObjectId(ANSITEXTVIEW("ValueId"), Request.Id);
-						}
-						if (Request.RawOffset != 0)
-						{
-							BatchRequest << ANSITEXTVIEW("RawOffset") << Request.RawOffset;
-						}
-						if (Request.RawSize != MAX_uint64)
-						{
-							BatchRequest << ANSITEXTVIEW("RawSize") << Request.RawSize;
-						}
-						if (!Request.RawHash.IsZero())
-						{
-							BatchRequest << ANSITEXTVIEW("ChunkId") << Request.RawHash;
-						}
-						if (Request.Policy != DefaultPolicy)
-						{
-							BatchRequest << ANSITEXTVIEW("Policy") << WriteToString<128>(Request.Policy);
-						}
-					}
-					BatchRequest.EndObject();
-				}
-				BatchRequest.EndArray();
-			}
-			BatchRequest.EndObject();
-		}
-		BatchRequest.EndObject();
+bool FZenCacheStore::IsServiceReady()
+{
+	return ZenService.GetInstance().IsServiceReady();
+}
 
-		FCbPackage BatchResponse;
-		FHttpRequest::EResult HttpResult = FHttpRequest::EResult::Failed;
+FHttpRequest::EResult FZenCacheStore::ParseRpcResponse(FHttpRequest::EResult ResultFromPost, FHttpRequest& Request, FCbPackage& OutResponse)
+{
+	if (ResultFromPost != FHttpRequest::EResult::Success || !FHttpRequest::IsSuccessResponse(Request.GetResponseCode()))
+	{
+		return FHttpRequest::EResult::Failed;
+	}
 
+	const TArray64<uint8>& ResponseBuffer = Request.GetResponseBuffer();
+	if (ResponseBuffer.Num())
+	{
+		FLargeMemoryReader Ar(ResponseBuffer.GetData(), ResponseBuffer.Num());
+		if (!OutResponse.TryLoad(Ar))
 		{
-			LLM_SCOPE_BYTAG(UntaggedDDCResult);
-			FScopedHttpPoolRequestPtr Request(RequestPool.Get());
-			HttpResult = PerformBlockingRpc(*Request.Get(), TEXTVIEW("/z$/$rpc"), BatchRequest.Save().AsObject(), BatchResponse);
+			return FHttpRequest::EResult::Failed;
 		}
+	}
 
-		int32 RequestIndex = 0;
-		if (HttpResult == FHttpRequest::EResult::Success)
-		{
-			const FCbObject& ResponseObj = BatchResponse.GetObject();
+	return FHttpRequest::EResult::Success;
+}
 
-			for (FCbFieldView ResultView : ResponseObj[ANSITEXTVIEW("Result")])
-			{
-				if (RequestIndex >= Batch.Num())
-				{
-					++RequestIndex;
-					continue;
-				}
-				const FCacheGetChunkRequest& Request = Batch[RequestIndex++];
-					
-				FIoHash RawHash;
-				bool Succeeded = false;
-				uint64 RawSize = 0;
-				FCbObjectView ResultObject = ResultView.AsObjectView();
-				FSharedBuffer RequestedBytes;
-				if (DebugOptions.ShouldSimulateGetMiss(Request.Key))
-				{
-					UE_LOG(LogDerivedDataCache, Verbose, TEXT("%s: Simulated miss for get of '%s' from '%s'"),
-						*GetName(), *WriteToString<96>(Request.Key, '/', Request.Id), *Request.Name);
-				}
-				else
-				{
-					FCbFieldView HashView = ResultObject[ANSITEXTVIEW("RawHash")];
-					RawHash = HashView.AsHash();
-					if (!HashView.HasError())
-					{
-						if (const FCbAttachment* Attachment = BatchResponse.FindAttachment(HashView.AsHash()))
-						{
-							FCompressedBuffer CompressedBuffer = Attachment->AsCompressedBinary();
-							if (CompressedBuffer)
-							{
-								TRACE_COUNTER_ADD(ZenDDC_BytesReceived, CompressedBuffer.GetCompressedSize());
-								RequestedBytes = FCompressedBufferReader(CompressedBuffer).Decompress(Request.RawOffset, Request.RawSize);
-								RawSize = RequestedBytes.GetSize();
-								Succeeded = true;
-							}
-						}
-						else
-						{
-							FCbFieldView RawSizeField = ResultObject[ANSITEXTVIEW("RawSize")];
-							uint64 TotalSize = RawSizeField.AsUInt64();
-							Succeeded = !RawSizeField.HasError();
-							if (Succeeded)
-							{
-								RawSize = FMath::Min(Request.RawSize, TotalSize - FMath::Min(Request.RawOffset, TotalSize));
-							}
-						}
-					}
-				}
-				Succeeded ? OnHit(Request, MoveTemp(RawHash), RawSize, MoveTemp(RequestedBytes)) : OnMiss(Request);
-			}
-		}
-		for (const FCacheGetChunkRequest& Request : Batch.RightChop(RequestIndex))
-		{
-			OnMiss(Request);
-		}
-	});
+FHttpRequest::EResult FZenCacheStore::PerformBlockingRpc(FHttpRequest& Request,
+	FStringView Uri,
+	FCbObject RequestObject,
+	FCbPackage& OutResponse)
+{
+	return ParseRpcResponse(
+		Request.PerformBlockingPost(Uri, RequestObject.GetBuffer(), EHttpContentType::CbObject, EHttpContentType::CbPackage), Request, OutResponse);
+}
 
-	TRACE_COUNTER_SUBTRACT(ZenDDC_ChunkRequestCount, int64(Requests.Num()));
+FHttpRequest::EResult FZenCacheStore::PerformBlockingRpc(FHttpRequest& Request,
+	FStringView Uri,
+	const FCbPackage& RequestPackage,
+	FCbPackage& OutResponse)
+{
+	FLargeMemoryWriter PackageMemory;
+	UE::Zen::Http::SaveCbPackage(RequestPackage, PackageMemory);
+
+	return ParseRpcResponse(
+		Request.PerformBlockingPost(Uri, FCompositeBuffer(FSharedBuffer::MakeView(PackageMemory.GetView())), EHttpContentType::CbPackage, EHttpContentType::CbPackage), Request, OutResponse);
+}
+
+void FZenCacheStore::EnqueueAsyncRpc(FHttpRequest& Request,
+	UE::DerivedData::IRequestOwner& Owner,
+	FHttpRequestPool* Pool,
+	FStringView Uri,
+	FCbObject RequestObject,
+	FOnRpcComplete&& OnComplete)
+{
+	auto OnHttpRequestComplete = [OnComplete = MoveTemp(OnComplete)]
+	(FHttpRequest::EResult HttpResult, FHttpRequest* Request)
+	{
+		FCbPackage Response;
+		FHttpRequest::EResult ParsedResponseResult = ParseRpcResponse(HttpResult, *Request, Response);
+		OnComplete(ParsedResponseResult, Request, Response);
+		return FHttpRequest::ECompletionBehavior::Done;
+	};
+	Request.EnqueueAsyncPost(Owner, Pool, Uri, RequestObject.GetBuffer(), MoveTemp(OnHttpRequestComplete), EHttpContentType::CbObject, EHttpContentType::CbPackage);
+}
+
+void FZenCacheStore::EnqueueAsyncRpc(FHttpRequest& Request,
+	UE::DerivedData::IRequestOwner& Owner,
+	FHttpRequestPool* Pool,
+	FStringView Uri,
+	const FCbPackage& RequestPackage,
+	FOnRpcComplete&& OnComplete)
+{
+	auto OnHttpRequestComplete = [OnComplete = MoveTemp(OnComplete)]
+	(FHttpRequest::EResult HttpResult, FHttpRequest* Request)
+	{
+		FCbPackage Response;
+		FHttpRequest::EResult ParsedResponseResult = ParseRpcResponse(HttpResult, *Request, Response);
+		OnComplete(ParsedResponseResult, Request, Response);
+		return FHttpRequest::ECompletionBehavior::Done;
+	};
+	FLargeMemoryWriter PackageMemory;
+	UE::Zen::Http::SaveCbPackage(RequestPackage, PackageMemory);
+	uint64 PackageMemorySize = PackageMemory.TotalSize();
+	FSharedBuffer PackageSharedBuffer = FSharedBuffer::TakeOwnership(PackageMemory.ReleaseOwnership(), PackageMemorySize, FMemory::Free);
+	Request.EnqueueAsyncPost(Owner, Pool, Uri, FCompositeBuffer(PackageSharedBuffer), MoveTemp(OnHttpRequestComplete), EHttpContentType::CbPackage, EHttpContentType::CbPackage);
+}
+
+void FZenCacheStore::LegacyStats(FDerivedDataCacheStatsNode& OutNode)
+{
+	Zen::FZenStats ZenStats;
+
+	FDerivedDataCacheUsageStats LocalStats;
+	FDerivedDataCacheUsageStats RemoteStats;
+
+#if ENABLE_COOK_STATS
+	using EHitOrMiss = FCookStats::CallStats::EHitOrMiss;
+	using ECacheStatType = FCookStats::CallStats::EStatType;
+
+	Zen::GetDefaultServiceInstance().GetStats(ZenStats);
+
+	const int64 RemotePutSize = int64(ZenStats.UpstreamStats.TotalUploadedMB * 1024 * 1024);
+	const int64 RemoteGetSize = int64(ZenStats.UpstreamStats.TotalDownloadedMB * 1024 * 1024);
+	const int64 LocalGetSize = FMath::Max<int64>(0, UsageStats.GetStats.GetAccumulatedValueAnyThread(EHitOrMiss::Hit, ECacheStatType::Bytes) - RemoteGetSize);
+
+	LocalStats.PutStats = UsageStats.PutStats;
+	LocalStats.ExistsStats = UsageStats.ExistsStats;
+	LocalStats.PrefetchStats = UsageStats.PrefetchStats;
+
+	LocalStats.GetStats.Accumulate(EHitOrMiss::Hit, ECacheStatType::Counter, ZenStats.CacheStats.Hits - ZenStats.CacheStats.UpstreamHits, /*bIsInGameThread*/ false);
+	LocalStats.GetStats.Accumulate(EHitOrMiss::Miss, ECacheStatType::Counter, ZenStats.CacheStats.Misses + ZenStats.CacheStats.UpstreamHits, /*bIsInGameThread*/ false);
+	RemoteStats.GetStats.Accumulate(EHitOrMiss::Hit, ECacheStatType::Counter, ZenStats.CacheStats.UpstreamHits, /*bIsInGameThread*/ false);
+	RemoteStats.GetStats.Accumulate(EHitOrMiss::Miss, ECacheStatType::Counter, ZenStats.CacheStats.Misses, /*bIsInGameThread*/ false);
+
+	LocalStats.GetStats.Accumulate(EHitOrMiss::Hit, ECacheStatType::Bytes, LocalGetSize, /*bIsInGameThread*/ false);
+	RemoteStats.GetStats.Accumulate(EHitOrMiss::Hit, ECacheStatType::Bytes, RemoteGetSize, /*bIsInGameThread*/ false);
+	RemoteStats.PutStats.Accumulate(EHitOrMiss::Hit, ECacheStatType::Bytes, RemotePutSize, /*bIsInGameThread*/ false);
+#endif
+
+	if (ZenStats.UpstreamStats.EndPointStats.IsEmpty())
+	{
+		OutNode = {TEXT("Zen"), ZenService.GetInstance().GetURL(), /*bIsLocal*/ true};
+		OutNode.UsageStats.Add(TEXT(""), LocalStats);
+		return;
+	}
+
+	TSharedRef<FDerivedDataCacheStatsNode> LocalNode =
+		MakeShared<FDerivedDataCacheStatsNode>(TEXT("Zen"), ZenService.GetInstance().GetURL(), /*bIsLocal*/ true);
+	LocalNode->UsageStats.Add(TEXT(""), LocalStats);
+
+	TSharedRef<FDerivedDataCacheStatsNode> RemoteNode =
+		MakeShared<FDerivedDataCacheStatsNode>(ZenStats.UpstreamStats.EndPointStats[0].Name, ZenStats.UpstreamStats.EndPointStats[0].Url, /*bIsLocal*/ false);
+	RemoteNode->UsageStats.Add(TEXT(""), RemoteStats);
+
+	OutNode = {TEXT("Zen Group"), TEXT(""), /*bIsLocal*/ true};
+	OutNode.Children.Add(LocalNode);
+	OutNode.Children.Add(RemoteNode);
+}
+
+bool FZenCacheStore::LegacyDebugOptions(FBackendDebugOptions& InOptions)
+{
+	DebugOptions = InOptions;
+	return true;
+}
+
+void FZenCacheStore::Put(
+	const TConstArrayView<FCachePutRequest> Requests,
+	IRequestOwner& Owner,
+	FOnCachePutComplete&& OnComplete)
+{
+	TRACE_CPUPROFILER_EVENT_SCOPE(ZenDDC::PutCachedRecord);
+	TRefCountPtr<FPutOp> PutOp = MakeAsyncOp<FPutOp>(*this, Owner, Requests, MoveTemp(OnComplete));
+	PutOp->IssueRequests();
+}
+
+void FZenCacheStore::Get(
+	const TConstArrayView<FCacheGetRequest> Requests,
+	IRequestOwner& Owner,
+	FOnCacheGetComplete&& OnComplete)
+{
+	TRACE_CPUPROFILER_EVENT_SCOPE(ZenDDC::GetCacheRecord);
+	TRefCountPtr<FGetOp> GetOp = MakeAsyncOp<FGetOp>(*this, Owner, Requests, MoveTemp(OnComplete));
+	GetOp->IssueRequests();
+}
+
+void FZenCacheStore::PutValue(
+	TConstArrayView<FCachePutValueRequest> Requests,
+	IRequestOwner& Owner,
+	FOnCachePutValueComplete&& OnComplete)
+{
+	TRACE_CPUPROFILER_EVENT_SCOPE(ZenDDC::PutValue);
+	TRefCountPtr<FPutValueOp> PutValueOp = MakeAsyncOp<FPutValueOp>(*this, Owner, Requests, MoveTemp(OnComplete));
+	PutValueOp->IssueRequests();
+}
+
+void FZenCacheStore::GetValue(
+	TConstArrayView<FCacheGetValueRequest> Requests,
+	IRequestOwner& Owner,
+	FOnCacheGetValueComplete&& OnComplete)
+{
+	TRACE_CPUPROFILER_EVENT_SCOPE(ZenDDC::GetValue);
+	TRefCountPtr<FGetValueOp> GetValueOp = MakeAsyncOp<FGetValueOp>(*this, Owner, Requests, MoveTemp(OnComplete));
+	GetValueOp->IssueRequests();
+}
+
+void FZenCacheStore::GetChunks(
+	TConstArrayView<FCacheGetChunkRequest> Requests,
+	IRequestOwner& Owner,
+	FOnCacheGetChunkComplete&& OnComplete)
+{
+	TRACE_CPUPROFILER_EVENT_SCOPE(ZenDDC::GetChunks);
+	TRefCountPtr<FGetChunksOp> GetChunksOp = MakeAsyncOp<FGetChunksOp>(*this, Owner, Requests, MoveTemp(OnComplete));
+	GetChunksOp->IssueRequests();
 }
 
 TTuple<ILegacyCacheStore*, ECacheStoreFlags> CreateZenCacheStore(const TCHAR* NodeName, const TCHAR* Config)
