DDC: Removed the remaining legacy cache implementations and made the default implementation non-virtual

#rnx

diff --git a/Engine/Source/Developer/DerivedDataCache/Private/DerivedDataCacheStoreAsync.cpp b/Engine/Source/Developer/DerivedDataCache/Private/DerivedDataCacheStoreAsync.cpp
--- a/Engine/Source/Developer/DerivedDataCache/Private/DerivedDataCacheStoreAsync.cpp
+++ b/Engine/Source/Developer/DerivedDataCache/Private/DerivedDataCacheStoreAsync.cpp
@@ -93,43 +93,6 @@ public:
 		Execute(COOK_STAT(&FDerivedDataCacheUsageStats::TimeGet,) Requests, Owner, MoveTemp(OnComplete), &ICacheStore::GetChunks);
 	}
 
-	virtual void LegacyPut(
-		TConstArrayView<FLegacyCachePutRequest> Requests,
-		IRequestOwner& Owner,
-		FOnLegacyCachePutComplete&& OnComplete) override
-	{
-		if (MemoryCache)
-		{
-			MemoryCache->LegacyPut(Requests, Owner, [](auto&&){});
-			Execute(COOK_STAT(&FDerivedDataCacheUsageStats::TimePut,) Requests, Owner,
-				[this, OnComplete = MoveTemp(OnComplete)](FLegacyCachePutResponse&& Response)
-				{
-					MemoryCache->LegacyDelete(Response.Key);
-					OnComplete(MoveTemp(Response));
-				}, &ILegacyCacheStore::LegacyPut);
-		}
-		else
-		{
-			Execute(COOK_STAT(&FDerivedDataCacheUsageStats::TimePut,) Requests, Owner, MoveTemp(OnComplete), &ILegacyCacheStore::LegacyPut);
-		}
-	}
-
-	virtual void LegacyGet(
-		TConstArrayView<FLegacyCacheGetRequest> Requests,
-		IRequestOwner& Owner,
-		FOnLegacyCacheGetComplete&& OnComplete) override
-	{
-		Execute(COOK_STAT(&FDerivedDataCacheUsageStats::TimeGet,) Requests, Owner, MoveTemp(OnComplete), &ILegacyCacheStore::LegacyGet);
-	}
-
-	virtual void LegacyDelete(
-		TConstArrayView<FLegacyCacheDeleteRequest> Requests,
-		IRequestOwner& Owner,
-		FOnLegacyCacheDeleteComplete&& OnComplete) override
-	{
-		Execute(COOK_STAT(&FDerivedDataCacheUsageStats::TimePut,) Requests, Owner, MoveTemp(OnComplete), &ILegacyCacheStore::LegacyDelete);
-	}
-
 	virtual void LegacyStats(FDerivedDataCacheStatsNode& OutNode) override;
 
 	virtual bool LegacyDebugOptions(FBackendDebugOptions& Options) override
diff --git a/Engine/Source/Developer/DerivedDataCache/Private/DerivedDataCacheStoreHierarchy.cpp b/Engine/Source/Developer/DerivedDataCache/Private/DerivedDataCacheStoreHierarchy.cpp
--- a/Engine/Source/Developer/DerivedDataCache/Private/DerivedDataCacheStoreHierarchy.cpp
+++ b/Engine/Source/Developer/DerivedDataCache/Private/DerivedDataCacheStoreHierarchy.cpp
@@ -64,21 +64,6 @@ public:
 		IRequestOwner& Owner,
 		FOnCacheGetChunkComplete&& OnComplete) final;
 
-	void LegacyPut(
-		TConstArrayView<FLegacyCachePutRequest> Requests,
-		IRequestOwner& Owner,
-		FOnLegacyCachePutComplete&& OnComplete) final;
-
-	void LegacyGet(
-		TConstArrayView<FLegacyCacheGetRequest> Requests,
-		IRequestOwner& Owner,
-		FOnLegacyCacheGetComplete&& OnComplete) final;
-
-	void LegacyDelete(
-		TConstArrayView<FLegacyCacheDeleteRequest> Requests,
-		IRequestOwner& Owner,
-		FOnLegacyCacheDeleteComplete&& OnComplete) final;
-
 	void LegacyStats(FDerivedDataCacheStatsNode& OutNode) final;
 
 	bool LegacyDebugOptions(FBackendDebugOptions& Options) final;
@@ -86,7 +71,6 @@ public:
 	template <typename PutRequestType, typename GetRequestType> struct TBatchParams;
 	using FCacheRecordBatchParams = TBatchParams<FCachePutRequest, FCacheGetRequest>;
 	using FCacheValueBatchParams = TBatchParams<FCachePutValueRequest, FCacheGetValueRequest>;
-	using FLegacyCacheBatchParams = TBatchParams<FLegacyCachePutRequest, FLegacyCacheGetRequest>;
 
 private:
 	// Caller must hold a write lock on CacheStoresLock.
@@ -99,7 +83,6 @@ private:
 	template <typename Params> class TGetBatch;
 
 	class FGetChunksBatch;
-	class FLegacyDeleteBatch;
 
 	enum class ECacheStoreNodeFlags : uint32;
 	FRIEND_ENUM_CLASS_FLAGS(ECacheStoreNodeFlags);
@@ -1071,194 +1054,6 @@ void FCacheStoreHierarchy::GetChunks(
 
 ///////////////////////////////////////////////////////////////////////////////////////////////////////////////////////
 
-template <>
-auto FCacheStoreHierarchy::FLegacyCacheBatchParams::Put() -> PutFunctionType
-{
-	return &ILegacyCacheStore::LegacyPut;
-};
-
-template <>
-auto FCacheStoreHierarchy::FLegacyCacheBatchParams::Get() -> GetFunctionType
-{
-	return &ILegacyCacheStore::LegacyGet;
-};
-
-template <>
-bool FCacheStoreHierarchy::FLegacyCacheBatchParams::HasResponseData(const FLegacyCacheGetResponse& Response)
-{
-	return Response.Value.HasData();
-}
-
-template <>
-void FCacheStoreHierarchy::FLegacyCacheBatchParams::FilterResponseByRequest(
-	FLegacyCacheGetResponse& Response,
-	const FLegacyCacheGetRequest& Request)
-{
-	if (Response.Value && EnumHasAnyFlags(Request.Policy, ECachePolicy::SkipData))
-	{
-		Response.Value.Reset();
-	}
-}
-
-template <>
-FLegacyCachePutRequest FCacheStoreHierarchy::FLegacyCacheBatchParams::MakePutRequest(
-	const FLegacyCacheGetResponse& Response,
-	const FLegacyCacheGetRequest& Request)
-{
-	return {Response.Name, Response.Key, Response.Value, Request.Policy};
-}
-
-template <>
-FLegacyCacheGetRequest FCacheStoreHierarchy::FLegacyCacheBatchParams::MakeGetRequest(
-	const FLegacyCachePutRequest& Request,
-	const uint64 UserData)
-{
-	return {Request.Name, Request.Key, AddPolicy(Request.Policy, ECachePolicy::SkipData), UserData};
-}
-
-void FCacheStoreHierarchy::LegacyPut(
-	const TConstArrayView<FLegacyCachePutRequest> Requests,
-	IRequestOwner& Owner,
-	FOnLegacyCachePutComplete&& OnComplete)
-{
-	TPutBatch<FLegacyCacheBatchParams>::Begin(*this, Requests, Owner, MoveTemp(OnComplete));
-}
-
-void FCacheStoreHierarchy::LegacyGet(
-	const TConstArrayView<FLegacyCacheGetRequest> Requests,
-	IRequestOwner& Owner,
-	FOnLegacyCacheGetComplete&& OnComplete)
-{
-	TGetBatch<FLegacyCacheBatchParams>::Begin(*this, Requests, Owner, MoveTemp(OnComplete));
-}
-
-///////////////////////////////////////////////////////////////////////////////////////////////////////////////////////
-
-class FCacheStoreHierarchy::FLegacyDeleteBatch final : public FBatchBase
-{
-public:
-	static void Begin(
-		const FCacheStoreHierarchy& Hierarchy,
-		TConstArrayView<FLegacyCacheDeleteRequest> Requests,
-		IRequestOwner& Owner,
-		FOnLegacyCacheDeleteComplete&& OnComplete);
-
-private:
-	FLegacyDeleteBatch(
-		const FCacheStoreHierarchy& InHierarchy,
-		const TConstArrayView<FLegacyCacheDeleteRequest> InRequests,
-		IRequestOwner& InOwner,
-		FOnLegacyCacheDeleteComplete&& InOnComplete)
-		: Hierarchy(InHierarchy)
-		, Requests(InRequests)
-		, BatchOwner(InOwner)
-		, OnComplete(MoveTemp(InOnComplete))
-	{
-		States.SetNum(Requests.Num());
-	}
-
-	void DispatchRequests();
-	void CompleteRequest(FLegacyCacheDeleteResponse&& Response);
-
-	struct FRequestState
-	{
-		bool bOk = false;
-	};
-
-	const FCacheStoreHierarchy& Hierarchy;
-	TArray<FLegacyCacheDeleteRequest, TInlineAllocator<1>> Requests;
-	IRequestOwner& BatchOwner;
-	FOnLegacyCacheDeleteComplete OnComplete;
-
-	TArray<FRequestState, TInlineAllocator<1>> States;
-	FCounterEvent RemainingRequestCount;
-	int32 NodeIndex = 0;
-};
-
-void FCacheStoreHierarchy::FLegacyDeleteBatch::Begin(
-	const FCacheStoreHierarchy& InHierarchy,
-	const TConstArrayView<FLegacyCacheDeleteRequest> InRequests,
-	IRequestOwner& InOwner,
-	FOnLegacyCacheDeleteComplete&& InOnComplete)
-{
-	if (InRequests.IsEmpty() || !EnumHasAnyFlags(InHierarchy.CombinedNodeFlags, ECacheStoreNodeFlags::HasStoreNode))
-	{
-		return CompleteWithStatus(InRequests, InOnComplete, EStatus::Error);
-	}
-
-	TRefCountPtr<FLegacyDeleteBatch> State = new FLegacyDeleteBatch(InHierarchy, InRequests, InOwner, MoveTemp(InOnComplete));
-	State->DispatchRequests();
-}
-
-void FCacheStoreHierarchy::FLegacyDeleteBatch::DispatchRequests()
-{
-	FReadScopeLock Lock(Hierarchy.NodesLock);
-
-	TArray<FLegacyCacheDeleteRequest, TInlineAllocator<1>> NodeRequests;
-	NodeRequests.Reserve(Requests.Num());
-	for (const int32 NodeCount = Hierarchy.Nodes.Num(); NodeIndex < NodeCount && !BatchOwner.IsCanceled(); ++NodeIndex)
-	{
-		const FCacheStoreNode& Node = Hierarchy.Nodes[NodeIndex];
-
-		uint64 RequestIndex = 0;
-		for (const FLegacyCacheDeleteRequest& Request : Requests)
-		{
-			if (CanStore(Request.Policy, Node.CacheFlags))
-			{
-				NodeRequests.Add_GetRef(Request).UserData = RequestIndex;
-			}
-			++RequestIndex;
-		}
-
-		if (const int32 NodeRequestsCount = NodeRequests.Num())
-		{
-			RemainingRequestCount.Reset(NodeRequestsCount + 1);
-			Node.Cache->LegacyDelete(NodeRequests, BatchOwner,
-				[State = TRefCountPtr(this)](FLegacyCacheDeleteResponse&& Response)
-				{
-					State->CompleteRequest(MoveTemp(Response));
-				});
-			NodeRequests.Reset();
-			if (!RemainingRequestCount.Signal())
-			{
-				return;
-			}
-		}
-	}
-
-	int32 RequestIndex = 0;
-	for (const FLegacyCacheDeleteRequest& Request : Requests)
-	{
-		const bool bOk = States[RequestIndex].bOk;
-		const EStatus Status = bOk ? EStatus::Ok : BatchOwner.IsCanceled() ? EStatus::Canceled : EStatus::Error;
-		OnComplete(Request.MakeResponse(Status));
-		++RequestIndex;
-	}
-}
-
-void FCacheStoreHierarchy::FLegacyDeleteBatch::CompleteRequest(FLegacyCacheDeleteResponse&& Response)
-{
-	if (Response.Status == EStatus::Ok)
-	{
-		States[int32(Response.UserData)].bOk = true;
-	}
-	if (RemainingRequestCount.Signal())
-	{
-		++NodeIndex;
-		DispatchRequests();
-	}
-}
-
-void FCacheStoreHierarchy::LegacyDelete(
-	const TConstArrayView<FLegacyCacheDeleteRequest> Requests,
-	IRequestOwner& Owner,
-	FOnLegacyCacheDeleteComplete&& OnComplete)
-{
-	FLegacyDeleteBatch::Begin(*this, Requests, Owner, MoveTemp(OnComplete));
-}
-
-///////////////////////////////////////////////////////////////////////////////////////////////////////////////////////
-
 void FCacheStoreHierarchy::LegacyStats(FDerivedDataCacheStatsNode& OutNode)
 {
 	FReadScopeLock Lock(NodesLock);
diff --git a/Engine/Source/Developer/DerivedDataCache/Private/DerivedDataCacheStoreThrottle.cpp b/Engine/Source/Developer/DerivedDataCache/Private/DerivedDataCacheStoreThrottle.cpp
--- a/Engine/Source/Developer/DerivedDataCache/Private/DerivedDataCacheStoreThrottle.cpp
+++ b/Engine/Source/Developer/DerivedDataCache/Private/DerivedDataCacheStoreThrottle.cpp
@@ -120,58 +120,6 @@ public:
 
 	}
 
-	void LegacyPut(
-		const TConstArrayView<FLegacyCachePutRequest> Requests,
-		IRequestOwner& Owner,
-		FOnLegacyCachePutComplete&& OnComplete) final
-	{
-		struct FValueSize
-		{
-			FCacheKey Key;
-			uint64 Size;
-		};
-		TArray<FValueSize, TInlineAllocator<1>> ValueSizes;
-		ValueSizes.Reserve(Requests.Num());
-		Algo::Transform(Requests, ValueSizes, [](const FLegacyCachePutRequest& Request) -> FValueSize
-		{
-			return {Request.Key.GetKey(), Request.Value.GetRawSize()};
-		});
-
-		InnerCache->LegacyPut(Requests, Owner,
-			[this, ValueSizes = MoveTemp(ValueSizes), State = EnterThrottlingScope(), OnComplete = MoveTemp(OnComplete)](FLegacyCachePutResponse&& Response)
-			{
-				const FValueSize* Size = Algo::FindBy(ValueSizes, Response.Key.GetKey(), &FValueSize::Key);
-				CloseThrottlingScope(State, FThrottlingState(this, Size ? Size->Size : 0));
-				OnComplete(MoveTemp(Response));
-			});
-	}
-
-	void LegacyGet(
-		const TConstArrayView<FLegacyCacheGetRequest> Requests,
-		IRequestOwner& Owner,
-		FOnLegacyCacheGetComplete&& OnComplete) final
-	{
-		InnerCache->LegacyGet(Requests, Owner,
-			[this, State = EnterThrottlingScope(), OnComplete = MoveTemp(OnComplete)](FLegacyCacheGetResponse&& Response)
-			{
-				CloseThrottlingScope(State, FThrottlingState(this, Response.Value.GetRawSize()));
-				OnComplete(MoveTemp(Response));
-			});
-	}
-
-	void LegacyDelete(
-		const TConstArrayView<FLegacyCacheDeleteRequest> Requests,
-		IRequestOwner& Owner,
-		FOnLegacyCacheDeleteComplete&& OnComplete) final
-	{
-		InnerCache->LegacyDelete(Requests, Owner,
-			[this, State = EnterThrottlingScope(), OnComplete = MoveTemp(OnComplete)](FLegacyCacheDeleteResponse&& Response)
-			{
-				CloseThrottlingScope(State, FThrottlingState(this, 0));
-				OnComplete(MoveTemp(Response));
-			});
-	}
-
 	void LegacyStats(FDerivedDataCacheStatsNode& OutNode) final
 	{
 		InnerCache->LegacyStats(OutNode);
diff --git a/Engine/Source/Developer/DerivedDataCache/Private/DerivedDataCacheStoreVerify.cpp b/Engine/Source/Developer/DerivedDataCache/Private/DerivedDataCacheStoreVerify.cpp
--- a/Engine/Source/Developer/DerivedDataCache/Private/DerivedDataCacheStoreVerify.cpp
+++ b/Engine/Source/Developer/DerivedDataCache/Private/DerivedDataCacheStoreVerify.cpp
@@ -1,7 +1,6 @@
 // Copyright Epic Games, Inc. All Rights Reserved.
 
 #include "Algo/Compare.h"
-#include "Compression/OodleDataCompression.h"
 #include "Containers/Set.h"
 #include "DerivedDataCacheKey.h"
 #include "DerivedDataCacheKeyFilter.h"
@@ -90,24 +89,6 @@ public:
 		IRequestOwner& Owner,
 		FOnCacheGetChunkComplete&& OnComplete) final;
 
-	void LegacyPut(
-		TConstArrayView<FLegacyCachePutRequest> Requests,
-		IRequestOwner& Owner,
-		FOnLegacyCachePutComplete&& OnComplete) final;
-
-	void LegacyGet(
-		TConstArrayView<FLegacyCacheGetRequest> Requests,
-		IRequestOwner& Owner,
-		FOnLegacyCacheGetComplete&& OnComplete) final;
-
-	void LegacyDelete(
-		TConstArrayView<FLegacyCacheDeleteRequest> Requests,
-		IRequestOwner& Owner,
-		FOnLegacyCacheDeleteComplete&& OnComplete) final
-	{
-		InnerCache->LegacyDelete(Requests, Owner, MoveTemp(OnComplete));
-	}
-
 	void LegacyStats(FDerivedDataCacheStatsNode& OutNode) final
 	{
 		InnerCache->LegacyStats(OutNode);
@@ -137,15 +118,6 @@ private:
 		FRWLock Lock;
 	};
 
-	struct FVerifyLegacyPutState
-	{
-		TArray<FLegacyCachePutRequest> ForwardRequests;
-		TArray<FLegacyCachePutRequest> VerifyRequests;
-		FOnLegacyCachePutComplete OnComplete;
-		int32 ActiveRequests = 0;
-		FRWLock Lock;
-	};
-
 	void GetMetaComplete(IRequestOwner& Owner, FVerifyPutState* State, FCacheGetResponse&& Response);
 	void GetDataComplete(IRequestOwner& Owner, FVerifyPutState* State, FCacheGetResponse&& Response);
 	void GetComplete(IRequestOwner& Owner, FVerifyPutState* State);
@@ -154,10 +126,6 @@ private:
 	void GetDataComplete(IRequestOwner& Owner, FVerifyPutValueState* State, FCacheGetValueResponse&& Response);
 	void GetComplete(IRequestOwner& Owner, FVerifyPutValueState* State);
 
-	void GetMetaComplete(IRequestOwner& Owner, FVerifyLegacyPutState* State, FLegacyCacheGetResponse&& Response);
-	void GetDataComplete(IRequestOwner& Owner, FVerifyLegacyPutState* State, FLegacyCacheGetResponse&& Response);
-	void GetComplete(IRequestOwner& Owner, FVerifyLegacyPutState* State);
-
 	static bool CompareRecords(const FCacheRecord& PutRecord, const FCacheRecord& GetRecord, const FSharedString& Name);
 
 	static void LogChangedValue(
@@ -531,139 +499,6 @@ void FCacheStoreVerify::GetChunks(
 	}
 }
 
-void FCacheStoreVerify::LegacyPut(
-	const TConstArrayView<FLegacyCachePutRequest> Requests,
-	IRequestOwner& Owner,
-	FOnLegacyCachePutComplete&& OnComplete)
-{
-	TUniquePtr<FVerifyLegacyPutState> State = MakeUnique<FVerifyLegacyPutState>();
-	State->VerifyRequests.Reserve(Requests.Num());
-	{
-		FScopeLock Lock(&AlreadyTestedLock);
-		for (const FLegacyCachePutRequest& Request : Requests)
-		{
-			const FCacheKey& Key = Request.Key.GetKey();
-			bool bForward = !Filter.IsMatch(Key);
-			if (!bForward)
-			{
-				AlreadyTested.Add(Key, &bForward);
-			}
-			(bForward ? State->ForwardRequests : State->VerifyRequests).Add(Request);
-		}
-	}
-
-	if (State->VerifyRequests.IsEmpty())
-	{
-		return InnerCache->LegacyPut(State->ForwardRequests, Owner, MoveTemp(OnComplete));
-	}
-
-	TArray<FLegacyCacheGetRequest> GetDataRequests;
-	GetDataRequests.Reserve(State->VerifyRequests.Num());
-	{
-		uint64 PutIndex = 0;
-		const ECachePolicy GetPolicy = ECachePolicy::Query;
-		for (const FLegacyCachePutRequest& PutRequest : State->VerifyRequests)
-		{
-			GetDataRequests.Add({PutRequest.Name, PutRequest.Key, GetPolicy, PutIndex++});
-		}
-	}
-
-	State->OnComplete = MoveTemp(OnComplete);
-	State->ActiveRequests = GetDataRequests.Num();
-	InnerCache->LegacyGet(GetDataRequests, Owner, [this, &Owner, State = State.Release()](FLegacyCacheGetResponse&& MetaResponse)
-	{
-		GetDataComplete(Owner, State, MoveTemp(MetaResponse));
-	});
-}
-
-void FCacheStoreVerify::GetDataComplete(IRequestOwner& Owner, FVerifyLegacyPutState* State, FLegacyCacheGetResponse&& Response)
-{
-	FLegacyCachePutRequest& Request = State->VerifyRequests[int32(Response.UserData)];
-
-	if (Response.Status == EStatus::Ok)
-	{
-		const FLegacyCacheValue& NewValue = Request.Value;
-		const FLegacyCacheValue& OldValue = Response.Value;
-		if (NewValue.GetRawHash() == OldValue.GetRawHash())
-		{
-			UE_LOG(LogDerivedDataCache, Verbose,
-				TEXT("Verify: Data in the cache matches newly generated data for %s from '%s'."),
-				*WriteToString<96>(Request.Key.GetKey()), *Request.Name);
-			State->OnComplete(Request.MakeResponse(EStatus::Ok));
-		}
-		else
-		{
-			LogChangedValue(Request.Name, Request.Key.GetKey(), FValueId::Null,
-				NewValue.GetRawHash(), OldValue.GetRawHash(),
-				NewValue.GetRawData(), OldValue.GetRawData());
-			if (bPutOnError)
-			{
-				// Ask to overwrite existing values to potentially eliminate the mismatch.
-				UE_LOG(LogDerivedDataCache, Display,
-					TEXT("Verify: Writing newly generated data to the cache for %s from '%s'."),
-					*WriteToString<96>(Request.Key.GetKey()), *Request.Name);
-				Request.Policy &= ~ECachePolicy::Query;
-				FWriteScopeLock Lock(State->Lock);
-				State->ForwardRequests.Add(MoveTemp(Request));
-			}
-			else
-			{
-				State->OnComplete(Request.MakeResponse(EStatus::Ok));
-			}
-		}
-	}
-	else
-	{
-		UE_LOG(LogDerivedDataCache, Display,
-			TEXT("Verify: Cache did not contain a value for %s from '%s'."),
-			*WriteToString<96>(Request.Key.GetKey()), *Request.Name);
-		FWriteScopeLock Lock(State->Lock);
-		State->ForwardRequests.Add(MoveTemp(Request));
-	}
-
-	GetComplete(Owner, State);
-}
-
-void FCacheStoreVerify::GetComplete(IRequestOwner& Owner, FVerifyLegacyPutState* State)
-{
-	if (FWriteScopeLock Lock(State->Lock); --State->ActiveRequests > 0)
-	{
-		return;
-	}
-	if (!State->ForwardRequests.IsEmpty())
-	{
-		InnerCache->LegacyPut(State->ForwardRequests, Owner, MoveTemp(State->OnComplete));
-	}
-	delete State;
-}
-
-void FCacheStoreVerify::LegacyGet(
-	const TConstArrayView<FLegacyCacheGetRequest> Requests,
-	IRequestOwner& Owner,
-	FOnLegacyCacheGetComplete&& OnComplete)
-{
-	TArray<FLegacyCacheGetRequest, TInlineAllocator<8>> ForwardRequests;
-	TArray<FLegacyCacheGetRequest, TInlineAllocator<8>> VerifyRequests;
-	ForwardRequests.Reserve(Requests.Num());
-	VerifyRequests.Reserve(Requests.Num());
-	{
-		FScopeLock Lock(&AlreadyTestedLock);
-		for (const FLegacyCacheGetRequest& Request : Requests)
-		{
-			const FCacheKey& Key = Request.Key.GetKey();
-			const bool bForward = !Filter.IsMatch(Key) || AlreadyTested.Contains(Key);
-			(bForward ? ForwardRequests : VerifyRequests).Add(Request);
-		}
-	}
-
-	CompleteWithStatus(VerifyRequests, OnComplete, EStatus::Error);
-
-	if (!ForwardRequests.IsEmpty())
-	{
-		InnerCache->LegacyGet(ForwardRequests, Owner, MoveTemp(OnComplete));
-	}
-}
-
 bool FCacheStoreVerify::CompareRecords(const FCacheRecord& PutRecord, const FCacheRecord& GetRecord, const FSharedString& Name)
 {
 	bool bEqual = true;
diff --git a/Engine/Source/Developer/DerivedDataCache/Private/DerivedDataLegacyCacheStore.h b/Engine/Source/Developer/DerivedDataCache/Private/DerivedDataLegacyCacheStore.h
--- a/Engine/Source/Developer/DerivedDataCache/Private/DerivedDataLegacyCacheStore.h
+++ b/Engine/Source/Developer/DerivedDataCache/Private/DerivedDataLegacyCacheStore.h
@@ -33,17 +33,17 @@ using FOnLegacyCacheDeleteComplete = TUniqueFunction<void (FLegacyCacheDeleteRes
 class ILegacyCacheStore : public ICacheStore
 {
 public:
-	virtual void LegacyPut(
+	void LegacyPut(
 		TConstArrayView<FLegacyCachePutRequest> Requests,
 		IRequestOwner& Owner,
 		FOnLegacyCachePutComplete&& OnComplete);
 
-	virtual void LegacyGet(
+	void LegacyGet(
 		TConstArrayView<FLegacyCacheGetRequest> Requests,
 		IRequestOwner& Owner,
 		FOnLegacyCacheGetComplete&& OnComplete);
 
-	virtual void LegacyDelete(
+	void LegacyDelete(
 		TConstArrayView<FLegacyCacheDeleteRequest> Requests,
 		IRequestOwner& Owner,
 		FOnLegacyCacheDeleteComplete&& OnComplete);
diff --git a/Engine/Source/Developer/DerivedDataCache/Private/MemoryCacheStore.cpp b/Engine/Source/Developer/DerivedDataCache/Private/MemoryCacheStore.cpp
--- a/Engine/Source/Developer/DerivedDataCache/Private/MemoryCacheStore.cpp
+++ b/Engine/Source/Developer/DerivedDataCache/Private/MemoryCacheStore.cpp
@@ -55,19 +55,6 @@ public:
 
 	// ILegacyCacheStore Interface
 
-	void LegacyPut(
-		TConstArrayView<FLegacyCachePutRequest> Requests,
-		IRequestOwner& Owner,
-		FOnLegacyCachePutComplete&& OnComplete) final;
-	void LegacyGet(
-		TConstArrayView<FLegacyCacheGetRequest> Requests,
-		IRequestOwner& Owner,
-		FOnLegacyCacheGetComplete&& OnComplete) final;
-	void LegacyDelete(
-		TConstArrayView<FLegacyCacheDeleteRequest> Requests,
-		IRequestOwner& Owner,
-		FOnLegacyCacheDeleteComplete&& OnComplete) final;
-
 	void LegacyStats(FDerivedDataCacheStatsNode& OutNode) final;
 	bool LegacyDebugOptions(FBackendDebugOptions& Options) final;
 
@@ -75,7 +62,6 @@ public:
 
 	void Delete(const FCacheKey& Key) final;
 	void DeleteValue(const FCacheKey& Key) final;
-	void LegacyDelete(const FLegacyCacheKey& Key) final;
 
 	void Disable() final;
 
@@ -91,8 +77,6 @@ private:
 	/** Name of this cache (used for debugging) */
 	FString Name;
 
-	/** Set of legacy values in this cache. */
-	TMap<FCacheKey, FLegacyCacheValue> LegacyCacheValues;
 	/** Set of records in this cache. */
 	TMap<FCacheKey, FCacheRecordComponents> CacheRecords;
 	/** Set of values in this cache. */
@@ -156,7 +140,6 @@ void FMemoryCacheStore::Disable()
 	bDisabled = true;
 	CacheRecords.Empty();
 	CacheValues.Empty();
-	LegacyCacheValues.Empty();
 	CurrentCacheSize = 0;
 }
 
@@ -629,148 +612,6 @@ void FMemoryCacheStore::GetChunks(
 	}
 }
 
-void FMemoryCacheStore::LegacyPut(
-	const TConstArrayView<FLegacyCachePutRequest> Requests,
-	IRequestOwner& Owner,
-	FOnLegacyCachePutComplete&& OnComplete)
-{
-	if (bDisabled)
-	{
-		return CompleteWithStatus(Requests, OnComplete, EStatus::Error);
-	}
-
-	for (const FLegacyCachePutRequest& Request : Requests)
-	{
-		EStatus Status = EStatus::Error;
-		ON_SCOPE_EXIT
-		{
-			OnComplete(Request.MakeResponse(Status));
-		};
-
-		const FCacheKey& Key = Request.Key.GetKey();
-		const FLegacyCacheValue& Value = Request.Value;
-
-		if (!Value.HasData())
-		{
-			continue;
-		}
-
-		if (bCanBeDisabled && DebugOptions.ShouldSimulatePutMiss(Key))
-		{
-			UE_LOG(LogDerivedDataCache, Verbose, TEXT("%s: Simulated miss for put of %s from '%s'"),
-				*Name, *WriteToString<96>(Key), *Request.Name);
-			continue;
-		}
-
-		COOK_STAT(auto Timer = UsageStats.TimePut());
-		const int64 ValueSize = Value.GetRawSize();
-		const bool bReplaceExisting = !EnumHasAnyFlags(Request.Policy, ECachePolicy::QueryLocal);
-
-		FWriteScopeLock ScopeLock(SynchronizationObject);
-		FLegacyCacheValue* const ExistingValue = LegacyCacheValues.Find(Key);
-
-		if (ExistingValue && !bReplaceExisting)
-		{
-			Status = EStatus::Ok;
-			continue;
-		}
-
-		const int64 ExistingValueSize = ExistingValue ? ExistingValue->GetRawSize() : 0;
-		const int64 RequiredSize = ValueSize - ExistingValueSize;
-
-		if (MaxCacheSize > 0 && (CurrentCacheSize + RequiredSize) > MaxCacheSize)
-		{
-			UE_CLOG(!bMaxSizeExceeded, LogDerivedDataCache, Display,
-				TEXT("Failed to cache data. Maximum cache size reached. ")
-				TEXT("CurrentSize %" UINT64_FMT " KiB / MaxSize: %" UINT64_FMT " KiB"),
-				CurrentCacheSize / 1024, MaxCacheSize / 1024);
-			bMaxSizeExceeded = true;
-			continue;
-		}
-
-		CurrentCacheSize += RequiredSize;
-		if (ExistingValue)
-		{
-			*ExistingValue = Value;
-		}
-		else
-		{
-			LegacyCacheValues.Add(Key, Value);
-		}
-
-		COOK_STAT(Timer.AddHit(ValueSize));
-		Status = EStatus::Ok;
-	}
-}
-
-void FMemoryCacheStore::LegacyGet(
-	const TConstArrayView<FLegacyCacheGetRequest> Requests,
-	IRequestOwner& Owner,
-	FOnLegacyCacheGetComplete&& OnComplete)
-{
-	if (bDisabled)
-	{
-		return CompleteWithStatus(Requests, OnComplete, EStatus::Error);
-	}
-
-	for (const FLegacyCacheGetRequest& Request : Requests)
-	{
-		const FCacheKey& Key = Request.Key.GetKey();
-		const bool bExistsOnly = EnumHasAllFlags(Request.Policy, ECachePolicy::SkipData);
-		COOK_STAT(auto Timer = bExistsOnly ? UsageStats.TimeProbablyExists() : UsageStats.TimeGet());
-
-		FLegacyCacheValue Value;
-		EStatus Status = EStatus::Error;
-		if (bCanBeDisabled && DebugOptions.ShouldSimulateGetMiss(Key))
-		{
-			UE_LOG(LogDerivedDataCache, Verbose, TEXT("%s: Simulated miss for get of %s from '%s'"),
-				*Name, *WriteToString<96>(Key), *Request.Name);
-		}
-		else if (bExistsOnly && bCanBeDisabled)
-		{
-		}
-		else if (FReadScopeLock ScopeLock(SynchronizationObject); const FLegacyCacheValue* CacheValue = LegacyCacheValues.Find(Key))
-		{
-			Status = EStatus::Ok;
-			if (!bExistsOnly)
-			{
-				Value = *CacheValue;
-			}
-			COOK_STAT(Timer.AddHit(Value.GetRawSize()));
-		}
-
-		OnComplete({Request.Name, Request.Key, MoveTemp(Value), Request.UserData, Status});
-	}
-}
-
-void FMemoryCacheStore::LegacyDelete(const FLegacyCacheKey& Key)
-{
-	FWriteScopeLock ScopeLock(SynchronizationObject);
-	FLegacyCacheValue Value;
-	if (LegacyCacheValues.RemoveAndCopyValue(Key.GetKey(), Value))
-	{
-		CurrentCacheSize -= Value.GetRawSize();
-		bMaxSizeExceeded = false;
-	}
-}
-
-void FMemoryCacheStore::LegacyDelete(
-	const TConstArrayView<FLegacyCacheDeleteRequest> Requests,
-	IRequestOwner& Owner,
-	FOnLegacyCacheDeleteComplete&& OnComplete)
-{
-	if (bDisabled)
-	{
-		return CompleteWithStatus(Requests, OnComplete, EStatus::Error);
-	}
-
-	for (const FLegacyCacheDeleteRequest& Request : Requests)
-	{
-		LegacyDelete(Request.Key);
-	}
-	CompleteWithStatus(Requests, OnComplete, EStatus::Ok);
-}
-
 IMemoryCacheStore* CreateMemoryCacheStore(const TCHAR* Name, int64 MaxCacheSize, bool bCanBeDisabled)
 {
 	return new FMemoryCacheStore(Name, MaxCacheSize, bCanBeDisabled);
diff --git a/Engine/Source/Developer/DerivedDataCache/Private/MemoryCacheStore.h b/Engine/Source/Developer/DerivedDataCache/Private/MemoryCacheStore.h
--- a/Engine/Source/Developer/DerivedDataCache/Private/MemoryCacheStore.h
+++ b/Engine/Source/Developer/DerivedDataCache/Private/MemoryCacheStore.h
@@ -10,11 +10,8 @@ namespace UE::DerivedData
 class IMemoryCacheStore : public ILegacyCacheStore
 {
 public:
-	using ILegacyCacheStore::LegacyDelete;
-
 	virtual void Delete(const FCacheKey& Key) = 0;
 	virtual void DeleteValue(const FCacheKey& Key) = 0;
-	virtual void LegacyDelete(const FLegacyCacheKey& Key) = 0;
 
 	virtual void Disable() = 0;
 };
