DDC: Added detail to the error when cache graph creation fails

#rnx

diff --git a/Engine/Source/Developer/DerivedDataCache/Private/DerivedDataBackends.cpp b/Engine/Source/Developer/DerivedDataCache/Private/DerivedDataBackends.cpp
--- a/Engine/Source/Developer/DerivedDataCache/Private/DerivedDataBackends.cpp
+++ b/Engine/Source/Developer/DerivedDataCache/Private/DerivedDataBackends.cpp
@@ -116,18 +116,18 @@ public:
 		}
 		else
 		{
+			const TCHAR* const RootName = TEXT("Root");
+
 			if (!GraphName.IsEmpty() && GraphName != TEXT("Default"))
 			{
-				RootNode = ParseNode(TEXT("Root"), GEngineIni, *GraphName, ParsedNodes);
-
+				RootNode = ParseNode(RootName, GEngineIni, *GraphName, ParsedNodes);
 				if (!RootNode.Key)
 				{
 					// Destroy any cache stores that have been created.
 					ParsedNodes.Empty();
 					DestroyCreatedBackends();
 					UE_LOG(LogDerivedDataCache, Warning,
-						TEXT("Unable to create cache graph using the requested graph settings (%s). "
-						"Reverting to the default graph."), *GraphName);
+						TEXT("Unable to create cache graph '%s'. Reverting to the default graph."), *GraphName);
 				}
 			}
 
@@ -135,20 +135,28 @@ public:
 			{
 				// Try to use the default graph.
 				GraphName = FApp::IsEngineInstalled() ? TEXT("InstalledDerivedDataBackendGraph") : TEXT("DerivedDataBackendGraph");
-				FString Entry;
-				if (!GConfig->GetString(*GraphName, TEXT("Root"), Entry, GEngineIni) || !Entry.Len())
-				{
-					UE_LOG(LogDerivedDataCache, Fatal,
-						TEXT("Unable to create cache graph using the default graph settings (%s) ini=%s."),
-						*GraphName, *GEngineIni);
-				}
-				RootNode = ParseNode(TEXT("Root"), GEngineIni, *GraphName, ParsedNodes);
+				RootNode = ParseNode(RootName, GEngineIni, *GraphName, ParsedNodes);
 				if (!RootNode.Key)
 				{
-					UE_LOG(LogDerivedDataCache, Fatal,
-						TEXT("Unable to create cache graph using the default graph settings (%s) ini=%s. ")
-						TEXT("At least one cache store in the graph must be available."),
-						*GraphName, *GEngineIni);
+					FString Entry;
+					if (!GConfig->DoesSectionExist(*GraphName, GEngineIni))
+					{
+						UE_LOG(LogDerivedDataCache, Fatal,
+							TEXT("Unable to create default cache graph '%s' because its config section is missing in the '%s' config."),
+							*GraphName, *GEngineIni);
+					}
+					else if (!GConfig->GetString(*GraphName, RootName, Entry, GEngineIni) || !Entry.Len())
+					{
+						UE_LOG(LogDerivedDataCache, Fatal,
+							TEXT("Unable to create default cache graph '%s' because the root node '%s' is missing."),
+							*GraphName, RootName);
+					}
+					else
+					{
+						UE_LOG(LogDerivedDataCache, Fatal,
+							TEXT("Unable to create default cache graph '%s' because no cache store was available."),
+							*GraphName);
+					}
 				}
 			}
 		}
