DDC: Removed the legacy cache backend interface

#rnx

diff --git a/Engine/Source/Developer/DerivedDataCache/Private/DerivedDataBackendInterface.h b/Engine/Source/Developer/DerivedDataCache/Private/DerivedDataBackendInterface.h
--- a/Engine/Source/Developer/DerivedDataCache/Private/DerivedDataBackendInterface.h
+++ b/Engine/Source/Developer/DerivedDataCache/Private/DerivedDataBackendInterface.h
@@ -85,151 +85,6 @@ struct FBackendDebugOptions
 	bool ShouldSimulateGetMiss(const TCHAR* LegacyKey);
 };
 
-/**
- * Interface for cache server backends.
- * The entire API should be callable from any thread (except the singleton can be assumed to be called at least once before concurrent access).
- */
-class FDerivedDataBackendInterface : public ILegacyCacheStore
-{
-public:
-	using ESpeedClass = UE::DerivedData::EBackendSpeedClass;
-	using FBackendDebugOptions = UE::DerivedData::FBackendDebugOptions;
-
-	/** Status of a put operation. */
-	enum class EPutStatus
-	{
-		/** The put is executing asynchronously. */
-		Executing,
-		/** The put completed synchronously and the data was not cached. */
-		NotCached,
-		/** The put completed synchronously and the data was cached. */
-		Cached,
-		/** The put was skipped and should not be retried. */
-		Skipped,
-	};
-
-	virtual ~FDerivedDataBackendInterface()
-	{
-	}
-
-	/** Return a name for this interface */
-	virtual FString GetName() const = 0;
-
-	/** return true if this cache is writable */
-	virtual bool IsWritable() const = 0;
-
-	/**
-	 * Returns true if hits on this cache should propagate to lower cache level. Typically false for a PAK file.
-	 * Caution! This generally isn't propagated, so the thing that returns false must be a direct child of the hierarchical cache.
-	 */
-	virtual bool BackfillLowerCacheLevels() const
-	{
-		return true;
-	}
-
-	/**
-	 * Returns a class of speed for this interface
-	 */
-	virtual ESpeedClass GetSpeedClass() const = 0;
-
-	/**
-	 * Synchronous test for the existence of a cache item
-	 *
-	 * @param	CacheKey	Alphanumeric+underscore key of this cache item
-	 * @return				true if the data probably will be found, this can't be guaranteed because of concurrency in the backends, corruption, etc
-	 */
-	virtual bool CachedDataProbablyExists(const TCHAR* CacheKey)=0;
-
-	/**
-	 * Synchronous test for the existence of multiple cache items
-	 *
-	 * @param	CacheKeys	Alphanumeric+underscore key of the cache items
-	 * @return				A bit array with bits indicating whether the data for the corresponding key will probably be found
-	 */
-	virtual TBitArray<> CachedDataProbablyExistsBatch(TConstArrayView<FString> CacheKeys)
-	{
-		TBitArray<> Result;
-		Result.Reserve(CacheKeys.Num());
-		for (const FString& Key : CacheKeys)
-		{
-			Result.Add(CachedDataProbablyExists(*Key));
-		}
-		return Result;
-	}
-
-	/**
-	 * Synchronous retrieve of a cache item
-	 *
-	 * @param	CacheKey	Alphanumeric+underscore key of this cache item
-	 * @param	OutData		Buffer to receive the results, if any were found
-	 * @return				true if any data was found, and in this case OutData is non-empty
-	 */
-	virtual bool GetCachedData(const TCHAR* CacheKey, TArray<uint8>& OutData) = 0;
-
-	/**
-	 * Asynchronous, fire-and-forget placement of a cache item
-	 *
-	 * @param	CacheKey			Alphanumeric+underscore key of this cache item
-	 * @param	InData				Buffer containing the data to cache, can be destroyed after the call returns, immediately
-	 * @param	bPutEvenIfExists	If true, then do not attempt skip the put even if CachedDataProbablyExists returns true
-	 */
-	virtual EPutStatus PutCachedData(const TCHAR* CacheKey, TArrayView<const uint8> InData, bool bPutEvenIfExists) = 0;
-
-	/**
-	 * Remove data from cache (used in the event that corruption is detected at a higher level and possibly house keeping)
-	 *
-	 * @param	CacheKey	Alphanumeric+underscore key of this cache item
-	 * @param	bTransient	true if the data is transient and it is up to the backend to decide when and if to remove cached data.
-	 */
-	virtual void RemoveCachedData(const TCHAR* CacheKey, bool bTransient)=0;
-
-	/**
-	 * Retrieve usage stats for this backend. If the backend holds inner backends, this is expected to be passed down recursively.
-	 */
-	virtual TSharedRef<FDerivedDataCacheStatsNode> GatherUsageStats() const = 0;
-
-	/**
-	 * Synchronous attempt to make sure the cached data will be available as optimally as possible.
-	 *
-	 * @param	CacheKeys	Alphanumeric+underscore keys of the cache items
-	 * @return				A bit array with bits indicating whether the data for the corresponding key will probably be found in a fast backend on a future request.
-	 */
-	virtual TBitArray<> TryToPrefetch(TConstArrayView<FString> CacheKeys) = 0;
-
-	/**
-	 * Allows the DDC backend to determine if it wants to cache the provided data. Reasons for returning false could be a slow connection,
-	 * a file size limit, etc.
-	 */
-	virtual bool WouldCache(const TCHAR* CacheKey, TArrayView<const uint8> InData) = 0;
-
-	/**
-	 * Ask a backend to apply debug behavior to simulate different conditions. Backends that don't support these options should return
-	 * false which will result in a warning if an attempt is made to apply these options.
-	 */
-	virtual bool ApplyDebugOptions(FBackendDebugOptions& InOptions) = 0;
-
-	virtual EBackendLegacyMode GetLegacyMode() const = 0;
-
-	virtual void LegacyPut(
-		TConstArrayView<FLegacyCachePutRequest> Requests,
-		IRequestOwner& Owner,
-		FOnLegacyCachePutComplete&& OnComplete) final;
-
-	virtual void LegacyGet(
-		TConstArrayView<FLegacyCacheGetRequest> Requests,
-		IRequestOwner& Owner,
-		FOnLegacyCacheGetComplete&& OnComplete) final;
-
-	virtual void LegacyDelete(
-		TConstArrayView<FLegacyCacheDeleteRequest> Requests,
-		IRequestOwner& Owner,
-		FOnLegacyCacheDeleteComplete&& OnComplete) final;
-
-	virtual void LegacyStats(FDerivedDataCacheStatsNode& OutNode) final;
-
-	virtual bool LegacyDebugOptions(FBackendDebugOptions& Options) final;
-};
-
 class FDerivedDataBackend
 {
 public:
@@ -324,19 +179,19 @@ public:
 	return false;
 }
 
-inline const TCHAR* LexToString(FDerivedDataBackendInterface::ESpeedClass SpeedClass)
+inline const TCHAR* LexToString(EBackendSpeedClass SpeedClass)
 {
 	switch (SpeedClass)
 	{
-	case FDerivedDataBackendInterface::ESpeedClass::Unknown:
+	case EBackendSpeedClass::Unknown:
 		return TEXT("Unknown");
-	case FDerivedDataBackendInterface::ESpeedClass::Slow:
+	case EBackendSpeedClass::Slow:
 		return TEXT("Slow");
-	case FDerivedDataBackendInterface::ESpeedClass::Ok:
+	case EBackendSpeedClass::Ok:
 		return TEXT("Ok");
-	case FDerivedDataBackendInterface::ESpeedClass::Fast:
+	case EBackendSpeedClass::Fast:
 		return TEXT("Fast");
-	case FDerivedDataBackendInterface::ESpeedClass::Local:
+	case EBackendSpeedClass::Local:
 		return TEXT("Local");
 	}
 
@@ -344,25 +199,25 @@ inline const TCHAR* LexToString(FDerivedDataBackendInterface::ESpeedClass SpeedC
 	return TEXT("Unknown value! (Update LexToString!)");
 }
 
-inline void LexFromString(FDerivedDataBackendInterface::ESpeedClass& OutValue, const TCHAR* Buffer)
+inline void LexFromString(EBackendSpeedClass& OutValue, const TCHAR* Buffer)
 {
-	OutValue = FDerivedDataBackendInterface::ESpeedClass::Unknown;
+	OutValue = EBackendSpeedClass::Unknown;
 
 	if (FCString::Stricmp(Buffer, TEXT("Slow")) == 0)
 	{
-		OutValue = FDerivedDataBackendInterface::ESpeedClass::Slow;
+		OutValue = EBackendSpeedClass::Slow;
 	}
 	else if (FCString::Stricmp(Buffer, TEXT("Ok")) == 0)
 	{
-		OutValue = FDerivedDataBackendInterface::ESpeedClass::Ok;
+		OutValue = EBackendSpeedClass::Ok;
 	}
 	else if (FCString::Stricmp(Buffer, TEXT("Fast")) == 0)
 	{
-		OutValue = FDerivedDataBackendInterface::ESpeedClass::Fast;
+		OutValue = EBackendSpeedClass::Fast;
 	}
 	else if (FCString::Stricmp(Buffer, TEXT("Local")) == 0)
 	{
-		OutValue = FDerivedDataBackendInterface::ESpeedClass::Local;
+		OutValue = EBackendSpeedClass::Local;
 	}
 }
 
