Some missing TEXT macros added.

diff --git a/Engine/Source/Developer/DerivedDataCache/Public/DerivedDataCacheInterface.h b/Engine/Source/Developer/DerivedDataCache/Public/DerivedDataCacheInterface.h
--- a/Engine/Source/Developer/DerivedDataCache/Public/DerivedDataCacheInterface.h
+++ b/Engine/Source/Developer/DerivedDataCache/Public/DerivedDataCacheInterface.h
@@ -127,9 +127,9 @@ public:
 		FString Out;
 		Out.Reserve(PluginName.Len() + 1 + VersionString.Len() + 1 + PluginSpecificCacheKeySuffix.Len());
 		AppendSanitized(Out, PluginName);
-		Out.AppendChar('_');
+		Out.AppendChar(TEXT('_'));
 		AppendSanitized(Out, VersionString);
-		Out.AppendChar('_');
+		Out.AppendChar(TEXT('_'));
 		AppendSanitized(Out, PluginSpecificCacheKeySuffix);
 		return Out;
 	}
