DDC: Added TryLoadFromCompactBinary for FCacheKey

#rnx

diff --git a/Engine/Source/Developer/DerivedDataCache/Private/DerivedDataCacheKey.cpp b/Engine/Source/Developer/DerivedDataCache/Private/DerivedDataCacheKey.cpp
--- a/Engine/Source/Developer/DerivedDataCache/Private/DerivedDataCacheKey.cpp
+++ b/Engine/Source/Developer/DerivedDataCache/Private/DerivedDataCacheKey.cpp
@@ -9,6 +9,7 @@
 #include "Math/UnrealMathUtility.h"
 #include "Misc/ScopeRWLock.h"
 #include "Misc/StringBuilder.h"
+#include "Serialization/CompactBinary.h"
 #include "Serialization/CompactBinaryWriter.h"
 #include "String/Find.h"
 
@@ -124,6 +125,18 @@ FCacheBucket::FCacheBucket(FWideStringView InName)
 {
 }
 
+bool TryLoadFromCompactBinary(const FCbObjectView Object, FCacheKey& OutKey)
+{
+	const FUtf8StringView Bucket = Object[ANSITEXTVIEW("Bucket")].AsString();
+	if (!FCacheBucket::IsValidName(Bucket))
+	{
+		return false;
+	}
+	OutKey.Bucket = FCacheBucket(Bucket);
+	OutKey.Hash = Object[ANSITEXTVIEW("Hash")].AsHash();
+	return true;
+}
+
 FCbWriter& operator<<(FCbWriter& Writer, const FCacheKey& Key)
 {
 	Writer.BeginObject();
diff --git a/Engine/Source/Developer/DerivedDataCache/Private/DerivedDataCacheRecord.cpp b/Engine/Source/Developer/DerivedDataCache/Private/DerivedDataCacheRecord.cpp
--- a/Engine/Source/Developer/DerivedDataCache/Private/DerivedDataCacheRecord.cpp
+++ b/Engine/Source/Developer/DerivedDataCache/Private/DerivedDataCacheRecord.cpp
@@ -207,21 +207,10 @@ FOptionalCacheRecord FCacheRecord::Load(const FCbPackage& Attachments, const FCb
 	}
 
 	FCacheKey Key;
-	FCbObjectView KeyObject = ObjectView[ANSITEXTVIEW("Key")].AsObjectView();
-	auto TrySetBucketName = [](FUtf8StringView Name, FCacheKey& Key)
-	{
-		if (FCacheBucket::IsValidName(Name))
-		{
-			Key.Bucket = FCacheBucket(Name);
-			return true;
-		}
-		return false;
-	};
-	if (!TrySetBucketName(KeyObject[ANSITEXTVIEW("Bucket")].AsString(), Key))
+	if (!TryLoadFromCompactBinary(ObjectView[ANSITEXTVIEW("Key")].AsObjectView(), Key))
 	{
 		return FOptionalCacheRecord();
 	}
-	Key.Hash = KeyObject[ANSITEXTVIEW("Hash")].AsHash();
 
 	FCacheRecordBuilder Builder(Key);
 
diff --git a/Engine/Source/Developer/DerivedDataCache/Public/DerivedDataCacheKey.h b/Engine/Source/Developer/DerivedDataCache/Public/DerivedDataCacheKey.h
--- a/Engine/Source/Developer/DerivedDataCache/Public/DerivedDataCacheKey.h
+++ b/Engine/Source/Developer/DerivedDataCache/Public/DerivedDataCacheKey.h
@@ -11,6 +11,7 @@
 
 #define UE_API DERIVEDDATACACHE_API
 
+class FCbObjectView;
 class FCbWriter;
 
 namespace UE::DerivedData
@@ -135,6 +136,7 @@ inline TStringBuilderBase<CharType>& operator<<(TStringBuilderBase<CharType>& Bu
 	return Builder << Key.Bucket << CharType('/') << Key.Hash;
 }
 
+UE_API bool TryLoadFromCompactBinary(FCbObjectView Object, FCacheKey& OutKey);
 UE_API FCbWriter& operator<<(FCbWriter& Writer, const FCacheKey& Key);
 
 ///////////////////////////////////////////////////////////////////////////////////////////////////////////////////////
