Select maximum HTTP connections for ZenCacheStore based on logical core count (clamped between 8 and 64).  Speculative solution for farm cooks using Zen experiencing "couldn't connect errors".

#rnx

diff --git a/Engine/Source/Developer/DerivedDataCache/Private/ZenCacheStore.cpp b/Engine/Source/Developer/DerivedDataCache/Private/ZenCacheStore.cpp
--- a/Engine/Source/Developer/DerivedDataCache/Private/ZenCacheStore.cpp
+++ b/Engine/Source/Developer/DerivedDataCache/Private/ZenCacheStore.cpp
@@ -13,6 +13,7 @@
 #include "DerivedDataChunk.h"
 #include "DerivedDataRequestOwner.h"
 #include "Http/HttpClient.h"
+#include "Math/UnrealMathUtility.h"
 #include "Misc/App.h"
 #include "Misc/ConfigCacheIni.h"
 #include "Misc/Optional.h"
@@ -913,7 +914,7 @@ FZenCacheStore::FZenCacheStore(
 	, StructuredNamespace(InStructuredNamespace)
 	, ZenService(InServiceUrl)
 {
-	constexpr uint32 MaxConnections = 64;
+	const uint32 MaxConnections = FMath::Clamp(static_cast<uint32>(FPlatformMisc::NumberOfCoresIncludingHyperthreads()), 8, 64);
 	constexpr uint32 RequestPoolSize = 128;
 	constexpr uint32 RequestPoolOverflowSize = 128;
 	
