Ensure the connection pool size matches the total number of connections for the DDC HTTP implementation.

#rnx

diff --git a/Engine/Source/Developer/DerivedDataCache/Private/Http/HttpClient.cpp b/Engine/Source/Developer/DerivedDataCache/Private/Http/HttpClient.cpp
--- a/Engine/Source/Developer/DerivedDataCache/Private/Http/HttpClient.cpp
+++ b/Engine/Source/Developer/DerivedDataCache/Private/Http/HttpClient.cpp
@@ -106,6 +106,7 @@ FHttpSharedData::FHttpSharedData(uint32 OverrideMaxConnections)
 	curl_share_setopt(Internals->CurlShare, CURLSHOPT_SHARE, CURL_LOCK_DATA_SSL_SESSION);
 	Internals->CurlMulti = curl_multi_init();
 	curl_multi_setopt(Internals->CurlMulti, CURLMOPT_MAX_TOTAL_CONNECTIONS, OverrideMaxConnections ? OverrideMaxConnections : Http::Private::MaxTotalConnections);
+	curl_multi_setopt(Internals->CurlMulti, CURLMOPT_MAXCONNECTS, OverrideMaxConnections ? OverrideMaxConnections : Http::Private::MaxTotalConnections); // Keep the connection pool at exactly the number of total connections
 	curl_multi_setopt(Internals->CurlMulti, CURLMOPT_PIPELINING, CURLPIPE_MULTIPLEX);
 }
 
