DDC: Added functions to serialize cache types to and from compact binary, and strings where appropriate

diff --git a/Engine/Source/Developer/DerivedDataCache/Private/DerivedDataCache.cpp b/Engine/Source/Developer/DerivedDataCache/Private/DerivedDataCache.cpp
--- a/Engine/Source/Developer/DerivedDataCache/Private/DerivedDataCache.cpp
+++ b/Engine/Source/Developer/DerivedDataCache/Private/DerivedDataCache.cpp
@@ -23,6 +23,9 @@
 #include "Misc/CoreMisc.h"
 #include "Misc/ScopeLock.h"
 #include "ProfilingDebugging/CookStats.h"
+#include "Serialization/CompactBinary.h"
+#include "Serialization/CompactBinarySerialization.h"
+#include "Serialization/CompactBinaryWriter.h"
 #include "Stats/Stats.h"
 #include "Stats/StatsMisc.h"
 #include "ZenServerInterface.h"
@@ -204,6 +207,116 @@ FCacheGetChunkResponse FCacheGetChunkRequest::MakeResponse(const EStatus Status)
 	return {Name, Key, Id, RawOffset, 0, {}, {}, UserData, Status};
 }
 
+FCbWriter& operator<<(FCbWriter& Writer, const FCacheGetRequest& Request)
+{
+	Writer.BeginObject();
+	if (!Request.Name.IsEmpty())
+	{
+		Writer << ANSITEXTVIEW("Name") << Request.Name;
+	}
+	Writer << ANSITEXTVIEW("Key") << Request.Key;
+	if (!Request.Policy.IsDefault())
+	{
+		Writer << ANSITEXTVIEW("Policy") << Request.Policy;
+	}
+	if (Request.UserData != 0)
+	{
+		Writer << ANSITEXTVIEW("UserData") << Request.UserData;
+	}
+	Writer.EndObject();
+	return Writer;
+}
+
+bool LoadFromCompactBinary(FCbFieldView Field, FCacheGetRequest& Request)
+{
+	bool bOk = Field.IsObject();
+	LoadFromCompactBinary(Field[ANSITEXTVIEW("Name")], Request.Name);
+	bOk &= LoadFromCompactBinary(Field[ANSITEXTVIEW("Key")], Request.Key);
+	LoadFromCompactBinary(Field[ANSITEXTVIEW("Policy")], Request.Policy);
+	LoadFromCompactBinary(Field[ANSITEXTVIEW("UserData")], Request.UserData);
+	return bOk;
+}
+
+FCbWriter& operator<<(FCbWriter& Writer, const FCacheGetValueRequest& Request)
+{
+	Writer.BeginObject();
+	if (!Request.Name.IsEmpty())
+	{
+		Writer << ANSITEXTVIEW("Name") << MakeStringView(Request.Name);
+	}
+	Writer << ANSITEXTVIEW("Key") << Request.Key;
+	if (Request.Policy != ECachePolicy::Default)
+	{
+		Writer << ANSITEXTVIEW("Policy") << Request.Policy;
+	}
+	if (Request.UserData != 0)
+	{
+		Writer << ANSITEXTVIEW("UserData") << Request.UserData;
+	}
+	Writer.EndObject();
+	return Writer;
+}
+
+bool LoadFromCompactBinary(FCbFieldView Field, FCacheGetValueRequest& Request)
+{
+	bool bOk = Field.IsObject();
+	LoadFromCompactBinary(Field[ANSITEXTVIEW("Name")], Request.Name);
+	bOk &= LoadFromCompactBinary(Field[ANSITEXTVIEW("Key")], Request.Key);
+	LoadFromCompactBinary(Field[ANSITEXTVIEW("Policy")], Request.Policy);
+	LoadFromCompactBinary(Field[ANSITEXTVIEW("UserData")], Request.UserData);
+	return bOk;
+}
+
+FCbWriter& operator<<(FCbWriter& Writer, const FCacheGetChunkRequest& Request)
+{
+	Writer.BeginObject();
+	if (!Request.Name.IsEmpty())
+	{
+		Writer << ANSITEXTVIEW("Name") << MakeStringView(Request.Name);
+	}
+	Writer << ANSITEXTVIEW("Key") << Request.Key;
+	if (Request.Id.IsValid())
+	{
+		Writer << ANSITEXTVIEW("Id") << Request.Id;
+	}
+	if (Request.RawOffset != 0)
+	{
+		Writer << ANSITEXTVIEW("RawOffset") << Request.RawOffset;
+	}
+	if (Request.RawSize != MAX_uint64)
+	{
+		Writer << ANSITEXTVIEW("RawSize") << Request.RawSize;
+	}
+	if (!Request.RawHash.IsZero())
+	{
+		Writer << ANSITEXTVIEW("RawHash") << Request.RawHash;
+	}
+	if (Request.Policy != ECachePolicy::Default)
+	{
+		Writer << ANSITEXTVIEW("Policy") << Request.Policy;
+	}
+	if (Request.UserData)
+	{
+		Writer << ANSITEXTVIEW("UserData") << Request.UserData;
+	}
+	Writer.EndObject();
+	return Writer;
+}
+
+bool LoadFromCompactBinary(FCbFieldView Field, FCacheGetChunkRequest& OutRequest)
+{
+	bool bOk = Field.IsObject();
+	LoadFromCompactBinary(Field[ANSITEXTVIEW("Name")], OutRequest.Name);
+	bOk &= LoadFromCompactBinary(Field[ANSITEXTVIEW("Key")], OutRequest.Key);
+	LoadFromCompactBinary(Field[ANSITEXTVIEW("Id")], OutRequest.Id);
+	LoadFromCompactBinary(Field[ANSITEXTVIEW("RawOffset")], OutRequest.RawOffset, 0);
+	LoadFromCompactBinary(Field[ANSITEXTVIEW("RawSize")], OutRequest.RawSize, MAX_uint64);
+	LoadFromCompactBinary(Field[ANSITEXTVIEW("RawHash")], OutRequest.RawHash);
+	LoadFromCompactBinary(Field[ANSITEXTVIEW("Policy")], OutRequest.Policy);
+	LoadFromCompactBinary(Field[ANSITEXTVIEW("UserData")], OutRequest.UserData);
+	return bOk;
+}
+
 } // UE::DerivedData
 
 namespace UE::DerivedData::Private
diff --git a/Engine/Source/Developer/DerivedDataCache/Private/DerivedDataCacheKey.cpp b/Engine/Source/Developer/DerivedDataCache/Private/DerivedDataCacheKey.cpp
--- a/Engine/Source/Developer/DerivedDataCache/Private/DerivedDataCacheKey.cpp
+++ b/Engine/Source/Developer/DerivedDataCache/Private/DerivedDataCacheKey.cpp
@@ -10,6 +10,7 @@
 #include "Misc/ScopeRWLock.h"
 #include "Misc/StringBuilder.h"
 #include "Serialization/CompactBinary.h"
+#include "Serialization/CompactBinarySerialization.h"
 #include "Serialization/CompactBinaryWriter.h"
 #include "String/Find.h"
 
@@ -125,27 +126,40 @@ FCacheBucket::FCacheBucket(FWideStringView InName)
 {
 }
 
-bool TryLoadFromCompactBinary(const FCbObjectView Object, FCacheKey& OutKey)
+FCbWriter& operator<<(FCbWriter& Writer, const FCacheBucket Bucket)
 {
-	const FUtf8StringView Bucket = Object[ANSITEXTVIEW("Bucket")].AsString();
-	if (!FCacheBucket::IsValidName(Bucket))
+	Writer.AddString(Bucket.ToString());
+	return Writer;
+}
+
+bool LoadFromCompactBinary(FCbFieldView Field, FCacheBucket& OutBucket)
+{
+	if (const FUtf8StringView Bucket = Field.AsString(); !Field.HasError() && FCacheBucket::IsValidName(Bucket))
 	{
-		return false;
+		OutBucket = FCacheBucket(Bucket);
+		return true;
 	}
-	OutKey.Bucket = FCacheBucket(Bucket);
-	OutKey.Hash = Object[ANSITEXTVIEW("Hash")].AsHash();
-	return true;
+	OutBucket.Reset();
+	return false;
 }
 
 FCbWriter& operator<<(FCbWriter& Writer, const FCacheKey& Key)
 {
 	Writer.BeginObject();
-	Writer.AddString(ANSITEXTVIEW("Bucket"), Key.Bucket.ToString());
-	Writer.AddHash(ANSITEXTVIEW("Hash"), Key.Hash);
+	Writer << ANSITEXTVIEW("Bucket") << Key.Bucket;
+	Writer << ANSITEXTVIEW("Hash") << Key.Hash;
 	Writer.EndObject();
 	return Writer;
 }
 
+bool LoadFromCompactBinary(const FCbFieldView Field, FCacheKey& OutKey)
+{
+	bool bOk = Field.IsObject();
+	bOk &= LoadFromCompactBinary(Field[ANSITEXTVIEW("Bucket")], OutKey.Bucket);
+	bOk &= LoadFromCompactBinary(Field[ANSITEXTVIEW("Hash")], OutKey.Hash);
+	return bOk;
+}
+
 FCacheKey ConvertLegacyCacheKey(const FStringView Key)
 {
 	FTCHARToUTF8 Utf8Key(Key);
diff --git a/Engine/Source/Developer/DerivedDataCache/Private/DerivedDataCachePolicy.cpp b/Engine/Source/Developer/DerivedDataCache/Private/DerivedDataCachePolicy.cpp
--- a/Engine/Source/Developer/DerivedDataCache/Private/DerivedDataCachePolicy.cpp
+++ b/Engine/Source/Developer/DerivedDataCache/Private/DerivedDataCachePolicy.cpp
@@ -17,7 +17,7 @@ namespace UE::DerivedData::Private { class FCacheRecordPolicyShared; }
 namespace UE::DerivedData::Private
 {
 
-constexpr ANSICHAR CachePolicyDelimiter = ',';
+static constexpr ANSICHAR CachePolicyDelimiter = ',';
 
 struct FCachePolicyToText
 {
@@ -25,7 +25,7 @@ struct FCachePolicyToText
 	FAnsiStringView Text;
 };
 
-const FCachePolicyToText CachePolicyToText[]
+static const FCachePolicyToText CachePolicyToText[]
 {
 	// Flags with multiple bits are ordered by bit count to minimize token count in the text format.
 	{ ECachePolicy::Default,       ANSITEXTVIEW("Default") },
@@ -46,7 +46,7 @@ const FCachePolicyToText CachePolicyToText[]
 	{ ECachePolicy::None,          ANSITEXTVIEW("None") },
 };
 
-constexpr ECachePolicy CachePolicyKnownFlags
+static constexpr ECachePolicy CachePolicyKnownFlags
 	= ECachePolicy::Default
 	| ECachePolicy::SkipMeta
 	| ECachePolicy::SkipData
@@ -54,7 +54,7 @@ constexpr ECachePolicy CachePolicyKnownFlags
 	| ECachePolicy::KeepAlive;
 
 template <typename CharType>
-TStringBuilderBase<CharType>& CachePolicyToString(TStringBuilderBase<CharType>& Builder, ECachePolicy Policy)
+static TStringBuilderBase<CharType>& CachePolicyToString(TStringBuilderBase<CharType>& Builder, ECachePolicy Policy)
 {
 	// Mask out unknown flags. None will be written if no flags are known.
 	Policy &= CachePolicyKnownFlags;
@@ -75,19 +75,22 @@ TStringBuilderBase<CharType>& CachePolicyToString(TStringBuilderBase<CharType>&
 }
 
 template <typename CharType>
-ECachePolicy ParseCachePolicy(const TStringView<CharType> Text)
+static bool CachePolicyFromString(ECachePolicy& OutPolicy, const TStringView<CharType> String)
 {
-	checkf(!Text.IsEmpty(), TEXT("ParseCachePolicy requires a non-empty string."));
-	ECachePolicy Policy = ECachePolicy::None;
-	String::ParseTokens(StringCast<UTF8CHAR, 128>(Text.GetData(), Text.Len()), UTF8CHAR(CachePolicyDelimiter),
-		[&Policy, Index = int32(0)](FUtf8StringView Token) mutable
+	if (String.IsEmpty())
+	{
+		return false;
+	}
+
+	String::ParseTokens(StringCast<UTF8CHAR, 128>(String.GetData(), String.Len()), UTF8CHAR(CachePolicyDelimiter),
+		[&OutPolicy, Index = int32(0)](FUtf8StringView Token) mutable
 		{
 			const int32 EndIndex = Index;
 			for (; Index < UE_ARRAY_COUNT(CachePolicyToText); ++Index)
 			{
 				if (CachePolicyToText[Index].Text == Token)
 				{
-					Policy |= CachePolicyToText[Index].Policy;
+					OutPolicy |= CachePolicyToText[Index].Policy;
 					++Index;
 					return;
 				}
@@ -96,12 +99,21 @@ ECachePolicy ParseCachePolicy(const TStringView<CharType> Text)
 			{
 				if (CachePolicyToText[Index].Text == Token)
 				{
-					Policy |= CachePolicyToText[Index].Policy;
+					OutPolicy |= CachePolicyToText[Index].Policy;
 					++Index;
 					return;
 				}
 			}
 		});
+	return true;
+}
+
+template <typename CharType>
+static ECachePolicy ParseCachePolicy(const TStringView<CharType> Text)
+{
+	checkf(!Text.IsEmpty(), TEXT("ParseCachePolicy requires a non-empty string."));
+	ECachePolicy Policy = ECachePolicy::None;
+	CachePolicyFromString(Policy, Text);
 	return Policy;
 }
 
@@ -114,10 +126,46 @@ FAnsiStringBuilderBase& operator<<(FAnsiStringBuilderBase& Builder, ECachePolicy
 FWideStringBuilderBase& operator<<(FWideStringBuilderBase& Builder, ECachePolicy Policy) { return Private::CachePolicyToString(Builder, Policy); }
 FUtf8StringBuilderBase& operator<<(FUtf8StringBuilderBase& Builder, ECachePolicy Policy) { return Private::CachePolicyToString(Builder, Policy); }
 
+bool TryLexFromString(ECachePolicy& OutPolicy, const FUtf8StringView String) { return Private::CachePolicyFromString(OutPolicy, String); }
+bool TryLexFromString(ECachePolicy& OutPolicy, const FWideStringView String) { return Private::CachePolicyFromString(OutPolicy, String); }
+
 ECachePolicy ParseCachePolicy(FAnsiStringView Text) { return Private::ParseCachePolicy(Text); }
 ECachePolicy ParseCachePolicy(FWideStringView Text) { return Private::ParseCachePolicy(Text); }
 ECachePolicy ParseCachePolicy(FUtf8StringView Text) { return Private::ParseCachePolicy(Text); }
 
+FCbWriter& operator<<(FCbWriter& Writer, const ECachePolicy Policy)
+{
+	Writer.AddString(WriteToUtf8String<64>(Policy));
+	return Writer;
+}
+
+bool LoadFromCompactBinary(FCbFieldView Field, ECachePolicy& OutPolicy, const ECachePolicy Default)
+{
+	if (TryLexFromString(OutPolicy, Field.AsString()))
+	{
+		return true;
+	}
+	OutPolicy = Default;
+	return false;
+}
+
+FCbWriter& operator<<(FCbWriter& Writer, const FCacheValuePolicy& Policy)
+{
+	Writer.BeginObject();
+	Writer << ANSITEXTVIEW("Id") << Policy.Id;
+	Writer << ANSITEXTVIEW("Policy") << Policy.Policy;
+	Writer.EndObject();
+	return Writer;
+}
+
+bool LoadFromCompactBinary(FCbFieldView Field, FCacheValuePolicy& OutPolicy)
+{
+	bool bOk = Field.IsObject();
+	bOk &= LoadFromCompactBinary(Field[ANSITEXTVIEW("Id")], OutPolicy.Id);
+	bOk &= LoadFromCompactBinary(Field[ANSITEXTVIEW("Policy")], OutPolicy.Policy);
+	return bOk;
+}
+
 class Private::FCacheRecordPolicyShared final : public Private::ICacheRecordPolicyShared
 {
 public:
@@ -185,48 +233,17 @@ FCacheRecordPolicy FCacheRecordPolicy::Transform(const TFunctionRef<ECachePolicy
 
 void FCacheRecordPolicy::Save(FCbWriter& Writer) const
 {
-	Writer.BeginObject();
-	Writer.AddString(ANSITEXTVIEW("BasePolicy"), WriteToUtf8String<128>(GetBasePolicy()));
-	if (!IsUniform())
-	{
-		Writer.BeginArray(ANSITEXTVIEW("ValuePolicies"));
-		for (const FCacheValuePolicy& Value : GetValuePolicies())
-		{
-			Writer.BeginObject();
-			Writer.AddObjectId(ANSITEXTVIEW("Id"), Value.Id);
-			Writer.AddString(ANSITEXTVIEW("Policy"), WriteToUtf8String<128>(Value.Policy));
-			Writer.EndObject();
-		}
-		Writer.EndArray();
-	}
-	Writer.EndObject();
+	Writer << *this;
 }
 
 FOptionalCacheRecordPolicy FCacheRecordPolicy::Load(const FCbObjectView Object)
 {
-	const FUtf8StringView BasePolicyText = Object[ANSITEXTVIEW("BasePolicy")].AsString();
-	if (BasePolicyText.IsEmpty())
+	FCacheRecordPolicy Policy;
+	if (LoadFromCompactBinary(Object.AsFieldView(), Policy))
 	{
-		return {};
+		return Policy;
 	}
-
-	FCacheRecordPolicyBuilder Builder(ParseCachePolicy(BasePolicyText));
-	for (const FCbFieldView Value : Object[ANSITEXTVIEW("ValuePolicies")])
-	{
-		const FValueId Id = Value[ANSITEXTVIEW("Id")].AsObjectId();
-		const FUtf8StringView PolicyText = Value[ANSITEXTVIEW("Policy")].AsString();
-		if (Id.IsNull() || PolicyText.IsEmpty())
-		{
-			return {};
-		}
-		const ECachePolicy Policy = ParseCachePolicy(PolicyText);
-		if (EnumHasAnyFlags(Policy, ~FCacheValuePolicy::PolicyMask))
-		{
-			return {};
-		}
-		Builder.AddValuePolicy(Id, Policy);
-	}
-	return Builder.Build();
+	return {};
 }
 
 void FCacheRecordPolicyBuilder::AddValuePolicy(const FCacheValuePolicy& Value)
@@ -261,4 +278,48 @@ FCacheRecordPolicy FCacheRecordPolicyBuilder::Build()
 	return Policy;
 }
 
+FCbWriter& operator<<(FCbWriter& Writer, const FCacheRecordPolicy& Policy)
+{
+	Writer.BeginObject();
+	Writer << ANSITEXTVIEW("BasePolicy") << Policy.GetBasePolicy();
+	if (!Policy.IsUniform())
+	{
+		Writer.BeginArray(ANSITEXTVIEW("ValuePolicies"));
+		for (const FCacheValuePolicy& Value : Policy.GetValuePolicies())
+		{
+			Writer << Value;
+		}
+		Writer.EndArray();
+	}
+	Writer.EndObject();
+	return Writer;
+}
+
+bool LoadFromCompactBinary(FCbFieldView Field, FCacheRecordPolicy& OutPolicy)
+{
+	ECachePolicy BasePolicy;
+	if (!LoadFromCompactBinary(Field[ANSITEXTVIEW("BasePolicy")], BasePolicy))
+	{
+		OutPolicy = {};
+		return false;
+	}
+	FCacheRecordPolicyBuilder Builder(BasePolicy);
+
+	for (FCbFieldView Value : Field[ANSITEXTVIEW("ValuePolicies")])
+	{
+		FCacheValuePolicy ValuePolicy;
+		if (!LoadFromCompactBinary(Value, ValuePolicy) ||
+			ValuePolicy.Id.IsNull() ||
+			EnumHasAnyFlags(ValuePolicy.Policy, ~FCacheValuePolicy::PolicyMask))
+		{
+			OutPolicy = {};
+			return false;
+		}
+		Builder.AddValuePolicy(ValuePolicy);
+	}
+
+	OutPolicy = Builder.Build();
+	return true;
+}
+
 } // UE::DerivedData
diff --git a/Engine/Source/Developer/DerivedDataCache/Private/DerivedDataCacheRecord.cpp b/Engine/Source/Developer/DerivedDataCache/Private/DerivedDataCacheRecord.cpp
--- a/Engine/Source/Developer/DerivedDataCache/Private/DerivedDataCacheRecord.cpp
+++ b/Engine/Source/Developer/DerivedDataCache/Private/DerivedDataCacheRecord.cpp
@@ -207,7 +207,7 @@ FOptionalCacheRecord FCacheRecord::Load(const FCbPackage& Attachments, const FCb
 	}
 
 	FCacheKey Key;
-	if (!TryLoadFromCompactBinary(ObjectView[ANSITEXTVIEW("Key")].AsObjectView(), Key))
+	if (!LoadFromCompactBinary(ObjectView[ANSITEXTVIEW("Key")], Key))
 	{
 		return FOptionalCacheRecord();
 	}
diff --git a/Engine/Source/Developer/DerivedDataCache/Private/DerivedDataRequestTypes.cpp b/Engine/Source/Developer/DerivedDataCache/Private/DerivedDataRequestTypes.cpp
new file mode 100644
--- /dev/null
+++ b/Engine/Source/Developer/DerivedDataCache/Private/DerivedDataRequestTypes.cpp
@@ -0,0 +1,150 @@
+// Copyright Epic Games, Inc. All Rights Reserved.
+
+#include "DerivedDataRequestTypes.h"
+
+#include "Containers/StringConv.h"
+#include "Containers/StringView.h"
+#include "Misc/StringBuilder.h"
+#include "Serialization/CompactBinary.h"
+#include "Serialization/CompactBinaryWriter.h"
+
+namespace UE::DerivedData::Private
+{
+
+template <typename CharType>
+static TStringBuilderBase<CharType>& PriorityToString(TStringBuilderBase<CharType>& Builder, const EPriority Priority)
+{
+	switch (Priority)
+	{
+	case EPriority::Lowest:   return Builder << ANSITEXTVIEW("Lowest");
+	case EPriority::Low:      return Builder << ANSITEXTVIEW("Low");
+	case EPriority::Normal:   return Builder << ANSITEXTVIEW("Normal");
+	case EPriority::High:     return Builder << ANSITEXTVIEW("High");
+	case EPriority::Highest:  return Builder << ANSITEXTVIEW("Highest");
+	case EPriority::Blocking: return Builder << ANSITEXTVIEW("Blocking");
+	}
+	return Builder << ANSITEXTVIEW("Unknown");
+}
+
+template <typename CharType>
+static bool PriorityFromString(EPriority& OutPriority, const TStringView<CharType> String)
+{
+	const auto ConvertedString = StringCast<UTF8CHAR, 16>(String.GetData(), String.Len());
+	if (ConvertedString == UTF8TEXTVIEW("Lowest"))
+	{
+		OutPriority = EPriority::Lowest;
+	}
+	else if (ConvertedString == UTF8TEXTVIEW("Low"))
+	{
+		OutPriority = EPriority::Low;
+	}
+	else if (ConvertedString == UTF8TEXTVIEW("Normal"))
+	{
+		OutPriority = EPriority::Normal;
+	}
+	else if (ConvertedString == UTF8TEXTVIEW("High"))
+	{
+		OutPriority = EPriority::High;
+	}
+	else if (ConvertedString == UTF8TEXTVIEW("Highest"))
+	{
+		OutPriority = EPriority::Highest;
+	}
+	else if (ConvertedString == UTF8TEXTVIEW("Blocking"))
+	{
+		OutPriority = EPriority::Blocking;
+	}
+	else
+	{
+		return false;
+	}
+	return true;
+}
+
+template <typename CharType>
+static TStringBuilderBase<CharType>& StatusToString(TStringBuilderBase<CharType>& Builder, const EStatus Status)
+{
+	switch (Status)
+	{
+	case EStatus::Ok:       return Builder << ANSITEXTVIEW("Ok");
+	case EStatus::Error:    return Builder << ANSITEXTVIEW("Error");
+	case EStatus::Canceled: return Builder << ANSITEXTVIEW("Canceled");
+	}
+	return Builder << ANSITEXTVIEW("Unknown");
+}
+
+template <typename CharType>
+static bool StatusFromString(EStatus& OutStatus, const TStringView<CharType> String)
+{
+	const auto ConvertedString = StringCast<UTF8CHAR, 16>(String.GetData(), String.Len());
+	if (ConvertedString == UTF8TEXTVIEW("Ok"))
+	{
+		OutStatus = EStatus::Ok;
+	}
+	else if (ConvertedString == UTF8TEXTVIEW("Error"))
+	{
+		OutStatus = EStatus::Error;
+	}
+	else if (ConvertedString == UTF8TEXTVIEW("Canceled"))
+	{
+		OutStatus = EStatus::Canceled;
+	}
+	else
+	{
+		return false;
+	}
+	return true;
+}
+
+} // UE::DerivedData::Private
+
+namespace UE::DerivedData
+{
+
+FAnsiStringBuilderBase& operator<<(FAnsiStringBuilderBase& Builder, EPriority Priority) { return Private::PriorityToString(Builder, Priority); }
+FWideStringBuilderBase& operator<<(FWideStringBuilderBase& Builder, EPriority Priority) { return Private::PriorityToString(Builder, Priority); }
+FUtf8StringBuilderBase& operator<<(FUtf8StringBuilderBase& Builder, EPriority Priority) { return Private::PriorityToString(Builder, Priority); }
+
+bool TryLexFromString(EPriority& OutPriority, FUtf8StringView String) { return Private::PriorityFromString(OutPriority, String); }
+bool TryLexFromString(EPriority& OutPriority, FWideStringView String) { return Private::PriorityFromString(OutPriority, String); }
+
+FCbWriter& operator<<(FCbWriter& Writer, const EPriority Priority)
+{
+	Writer.AddString(WriteToUtf8String<16>(Priority));
+	return Writer;
+}
+
+bool LoadFromCompactBinary(FCbFieldView Field, EPriority& OutPriority, const EPriority Default)
+{
+	if (TryLexFromString(OutPriority, Field.AsString()))
+	{
+		return true;
+	}
+	OutPriority = Default;
+	return false;
+}
+
+FAnsiStringBuilderBase& operator<<(FAnsiStringBuilderBase& Builder, EStatus Status) { return Private::StatusToString(Builder, Status); }
+FWideStringBuilderBase& operator<<(FWideStringBuilderBase& Builder, EStatus Status) { return Private::StatusToString(Builder, Status); }
+FUtf8StringBuilderBase& operator<<(FUtf8StringBuilderBase& Builder, EStatus Status) { return Private::StatusToString(Builder, Status); }
+
+bool TryLexFromString(EStatus& OutStatus, FUtf8StringView String) { return Private::StatusFromString(OutStatus, String); }
+bool TryLexFromString(EStatus& OutStatus, FWideStringView String) { return Private::StatusFromString(OutStatus, String); }
+
+FCbWriter& operator<<(FCbWriter& Writer, const EStatus Status)
+{
+	Writer.AddString(WriteToUtf8String<16>(Status));
+	return Writer;
+}
+
+bool LoadFromCompactBinary(FCbFieldView Field, EStatus& OutStatus, const EStatus Default)
+{
+	if (TryLexFromString(OutStatus, Field.AsString()))
+	{
+		return true;
+	}
+	OutStatus = Default;
+	return false;
+}
+
+} // UE::DerivedData
diff --git a/Engine/Source/Developer/DerivedDataCache/Private/DerivedDataSharedString.cpp b/Engine/Source/Developer/DerivedDataCache/Private/DerivedDataSharedString.cpp
new file mode 100644
--- /dev/null
+++ b/Engine/Source/Developer/DerivedDataCache/Private/DerivedDataSharedString.cpp
@@ -0,0 +1,30 @@
+// Copyright Epic Games, Inc. All Rights Reserved.
+
+#include "DerivedDataSharedString.h"
+
+#include "Misc/StringBuilder.h"
+#include "Serialization/CompactBinary.h"
+#include "Serialization/CompactBinarySerialization.h"
+
+namespace UE::DerivedData
+{
+
+bool LoadFromCompactBinary(FCbFieldView Field, FUtf8SharedString& OutString)
+{
+	OutString = Field.AsString();
+	return !Field.HasError();
+}
+
+bool LoadFromCompactBinary(FCbFieldView Field, FWideSharedString& OutString)
+{
+	TWideStringBuilder<128> String;
+	if (LoadFromCompactBinary(Field, String))
+	{
+		OutString = String;
+		return true;
+	}
+	OutString.Reset();
+	return false;
+}
+
+} // UE::DerivedData
diff --git a/Engine/Source/Developer/DerivedDataCache/Private/DerivedDataValue.cpp b/Engine/Source/Developer/DerivedDataCache/Private/DerivedDataValue.cpp
--- a/Engine/Source/Developer/DerivedDataCache/Private/DerivedDataValue.cpp
+++ b/Engine/Source/Developer/DerivedDataCache/Private/DerivedDataValue.cpp
@@ -6,6 +6,8 @@
 #include "Containers/StringConv.h"
 #include "Hash/xxhash.h"
 #include "Serialization/CompactBinary.h"
+#include "Serialization/CompactBinarySerialization.h"
+#include "Serialization/CompactBinaryWriter.h"
 #include "Templates/UnrealTemplate.h"
 
 namespace UE::DerivedData
@@ -60,4 +62,21 @@ FValueId FValueId::FromName(const FWideStringView Name)
 	return FValueId::FromName(FTCHARToUTF8(Name));
 }
 
+FCbWriter& operator<<(FCbWriter& Writer, const FValueId& Id)
+{
+	Writer.AddObjectId(Id);
+	return Writer;
+}
+
+bool LoadFromCompactBinary(FCbFieldView Field, FValueId& OutId)
+{
+	if (const FCbObjectId ObjectId = Field.AsObjectId(); !Field.HasError())
+	{
+		OutId = ObjectId;
+		return true;
+	}
+	OutId.Reset();
+	return false;
+}
+
 } // UE::DerivedData
diff --git a/Engine/Source/Developer/DerivedDataCache/Public/DerivedDataCache.h b/Engine/Source/Developer/DerivedDataCache/Public/DerivedDataCache.h
--- a/Engine/Source/Developer/DerivedDataCache/Public/DerivedDataCache.h
+++ b/Engine/Source/Developer/DerivedDataCache/Public/DerivedDataCache.h
@@ -17,6 +17,9 @@
 
 #define UE_API DERIVEDDATACACHE_API
 
+class FCbFieldView;
+class FCbWriter;
+
 namespace UE::DerivedData { class ICacheStoreMaintainer; }
 namespace UE::DerivedData { class IRequestOwner; }
 namespace UE::DerivedData { struct FCacheGetChunkRequest; }
@@ -354,6 +357,16 @@ struct FCacheGetChunkResponse
 
 ///////////////////////////////////////////////////////////////////////////////////////////////////////////////////////
 
+UE_API FCbWriter& operator<<(FCbWriter& Writer, const FCacheGetRequest& Request);
+UE_API FCbWriter& operator<<(FCbWriter& Writer, const FCacheGetValueRequest& Request);
+UE_API FCbWriter& operator<<(FCbWriter& Writer, const FCacheGetChunkRequest& Request);
+
+UE_API bool LoadFromCompactBinary(FCbFieldView Field, FCacheGetRequest& OutRequest);
+UE_API bool LoadFromCompactBinary(FCbFieldView Field, FCacheGetValueRequest& OutRequest);
+UE_API bool LoadFromCompactBinary(FCbFieldView Field, FCacheGetChunkRequest& OutRequest);
+
+///////////////////////////////////////////////////////////////////////////////////////////////////////////////////////
+
 /** Returns a reference to the cache. Asserts if not available. */
 UE_API ICache& GetCache();
 
diff --git a/Engine/Source/Developer/DerivedDataCache/Public/DerivedDataCacheKey.h b/Engine/Source/Developer/DerivedDataCache/Public/DerivedDataCacheKey.h
--- a/Engine/Source/Developer/DerivedDataCache/Public/DerivedDataCacheKey.h
+++ b/Engine/Source/Developer/DerivedDataCache/Public/DerivedDataCacheKey.h
@@ -11,7 +11,7 @@
 
 #define UE_API DERIVEDDATACACHE_API
 
-class FCbObjectView;
+class FCbFieldView;
 class FCbWriter;
 
 namespace UE::DerivedData
@@ -106,6 +106,9 @@ inline TStringBuilderBase<CharType>& operator<<(TStringBuilderBase<CharType>& Bu
 	return Builder << Bucket.ToString();
 }
 
+UE_API FCbWriter& operator<<(FCbWriter& Writer, FCacheBucket Bucket);
+UE_API bool LoadFromCompactBinary(FCbFieldView Field, FCacheBucket& OutBucket);
+
 ///////////////////////////////////////////////////////////////////////////////////////////////////////////////////////
 
 inline bool operator==(const FCacheKey& A, const FCacheKey& B)
@@ -136,8 +139,8 @@ inline TStringBuilderBase<CharType>& operator<<(TStringBuilderBase<CharType>& Bu
 	return Builder << Key.Bucket << CharType('/') << Key.Hash;
 }
 
-UE_API bool TryLoadFromCompactBinary(FCbObjectView Object, FCacheKey& OutKey);
 UE_API FCbWriter& operator<<(FCbWriter& Writer, const FCacheKey& Key);
+UE_API bool LoadFromCompactBinary(FCbFieldView Field, FCacheKey& OutKey);
 
 ///////////////////////////////////////////////////////////////////////////////////////////////////////////////////////
 
diff --git a/Engine/Source/Developer/DerivedDataCache/Public/DerivedDataCachePolicy.h b/Engine/Source/Developer/DerivedDataCache/Public/DerivedDataCachePolicy.h
--- a/Engine/Source/Developer/DerivedDataCache/Public/DerivedDataCachePolicy.h
+++ b/Engine/Source/Developer/DerivedDataCache/Public/DerivedDataCachePolicy.h
@@ -12,6 +12,7 @@
 
 #define UE_API DERIVEDDATACACHE_API
 
+class FCbFieldView;
 class FCbObjectView;
 class FCbWriter;
 
@@ -105,10 +106,16 @@ UE_API FAnsiStringBuilderBase& operator<<(FAnsiStringBuilderBase& Builder, ECach
 UE_API FWideStringBuilderBase& operator<<(FWideStringBuilderBase& Builder, ECachePolicy Policy);
 UE_API FUtf8StringBuilderBase& operator<<(FUtf8StringBuilderBase& Builder, ECachePolicy Policy);
 
-/** Parse non-empty text written by operator<< into a policy. */
-UE_API ECachePolicy ParseCachePolicy(FAnsiStringView Text);
-UE_API ECachePolicy ParseCachePolicy(FWideStringView Text);
-UE_API ECachePolicy ParseCachePolicy(FUtf8StringView Text);
+/** Try to parse a policy from text written by operator<<. */
+UE_API bool TryLexFromString(ECachePolicy& OutPolicy, FUtf8StringView String);
+UE_API bool TryLexFromString(ECachePolicy& OutPolicy, FWideStringView String);
+
+UE_DEPRECATED(5.1, "Replace ParseCachePolicy with TryLexFromString.") UE_API ECachePolicy ParseCachePolicy(FAnsiStringView Text);
+UE_DEPRECATED(5.1, "Replace ParseCachePolicy with TryLexFromString.") UE_API ECachePolicy ParseCachePolicy(FWideStringView Text);
+UE_DEPRECATED(5.1, "Replace ParseCachePolicy with TryLexFromString.") UE_API ECachePolicy ParseCachePolicy(FUtf8StringView Text);
+
+UE_API FCbWriter& operator<<(FCbWriter& Writer, ECachePolicy Policy);
+UE_API bool LoadFromCompactBinary(FCbFieldView Field, ECachePolicy& OutPolicy, ECachePolicy Default = ECachePolicy::Default);
 
 /** A value ID and the cache policy to use for that value. */
 struct FCacheValuePolicy
@@ -120,6 +127,9 @@ struct FCacheValuePolicy
 	static constexpr ECachePolicy PolicyMask = ECachePolicy::Default | ECachePolicy::SkipData;
 };
 
+UE_API FCbWriter& operator<<(FCbWriter& Writer, const FCacheValuePolicy& Policy);
+UE_API bool LoadFromCompactBinary(FCbFieldView Field, FCacheValuePolicy& OutPolicy);
+
 /** Interface for the private implementation of the cache record policy. */
 class Private::ICacheRecordPolicyShared
 {
@@ -153,6 +163,9 @@ public:
 	{
 	}
 
+	/** Returns true if this is the default cache policy with no overrides for values. */
+	inline bool IsDefault() const { return !Shared && RecordPolicy == ECachePolicy::Default; }
+
 	/** Returns true if the record and every value use the same cache policy. */
 	inline bool IsUniform() const { return !Shared; }
 
@@ -178,9 +191,11 @@ public:
 	UE_API FCacheRecordPolicy Transform(TFunctionRef<ECachePolicy (ECachePolicy)> Op) const;
 
 	/** Saves the cache record policy to a compact binary object. */
+	UE_DEPRECATED(5.1, "Replace Policy.Save(Writer) with Writer << Policy.")
 	UE_API void Save(FCbWriter& Writer) const;
 
 	/** Loads a cache record policy from an object. */
+	UE_DEPRECATED(5.1, "Replace Load(Object) with LoadFromCompactBinary(Field, Policy).")
 	UE_API static FOptionalCacheRecordPolicy Load(FCbObjectView Object);
 
 private:
@@ -243,6 +258,9 @@ public:
 	inline void Reset() { *this = FOptionalCacheRecordPolicy(); }
 };
 
-} // namespace UE::DerivedData
+UE_API FCbWriter& operator<<(FCbWriter& Writer, const FCacheRecordPolicy& Policy);
+UE_API bool LoadFromCompactBinary(FCbFieldView Field, FCacheRecordPolicy& OutPolicy);
+
+} // UE::DerivedData
 
 #undef UE_API
diff --git a/Engine/Source/Developer/DerivedDataCache/Public/DerivedDataRequestTypes.h b/Engine/Source/Developer/DerivedDataCache/Public/DerivedDataRequestTypes.h
--- a/Engine/Source/Developer/DerivedDataCache/Public/DerivedDataRequestTypes.h
+++ b/Engine/Source/Developer/DerivedDataCache/Public/DerivedDataRequestTypes.h
@@ -3,6 +3,12 @@
 #pragma once
 
 #include "CoreTypes.h"
+#include "Containers/StringFwd.h"
+
+#define UE_API DERIVEDDATACACHE_API
+
+class FCbFieldView;
+class FCbWriter;
 
 namespace UE::DerivedData
 {
@@ -43,6 +49,18 @@ enum class EPriority : uint8
 	Blocking,
 };
 
+/** Append a non-empty text version of the priority to the builder. */
+UE_API FAnsiStringBuilderBase& operator<<(FAnsiStringBuilderBase& Builder, EPriority Priority);
+UE_API FWideStringBuilderBase& operator<<(FWideStringBuilderBase& Builder, EPriority Priority);
+UE_API FUtf8StringBuilderBase& operator<<(FUtf8StringBuilderBase& Builder, EPriority Priority);
+
+/** Try to parse a priority from text written by operator<<. */
+UE_API bool TryLexFromString(EPriority& OutPriority, FUtf8StringView String);
+UE_API bool TryLexFromString(EPriority& OutPriority, FWideStringView String);
+
+UE_API FCbWriter& operator<<(FCbWriter& Writer, EPriority Priority);
+UE_API bool LoadFromCompactBinary(FCbFieldView Field, EPriority& OutPriority, EPriority Default = EPriority::Normal);
+
 /** Status of a request that has completed. */
 enum class EStatus : uint8
 {
@@ -54,4 +72,18 @@ enum class EStatus : uint8
 	Canceled,
 };
 
+/** Append a non-empty text version of the status to the builder. */
+UE_API FAnsiStringBuilderBase& operator<<(FAnsiStringBuilderBase& Builder, EStatus Status);
+UE_API FWideStringBuilderBase& operator<<(FWideStringBuilderBase& Builder, EStatus Status);
+UE_API FUtf8StringBuilderBase& operator<<(FUtf8StringBuilderBase& Builder, EStatus Status);
+
+/** Try to parse a status from text written by operator<<. */
+UE_API bool TryLexFromString(EStatus& OutStatus, FUtf8StringView String);
+UE_API bool TryLexFromString(EStatus& OutStatus, FWideStringView String);
+
+UE_API FCbWriter& operator<<(FCbWriter& Writer, EStatus Status);
+UE_API bool LoadFromCompactBinary(FCbFieldView Field, EStatus& OutStatus, EStatus Default = EStatus::Ok);
+
 } // UE::DerivedData
+
+#undef UE_API
diff --git a/Engine/Source/Developer/DerivedDataCache/Public/DerivedDataSharedString.h b/Engine/Source/Developer/DerivedDataCache/Public/DerivedDataSharedString.h
--- a/Engine/Source/Developer/DerivedDataCache/Public/DerivedDataSharedString.h
+++ b/Engine/Source/Developer/DerivedDataCache/Public/DerivedDataSharedString.h
@@ -8,6 +8,11 @@
 #include "HAL/UnrealMemory.h"
 #include <atomic>
 
+#define UE_API DERIVEDDATACACHE_API
+
+class FCbFieldView;
+class FCbWriter;
+
 namespace UE::DerivedData
 {
 
@@ -174,4 +179,9 @@ inline bool operator<(const TSharedString<CharType>& Lhs, const TSharedString<Ch
 	return MakeStringView(Lhs).Compare(MakeStringView(Rhs), ESearchCase::IgnoreCase) < 0;
 }
 
+UE_API bool LoadFromCompactBinary(FCbFieldView Field, FUtf8SharedString& OutString);
+UE_API bool LoadFromCompactBinary(FCbFieldView Field, FWideSharedString& OutString);
+
 } // UE::DerivedData
+
+#undef UE_API
diff --git a/Engine/Source/Developer/DerivedDataCache/Public/DerivedDataValueId.h b/Engine/Source/Developer/DerivedDataCache/Public/DerivedDataValueId.h
--- a/Engine/Source/Developer/DerivedDataCache/Public/DerivedDataValueId.h
+++ b/Engine/Source/Developer/DerivedDataCache/Public/DerivedDataValueId.h
@@ -9,7 +9,9 @@
 
 #define UE_API DERIVEDDATACACHE_API
 
+class FCbFieldView;
 class FCbObjectId;
+class FCbWriter;
 struct FIoHash;
 
 namespace UE::DerivedData
@@ -104,6 +106,9 @@ inline bool FValueId::IsNull() const
 	return *this == FValueId();
 }
 
+UE_API FCbWriter& operator<<(FCbWriter& Writer, const FValueId& Id);
+UE_API bool LoadFromCompactBinary(FCbFieldView Field, FValueId& OutId);
+
 } // UE::DerivedData
 
 #undef UE_API
