DDC: Fixed deprecation warnings in the Zen cache store

#rnx

diff --git a/Engine/Source/Developer/DerivedDataCache/Private/ZenCacheStore.cpp b/Engine/Source/Developer/DerivedDataCache/Private/ZenCacheStore.cpp
--- a/Engine/Source/Developer/DerivedDataCache/Private/ZenCacheStore.cpp
+++ b/Engine/Source/Developer/DerivedDataCache/Private/ZenCacheStore.cpp
@@ -217,10 +217,9 @@ public:
 						{
 							BatchWriter.SetName(ANSITEXTVIEW("Record"));
 							Record.Save(BatchPackage, BatchWriter);
-							if (!Request.Policy.IsUniform() || Request.Policy.GetRecordPolicy() != BatchDefaultPolicy)
+							if (!Request.Policy.IsDefault())
 							{
-								BatchWriter.SetName(ANSITEXTVIEW("Policy"));
-								Request.Policy.Save(BatchWriter);
+								BatchWriter << ANSITEXTVIEW("Policy") << Request.Policy;
 							}
 						}
 						BatchWriter.EndObject();
@@ -363,10 +362,9 @@ public:
 							BatchRequest.BeginObject();
 							{
 								BatchRequest << ANSITEXTVIEW("Key") << Request.Key;
-								if (!Request.Policy.IsUniform() || Request.Policy.GetRecordPolicy() != BatchDefaultPolicy)
+								if (!Request.Policy.IsDefault())
 								{
-									BatchRequest.SetName(ANSITEXTVIEW("Policy"));
-									Request.Policy.Save(BatchRequest);
+									BatchRequest << ANSITEXTVIEW("Policy") << Request.Policy;
 								}
 							}
 							BatchRequest.EndObject();
