DDC: Added cache replay functionality

A cache replay can be created from any program that uses the cache, and can be replayed through any cache graph in any program that uses the cache.

Save a cache replay by adding -DDC-ReplaySave=PathToReplay.ddr to the command line, and compress when saving by adding -DDC-ReplayCompress to reduce replay size by approximately 85%.

Load a cache replay by adding -DDC-ReplayLoad=PathToReplay.ddr to the command line one or more times.
- Modify the cache priority on loaded replays with -DDC-ReplayLoadPriority=Priority.
- For example, -DDC-ReplayLoadPriority=Lowest will queue replayed requests to execute asynchronously.
- Modify the cache policy on loaded replays with -DDC-ReplayLoadAddPolicy=CachePolicy and -DDC-ReplayLoadRemovePolicy=CachePolicy.
- For example, -DDC-ReplayLoadAddPolicy=SkipData will make the replay prefetch into the cache stores earlier in the cache hierarchy.
- For example, -DDC-ReplayLoadRemovePolicy=QueryLocal will make the replay query only remote cache stores.

Filter save or load operations by combining these arguments:
-DDC-ReplayMethods=Get+GetValue+GetChunks will only process requests of the corresponding methods.
-DDC-ReplayTypes=Type1[[REDACTED_HANDLE]][+Type2[[REDACTED_HANDLE]]...] will only process requests of the corresponding types.
-DDC-ReplayRate=Rate will only process Rate% of the requests that do not match one of the type filters.
-DDC-ReplaySalt=Salt configures the salt used to filter by type and rate, which makes them match consistently between sessions instead of being random.

Load a cache replay with the console command: DDC.LoadReplay ReplayPath [Priority=PriorityName]

diff --git a/Engine/Source/Developer/DerivedDataCache/Private/DerivedDataBackends.cpp b/Engine/Source/Developer/DerivedDataCache/Private/DerivedDataBackends.cpp
--- a/Engine/Source/Developer/DerivedDataCache/Private/DerivedDataBackends.cpp
+++ b/Engine/Source/Developer/DerivedDataCache/Private/DerivedDataBackends.cpp
@@ -57,6 +57,8 @@ IMemoryCacheStore* CreateMemoryCacheStore(const TCHAR* Name, int64 MaxCacheSize,
 IPakFileCacheStore* CreatePakFileCacheStore(const TCHAR* Filename, bool bWriting, bool bCompressed);
 ILegacyCacheStore* CreateS3CacheStore(const TCHAR* RootManifestPath, const TCHAR* BaseUrl, const TCHAR* Region, const TCHAR* CanaryObjectKey, const TCHAR* CachePath);
 TTuple<ILegacyCacheStore*, ECacheStoreFlags> CreateZenCacheStore(const TCHAR* NodeName, const TCHAR* Config);
+ILegacyCacheStore* TryCreateCacheStoreReplay(ILegacyCacheStore* InnerCache);
+bool TryLoadCacheStoreReplay(ILegacyCacheStore* TargetCache, const TCHAR* ReplayPath, const TCHAR* PriorityName = nullptr);
 
 /**
  * This class is used to create a singleton that represents the derived data cache hierarchy and all of the wrappers necessary
@@ -85,10 +87,14 @@ public:
 			TEXT("DDC.MountPak"),
 			*LOCTEXT("CommandText_DDCMountPak", "Mounts read-only pak file").ToString(),
 			FConsoleCommandWithArgsDelegate::CreateRaw(this, &FDerivedDataBackendGraph::MountPakCommandHandler))
-		, UnountPakCommand(
+		, UnmountPakCommand(
 			TEXT("DDC.UnmountPak"),
 			*LOCTEXT("CommandText_DDCUnmountPak", "Unmounts read-only pak file").ToString(),
 			FConsoleCommandWithArgsDelegate::CreateRaw(this, &FDerivedDataBackendGraph::UnmountPakCommandHandler))
+		, LoadReplayCommand(
+			TEXT("DDC.LoadReplay"),
+			*LOCTEXT("CommandText_DDCLoadReplay", "Loads a cache replay file created by -DDC-ReplaySave=<Path>").ToString(),
+			FConsoleCommandWithArgsDelegate::CreateRaw(this, &FDerivedDataBackendGraph::LoadReplayCommandHandler))
 	{
 		check(!StaticGraph);
 		StaticGraph = this;
@@ -186,7 +192,7 @@ public:
 		{
 			FString VerifyArg;
 			if (FParse::Value(FCommandLine::Get(), TEXT("-DDC-Verify="), VerifyArg) ||
-				FParse::Param(FCommandLine::Get(), TEXT("-DDC-Verify")))
+				FParse::Param(FCommandLine::Get(), TEXT("DDC-Verify")))
 			{
 				VerifyNode = CreateCacheStoreVerify(RootNode.Key, /*bPutOnError*/ false);
 				CreatedNodes.AddUnique(VerifyNode);
@@ -194,6 +200,13 @@ public:
 			}
 		}
 
+		// Create a Replay node when requested on the command line.
+		if (ILegacyCacheStore* ReplayNode = TryCreateCacheStoreReplay(RootNode.Key))
+		{
+			CreatedNodes.AddUnique(ReplayNode);
+			RootNode.Key = ReplayNode;
+		}
+
 		if (MaxKeyLength == 0)
 		{
 			MaxKeyLength = MAX_BACKEND_KEY_LENGTH;
@@ -1086,6 +1099,23 @@ private:
 		MountPakFile(*Args[0]);
 	}
 
+	void LoadReplayCommandHandler(const TArray<FString>& Args)
+	{
+		if (Args.Num() < 1)
+		{
+			UE_LOG(LogDerivedDataCache, Log, TEXT("Usage: DDC.LoadReplay ReplayPath [Priority=<Name>]"));
+			return;
+		}
+
+		FString PriorityName;
+		for (int32 ArgIndex = 1; ArgIndex < Args.Num(); ++ArgIndex)
+		{
+			FParse::Value(*Args[ArgIndex], TEXT("Priority="), PriorityName);
+		}
+
+		TryLoadCacheStoreReplay(RootCache, *Args[0], *PriorityName);
+	}
+
 	static inline FDerivedDataBackendGraph*			StaticGraph;
 
 	FThreadSafeCounter								AsyncCompletionCounter;
@@ -1127,7 +1157,9 @@ private:
 	/** MountPak console command */
 	FAutoConsoleCommand MountPakCommand;
 	/** UnmountPak console command */
-	FAutoConsoleCommand UnountPakCommand;
+	FAutoConsoleCommand UnmountPakCommand;
+
+	FAutoConsoleCommand LoadReplayCommand;
 };
 
 } // UE::DerivedData
diff --git a/Engine/Source/Developer/DerivedDataCache/Private/DerivedDataCacheStoreReplay.cpp b/Engine/Source/Developer/DerivedDataCache/Private/DerivedDataCacheStoreReplay.cpp
new file mode 100644
--- /dev/null
+++ b/Engine/Source/Developer/DerivedDataCache/Private/DerivedDataCacheStoreReplay.cpp
@@ -0,0 +1,778 @@
+// Copyright Epic Games, Inc. All Rights Reserved.
+
+#include "Algo/AnyOf.h"
+#include "Compression/CompressedBuffer.h"
+#include "Containers/StringView.h"
+#include "DerivedDataCacheKeyFilter.h"
+#include "DerivedDataLegacyCacheStore.h"
+#include "DerivedDataRequestOwner.h"
+#include "DerivedDataValue.h"
+#include "HAL/CriticalSection.h"
+#include "HAL/FileManager.h"
+#include "HAL/PlatformTime.h"
+#include "Misc/CommandLine.h"
+#include "Misc/Parse.h"
+#include "Misc/ScopeLock.h"
+#include "ProfilingDebugging/CpuProfilerTrace.h"
+#include "Serialization/CompactBinary.h"
+#include "Serialization/CompactBinarySerialization.h"
+#include "Serialization/CompactBinaryWriter.h"
+#include "Serialization/MemoryReader.h"
+#include "String/Find.h"
+#include "String/ParseTokens.h"
+#include "Templates/Invoke.h"
+
+namespace UE::DerivedData
+{
+
+static constexpr uint64 GCacheReplayCompressionBlockSize = 256 * 1024;
+
+///////////////////////////////////////////////////////////////////////////////////////////////////////////////////////
+
+enum class ECacheMethod : uint8
+{
+	Put,
+	Get,
+	PutValue,
+	GetValue,
+	GetChunks,
+};
+
+template <typename CharType>
+static TStringBuilderBase<CharType>& operator<<(TStringBuilderBase<CharType>& Builder, const ECacheMethod Method)
+{
+	switch (Method)
+	{
+	case ECacheMethod::Put:       return Builder << ANSITEXTVIEW("Put");
+	case ECacheMethod::Get:       return Builder << ANSITEXTVIEW("Get");
+	case ECacheMethod::PutValue:  return Builder << ANSITEXTVIEW("PutValue");
+	case ECacheMethod::GetValue:  return Builder << ANSITEXTVIEW("GetValue");
+	case ECacheMethod::GetChunks: return Builder << ANSITEXTVIEW("GetChunks");
+	}
+	return Builder << ANSITEXTVIEW("Unknown");
+}
+
+template <typename CharType>
+static bool TryLexFromString(ECacheMethod& OutMethod, const TStringView<CharType> String)
+{
+	const auto ConvertedString = StringCast<UTF8CHAR, 16>(String.GetData(), String.Len());
+	if (ConvertedString == UTF8TEXTVIEW("Put"))
+	{
+		OutMethod = ECacheMethod::Put;
+	}
+	else if (ConvertedString == UTF8TEXTVIEW("Get"))
+	{
+		OutMethod = ECacheMethod::Get;
+	}
+	else if (ConvertedString == UTF8TEXTVIEW("PutValue"))
+	{
+		OutMethod = ECacheMethod::PutValue;
+	}
+	else if (ConvertedString == UTF8TEXTVIEW("GetValue"))
+	{
+		OutMethod = ECacheMethod::GetValue;
+	}
+	else if (ConvertedString == UTF8TEXTVIEW("GetChunks"))
+	{
+		OutMethod = ECacheMethod::GetChunks;
+	}
+	else
+	{
+		return false;
+	}
+	return true;
+}
+
+static FCbWriter& operator<<(FCbWriter& Writer, const ECacheMethod Method)
+{
+	Writer.AddString(WriteToUtf8String<16>(Method));
+	return Writer;
+}
+
+static bool LoadFromCompactBinary(FCbFieldView Field, ECacheMethod& OutMethod)
+{
+	if (TryLexFromString(OutMethod, Field.AsString()))
+	{
+		return true;
+	}
+	OutMethod = {};
+	return false;
+}
+
+///////////////////////////////////////////////////////////////////////////////////////////////////////////////////////
+
+class FCacheMethodFilter
+{
+public:
+	static FCacheMethodFilter Parse(const TCHAR* CommandLine);
+
+	inline bool IsMatch(ECacheMethod Method) const
+	{
+		return (MethodMask & (1 << uint32(Method))) == 0;
+	}
+
+private:
+	uint32 MethodMask = 0;
+};
+
+class FCachePolicyFilter
+{
+public:
+	static FCachePolicyFilter Parse(const TCHAR* CommandLine);
+
+	inline ECachePolicy operator()(ECachePolicy Policy) const
+	{
+		EnumAddFlags(Policy, FlagsToAdd);
+		EnumRemoveFlags(Policy, FlagsToRemove);
+		return Policy;
+	}
+
+private:
+	ECachePolicy FlagsToAdd = ECachePolicy::None;
+	ECachePolicy FlagsToRemove = ECachePolicy::None;
+};
+
+///////////////////////////////////////////////////////////////////////////////////////////////////////////////////////
+
+class FCacheStoreReplay final : public ILegacyCacheStore
+{
+public:
+	FCacheStoreReplay(
+		ILegacyCacheStore* InnerCache,
+		FCacheKeyFilter KeyFilter,
+		FCacheMethodFilter MethodFilter,
+		FString&& ReplayPath,
+		uint64 CompressionBlockSize = 0);
+
+	~FCacheStoreReplay();
+
+	void Put(
+		const TConstArrayView<FCachePutRequest> Requests,
+		IRequestOwner& Owner,
+		FOnCachePutComplete&& OnComplete) final;
+	void Get(
+		const TConstArrayView<FCacheGetRequest> Requests,
+		IRequestOwner& Owner,
+		FOnCacheGetComplete&& OnComplete) final;
+	void PutValue(
+		const TConstArrayView<FCachePutValueRequest> Requests,
+		IRequestOwner& Owner,
+		FOnCachePutValueComplete&& OnComplete) final;
+	void GetValue(
+		const TConstArrayView<FCacheGetValueRequest> Requests,
+		IRequestOwner& Owner,
+		FOnCacheGetValueComplete&& OnComplete) final;
+	void GetChunks(
+		const TConstArrayView<FCacheGetChunkRequest> Requests,
+		IRequestOwner& Owner,
+		FOnCacheGetChunkComplete&& OnComplete) final;
+
+	void LegacyStats(FDerivedDataCacheStatsNode& OutNode) final
+	{
+		InnerCache->LegacyStats(OutNode);
+	}
+
+	bool LegacyDebugOptions(FBackendDebugOptions& Options) final
+	{
+		return InnerCache->LegacyDebugOptions(Options);
+	}
+
+private:
+	template <typename RequestType>
+	void SerializeRequests(TConstArrayView<RequestType> Requests, ECacheMethod Method, EPriority Priority);
+
+	void WriteBinaryToArchive(const FCompositeBuffer& RawBinary);
+	void WriteToArchive(FCbWriter& Writer);
+	void FlushToArchive();
+
+	ILegacyCacheStore* InnerCache;
+	FCacheKeyFilter KeyFilter;
+	FCacheMethodFilter MethodFilter;
+	FString ReplayPath;
+	TUniquePtr<FArchive> ReplayAr;
+	FUniqueBuffer RawBlock;
+	FMutableMemoryView RawBlockTail;
+	FCriticalSection Lock;
+};
+
+FCacheStoreReplay::FCacheStoreReplay(
+	ILegacyCacheStore* InInnerCache,
+	FCacheKeyFilter InKeyFilter,
+	FCacheMethodFilter InMethodFilter,
+	FString&& InReplayPath,
+	uint64 CompressionBlockSize)
+	: InnerCache(InInnerCache)
+	, KeyFilter(MoveTemp(InKeyFilter))
+	, MethodFilter(MoveTemp(InMethodFilter))
+	, ReplayPath(MoveTemp(InReplayPath))
+	, ReplayAr(IFileManager::Get().CreateFileWriter(*ReplayPath, FILEWRITE_NoFail))
+{
+	if (CompressionBlockSize)
+	{
+		RawBlock = FUniqueBuffer::Alloc(CompressionBlockSize);
+		RawBlockTail = RawBlock;
+	}
+	UE_LOG(LogDerivedDataCache, Display, TEXT("Replay: Saving cache replay to '%s'"), *ReplayPath);
+}
+
+FCacheStoreReplay::~FCacheStoreReplay()
+{
+	FlushToArchive();
+}
+
+template <typename RequestType>
+void FCacheStoreReplay::SerializeRequests(
+	const TConstArrayView<RequestType> Requests,
+	const ECacheMethod Method,
+	const EPriority Priority)
+{
+	TRACE_CPUPROFILER_EVENT_SCOPE(ReplayDDC_Serialize);
+
+	const auto IsKeyMatch = [this](const RequestType& Request) { return KeyFilter.IsMatch(Request.Key); };
+	if (!MethodFilter.IsMatch(Method) || !Algo::AnyOf(Requests, IsKeyMatch))
+	{
+		return;
+	}
+
+	TCbWriter<512> Writer;
+	Writer.BeginObject();
+	Writer << ANSITEXTVIEW("Method") << Method;
+	Writer << ANSITEXTVIEW("Priority") << Priority;
+	Writer.BeginArray(ANSITEXTVIEW("Requests"));
+	for (const RequestType& Request : Requests)
+	{
+		if (KeyFilter.IsMatch(Request.Key))
+		{
+			Writer << Request;
+		}
+	}
+	Writer.EndArray();
+	Writer.EndObject();
+
+	if (UE_LOG_ACTIVE(LogDerivedDataCache, Verbose))
+	{
+		TUtf8StringBuilder<1024> Batch;
+		CompactBinaryToCompactJson(Writer.Save().AsObject(), Batch);
+		UE_LOG(LogDerivedDataCache, Verbose, TEXT("Replay: %hs"), *Batch);
+	}
+
+	WriteToArchive(Writer);
+}
+
+void FCacheStoreReplay::WriteBinaryToArchive(const FCompositeBuffer& RawBinary)
+{
+	const FValue CompressedBinary = FValue::Compress(RawBinary, RawBlock.GetSize());
+	TCbWriter<64> BinaryWriter;
+	BinaryWriter.AddBinary(CompressedBinary.GetData().GetCompressed());
+	BinaryWriter.Save(*ReplayAr);
+}
+
+void FCacheStoreReplay::WriteToArchive(FCbWriter& Writer)
+{
+	FScopeLock ScopeLock(&Lock);
+	if (RawBlock)
+	{
+		const uint64 SaveSize = Writer.GetSaveSize();
+		if (RawBlockTail.GetSize() < SaveSize)
+		{
+			FlushToArchive();
+		}
+		if (RawBlockTail.GetSize() < SaveSize)
+		{
+			WriteBinaryToArchive(Writer.Save().AsObject().GetBuffer());
+		}
+		else
+		{
+			Writer.Save(RawBlockTail.Left(SaveSize));
+			RawBlockTail += SaveSize;
+		}
+	}
+	else
+	{
+		Writer.Save(*ReplayAr);
+	}
+}
+
+void FCacheStoreReplay::FlushToArchive()
+{
+	const FSharedBuffer RawBlockHead = FSharedBuffer::MakeView(RawBlock.GetView().LeftChop(RawBlockTail.GetSize()));
+	if (RawBlockHead.GetSize() > 0)
+	{
+		WriteBinaryToArchive(FCompositeBuffer(RawBlockHead));
+		RawBlockTail = RawBlock;
+	}
+}
+
+void FCacheStoreReplay::Put(
+	const TConstArrayView<FCachePutRequest> Requests,
+	IRequestOwner& Owner,
+	FOnCachePutComplete&& OnComplete)
+{
+	InnerCache->Put(Requests, Owner, MoveTemp(OnComplete));
+}
+
+void FCacheStoreReplay::Get(
+	const TConstArrayView<FCacheGetRequest> Requests,
+	IRequestOwner& Owner,
+	FOnCacheGetComplete&& OnComplete)
+{
+	SerializeRequests(Requests, ECacheMethod::Get, Owner.GetPriority());
+	InnerCache->Get(Requests, Owner, MoveTemp(OnComplete));
+}
+
+void FCacheStoreReplay::PutValue(
+	const TConstArrayView<FCachePutValueRequest> Requests,
+	IRequestOwner& Owner,
+	FOnCachePutValueComplete&& OnComplete)
+{
+	InnerCache->PutValue(Requests, Owner, MoveTemp(OnComplete));
+}
+
+void FCacheStoreReplay::GetValue(
+	const TConstArrayView<FCacheGetValueRequest> Requests,
+	IRequestOwner& Owner,
+	FOnCacheGetValueComplete&& OnComplete)
+{
+	SerializeRequests(Requests, ECacheMethod::GetValue, Owner.GetPriority());
+	InnerCache->GetValue(Requests, Owner, MoveTemp(OnComplete));
+}
+
+void FCacheStoreReplay::GetChunks(
+	const TConstArrayView<FCacheGetChunkRequest> Requests,
+	IRequestOwner& Owner,
+	FOnCacheGetChunkComplete&& OnComplete)
+{
+	SerializeRequests(Requests, ECacheMethod::GetChunks, Owner.GetPriority());
+	InnerCache->GetChunks(Requests, Owner, MoveTemp(OnComplete));
+}
+
+///////////////////////////////////////////////////////////////////////////////////////////////////////////////////////
+
+class FCacheReplayReader
+{
+public:
+	static constexpr uint64 DefaultScratchSize = 1024;
+
+	FCacheReplayReader(
+		ILegacyCacheStore* TargetCache,
+		FCacheKeyFilter KeyFilter,
+		FCacheMethodFilter MethodFilter,
+		FCachePolicyFilter PolicyFilter,
+		const EPriority* ForcedPriority = nullptr);
+
+	~FCacheReplayReader();
+
+	bool ReadFromFile(const TCHAR* ReplayPath, uint64 ScratchSize = DefaultScratchSize);
+	bool ReadFromArchive(FArchive& ReplayAr, uint64 ScratchSize = DefaultScratchSize);
+	bool ReadFromObject(const FCbObject& Object);
+
+private:
+	template <typename RequestType, typename FunctionType>
+	bool DispatchRequests(FCbObjectView Object, ECacheMethod Method, EPriority Priority, FunctionType Function);
+
+	void ApplyPolicyFilter(FCacheGetRequest& Request);
+	void ApplyPolicyFilter(FCacheGetValueRequest& Request);
+	void ApplyPolicyFilter(FCacheGetChunkRequest& Request);
+
+	ILegacyCacheStore* TargetCache;
+	FCacheKeyFilter KeyFilter;
+	FCacheMethodFilter MethodFilter;
+	FCachePolicyFilter PolicyFilter;
+	FRequestOwner BlockingTaskOwner;
+	TArray<FRequestOwner, TInlineAllocator<6>> Owners;
+	int64 DispatchCount = 0;
+	double CreationTime = FPlatformTime::Seconds();
+};
+
+FCacheReplayReader::FCacheReplayReader(
+	ILegacyCacheStore* const InTargetCache,
+	FCacheKeyFilter InKeyFilter,
+	FCacheMethodFilter InMethodFilter,
+	FCachePolicyFilter InPolicyFilter,
+	const EPriority* const InForcedPriority)
+	: TargetCache(InTargetCache)
+	, KeyFilter(MoveTemp(InKeyFilter))
+	, MethodFilter(MoveTemp(InMethodFilter))
+	, PolicyFilter(MoveTemp(InPolicyFilter))
+	, BlockingTaskOwner(EPriority::Normal)
+{
+	static_assert(uint8(EPriority::Lowest) == 0);
+	Owners.Emplace(InForcedPriority ? *InForcedPriority : EPriority::Lowest);
+	static_assert(uint8(EPriority::Low) == 1);
+	Owners.Emplace(InForcedPriority ? *InForcedPriority : EPriority::Low);
+	static_assert(uint8(EPriority::Normal) == 2);
+	Owners.Emplace(InForcedPriority ? *InForcedPriority : EPriority::Normal);
+	static_assert(uint8(EPriority::High) == 3);
+	Owners.Emplace(InForcedPriority ? *InForcedPriority : EPriority::High);
+	static_assert(uint8(EPriority::Highest) == 4);
+	Owners.Emplace(InForcedPriority ? *InForcedPriority : EPriority::Highest);
+	static_assert(uint8(EPriority::Blocking) == 5);
+	Owners.Emplace(InForcedPriority ? *InForcedPriority : EPriority::Blocking);
+}
+
+FCacheReplayReader::~FCacheReplayReader()
+{
+	UE_LOG(LogDerivedDataCache, Display, TEXT("Replay: Dispatched %" INT64_FMT " requests in %.3lf seconds."),
+		DispatchCount, FPlatformTime::Seconds() - CreationTime);
+
+	TRACE_CPUPROFILER_EVENT_SCOPE(ReplayDDC_Wait);
+	BlockingTaskOwner.Wait();
+
+	for (FRequestOwner& Owner : Owners)
+	{
+		Owner.Wait();
+	}
+
+	UE_LOG(LogDerivedDataCache, Display, TEXT("Replay: Completed %" INT64_FMT " requests in %.3lf seconds."),
+		DispatchCount, FPlatformTime::Seconds() - CreationTime);
+}
+
+template <typename RequestType, typename FunctionType>
+bool FCacheReplayReader::DispatchRequests(
+	const FCbObjectView Object,
+	const ECacheMethod Method,
+	const EPriority Priority,
+	const FunctionType Function)
+{
+	TArray<RequestType, TInlineAllocator<16>> Requests;
+	{
+		TRACE_CPUPROFILER_EVENT_SCOPE(ReplayDDC_Serialize);
+		const FCbArrayView Array = Object[ANSITEXTVIEW("Requests")].AsArrayView();
+		Requests.Reserve(Array.Num());
+		for (FCbFieldView Field : Array)
+		{
+			RequestType& Request = Requests.AddDefaulted_GetRef();
+			if (!LoadFromCompactBinary(Field, Request))
+			{
+				return false;
+			}
+			if (KeyFilter.IsMatch(Request.Key))
+			{
+				ApplyPolicyFilter(Request);
+			}
+			else
+			{
+				Requests.Pop(/*bAllowShrinking*/ false);
+			}
+		}
+	}
+
+	if (!Requests.IsEmpty())
+	{
+		if (UE_LOG_ACTIVE(LogDerivedDataCache, Verbose))
+		{
+			TUtf8StringBuilder<1024> Batch;
+			CompactBinaryToCompactJson(Object, Batch);
+			UE_LOG(LogDerivedDataCache, Verbose, TEXT("Replay: %hs"), *Batch);
+		}
+
+		TRACE_CPUPROFILER_EVENT_SCOPE(ReplayDDC_Dispatch);
+
+		DispatchCount += Requests.Num();
+
+		FRequestOwner& Owner = Owners[uint8(Priority)];
+		if (Owner.GetPriority() < EPriority::Blocking)
+		{
+			// Owners with non-blocking priority can execute the request directly because it will be async.
+			FRequestBarrier Barrier(Owner);
+			Invoke(Function, TargetCache, Requests, Owner, [](auto&&){});
+		}
+		else
+		{
+			// Owners with blocking priority launch a task to execute the blocking request to allow concurrent replay.
+			FRequestBarrier Barrier(BlockingTaskOwner);
+			((IRequestOwner&)BlockingTaskOwner).LaunchTask(TEXT("CacheReplayTask"),
+				[this, Requests = MoveTemp(Requests), Function]
+				{
+					FRequestOwner BlockingOwner(EPriority::Blocking);
+					Invoke(Function, TargetCache, Requests, BlockingOwner, [](auto&&){});
+				});
+		}
+	}
+
+	return true;
+}
+
+void FCacheReplayReader::ApplyPolicyFilter(FCacheGetRequest& Request)
+{
+	Request.Policy = Request.Policy.Transform([this](ECachePolicy Policy) { return PolicyFilter(Policy); });
+}
+
+void FCacheReplayReader::ApplyPolicyFilter(FCacheGetValueRequest& Request)
+{
+	Request.Policy = PolicyFilter(Request.Policy);
+}
+
+void FCacheReplayReader::ApplyPolicyFilter(FCacheGetChunkRequest& Request)
+{
+	Request.Policy = PolicyFilter(Request.Policy);
+}
+
+bool FCacheReplayReader::ReadFromFile(const TCHAR* const ReplayPath, const uint64 ScratchSize)
+{
+	TUniquePtr<FArchive> ReplayAr(IFileManager::Get().CreateFileReader(ReplayPath, FILEREAD_Silent));
+	if (!ReplayAr)
+	{
+		UE_LOG(LogDerivedDataCache, Warning, TEXT("Replay: File '%s' failed to open."), ReplayPath);
+		return false;
+	}
+
+	UE_LOG(LogDerivedDataCache, Display, TEXT("Replay: Loading cache replay from '%s'"), ReplayPath);
+	return ReadFromArchive(*ReplayAr, FMath::Max(ScratchSize, GCacheReplayCompressionBlockSize));
+}
+
+bool FCacheReplayReader::ReadFromArchive(FArchive& ReplayAr, const uint64 ScratchSize)
+{
+	// A scratch buffer for the compact binary fields.
+	FUniqueBuffer Scratch = FUniqueBuffer::Alloc(ScratchSize);
+	const auto Alloc = [&Scratch](const uint64 Size) -> FUniqueBuffer
+	{
+		if (Size <= Scratch.GetSize())
+		{
+			return FUniqueBuffer::MakeView(Scratch.GetView().Left(Size));
+		}
+		return FUniqueBuffer::Alloc(Size);
+	};
+
+	// A scratch buffer for decompressed blocks of fields.
+	FUniqueBuffer BlockScratch;
+
+	for (int64 Offset = ReplayAr.Tell(); Offset < ReplayAr.TotalSize(); Offset = ReplayAr.Tell())
+	{
+		FCbField Field = LoadCompactBinary(ReplayAr, Alloc);
+		if (!Field)
+		{
+			UE_LOG(LogDerivedDataCache, Warning,
+				TEXT("Replay: Failed to load compact binary at offset %" INT64_FMT ". "
+					 "Archive is at offset %" INT64_FMT " and has total size %" INT64_FMT "."),
+				Offset, ReplayAr.Tell(), ReplayAr.TotalSize());
+			return false;
+		}
+
+		// A binary field is used to store a compressed buffer containing a sequence of compact binary objects.
+		if (Field.IsBinary())
+		{
+			FCompressedBuffer CompressedBuffer = FCompressedBuffer::FromCompressed(Field.AsBinary());
+			if (!CompressedBuffer)
+			{
+				UE_LOG(LogDerivedDataCache, Warning,
+					TEXT("Replay: Failed to load compressed buffer from binary field at offset %" INT64_FMT ". "
+						 "Archive is at offset %" INT64_FMT " and has total size %" INT64_FMT "."),
+					Offset, ReplayAr.Tell(), ReplayAr.TotalSize());
+				return false;
+			}
+
+			const uint64 RawBlockSize = CompressedBuffer.GetRawSize();
+			if (BlockScratch.GetSize() < RawBlockSize)
+			{
+				BlockScratch = FUniqueBuffer::Alloc(RawBlockSize);
+			}
+
+			const FMutableMemoryView RawBlockView = BlockScratch.GetView().Left(RawBlockSize);
+			if (!CompressedBuffer.TryDecompressTo(RawBlockView))
+			{
+				UE_LOG(LogDerivedDataCache, Warning,
+					TEXT("Replay: Failed to decompress compressed buffer from binary field at offset %" INT64_FMT ". "
+						 "Archive is at offset %" INT64_FMT " and has total size %" INT64_FMT "."),
+					Offset, ReplayAr.Tell(), ReplayAr.TotalSize());
+				return false;
+			}
+
+			FMemoryReaderView InnerAr(RawBlockView);
+			if (!ReadFromArchive(InnerAr))
+			{
+				return false;
+			}
+		}
+
+		// An object field is used to store one batch of cache requests.
+		if (Field.IsObject() && !ReadFromObject(Field.AsObject()))
+		{
+			UE_LOG(LogDerivedDataCache, Warning,
+				TEXT("Replay: Failed to load cache request from object field at offset %" INT64_FMT ". "
+						"Archive is at offset %" INT64_FMT " and has total size %" INT64_FMT "."),
+				Offset, ReplayAr.Tell(), ReplayAr.TotalSize());
+			return false;
+		}
+	}
+
+	return true;
+}
+
+bool FCacheReplayReader::ReadFromObject(const FCbObject& Object)
+{
+	ECacheMethod Method{};
+	EPriority Priority{};
+	if (!LoadFromCompactBinary(Object[ANSITEXTVIEW("Method")], Method) ||
+		!LoadFromCompactBinary(Object[ANSITEXTVIEW("Priority")], Priority))
+	{
+		return false;
+	}
+
+	if (!MethodFilter.IsMatch(Method))
+	{
+		return true;
+	}
+
+	switch (Method)
+	{
+	case ECacheMethod::Get:
+		return DispatchRequests<FCacheGetRequest>(Object, Method, Priority, &ICacheStore::Get);
+	case ECacheMethod::GetValue:
+		return DispatchRequests<FCacheGetValueRequest>(Object, Method, Priority, &ICacheStore::GetValue);
+	case ECacheMethod::GetChunks:
+		return DispatchRequests<FCacheGetChunkRequest>(Object, Method, Priority, &ICacheStore::GetChunks);
+	}
+
+	return false;
+}
+
+///////////////////////////////////////////////////////////////////////////////////////////////////////////////////////
+
+FCacheMethodFilter FCacheMethodFilter::Parse(const TCHAR* const CommandLine)
+{
+	FCacheMethodFilter MethodFilter;
+	FString MethodNames;
+	if (FParse::Value(CommandLine, TEXT("-DDC-ReplayMethods="), MethodNames))
+	{
+		MethodFilter.MethodMask = ~uint32(0);
+		String::ParseTokens(MethodNames, TEXT('+'), [&MethodFilter](FStringView MethodName)
+		{
+			ECacheMethod Method;
+			if (TryLexFromString(Method, MethodName))
+			{
+				MethodFilter.MethodMask &= ~(1 << uint32(Method));
+			}
+		});
+	}
+	return MethodFilter;
+}
+
+FCachePolicyFilter FCachePolicyFilter::Parse(const TCHAR* const CommandLine)
+{
+	FCachePolicyFilter PolicyFilter;
+
+	FString FlagNamesToAdd;
+	if (FParse::Value(CommandLine, TEXT("-DDC-ReplayLoadAddPolicy="), FlagNamesToAdd))
+	{
+		TryLexFromString(PolicyFilter.FlagsToAdd, FlagNamesToAdd);
+	}
+
+	FString FlagNamesToRemove;
+	if (FParse::Value(CommandLine, TEXT("-DDC-ReplayLoadRemovePolicy="), FlagNamesToRemove))
+	{
+		TryLexFromString(PolicyFilter.FlagsToRemove, FlagNamesToRemove);
+	}
+
+	return PolicyFilter;
+}
+
+static FCacheKeyFilter ParseReplayKeyFilter(const TCHAR* const CommandLine)
+{
+	const bool bDefaultMatch = String::FindFirst(CommandLine, TEXT("-DDC-ReplayTypes="), ESearchCase::IgnoreCase) == INDEX_NONE;
+	float DefaultRate = bDefaultMatch ? 100.0f : 0.0f;
+	FParse::Value(CommandLine, TEXT("-DDC-ReplayRate="), DefaultRate);
+
+	FCacheKeyFilter KeyFilter = FCacheKeyFilter::Parse(CommandLine, TEXT("-DDC-ReplayTypes="), DefaultRate);
+
+	if (KeyFilter)
+	{
+		uint32 Salt;
+		if (FParse::Value(CommandLine, TEXT("-DDC-ReplaySalt="), Salt))
+		{
+			if (Salt == 0)
+			{
+				UE_LOG(LogDerivedDataCache, Warning,
+					TEXT("Replay: Ignoring salt of 0. The salt must be a positive integer."));
+			}
+			else
+			{
+				KeyFilter.SetSalt(Salt);
+			}
+		}
+
+		UE_LOG(LogDerivedDataCache, Display,
+			TEXT("Replay: Using salt -DDC-ReplaySalt=%u to filter cache keys to replay."), KeyFilter.GetSalt());
+	}
+
+	return KeyFilter;
+}
+
+///////////////////////////////////////////////////////////////////////////////////////////////////////////////////////
+
+ILegacyCacheStore* TryCreateCacheStoreReplay(ILegacyCacheStore* InnerCache)
+{
+	const TCHAR* const CommandLine = FCommandLine::Get();
+	const bool bHasReplayLoad = String::FindFirst(CommandLine, TEXT("-DDC-ReplayLoad=")) != INDEX_NONE;
+
+	FString ReplaySavePath;
+	if (!FParse::Value(CommandLine, TEXT("-DDC-ReplaySave="), ReplaySavePath) && !bHasReplayLoad)
+	{
+		return nullptr;
+	}
+
+	ILegacyCacheStore* ReplayTarget = InnerCache;
+	ILegacyCacheStore* ReplayStore = nullptr;
+
+	const FCacheKeyFilter KeyFilter = ParseReplayKeyFilter(CommandLine);
+	const FCacheMethodFilter MethodFilter = FCacheMethodFilter::Parse(CommandLine);
+	const FCachePolicyFilter PolicyFilter = FCachePolicyFilter::Parse(CommandLine);
+
+	// Create the replay cache store to save requests that pass the filters.
+	if (!ReplaySavePath.IsEmpty())
+	{
+		const uint64 BlockSize = FParse::Param(CommandLine, TEXT("DDC-ReplayCompress")) ? GCacheReplayCompressionBlockSize : 0;
+		ReplayTarget = ReplayStore = new FCacheStoreReplay(InnerCache, KeyFilter, MethodFilter, *ReplaySavePath, BlockSize);
+	}
+
+	if (bHasReplayLoad)
+	{
+		// Allow the captured priority to be overridden on the command line.
+		const EPriority* ForcedLoadPriority = nullptr;
+		EPriority ReplayLoadPriority = EPriority::Lowest;
+		FString ReplayLoadPriorityName;
+		if (FParse::Value(CommandLine, TEXT("-DDC-ReplayLoadPriority="), ReplayLoadPriorityName) &&
+			TryLexFromString(ReplayLoadPriority, ReplayLoadPriorityName))
+		{
+			ForcedLoadPriority = &ReplayLoadPriority;
+		}
+
+		// Load every cache replay file that was requested on the command line.
+		FCacheReplayReader Reader(ReplayTarget, KeyFilter, MethodFilter, PolicyFilter, ForcedLoadPriority);
+		const TCHAR* Tokens = CommandLine;
+		for (FString Token; FParse::Token(Tokens, Token, /*UseEscape*/ false);)
+		{
+			FString ReplayLoadPath;
+			if (FParse::Value(*Token, TEXT("-DDC-ReplayLoad="), ReplayLoadPath))
+			{
+				Reader.ReadFromFile(*ReplayLoadPath);
+			}
+		}
+	}
+
+	return ReplayStore;
+}
+
+bool TryLoadCacheStoreReplay(ILegacyCacheStore* TargetCache, const TCHAR* ReplayPath, const TCHAR* PriorityName)
+{
+	const EPriority* ForcedPriority = nullptr;
+	EPriority Priority = EPriority::Lowest;
+	if (PriorityName && TryLexFromString(Priority, PriorityName))
+	{
+		ForcedPriority = &Priority;
+	}
+
+	const TCHAR* const CommandLine = FCommandLine::Get();
+	FCacheKeyFilter KeyFilter = ParseReplayKeyFilter(CommandLine);
+	FCacheMethodFilter MethodFilter = FCacheMethodFilter::Parse(CommandLine);
+	FCachePolicyFilter PolicyFilter = FCachePolicyFilter::Parse(CommandLine);
+	FCacheReplayReader Reader(
+		TargetCache,
+		MoveTemp(KeyFilter),
+		MoveTemp(MethodFilter),
+		MoveTemp(PolicyFilter),
+		ForcedPriority);
+	return Reader.ReadFromFile(ReplayPath);
+}
+
+} // UE::DerivedData
