DDC: Extended cache replay to read replay files async and with a more configurable console command

#rnx

diff --git a/Engine/Source/Developer/DerivedDataCache/Internal/DerivedDataCacheMethod.h b/Engine/Source/Developer/DerivedDataCache/Internal/DerivedDataCacheMethod.h
new file mode 100644
--- /dev/null
+++ b/Engine/Source/Developer/DerivedDataCache/Internal/DerivedDataCacheMethod.h
@@ -0,0 +1,52 @@
+// Copyright Epic Games, Inc. All Rights Reserved.
+
+#pragma once
+
+#include "CoreTypes.h"
+#include "Containers/StringFwd.h"
+
+#define UE_API DERIVEDDATACACHE_API
+
+class FCbFieldView;
+class FCbWriter;
+
+namespace UE::DerivedData
+{
+
+enum class ECacheMethod : uint8
+{
+	Put,
+	Get,
+	PutValue,
+	GetValue,
+	GetChunks,
+};
+
+UE_API FAnsiStringBuilderBase& operator<<(FAnsiStringBuilderBase& Builder, ECacheMethod Method);
+UE_API FWideStringBuilderBase& operator<<(FWideStringBuilderBase& Builder, ECacheMethod Method);
+UE_API FUtf8StringBuilderBase& operator<<(FUtf8StringBuilderBase& Builder, ECacheMethod Method);
+
+UE_API bool TryLexFromString(ECacheMethod& OutMethod, FUtf8StringView String);
+UE_API bool TryLexFromString(ECacheMethod& OutMethod, FWideStringView String);
+
+UE_API FCbWriter& operator<<(FCbWriter& Writer, ECacheMethod Method);
+UE_API bool LoadFromCompactBinary(FCbFieldView Field, ECacheMethod& OutMethod);
+
+class FCacheMethodFilter
+{
+public:
+	/** Parse from method names delimited by ',' or '+'. */
+	UE_API static FCacheMethodFilter Parse(FStringView MethodNames);
+
+	inline bool IsMatch(ECacheMethod Method) const
+	{
+		return (MethodMask & (1 << uint32(Method))) == 0;
+	}
+
+private:
+	uint32 MethodMask = 0;
+};
+
+} // UE::DerivedData
+
+#undef UE_API
diff --git a/Engine/Source/Developer/DerivedDataCache/Internal/DerivedDataCacheReplay.h b/Engine/Source/Developer/DerivedDataCache/Internal/DerivedDataCacheReplay.h
new file mode 100644
--- /dev/null
+++ b/Engine/Source/Developer/DerivedDataCache/Internal/DerivedDataCacheReplay.h
@@ -0,0 +1,47 @@
+// Copyright Epic Games, Inc. All Rights Reserved.
+
+#pragma once
+
+#include "CoreTypes.h"
+#include "Templates/PimplPtr.h"
+
+#define UE_API DERIVEDDATACACHE_API
+
+class FArchive;
+class FCbObjectView;
+
+namespace UE::DerivedData { class FCacheKeyFilter; }
+namespace UE::DerivedData { class FCacheMethodFilter; }
+namespace UE::DerivedData { class ILegacyCacheStore; }
+namespace UE::DerivedData { enum class ECacheMethod : uint8; }
+namespace UE::DerivedData { enum class ECachePolicy : uint32; }
+namespace UE::DerivedData { enum class EPriority : uint8; }
+
+namespace UE::DerivedData
+{
+
+class FCacheReplayReader
+{
+public:
+	static constexpr uint64 DefaultScratchSize = 1024;
+
+	UE_API explicit FCacheReplayReader(ILegacyCacheStore* TargetCache);
+
+	UE_API void SetKeyFilter(FCacheKeyFilter KeyFilter);
+	UE_API void SetMethodFilter(FCacheMethodFilter MethodFilter);
+	UE_API void SetPolicyTransform(ECachePolicy AddFlags, ECachePolicy RemoveFlags);
+	UE_API void SetPriorityOverride(EPriority Priority);
+
+	UE_API void ReadFromFileAsync(const TCHAR* ReplayPath, uint64 ScratchSize = DefaultScratchSize);
+	UE_API bool ReadFromFile(const TCHAR* ReplayPath, uint64 ScratchSize = DefaultScratchSize);
+	UE_API bool ReadFromArchive(FArchive& ReplayAr, uint64 ScratchSize = DefaultScratchSize);
+	UE_API bool ReadFromObject(const FCbObject& Object);
+
+private:
+	class FState;
+	TPimplPtr<FState> State;
+};
+
+} // UE::DerivedData
+
+#undef UE_API
diff --git a/Engine/Source/Developer/DerivedDataCache/Private/DerivedDataBackends.cpp b/Engine/Source/Developer/DerivedDataCache/Private/DerivedDataBackends.cpp
--- a/Engine/Source/Developer/DerivedDataCache/Private/DerivedDataBackends.cpp
+++ b/Engine/Source/Developer/DerivedDataCache/Private/DerivedDataBackends.cpp
@@ -3,7 +3,9 @@
 #include "CoreTypes.h"
 #include "Containers/StringView.h"
 #include "DerivedDataBackendInterface.h"
+#include "DerivedDataCacheMethod.h"
 #include "DerivedDataCachePrivate.h"
+#include "DerivedDataCacheReplay.h"
 #include "DerivedDataCacheStore.h"
 #include "DerivedDataCacheUsageStats.h"
 #include "DerivedDataRequestOwner.h"
@@ -58,7 +60,6 @@ IPakFileCacheStore* CreatePakFileCacheStore(const TCHAR* Filename, bool bWriting
 ILegacyCacheStore* CreateS3CacheStore(const TCHAR* RootManifestPath, const TCHAR* BaseUrl, const TCHAR* Region, const TCHAR* CanaryObjectKey, const TCHAR* CachePath);
 TTuple<ILegacyCacheStore*, ECacheStoreFlags> CreateZenCacheStore(const TCHAR* NodeName, const TCHAR* Config);
 ILegacyCacheStore* TryCreateCacheStoreReplay(ILegacyCacheStore* InnerCache);
-bool TryLoadCacheStoreReplay(ILegacyCacheStore* TargetCache, const TCHAR* ReplayPath, const TCHAR* PriorityName = nullptr);
 
 /**
  * This class is used to create a singleton that represents the derived data cache hierarchy and all of the wrappers necessary
@@ -870,6 +871,7 @@ public:
 	virtual ~FDerivedDataBackendGraph()
 	{
 		check(StaticGraph == this);
+		Replays.Empty();
 		RootCache = nullptr;
 		DestroyCreatedBackends();
 		StaticGraph = nullptr;
@@ -1103,17 +1105,84 @@ private:
 	{
 		if (Args.Num() < 1)
 		{
-			UE_LOG(LogDerivedDataCache, Log, TEXT("Usage: DDC.LoadReplay ReplayPath [Priority=<Name>]"));
+			UE_LOG(LogDerivedDataCache, Log, TEXT("Usage: DDC.LoadReplay ReplayPath"
+				" [Methods=Get+GetValue+GetChunks]"
+				" [Rate=<0-100>]"
+				" [Types=Type1[@Rate1][+Type2[@Rate2]]..."
+				" [Salt=PositiveInt32]"
+				" [Priority=<Lowest-Blocking>]"
+				" [AddPolicy=Query,SkipData]"
+				" [RemovePolicy=SkipMeta]"));
 			return;
 		}
 
+		TStringBuilder<512> JoinedArgs;
+		JoinedArgs.Join(MakeArrayView(Args).RightChop(1), TEXT(' '));
+
+		FCacheReplayReader Replay(RootCache);
+
+		// Parse Key Filter
+		const bool bDefaultMatch = String::FindFirst(*JoinedArgs, TEXT("Types="), ESearchCase::IgnoreCase) == INDEX_NONE;
+		float DefaultRate = bDefaultMatch ? 100.0f : 0.0f;
+		FParse::Value(*JoinedArgs, TEXT("Rate="), DefaultRate);
+
+		FCacheKeyFilter KeyFilter = FCacheKeyFilter::Parse(*JoinedArgs, TEXT("Types="), DefaultRate);
+
+		if (KeyFilter)
+		{
+			uint32 Salt;
+			if (FParse::Value(*JoinedArgs, TEXT("Salt="), Salt))
+			{
+				if (Salt == 0)
+				{
+					UE_LOG(LogDerivedDataCache, Warning,
+						TEXT("Replay: Ignoring salt of 0. The salt must be a positive integer."));
+				}
+				else
+				{
+					KeyFilter.SetSalt(Salt);
+				}
+			}
+
+			UE_LOG(LogDerivedDataCache, Display,
+				TEXT("Replay: Using salt %u to filter cache keys to replay."), KeyFilter.GetSalt());
+		}
+
+		Replay.SetKeyFilter(MoveTemp(KeyFilter));
+
+		// Parse Method Filter
+		FString MethodNames;
+		if (FParse::Value(*JoinedArgs, TEXT("Methods="), MethodNames))
+		{
+			Replay.SetMethodFilter(FCacheMethodFilter::Parse(MethodNames));
+		}
+
+		// Parse Policy Transform
+		ECachePolicy FlagsToAdd = ECachePolicy::None;
+		FString FlagNamesToAdd;
+		if (FParse::Value(*JoinedArgs, TEXT("AddPolicy="), FlagNamesToAdd))
+		{
+			TryLexFromString(FlagsToAdd, FlagNamesToAdd);
+		}
+		ECachePolicy FlagsToRemove = ECachePolicy::None;
+		FString FlagNamesToRemove;
+		if (FParse::Value(*JoinedArgs, TEXT("RemovePolicy="), FlagNamesToRemove))
+		{
+			TryLexFromString(FlagsToRemove, FlagNamesToRemove);
+		}
+		Replay.SetPolicyTransform(FlagsToAdd, FlagsToRemove);
+
+		// Parse Priority Override
+		EPriority Priority{};
 		FString PriorityName;
-		for (int32 ArgIndex = 1; ArgIndex < Args.Num(); ++ArgIndex)
+		if (FParse::Value(*JoinedArgs, TEXT("Priority="), PriorityName) &&
+			TryLexFromString(Priority, PriorityName))
 		{
-			FParse::Value(*Args[ArgIndex], TEXT("Priority="), PriorityName);
+			Replay.SetPriorityOverride(Priority);
 		}
 
-		TryLoadCacheStoreReplay(RootCache, *Args[0], *PriorityName);
+		Replay.ReadFromFileAsync(*Args[0]);
+		Replays.Add(MoveTemp(Replay));
 	}
 
 	static inline FDerivedDataBackendGraph*			StaticGraph;
@@ -1160,6 +1229,7 @@ private:
 	FAutoConsoleCommand UnmountPakCommand;
 
 	FAutoConsoleCommand LoadReplayCommand;
+	TArray<FCacheReplayReader> Replays;
 };
 
 } // UE::DerivedData
diff --git a/Engine/Source/Developer/DerivedDataCache/Private/DerivedDataCacheKeyFilter.cpp b/Engine/Source/Developer/DerivedDataCache/Private/DerivedDataCacheKeyFilter.cpp
--- a/Engine/Source/Developer/DerivedDataCache/Private/DerivedDataCacheKeyFilter.cpp
+++ b/Engine/Source/Developer/DerivedDataCache/Private/DerivedDataCacheKeyFilter.cpp
@@ -76,7 +76,7 @@ FCacheKeyFilter FCacheKeyFilter::Parse(const TCHAR* const Config, const TCHAR* c
 	{
 		const int32 PrefixIndex = String::FindFirst(ConfigView, PrefixView, ESearchCase::IgnoreCase);
 		ConfigView.RightChopInline(PrefixIndex + PrefixView.Len() + TypeConfigArray.Len());
-		String::ParseTokens(TypeConfigArray, TEXT('+'), [&Types](const FStringView TypeConfig)
+		String::ParseTokensMultiple(TypeConfigArray, {TEXT('+'), TEXT(',')}, [&Types](const FStringView TypeConfig)
 		{
 			FCacheKeyFilterState::FTypeRate TypeRate;
 			if (FCacheKeyFilterState::TryParseTypeRate(TypeConfig, TypeRate))
diff --git a/Engine/Source/Developer/DerivedDataCache/Private/DerivedDataCacheMethod.cpp b/Engine/Source/Developer/DerivedDataCache/Private/DerivedDataCacheMethod.cpp
new file mode 100644
--- /dev/null
+++ b/Engine/Source/Developer/DerivedDataCache/Private/DerivedDataCacheMethod.cpp
@@ -0,0 +1,103 @@
+// Copyright Epic Games, Inc. All Rights Reserved.
+
+#include "DerivedDataCacheMethod.h"
+
+#include "Containers/StringConv.h"
+#include "Containers/StringView.h"
+#include "Misc/StringBuilder.h"
+#include "Serialization/CompactBinary.h"
+#include "Serialization/CompactBinaryWriter.h"
+#include "String/ParseTokens.h"
+
+namespace UE::DerivedData::Private
+{
+
+template <typename CharType>
+static TStringBuilderBase<CharType>& MethodToString(TStringBuilderBase<CharType>& Builder, const ECacheMethod Method)
+{
+	switch (Method)
+	{
+	case ECacheMethod::Put:       return Builder << ANSITEXTVIEW("Put");
+	case ECacheMethod::Get:       return Builder << ANSITEXTVIEW("Get");
+	case ECacheMethod::PutValue:  return Builder << ANSITEXTVIEW("PutValue");
+	case ECacheMethod::GetValue:  return Builder << ANSITEXTVIEW("GetValue");
+	case ECacheMethod::GetChunks: return Builder << ANSITEXTVIEW("GetChunks");
+	}
+	return Builder << ANSITEXTVIEW("Unknown");
+}
+
+template <typename CharType>
+static bool MethodFromString(ECacheMethod& OutMethod, const TStringView<CharType> String)
+{
+	const auto ConvertedString = StringCast<UTF8CHAR, 16>(String.GetData(), String.Len());
+	if (ConvertedString == UTF8TEXTVIEW("Put"))
+	{
+		OutMethod = ECacheMethod::Put;
+	}
+	else if (ConvertedString == UTF8TEXTVIEW("Get"))
+	{
+		OutMethod = ECacheMethod::Get;
+	}
+	else if (ConvertedString == UTF8TEXTVIEW("PutValue"))
+	{
+		OutMethod = ECacheMethod::PutValue;
+	}
+	else if (ConvertedString == UTF8TEXTVIEW("GetValue"))
+	{
+		OutMethod = ECacheMethod::GetValue;
+	}
+	else if (ConvertedString == UTF8TEXTVIEW("GetChunks"))
+	{
+		OutMethod = ECacheMethod::GetChunks;
+	}
+	else
+	{
+		return false;
+	}
+	return true;
+}
+
+} // UE::DerivedData::Private
+
+namespace UE::DerivedData
+{
+
+FAnsiStringBuilderBase& operator<<(FAnsiStringBuilderBase& Builder, ECacheMethod Method) { return Private::MethodToString(Builder, Method); }
+FWideStringBuilderBase& operator<<(FWideStringBuilderBase& Builder, ECacheMethod Method) { return Private::MethodToString(Builder, Method); }
+FUtf8StringBuilderBase& operator<<(FUtf8StringBuilderBase& Builder, ECacheMethod Method) { return Private::MethodToString(Builder, Method); }
+
+bool TryLexFromString(ECacheMethod& OutMethod, FUtf8StringView String) { return Private::MethodFromString(OutMethod, String); }
+bool TryLexFromString(ECacheMethod& OutMethod, FWideStringView String) { return Private::MethodFromString(OutMethod, String); }
+
+FCbWriter& operator<<(FCbWriter& Writer, const ECacheMethod Method)
+{
+	Writer.AddString(WriteToUtf8String<16>(Method));
+	return Writer;
+}
+
+bool LoadFromCompactBinary(FCbFieldView Field, ECacheMethod& OutMethod)
+{
+	if (TryLexFromString(OutMethod, Field.AsString()))
+	{
+		return true;
+	}
+	OutMethod = {};
+	return false;
+}
+
+FCacheMethodFilter FCacheMethodFilter::Parse(const FStringView MethodNames)
+{
+	FCacheMethodFilter MethodFilter;
+	MethodFilter.MethodMask = ~uint32(0);
+	String::ParseTokensMultiple(MethodNames, {TEXT('+'), TEXT(',')}, [&MethodFilter](FStringView MethodName)
+	{
+		ECacheMethod Method;
+		if (TryLexFromString(Method, MethodName))
+		{
+			MethodFilter.MethodMask &= ~(1 << uint32(Method));
+		}
+	});
+	return MethodFilter;
+}
+
+} // UE::DerivedData
diff --git a/Engine/Source/Developer/DerivedDataCache/Private/DerivedDataCacheReplay.cpp b/Engine/Source/Developer/DerivedDataCache/Private/DerivedDataCacheReplay.cpp
new file mode 100644
--- /dev/null
+++ b/Engine/Source/Developer/DerivedDataCache/Private/DerivedDataCacheReplay.cpp
@@ -0,0 +1,765 @@
+// Copyright Epic Games, Inc. All Rights Reserved.
+
+#include "DerivedDataCacheReplay.h"
+
+#include "Algo/AnyOf.h"
+#include "Compression/CompressedBuffer.h"
+#include "Containers/StringView.h"
+#include "DerivedDataCacheKeyFilter.h"
+#include "DerivedDataCacheMethod.h"
+#include "DerivedDataLegacyCacheStore.h"
+#include "DerivedDataRequestOwner.h"
+#include "DerivedDataValue.h"
+#include "HAL/CriticalSection.h"
+#include "HAL/FileManager.h"
+#include "HAL/PlatformTime.h"
+#include "Misc/CommandLine.h"
+#include "Misc/Optional.h"
+#include "Misc/Parse.h"
+#include "Misc/ScopeLock.h"
+#include "ProfilingDebugging/CpuProfilerTrace.h"
+#include "Serialization/CompactBinary.h"
+#include "Serialization/CompactBinarySerialization.h"
+#include "Serialization/CompactBinaryWriter.h"
+#include "Serialization/MemoryReader.h"
+#include "String/Find.h"
+#include "Tasks/Pipe.h"
+#include "Tasks/Task.h"
+#include "Templates/Invoke.h"
+
+namespace UE::DerivedData
+{
+
+static constexpr uint64 GCacheReplayCompressionBlockSize = 256 * 1024;
+
+///////////////////////////////////////////////////////////////////////////////////////////////////////////////////////
+
+class FCacheStoreReplay final : public ILegacyCacheStore
+{
+public:
+	FCacheStoreReplay(
+		ILegacyCacheStore* InnerCache,
+		FCacheKeyFilter KeyFilter,
+		FCacheMethodFilter MethodFilter,
+		FString&& ReplayPath,
+		uint64 CompressionBlockSize = 0);
+
+	~FCacheStoreReplay();
+
+	void Put(
+		const TConstArrayView<FCachePutRequest> Requests,
+		IRequestOwner& Owner,
+		FOnCachePutComplete&& OnComplete) final;
+	void Get(
+		const TConstArrayView<FCacheGetRequest> Requests,
+		IRequestOwner& Owner,
+		FOnCacheGetComplete&& OnComplete) final;
+	void PutValue(
+		const TConstArrayView<FCachePutValueRequest> Requests,
+		IRequestOwner& Owner,
+		FOnCachePutValueComplete&& OnComplete) final;
+	void GetValue(
+		const TConstArrayView<FCacheGetValueRequest> Requests,
+		IRequestOwner& Owner,
+		FOnCacheGetValueComplete&& OnComplete) final;
+	void GetChunks(
+		const TConstArrayView<FCacheGetChunkRequest> Requests,
+		IRequestOwner& Owner,
+		FOnCacheGetChunkComplete&& OnComplete) final;
+
+	void LegacyStats(FDerivedDataCacheStatsNode& OutNode) final
+	{
+		InnerCache->LegacyStats(OutNode);
+	}
+
+	bool LegacyDebugOptions(FBackendDebugOptions& Options) final
+	{
+		return InnerCache->LegacyDebugOptions(Options);
+	}
+
+	void SetReader(FCacheReplayReader&& Reader)
+	{
+		ReplayReader = MoveTemp(Reader);
+	}
+
+private:
+	template <typename RequestType>
+	void SerializeRequests(TConstArrayView<RequestType> Requests, ECacheMethod Method, EPriority Priority);
+
+	void WriteBinaryToArchive(const FCompositeBuffer& RawBinary);
+	void WriteToArchive(FCbWriter& Writer);
+	void FlushToArchive();
+
+	ILegacyCacheStore* InnerCache;
+	FCacheKeyFilter KeyFilter;
+	FCacheMethodFilter MethodFilter;
+	FString ReplayPath;
+	TUniquePtr<FArchive> ReplayAr;
+	FUniqueBuffer RawBlock;
+	FMutableMemoryView RawBlockTail;
+	FCriticalSection Lock;
+	TOptional<FCacheReplayReader> ReplayReader;
+};
+
+FCacheStoreReplay::FCacheStoreReplay(
+	ILegacyCacheStore* InInnerCache,
+	FCacheKeyFilter InKeyFilter,
+	FCacheMethodFilter InMethodFilter,
+	FString&& InReplayPath,
+	uint64 CompressionBlockSize)
+	: InnerCache(InInnerCache)
+	, KeyFilter(MoveTemp(InKeyFilter))
+	, MethodFilter(MoveTemp(InMethodFilter))
+	, ReplayPath(MoveTemp(InReplayPath))
+{
+	if (!ReplayPath.IsEmpty())
+	{
+		ReplayAr.Reset(IFileManager::Get().CreateFileWriter(*ReplayPath, FILEWRITE_NoFail));
+	}
+	if (CompressionBlockSize)
+	{
+		RawBlock = FUniqueBuffer::Alloc(CompressionBlockSize);
+		RawBlockTail = RawBlock;
+	}
+	UE_CLOG(ReplayAr, LogDerivedDataCache, Display, TEXT("Replay: Saving cache replay to '%s'"), *ReplayPath);
+}
+
+FCacheStoreReplay::~FCacheStoreReplay()
+{
+	FlushToArchive();
+}
+
+template <typename RequestType>
+void FCacheStoreReplay::SerializeRequests(
+	const TConstArrayView<RequestType> Requests,
+	const ECacheMethod Method,
+	const EPriority Priority)
+{
+	if (!ReplayAr)
+	{
+		return;
+	}
+
+	TRACE_CPUPROFILER_EVENT_SCOPE(ReplayDDC_Serialize);
+
+	const auto IsKeyMatch = [this](const RequestType& Request) { return KeyFilter.IsMatch(Request.Key); };
+	if (!MethodFilter.IsMatch(Method) || !Algo::AnyOf(Requests, IsKeyMatch))
+	{
+		return;
+	}
+
+	TCbWriter<512> Writer;
+	Writer.BeginObject();
+	Writer << ANSITEXTVIEW("Method") << Method;
+	Writer << ANSITEXTVIEW("Priority") << Priority;
+	Writer.BeginArray(ANSITEXTVIEW("Requests"));
+	for (const RequestType& Request : Requests)
+	{
+		if (KeyFilter.IsMatch(Request.Key))
+		{
+			Writer << Request;
+		}
+	}
+	Writer.EndArray();
+	Writer.EndObject();
+
+	if (UE_LOG_ACTIVE(LogDerivedDataCache, Verbose))
+	{
+		TUtf8StringBuilder<1024> Batch;
+		CompactBinaryToCompactJson(Writer.Save().AsObject(), Batch);
+		UE_LOG(LogDerivedDataCache, Verbose, TEXT("Replay: %hs"), *Batch);
+	}
+
+	WriteToArchive(Writer);
+}
+
+void FCacheStoreReplay::WriteBinaryToArchive(const FCompositeBuffer& RawBinary)
+{
+	const FValue CompressedBinary = FValue::Compress(RawBinary, RawBlock.GetSize());
+	TCbWriter<64> BinaryWriter;
+	BinaryWriter.AddBinary(CompressedBinary.GetData().GetCompressed());
+	BinaryWriter.Save(*ReplayAr);
+}
+
+void FCacheStoreReplay::WriteToArchive(FCbWriter& Writer)
+{
+	FScopeLock ScopeLock(&Lock);
+	if (RawBlock)
+	{
+		const uint64 SaveSize = Writer.GetSaveSize();
+		if (RawBlockTail.GetSize() < SaveSize)
+		{
+			FlushToArchive();
+		}
+		if (RawBlockTail.GetSize() < SaveSize)
+		{
+			WriteBinaryToArchive(Writer.Save().AsObject().GetBuffer());
+		}
+		else
+		{
+			Writer.Save(RawBlockTail.Left(SaveSize));
+			RawBlockTail += SaveSize;
+		}
+	}
+	else
+	{
+		Writer.Save(*ReplayAr);
+	}
+}
+
+void FCacheStoreReplay::FlushToArchive()
+{
+	const FSharedBuffer RawBlockHead = FSharedBuffer::MakeView(RawBlock.GetView().LeftChop(RawBlockTail.GetSize()));
+	if (RawBlockHead.GetSize() > 0)
+	{
+		WriteBinaryToArchive(FCompositeBuffer(RawBlockHead));
+		RawBlockTail = RawBlock;
+	}
+}
+
+void FCacheStoreReplay::Put(
+	const TConstArrayView<FCachePutRequest> Requests,
+	IRequestOwner& Owner,
+	FOnCachePutComplete&& OnComplete)
+{
+	InnerCache->Put(Requests, Owner, MoveTemp(OnComplete));
+}
+
+void FCacheStoreReplay::Get(
+	const TConstArrayView<FCacheGetRequest> Requests,
+	IRequestOwner& Owner,
+	FOnCacheGetComplete&& OnComplete)
+{
+	SerializeRequests(Requests, ECacheMethod::Get, Owner.GetPriority());
+	InnerCache->Get(Requests, Owner, MoveTemp(OnComplete));
+}
+
+void FCacheStoreReplay::PutValue(
+	const TConstArrayView<FCachePutValueRequest> Requests,
+	IRequestOwner& Owner,
+	FOnCachePutValueComplete&& OnComplete)
+{
+	InnerCache->PutValue(Requests, Owner, MoveTemp(OnComplete));
+}
+
+void FCacheStoreReplay::GetValue(
+	const TConstArrayView<FCacheGetValueRequest> Requests,
+	IRequestOwner& Owner,
+	FOnCacheGetValueComplete&& OnComplete)
+{
+	SerializeRequests(Requests, ECacheMethod::GetValue, Owner.GetPriority());
+	InnerCache->GetValue(Requests, Owner, MoveTemp(OnComplete));
+}
+
+void FCacheStoreReplay::GetChunks(
+	const TConstArrayView<FCacheGetChunkRequest> Requests,
+	IRequestOwner& Owner,
+	FOnCacheGetChunkComplete&& OnComplete)
+{
+	SerializeRequests(Requests, ECacheMethod::GetChunks, Owner.GetPriority());
+	InnerCache->GetChunks(Requests, Owner, MoveTemp(OnComplete));
+}
+
+///////////////////////////////////////////////////////////////////////////////////////////////////////////////////////
+
+class FCacheReplayReader::FState
+{
+public:
+	~FState();
+
+	void WaitForAsyncReads();
+
+	void ReadFromFileAsync(const TCHAR* ReplayPath, uint64 ScratchSize);
+	bool ReadFromFile(const TCHAR* ReplayPath, uint64 ScratchSize);
+	bool ReadFromArchive(FArchive& ReplayAr, uint64 ScratchSize);
+	bool ReadFromObject(const FCbObject& Object);
+
+	static_assert(uint8(EPriority::Lowest) == 0);
+	static_assert(uint8(EPriority::Low) == 1);
+	static_assert(uint8(EPriority::Normal) == 2);
+	static_assert(uint8(EPriority::High) == 3);
+	static_assert(uint8(EPriority::Highest) == 4);
+	static_assert(uint8(EPriority::Blocking) == 5);
+
+	ILegacyCacheStore* TargetCache = nullptr;
+	FCacheKeyFilter KeyFilter;
+	FCacheMethodFilter MethodFilter;
+	ECachePolicy PolicyFlagsToAdd = ECachePolicy::None;
+	ECachePolicy PolicyFlagsToRemove = ECachePolicy::None;
+	FRequestOwner Owners[6]
+	{
+		FRequestOwner(EPriority::Lowest),
+		FRequestOwner(EPriority::Low),
+		FRequestOwner(EPriority::Normal),
+		FRequestOwner(EPriority::High),
+		FRequestOwner(EPriority::Highest),
+		FRequestOwner(EPriority::Blocking),
+	};
+
+private:
+	template <typename RequestType, typename FunctionType>
+	bool DispatchRequests(FCbObjectView Object, ECacheMethod Method, FunctionType Function);
+
+	void ApplyPolicyTransform(FCacheGetRequest& Request);
+	void ApplyPolicyTransform(FCacheGetValueRequest& Request);
+	void ApplyPolicyTransform(FCacheGetChunkRequest& Request);
+
+	Tasks::FPipe ReadAsyncPipe{TEXT("CacheReplayReadAsync")};
+	TArray<Tasks::FTask> BlockingTasks;
+	int32 DispatchCount = 0;
+	int32 DispatchScope = 0;
+
+	class FDispatchScope;
+};
+
+class FCacheReplayReader::FState::FDispatchScope
+{
+public:
+	explicit FDispatchScope(FState& InState)
+		: State(InState)
+	{
+		if (State.DispatchScope++ == 0)
+		{
+			StartTime = FPlatformTime::Seconds();
+			StartDispatchCount = State.DispatchCount;
+		}
+	}
+
+	~FDispatchScope()
+	{
+		if (--State.DispatchScope == 0)
+		{
+			UE_LOG(LogDerivedDataCache, Display, TEXT("Replay: Dispatched %d requests in %.3lf seconds."),
+				State.DispatchCount - StartDispatchCount, FPlatformTime::Seconds() - StartTime);
+		}
+	}
+
+private:
+	FState& State;
+	double StartTime = 0.0;
+	int32 StartDispatchCount = 0;
+};
+
+FCacheReplayReader::FState::~FState()
+{
+	WaitForAsyncReads();
+
+	TRACE_CPUPROFILER_EVENT_SCOPE(ReplayDDC_RequestWait);
+	Tasks::Wait(BlockingTasks);
+
+	for (FRequestOwner& Owner : Owners)
+	{
+		Owner.Wait();
+	}
+}
+
+void FCacheReplayReader::FState::WaitForAsyncReads()
+{
+	if (ReadAsyncPipe.HasWork())
+	{
+		TRACE_CPUPROFILER_EVENT_SCOPE(ReplayDDC_DispatchWait);
+		ReadAsyncPipe.WaitUntilEmpty();
+	}
+}
+
+template <typename RequestType, typename FunctionType>
+bool FCacheReplayReader::FState::DispatchRequests(
+	const FCbObjectView Object,
+	const ECacheMethod Method,
+	const FunctionType Function)
+{
+	if (!MethodFilter.IsMatch(Method))
+	{
+		return true;
+	}
+
+	EPriority Priority = EPriority::Normal;
+	LoadFromCompactBinary(Object[ANSITEXTVIEW("Priority")], Priority, Priority);
+
+	TArray<RequestType, TInlineAllocator<16>> Requests;
+	{
+		TRACE_CPUPROFILER_EVENT_SCOPE(ReplayDDC_Serialize);
+		const FCbArrayView Array = Object[ANSITEXTVIEW("Requests")].AsArrayView();
+		Requests.Reserve(Array.Num());
+		for (FCbFieldView Field : Array)
+		{
+			RequestType& Request = Requests.AddDefaulted_GetRef();
+			if (!LoadFromCompactBinary(Field, Request))
+			{
+				return false;
+			}
+			if (KeyFilter.IsMatch(Request.Key))
+			{
+				ApplyPolicyTransform(Request);
+			}
+			else
+			{
+				Requests.Pop(/*bAllowShrinking*/ false);
+			}
+		}
+	}
+
+	if (!Requests.IsEmpty())
+	{
+		if (UE_LOG_ACTIVE(LogDerivedDataCache, Verbose))
+		{
+			TUtf8StringBuilder<1024> Batch;
+			CompactBinaryToCompactJson(Object, Batch);
+			UE_LOG(LogDerivedDataCache, Verbose, TEXT("Replay: %hs"), *Batch);
+		}
+
+		TRACE_CPUPROFILER_EVENT_SCOPE(ReplayDDC_Dispatch);
+
+		DispatchCount += Requests.Num();
+
+		FRequestOwner& Owner = Owners[uint8(Priority)];
+		if (Owner.GetPriority() < EPriority::Blocking)
+		{
+			// Owners with non-blocking priority can execute the request directly because it will be async.
+			FRequestBarrier Barrier(Owner);
+			Invoke(Function, TargetCache, Requests, Owner, [](auto&&){});
+		}
+		else
+		{
+			// Owners with blocking priority launch a task to execute the blocking request to allow concurrent replay.
+			BlockingTasks.Add(Tasks::Launch(TEXT("CacheReplayTask"), [this, Requests = MoveTemp(Requests), Function]
+			{
+				FRequestOwner BlockingOwner(EPriority::Blocking);
+				Invoke(Function, TargetCache, Requests, BlockingOwner, [](auto&&){});
+				BlockingOwner.Wait();
+			}));
+		}
+	}
+
+	return true;
+}
+
+void FCacheReplayReader::FState::ApplyPolicyTransform(FCacheGetRequest& Request)
+{
+	Request.Policy = Request.Policy.Transform([this](ECachePolicy Policy)
+	{
+		EnumAddFlags(Policy, PolicyFlagsToAdd);
+		EnumRemoveFlags(Policy, PolicyFlagsToRemove);
+		return Policy;
+	});
+}
+
+void FCacheReplayReader::FState::ApplyPolicyTransform(FCacheGetValueRequest& Request)
+{
+	EnumAddFlags(Request.Policy, PolicyFlagsToAdd);
+	EnumRemoveFlags(Request.Policy, PolicyFlagsToRemove);
+}
+
+void FCacheReplayReader::FState::ApplyPolicyTransform(FCacheGetChunkRequest& Request)
+{
+	EnumAddFlags(Request.Policy, PolicyFlagsToAdd);
+	EnumRemoveFlags(Request.Policy, PolicyFlagsToRemove);
+}
+
+void FCacheReplayReader::FState::ReadFromFileAsync(const TCHAR* const ReplayPath, const uint64 ScratchSize)
+{
+	ReadAsyncPipe.Launch(TEXT("CacheReplayReadFromFileAsync"), [this, ReplayPath = FString(ReplayPath), ScratchSize]
+	{
+		TRACE_CPUPROFILER_EVENT_SCOPE(ReplayDDC_ReadFromFileAsync);
+		ReadFromFile(*ReplayPath, ScratchSize);
+	});
+}
+
+bool FCacheReplayReader::FState::ReadFromFile(const TCHAR* const ReplayPath, const uint64 ScratchSize)
+{
+	TUniquePtr<FArchive> ReplayAr(IFileManager::Get().CreateFileReader(ReplayPath, FILEREAD_Silent));
+	if (!ReplayAr)
+	{
+		UE_LOG(LogDerivedDataCache, Warning, TEXT("Replay: File '%s' failed to open."), ReplayPath);
+		return false;
+	}
+
+	FDispatchScope Dispatch(*this);
+	UE_LOG(LogDerivedDataCache, Display, TEXT("Replay: Loading cache replay from '%s'"), ReplayPath);
+	return ReadFromArchive(*ReplayAr, FMath::Max(ScratchSize, GCacheReplayCompressionBlockSize));
+}
+
+bool FCacheReplayReader::FState::ReadFromArchive(FArchive& ReplayAr, const uint64 ScratchSize)
+{
+	FDispatchScope Dispatch(*this);
+
+	// A scratch buffer for the compact binary fields.
+	FUniqueBuffer Scratch = FUniqueBuffer::Alloc(ScratchSize);
+	const auto Alloc = [&Scratch](const uint64 Size) -> FUniqueBuffer
+	{
+		if (Size <= Scratch.GetSize())
+		{
+			return FUniqueBuffer::MakeView(Scratch.GetView().Left(Size));
+		}
+		return FUniqueBuffer::Alloc(Size);
+	};
+
+	// A scratch buffer for decompressed blocks of fields.
+	FUniqueBuffer BlockScratch;
+
+	for (int64 Offset = ReplayAr.Tell(); Offset < ReplayAr.TotalSize(); Offset = ReplayAr.Tell())
+	{
+		FCbField Field = LoadCompactBinary(ReplayAr, Alloc);
+		if (!Field)
+		{
+			UE_LOG(LogDerivedDataCache, Warning,
+				TEXT("Replay: Failed to load compact binary at offset %" INT64_FMT ". "
+					 "Archive is at offset %" INT64_FMT " and has total size %" INT64_FMT "."),
+				Offset, ReplayAr.Tell(), ReplayAr.TotalSize());
+			return false;
+		}
+
+		// A binary field is used to store a compressed buffer containing a sequence of compact binary objects.
+		if (Field.IsBinary())
+		{
+			FCompressedBuffer CompressedBuffer = FCompressedBuffer::FromCompressed(Field.AsBinary());
+			if (!CompressedBuffer)
+			{
+				UE_LOG(LogDerivedDataCache, Warning,
+					TEXT("Replay: Failed to load compressed buffer from binary field at offset %" INT64_FMT ". "
+						 "Archive is at offset %" INT64_FMT " and has total size %" INT64_FMT "."),
+					Offset, ReplayAr.Tell(), ReplayAr.TotalSize());
+				return false;
+			}
+
+			const uint64 RawBlockSize = CompressedBuffer.GetRawSize();
+			if (BlockScratch.GetSize() < RawBlockSize)
+			{
+				BlockScratch = FUniqueBuffer::Alloc(RawBlockSize);
+			}
+
+			const FMutableMemoryView RawBlockView = BlockScratch.GetView().Left(RawBlockSize);
+			if (!CompressedBuffer.TryDecompressTo(RawBlockView))
+			{
+				UE_LOG(LogDerivedDataCache, Warning,
+					TEXT("Replay: Failed to decompress compressed buffer from binary field at offset %" INT64_FMT ". "
+						 "Archive is at offset %" INT64_FMT " and has total size %" INT64_FMT "."),
+					Offset, ReplayAr.Tell(), ReplayAr.TotalSize());
+				return false;
+			}
+
+			FMemoryReaderView InnerAr(RawBlockView);
+			if (!ReadFromArchive(InnerAr, DefaultScratchSize))
+			{
+				return false;
+			}
+		}
+
+		// An object field is used to store one batch of cache requests.
+		if (Field.IsObject() && !ReadFromObject(Field.AsObject()))
+		{
+			UE_LOG(LogDerivedDataCache, Warning,
+				TEXT("Replay: Failed to load cache request from object field at offset %" INT64_FMT ". "
+						"Archive is at offset %" INT64_FMT " and has total size %" INT64_FMT "."),
+				Offset, ReplayAr.Tell(), ReplayAr.TotalSize());
+			return false;
+		}
+	}
+
+	return true;
+}
+
+bool FCacheReplayReader::FState::ReadFromObject(const FCbObject& Object)
+{
+	ECacheMethod Method{};
+	if (!LoadFromCompactBinary(Object[ANSITEXTVIEW("Method")], Method))
+	{
+		return false;
+	}
+
+	switch (Method)
+	{
+	case ECacheMethod::Get:
+		return DispatchRequests<FCacheGetRequest>(Object, Method, &ICacheStore::Get);
+	case ECacheMethod::GetValue:
+		return DispatchRequests<FCacheGetValueRequest>(Object, Method, &ICacheStore::GetValue);
+	case ECacheMethod::GetChunks:
+		return DispatchRequests<FCacheGetChunkRequest>(Object, Method, &ICacheStore::GetChunks);
+	}
+
+	return false;
+}
+
+FCacheReplayReader::FCacheReplayReader(ILegacyCacheStore* const TargetCache)
+	: State(MakePimpl<FState>())
+{
+	State->TargetCache = TargetCache;
+}
+
+void FCacheReplayReader::ReadFromFileAsync(const TCHAR* ReplayPath, const uint64 ScratchSize)
+{
+	return State->ReadFromFileAsync(ReplayPath, ScratchSize);
+}
+
+bool FCacheReplayReader::ReadFromFile(const TCHAR* ReplayPath, const uint64 ScratchSize)
+{
+	State->WaitForAsyncReads();
+	TRACE_CPUPROFILER_EVENT_SCOPE(ReplayDDC_ReadFromFile);
+	return State->ReadFromFile(ReplayPath, ScratchSize);
+}
+
+bool FCacheReplayReader::ReadFromArchive(FArchive& ReplayAr, const uint64 ScratchSize)
+{
+	State->WaitForAsyncReads();
+	TRACE_CPUPROFILER_EVENT_SCOPE(ReplayDDC_ReadFromArchive);
+	return State->ReadFromArchive(ReplayAr, ScratchSize);
+}
+
+bool FCacheReplayReader::ReadFromObject(const FCbObject& Object)
+{
+	State->WaitForAsyncReads();
+	TRACE_CPUPROFILER_EVENT_SCOPE(ReplayDDC_ReadFromObject);
+	return State->ReadFromObject(Object);
+}
+
+void FCacheReplayReader::SetKeyFilter(FCacheKeyFilter KeyFilter)
+{
+	State->KeyFilter = MoveTemp(KeyFilter);
+}
+
+void FCacheReplayReader::SetMethodFilter(FCacheMethodFilter MethodFilter)
+{
+	State->MethodFilter = MoveTemp(MethodFilter);
+}
+
+void FCacheReplayReader::SetPolicyTransform(ECachePolicy AddFlags, ECachePolicy RemoveFlags)
+{
+	State->PolicyFlagsToAdd = AddFlags;
+	State->PolicyFlagsToRemove = RemoveFlags;
+}
+
+void FCacheReplayReader::SetPriorityOverride(EPriority Priority)
+{
+	for (FRequestOwner& Owner : State->Owners)
+	{
+		Owner.SetPriority(Priority);
+	}
+}
+
+///////////////////////////////////////////////////////////////////////////////////////////////////////////////////////
+
+static FCacheMethodFilter ParseReplayMethodFilter(const TCHAR* const CommandLine)
+{
+	FCacheMethodFilter MethodFilter;
+	FString MethodNames;
+	if (FParse::Value(CommandLine, TEXT("-DDC-ReplayMethods="), MethodNames))
+	{
+		MethodFilter = FCacheMethodFilter::Parse(MethodNames);
+	}
+	return MethodFilter;
+}
+
+static FCacheKeyFilter ParseReplayKeyFilter(const TCHAR* const CommandLine)
+{
+	const bool bDefaultMatch = String::FindFirst(CommandLine, TEXT("-DDC-ReplayTypes="), ESearchCase::IgnoreCase) == INDEX_NONE;
+	float DefaultRate = bDefaultMatch ? 100.0f : 0.0f;
+	FParse::Value(CommandLine, TEXT("-DDC-ReplayRate="), DefaultRate);
+
+	FCacheKeyFilter KeyFilter = FCacheKeyFilter::Parse(CommandLine, TEXT("-DDC-ReplayTypes="), DefaultRate);
+
+	if (KeyFilter)
+	{
+		uint32 Salt;
+		if (FParse::Value(CommandLine, TEXT("-DDC-ReplaySalt="), Salt))
+		{
+			if (Salt == 0)
+			{
+				UE_LOG(LogDerivedDataCache, Warning,
+					TEXT("Replay: Ignoring salt of 0. The salt must be a positive integer."));
+			}
+			else
+			{
+				KeyFilter.SetSalt(Salt);
+			}
+		}
+
+		UE_LOG(LogDerivedDataCache, Display,
+			TEXT("Replay: Using salt -DDC-ReplaySalt=%u to filter cache keys to replay."), KeyFilter.GetSalt());
+	}
+
+	return KeyFilter;
+}
+
+static void ParseReplayPolicyTransform(const TCHAR* const CommandLine, FCacheReplayReader& Reader)
+{
+	ECachePolicy FlagsToAdd = ECachePolicy::None;
+	FString FlagNamesToAdd;
+	if (FParse::Value(CommandLine, TEXT("-DDC-ReplayLoadAddPolicy="), FlagNamesToAdd))
+	{
+		TryLexFromString(FlagsToAdd, FlagNamesToAdd);
+	}
+
+	ECachePolicy FlagsToRemove = ECachePolicy::None;
+	FString FlagNamesToRemove;
+	if (FParse::Value(CommandLine, TEXT("-DDC-ReplayLoadRemovePolicy="), FlagNamesToRemove))
+	{
+		TryLexFromString(FlagsToRemove, FlagNamesToRemove);
+	}
+
+	Reader.SetPolicyTransform(FlagsToAdd, FlagsToRemove);
+}
+
+///////////////////////////////////////////////////////////////////////////////////////////////////////////////////////
+
+ILegacyCacheStore* TryCreateCacheStoreReplay(ILegacyCacheStore* InnerCache)
+{
+	const TCHAR* const CommandLine = FCommandLine::Get();
+	const bool bHasReplayLoad = String::FindFirst(CommandLine, TEXT("-DDC-ReplayLoad=")) != INDEX_NONE;
+
+	FString ReplaySavePath;
+	if (!FParse::Value(CommandLine, TEXT("-DDC-ReplaySave="), ReplaySavePath) && !bHasReplayLoad)
+	{
+		return nullptr;
+	}
+
+	ILegacyCacheStore* ReplayTarget = InnerCache;
+	FCacheStoreReplay* ReplayStore = nullptr;
+
+	const FCacheKeyFilter KeyFilter = ParseReplayKeyFilter(CommandLine);
+	const FCacheMethodFilter MethodFilter = ParseReplayMethodFilter(CommandLine);
+
+	if (!ReplaySavePath.IsEmpty())
+	{
+		// Create the replay cache store to save requests that pass the filters.
+		const uint64 BlockSize = FParse::Param(CommandLine, TEXT("DDC-ReplayCompress")) ? GCacheReplayCompressionBlockSize : 0;
+		ReplayTarget = ReplayStore = new FCacheStoreReplay(InnerCache, KeyFilter, MethodFilter, MoveTemp(ReplaySavePath), BlockSize);
+	}
+	else
+	{
+		// Create a replay store to own the reader without saving a new replay.
+		ReplayStore = new FCacheStoreReplay(InnerCache, {}, {}, {}, {});
+	}
+
+	// Load every cache replay file that was requested on the command line.
+	if (bHasReplayLoad)
+	{
+		FCacheReplayReader Reader(ReplayTarget);
+		Reader.SetKeyFilter(KeyFilter);
+		Reader.SetMethodFilter(MethodFilter);
+		ParseReplayPolicyTransform(CommandLine, Reader);
+
+		EPriority ReplayLoadPriority = EPriority::Lowest;
+		FString ReplayLoadPriorityName;
+		if (FParse::Value(CommandLine, TEXT("-DDC-ReplayLoadPriority="), ReplayLoadPriorityName) &&
+			TryLexFromString(ReplayLoadPriority, ReplayLoadPriorityName))
+		{
+			Reader.SetPriorityOverride(ReplayLoadPriority);
+		}
+
+		const TCHAR* Tokens = CommandLine;
+		for (FString Token; FParse::Token(Tokens, Token, /*UseEscape*/ false);)
+		{
+			FString ReplayLoadPath;
+			if (FParse::Value(*Token, TEXT("-DDC-ReplayLoad="), ReplayLoadPath))
+			{
+				Reader.ReadFromFileAsync(*ReplayLoadPath);
+			}
+		}
+
+		ReplayStore->SetReader(MoveTemp(Reader));
+	}
+
+	return ReplayStore;
+}
+
+} // UE::DerivedData
diff --git a/Engine/Source/Developer/DerivedDataCache/Private/DerivedDataCacheStoreReplay.cpp b/Engine/Source/Developer/DerivedDataCache/Private/DerivedDataCacheStoreReplay.cpp
deleted file mode 100644
