DDC: Fixed a non-unity compile error

#rnx

diff --git a/Engine/Source/Developer/DerivedDataCache/Internal/DerivedDataCacheReplay.h b/Engine/Source/Developer/DerivedDataCache/Internal/DerivedDataCacheReplay.h
--- a/Engine/Source/Developer/DerivedDataCache/Internal/DerivedDataCacheReplay.h
+++ b/Engine/Source/Developer/DerivedDataCache/Internal/DerivedDataCacheReplay.h
@@ -35,7 +35,7 @@ public:
 	UE_API void ReadFromFileAsync(const TCHAR* ReplayPath, uint64 ScratchSize = DefaultScratchSize);
 	UE_API bool ReadFromFile(const TCHAR* ReplayPath, uint64 ScratchSize = DefaultScratchSize);
 	UE_API bool ReadFromArchive(FArchive& ReplayAr, uint64 ScratchSize = DefaultScratchSize);
-	UE_API bool ReadFromObject(const FCbObject& Object);
+	UE_API bool ReadFromObject(FCbObjectView Object);
 
 private:
 	class FState;
diff --git a/Engine/Source/Developer/DerivedDataCache/Private/DerivedDataCacheReplay.cpp b/Engine/Source/Developer/DerivedDataCache/Private/DerivedDataCacheReplay.cpp
--- a/Engine/Source/Developer/DerivedDataCache/Private/DerivedDataCacheReplay.cpp
+++ b/Engine/Source/Developer/DerivedDataCache/Private/DerivedDataCacheReplay.cpp
@@ -272,7 +272,7 @@ public:
 	void ReadFromFileAsync(const TCHAR* ReplayPath, uint64 ScratchSize);
 	bool ReadFromFile(const TCHAR* ReplayPath, uint64 ScratchSize);
 	bool ReadFromArchive(FArchive& ReplayAr, uint64 ScratchSize);
-	bool ReadFromObject(const FCbObject& Object);
+	bool ReadFromObject(FCbObjectView Object);
 
 	static_assert(uint8(EPriority::Lowest) == 0);
 	static_assert(uint8(EPriority::Low) == 1);
@@ -559,7 +559,7 @@ bool FCacheReplayReader::FState::ReadFromArchive(FArchive& ReplayAr, const uint6
 	return true;
 }
 
-bool FCacheReplayReader::FState::ReadFromObject(const FCbObject& Object)
+bool FCacheReplayReader::FState::ReadFromObject(const FCbObjectView Object)
 {
 	ECacheMethod Method{};
 	if (!LoadFromCompactBinary(Object[ANSITEXTVIEW("Method")], Method))
@@ -605,7 +605,7 @@ bool FCacheReplayReader::ReadFromArchive(FArchive& ReplayAr, const uint64 Scratc
 	return State->ReadFromArchive(ReplayAr, ScratchSize);
 }
 
-bool FCacheReplayReader::ReadFromObject(const FCbObject& Object)
+bool FCacheReplayReader::ReadFromObject(const FCbObjectView Object)
 {
 	State->WaitForAsyncReads();
 	TRACE_CPUPROFILER_EVENT_SCOPE(ReplayDDC_ReadFromObject);
