Add missing include.

#rnx

diff --git a/Engine/Source/Developer/DerivedDataCache/Private/Http/HttpRequest.cpp b/Engine/Source/Developer/DerivedDataCache/Private/Http/HttpRequest.cpp
--- a/Engine/Source/Developer/DerivedDataCache/Private/Http/HttpRequest.cpp
+++ b/Engine/Source/Developer/DerivedDataCache/Private/Http/HttpRequest.cpp
@@ -11,6 +11,7 @@
 #include "Serialization/CompactBinary.h"
 #include "Serialization/JsonReader.h"
 #include "Serialization/JsonSerializer.h"
+#include "String/Find.h"
 
 #if WITH_SSL
 #include "Ssl.h"
