DDC: Optimized cache key construction and shared string serialization

#rnx

diff --git a/Engine/Source/Developer/DerivedDataCache/Private/DerivedDataCacheKey.cpp b/Engine/Source/Developer/DerivedDataCache/Private/DerivedDataCacheKey.cpp
--- a/Engine/Source/Developer/DerivedDataCache/Private/DerivedDataCacheKey.cpp
+++ b/Engine/Source/Developer/DerivedDataCache/Private/DerivedDataCacheKey.cpp
@@ -5,6 +5,7 @@
 #include "Containers/Set.h"
 #include "Containers/StringConv.h"
 #include "DerivedDataCachePrivate.h"
+#include "Hash/xxhash.h"
 #include "IO/IoHash.h"
 #include "Math/UnrealMathUtility.h"
 #include "Misc/ScopeRWLock.h"
@@ -17,18 +18,10 @@
 namespace UE::DerivedData::Private
 {
 
-template <typename CharType>
-static inline void AssertValidCacheBucketName(TStringView<CharType> Name)
-{
-	checkf(FCacheBucket::IsValidName(Name),
-		TEXT("A cache bucket name must be alphanumeric, non-empty, and contain at most %d code units. Name: '%s'"),
-		*WriteToString<256>(Name), FCacheBucket::MaxNameLen);
-}
-
 class FCacheBucketOwner : public FCacheBucket
 {
 public:
-	inline explicit FCacheBucketOwner(FAnsiStringView Bucket);
+	inline explicit FCacheBucketOwner(FUtf8StringView Bucket);
 	inline FCacheBucketOwner(FCacheBucketOwner&& Bucket);
 	inline ~FCacheBucketOwner();
 
@@ -37,26 +30,44 @@ public:
 
 	using FCacheBucket::operator==;
 
-	inline bool operator==(FAnsiStringView Bucket) const { return ToString().Equals(Bucket, ESearchCase::IgnoreCase); }
+	inline bool operator==(FUtf8StringView Bucket) const { return ToString().Equals(Bucket, ESearchCase::IgnoreCase); }
+
+	static inline uint32 GetBucketHash(const FUtf8StringView Bucket)
+	{
+		const int32 Len = Bucket.Len();
+		check(Len <= MaxNameLen);
+		UTF8CHAR LowerBucket[MaxNameLen];
+		UTF8CHAR* LowerBucketIt = LowerBucket;
+		for (const UTF8CHAR& Char : Bucket)
+		{
+			*LowerBucketIt++ = TChar<UTF8CHAR>::ToLower(Char);
+		}
+		return uint32(FXxHash64::HashBuffer(LowerBucket, Len).Hash);
+	}
 
 	friend inline uint32 GetTypeHash(const FCacheBucketOwner& Bucket)
 	{
-		return ::GetTypeHash(Bucket.ToString());
+		return GetBucketHash(Bucket.ToString());
 	}
 };
 
-inline FCacheBucketOwner::FCacheBucketOwner(FAnsiStringView Bucket)
+inline FCacheBucketOwner::FCacheBucketOwner(FUtf8StringView Bucket)
 {
+	checkf(FCacheBucket::IsValidName(Bucket),
+		TEXT("A cache bucket name must be alphanumeric, non-empty, and contain at most %d code units. Name: '%s'"),
+		*WriteToString<256>(Bucket), FCacheBucket::MaxNameLen);
+
+	static_assert(sizeof(ANSICHAR) == sizeof(UTF8CHAR));
 	const int32 BucketLen = Bucket.Len();
 	const int32 PrefixSize = FMath::DivideAndRoundUp<int32>(1, sizeof(ANSICHAR));
-	ANSICHAR* Buffer = new ANSICHAR[PrefixSize + BucketLen + 1];
+	UTF8CHAR* Buffer = new UTF8CHAR[PrefixSize + BucketLen + 1];
 	Buffer += PrefixSize;
-	Name = Buffer;
+	Name = reinterpret_cast<ANSICHAR*>(Buffer);
 
 	reinterpret_cast<uint8*>(Buffer)[LengthOffset] = static_cast<uint8>(BucketLen);
 	Bucket.CopyString(Buffer, BucketLen);
 	Buffer += BucketLen;
-	*Buffer = '\0';
+	*Buffer = UTF8CHAR('\0');
 }
 
 inline FCacheBucketOwner::FCacheBucketOwner(FCacheBucketOwner&& Bucket)
@@ -86,21 +97,23 @@ private:
 };
 
 template <typename CharType>
-inline FCacheBucket FCacheBuckets::FindOrAdd(TStringView<CharType> Name)
+inline FCacheBucket FCacheBuckets::FindOrAdd(const TStringView<CharType> Name)
 {
-	const auto AnsiCast = StringCast<ANSICHAR>(Name.GetData(), Name.Len());
-	const FAnsiStringView AnsiName = AnsiCast;
+	const auto NameCast = StringCast<UTF8CHAR, FCacheBucket::MaxNameLen + 1>(Name.GetData(), Name.Len());
+	const FUtf8StringView NameView = NameCast;
+	uint32 Hash = 0;
 
-	const uint32 Hash = GetTypeHash(AnsiName);
-	if (FReadScopeLock ReadLock(Lock); const FCacheBucketOwner* Bucket = Buckets.FindByHash(Hash, AnsiName))
+	if (NameView.Len() <= FCacheBucket::MaxNameLen)
 	{
-		return *Bucket;
+		Hash = FCacheBucketOwner::GetBucketHash(NameView);
+		FReadScopeLock ReadLock(Lock);
+		if (const FCacheBucketOwner* Bucket = Buckets.FindByHash(Hash, NameView))
+		{
+			return *Bucket;
+		}
 	}
 
-	// It is valid to assert after because the "bogus char" '?' is not valid in a bucket name.
-	AssertValidCacheBucketName(Name);
-
-	FCacheBucketOwner LocalBucket(AnsiName);
+	FCacheBucketOwner LocalBucket(NameView);
 	FWriteScopeLock WriteLock(Lock);
 	return Buckets.FindOrAddByHash(Hash, MoveTemp(LocalBucket));
 }
diff --git a/Engine/Source/Developer/DerivedDataCache/Private/DerivedDataSharedString.cpp b/Engine/Source/Developer/DerivedDataCache/Private/DerivedDataSharedString.cpp
--- a/Engine/Source/Developer/DerivedDataCache/Private/DerivedDataSharedString.cpp
+++ b/Engine/Source/Developer/DerivedDataCache/Private/DerivedDataSharedString.cpp
@@ -17,7 +17,7 @@ bool LoadFromCompactBinary(FCbFieldView Field, FUtf8SharedString& OutString)
 
 bool LoadFromCompactBinary(FCbFieldView Field, FWideSharedString& OutString)
 {
-	TWideStringBuilder<128> String;
+	TWideStringBuilder<512> String;
 	if (LoadFromCompactBinary(Field, String))
 	{
 		OutString = String;
