DDC: Added detail to the User-Agent HTTP header

User-Agent: UnrealEngine/X.Y.Z-<CL> (<Platform>; <Config> <TargetType>; <BranchName>) <AppName> (<ProjectName>)

#rnx

diff --git a/Engine/Source/Developer/DerivedDataCache/Private/Http/HttpRequest.cpp b/Engine/Source/Developer/DerivedDataCache/Private/Http/HttpRequest.cpp
--- a/Engine/Source/Developer/DerivedDataCache/Private/Http/HttpRequest.cpp
+++ b/Engine/Source/Developer/DerivedDataCache/Private/Http/HttpRequest.cpp
@@ -6,8 +6,11 @@
 #include "DerivedDataRequest.h"
 #include "DerivedDataRequestOwner.h"
 #include "HAL/IConsoleManager.h"
+#include "HAL/PlatformProperties.h"
 #include "Logging/LogMacros.h"
+#include "Misc/App.h"
 #include "Misc/AsciiSet.h"
+#include "Misc/EngineVersion.h"
 #include "Serialization/CompactBinary.h"
 #include "Serialization/JsonReader.h"
 #include "Serialization/JsonSerializer.h"
@@ -549,7 +552,21 @@ Http::Private::FCurlStringList FHttpRequest::PrepareToIssueRequest(FStringView U
 
 	AddHeader(GetSessionIdHeader());
 	AddHeader(*WriteToAnsiString<32>(ANSITEXTVIEW("UE-Request: "), GRequestId.fetch_add(1, std::memory_order_relaxed)));
-	AddHeader("User-Agent: Unreal Engine");
+
+	// User-Agent: UnrealEngine/X.Y.Z-<CL> (<Platform>; <Config> <TargetType>; <BranchName>) <AppName> (<ProjectName>)
+	const FEngineVersion& Version = FEngineVersion::Current();
+	TUtf8StringBuilder<256> UserAgent;
+	UserAgent << "User-Agent: UnrealEngine/"
+		<< Version.GetMajor() << '.' << Version.GetMinor() << '.' << Version.GetPatch() << '-' << Version.GetChangelist()
+		<< " (" << FPlatformProperties::PlatformName()
+		<< "; " << LexToString(FApp::GetBuildConfiguration()) << ' ' << LexToString(FApp::GetBuildTargetType())
+		<< "; " << FApp::GetBranchName()
+		<< ") " << FApp::GetName();
+	if (FApp::HasProjectName())
+	{
+		UserAgent << " (" << FApp::GetProjectName() << ")";
+	}
+	AddHeader((const ANSICHAR*)*UserAgent);
 
 	// Strip any Expect: 100-Continue header since this just introduces latency.
 	AddHeader("Expect:");
