Fix checkSlow failure in Debug caused by missing KeyFuncs when using a TSet.

diff --git a/Engine/Source/Developer/DerivedDataCache/Private/DerivedDataCacheKey.cpp b/Engine/Source/Developer/DerivedDataCache/Private/DerivedDataCacheKey.cpp
--- a/Engine/Source/Developer/DerivedDataCache/Private/DerivedDataCacheKey.cpp
+++ b/Engine/Source/Developer/DerivedDataCache/Private/DerivedDataCacheKey.cpp
@@ -85,6 +85,15 @@ inline FCacheBucketOwner::~FCacheBucketOwner()
 	}
 }
 
+struct FCacheBucketOwnerKeyFuncs : BaseKeyFuncs<FCacheBucketOwner, FCacheBucketOwner, false>
+{
+	static const FCacheBucketOwner& GetSetKey(const FCacheBucketOwner& Element) { return Element; }
+	static bool Matches(const FCacheBucketOwner& A, const FCacheBucketOwner& B) { return A == B; }
+	static bool Matches(const FCacheBucketOwner& A, const FUtf8StringView B) { return A.ToString() == B; }
+	static uint32 GetKeyHash(const FCacheBucketOwner& Key) { return GetTypeHash(Key); }
+	static uint32 GetKeyHash(const FUtf8StringView Key) { return FCacheBucketOwner::GetBucketHash(Key); }
+};
+
 class FCacheBuckets
 {
 public:
@@ -93,7 +102,7 @@ public:
 
 private:
 	FRWLock Lock;
-	TSet<FCacheBucketOwner> Buckets;
+	TSet<FCacheBucketOwner, FCacheBucketOwnerKeyFuncs> Buckets;
 };
 
 template <typename CharType>
