DDC: Fixed invalid or inconsistent whitespace in HttpCacheStore

#rnx

diff --git a/Engine/Source/Developer/DerivedDataCache/Private/HttpCacheStore.cpp b/Engine/Source/Developer/DerivedDataCache/Private/HttpCacheStore.cpp
--- a/Engine/Source/Developer/DerivedDataCache/Private/HttpCacheStore.cpp
+++ b/Engine/Source/Developer/DerivedDataCache/Private/HttpCacheStore.cpp
@@ -208,9 +208,9 @@ private:
 	enum class OperationCategory
 	{
 		Get,
-		Put
+		Put,
 	};
-	template<OperationCategory Category>
+	template <OperationCategory Category>
 	FHttpRequest* WaitForHttpRequestForOwner(IRequestOwner& Owner, bool bUnboundedOverflow, FHttpRequestPool*& OutPool)
 	{
 		if (!FHttpRequest::AllowAsync())
@@ -420,7 +420,7 @@ public:
 	using FOnGetCachedDataBatchComplete = TUniqueFunction<void(FGetCachedDataBatchResponse&& Response)>;
 
 	/** Utility method for fetching a batch of value data. */
-	template<typename ValueType, typename ValueIdGetterType>
+	template <typename ValueType, typename ValueIdGetterType>
 	static void GetDataBatch(
 		FHttpCacheStore& CacheStore,
 		IRequestOwner& Owner,
@@ -572,13 +572,13 @@ void FHttpCacheStore::FPutPackageOp::PutRefAsync(
 			if (bPutRefBlobsAlways && !bFinalize)
 			{
 				Object.IterateAttachments([&NeededBlobHashes](FCbFieldView AttachmentFieldView)
+				{
+					FIoHash AttachmentHash = AttachmentFieldView.AsHash();
+					if (!AttachmentHash.IsZero())
 					{
-						FIoHash AttachmentHash = AttachmentFieldView.AsHash();
-						if (!AttachmentHash.IsZero())
-						{
-							NeededBlobHashes.Add(AttachmentHash);
-						}
-					});
+						NeededBlobHashes.Add(AttachmentHash);
+					}
+				});
 			}
 			else if (TSharedPtr<FJsonObject> ResponseObject = Request->GetResponseAsJsonObject())
 			{
@@ -812,7 +812,7 @@ void FHttpCacheStore::FGetRecordOp::GetRecord(
 	});
 }
 
-template<typename ValueType, typename ValueIdGetterType>
+template <typename ValueType, typename ValueIdGetterType>
 void FHttpCacheStore::FGetRecordOp::GetDataBatch(
 	FHttpCacheStore& CacheStore,
 	IRequestOwner& Owner,
@@ -1740,10 +1740,10 @@ void FHttpCacheStore::GetCacheValueAsync(
 			return FHttpRequest::ECompletionBehavior::Done;
 		}
 
-		 if (!ShouldAbortForShutdown() && !Owner.IsCanceled() && ShouldRetryOnError(HttpResult, ResponseCode) && ((Request->GetAttempts()+1) < UE_HTTPDDC_MAX_ATTEMPTS))
-		 {
+		if (!ShouldAbortForShutdown() && !Owner.IsCanceled() && ShouldRetryOnError(HttpResult, ResponseCode) && ((Request->GetAttempts()+1) < UE_HTTPDDC_MAX_ATTEMPTS))
+		{
 			return FHttpRequest::ECompletionBehavior::Retry;
-		 }
+		}
 
 		UE_LOG(LogDerivedDataCache, Verbose, TEXT("%s: Cache miss with failed HTTP request for %s from '%s'"),
 			*Domain, *WriteToString<96>(Key), *Name);
@@ -1778,12 +1778,12 @@ void FHttpCacheStore::RefCachedDataProbablyExistsBatchAsync(
 	if (!IsUsable())
 	{
 		for (const FCacheGetValueRequest& ValueRef : ValueRefs)
-	{
+		{
 			UE_LOG(LogDerivedDataCache, VeryVerbose,
 				TEXT("%s: Skipped exists check of %s from '%s' because this cache store is not available"),
 				*Domain, *WriteToString<96>(ValueRef.Key), *ValueRef.Name);
 			OnComplete(ValueRef.MakeResponse(EStatus::Error));
-	}
+		}
 		return;
 	}
 
@@ -1829,10 +1829,10 @@ void FHttpCacheStore::RefCachedDataProbablyExistsBatchAsync(
 			{
 				for (const FCacheGetValueRequest& ValueRef : ValueRefs)
 				{
-				UE_LOG(LogDerivedDataCache, Log,
-					TEXT("%s: Cache exists returned invalid results."),
-					*Domain);
-					OnComplete(ValueRef.MakeResponse(EStatus::Error));
+					UE_LOG(LogDerivedDataCache, Log,
+						TEXT("%s: Cache exists returned invalid results."),
+						*Domain);
+						OnComplete(ValueRef.MakeResponse(EStatus::Error));
 				}
 				return FHttpRequest::ECompletionBehavior::Done;
 			}
@@ -2126,9 +2126,9 @@ void FHttpCacheStore::GetChunks(
 
 					FRequestOwner BlockingOwner(EPriority::Blocking);
 					GetCacheRecordOnlyAsync(BlockingOwner, Request.Name, Request.Key, PolicyBuilder.Build(), 0, [&Record](FGetCacheRecordOnlyResponse&& Response)
-						{
-							Record = MoveTemp(Response.Record);
-						});
+					{
+						Record = MoveTemp(Response.Record);
+					});
 					BlockingOwner.Wait();
 				}
 				if (Record)
@@ -2154,13 +2154,13 @@ void FHttpCacheStore::GetChunks(
 						bool bSucceeded = false;
 						FCompressedBuffer NewBuffer;
 						FGetRecordOp::GetDataBatch(*this, BlockingOwner, Request.Name, Request.Key, ::MakeArrayView({ ValueWithId }), IdGetter, [&bSucceeded, &NewBuffer](FGetRecordOp::FGetCachedDataBatchResponse&& Response)
+						{
+							if (Response.Status == EStatus::Ok)
 							{
-								if (Response.Status == EStatus::Ok)
-								{
-									bSucceeded = true;
-									NewBuffer = MoveTemp(Response.DataBuffer);
-								}
-							});
+								bSucceeded = true;
+								NewBuffer = MoveTemp(Response.DataBuffer);
+							}
+						});
 						BlockingOwner.Wait();
 
 						if (bSucceeded)
@@ -2184,10 +2184,10 @@ void FHttpCacheStore::GetChunks(
 					FRequestOwner BlockingOwner(EPriority::Blocking);
 					bool bSucceeded = false;
 					GetCacheValueAsync(BlockingOwner, Request.Name, Request.Key, Request.Policy, 0, [&bSucceeded, &Value](FCacheGetValueResponse&& Response)
-						{
-							Value = MoveTemp(Response.Value);
-							bSucceeded = Response.Status == EStatus::Ok;
-						});
+					{
+						Value = MoveTemp(Response.Value);
+						bSucceeded = Response.Status == EStatus::Ok;
+					});
 					BlockingOwner.Wait();
 					bHasValue = bSucceeded;
 				}
@@ -2209,13 +2209,13 @@ void FHttpCacheStore::GetChunks(
 						bool bSucceeded = false;
 						FCompressedBuffer NewBuffer;
 						FGetRecordOp::GetDataBatch(*this, BlockingOwner, Request.Name, Request.Key, ::MakeArrayView({ Value }), IdGetter, [&bSucceeded, &NewBuffer](FGetRecordOp::FGetCachedDataBatchResponse&& Response)
+						{
+							if (Response.Status == EStatus::Ok)
 							{
-								if (Response.Status == EStatus::Ok)
-								{
-									bSucceeded = true;
-									NewBuffer = MoveTemp(Response.DataBuffer);
-								}
-							});
+								bSucceeded = true;
+								NewBuffer = MoveTemp(Response.DataBuffer);
+							}
+						});
 						BlockingOwner.Wait();
 
 						if (bSucceeded)
