DDC: Force SSL to start up before DDC to enforce shutdown order, since DDC may be using SSL until it finishes shutting down

#rnx

diff --git a/Engine/Source/Developer/DerivedDataCache/Private/DerivedDataCacheModule.cpp b/Engine/Source/Developer/DerivedDataCache/Private/DerivedDataCacheModule.cpp
--- a/Engine/Source/Developer/DerivedDataCache/Private/DerivedDataCacheModule.cpp
+++ b/Engine/Source/Developer/DerivedDataCache/Private/DerivedDataCacheModule.cpp
@@ -56,6 +56,13 @@ public:
 		}
 	}
 
+	void StartupModule() final
+	{
+		// Required to guarantee that SSL shuts down after DDC. Without this, waiting for active
+		// cache requests on shutdown can crash when accessing SSL.
+		FModuleManager::Get().LoadModuleChecked(TEXT("SSL"));
+	}
+
 	void ShutdownModule() final
 	{
 		delete GDerivedDataBuild;
