DDC: Simplify FCacheBucketOwnerKeyFuncs

#rnx

diff --git a/Engine/Source/Developer/DerivedDataCache/Private/DerivedDataCacheKey.cpp b/Engine/Source/Developer/DerivedDataCache/Private/DerivedDataCacheKey.cpp
--- a/Engine/Source/Developer/DerivedDataCache/Private/DerivedDataCacheKey.cpp
+++ b/Engine/Source/Developer/DerivedDataCache/Private/DerivedDataCacheKey.cpp
@@ -29,33 +29,14 @@ public:
 	FCacheBucketOwner& operator=(const FCacheBucketOwner&) = delete;
 
 	using FCacheBucket::operator==;
-
-	inline bool operator==(FUtf8StringView Bucket) const { return ToString().Equals(Bucket, ESearchCase::IgnoreCase); }
-
-	static inline uint32 GetBucketHash(const FUtf8StringView Bucket)
-	{
-		const int32 Len = Bucket.Len();
-		check(Len <= MaxNameLen);
-		UTF8CHAR LowerBucket[MaxNameLen];
-		UTF8CHAR* LowerBucketIt = LowerBucket;
-		for (const UTF8CHAR& Char : Bucket)
-		{
-			*LowerBucketIt++ = TChar<UTF8CHAR>::ToLower(Char);
-		}
-		return uint32(FXxHash64::HashBuffer(LowerBucket, Len).Hash);
-	}
-
-	friend inline uint32 GetTypeHash(const FCacheBucketOwner& Bucket)
-	{
-		return GetBucketHash(Bucket.ToString());
-	}
+	inline bool operator==(FUtf8StringView Bucket) const { return ToString() == Bucket; }
 };
 
 inline FCacheBucketOwner::FCacheBucketOwner(FUtf8StringView Bucket)
 {
 	checkf(FCacheBucket::IsValidName(Bucket),
 		TEXT("A cache bucket name must be alphanumeric, non-empty, and contain at most %d code units. Name: '%s'"),
-		*WriteToString<256>(Bucket), FCacheBucket::MaxNameLen);
+		FCacheBucket::MaxNameLen, *WriteToString<256>(Bucket));
 
 	static_assert(sizeof(ANSICHAR) == sizeof(UTF8CHAR));
 	const int32 BucketLen = Bucket.Len();
@@ -85,13 +66,25 @@ inline FCacheBucketOwner::~FCacheBucketOwner()
 	}
 }
 
-struct FCacheBucketOwnerKeyFuncs : BaseKeyFuncs<FCacheBucketOwner, FCacheBucketOwner, false>
+struct FCacheBucketOwnerKeyFuncs : DefaultKeyFuncs<FCacheBucketOwner>
 {
-	static const FCacheBucketOwner& GetSetKey(const FCacheBucketOwner& Element) { return Element; }
-	static bool Matches(const FCacheBucketOwner& A, const FCacheBucketOwner& B) { return A == B; }
-	static bool Matches(const FCacheBucketOwner& A, const FUtf8StringView B) { return A.ToString() == B; }
-	static uint32 GetKeyHash(const FCacheBucketOwner& Key) { return GetTypeHash(Key); }
-	static uint32 GetKeyHash(const FUtf8StringView Key) { return FCacheBucketOwner::GetBucketHash(Key); }
+	static uint32 GetKeyHash(const FUtf8StringView Key)
+	{
+		const int32 Len = Key.Len();
+		check(Len <= FCacheBucket::MaxNameLen);
+		UTF8CHAR LowerKey[FCacheBucket::MaxNameLen];
+		UTF8CHAR* LowerKeyIt = LowerKey;
+		for (const UTF8CHAR& Char : Key)
+		{
+			*LowerKeyIt++ = TChar<UTF8CHAR>::ToLower(Char);
+		}
+		return uint32(FXxHash64::HashBuffer(LowerKey, Len).Hash);
+	}
+
+	static uint32 GetKeyHash(const FCacheBucketOwner& Key)
+	{
+		return GetKeyHash(Key.ToString());
+	}
 };
 
 class FCacheBuckets
@@ -114,7 +107,7 @@ inline FCacheBucket FCacheBuckets::FindOrAdd(const TStringView<CharType> Name)
 
 	if (NameView.Len() <= FCacheBucket::MaxNameLen)
 	{
-		Hash = FCacheBucketOwner::GetBucketHash(NameView);
+		Hash = FCacheBucketOwnerKeyFuncs::GetKeyHash(NameView);
 		FReadScopeLock ReadLock(Lock);
 		if (const FCacheBucketOwner* Bucket = Buckets.FindByHash(Hash, NameView))
 		{
