DDC: Renamed HTTP enums to match modern terminology

#rnx

diff --git a/Engine/Source/Developer/DerivedDataCache/Private/Http/HttpClient.h b/Engine/Source/Developer/DerivedDataCache/Private/Http/HttpClient.h
--- a/Engine/Source/Developer/DerivedDataCache/Private/Http/HttpClient.h
+++ b/Engine/Source/Developer/DerivedDataCache/Private/Http/HttpClient.h
@@ -37,49 +37,57 @@ namespace UE
 
 class FHttpRequestPool;
 
-enum class EHttpContentType : int16
+enum class EHttpMethod : uint8
+{
+	Get,
+	Put,
+	Post,
+	Head,
+	Delete,
+};
+
+enum class EHttpMediaType : int16
 {
 	// Negative integers reserved for future custom string content types
 	UnspecifiedContentType = 0,
-	AnyContentType,
+	Any,
 	Binary,
 	Text,
-	JSON,
+	Json,
+	Yaml,
 	CbObject,
 	CbPackage,
-	YAML,
 	CbPackageOffer,
 	CompressedBinary,
 	FormUrlEncoded,
-	Count
 };
 
-inline FStringView GetHttpMimeType(EHttpContentType Type)
+inline FStringView GetHttpMimeType(EHttpMediaType Type)
 {
 	switch (Type)
 	{
-		case EHttpContentType::UnspecifiedContentType:
+		case EHttpMediaType::UnspecifiedContentType:
 			checkNoEntry();
 			return FStringView();
-		case EHttpContentType::AnyContentType:
+		case EHttpMediaType::Any:
 			return TEXTVIEW("*/*");
-		case EHttpContentType::Binary:
+		case EHttpMediaType::Binary:
 			return TEXTVIEW("application/octet-stream");
-		case EHttpContentType::Text:
+		case EHttpMediaType::Text:
 			return TEXTVIEW("text/plain");
-		case EHttpContentType::JSON:
+		case EHttpMediaType::Json:
 			return TEXTVIEW("application/json");
-		case EHttpContentType::CbObject:
+		case EHttpMediaType::Yaml:
+			return TEXTVIEW("text/yaml");
+		case EHttpMediaType::CbObject:
 			return TEXTVIEW("application/x-ue-cb");
-		case EHttpContentType::CbPackage:
+		case EHttpMediaType::CbPackage:
 			return TEXTVIEW("application/x-ue-cbpkg");
-		case EHttpContentType::YAML:
-			return TEXTVIEW("text/yaml");
-		case EHttpContentType::CbPackageOffer:
+		case EHttpMediaType::CbPackageOffer:
 			return TEXTVIEW("application/x-ue-offer");
-		case EHttpContentType::CompressedBinary:
+		case EHttpMediaType::CompressedBinary:
 			return TEXTVIEW("application/x-ue-comp");
-		case EHttpContentType::FormUrlEncoded:
+		case EHttpMediaType::FormUrlEncoded:
 			return TEXTVIEW("application/x-www-form-urlencoded");
 		default:
 			return TEXTVIEW("unknown");
@@ -152,18 +160,6 @@ public:
 		Retry
 	};
 
-	/**
-	 * Supported request verbs
-	 */
-	enum class ERequestVerb
-	{
-		Get,
-		Put,
-		Post,
-		Delete,
-		Head
-	};
-
 	using FResultCode = int;
 
 	using FOnHttpRequestComplete = TUniqueFunction<ECompletionBehavior(EResult HttpResult, FHttpRequest* Request)>;
@@ -234,12 +230,12 @@ public:
 	EResult PerformBlockingPut(
 		FStringView Uri,
 		const FCompositeBuffer& Buffer,
-		EHttpContentType ContentType = EHttpContentType::UnspecifiedContentType,
+		EHttpMediaType ContentType = EHttpMediaType::UnspecifiedContentType,
 		TConstArrayView<long> ExpectedErrorCodes = {});
 	EResult PerformBlockingPost(FStringView Uri,
 		const FCompositeBuffer& Buffer,
-		EHttpContentType ContentType = EHttpContentType::UnspecifiedContentType,
-		EHttpContentType AcceptType = EHttpContentType::UnspecifiedContentType,
+		EHttpMediaType ContentType = EHttpMediaType::UnspecifiedContentType,
+		EHttpMediaType AcceptType = EHttpMediaType::UnspecifiedContentType,
 		TConstArrayView<long> ExpectedErrorCodes = {});
 	void EnqueueAsyncPut(
 		DerivedData::IRequestOwner& Owner,
@@ -247,7 +243,7 @@ public:
 		FStringView Uri,
 		const FCompositeBuffer& Buffer,
 		FOnHttpRequestComplete&& OnComplete,
-		EHttpContentType ContentType = EHttpContentType::UnspecifiedContentType,
+		EHttpMediaType ContentType = EHttpMediaType::UnspecifiedContentType,
 		TConstArrayView<long> ExpectedErrorCodes = {});
 	void EnqueueAsyncPost(
 		DerivedData::IRequestOwner& Owner,
@@ -255,8 +251,8 @@ public:
 		FStringView Uri,
 		const FCompositeBuffer& Buffer,
 		FOnHttpRequestComplete&& OnComplete,
-		EHttpContentType ContentType = EHttpContentType::UnspecifiedContentType,
-		EHttpContentType AcceptType = EHttpContentType::UnspecifiedContentType,
+		EHttpMediaType ContentType = EHttpMediaType::UnspecifiedContentType,
+		EHttpMediaType AcceptType = EHttpMediaType::UnspecifiedContentType,
 		TConstArrayView<long> ExpectedErrorCodes = {});
 
 	/**
@@ -269,14 +265,14 @@ public:
 	EResult PerformBlockingDownload(
 		FStringView Uri,
 		TArray64<uint8>* Buffer,
-		EHttpContentType AcceptType = EHttpContentType::UnspecifiedContentType,
+		EHttpMediaType AcceptType = EHttpMediaType::UnspecifiedContentType,
 		TConstArrayView<long> ExpectedErrorCodes = {400});
 	void EnqueueAsyncDownload(
 		DerivedData::IRequestOwner& Owner,
 		FHttpRequestPool* Pool,
 		FStringView Uri,
 		FOnHttpRequestComplete&& OnComplete,
-		EHttpContentType AcceptType = EHttpContentType::UnspecifiedContentType,
+		EHttpMediaType AcceptType = EHttpMediaType::UnspecifiedContentType,
 		TConstArrayView<long> ExpectedErrorCodes = {400});
 
 	/**
@@ -286,7 +282,7 @@ public:
 	 */
 	EResult PerformBlockingHead(
 		FStringView Uri,
-		EHttpContentType AcceptType = EHttpContentType::UnspecifiedContentType,
+		EHttpMediaType AcceptType = EHttpMediaType::UnspecifiedContentType,
 		TConstArrayView<long> ExpectedErrorCodes = {400});
 	EResult PerformBlockingDelete(
 		FStringView Uri,
@@ -296,7 +292,7 @@ public:
 		FHttpRequestPool* Pool,
 		FStringView Uri,
 		FOnHttpRequestComplete&& OnComplete,
-		EHttpContentType AcceptType = EHttpContentType::UnspecifiedContentType,
+		EHttpMediaType AcceptType = EHttpMediaType::UnspecifiedContentType,
 		TConstArrayView<long> ExpectedErrorCodes = {400});
 	void EnqueueAsyncDelete(
 		DerivedData::IRequestOwner& Owner,
@@ -394,8 +390,8 @@ private:
 	FString EffectiveDomain;
 	const FHttpAccessToken* AuthorizationToken;
 
-	void AddContentTypeHeader(FStringView Header, EHttpContentType Type);
-	Http::Private::FCurlStringList PrepareToIssueRequest(FStringView Uri, ERequestVerb Verb, uint64 ContentLength);
+	void AddContentTypeHeader(FStringView Header, EHttpMediaType Type);
+	Http::Private::FCurlStringList PrepareToIssueRequest(FStringView Uri, EHttpMethod Verb, uint64 ContentLength);
 
 	/**
 	 * Performs the request, blocking until finished.
@@ -405,18 +401,18 @@ private:
 	 * @param ExpectedErrorCodes An array of expected return codes outside of the success range that should NOT be logged as an abnormal/exceptional outcome.
 	 * If unset the response body will be stored in the request.
 	 */
-	EResult PerformBlocking(FStringView Uri, ERequestVerb Verb, uint64 ContentLength, TConstArrayView<long> ExpectedErrorCodes);
+	EResult PerformBlocking(FStringView Uri, EHttpMethod Verb, uint64 ContentLength, TConstArrayView<long> ExpectedErrorCodes);
 
 	void EnqueueAsync(
 		DerivedData::IRequestOwner& Owner,
 		FHttpRequestPool* Pool,
 		FStringView Uri,
-		ERequestVerb Verb,
+		EHttpMethod Verb,
 		uint64 ContentLength,
 		FOnHttpRequestComplete&& OnComplete,
 		TConstArrayView<long> ExpectedErrorCodes);
 
-	void LogResult(FResultCode EResult, FStringView Uri, ERequestVerb Verb, TConstArrayView<long> ExpectedErrorCodes) const;
+	void LogResult(FResultCode EResult, FStringView Uri, EHttpMethod Verb, TConstArrayView<long> ExpectedErrorCodes) const;
 
 	static FString GetAnsiBufferAsString(TConstArrayView64<uint8> Buffer)
 	{
diff --git a/Engine/Source/Developer/DerivedDataCache/Private/Http/HttpRequest.cpp b/Engine/Source/Developer/DerivedDataCache/Private/Http/HttpRequest.cpp
--- a/Engine/Source/Developer/DerivedDataCache/Private/Http/HttpRequest.cpp
+++ b/Engine/Source/Developer/DerivedDataCache/Private/Http/HttpRequest.cpp
@@ -74,7 +74,7 @@ struct FAsyncRequestData final : public DerivedData::FRequestBase
 	FHttpRequestPool* Pool = nullptr;
 	FCurlStringList CurlHeaders;
 	FString Uri;
-	FHttpRequest::ERequestVerb Verb;
+	EHttpMethod Verb;
 	TArray<long, TInlineAllocator<4>> ExpectedErrorCodes;
 	FHttpRequest::FOnHttpRequestComplete OnComplete;
 	FLazyEvent Event {EEventMode::ManualReset};
@@ -378,7 +378,7 @@ void FHttpRequest::PrepareToRetry()
 FHttpRequest::EResult FHttpRequest::PerformBlockingDownload(
 	const FStringView Uri,
 	TArray64<uint8>* const Buffer,
-	const EHttpContentType AcceptType,
+	const EHttpMediaType AcceptType,
 	const TConstArrayView<long> ExpectedErrorCodes)
 {
 	curl_easy_setopt(static_cast<CURL*>(Curl), CURLOPT_HTTPGET, 1L);
@@ -386,7 +386,7 @@ FHttpRequest::EResult FHttpRequest::PerformBlockingDownload(
 
 	AddContentTypeHeader(TEXTVIEW("Accept"), AcceptType);
 
-	return PerformBlocking(Uri, ERequestVerb::Get, 0u, ExpectedErrorCodes);
+	return PerformBlocking(Uri, EHttpMethod::Get, 0u, ExpectedErrorCodes);
 }
 
 void FHttpRequest::EnqueueAsyncDownload(
@@ -394,12 +394,12 @@ void FHttpRequest::EnqueueAsyncDownload(
 	FHttpRequestPool* const Pool,
 	const FStringView Uri,
 	FOnHttpRequestComplete&& OnComplete,
-	const EHttpContentType AcceptType,
+	const EHttpMediaType AcceptType,
 	const TConstArrayView<long> ExpectedErrorCodes)
 {
 	curl_easy_setopt(static_cast<CURL*>(Curl), CURLOPT_HTTPGET, 1L);
 
-	return EnqueueAsync(Owner, Pool, Uri, ERequestVerb::Get, 0u, MoveTemp(OnComplete), ExpectedErrorCodes);
+	return EnqueueAsync(Owner, Pool, Uri, EHttpMethod::Get, 0u, MoveTemp(OnComplete), ExpectedErrorCodes);
 }
 
 void FHttpRequest::AddHeader(const FStringView Header, const FStringView Value)
@@ -504,9 +504,9 @@ void FHttpRequest::CompleteAsync(FResultCode Result)
 	}
 }
 
-void FHttpRequest::AddContentTypeHeader(FStringView Header, EHttpContentType Type)
+void FHttpRequest::AddContentTypeHeader(FStringView Header, EHttpMediaType Type)
 {
-	if (Type != EHttpContentType::UnspecifiedContentType)
+	if (Type != EHttpMediaType::UnspecifiedContentType)
 	{
 		AddHeader(Header, GetHttpMimeType(Type));
 	}
@@ -528,7 +528,7 @@ static const char* GetSessionIdHeader()
 
 static std::atomic<int> GRequestId = 1;
 
-Http::Private::FCurlStringList FHttpRequest::PrepareToIssueRequest(FStringView Uri, ERequestVerb Verb, uint64 ContentLength)
+Http::Private::FCurlStringList FHttpRequest::PrepareToIssueRequest(FStringView Uri, EHttpMethod Verb, uint64 ContentLength)
 {
 	// Strip leading slashes because the separating slash is added below.
 	Uri = FAsciiSet::TrimPrefixWith(Uri, "/");
@@ -576,7 +576,7 @@ Http::Private::FCurlStringList FHttpRequest::PrepareToIssueRequest(FStringView U
 		AddHeader(TCHAR_TO_ANSI(*Header));
 	}
 
-	if ((Verb != ERequestVerb::Get) && (Verb != ERequestVerb::Delete))
+	if ((Verb != EHttpMethod::Get) && (Verb != EHttpMethod::Delete))
 	{
 		AddHeader(*WriteToAnsiString<32>(ANSITEXTVIEW("Content-Length: "), ContentLength));
 	}
@@ -592,7 +592,7 @@ Http::Private::FCurlStringList FHttpRequest::PrepareToIssueRequest(FStringView U
 
 FHttpRequest::EResult FHttpRequest::PerformBlocking(
 	const FStringView Uri,
-	const ERequestVerb Verb,
+	const EHttpMethod Verb,
 	const uint64 ContentLength,
 	const TConstArrayView<long> ExpectedErrorCodes)
 {
@@ -613,7 +613,7 @@ void FHttpRequest::EnqueueAsync(
 	DerivedData::IRequestOwner& Owner,
 	FHttpRequestPool* const Pool,
 	const FStringView Uri,
-	const ERequestVerb Verb,
+	const EHttpMethod Verb,
 	const uint64 ContentLength,
 	FOnHttpRequestComplete&& OnComplete,
 	const TConstArrayView<long> ExpectedErrorCodes)
@@ -648,7 +648,7 @@ void FHttpRequest::EnqueueAsync(
 void FHttpRequest::LogResult(
 	const FResultCode Result,
 	const FStringView Uri,
-	const ERequestVerb Verb,
+	const EHttpMethod Verb,
 	const TConstArrayView<long> ExpectedErrorCodes) const
 {
 	if (Result == CURLE_OK)
@@ -659,21 +659,21 @@ void FHttpRequest::LogResult(
 
 		switch (Verb)
 		{
-		case ERequestVerb::Head:
+		case EHttpMethod::Head:
 			VerbStr = TEXT("querying");
 			break;
-		case ERequestVerb::Get:
+		case EHttpMethod::Get:
 			VerbStr = TEXT("fetching");
 			AdditionalInfo << TEXTVIEW("Received: ") << uint64(BytesReceived) << TEXTVIEW(" bytes.");
 			break;
-		case ERequestVerb::Put:
+		case EHttpMethod::Put:
 			VerbStr = TEXT("updating");
 			AdditionalInfo << TEXTVIEW("Sent: ") << uint64(BytesSent) << TEXTVIEW(" bytes.");
 			break;
-		case ERequestVerb::Post:
+		case EHttpMethod::Post:
 			VerbStr = TEXT("posting");
 			break;
-		case ERequestVerb::Delete:
+		case EHttpMethod::Delete:
 			VerbStr = TEXT("deleting");
 			break;
 		default:
@@ -719,7 +719,7 @@ void FHttpRequest::LogResult(
 FHttpRequest::EResult FHttpRequest::PerformBlockingPut(
 	const FStringView Uri,
 	const FCompositeBuffer& Buffer,
-	const EHttpContentType ContentType,
+	const EHttpMediaType ContentType,
 	const TConstArrayView<long> ExpectedErrorCodes)
 {
 	uint64 ContentLength = 0u;
@@ -733,14 +733,14 @@ FHttpRequest::EResult FHttpRequest::PerformBlockingPut(
 	ContentLength = Buffer.GetSize();
 	ReadCompositeBuffer = Buffer;
 
-	return PerformBlocking(Uri, ERequestVerb::Put, ContentLength, ExpectedErrorCodes);
+	return PerformBlocking(Uri, EHttpMethod::Put, ContentLength, ExpectedErrorCodes);
 }
 
 FHttpRequest::EResult FHttpRequest::PerformBlockingPost(
 	const FStringView Uri,
 	const FCompositeBuffer& Buffer,
-	const EHttpContentType ContentType,
-	const EHttpContentType AcceptType,
+	const EHttpMediaType ContentType,
+	const EHttpMediaType AcceptType,
 	const TConstArrayView<long> ExpectedErrorCodes)
 {
 	uint64 ContentLength = 0;
@@ -755,7 +755,7 @@ FHttpRequest::EResult FHttpRequest::PerformBlockingPost(
 	ContentLength = Buffer.GetSize();
 	ReadCompositeBuffer = Buffer;
 
-	return PerformBlocking(Uri, ERequestVerb::Post, ContentLength, ExpectedErrorCodes);
+	return PerformBlocking(Uri, EHttpMethod::Post, ContentLength, ExpectedErrorCodes);
 }
 
 void FHttpRequest::EnqueueAsyncPut(
@@ -764,7 +764,7 @@ void FHttpRequest::EnqueueAsyncPut(
 	const FStringView Uri,
 	const FCompositeBuffer& Buffer,
 	FOnHttpRequestComplete&& OnComplete,
-	const EHttpContentType ContentType,
+	const EHttpMediaType ContentType,
 	const TConstArrayView<long> ExpectedErrorCodes)
 {
 	uint64 ContentLength = 0;
@@ -778,7 +778,7 @@ void FHttpRequest::EnqueueAsyncPut(
 	ReadCompositeBuffer = Buffer;
 	ContentLength = Buffer.GetSize();
 
-	return EnqueueAsync(Owner, Pool, Uri, ERequestVerb::Put, ContentLength, MoveTemp(OnComplete), ExpectedErrorCodes);
+	return EnqueueAsync(Owner, Pool, Uri, EHttpMethod::Put, ContentLength, MoveTemp(OnComplete), ExpectedErrorCodes);
 }
 
 void FHttpRequest::EnqueueAsyncPost(
@@ -787,8 +787,8 @@ void FHttpRequest::EnqueueAsyncPost(
 	const FStringView Uri,
 	const FCompositeBuffer& Buffer,
 	FOnHttpRequestComplete&& OnComplete,
-	const EHttpContentType ContentType,
-	const EHttpContentType AcceptType,
+	const EHttpMediaType ContentType,
+	const EHttpMediaType AcceptType,
 	const TConstArrayView<long> ExpectedErrorCodes)
 {
 	uint64 ContentLength = 0;
@@ -803,17 +803,17 @@ void FHttpRequest::EnqueueAsyncPost(
 	ReadCompositeBuffer = Buffer;
 	ContentLength = Buffer.GetSize();
 
-	return EnqueueAsync(Owner, Pool, Uri, ERequestVerb::Post, ContentLength, MoveTemp(OnComplete), ExpectedErrorCodes);
+	return EnqueueAsync(Owner, Pool, Uri, EHttpMethod::Post, ContentLength, MoveTemp(OnComplete), ExpectedErrorCodes);
 }
 
 FHttpRequest::EResult FHttpRequest::PerformBlockingHead(
 	const FStringView Uri,
-	const EHttpContentType AcceptType,
+	const EHttpMediaType AcceptType,
 	const TConstArrayView<long> ExpectedErrorCodes)
 {
 	curl_easy_setopt(static_cast<CURL*>(Curl), CURLOPT_NOBODY, 1L);
 	AddContentTypeHeader(TEXTVIEW("Accept"), AcceptType);
-	return PerformBlocking(Uri, ERequestVerb::Head, 0, ExpectedErrorCodes);
+	return PerformBlocking(Uri, EHttpMethod::Head, 0, ExpectedErrorCodes);
 }
 
 FHttpRequest::EResult FHttpRequest::PerformBlockingDelete(
@@ -821,7 +821,7 @@ FHttpRequest::EResult FHttpRequest::PerformBlockingDelete(
 	const TConstArrayView<long> ExpectedErrorCodes)
 {
 	curl_easy_setopt(static_cast<CURL*>(Curl), CURLOPT_CUSTOMREQUEST, "DELETE");
-	return PerformBlocking(Uri, ERequestVerb::Delete, 0, ExpectedErrorCodes);
+	return PerformBlocking(Uri, EHttpMethod::Delete, 0, ExpectedErrorCodes);
 }
 
 void FHttpRequest::EnqueueAsyncHead(
@@ -829,12 +829,12 @@ void FHttpRequest::EnqueueAsyncHead(
 	FHttpRequestPool* const Pool,
 	const FStringView Uri,
 	FOnHttpRequestComplete&& OnComplete,
-	const EHttpContentType AcceptType,
+	const EHttpMediaType AcceptType,
 	const TConstArrayView<long> ExpectedErrorCodes)
 {
 	curl_easy_setopt(static_cast<CURL*>(Curl), CURLOPT_NOBODY, 1L);
 	AddContentTypeHeader(TEXTVIEW("Accept"), AcceptType);
-	return EnqueueAsync(Owner, Pool, Uri, ERequestVerb::Head, 0, MoveTemp(OnComplete), ExpectedErrorCodes);
+	return EnqueueAsync(Owner, Pool, Uri, EHttpMethod::Head, 0, MoveTemp(OnComplete), ExpectedErrorCodes);
 }
 
 void FHttpRequest::EnqueueAsyncDelete(
@@ -845,7 +845,7 @@ void FHttpRequest::EnqueueAsyncDelete(
 	const TConstArrayView<long> ExpectedErrorCodes)
 {
 	curl_easy_setopt(static_cast<CURL*>(Curl), CURLOPT_CUSTOMREQUEST, "DELETE");
-	return EnqueueAsync(Owner, Pool, Uri, ERequestVerb::Delete, 0, MoveTemp(OnComplete), ExpectedErrorCodes);
+	return EnqueueAsync(Owner, Pool, Uri, EHttpMethod::Delete, 0, MoveTemp(OnComplete), ExpectedErrorCodes);
 }
 
 } // UE
diff --git a/Engine/Source/Developer/DerivedDataCache/Private/HttpCacheStore.cpp b/Engine/Source/Developer/DerivedDataCache/Private/HttpCacheStore.cpp
--- a/Engine/Source/Developer/DerivedDataCache/Private/HttpCacheStore.cpp
+++ b/Engine/Source/Developer/DerivedDataCache/Private/HttpCacheStore.cpp
@@ -612,12 +612,12 @@ void FHttpCacheStore::FPutPackageOp::PutRefAsync(
 
 	if (bFinalize)
 	{
-		Request->EnqueueAsyncPost(Owner, Pool, *RefsUri, FCompositeBuffer(), MoveTemp(OnHttpRequestComplete), EHttpContentType::FormUrlEncoded);
+		Request->EnqueueAsyncPost(Owner, Pool, *RefsUri, FCompositeBuffer(), MoveTemp(OnHttpRequestComplete), EHttpMediaType::FormUrlEncoded);
 	}
 	else
 	{
 		Request->AddHeader(TEXTVIEW("X-Jupiter-IoHash"), WriteToString<48>(ObjectHash));
-		Request->EnqueueAsyncPut(Owner, Pool, *RefsUri, Object.GetBuffer(), MoveTemp(OnHttpRequestComplete), EHttpContentType::CbObject);
+		Request->EnqueueAsyncPut(Owner, Pool, *RefsUri, Object.GetBuffer(), MoveTemp(OnHttpRequestComplete), EHttpMediaType::CbObject);
 	}
 }
 
@@ -732,7 +732,7 @@ void FHttpCacheStore::FPutPackageOp::OnPackagePutRefComplete(
 			{
 				return PutPackageOp->OnCompressedBlobUploadComplete(HttpResult, Request);
 			},
-			EHttpContentType::CompressedBinary);
+			EHttpMediaType::CompressedBinary);
 	}
 }
 
@@ -904,7 +904,7 @@ void FHttpCacheStore::FGetRecordOp::GetDataBatch(
 
 		TStringBuilder<256> CompressedBlobsUri;
 		CompressedBlobsUri << "api/v1/compressed-blobs/" << CacheStore.StructuredNamespace << "/" << RawHash;
-		Request->EnqueueAsyncDownload(Owner, Pool, *CompressedBlobsUri, MoveTemp(OnHttpRequestComplete), EHttpContentType::AnyContentType, { 404 });
+		Request->EnqueueAsyncDownload(Owner, Pool, *CompressedBlobsUri, MoveTemp(OnHttpRequestComplete), EHttpMediaType::Any, { 404 });
 	}
 }
 
@@ -1138,7 +1138,7 @@ void FHttpCacheStore::FGetRecordOp::DataProbablyExistsBatch(
 		return FHttpRequest::ECompletionBehavior::Done;
 	};
 	FCompositeBuffer DummyBuffer;
-	Request->EnqueueAsyncPost(Owner, Pool, *CompressedBlobsUri, DummyBuffer, MoveTemp(OnHttpRequestComplete), EHttpContentType::FormUrlEncoded);
+	Request->EnqueueAsyncPost(Owner, Pool, *CompressedBlobsUri, DummyBuffer, MoveTemp(OnHttpRequestComplete), EHttpMediaType::FormUrlEncoded);
 }
 
 void FHttpCacheStore::FGetRecordOp::FinishDataStep(bool bSuccess, uint64 InBytesReceived)
@@ -1356,7 +1356,7 @@ bool FHttpCacheStore::AcquireAccessToken()
 			auto OAuthFormDataUTF8 = FTCHARToUTF8(*OAuthFormData);
 			FormData.Append((uint8*)OAuthFormDataUTF8.Get(), OAuthFormDataUTF8.Length());
 
-			Result = Request.PerformBlockingPost(*Uri, FCompositeBuffer(FSharedBuffer::MakeView(FormData.GetData(), FormData.Num())), EHttpContentType::FormUrlEncoded);
+			Result = Request.PerformBlockingPost(*Uri, FCompositeBuffer(FSharedBuffer::MakeView(FormData.GetData(), FormData.Num())), EHttpMediaType::FormUrlEncoded);
 		}
 
 		if (Result == FHttpRequest::EResult::Success && Request.GetResponseCode() == 200)
@@ -1512,7 +1512,7 @@ void FHttpCacheStore::GetCacheRecordOnlyAsync(
 		OnComplete({ Name, Key, UserData, Request->GetBytesReceived(), {}, EStatus::Error });
 		return FHttpRequest::ECompletionBehavior::Done;
 	};
-	Request->EnqueueAsyncDownload(Owner, Pool, *RefsUri, MoveTemp(OnHttpRequestComplete), EHttpContentType::CbObject, { 401, 404 });
+	Request->EnqueueAsyncDownload(Owner, Pool, *RefsUri, MoveTemp(OnHttpRequestComplete), EHttpMediaType::CbObject, { 401, 404 });
 }
 
 void FHttpCacheStore::PutCacheRecordAsync(
@@ -1751,7 +1751,7 @@ void FHttpCacheStore::GetCacheValueAsync(
 		return FHttpRequest::ECompletionBehavior::Done;
 	};
 
-	Request->EnqueueAsyncDownload(Owner, Pool, *RefsUri, MoveTemp(OnHttpRequestComplete), EHttpContentType::UnspecifiedContentType, { 401, 404 });
+	Request->EnqueueAsyncDownload(Owner, Pool, *RefsUri, MoveTemp(OnHttpRequestComplete), EHttpMediaType::UnspecifiedContentType, { 401, 404 });
 }
 
 void FHttpCacheStore::GetCacheRecordAsync(
@@ -1914,7 +1914,7 @@ void FHttpCacheStore::RefCachedDataProbablyExistsBatchAsync(
 		return FHttpRequest::ECompletionBehavior::Done;
 	};
 
-	Request->EnqueueAsyncPost(Owner, Pool, *RefsUri, FCompositeBuffer(RequestFields.GetOuterBuffer()), MoveTemp(OnHttpRequestComplete), EHttpContentType::CbObject);
+	Request->EnqueueAsyncPost(Owner, Pool, *RefsUri, FCompositeBuffer(RequestFields.GetOuterBuffer()), MoveTemp(OnHttpRequestComplete), EHttpMediaType::CbObject);
 }
 
 void FHttpCacheStore::LegacyStats(FDerivedDataCacheStatsNode& OutNode)
diff --git a/Engine/Source/Developer/DerivedDataCache/Private/ZenCacheStore.cpp b/Engine/Source/Developer/DerivedDataCache/Private/ZenCacheStore.cpp
--- a/Engine/Source/Developer/DerivedDataCache/Private/ZenCacheStore.cpp
+++ b/Engine/Source/Developer/DerivedDataCache/Private/ZenCacheStore.cpp
@@ -963,7 +963,7 @@ FHttpRequest::EResult FZenCacheStore::PerformBlockingRpc(FHttpRequest& Request,
 	FCbPackage& OutResponse)
 {
 	return ParseRpcResponse(
-		Request.PerformBlockingPost(Uri, RequestObject.GetBuffer(), EHttpContentType::CbObject, EHttpContentType::CbPackage), Request, OutResponse);
+		Request.PerformBlockingPost(Uri, RequestObject.GetBuffer(), EHttpMediaType::CbObject, EHttpMediaType::CbPackage), Request, OutResponse);
 }
 
 FHttpRequest::EResult FZenCacheStore::PerformBlockingRpc(FHttpRequest& Request,
@@ -975,7 +975,7 @@ FHttpRequest::EResult FZenCacheStore::PerformBlockingRpc(FHttpRequest& Request,
 	UE::Zen::Http::SaveCbPackage(RequestPackage, PackageMemory);
 
 	return ParseRpcResponse(
-		Request.PerformBlockingPost(Uri, FCompositeBuffer(FSharedBuffer::MakeView(PackageMemory.GetView())), EHttpContentType::CbPackage, EHttpContentType::CbPackage), Request, OutResponse);
+		Request.PerformBlockingPost(Uri, FCompositeBuffer(FSharedBuffer::MakeView(PackageMemory.GetView())), EHttpMediaType::CbPackage, EHttpMediaType::CbPackage), Request, OutResponse);
 }
 
 void FZenCacheStore::EnqueueAsyncRpc(FHttpRequest& Request,
@@ -993,7 +993,7 @@ void FZenCacheStore::EnqueueAsyncRpc(FHttpRequest& Request,
 		OnComplete(ParsedResponseResult, Request, Response);
 		return FHttpRequest::ECompletionBehavior::Done;
 	};
-	Request.EnqueueAsyncPost(Owner, Pool, Uri, RequestObject.GetBuffer(), MoveTemp(OnHttpRequestComplete), EHttpContentType::CbObject, EHttpContentType::CbPackage);
+	Request.EnqueueAsyncPost(Owner, Pool, Uri, RequestObject.GetBuffer(), MoveTemp(OnHttpRequestComplete), EHttpMediaType::CbObject, EHttpMediaType::CbPackage);
 }
 
 void FZenCacheStore::EnqueueAsyncRpc(FHttpRequest& Request,
@@ -1015,7 +1015,7 @@ void FZenCacheStore::EnqueueAsyncRpc(FHttpRequest& Request,
 	UE::Zen::Http::SaveCbPackage(RequestPackage, PackageMemory);
 	uint64 PackageMemorySize = PackageMemory.TotalSize();
 	FSharedBuffer PackageSharedBuffer = FSharedBuffer::TakeOwnership(PackageMemory.ReleaseOwnership(), PackageMemorySize, FMemory::Free);
-	Request.EnqueueAsyncPost(Owner, Pool, Uri, FCompositeBuffer(PackageSharedBuffer), MoveTemp(OnHttpRequestComplete), EHttpContentType::CbPackage, EHttpContentType::CbPackage);
+	Request.EnqueueAsyncPost(Owner, Pool, Uri, FCompositeBuffer(PackageSharedBuffer), MoveTemp(OnHttpRequestComplete), EHttpMediaType::CbPackage, EHttpMediaType::CbPackage);
 }
 
 void FZenCacheStore::LegacyStats(FDerivedDataCacheStatsNode& OutNode)
