Add an optional "DeactivateAt" attribute for DDC filesystem cache stores to allow the store to be deactivated for the session if a speedtest is run and the speedtest indicates that the latency is higher than the configured threshold.  Default state is that there is no deactivation threshold.

diff --git a/Engine/Source/Developer/DerivedDataCache/Private/FileSystemCacheStore.cpp b/Engine/Source/Developer/DerivedDataCache/Private/FileSystemCacheStore.cpp
--- a/Engine/Source/Developer/DerivedDataCache/Private/FileSystemCacheStore.cpp
+++ b/Engine/Source/Developer/DerivedDataCache/Private/FileSystemCacheStore.cpp
@@ -26,6 +26,7 @@
 #include "Misc/CommandLine.h"
 #include "Misc/ConfigCacheIni.h"
 #include "Misc/CoreMisc.h"
+#include "Misc/DateTime.h"
 #include "Misc/FileHelper.h"
 #include "Misc/Guid.h"
 #include "Misc/MessageDialog.h"
@@ -33,6 +34,7 @@
 #include "Misc/PathViews.h"
 #include "Misc/ScopeExit.h"
 #include "Misc/ScopeLock.h"
+#include "Misc/ScopeRWLock.h"
 #include "Misc/StringBuilder.h"
 #include "ProfilingDebugging/CookStats.h"
 #include "ProfilingDebugging/CountersTrace.h"
@@ -41,8 +43,11 @@
 #include "Serialization/CompactBinaryPackage.h"
 #include "Serialization/CompactBinaryValidation.h"
 #include "Serialization/CompactBinaryWriter.h"
+#include "Tasks/Task.h"
 #include "Templates/Greater.h"
 
+#include <atomic>
+
 namespace UE::DerivedData
 {
 
@@ -744,11 +749,13 @@ public:
 		const TCHAR* AccessLogPath,
 		ECacheStoreFlags& OutFlags);
 
-	bool RunSpeedTest(
-		double InSkipTestsIfSeeksExceedMS,
+	static bool RunSpeedTest(
+		FString CachePath,
+		bool bReadOnly,
+		bool bSeekTimeOnly,
 		double& OutSeekTimeMS,
 		double& OutReadSpeedMBs,
-		double& OutWriteSpeedMBs) const;
+		double& OutWriteSpeedMBs);
 
 	// ICacheStore Interface
 
@@ -826,6 +833,7 @@ private:
 	[[nodiscard]] TUniquePtr<FArchive> OpenFileRead(FStringBuilderBase& Path, FStringView DebugName) const;
 
 	[[nodiscard]] bool FileExists(FStringBuilderBase& Path) const;
+	[[nodiscard]] bool IsDeactivatedForPerformance();
 
 private:
 	/** Base path we are storing the cache files in. */
@@ -836,6 +844,7 @@ private:
 	bool		bReadOnly;
 	/** If true, CachedDataProbablyExists will update the file timestamps. */
 	bool		bTouch;
+
 	/** Age of file when it should be deleted from DDC cache. */
 	double		DaysToDeleteUnusedFiles;
 
@@ -856,6 +865,21 @@ private:
 	FDerivedDataCacheSpeedStats SpeedStats;
 
 	TUniquePtr<FFileSystemCacheStoreMaintainer> Maintainer;
+
+	TUniquePtr<FFileSystemCacheStoreMaintainerParams> DeactivationDeferredMaintainerParams;
+
+	enum class EPerformanceReEvaluationResult
+	{
+		Invalid = 0,
+		PerformanceActivate,
+		PerformanceDeactivate
+	};
+	FRWLock PerformanceReEvaluationTaskLock;
+	Tasks::TTask<std::atomic<EPerformanceReEvaluationResult>> PerformanceReEvaluationTask;
+	std::atomic<int64> LastPerformanceEvaluationTicks;
+	std::atomic<bool> bDeactivedForPerformance;
+	bool bDeactivationDeferredClean;
+	float DeactivateAtMS;
 };
 
 FFileSystemCacheStore::FFileSystemCacheStore(
@@ -868,6 +892,9 @@ FFileSystemCacheStore::FFileSystemCacheStore(
 	, bReadOnly(false)
 	, bTouch(false)
 	, DaysToDeleteUnusedFiles(15.0)
+	, bDeactivedForPerformance(false)
+	, bDeactivationDeferredClean(false)
+	, DeactivateAtMS(-1.f)
 {
 	// If we find a platform that has more stringent limits, this needs to be rethought.
 	checkf(GMaxCacheRootLen + GMaxCacheKeyLen <= FPlatformMisc::GetMaxPathLength(),
@@ -898,6 +925,7 @@ FFileSystemCacheStore::FFileSystemCacheStore(
 	/* Speeds faster than this are ok. Everything else is slow. This value can be overridden in the ini file */
 	float ConsiderSlowAtMS = 50;
 	FParse::Value(InParams, TEXT("ConsiderSlowAt="), ConsiderSlowAtMS);
+	FParse::Value(InParams, TEXT("DeactivateAt="), DeactivateAtMS);
 
 	// can skip the speed test so everything acts as local (e.g. 4.25 and earlier behavior).
 	bool SkipSpeedTest = !WITH_EDITOR || FParse::Param(FCommandLine::Get(), TEXT("DDCSkipSpeedTest"));
@@ -909,19 +937,32 @@ FFileSystemCacheStore::FFileSystemCacheStore(
 		UE_LOG(LogDerivedDataCache, Log, TEXT("Skipping speed test to %s. Assuming local performance"), *CachePath);
 	}
 
-	if (!SkipSpeedTest && !RunSpeedTest(ConsiderSlowAtMS * 2, SpeedStats.LatencyMS, SpeedStats.ReadSpeedMBs, SpeedStats.WriteSpeedMBs))
+	if (!SkipSpeedTest &&
+		!RunSpeedTest(CachePath,
+			bReadOnly,
+			false /* bSeekTimeOnly */,
+			SpeedStats.LatencyMS,
+			SpeedStats.ReadSpeedMBs,
+			SpeedStats.WriteSpeedMBs))
 	{
+		LastPerformanceEvaluationTicks.store(FDateTime::UtcNow().GetTicks(), std::memory_order_relaxed);
+
 		OutFlags = ECacheStoreFlags::None;
 		UE_LOG(LogDerivedDataCache, Warning, TEXT("No read or write access to %s"), *CachePath);
 	}
 	else
 	{
+		LastPerformanceEvaluationTicks.store(FDateTime::UtcNow().GetTicks(), std::memory_order_relaxed);
+
 		bool bReadTestPassed = SpeedStats.ReadSpeedMBs > 0.0;
 		bool bWriteTestPassed = SpeedStats.WriteSpeedMBs > 0.0;
 
 		// if we failed writes mark this as read only
 		bReadOnly = bReadOnly || !bWriteTestPassed;
 
+		const bool bLocalDeactivatedForPerformance = (DeactivateAtMS > 0.f) && (SpeedStats.LatencyMS >= DeactivateAtMS);
+		bDeactivedForPerformance.store(bLocalDeactivatedForPerformance, std::memory_order_relaxed);
+
 		// classify and report on these times
 		if (SpeedStats.LatencyMS < 1)
 		{
@@ -942,8 +983,34 @@ FFileSystemCacheStore::FFileSystemCacheStore(
 
 		UE_LOG(LogDerivedDataCache, Display,
 			TEXT("Performance to %s: Latency=%.02fms. RandomReadSpeed=%.02fMBs, RandomWriteSpeed=%.02fMBs. "
-			     "Assigned SpeedClass '%s'"),
-			*CachePath, SpeedStats.LatencyMS, SpeedStats.ReadSpeedMBs, SpeedStats.WriteSpeedMBs, LexToString(SpeedClass));
+				 "Assigned SpeedClass '%s'"),
+			*CachePath,
+			SpeedStats.LatencyMS,
+			SpeedStats.ReadSpeedMBs,
+			SpeedStats.WriteSpeedMBs,
+			LexToString(SpeedClass));
+
+		if (bLocalDeactivatedForPerformance)
+		{
+			if (GIsBuildMachine)
+			{
+				UE_LOG(LogDerivedDataCache, Display,
+					TEXT("%s: Performance does not meet minimum criteria. "
+						 "It will be deactivated until performance measurements improve. "
+						 "If this is consistent, consider disabling this cache store through "
+						 "environment variables or other configuration."),
+					*CachePath);
+			}
+			else
+			{
+				UE_LOG(LogDerivedDataCache, Warning,
+					TEXT("%s: Performance does not meet minimum criteria. "
+						 "It will be deactivated until performance measurements improve. "
+						 "If this is consistent, consider disabling this cache store through "
+						 "environment variables or other configuration."),
+					*CachePath);
+			}
+		}
 
 		if (SpeedClass <= EBackendSpeedClass::Slow && !bReadOnly)
 		{
@@ -982,42 +1049,60 @@ FFileSystemCacheStore::FFileSystemCacheStore(
 			FParse::Bool(InParams, TEXT("DeleteUnused="), bDeleteUnused);
 			bDeleteUnused = bDeleteUnused && !FParse::Param(FCommandLine::Get(), TEXT("NODDCCLEANUP"));
 
+			if (bClean && bLocalDeactivatedForPerformance)
+			{
+				bDeactivationDeferredClean = true;
+			}
+
 			if (bClean || bDeleteUnused)
 			{
-				FFileSystemCacheStoreMaintainerParams MaintainerParams;
-				MaintainerParams.MaxFileAge = FTimespan::FromDays(DaysToDeleteUnusedFiles);
+				FFileSystemCacheStoreMaintainerParams* MaintainerParams;
+				FFileSystemCacheStoreMaintainerParams LocalMaintainerParams;
+				if (bLocalDeactivatedForPerformance)
+				{
+					DeactivationDeferredMaintainerParams = MakeUnique<FFileSystemCacheStoreMaintainerParams>();
+					MaintainerParams = DeactivationDeferredMaintainerParams.Get();
+				}
+				else
+				{
+					MaintainerParams = &LocalMaintainerParams;
+				}
+				MaintainerParams->MaxFileAge = FTimespan::FromDays(DaysToDeleteUnusedFiles);
 				if (bDeleteUnused)
 				{
-					if (!FParse::Value(InParams, TEXT("MaxFileChecksPerSec="), MaintainerParams.MaxScanRate))
+					if (!FParse::Value(InParams, TEXT("MaxFileChecksPerSec="), MaintainerParams->MaxScanRate))
 					{
 						int32 MaxFileScanRate;
 						if (GConfig->GetInt(TEXT("DDCCleanup"), TEXT("MaxFileChecksPerSec"), MaxFileScanRate, GEngineIni))
 						{
-							MaintainerParams.MaxScanRate = uint32(MaxFileScanRate);
+							MaintainerParams->MaxScanRate = uint32(MaxFileScanRate);
 						}
 					}
-					FParse::Value(InParams, TEXT("FoldersToClean="), MaintainerParams.MaxDirectoryScanCount);
+					FParse::Value(InParams, TEXT("FoldersToClean="), MaintainerParams->MaxDirectoryScanCount);
 				}
 				else
 				{
-					MaintainerParams.ScanFrequency = FTimespan::MaxValue();
+					MaintainerParams->ScanFrequency = FTimespan::MaxValue();
 				}
 				double TimeToWaitAfterInit;
 				if (bClean)
 				{
-					MaintainerParams.TimeToWaitAfterInit = FTimespan::Zero();
+					MaintainerParams->TimeToWaitAfterInit = FTimespan::Zero();
 				}
 				else if (GConfig->GetDouble(TEXT("DDCCleanup"), TEXT("TimeToWaitAfterInit"), TimeToWaitAfterInit, GEngineIni))
 				{
-					MaintainerParams.TimeToWaitAfterInit = FTimespan::FromSeconds(TimeToWaitAfterInit);
+					MaintainerParams->TimeToWaitAfterInit = FTimespan::FromSeconds(TimeToWaitAfterInit);
 				}
 
-				Maintainer = MakeUnique<FFileSystemCacheStoreMaintainer>(MaintainerParams, CachePath);
-
-				if (bClean)
+				if (!bLocalDeactivatedForPerformance)
 				{
-					Maintainer->BoostPriority();
-					Maintainer->WaitForIdle();
+					Maintainer = MakeUnique<FFileSystemCacheStoreMaintainer>(*MaintainerParams, CachePath);
+
+					if (bClean)
+					{
+						Maintainer->BoostPriority();
+						Maintainer->WaitForIdle();
+					}
 				}
 			}
 		}
@@ -1035,10 +1120,12 @@ FFileSystemCacheStore::FFileSystemCacheStore(
 }
 
 bool FFileSystemCacheStore::RunSpeedTest(
-	const double InSkipTestsIfSeeksExceedMS,
+	FString CachePath,
+	bool bReadOnly,
+	bool bSeekTimeOnly,
 	double& OutSeekTimeMS,
 	double& OutReadSpeedMBs,
-	double& OutWriteSpeedMBs) const
+	double& OutWriteSpeedMBs)
 {
 	SCOPED_BOOT_TIMING("RunSpeedTest");
 
@@ -1110,19 +1197,10 @@ bool FFileSystemCacheStore::RunSpeedTest(
 
 	UE_LOG(LogDerivedDataCache, Verbose, TEXT("Stat tests to %s took %.02f seconds"), *CachePath, TotalSeekTime);
 
-	// if seek times are very slow do a single read/write test just to confirm access
-	/*if (OutSeekTimeMS >= InSkipTestsIfSeeksExceedMS)
+	if (bSeekTimeOnly)
 	{
-		UE_LOG(LogDerivedDataCache, Warning,
-			TEXT("Limiting read/write speed tests due to seek times of %.02f exceeding %.02fms. ")
-			TEXT("Values will be inaccurate."), OutSeekTimeMS, InSkipTestsIfSeeksExceedMS);
-
-		FString Path = TestFileEntries.begin()->Key;
-		int Size = TestFileEntries.begin()->Value;
-
-		TestFileEntries.Reset();
-		TestFileEntries.Add(Path, Size);
-	}*/
+		return true;
+	}
 
 	// create any files that were missing
 	if (!bReadOnly)
@@ -1468,11 +1546,16 @@ bool FFileSystemCacheStore::PutCacheRecord(
 	const FCacheRecordPolicy& Policy,
 	uint64& OutWriteSize)
 {
-	if (bReadOnly)
+	const bool bLocalDeactivatedForPerformance = IsDeactivatedForPerformance();
+	if (bLocalDeactivatedForPerformance || bReadOnly)
 	{
 		UE_LOG(LogDerivedDataCache, VeryVerbose,
-			TEXT("%s: Skipped put of %s from '%.*s' because this cache store is read-only"),
-			*CachePath, *WriteToString<96>(Record.GetKey()), Name.Len(), Name.GetData());
+			TEXT("%s: Skipped put of %s from '%.*s' because this cache store is %s"),
+			*CachePath,
+			*WriteToString<96>(Record.GetKey()),
+			Name.Len(),
+			Name.GetData(),
+			bLocalDeactivatedForPerformance ? TEXT("deactivated due to low performance") : TEXT("read-only"));
 		return false;
 	}
 
@@ -1636,10 +1719,17 @@ FOptionalCacheRecord FFileSystemCacheStore::GetCacheRecordOnly(
 {
 	// Skip the request if querying the cache is disabled.
 	const ECachePolicy QueryFlag = SpeedClass == EBackendSpeedClass::Local ? ECachePolicy::QueryLocal : ECachePolicy::QueryRemote;
-	if (!EnumHasAnyFlags(Policy.GetRecordPolicy(), QueryFlag))
-	{
-		UE_LOG(LogDerivedDataCache, VeryVerbose, TEXT("%s: Skipped get of %s from '%.*s' due to cache policy"),
-			*CachePath, *WriteToString<96>(Key), Name.Len(), Name.GetData());
+	const bool bLocalDeactivatedForPerformance = IsDeactivatedForPerformance();
+	if (bLocalDeactivatedForPerformance || !EnumHasAnyFlags(Policy.GetRecordPolicy(), QueryFlag))
+	{
+		UE_LOG(LogDerivedDataCache, VeryVerbose, TEXT("%s: Skipped get of %s from '%.*s' %s"),
+			*CachePath,
+			*WriteToString<96>(Key),
+			Name.Len(),
+			Name.GetData(),
+			bLocalDeactivatedForPerformance ?
+				TEXT("because this cache store is deactivated due to low performance") :
+				TEXT("due to cache policy"));
 		return FOptionalCacheRecord();
 	}
 
@@ -1751,11 +1841,18 @@ bool FFileSystemCacheStore::PutCacheValue(
 	const ECachePolicy Policy,
 	uint64& OutWriteSize)
 {
-	if (bReadOnly)
+	const bool bLocalDeactivatedForPerformance = IsDeactivatedForPerformance();
+	if (bLocalDeactivatedForPerformance || bReadOnly)
 	{
 		UE_LOG(LogDerivedDataCache, VeryVerbose,
-			TEXT("%s: Skipped put of %s from '%.*s' because this cache store is read-only"),
-			*CachePath, *WriteToString<96>(Key), Name.Len(), Name.GetData());
+			TEXT("%s: Skipped put of %s from '%.*s' because this cache store is %s"),
+			*CachePath,
+			*WriteToString<96>(Key),
+			Name.Len(),
+			Name.GetData(),
+			bLocalDeactivatedForPerformance ?
+				TEXT("deactivated due to low performance") :
+				TEXT("read-only"));
 		return false;
 	}
 
@@ -1898,10 +1995,17 @@ bool FFileSystemCacheStore::GetCacheValueOnly(
 {
 	// Skip the request if querying the cache is disabled.
 	const ECachePolicy QueryFlag = SpeedClass == EBackendSpeedClass::Local ? ECachePolicy::QueryLocal : ECachePolicy::QueryRemote;
-	if (!EnumHasAnyFlags(Policy, QueryFlag))
-	{
-		UE_LOG(LogDerivedDataCache, VeryVerbose, TEXT("%s: Skipped get of %s from '%.*s' due to cache policy"),
-			*CachePath, *WriteToString<96>(Key), Name.Len(), Name.GetData());
+	const bool bLocalDeactivatedForPerformance = IsDeactivatedForPerformance();
+	if (bLocalDeactivatedForPerformance || !EnumHasAnyFlags(Policy, QueryFlag))
+	{
+		UE_LOG(LogDerivedDataCache, VeryVerbose, TEXT("%s: Skipped get of %s from '%.*s' %s"),
+			*CachePath,
+			*WriteToString<96>(Key),
+			Name.Len(),
+			Name.GetData(),
+			bLocalDeactivatedForPerformance ?
+			TEXT("because this cache store is deactivated due to low performance") :
+			TEXT("due to cache policy"));
 		return false;
 	}
 
@@ -2348,6 +2452,119 @@ bool FFileSystemCacheStore::FileExists(FStringBuilderBase& Path) const
 	return true;
 }
 
+bool FFileSystemCacheStore::IsDeactivatedForPerformance()
+{
+	if ((DeactivateAtMS <= 0.f) || !bDeactivedForPerformance.load(std::memory_order_relaxed))
+	{
+		return false;
+	}
+
+	// Look for an opportunity to consume the output of an existing completed performance evaluation task
+	{
+		FReadScopeLock ReadLock(PerformanceReEvaluationTaskLock);
+		if (PerformanceReEvaluationTask.IsValid())
+		{
+			if (PerformanceReEvaluationTask.IsCompleted())
+			{
+				EPerformanceReEvaluationResult Result =
+					PerformanceReEvaluationTask.GetResult().exchange(
+						EPerformanceReEvaluationResult::Invalid,
+						std::memory_order_relaxed);
+
+				if (Result != EPerformanceReEvaluationResult::Invalid)
+				{
+					LastPerformanceEvaluationTicks.store(FDateTime::UtcNow().GetTicks(), std::memory_order_relaxed);
+					bool bLocalDeactivatedForPerformance =
+						Result == EPerformanceReEvaluationResult::PerformanceDeactivate;
+
+					if (!bLocalDeactivatedForPerformance)
+					{
+						// We're no longer deactivated for performance.  If maintenance was deferred, do it now.
+						if (FFileSystemCacheStoreMaintainerParams* MaintainerParams = DeactivationDeferredMaintainerParams.Get())
+						{
+							Maintainer = MakeUnique<FFileSystemCacheStoreMaintainer>(*MaintainerParams, CachePath);
+							DeactivationDeferredMaintainerParams.Reset();
+
+							if (bDeactivationDeferredClean)
+							{
+								Maintainer->BoostPriority();
+								Maintainer->WaitForIdle();
+							}
+						}
+
+						UE_LOG(LogDerivedDataCache, Display,
+							TEXT("%s: Performance has improved and meets minimum performance criteria. "
+								"It will be reactivated now."),
+							*CachePath);
+					}
+
+					bDeactivedForPerformance.store(bLocalDeactivatedForPerformance, std::memory_order_relaxed);
+					return bLocalDeactivatedForPerformance;
+				}
+			}
+			else
+			{
+				// Avoid attempting to get a write lock and see if you can spawn a new evaluation task
+				return true;
+			}
+		}
+	}
+
+	// Look for an opportunity to start a new performance evaluation task
+	FTimespan TimespanSinceLastPerfEval =
+		FDateTime::UtcNow() - FDateTime(LastPerformanceEvaluationTicks.load(std::memory_order_relaxed));
+
+	if (TimespanSinceLastPerfEval > FTimespan::FromMinutes(1))
+	{
+		FWriteScopeLock WriteLock(PerformanceReEvaluationTaskLock);
+		// After acquiring the write lock, ensure that the task hasn't been re-launched
+		// (and possibly completed and consumed) by someone else while we were waiting.
+		// This is evaluated by checking that:
+		// 1. Task is invalid or task is valid and has a consumed result
+		// and
+		// 2. Task consumption time is still larger than our re-evaluation interval
+		if (!PerformanceReEvaluationTask.IsValid() ||
+			(PerformanceReEvaluationTask.IsCompleted() &&
+				(PerformanceReEvaluationTask.GetResult().load(std::memory_order_relaxed) ==
+					EPerformanceReEvaluationResult::Invalid)
+			))
+		{
+			TimespanSinceLastPerfEval =
+				FDateTime::UtcNow() - FDateTime(LastPerformanceEvaluationTicks.load(std::memory_order_relaxed));
+
+			if (TimespanSinceLastPerfEval > FTimespan::FromMinutes(1))
+			{
+				PerformanceReEvaluationTask = Tasks::Launch(TEXT("FFileSystemCacheStore::ReEvaluatePerformance"),
+					[CachePath = this->CachePath, DeactivateAtMS = this->DeactivateAtMS]() ->
+						std::atomic<EPerformanceReEvaluationResult>
+					{
+						check(DeactivateAtMS > 0.f);
+
+						FDerivedDataCacheSpeedStats LocalSpeedStats;
+						LocalSpeedStats.ReadSpeedMBs = 999;
+						LocalSpeedStats.WriteSpeedMBs = 999;
+						LocalSpeedStats.LatencyMS = 0;
+
+						RunSpeedTest(CachePath,
+							true /* bReadOnly */,
+							true /* bSeekTimeOnly */,
+							LocalSpeedStats.LatencyMS,
+							LocalSpeedStats.ReadSpeedMBs,
+							LocalSpeedStats.WriteSpeedMBs);
+
+						if (LocalSpeedStats.LatencyMS >= DeactivateAtMS)
+						{
+							return EPerformanceReEvaluationResult::PerformanceDeactivate;
+						}
+						return EPerformanceReEvaluationResult::PerformanceActivate;
+					});
+			}
+		}
+	}
+
+	return true;
+}
+
 ILegacyCacheStore* CreateFileSystemCacheStore(
 	const TCHAR* CachePath,
 	const TCHAR* Params,
