DDC: Fixed auth token refresh in HttpCacheStore to unregister from the ticker

Removed duplicate code, unnecessary scopes, and unused variables as well.

#rnx

diff --git a/Engine/Source/Developer/DerivedDataCache/Private/HttpCacheStore.cpp b/Engine/Source/Developer/DerivedDataCache/Private/HttpCacheStore.cpp
--- a/Engine/Source/Developer/DerivedDataCache/Private/HttpCacheStore.cpp
+++ b/Engine/Source/Developer/DerivedDataCache/Private/HttpCacheStore.cpp
@@ -194,7 +194,6 @@ private:
 	FString EffectiveDomain;
 	FString Namespace;
 	FString StructuredNamespace;
-	FString DefaultBucket;
 	FString OAuthProvider;
 	FString OAuthClientId;
 	FString OAuthSecret;
@@ -202,21 +201,26 @@ private:
 	FString OAuthProviderIdentifier;
 	FString OAuthAccessToken;
 
-	FCriticalSection AccessCs;
 	FDerivedDataCacheUsageStats UsageStats;
 	FBackendDebugOptions DebugOptions;
 	TUniquePtr<FHttpSharedData> SharedData;
 	TUniquePtr<FHttpRequestPool> GetRequestPools[2];
 	TUniquePtr<FHttpRequestPool> PutRequestPools[2];
 	TUniquePtr<FHttpRequestPool> NonBlockingRequestPools;
+
+	FCriticalSection AccessCs;
 	TUniquePtr<FHttpAccessToken> Access;
-	bool bIsUsable;
-	bool bReadOnly;
-	uint32 FailedLoginAttempts;
+	FTSTicker::FDelegateHandle RefreshAccessTokenHandle;
+	uint32 FailedLoginAttempts = 0;
+
+	bool bIsUsable = false;
+	bool bReadOnly = false;
+
 	static inline FHttpCacheStore* AnyInstance = nullptr;
 
 	bool IsServiceReady();
 	bool AcquireAccessToken();
+	void SetAccessToken(FStringView Token, double RefreshDelay = 0.0);
 	bool ShouldRetryOnError(FHttpRequest::EResult Result, int64 ResponseCode);
 	bool ShouldRetryOnError(int64 ResponseCode) { return ShouldRetryOnError(FHttpRequest::EResult::Success, ResponseCode); }
 
@@ -1190,17 +1194,13 @@ FHttpCacheStore::FHttpCacheStore(const FHttpCacheStoreParams& Params)
 	, EffectiveDomain(Params.Host)
 	, Namespace(Params.Namespace)
 	, StructuredNamespace(Params.StructuredNamespace)
-	, DefaultBucket(TEXT("default"))
 	, OAuthProvider(Params.OAuthProvider)
 	, OAuthClientId(Params.OAuthClientId)
 	, OAuthSecret(Params.OAuthSecret)
 	, OAuthScope(Params.OAuthScope)
 	, OAuthProviderIdentifier(Params.OAuthProviderIdentifier)
 	, OAuthAccessToken(Params.OAuthAccessToken)
-	, Access(nullptr)
-	, bIsUsable(false)
 	, bReadOnly(Params.bReadOnly)
-	, FailedLoginAttempts(0)
 {
 	SharedData = MakeUnique<FHttpSharedData>();
 	if (IsServiceReady() && AcquireAccessToken())
@@ -1265,6 +1265,11 @@ FHttpCacheStore::FHttpCacheStore(const FHttpCacheStoreParams& Params)
 
 FHttpCacheStore::~FHttpCacheStore()
 {
+	if (RefreshAccessTokenHandle.IsValid())
+	{
+		FTSTicker::GetCoreTicker().RemoveTicker(RefreshAccessTokenHandle);
+	}
+
 	if (AnyInstance == this)
 	{
 		AnyInstance = nullptr;
@@ -1299,186 +1304,150 @@ bool FHttpCacheStore::AcquireAccessToken()
 {
 	if (Domain.StartsWith(TEXT("http://localhost")))
 	{
-		UE_LOG(LogDerivedDataCache, Log, TEXT("Connecting to a local host '%s', so skipping authorization"), *Domain);
+		UE_LOG(LogDerivedDataCache, Log, TEXT("%s: Skipping authorization for connection to localhost."), *Domain);
 		return true;
 	}
 
-	// Avoid spamming the this if the service is down
+	// Avoid spamming this if the service is down.
 	if (FailedLoginAttempts > UE_HTTPDDC_MAX_FAILED_LOGIN_ATTEMPTS)
 	{
 		return false;
 	}
 
 	ensureMsgf(OAuthProvider.StartsWith(TEXT("http://")) || OAuthProvider.StartsWith(TEXT("https://")),
-		TEXT("The OAuth provider %s is not valid. Needs to be a fully qualified url."),
-		*OAuthProvider
-	);
+		TEXT("%s: OAuth provider %s must be a complete URI including the scheme."), *Domain, *OAuthProvider);
 
 	// In case many requests wants to update the token at the same time
 	// get the current serial while we wait to take the CS.
-	const uint32 WantsToUpdateTokenSerial = Access.IsValid() ? Access->GetSerial() : 0u;
+	const uint32 WantsToUpdateTokenSerial = Access ? Access->GetSerial() : 0;
+
+	FScopeLock Lock(&AccessCs);
+
+	// If the token was updated while we waited to take the lock, then it should now be valid.
+	if (Access && Access->GetSerial() > WantsToUpdateTokenSerial)
+	{
+		return true;
+	}
+
+	if (!OAuthAccessToken.IsEmpty())
+	{
+		SetAccessToken(OAuthAccessToken);
+		return true;
+	}
 
 	if (bHttpEnableOidc && !OAuthProviderIdentifier.IsEmpty())
 	{
 		FString AccessTokenString;
 		FDateTime TokenExpiresAt;
-		
+		if (FDesktopPlatformModule::Get()->GetOidcAccessToken(FPaths::RootDir(), FPaths::GetProjectFilePath(), OAuthProviderIdentifier, /* Unattended */ true, GWarn, AccessTokenString, TokenExpiresAt))
 		{
-			FScopeLock Lock(&AccessCs);
+			const double ExpiryTimeSeconds = (TokenExpiresAt - FDateTime::UtcNow()).GetTotalSeconds();
+			UE_LOG(LogDerivedDataCache, Display,
+				TEXT("%s: OidcToken: Logged in to HTTP DDC services. Expires at %s which is in %.0f seconds."),
+				*Domain, *TokenExpiresAt.ToString(), ExpiryTimeSeconds);
+			SetAccessToken(AccessTokenString, ExpiryTimeSeconds);
+			return true;
+		}
 
-			// Check if someone has beaten us to update the token, then it 
-			// should now be valid.
-			if (Access.IsValid() && Access->GetSerial() > WantsToUpdateTokenSerial)
-			{
-				return true;
-			}
+		UE_LOG(LogDerivedDataCache, Error, TEXT("%s: OidcToken: Failed to log in to HTTP services."), *Domain);
+		FailedLoginAttempts++;
+		return false;
+	}
 
-			if (!OAuthAccessToken.IsEmpty())
-			{
-				if (!Access)
-				{
-					Access = MakeUnique<FHttpAccessToken>();
-				}
-				Access->SetToken(OAuthAccessToken);
-				FailedLoginAttempts = 0;
-				return true;
-			}
+	const uint32 SchemeEnd = OAuthProvider.Find(TEXT("://")) + 3;
+	const uint32 DomainEnd = OAuthProvider.Find(TEXT("/"), ESearchCase::CaseSensitive, ESearchDir::FromStart, SchemeEnd);
+	FString AuthDomain(DomainEnd, *OAuthProvider);
+	FString Uri(*OAuthProvider + DomainEnd + 1);
 
-			if (FDesktopPlatformModule::Get()->GetOidcAccessToken(FPaths::RootDir(), FPaths::GetProjectFilePath(), OAuthProviderIdentifier, /* Unattended */ true, GWarn, AccessTokenString, TokenExpiresAt))
-			{
-				int32 ExpiryTimeSeconds = (TokenExpiresAt - FDateTime::UtcNow()).GetTotalSeconds();
-				if (!Access)
-				{
-					Access = MakeUnique<FHttpAccessToken>();
-				}
-				Access->SetToken(AccessTokenString);
-				UE_LOG(LogDerivedDataCache, Display, TEXT("OidcToken: Logged in to HTTP DDC services. Expires at %s which is in %d seconds."), *TokenExpiresAt.ToString(), ExpiryTimeSeconds);
+	FHttpRequest Request(*AuthDomain, *AuthDomain, nullptr, SharedData.Get(), false);
+	FHttpRequest::EResult Result = FHttpRequest::EResult::Success;
+	if (OAuthProvider.StartsWith(TEXT("http://localhost")))
+	{
+		// Simple unauthenticated call to a local endpoint that mimics the result from an OIDC provider.
+		Result = Request.PerformBlockingDownload(*Uri, nullptr);
+	}
+	else
+	{
+		// Needs client id and secret to authenticate with an actual OIDC provider.
 
-				//Schedule a refresh of the token ahead of expiry time (this will not work in commandlets)
-				if (!IsRunningCommandlet())
-				{
-					FTSTicker::GetCoreTicker().AddTicker(FTickerDelegate::CreateLambda(
-						[this](float DeltaTime)
-						{
-							this->AcquireAccessToken();
-							return false;
-						}
-					), ExpiryTimeSeconds - 20.0f);
-				}
-				// Reset failed login attempts, the service is indeed alive.
-				FailedLoginAttempts = 0;
-				return true;
+		// If contents of the secret string is a file path, resolve and read form data.
+		if (OAuthSecret.StartsWith(TEXT("file://")))
+		{
+			FString FilePath = OAuthSecret.Mid(7, OAuthSecret.Len() - 7);
+				FString SecretFileContents;
+			if (FFileHelper::LoadFileToString(SecretFileContents, *FilePath))
+			{
+				// Overwrite the filepath with the actual content.
+				OAuthSecret = SecretFileContents;
 			}
 			else
 			{
-				UE_LOG(LogDerivedDataCache, Error, TEXT("OidcToken: Failed to log in to HTTP services. "));
-				FailedLoginAttempts++;
+				UE_LOG(LogDerivedDataCache, Warning, TEXT("%s: Failed to read OAuth form data file (%s)."), *Domain, *OAuthSecret);
+				return false;
 			}
 		}
-		
-	}
-	else // oidc not enabled or not configured
-	{
-		FScopeLock Lock(&AccessCs);
 
-		// Check if someone has beaten us to update the token, then it 
-		// should now be valid.
-		if (Access.IsValid() && Access->GetSerial() > WantsToUpdateTokenSerial)
-		{
-			return true;
-		}
+		TUtf8StringBuilder<256> OAuthFormData;
+		OAuthFormData
+			<< ANSITEXTVIEW("client_id=") << OAuthClientId
+			<< ANSITEXTVIEW("&scope=") << OAuthScope
+			<< ANSITEXTVIEW("&grant_type=client_credentials")
+			<< ANSITEXTVIEW("&client_secret=") << OAuthSecret;
 
-		const uint32 SchemeEnd = OAuthProvider.Find(TEXT("://")) + 3;
-		const uint32 DomainEnd = OAuthProvider.Find(TEXT("/"), ESearchCase::CaseSensitive, ESearchDir::FromStart, SchemeEnd);
-		FString AuthDomain(DomainEnd, *OAuthProvider);
-		FString Uri(*OAuthProvider + DomainEnd + 1);
+		Result = Request.PerformBlockingPost(*Uri, FCompositeBuffer(FSharedBuffer::MakeView(MakeMemoryView(OAuthFormData))), EHttpMediaType::FormUrlEncoded);
+	}
 
-		FHttpRequest Request(*AuthDomain, *AuthDomain, nullptr, SharedData.Get(), false);
-		FHttpRequest::EResult Result = FHttpRequest::EResult::Success;
-		if (OAuthProvider.StartsWith(TEXT("http://localhost")))
-		{
-			// Simple unauthenticated call to a local endpoint that mimics
-			// the result from an OIDC provider.
-			Result = Request.PerformBlockingDownload(*Uri, nullptr);
-		}
-		else
+	if (Result == FHttpRequest::EResult::Success && Request.GetResponseCode() == 200)
+	{
+		if (TSharedPtr<FJsonObject> ResponseObject = Request.GetResponseAsJsonObject())
 		{
-			// Needs client id and secret to authenticate with an actual OIDC provider.
-
-			// If contents of the secret string is a file path, resolve and read form data.
-			if (OAuthSecret.StartsWith(TEXT("file://")))
+			FString AccessTokenString;
+			double ExpiryTimeSeconds = 0.0;
+			if (ResponseObject->TryGetStringField(TEXT("access_token"), AccessTokenString) &&
+				ResponseObject->TryGetNumberField(TEXT("expires_in"), ExpiryTimeSeconds))
 			{
-				FString FilePath = OAuthSecret.Mid(7, OAuthSecret.Len() - 7);
-					FString SecretFileContents;
-				if (FFileHelper::LoadFileToString(SecretFileContents, *FilePath))
-				{
-					// Overwrite the filepath with the actual content.
-					OAuthSecret = SecretFileContents;
-				}
-				else
-				{
-					UE_LOG(LogDerivedDataCache, Warning, TEXT("%s: Failed to read OAuth form data file (%s)."), *Request.GetName(), *OAuthSecret);
-					return false;
-				}
+				UE_LOG(LogDerivedDataCache, Display,
+					TEXT("%s: Logged in to HTTP DDC services. Expires in %.0f seconds."), *Domain, ExpiryTimeSeconds);
+				SetAccessToken(AccessTokenString, ExpiryTimeSeconds);
+				return true;
 			}
+		}
+	}
 
-			FString OAuthFormData = FString::Printf(
-				TEXT("client_id=%s&scope=%s&grant_type=client_credentials&client_secret=%s"),
-				*OAuthClientId,
-				*OAuthScope,
-				*OAuthSecret
-			);
+	UE_LOG(LogDerivedDataCache, Warning, TEXT("%s: Failed to log in to HTTP services. Server responded with code %d."), *Domain, Request.GetResponseCode());
+	FailedLoginAttempts++;
+	return false;
+}
 
-			TArray64<uint8> FormData;
-			auto OAuthFormDataUTF8 = FTCHARToUTF8(*OAuthFormData);
-			FormData.Append((uint8*)OAuthFormDataUTF8.Get(), OAuthFormDataUTF8.Length());
+void FHttpCacheStore::SetAccessToken(FStringView Token, double RefreshDelay)
+{
+	if (RefreshAccessTokenHandle.IsValid())
+	{
+		FTSTicker::GetCoreTicker().RemoveTicker(RefreshAccessTokenHandle);
+		RefreshAccessTokenHandle.Reset();
+	}
 
-			Result = Request.PerformBlockingPost(*Uri, FCompositeBuffer(FSharedBuffer::MakeView(FormData.GetData(), FormData.Num())), EHttpMediaType::FormUrlEncoded);
-		}
+	if (!Access)
+	{
+		Access = MakeUnique<FHttpAccessToken>();
+	}
+	Access->SetToken(Token);
 
-		if (Result == FHttpRequest::EResult::Success && Request.GetResponseCode() == 200)
-		{
-			TSharedPtr<FJsonObject> ResponseObject = Request.GetResponseAsJsonObject();
-			if (ResponseObject)
+	// Schedule a refresh of the token ahead of expiry time (this will not work in commandlets)
+	constexpr float RefreshGracePeriod = 20.0f;
+	if (RefreshDelay > RefreshGracePeriod && !IsRunningCommandlet())
+	{
+		RefreshAccessTokenHandle = FTSTicker::GetCoreTicker().AddTicker(FTickerDelegate::CreateLambda(
+			[this](float DeltaTime)
 			{
-				FString AccessTokenString;
-				int32 ExpiryTimeSeconds = 0;
-				int32 CurrentTimeSeconds = int32(FPlatformTime::ToSeconds(FPlatformTime::Cycles()));
-
-				if (ResponseObject->TryGetStringField(TEXT("access_token"), AccessTokenString) &&
-					ResponseObject->TryGetNumberField(TEXT("expires_in"), ExpiryTimeSeconds))
-				{
-					if (!Access)
-					{
-						Access = MakeUnique<FHttpAccessToken>();
-					}
-					Access->SetToken(AccessTokenString);
-					UE_LOG(LogDerivedDataCache, Display, TEXT("%s: Logged in to HTTP DDC services. Expires in %d seconds."), *Request.GetName(), ExpiryTimeSeconds);
-
-					//Schedule a refresh of the token ahead of expiry time (this will not work in commandlets)
-					if (!IsRunningCommandlet())
-					{
-						FTSTicker::GetCoreTicker().AddTicker(FTickerDelegate::CreateLambda(
-							[this](float DeltaTime)
-							{
-								this->AcquireAccessToken();
-								return false;
-							}
-						), ExpiryTimeSeconds - 20.0f);
-					}
-					// Reset failed login attempts, the service is indeed alive.
-					FailedLoginAttempts = 0;
-					return true;
-				}
+				AcquireAccessToken();
+				return false;
 			}
-		}
-		else
-		{
-			UE_LOG(LogDerivedDataCache, Warning, TEXT("%s: Failed to log in to HTTP services. Server responed with code %d."), *Request.GetName(), Request.GetResponseCode());
-			FailedLoginAttempts++;
-		}
+		), RefreshDelay - RefreshGracePeriod);
 	}
-	return false;
+
+	// Reset failed login attempts, the service is indeed alive.
+	FailedLoginAttempts = 0;
 }
 
 bool FHttpCacheStore::ShouldRetryOnError(FHttpRequest::EResult Result, int64 ResponseCode)
@@ -2415,8 +2384,8 @@ void FHttpCacheStoreParams::Parse(const TCHAR* NodeName, const TCHAR* Config)
 		if (!AccessToken.IsEmpty())
 		{
 			OAuthAccessToken = AccessToken;
-			// we do not echo the access token as its sensitive information
-			UE_LOG(LogDerivedDataCache, Log, TEXT("%s: Found oauth access token in %s ."), NodeName, *OverrideName);
+			// We do not log the access token as it is sensitive information.
+			UE_LOG(LogDerivedDataCache, Log, TEXT("%s: Found OAuth access token in %s."), NodeName, *OverrideName);
 		}
 	}
 
