DDC: Enabled Horde Storage on Mac and Linux

diff --git a/Engine/Source/Developer/DerivedDataCache/DerivedDataCache.Build.cs b/Engine/Source/Developer/DerivedDataCache/DerivedDataCache.Build.cs
--- a/Engine/Source/Developer/DerivedDataCache/DerivedDataCache.Build.cs
+++ b/Engine/Source/Developer/DerivedDataCache/DerivedDataCache.Build.cs
@@ -17,7 +17,7 @@ public class DerivedDataCache : ModuleRules
 		AddEngineThirdPartyPrivateStaticDependencies(Target, "OpenSSL");
 
 		// Platform-specific opt-in
-		PrivateDefinitions.Add($"WITH_HTTP_DDC_BACKEND={(Target.Platform == UnrealTargetPlatform.Win64 ? 1 : 0)}");
+		PrivateDefinitions.Add($"WITH_HTTP_DDC_BACKEND=1");
 		PrivateDefinitions.Add($"WITH_S3_DDC_BACKEND={(Target.Platform == UnrealTargetPlatform.Win64 ? 1 : 0)}");
 	}
 }
