DDC: Removed filtering support for legacy cache keys

#rnx

diff --git a/Engine/Source/Developer/DerivedDataCache/Internal/DerivedDataCacheKeyFilter.h b/Engine/Source/Developer/DerivedDataCache/Internal/DerivedDataCacheKeyFilter.h
--- a/Engine/Source/Developer/DerivedDataCache/Internal/DerivedDataCacheKeyFilter.h
+++ b/Engine/Source/Developer/DerivedDataCache/Internal/DerivedDataCacheKeyFilter.h
@@ -45,9 +45,6 @@ public:
 	/** Returns true if the filter matches the key. The output is the same until the salt changes. */
 	UE_API bool IsMatch(const FCacheKey& Key) const;
 
-	/** Returns true if the filter matches the legacy key. The output is the same until the salt changes. */
-	UE_API bool IsLegacyMatch(const TCHAR* LegacyKey) const;
-
 	/** Returns true if this filter can ever match a cache key. */
 	inline explicit operator bool() const { return State.IsValid(); }
 
diff --git a/Engine/Source/Developer/DerivedDataCache/Private/DerivedDataBackendInterface.h b/Engine/Source/Developer/DerivedDataCache/Private/DerivedDataBackendInterface.h
--- a/Engine/Source/Developer/DerivedDataCache/Private/DerivedDataBackendInterface.h
+++ b/Engine/Source/Developer/DerivedDataCache/Private/DerivedDataBackendInterface.h
@@ -80,9 +80,6 @@ struct FBackendDebugOptions
 
 	bool ShouldSimulatePutMiss(const FCacheKey& Key);
 	bool ShouldSimulateGetMiss(const FCacheKey& Key);
-
-	bool ShouldSimulatePutMiss(const TCHAR* LegacyKey);
-	bool ShouldSimulateGetMiss(const TCHAR* LegacyKey);
 };
 
 class FDerivedDataBackend
diff --git a/Engine/Source/Developer/DerivedDataCache/Private/DerivedDataBackends.cpp b/Engine/Source/Developer/DerivedDataCache/Private/DerivedDataBackends.cpp
--- a/Engine/Source/Developer/DerivedDataCache/Private/DerivedDataBackends.cpp
+++ b/Engine/Source/Developer/DerivedDataCache/Private/DerivedDataBackends.cpp
@@ -1260,7 +1260,6 @@ struct Private::FBackendDebugMissState
 {
 	FCriticalSection Lock;
 	TMap<FCacheKey, EBackendDebugKeyState> Keys;
-	TMap<FName, EBackendDebugKeyState> LegacyKeys;
 };
 
 FBackendDebugOptions::FBackendDebugOptions()
@@ -1322,37 +1321,6 @@ bool FBackendDebugOptions::ParseFromTokens(FBackendDebugOptions& OutOptions, con
 	return true;
 }
 
-bool FBackendDebugOptions::ShouldSimulatePutMiss(const TCHAR* LegacyKey)
-{
-	if (!SimulateMissFilter.IsLegacyMatch(LegacyKey))
-	{
-		return false;
-	}
-
-	Private::FBackendDebugMissState& State = *SimulateMissState.Get();
-	const FName Key(LegacyKey);
-	const uint32 KeyHash = GetTypeHash(Key);
-
-	FScopeLock Lock(&State.Lock);
-	State.LegacyKeys.AddByHash(KeyHash, Key, EBackendDebugKeyState::HitGet);
-	return false;
-}
-
-bool FBackendDebugOptions::ShouldSimulateGetMiss(const TCHAR* LegacyKey)
-{
-	if (!SimulateMissFilter.IsLegacyMatch(LegacyKey))
-	{
-		return false;
-	}
-
-	Private::FBackendDebugMissState& State = *SimulateMissState.Get();
-	const FName Key(LegacyKey);
-	const uint32 KeyHash = GetTypeHash(Key);
-
-	FScopeLock Lock(&State.Lock);
-	return State.LegacyKeys.FindOrAddByHash(KeyHash, Key, EBackendDebugKeyState::MissGet) == EBackendDebugKeyState::MissGet;
-}
-
 bool FBackendDebugOptions::ShouldSimulatePutMiss(const FCacheKey& Key)
 {
 	if (!SimulateMissFilter.IsMatch(Key))
diff --git a/Engine/Source/Developer/DerivedDataCache/Private/DerivedDataCacheKeyFilter.cpp b/Engine/Source/Developer/DerivedDataCache/Private/DerivedDataCacheKeyFilter.cpp
--- a/Engine/Source/Developer/DerivedDataCache/Private/DerivedDataCacheKeyFilter.cpp
+++ b/Engine/Source/Developer/DerivedDataCache/Private/DerivedDataCacheKeyFilter.cpp
@@ -3,7 +3,6 @@
 #include "DerivedDataCacheKeyFilter.h"
 
 #include "Algo/Find.h"
-#include "Containers/StringConv.h"
 #include "DerivedDataCacheKey.h"
 #include "DerivedDataSharedString.h"
 #include "Hash/xxhash.h"
@@ -60,7 +59,7 @@ bool Private::FCacheKeyFilterState::TryParseTypeRate(const FStringView TypeConfi
 		return false;
 	}
 
-	OutTypeRate.Type = FUtf8StringView(FTCHARToUTF8(TypeView));
+	OutTypeRate.Type = WriteToUtf8String<64>(TypeView);
 	OutTypeRate.Rate = ConvertMatchRate(Rate);
 	return true;
 }
@@ -161,39 +160,4 @@ bool FCacheKeyFilter::IsMatch(const FCacheKey& Key) const
 	return State->ApplySalt(GetTypeHash(Key.Hash)) < TargetRate;
 }
 
-bool FCacheKeyFilter::IsLegacyMatch(const TCHAR* LegacyKey) const
-{
-	using namespace UE::DerivedData::Private;
-
-	if (!State)
-	{
-		return false;
-	}
-
-	const FStringView LegacyKeyView(LegacyKey);
-	uint32 TargetRate = State->DefaultRate;
-
-	if (!State->Types.IsEmpty())
-	{
-		TUtf8StringBuilder<256> Bucket;
-		Bucket << LegacyKeyView.Left(String::FindFirstChar(LegacyKeyView, TEXT('_')));
-
-		using FTypeRate = FCacheKeyFilterState::FTypeRate;
-		if (const FTypeRate* TypeRate = Algo::FindBy(State->Types, Bucket.ToView(), &FTypeRate::Type))
-		{
-			TargetRate = TypeRate->Rate;
-		}
-	}
-
-	if (TargetRate == 0)
-	{
-		return false;
-	}
-	if (TargetRate == MAX_uint32)
-	{
-		return true;
-	}
-	return State->ApplySalt(uint32(FXxHash64::HashBuffer(MakeMemoryView(LegacyKeyView)).Hash)) < TargetRate;
-}
-
 } // UE::DerivedData
diff --git a/Engine/Source/Developer/DerivedDataCache/Private/PakFileCacheStore.cpp b/Engine/Source/Developer/DerivedDataCache/Private/PakFileCacheStore.cpp
--- a/Engine/Source/Developer/DerivedDataCache/Private/PakFileCacheStore.cpp
+++ b/Engine/Source/Developer/DerivedDataCache/Private/PakFileCacheStore.cpp
@@ -259,7 +259,7 @@ void FPakFileCacheStore::Close()
 bool FPakFileCacheStore::CachedDataProbablyExists(const TCHAR* CacheKey)
 {
 	COOK_STAT(auto Timer = UsageStats.TimeProbablyExists());
-	if (bClosed || DebugOptions.ShouldSimulateGetMiss(CacheKey))
+	if (bClosed)
 	{
 		return false;
 	}
@@ -275,7 +275,7 @@ bool FPakFileCacheStore::CachedDataProbablyExists(const TCHAR* CacheKey)
 bool FPakFileCacheStore::GetCachedData(const TCHAR* CacheKey, TArray<uint8>& OutData)
 {
 	COOK_STAT(auto Timer = UsageStats.TimeGet());
-	if (bClosed || DebugOptions.ShouldSimulateGetMiss(CacheKey))
+	if (bClosed)
 	{
 		return false;
 	}
@@ -331,7 +331,7 @@ bool FPakFileCacheStore::GetCachedData(const TCHAR* CacheKey, TArray<uint8>& Out
 void FPakFileCacheStore::PutCachedData(const TCHAR* CacheKey, TArrayView<const uint8> InData, bool bPutEvenIfExists)
 {
 	COOK_STAT(auto Timer = UsageStats.TimePut());
-	if (bWriting && !bClosed && !DebugOptions.ShouldSimulateGetMiss(CacheKey))
+	if (bWriting && !bClosed)
 	{
 		FWriteScopeLock ScopeLock(SynchronizationObject);
 		FString Key(CacheKey);
