DDC: Refactored the HTTP API to add configurability and remove reliance on types from DDC

- FHttpSharedData is now IHttpConnectionPool. It is a connection pool of configurable size.
- FHttpRequestPool is now IHttpClient. It is a request pool of configurable size and the point at which most behaviors are configured.
- FHttpRequest is now split into IHttpRequest and IHttpResponse. This separation allows requests to be released back to the client for reuse while the response is still being processed.
- FScopedHttpPoolRequestPtr is now THttpUniquePtr. It is used uniformly with the pool, client, request, and response.
- The response body is now handled by IHttpReceiver, with a basic FHttpByteArrayReceiver provided to capture the whole body into an array.
- The response now includes stats about how much time was spent on various parts of the request, as well as how much data was sent and received and at what rate.
- The response supports cancellation at any point before it is complete, and IHttpResponseMonitor is provided to support thread-safe cancellation and waiting without requiring locks around the unique owner of the response.
- The response will be canceled automatically when destroyed through its THttpUniquePtr.
- The client does not directly support blocking on request creation. Blocking creation of requests is now provided by FHttpRequestQueue.
- The HTTP thread now relies on curl_multi_poll, available in newer libcurl versions, to allow immediate waking when a handle needs to be added or removed.

diff --git a/Engine/Source/Developer/DerivedDataCache/Private/Http/CurlHttpClient.cpp b/Engine/Source/Developer/DerivedDataCache/Private/Http/CurlHttpClient.cpp
new file mode 100644
--- /dev/null
+++ b/Engine/Source/Developer/DerivedDataCache/Private/Http/CurlHttpClient.cpp
@@ -0,0 +1,1274 @@
+// Copyright Epic Games, Inc. All Rights Reserved.
+
+#include "HttpClient.h"
+
+#include "Containers/DepletableMpscQueue.h"
+#include "Containers/LockFreeList.h"
+#include "Containers/StringView.h"
+#include "HAL/Event.h"
+#include "HAL/PlatformProcess.h"
+#include "HAL/Thread.h"
+#include "Memory/CompositeBuffer.h"
+#include "Memory/MemoryView.h"
+#include "Misc/App.h"
+#include "Misc/EngineVersion.h"
+#include "Misc/StringBuilder.h"
+#include "Serialization/CompactBinary.h"
+#include "String/Find.h"
+#include "Templates/RefCounting.h"
+#include "Templates/UnrealTemplate.h"
+
+#if WITH_SSL
+#include "Ssl.h"
+#include <openssl/ssl.h>
+#endif
+
+#ifndef CURL_NO_OLDIES
+#define CURL_NO_OLDIES
+#endif
+
+#if PLATFORM_MICROSOFT
+#include "Microsoft/AllowMicrosoftPlatformTypes.h"
+#endif
+#ifdef PLATFORM_CURL_INCLUDE
+#include PLATFORM_CURL_INCLUDE
+#else
+#include "curl/curl.h"
+#endif
+#if PLATFORM_MICROSOFT
+#include "Microsoft/HideMicrosoftPlatformTypes.h"
+#endif
+
+#include <atomic>
+
+namespace UE { class FCurlHttpClient; }
+namespace UE { class FCurlHttpRequest; }
+namespace UE { class FCurlHttpResponse; }
+
+namespace UE::CurlHttp::Private
+{
+
+template <typename T>
+static bool TryIncrement(std::atomic<T>& Value, const T Max)
+{
+	for (T Existing = Value.load(std::memory_order_relaxed);;)
+	{
+		if (Existing >= Max)
+		{
+			return false;
+		}
+		if (Value.compare_exchange_weak(Existing, Existing + 1, std::memory_order_relaxed))
+		{
+			return true;
+		}
+	}
+}
+
+template <typename T>
+static T AtomicFetchOr(std::atomic<T>& Atomic, const T Value, std::memory_order MemoryOrder = std::memory_order_seq_cst)
+{
+	for (T Existing = Atomic.load(MemoryOrder);;)
+	{
+		if (Atomic.compare_exchange_weak(Existing, Existing | Value, MemoryOrder))
+		{
+			return Existing;
+		}
+	}
+}
+
+template <typename T, typename PredicateType>
+static bool AtomicFetchOrIf(
+	std::atomic<T>& Atomic,
+	const T Value,
+	PredicateType Predicate,
+	std::memory_order MemoryOrder = std::memory_order_seq_cst)
+{
+	for (T Existing = Atomic.load(MemoryOrder);;)
+	{
+		if (!Predicate(Existing))
+		{
+			return false;
+		}
+		if (Atomic.compare_exchange_weak(Existing, Existing | Value, MemoryOrder))
+		{
+			return true;
+		}
+	}
+}
+
+} // UE::CurlHttp::Private
+
+///////////////////////////////////////////////////////////////////////////////////////////////////////////////////////
+
+namespace UE
+{
+
+DEFINE_LOG_CATEGORY_STATIC(LogHttp, Display, All);
+
+class FCurlHttpHeaders
+{
+public:
+	FCurlHttpHeaders() = default;
+	FCurlHttpHeaders(const FCurlHttpHeaders&) = delete;
+	FCurlHttpHeaders& operator=(const FCurlHttpHeaders&) = delete;
+
+	inline ~FCurlHttpHeaders()
+	{
+		curl_slist_free_all(List);
+	}
+
+	inline void Reset()
+	{
+		curl_slist_free_all(List);
+		List = nullptr;
+	}
+
+	inline void AddHeader(const ANSICHAR* Header)
+	{
+		List = curl_slist_append(List, Header);
+	}
+
+	inline curl_slist* GetList() const { return List; }
+
+private:
+	curl_slist* List = nullptr;
+};
+
+class FCurlHttpManager final : public IHttpManager
+{
+public:
+	FCurlHttpManager();
+
+	THttpUniquePtr<IHttpConnectionPool> CreateConnectionPool(const FHttpConnectionPoolParams& Params) final;
+
+	void SetDefaultOptions(CURL* Curl, FCurlHttpHeaders& Headers);
+
+private:
+	const FCbObjectId SessionId = FCbObjectId::NewObjectId();
+	std::atomic<uint32> RequestId = 1;
+	FUtf8StringBuilderBase UserAgent;
+};
+
+class FCurlHttpConnectionPool final : public IHttpConnectionPool
+{
+public:
+	FCurlHttpConnectionPool(FCurlHttpManager& Manager, const FHttpConnectionPoolParams& Params);
+
+	THttpUniquePtr<IHttpClient> CreateClient(const FHttpClientParams& Params) final;
+	void DeleteClient(FCurlHttpClient* Client);
+
+	void SetDefaultOptions(CURL* Curl, FCurlHttpHeaders& Headers) const;
+
+	bool BeginAsyncRequest(FCurlHttpResponse* Response);
+	void CancelAsyncRequest(FCurlHttpResponse* Response);
+
+private:
+	~FCurlHttpConnectionPool() final;
+	void Destroy() final { delete this; }
+
+	void ThreadLoop();
+	void CompleteRequest(CURL* Curl);
+
+	static void CurlLock(CURL* Curl, curl_lock_data Data, curl_lock_access Access, void* Param);
+	static void CurlUnlock(CURL* Curl, curl_lock_data Data, void* Param);
+
+	static void AssertShareCodeOk(CURLSHcode Code);
+	static void AssertMultiCodeOk(CURLMcode Code);
+
+	FCurlHttpManager& Manager;
+	CURLSH* CurlShare;
+	CURLM* CurlMulti;
+
+	FRWLock Locks[CURL_LOCK_DATA_LAST];
+	bool WriteLocked[CURL_LOCK_DATA_LAST]{};
+
+	std::atomic<uint32> ClientCount = 0;
+
+	enum class EThreadCommandType
+	{
+		Begin,
+		Cancel,
+	};
+
+	struct FThreadCommand
+	{
+		TRefCountPtr<FCurlHttpResponse> Response;
+		EThreadCommandType Type;
+	};
+
+	TDepletableMpscQueue<FThreadCommand> ThreadCommands;
+	FThread Thread;
+	std::atomic<bool> bThreadStarting;
+	std::atomic<bool> bThreadStopping;
+};
+
+class FCurlHttpClient final : public IHttpClient
+{
+public:
+	FCurlHttpClient(FCurlHttpConnectionPool& ConnectionPool, const FHttpClientParams& Params);
+
+	THttpUniquePtr<IHttpRequest> TryCreateRequest(const FHttpRequestParams& Params) final;
+	void DeleteRequest(FCurlHttpRequest* Request);
+
+	void SetDefaultOptions(CURL* Curl, FCurlHttpHeaders& Headers) const;
+
+	bool BeginAsyncRequest(FCurlHttpResponse* Response);
+	void CancelAsyncRequest(FCurlHttpResponse* Response);
+
+private:
+	~FCurlHttpClient() final;
+	void Destroy() final { delete this; }
+
+	static long ConvertTlsLevel(EHttpTlsLevel Level);
+
+	FCurlHttpConnectionPool& ConnectionPool;
+	TLockFreePointerListLIFO<FCurlHttpRequest> RequestPool;
+	std::atomic<uint32> RequestPoolCount = 0;
+	std::atomic<uint32> RequestCount = 0;
+	FHttpClientParams Params;
+};
+
+class FCurlHttpRequest final : public IHttpRequest
+{
+public:
+	explicit FCurlHttpRequest(FCurlHttpClient& Client);
+	~FCurlHttpRequest() final;
+
+	void Reset() final;
+	void SetUri(FAnsiStringView Uri) final;
+	void SetMethod(EHttpMethod Method) final;
+	void SetBody(const FCompositeBuffer& Body) final;
+	void AddHeader(FAnsiStringView Name, FAnsiStringView Value) final;
+
+	void Send(IHttpReceiver* Receiver, THttpUniquePtr<IHttpResponse>& OutResponse) final;
+	void SendAsync(IHttpReceiver* Receiver, THttpUniquePtr<IHttpResponse>& OutResponse) final;
+
+	void OnComplete(CURLcode Code);
+
+private:
+	void Destroy() final { Client.DeleteRequest(this); }
+
+	void SetDefaultOptions();
+	void CheckIdle(const TCHAR* FunctionName);
+
+	FCurlHttpResponse* CreateResponse(IHttpReceiver* Receiver, THttpUniquePtr<IHttpResponse>& OutResponse);
+
+	static size_t CurlRead(void* Ptr, size_t SizeInBlocks, size_t BlockSizeInBytes, void* Param);
+	static size_t CurlSeek(void* Param, curl_off_t Offset, int Origin);
+#if WITH_SSL
+	static CURLcode CurlSslContext(CURL* Curl, void* Context, void* Param);
+	static int SslCertVerify(int PreverifyOk, X509_STORE_CTX* Context);
+#endif
+
+	EHttpMethod Method;
+	FAnsiStringBuilderBase Uri;
+	FCurlHttpClient& Client;
+	FCurlHttpHeaders Headers;
+	FCompositeBuffer Body;
+	uint64 BodyOffset = 0;
+	CURL* Curl = nullptr;
+	std::atomic<FCurlHttpResponse*> Response = nullptr;
+};
+
+enum class ECurlHttpResponseState : uint8
+{
+	None = 0,
+	Complete = 1 << 0,
+	Canceled = 1 << 1,
+};
+
+ENUM_CLASS_FLAGS(ECurlHttpResponseState);
+
+class FCurlHttpResponse final : public IHttpResponse, public IHttpResponseMonitor
+{
+public:
+	FCurlHttpResponse(CURL* Curl, EHttpMethod Method, FAnsiStringView Uri, IHttpReceiver* Receiver);
+
+	bool Create();
+
+	CURL* GetCurl() { return Curl; }
+	void SetClient(FCurlHttpClient* InClient) { Client = InClient; }
+	void SetComplete(CURLcode Code);
+
+	TRefCountPtr<IHttpResponseMonitor> GetMonitor() final { return this; }
+
+	void Cancel() final;
+	void Wait() const final { CompleteEvent->Wait(); }
+	bool Poll() const final { return EnumHasAnyFlags(State.load(), ECurlHttpResponseState::Complete); }
+	bool IsCanceled() const { return EnumHasAnyFlags(State.load(std::memory_order_relaxed), ECurlHttpResponseState::Canceled); }
+
+	FAnsiStringView GetUri() const final { return Uri; }
+	EHttpMethod GetMethod() const final { return Method; }
+	int32 GetStatusCode() const final { return StatusCode; }
+	EHttpErrorCode GetErrorCode() const final { return ErrorCode; }
+	FAnsiStringView GetError() const final { return Error; }
+	const FHttpResponseStats& GetStats() const final;
+	TConstArrayView<FAnsiStringView> GetAllHeaders() const final { return HeaderViews; }
+
+	void AddRef() const final;
+	void Release() const final;
+
+private:
+	~FCurlHttpResponse() final;
+	void Destroy() final;
+
+	void ConditionallySetResponseStatus();
+	bool WriteHeader(FAnsiStringView Header);
+	bool WriteBody(FMemoryView Body);
+
+	static size_t CurlDebug(CURL* Curl, curl_infotype DebugInfoType, char* DebugInfo, size_t DebugInfoSize, void* Param);
+	static size_t CurlHeader(void* Ptr, size_t SizeInBlocks, size_t BlockSizeInBytes, void* Param);
+	static size_t CurlWrite(void* Ptr, size_t SizeInBlocks, size_t BlockSizeInBytes, void* Param);
+
+	static EHttpErrorCode ConvertErrorCode(CURLcode Code);
+
+	FAnsiStringBuilderBase Uri;
+	EHttpMethod Method;
+	bool bHasBody = false;
+	EHttpErrorCode ErrorCode = EHttpErrorCode::Unknown;
+	std::atomic<ECurlHttpResponseState> State = ECurlHttpResponseState::None;
+	std::atomic<int32> StatusCode = -1;
+	mutable std::atomic<uint32> ReferenceCount = 0;
+
+	CURL* Curl; // null when complete
+	IHttpReceiver* Receiver;
+	FCurlHttpClient* Client = nullptr;
+	TArray<FAnsiStringView, TInlineAllocator<32>> HeaderViews;
+	TArray<int32, TInlineAllocator<32>> HeaderLengths;
+	TAnsiStringBuilder<4096> Headers;
+	FAnsiStringView Error;
+	ANSICHAR ErrorBuffer[CURL_ERROR_SIZE]{};
+	FHttpResponseStats Stats;
+	FEventRef CompleteEvent{EEventMode::ManualReset};
+
+	enum class EHttpReceiverFunction : uint8
+	{
+		None = 0,
+		OnCreate,
+		OnHeaders,
+		OnBody,
+		OnComplete,
+	};
+	inline static thread_local EHttpReceiverFunction ReceiverFunction;
+};
+
+///////////////////////////////////////////////////////////////////////////////////////////////////////////////////////
+
+FCurlHttpManager::FCurlHttpManager()
+{
+	// User-Agent: UnrealEngine/X.Y.Z-<CL> (<Platform>; <Config> <TargetType>; <BranchName>) <AppName> (<ProjectName>)
+	const FEngineVersion& Version = FEngineVersion::Current();
+	UserAgent << ANSITEXTVIEW("User-Agent: UnrealEngine/")
+		<< Version.GetMajor() << '.' << Version.GetMinor() << '.' << Version.GetPatch() << '-' << Version.GetChangelist()
+		<< ANSITEXTVIEW(" (") << FPlatformProperties::PlatformName()
+		<< ANSITEXTVIEW("; ") << LexToString(FApp::GetBuildConfiguration()) << ' ' << LexToString(FApp::GetBuildTargetType())
+		<< ANSITEXTVIEW("; ") << FApp::GetBranchName()
+		<< ANSITEXTVIEW(") ") << FApp::GetName();
+	if (FApp::HasProjectName() && FApp::GetName() != FApp::GetProjectName())
+	{
+		UserAgent << ANSITEXTVIEW(" (") << FApp::GetProjectName() << ')';
+	}
+	UserAgent.ToString();
+}
+
+THttpUniquePtr<IHttpConnectionPool> FCurlHttpManager::CreateConnectionPool(const FHttpConnectionPoolParams& Params)
+{
+	return THttpUniquePtr<IHttpConnectionPool>(new FCurlHttpConnectionPool(*this, Params));
+}
+
+void FCurlHttpManager::SetDefaultOptions(CURL* Curl, FCurlHttpHeaders& Headers)
+{
+	curl_easy_setopt(Curl, CURLOPT_NOSIGNAL, 1L);
+	curl_easy_setopt(Curl, CURLOPT_USERAGENT, UserAgent.GetData());
+
+	Headers.AddHeader(*WriteToAnsiString<64>(ANSITEXTVIEW("UE-Session: "), SessionId));
+	Headers.AddHeader(*WriteToAnsiString<32>(ANSITEXTVIEW("UE-Request: "), RequestId.fetch_add(1, std::memory_order_relaxed)));
+
+	// Remove the Expect: 100-Continue header that curl adds by default because it adds latency.
+	Headers.AddHeader("Expect:");
+}
+
+///////////////////////////////////////////////////////////////////////////////////////////////////////////////////////
+
+FCurlHttpConnectionPool::FCurlHttpConnectionPool(FCurlHttpManager& InManager, const FHttpConnectionPoolParams& Params)
+	: Manager(InManager)
+	, CurlShare(curl_share_init())
+	, CurlMulti(curl_multi_init())
+	, bThreadStarting(false)
+	, bThreadStopping(!Params.bAllowAsync || !FPlatformProcess::SupportsMultithreading())
+{
+	curl_share_setopt(CurlShare, CURLSHOPT_USERDATA, this);
+	curl_share_setopt(CurlShare, CURLSHOPT_LOCKFUNC, CurlLock);
+	curl_share_setopt(CurlShare, CURLSHOPT_UNLOCKFUNC, CurlUnlock);
+	curl_share_setopt(CurlShare, CURLSHOPT_SHARE, CURL_LOCK_DATA_DNS);
+	curl_share_setopt(CurlShare, CURLSHOPT_SHARE, CURL_LOCK_DATA_SSL_SESSION);
+
+	checkf(!Params.MaxConnections || !Params.MinConnections || Params.MinConnections <= Params.MaxConnections,
+		TEXT("MinConnections (%u) exceeds MaxConnections (%u)"), Params.MinConnections, Params.MaxConnections);
+	constexpr uint32 DefaultConnectionCount = 8;
+	const uint32 MaxConnections = Params.MaxConnections ? Params.MaxConnections : FMath::Max(Params.MinConnections, DefaultConnectionCount);
+	const uint32 MinConnections = Params.MinConnections ? Params.MinConnections : FMath::Min(Params.MaxConnections, DefaultConnectionCount);
+	curl_multi_setopt(CurlMulti, CURLMOPT_MAXCONNECTS, MinConnections);
+	curl_multi_setopt(CurlMulti, CURLMOPT_MAX_TOTAL_CONNECTIONS, MaxConnections);
+	curl_multi_setopt(CurlMulti, CURLMOPT_PIPELINING, CURLPIPE_MULTIPLEX);
+}
+
+FCurlHttpConnectionPool::~FCurlHttpConnectionPool()
+{
+	checkf(ClientCount.load(std::memory_order_relaxed) == 0,
+		TEXT("Connection pool has %u clients that were not deleted."), ClientCount.load());
+
+	bThreadStopping.store(true, std::memory_order_relaxed);
+	AssertMultiCodeOk(curl_multi_wakeup(CurlMulti));
+	if (Thread.IsJoinable())
+	{
+		Thread.Join();
+	}
+
+	check(ThreadCommands.IsEmpty());
+
+	AssertMultiCodeOk(curl_multi_cleanup(CurlMulti));
+	AssertShareCodeOk(curl_share_cleanup(CurlShare));
+}
+
+THttpUniquePtr<IHttpClient> FCurlHttpConnectionPool::CreateClient(const FHttpClientParams& Params)
+{
+	ClientCount.fetch_add(1, std::memory_order_relaxed);
+	return THttpUniquePtr<IHttpClient>(new FCurlHttpClient(*this, Params));
+}
+
+void FCurlHttpConnectionPool::DeleteClient(FCurlHttpClient* Client)
+{
+	ClientCount.fetch_sub(1, std::memory_order_relaxed);
+}
+
+void FCurlHttpConnectionPool::SetDefaultOptions(CURL* Curl, FCurlHttpHeaders& Headers) const
+{
+	Manager.SetDefaultOptions(Curl, Headers);
+
+	curl_easy_setopt(Curl, CURLOPT_SHARE, CurlShare);
+}
+
+bool FCurlHttpConnectionPool::BeginAsyncRequest(FCurlHttpResponse* Response)
+{
+	if (bThreadStopping.load(std::memory_order_relaxed))
+	{
+		return false;
+	}
+	if (ThreadCommands.EnqueueAndReturnWasEmpty(FThreadCommand{Response, EThreadCommandType::Begin}))
+	{
+		AssertMultiCodeOk(curl_multi_wakeup(CurlMulti));
+	}
+	if (!bThreadStarting.load(std::memory_order_relaxed) && !bThreadStarting.exchange(true, std::memory_order_relaxed))
+	{
+		Thread = FThread(TEXT("HttpConnectionPool"), [this] { ThreadLoop(); }, 128 * 1024);
+	}
+	return true;
+}
+
+void FCurlHttpConnectionPool::CancelAsyncRequest(FCurlHttpResponse* Response)
+{
+	if (ThreadCommands.EnqueueAndReturnWasEmpty(FThreadCommand{Response, EThreadCommandType::Cancel}))
+	{
+		AssertMultiCodeOk(curl_multi_wakeup(CurlMulti));
+	}
+}
+
+void FCurlHttpConnectionPool::ThreadLoop()
+{
+	while (!ThreadCommands.IsEmpty() || !bThreadStopping.load(std::memory_order_relaxed))
+	{
+		ThreadCommands.Deplete([this](FThreadCommand Command)
+		{
+			if (CURL* Curl = Command.Response->GetCurl())
+			{
+				switch (Command.Type)
+				{
+				case EThreadCommandType::Begin:
+					AssertMultiCodeOk(curl_multi_add_handle(CurlMulti, Curl));
+					break;
+				case EThreadCommandType::Cancel:
+					CompleteRequest(Curl);
+					break;
+				}
+			}
+		});
+
+		int ActiveRequests = 0;
+		AssertMultiCodeOk(curl_multi_perform(CurlMulti, &ActiveRequests));
+
+		for (int MessagesInQueue = 0; CURLMsg* Message = curl_multi_info_read(CurlMulti, &MessagesInQueue);)
+		{
+			if (Message->msg == CURLMSG_DONE)
+			{
+				CompleteRequest(Message->easy_handle);
+			}
+		}
+
+		constexpr int WaitTimeMs = 60'000;
+		AssertMultiCodeOk(curl_multi_poll(CurlMulti, nullptr, 0, WaitTimeMs, nullptr));
+	}
+}
+
+void FCurlHttpConnectionPool::CompleteRequest(CURL* Curl)
+{
+	AssertMultiCodeOk(curl_multi_remove_handle(CurlMulti, Curl));
+	void* Request = nullptr;
+	curl_easy_getinfo(Curl, CURLINFO_PRIVATE, &Request);
+	((FCurlHttpRequest*)Request)->OnComplete(CURLE_OK);
+}
+
+void FCurlHttpConnectionPool::CurlLock(CURL* Curl, curl_lock_data Data, curl_lock_access Access, void* Param)
+{
+	FCurlHttpConnectionPool& ConnectionPool = *(FCurlHttpConnectionPool*)Param;
+	if (Access == CURL_LOCK_ACCESS_SHARED)
+	{
+		ConnectionPool.Locks[Data].ReadLock();
+	}
+	else
+	{
+		ConnectionPool.Locks[Data].WriteLock();
+		ConnectionPool.WriteLocked[Data] = true;
+	}
+}
+
+void FCurlHttpConnectionPool::CurlUnlock(CURL* Curl, curl_lock_data Data, void* Param)
+{
+	FCurlHttpConnectionPool& ConnectionPool = *(FCurlHttpConnectionPool*)Param;
+	if (!ConnectionPool.WriteLocked[Data])
+	{
+		ConnectionPool.Locks[Data].ReadUnlock();
+	}
+	else
+	{
+		ConnectionPool.WriteLocked[Data] = false;
+		ConnectionPool.Locks[Data].WriteUnlock();
+	}
+}
+
+inline void FCurlHttpConnectionPool::AssertShareCodeOk(const CURLSHcode Code)
+{
+	checkf(Code == CURLSHE_OK, TEXT("Error in curl_share operation: %hs"), curl_share_strerror(Code));
+}
+
+inline void FCurlHttpConnectionPool::AssertMultiCodeOk(const CURLMcode Code)
+{
+	checkf(Code == CURLM_OK, TEXT("Error in curl_multi operation: %hs"), curl_multi_strerror(Code));
+}
+
+///////////////////////////////////////////////////////////////////////////////////////////////////////////////////////
+
+FCurlHttpClient::FCurlHttpClient(FCurlHttpConnectionPool& InConnectionPool, const FHttpClientParams& InParams)
+	: ConnectionPool(InConnectionPool)
+	, Params(InParams)
+{
+	checkf(!Params.MaxRequests || !Params.MinRequests || Params.MinRequests <= Params.MaxRequests,
+		TEXT("MinRequests (%u) exceeds MaxRequests (%u)"), Params.MinRequests, Params.MaxRequests);
+	Params.MaxRequests = Params.MaxRequests ? Params.MaxRequests : FMath::Max(Params.MinRequests, 256u);
+	Params.MinRequests = Params.MinRequests ? Params.MinRequests : FMath::Min(Params.MaxRequests, 16u);
+}
+
+FCurlHttpClient::~FCurlHttpClient()
+{
+	while (FCurlHttpRequest* Request = RequestPool.Pop())
+	{
+		delete Request;
+		RequestCount.fetch_sub(1, std::memory_order_relaxed);
+		RequestPoolCount.fetch_sub(1, std::memory_order_relaxed);
+	}
+
+	checkf(RequestCount.load(std::memory_order_relaxed) == 0,
+		TEXT("Client has %u requests that were not destroyed by IHttpRequest::Destroy()."), RequestCount.load());
+
+	ConnectionPool.DeleteClient(this);
+}
+
+THttpUniquePtr<IHttpRequest> FCurlHttpClient::TryCreateRequest(const FHttpRequestParams& RequestParams)
+{
+	if (IHttpRequest* Request = RequestPool.Pop())
+	{
+		RequestPoolCount.fetch_sub(1, std::memory_order_relaxed);
+		return THttpUniquePtr<IHttpRequest>(Request);
+	}
+
+	if (RequestParams.bIgnoreMaxRequests)
+	{
+		RequestCount.fetch_add(1, std::memory_order_relaxed);
+	}
+	else if (!CurlHttp::Private::TryIncrement(RequestCount, Params.MaxRequests))
+	{
+		return nullptr;
+	}
+
+	return THttpUniquePtr<IHttpRequest>(new FCurlHttpRequest(*this));
+}
+
+void FCurlHttpClient::DeleteRequest(FCurlHttpRequest* Request)
+{
+	if (CurlHttp::Private::TryIncrement(RequestPoolCount, Params.MinRequests))
+	{
+		Request->Reset();
+		RequestPool.Push(Request);
+	}
+	else
+	{
+		delete Request;
+		RequestCount.fetch_sub(1, std::memory_order_relaxed);
+	}
+
+	if (Params.OnDestroyRequest)
+	{
+		Params.OnDestroyRequest();
+	}
+}
+
+void FCurlHttpClient::SetDefaultOptions(CURL* Curl, FCurlHttpHeaders& Headers) const
+{
+	ConnectionPool.SetDefaultOptions(Curl, Headers);
+
+	curl_easy_setopt(Curl, CURLOPT_HTTP_VERSION, CURL_HTTP_VERSION_2_0);
+
+	curl_easy_setopt(Curl, CURLOPT_DNS_CACHE_TIMEOUT, FMath::Min<long>(Params.DnsCacheTimeout, MAX_int32));
+	curl_easy_setopt(Curl, CURLOPT_CONNECTTIMEOUT_MS, long(Params.ConnectTimeout));
+	curl_easy_setopt(Curl, CURLOPT_LOW_SPEED_LIMIT, long(Params.LowSpeedLimit));
+	curl_easy_setopt(Curl, CURLOPT_LOW_SPEED_TIME, long(Params.LowSpeedTime));
+	curl_easy_setopt(Curl, CURLOPT_USE_SSL, ConvertTlsLevel(Params.TlsLevel));
+	curl_easy_setopt(Curl, CURLOPT_FOLLOWLOCATION, long(Params.bFollowRedirects));
+	curl_easy_setopt(Curl, CURLOPT_POSTREDIR, long(
+		(Params.bFollow301Post ? CURL_REDIR_POST_301 : 0) |
+		(Params.bFollow302Post ? CURL_REDIR_POST_302 : 0) |
+		(Params.bFollow303Post ? CURL_REDIR_POST_303 : 0)));
+	curl_easy_setopt(Curl, CURLOPT_VERBOSE, long(Params.bVerbose));
+}
+
+long FCurlHttpClient::ConvertTlsLevel(const EHttpTlsLevel Level)
+{
+	switch (Level)
+	{
+	case EHttpTlsLevel::None: return CURLUSESSL_NONE;
+	case EHttpTlsLevel::Try:  return CURLUSESSL_TRY;
+	case EHttpTlsLevel::All:  return CURLUSESSL_ALL;
+	default: checkNoEntry();  return CURLUSESSL_NONE;
+	}
+}
+
+inline bool FCurlHttpClient::BeginAsyncRequest(FCurlHttpResponse* Response)
+{
+	Response->SetClient(this);
+	return ConnectionPool.BeginAsyncRequest(Response);
+}
+
+inline void FCurlHttpClient::CancelAsyncRequest(FCurlHttpResponse* Response)
+{
+	ConnectionPool.CancelAsyncRequest(Response);
+}
+
+///////////////////////////////////////////////////////////////////////////////////////////////////////////////////////
+
+FCurlHttpRequest::FCurlHttpRequest(FCurlHttpClient& InClient)
+	: Client(InClient)
+	, Curl(curl_easy_init())
+{
+	SetDefaultOptions();
+}
+
+FCurlHttpRequest::~FCurlHttpRequest()
+{
+	curl_easy_cleanup(Curl);
+}
+
+void FCurlHttpRequest::Reset()
+{
+	CheckIdle(TEXT("Reset"));
+	curl_easy_reset(Curl);
+	Headers.Reset();
+	Uri.Reset();
+	Body.Reset();
+	SetDefaultOptions();
+}
+
+void FCurlHttpRequest::SetDefaultOptions()
+{
+	Method = EHttpMethod::Get;
+
+	Client.SetDefaultOptions(Curl, Headers);
+
+	curl_easy_setopt(Curl, CURLOPT_PRIVATE, this);
+	curl_easy_setopt(Curl, CURLOPT_READDATA, this);
+	curl_easy_setopt(Curl, CURLOPT_READFUNCTION, CurlRead);
+	curl_easy_setopt(Curl, CURLOPT_SEEKDATA, this);
+	curl_easy_setopt(Curl, CURLOPT_SEEKFUNCTION, CurlSeek);
+#if WITH_SSL
+	curl_easy_setopt(Curl, CURLOPT_SSL_CTX_DATA, this);
+	curl_easy_setopt(Curl, CURLOPT_SSL_CTX_FUNCTION, CurlSslContext);
+#endif
+}
+
+void FCurlHttpRequest::CheckIdle(const TCHAR* FunctionName)
+{
+	checkf(!Response.load(std::memory_order_relaxed),
+		TEXT("%.*hs %.*hs: %s requires the request to be idle."),
+		LexToString(Method).Len(), LexToString(Method).GetData(), Uri.Len(), Uri.GetData(), FunctionName);
+}
+
+void FCurlHttpRequest::SetUri(const FAnsiStringView InUri)
+{
+	CheckIdle(TEXT("SetUri"));
+	Uri.Reset();
+	Uri.Append(InUri);
+	curl_easy_setopt(Curl, CURLOPT_URL, *Uri);
+}
+
+void FCurlHttpRequest::SetMethod(const EHttpMethod InMethod)
+{
+	CheckIdle(TEXT("SetMethod"));
+	Method = InMethod;
+	switch (Method)
+	{
+	case EHttpMethod::Get:
+		curl_easy_setopt(Curl, CURLOPT_HTTPGET, 1L);
+		break;
+	case EHttpMethod::Put:
+		curl_easy_setopt(Curl, CURLOPT_UPLOAD, 1L);
+		break;
+	case EHttpMethod::Post:
+		curl_easy_setopt(Curl, CURLOPT_POST, 1L);
+		break;
+	case EHttpMethod::Head:
+		curl_easy_setopt(Curl, CURLOPT_NOBODY, 1L);
+		break;
+	case EHttpMethod::Delete:
+		curl_easy_setopt(Curl, CURLOPT_CUSTOMREQUEST, "DELETE");
+		break;
+	default:
+		checkNoEntry();
+		break;
+	}
+}
+
+void FCurlHttpRequest::SetBody(const FCompositeBuffer& InBody)
+{
+	CheckIdle(TEXT("SetBody"));
+	Body = InBody;
+}
+
+void FCurlHttpRequest::AddHeader(const FAnsiStringView Name, const FAnsiStringView Value)
+{
+	CheckIdle(TEXT("AddHeader"));
+	Headers.AddHeader(*WriteToAnsiString<256>(Name, ANSITEXTVIEW(": "), Value));
+}
+
+void FCurlHttpRequest::Send(IHttpReceiver* Receiver, THttpUniquePtr<IHttpResponse>& OutResponse)
+{
+	CheckIdle(TEXT("Send"));
+	if (CreateResponse(Receiver, OutResponse))
+	{
+		OnComplete(curl_easy_perform(Curl));
+	}
+}
+
+void FCurlHttpRequest::SendAsync(IHttpReceiver* Receiver, THttpUniquePtr<IHttpResponse>& OutResponse)
+{
+	CheckIdle(TEXT("SendAsync"));
+	if (FCurlHttpResponse* LocalResponse = CreateResponse(Receiver, OutResponse))
+	{
+		if (!Client.BeginAsyncRequest(LocalResponse))
+		{
+			OnComplete(curl_easy_perform(Curl));
+		}
+	}
+}
+
+FCurlHttpResponse* FCurlHttpRequest::CreateResponse(IHttpReceiver* Receiver, THttpUniquePtr<IHttpResponse>& OutResponse)
+{
+	BodyOffset = 0;
+
+	if (const uint64 BodySize = Body.GetSize(); BodySize > 0)
+	{
+		switch (Method)
+		{
+		case EHttpMethod::Put:
+			curl_easy_setopt(Curl, CURLOPT_INFILESIZE_LARGE, curl_off_t(BodySize));
+			break;
+		case EHttpMethod::Post:
+			curl_easy_setopt(Curl, CURLOPT_POSTFIELDSIZE_LARGE, curl_off_t(BodySize));
+			break;
+		default:
+			checkf(false, TEXT("%s"), *WriteToString<256>(TEXTVIEW("Method "), LexToString(Method),
+				TEXTVIEW(" must not include a body but a body of "), BodySize, TEXTVIEW(" bytes was provided.")));
+			break;
+		}
+	}
+
+	curl_easy_setopt(Curl, CURLOPT_HTTPHEADER, Headers.GetList());
+
+	FCurlHttpResponse* LocalResponse = new FCurlHttpResponse(Curl, Method, Uri, Receiver);
+	OutResponse = THttpUniquePtr<IHttpResponse>(LocalResponse);
+
+	if (LocalResponse->Create())
+	{
+		verify(!Response.exchange(LocalResponse, std::memory_order_release));
+		return LocalResponse;
+	}
+	LocalResponse->SetComplete(CURLE_OK);
+	return nullptr;
+}
+
+void FCurlHttpRequest::OnComplete(CURLcode Code)
+{
+	if (FCurlHttpResponse* LocalResponse = Response.exchange(nullptr, std::memory_order_acquire))
+	{
+		LocalResponse->SetComplete(Code);
+	}
+}
+
+size_t FCurlHttpRequest::CurlRead(void* Ptr, size_t SizeInBlocks, size_t BlockSizeInBytes, void* Param)
+{
+	FCurlHttpRequest* Request = (FCurlHttpRequest*)Param;
+	const size_t TotalSize = Request->Body.GetSize();
+	const size_t Offset = Request->BodyOffset;
+	const size_t Size = FMath::Min(TotalSize - Offset, SizeInBlocks * BlockSizeInBytes);
+	Request->Body.CopyTo(MakeMemoryView(Ptr, Size), Offset);
+	Request->BodyOffset += Size;
+	return Size;
+}
+
+size_t FCurlHttpRequest::CurlSeek(void* Param, curl_off_t Offset, int Origin)
+{
+	FCurlHttpRequest* Request = (FCurlHttpRequest*)Param;
+
+	const size_t TotalSize = Request->Body.GetSize();
+	size_t AbsoluteOffset = 0;
+
+	switch (Origin)
+	{
+	case SEEK_SET: AbsoluteOffset = Offset; break;
+	case SEEK_CUR: AbsoluteOffset = Request->BodyOffset + Offset; break;
+	case SEEK_END: AbsoluteOffset = TotalSize + Offset; break;
+	}
+
+	if (AbsoluteOffset >= TotalSize)
+	{
+		return CURL_SEEKFUNC_FAIL;
+	}
+
+	Request->BodyOffset = AbsoluteOffset;
+	return CURL_SEEKFUNC_OK;
+}
+
+#if WITH_SSL
+CURLcode FCurlHttpRequest::CurlSslContext(CURL* Curl, void* Context, void* Param)
+{
+	SSL_CTX* SslContext = (SSL_CTX*)Context;
+	FSslModule::Get().GetCertificateManager().AddCertificatesToSslContext(SslContext);
+	SSL_CTX_set_verify(SslContext, SSL_CTX_get_verify_mode(SslContext), SslCertVerify);
+	SSL_CTX_set_app_data(SslContext, Param);
+	return CURLE_OK;
+}
+
+int FCurlHttpRequest::SslCertVerify(int PreverifyOk, X509_STORE_CTX* Context)
+{
+	if (PreverifyOk == 1)
+	{
+		SSL* Ssl = (SSL*)X509_STORE_CTX_get_ex_data(Context, SSL_get_ex_data_X509_STORE_CTX_idx());
+		check(Ssl);
+
+		SSL_CTX* SslContext = SSL_get_SSL_CTX(Ssl);
+		check(SslContext);
+
+		FCurlHttpRequest* Request = (FCurlHttpRequest*)SSL_CTX_get_app_data(SslContext);
+		check(Request);
+
+		// Extract the domain from the URI.
+		FAnsiStringView Domain = Request->Uri;
+		if (const int32 SchemeIndex = String::FindFirst(Domain, ANSITEXTVIEW("://")); SchemeIndex != INDEX_NONE)
+		{
+			Domain.RightChopInline(SchemeIndex + ANSITEXTVIEW("://").Len());
+		}
+		if (const int32 SlashIndex = String::FindFirstChar(Domain, '/'); SlashIndex != INDEX_NONE)
+		{
+			Domain.LeftInline(SlashIndex);
+		}
+		if (const int32 AtIndex = String::FindFirstChar(Domain, '@'); AtIndex != INDEX_NONE)
+		{
+			Domain.RightChopInline(AtIndex + 1);
+		}
+		const auto RemovePort = [](FAnsiStringView& Authority)
+		{
+			if (const int32 ColonIndex = String::FindLastChar(Authority, ':'); ColonIndex != INDEX_NONE)
+			{
+				Authority.LeftInline(ColonIndex);
+			}
+		};
+		if (Domain.StartsWith('['))
+		{
+			if (const int32 LastBracketIndex = String::FindLastChar(Domain, ']'); LastBracketIndex != INDEX_NONE)
+			{
+				Domain.MidInline(1, LastBracketIndex - 1);
+			}
+			else
+			{
+				RemovePort(Domain);
+			}
+		}
+		else
+		{
+			RemovePort(Domain);
+		}
+
+		if (!FSslModule::Get().GetCertificateManager().VerifySslCertificates(Context, FString(Domain)))
+		{
+			PreverifyOk = 0;
+		}
+	}
+
+	return PreverifyOk;
+}
+#endif
+
+///////////////////////////////////////////////////////////////////////////////////////////////////////////////////////
+
+FCurlHttpResponse::FCurlHttpResponse(CURL* InCurl, EHttpMethod InMethod, FAnsiStringView InUri, IHttpReceiver* InReceiver)
+	: Method(InMethod)
+	, Curl(InCurl)
+	, Receiver(InReceiver)
+{
+	// Release() is called by Destroy().
+	AddRef();
+
+	Uri.Append(InUri);
+	checkf(Receiver, TEXT("Receiver must not be null for %s"), *WriteToString<256>(*this));
+
+	curl_easy_setopt(Curl, CURLOPT_ERRORBUFFER, ErrorBuffer);
+	curl_easy_setopt(Curl, CURLOPT_DEBUGDATA, this);
+	curl_easy_setopt(Curl, CURLOPT_DEBUGFUNCTION, CurlDebug);
+	curl_easy_setopt(Curl, CURLOPT_HEADERDATA, this);
+	curl_easy_setopt(Curl, CURLOPT_HEADERFUNCTION, CurlHeader);
+	curl_easy_setopt(Curl, CURLOPT_WRITEDATA, this);
+	curl_easy_setopt(Curl, CURLOPT_WRITEFUNCTION, CurlWrite);
+}
+
+bool FCurlHttpResponse::Create()
+{
+	for (TGuardValue GuardReceiverFunction(ReceiverFunction, EHttpReceiverFunction::OnCreate);
+		 IHttpReceiver* NewReceiver = Receiver->OnCreate(*this);
+		 Receiver = NewReceiver)
+	{
+		if (Receiver == NewReceiver)
+		{
+			break;
+		}
+		checkf(NewReceiver, TEXT("Receiver must not be null for %s"), *WriteToString<256>(*this));
+	}
+	return !IsCanceled();
+}
+
+FCurlHttpResponse::~FCurlHttpResponse()
+{
+	checkf(!Curl, TEXT("Response must be completed before it is destroyed for %s"), *WriteToString<256>(*this));
+}
+
+void FCurlHttpResponse::Destroy()
+{
+	if (Curl)
+	{
+		Cancel();
+	}
+
+	// AddRef() is called by the constructor.
+	Release();
+}
+
+void FCurlHttpResponse::ConditionallySetResponseStatus()
+{
+	if (StatusCode < 0)
+	{
+		long ResponseCode;
+		curl_easy_getinfo(Curl, CURLINFO_RESPONSE_CODE, &ResponseCode);
+		if (ResponseCode > 0)
+		{
+			StatusCode = int32(ResponseCode);
+		}
+	}
+}
+
+bool FCurlHttpResponse::WriteHeader(FAnsiStringView Header)
+{
+	if (bHasBody)
+	{
+		return !EnumHasAnyFlags(State.load(std::memory_order_relaxed), ECurlHttpResponseState::Canceled);
+	}
+
+	// Reset headers between responses and keep only the last.
+	if (!HeaderViews.IsEmpty())
+	{
+		HeaderLengths.Reset();
+		HeaderViews.Reset();
+		Headers.Reset();
+		StatusCode = -1;
+	}
+
+	ConditionallySetResponseStatus();
+
+	if (Header == ANSITEXTVIEW("\r\n"))
+	{
+		const ANSICHAR* Data = Headers.GetData();
+		for (const int32 Len : HeaderLengths)
+		{
+			HeaderViews.Add(MakeStringView(Data, Len - 2));
+			Data += Len;
+		}
+		for (TGuardValue GuardReceiverFunction(ReceiverFunction, EHttpReceiverFunction::OnHeaders);
+			 IHttpReceiver* NewReceiver = Receiver->OnHeaders(*this);
+			 Receiver = NewReceiver)
+		{
+			if (Receiver == NewReceiver)
+			{
+				break;
+			}
+			checkf(NewReceiver, TEXT("Receiver must not be null for %s"), *WriteToString<256>(*this));
+		}
+	}
+	else
+	{
+		HeaderLengths.Add(Header.Len());
+		Headers.Append(Header);
+	}
+
+	return !IsCanceled();
+}
+
+bool FCurlHttpResponse::WriteBody(FMemoryView Body)
+{
+	ConditionallySetResponseStatus();
+
+	bHasBody = true;
+
+	while (!Body.IsEmpty() && !IsCanceled())
+	{
+		const FMemoryView InitialBody = Body;
+		IHttpReceiver* NewReceiver;
+		{
+			TGuardValue GuardReceiverFunction(ReceiverFunction, EHttpReceiverFunction::OnBody);
+			NewReceiver = Receiver->OnBody(*this, Body);
+		}
+		if (Receiver == NewReceiver)
+		{
+			checkf(Body.IsEmpty() || Body == InitialBody,
+				TEXT("Receiver must consume the entire body for %s"), *WriteToString<256>(*this));
+			break;
+		}
+		checkf(Body.IsEmpty() || Body.GetDataEnd() == InitialBody.GetDataEnd(),
+			TEXT("Receiver must consume the body sequentially for %s"), *WriteToString<256>(*this));
+		checkf(NewReceiver, TEXT("Receiver must not be null for %s"), *WriteToString<256>(*this));
+		Receiver = NewReceiver;
+	}
+
+	return !IsCanceled();
+}
+
+void FCurlHttpResponse::SetComplete(CURLcode Code)
+{
+	Error = ErrorBuffer;
+	if (IsCanceled())
+	{
+		ErrorCode = EHttpErrorCode::Canceled;
+		Error = ANSITEXTVIEW("Canceled using IHttpResponse::Cancel()");
+		StatusCode = 0;
+	}
+	else if (Code == CURLE_OK)
+	{
+		ErrorCode = EHttpErrorCode::None;
+		ConditionallySetResponseStatus();
+	}
+	else
+	{
+		ErrorCode = ConvertErrorCode(Code);
+		if (Error.IsEmpty())
+		{
+			Error = curl_easy_strerror(Code);
+		}
+		StatusCode = 0;
+	}
+
+	const auto GetSizeInfo = [this](CURLINFO Info)
+	{
+		curl_off_t Value;
+		curl_easy_getinfo(Curl, Info, &Value);
+		return uint64(Value);
+	};
+	const auto GetDoubleInfo = [this](CURLINFO Info)
+	{
+		double Value;
+		curl_easy_getinfo(Curl, Info, &Value);
+		return Value;
+	};
+	Stats.SendSize = GetSizeInfo(CURLINFO_SIZE_UPLOAD_T);
+	Stats.RecvSize = GetSizeInfo(CURLINFO_SIZE_DOWNLOAD_T);
+	Stats.SendRate = GetSizeInfo(CURLINFO_SPEED_UPLOAD_T);
+	Stats.RecvRate = GetSizeInfo(CURLINFO_SPEED_DOWNLOAD_T);
+	Stats.NameResolveTime = GetDoubleInfo(CURLINFO_NAMELOOKUP_TIME);
+	Stats.ConnectTime = GetDoubleInfo(CURLINFO_CONNECT_TIME);
+	Stats.TlsConnectTime = GetDoubleInfo(CURLINFO_APPCONNECT_TIME);
+	Stats.StartTransferTime = GetDoubleInfo(CURLINFO_STARTTRANSFER_TIME);
+	Stats.TotalTime = GetDoubleInfo(CURLINFO_TOTAL_TIME);
+
+	Curl = nullptr;
+	Client = nullptr;
+
+	// Hold a reference to safely access Receiver, State, and CompleteEvent.
+	TRefCountPtr<FCurlHttpResponse> Self(this);
+
+	for (TGuardValue GuardReceiverFunction(ReceiverFunction, EHttpReceiverFunction::OnComplete);
+		 IHttpReceiver* NewReceiver = Receiver->OnComplete(*this);
+		 Receiver = NewReceiver)
+	{
+		if (Receiver == NewReceiver)
+		{
+			break;
+		}
+	}
+
+	CurlHttp::Private::AtomicFetchOr(State, ECurlHttpResponseState::Complete);
+	CompleteEvent->Trigger();
+
+	if (UE_LOG_ACTIVE(LogHttp, Verbose))
+	{
+		TStringBuilder<80> StatsText;
+		if (Stats.SendSize)
+		{
+			StatsText << TEXTVIEW("sent ") << Stats.SendSize << TEXTVIEW(" bytes, ");
+		}
+		if (Stats.RecvSize)
+		{
+			StatsText << TEXTVIEW("received ") << Stats.RecvSize << TEXTVIEW(" bytes, ");
+		}
+		StatsText.Appendf(TEXT("%.3f seconds"), Stats.TotalTime);
+		UE_LOG(LogHttp, Verbose, TEXT("%s (%s)"), *WriteToString<256>(*this), *StatsText);
+	}
+
+	// DO NOT ACCESS THIS AGAIN PAST THIS POINT!
+}
+
+void FCurlHttpResponse::Cancel()
+{
+	const EHttpReceiverFunction LocalReceiverFunction = ReceiverFunction;
+	checkf(LocalReceiverFunction != EHttpReceiverFunction::OnComplete,
+		TEXT("Cancel() must not be called from within OnComplete for %s"), *WriteToString<256>(*this));
+
+	// Cancel only once and only if not completed.
+	const auto NotCompleteOrCanceled = [](ECurlHttpResponseState CheckState) -> bool
+	{
+		return !EnumHasAnyFlags(CheckState, ECurlHttpResponseState::Complete | ECurlHttpResponseState::Canceled);
+	};
+	if (CurlHttp::Private::AtomicFetchOrIf(State, ECurlHttpResponseState::Canceled, NotCompleteOrCanceled))
+	{
+		// Queue cancellation of an async request when called from outside of a receiver function.
+		// Cancellation within a receiver is handled by WriterHeader or WriteBody returning false,
+		// which causes an error in curl that causes the request to complete immediately.
+		if (Client && LocalReceiverFunction == EHttpReceiverFunction::None)
+		{
+			Client->CancelAsyncRequest(this);
+		}
+	}
+
+	// Wait for completion when called from outside of a receiver function.
+	// Waiting within a receiver function would prevent the response from completing.
+	if (LocalReceiverFunction == EHttpReceiverFunction::None)
+	{
+		CompleteEvent->Wait();
+	}
+}
+
+const FHttpResponseStats& FCurlHttpResponse::GetStats() const
+{
+	checkf(!Curl, TEXT("Stats accessed before they are available for %s"), *WriteToString<256>(*this));
+	return Stats;
+}
+
+size_t FCurlHttpResponse::CurlDebug(CURL* Curl, curl_infotype DebugInfoType, char* DebugInfo, size_t DebugInfoSize, void* Param)
+{
+	if (DebugInfoType != CURLINFO_TEXT && DebugInfoType != CURLINFO_HEADER_IN)
+	{
+		return 0;
+	}
+
+	FAnsiStringView DebugText(DebugInfo, DebugInfoSize);
+	DebugText.TrimStartAndEndInline();
+	if (DebugText.IsEmpty())
+	{
+		return 0;
+	}
+
+	FCurlHttpResponse* Response = (FCurlHttpResponse*)Param;
+	const FAnsiStringView Method = LexToString(Response->GetMethod());
+	const FAnsiStringView Uri = Response->GetUri();
+	switch (DebugInfoType)
+	{
+	case CURLINFO_TEXT:
+		UE_LOG(LogHttp, VeryVerbose, TEXT("%.*hs %.*hs (%p): [TEXT] %.*hs"),
+			Method.Len(), Method.GetData(), Uri.Len(), Uri.GetData(), Response, DebugText.Len(), DebugText.GetData());
+		break;
+	case CURLINFO_HEADER_IN:
+		UE_LOG(LogHttp, VeryVerbose, TEXT("%.*hs %.*hs (%p): [RESPONSE] %.*hs"),
+			Method.Len(), Method.GetData(), Uri.Len(), Uri.GetData(), Response, DebugText.Len(), DebugText.GetData());
+		break;
+	}
+	return 0;
+}
+
+size_t FCurlHttpResponse::CurlHeader(void* Ptr, size_t SizeInBlocks, size_t BlockSizeInBytes, void* Param)
+{
+	const size_t HeaderSize = SizeInBlocks * BlockSizeInBytes;
+	return ((FCurlHttpResponse*)Param)->WriteHeader(MakeStringView((const ANSICHAR*)Ptr, HeaderSize)) ? HeaderSize : 0;
+}
+
+size_t FCurlHttpResponse::CurlWrite(void* Ptr, size_t SizeInBlocks, size_t BlockSizeInBytes, void* Param)
+{
+	const size_t Size = SizeInBlocks * BlockSizeInBytes;
+	return ((FCurlHttpResponse*)Param)->WriteBody(MakeMemoryView(Ptr, Size)) ? Size : 0;
+}
+
+EHttpErrorCode FCurlHttpResponse::ConvertErrorCode(CURLcode Code)
+{
+	switch (Code)
+	{
+	case CURLE_OK:
+		return EHttpErrorCode::None;
+	case CURLE_COULDNT_RESOLVE_HOST:
+		return EHttpErrorCode::ResolveHost;
+	case CURLE_COULDNT_CONNECT:
+		return EHttpErrorCode::Connect;
+	case CURLE_SSL_CONNECT_ERROR:
+		return EHttpErrorCode::TlsConnect;
+	case CURLE_PEER_FAILED_VERIFICATION:
+		return EHttpErrorCode::TlsPeerVerification;
+	case CURLE_OPERATION_TIMEDOUT:
+		return EHttpErrorCode::TimedOut;
+	default:
+		return EHttpErrorCode::Unknown;
+	}
+}
+
+void FCurlHttpResponse::AddRef() const
+{
+	ReferenceCount.fetch_add(1, std::memory_order_relaxed);
+}
+
+void FCurlHttpResponse::Release() const
+{
+	if (ReferenceCount.fetch_sub(1, std::memory_order_acq_rel) == 1)
+	{
+		delete this;
+	}
+}
+
+///////////////////////////////////////////////////////////////////////////////////////////////////////////////////////
+
+IHttpManager& IHttpManager::Get()
+{
+	static FCurlHttpManager CurlHttpManager;
+	return CurlHttpManager;
+}
+
+} // UE
diff --git a/Engine/Source/Developer/DerivedDataCache/Private/Http/HttpClient.cpp b/Engine/Source/Developer/DerivedDataCache/Private/Http/HttpClient.cpp
--- a/Engine/Source/Developer/DerivedDataCache/Private/Http/HttpClient.cpp
+++ b/Engine/Source/Developer/DerivedDataCache/Private/Http/HttpClient.cpp
@@ -2,377 +2,357 @@
 
 #include "HttpClient.h"
 
-#include "Algo/RemoveIf.h"
+#include "Containers/LockFreeList.h"
+#include "HAL/Event.h"
+#include "Memory/MemoryView.h"
 #include "Misc/AsciiSet.h"
-#include "Misc/ScopeRWLock.h"
-#include "Tasks/Task.h"
-
-#if !defined(CURL_NO_OLDIES)
-	#define CURL_NO_OLDIES
-#endif
-
-#if PLATFORM_MICROSOFT
-#include "Microsoft/AllowMicrosoftPlatformTypes.h"
-#endif
-
-#if defined(PLATFORM_CURL_INCLUDE)
-#include PLATFORM_CURL_INCLUDE
-#else
-#include "curl/curl.h"
-#endif // defined(PLATFORM_CURL_INCLUDE)
-
-#if PLATFORM_MICROSOFT
-#include "Microsoft/HideMicrosoftPlatformTypes.h"
-#endif
+#include "ProfilingDebugging/CpuProfilerTrace.h"
+#include "String/Find.h"
 
 namespace UE
 {
 
-namespace Http::Private
+FAnsiStringView LexToString(const EHttpMethod Method)
 {
+	switch (Method)
+	{
+	case EHttpMethod::Get:    return ANSITEXTVIEW("GET");
+	case EHttpMethod::Put:    return ANSITEXTVIEW("PUT");
+	case EHttpMethod::Post:   return ANSITEXTVIEW("POST");
+	case EHttpMethod::Head:   return ANSITEXTVIEW("HEAD");
+	case EHttpMethod::Delete: return ANSITEXTVIEW("DELETE");
+	default:                  return ANSITEXTVIEW("UNKNOWN");
+	}
+}
 
-static constexpr uint32 MaxTotalConnections = 8;
-static constexpr uint32 WaitIntervalMs = 10;
+bool TryLexFromString(EHttpMethod& OutMethod, const FAnsiStringView View)
+{
+	if (View == ANSITEXTVIEW("GET"))
+	{
+		OutMethod = EHttpMethod::Get;
+	}
+	else if (View == ANSITEXTVIEW("PUT"))
+	{
+		OutMethod = EHttpMethod::Put;
+	}
+	else if (View == ANSITEXTVIEW("POST"))
+	{
+		OutMethod = EHttpMethod::Post;
+	}
+	else if (View == ANSITEXTVIEW("HEAD"))
+	{
+		OutMethod = EHttpMethod::Head;
+	}
+	else if (View == ANSITEXTVIEW("DELETE"))
+	{
+		OutMethod = EHttpMethod::Delete;
+	}
+	else
+	{
+		return false;
+	}
+	return true;
+}
 
-struct FHttpSharedDataInternals
+FAnsiStringView LexToString(const EHttpMediaType MediaType)
 {
-	CURLSH* CurlShare;
-	CURLM* CurlMulti;
-	TDepletableMpscQueue<CURL*> PendingRequestAdditions;
-	FRWLock Locks[CURL_LOCK_DATA_LAST];
-	bool WriteLocked[CURL_LOCK_DATA_LAST]{};
-};
+	switch (MediaType)
+	{
+	case EHttpMediaType::Any:              return ANSITEXTVIEW("*/*");
+	case EHttpMediaType::Binary:           return ANSITEXTVIEW("application/octet-stream");
+	case EHttpMediaType::Text:             return ANSITEXTVIEW("text/plain");
+	case EHttpMediaType::Json:             return ANSITEXTVIEW("application/json");
+	case EHttpMediaType::Yaml:             return ANSITEXTVIEW("text/yaml");
+	case EHttpMediaType::CbObject:         return ANSITEXTVIEW("application/x-ue-cb");
+	case EHttpMediaType::CbPackage:        return ANSITEXTVIEW("application/x-ue-cbpkg");
+	case EHttpMediaType::CbPackageOffer:   return ANSITEXTVIEW("application/x-ue-offer");
+	case EHttpMediaType::CompressedBinary: return ANSITEXTVIEW("application/x-ue-comp");
+	case EHttpMediaType::FormUrlEncoded:   return ANSITEXTVIEW("application/x-www-form-urlencoded");
+	default:                               return ANSITEXTVIEW("unknown");
+	}
+}
 
-struct FHttpSharedDataStatics
+bool TryLexFromString(EHttpMediaType& OutMediaType, const FAnsiStringView View)
 {
-	static void LockFn(CURL* Handle, curl_lock_data Data, curl_lock_access Access, void* User)
+	const int32 SlashIndex = String::FindFirstChar(View, '/');
+	if (SlashIndex == INDEX_NONE)
 	{
-		FHttpSharedDataInternals* SharedDataInternals = ((FHttpSharedData*)User)->Internals.Get();
-		if (Access == CURL_LOCK_ACCESS_SHARED)
+		return false;
+	}
+
+	const FAnsiStringView Type = View.Left(SlashIndex);
+	const FAnsiStringView SubType = View.RightChop(SlashIndex + 1);
+	if (Type == ANSITEXTVIEW("application"))
+	{
+		if (SubType == ANSITEXTVIEW("octet-stream"))
+		{
+			OutMediaType = EHttpMediaType::Binary;
+		}
+		else if (SubType == ANSITEXTVIEW("json"))
+		{
+			OutMediaType = EHttpMediaType::Json;
+		}
+		else if (SubType == ANSITEXTVIEW("x-ue-cb"))
+		{
+			OutMediaType = EHttpMediaType::CbObject;
+		}
+		else if (SubType == ANSITEXTVIEW("x-ue-cbpkg"))
+		{
+			OutMediaType = EHttpMediaType::CbPackage;
+		}
+		else if (SubType == ANSITEXTVIEW("x-ue-offer"))
+		{
+			OutMediaType = EHttpMediaType::CbPackageOffer;
+		}
+		else if (SubType == ANSITEXTVIEW("x-ue-comp"))
+		{
+			OutMediaType = EHttpMediaType::CompressedBinary;
+		}
+		else if (SubType == ANSITEXTVIEW("x-www-form-urlencoded"))
 		{
-			SharedDataInternals->Locks[Data].ReadLock();
+			OutMediaType = EHttpMediaType::FormUrlEncoded;
 		}
 		else
 		{
-			SharedDataInternals->Locks[Data].WriteLock();
-			SharedDataInternals->WriteLocked[Data] = true;
+			return false;
 		}
 	}
-
-	static void UnlockFn(CURL* Handle, curl_lock_data Data, void* User)
+	else if (Type == ANSITEXTVIEW("text"))
 	{
-		FHttpSharedDataInternals* SharedDataInternals = ((FHttpSharedData*)User)->Internals.Get();
-		if (!SharedDataInternals->WriteLocked[Data])
+		if (SubType == ANSITEXTVIEW("plain"))
+		{
+			OutMediaType = EHttpMediaType::Text;
+		}
+		else if (SubType == ANSITEXTVIEW("yaml"))
 		{
-			SharedDataInternals->Locks[Data].ReadUnlock();
+			OutMediaType = EHttpMediaType::Yaml;
 		}
 		else
 		{
-			SharedDataInternals->WriteLocked[Data] = false;
-			SharedDataInternals->Locks[Data].WriteUnlock();
+			return false;
 		}
 	}
-};
+	else if (Type == ANSITEXTVIEW("*") && SubType == ANSITEXTVIEW("*"))
+	{
+		OutMediaType = EHttpMediaType::Any;
+	}
+	else
+	{
+		return false;
+	}
 
-} // Http::Private
+	return true;
+}
 
-void FHttpAccessToken::SetToken(const FStringView Token)
+///////////////////////////////////////////////////////////////////////////////////////////////////////////////////////
+
+void IHttpRequest::SetContentType(const EHttpMediaType Type, const FAnsiStringView Param)
 {
-	FWriteScopeLock WriteLock(Lock);
-	const FAnsiStringView Prefix = ANSITEXTVIEW("Authorization: Bearer ");
-	const int32 TokenLen = FPlatformString::ConvertedLength<ANSICHAR>(Token.GetData(), Token.Len());
-	Header.Empty(Prefix.Len() + TokenLen);
-	Header.Append(Prefix.GetData(), Prefix.Len());
-	const int32 TokenIndex = Header.AddUninitialized(TokenLen);
-	FPlatformString::Convert(Header.GetData() + TokenIndex, TokenLen, Token.GetData(), Token.Len());
-	Serial.fetch_add(1, std::memory_order_relaxed);
+	TAnsiStringBuilder<64> Value;
+	Value << LexToString(Type);
+	if (!Param.IsEmpty())
+	{
+		Value << ANSITEXTVIEW("; ") << Param;
+	}
+	AddHeader(ANSITEXTVIEW("Content-Type"), Value);
 }
 
-FAnsiStringBuilderBase& operator<<(FAnsiStringBuilderBase& Builder, const FHttpAccessToken& Token)
+void IHttpRequest::AddAcceptType(const EHttpMediaType Type, const float Weight)
 {
-	FReadScopeLock ReadLock(Token.Lock);
-	return Builder.Append(Token.Header);
+	TAnsiStringBuilder<64> Value;
+	Value << LexToString(Type);
+	if (Weight != 1.0f)
+	{
+		Value.Appendf(";q=%.3f", Weight);
+	}
+	AddHeader(ANSITEXTVIEW("Accept"), Value);
 }
 
-FHttpSharedData::FHttpSharedData(uint32 OverrideMaxConnections)
-	: PendingRequestEvent(EEventMode::AutoReset)
+///////////////////////////////////////////////////////////////////////////////////////////////////////////////////////
+
+FAnsiStringView IHttpResponse::GetHeader(const FAnsiStringView Name) const
 {
-	Internals = MakePimpl<Http::Private::FHttpSharedDataInternals>();
-	Internals->CurlShare = curl_share_init();
-	curl_share_setopt(Internals->CurlShare, CURLSHOPT_USERDATA, this);
-	curl_share_setopt(Internals->CurlShare, CURLSHOPT_LOCKFUNC, Http::Private::FHttpSharedDataStatics::LockFn);
-	curl_share_setopt(Internals->CurlShare, CURLSHOPT_UNLOCKFUNC, Http::Private::FHttpSharedDataStatics::UnlockFn);
-	curl_share_setopt(Internals->CurlShare, CURLSHOPT_SHARE, CURL_LOCK_DATA_DNS);
-	curl_share_setopt(Internals->CurlShare, CURLSHOPT_SHARE, CURL_LOCK_DATA_SSL_SESSION);
-	Internals->CurlMulti = curl_multi_init();
-	curl_multi_setopt(Internals->CurlMulti, CURLMOPT_MAX_TOTAL_CONNECTIONS, OverrideMaxConnections ? OverrideMaxConnections : Http::Private::MaxTotalConnections);
-	curl_multi_setopt(Internals->CurlMulti, CURLMOPT_MAXCONNECTS, OverrideMaxConnections ? OverrideMaxConnections : Http::Private::MaxTotalConnections); // Keep the connection pool at exactly the number of total connections
-	curl_multi_setopt(Internals->CurlMulti, CURLMOPT_PIPELINING, CURLPIPE_MULTIPLEX);
+	const int32 NameLen = Name.Len();
+	for (const FAnsiStringView& Header : GetAllHeaders())
+	{
+		if (Header.StartsWith(Name, ESearchCase::IgnoreCase) &&
+			Header.Len() > NameLen && Header.GetData()[NameLen] == ':')
+		{
+			return Header.TrimStartAndEnd();
+		}
+	}
+	return {};
 }
 
-FHttpSharedData::~FHttpSharedData()
+int32 IHttpResponse::GetHeaders(FAnsiStringView Name, TArrayView<FAnsiStringView> OutValues) const
 {
-	bAsyncThreadShutdownRequested.store(true, std::memory_order_relaxed);
-	if (AsyncServiceThread.IsJoinable())
+	int32 MatchIndex = 0;
+	const int32 NameLen = Name.Len();
+	const int32 ValuesCount = OutValues.Num();
+	FAnsiStringView* Values = OutValues.GetData();
+	for (const FAnsiStringView& Header : GetAllHeaders())
 	{
-		AsyncServiceThread.Join();
+		if (Header.StartsWith(Name, ESearchCase::IgnoreCase) &&
+			Header.Len() > NameLen && Header.GetData()[NameLen] == ':')
+		{
+			if (ValuesCount > MatchIndex)
+			{
+				Values[MatchIndex] = Header.TrimStartAndEnd();
+			}
+			++MatchIndex;
+		}
 	}
+	return MatchIndex;
+}
 
-	curl_multi_cleanup(Internals->CurlMulti);
-	curl_share_cleanup(Internals->CurlShare);
+EHttpMediaType IHttpResponse::GetContentType() const
+{
+	const FAnsiStringView ContentType = GetHeader(ANSITEXTVIEW("Content-Type"));
+	const FAnsiStringView ContentTypeNoParams = FAsciiSet::FindPrefixWithout(ContentType, "; \t");
+	EHttpMediaType MediaType = EHttpMediaType::Any;
+	TryLexFromString(MediaType, ContentTypeNoParams);
+	return MediaType;
 }
 
-void FHttpSharedData::AddRequest(void* Curl)
+FStringBuilderBase& operator<<(FStringBuilderBase& Builder, const IHttpResponse& Response)
 {
-	if (Internals->PendingRequestAdditions.EnqueueAndReturnWasEmpty(static_cast<CURL*>(Curl)))
+	Builder << LexToString(Response.GetMethod()) << TEXT(' ') << Response.GetUri();
+
+	if (const int32 StatusCode = Response.GetStatusCode(); StatusCode > 0)
 	{
-		PendingRequestEvent.Trigger();
+		Builder << TEXTVIEW(" -> ") << StatusCode;
 	}
 
-	if (!bAsyncThreadStarting.load(std::memory_order_relaxed) && !bAsyncThreadStarting.exchange(true, std::memory_order_relaxed))
+	if (const FAnsiStringView Error = Response.GetError(); !Error.IsEmpty())
 	{
-		AsyncServiceThread = FThread(TEXT("HttpService"), [this] { ProcessAsyncRequests(); }, 64 * 1024, TPri_Normal);
+		Builder << TEXTVIEW(": ") << Error;
 	}
+
+	return Builder;
 }
 
-void* FHttpSharedData::GetCurlShare() const
+///////////////////////////////////////////////////////////////////////////////////////////////////////////////////////
+
+FHttpByteArrayReceiver::FHttpByteArrayReceiver(TArray64<uint8>& OutArray, IHttpReceiver* InNext)
+	: Array(OutArray)
+	, Next(InNext)
 {
-	return Internals->CurlShare;
+	Array.Reset();
 }
 
-void FHttpSharedData::ProcessAsyncRequests()
+IHttpReceiver* FHttpByteArrayReceiver::OnBody(IHttpResponse& Response, FMemoryView& Data)
 {
-	int ActiveTransfers = 0;
-	TArray<Tasks::TTask<void>> PendingCompletionTasks;
-	PendingCompletionTasks.Reserve(128);
-
-	auto ProcessPendingRequests = [this, &ActiveTransfers, &PendingCompletionTasks]
+	if (Array.IsEmpty())
 	{
-		int CurrentActiveTransfers = -1;
-
-		do
+		constexpr int64 MaxReserveSize = 96 * 1024 * 1024;
+		if (const FAnsiStringView View = Response.GetHeader(ANSITEXTVIEW("Content-Length")); !View.IsEmpty())
 		{
-			Internals->PendingRequestAdditions.Deplete([this, &ActiveTransfers](CURL* Curl)
+			constexpr int32 MaxStringLen = 16;
+			ANSICHAR String[MaxStringLen];
+			if (const int32 CopyLen = View.CopyString(String, MaxStringLen); CopyLen < MaxStringLen)
 			{
-				curl_multi_add_handle(Internals->CurlMulti, Curl);
-				++ActiveTransfers;
-			});
-
-			curl_multi_perform(Internals->CurlMulti, &CurrentActiveTransfers);
-
-			if (CurrentActiveTransfers == 0 || ActiveTransfers != CurrentActiveTransfers)
-			{
-				for (;;)
+				String[CopyLen] = '\0';
+				const int64 ContentLength = FCStringAnsi::Atoi64(String);
+				if (ContentLength > 0 && ContentLength <= MaxReserveSize)
 				{
-					int MsgsStillInQueue = 0; // may use that to impose some upper limit we may spend in that loop
-					CURLMsg* Message = curl_multi_info_read(Internals->CurlMulti, &MsgsStillInQueue);
-
-					if (!Message)
-					{
-						break;
-					}
-
-					// find out which requests have completed
-					if (Message->msg == CURLMSG_DONE)
-					{
-						CURL* CompletedHandle = Message->easy_handle;
-						curl_multi_remove_handle(Internals->CurlMulti, CompletedHandle);
-
-						void* PrivateData = nullptr;
-						curl_easy_getinfo(CompletedHandle, CURLINFO_PRIVATE, &PrivateData);
-						FHttpRequest* CompletedRequest = (FHttpRequest*)PrivateData;
-
-						if (CompletedRequest)
-						{
-							// It is important that the CompleteAsync call doesn't happen on this thread as it is possible it will block waiting
-							// for a free HTTP request, and if that happens on this thread, we can deadlock as no HTTP requests will become 
-							// available while this thread is blocked.
-							PendingCompletionTasks.Emplace(Tasks::Launch(TEXT("FHttpRequest::CompleteAsync"), [CompletedRequest, Result = Message->data.result]() mutable
-							{
-								CompletedRequest->CompleteAsync(Result);
-							}));
-						}
-					}
+					Array.Reserve(ContentLength);
 				}
-				ActiveTransfers = CurrentActiveTransfers;
 			}
-
-			if (CurrentActiveTransfers > 0)
-			{
-				curl_multi_wait(Internals->CurlMulti, nullptr, 0, 1, nullptr);
-			}
-		}
-		while (CurrentActiveTransfers > 0);
-	};
-
-	do
-	{
-		ProcessPendingRequests();
-
-		PendingCompletionTasks.SetNum(Algo::StableRemoveIf(PendingCompletionTasks, [](const Tasks::TTask<void>& Task) { return Task.IsCompleted(); }), false);
-		if (!PendingCompletionTasks.IsEmpty())
-		{
-			Tasks::Wait(PendingCompletionTasks, FTimespan::FromMilliseconds(100));
-		}
-		else
-		{
-			PendingRequestEvent.Wait(100);
 		}
 	}
-	while (!FHttpSharedData::bAsyncThreadShutdownRequested.load(std::memory_order_relaxed));
-
-	// Process last requests before shutdown.  May want these to be aborted instead.
-	ProcessPendingRequests();
+	Array.Append(static_cast<const uint8*>(Data.GetData()), int64(Data.GetSize()));
+	return this;
 }
 
-FHttpRequestPool::FHttpRequestPool(
-	const FStringView InServiceUrl,
-	FStringView InEffectiveServiceUrl,
-	const FHttpAccessToken* const InAuthorizationToken,
-	FHttpSharedData* const InSharedData,
-	const uint32 PoolSize,
-	const uint32 InOverflowLimit)
-	: ActiveOverflowRequests(0)
-	, OverflowLimit(InOverflowLimit)
-{
-	InEffectiveServiceUrl = FAsciiSet::TrimSuffixWith(InEffectiveServiceUrl, "/");
+///////////////////////////////////////////////////////////////////////////////////////////////////////////////////////
 
-	Pool.AddUninitialized(PoolSize);
-	Requests.AddUninitialized(PoolSize);
-	for (int32 Index = 0; Index < Pool.Num(); ++Index)
-	{
-		Pool[Index].Usage = 0u;
-		Pool[Index].Request = new(&Requests[Index]) FHttpRequest(InServiceUrl, InEffectiveServiceUrl, InAuthorizationToken, InSharedData, /*bLogErrors*/ true);
-	}
-
-	InitData = MakeUnique<FInitData>(InServiceUrl, InEffectiveServiceUrl, InAuthorizationToken, InSharedData);
-}
-
-FHttpRequestPool::~FHttpRequestPool()
+struct Http::Private::FHttpRequestQueueData
 {
-	check(ActiveOverflowRequests.load() == 0);
-	for (const FEntry& Entry : Pool)
+	struct FWaiter
 	{
-		// No requests should be in use by now.
-		check(Entry.Usage.load(std::memory_order_acquire) == 0u);
-	}
-}
+		THttpUniquePtr<IHttpRequest> Request;
+		FEventRef Event{EEventMode::ManualReset};
+	};
 
-FHttpRequest* FHttpRequestPool::GetFreeRequest(bool bUnboundedOverflow)
-{
-	for (FEntry& Entry : Pool)
+	FHttpRequestQueueData(IHttpConnectionPool& ConnectionPool, const FHttpClientParams& ClientParams)
 	{
-		if (!Entry.Usage.load(std::memory_order_relaxed))
+		FHttpClientParams QueueParams = ClientParams;
+		QueueParams.OnDestroyRequest = [this, OnDestroyRequest = MoveTemp(QueueParams.OnDestroyRequest)]
 		{
-			uint8 Expected = 0u;
-			if (Entry.Usage.compare_exchange_strong(Expected, 1u))
+			if (OnDestroyRequest)
 			{
-				Entry.Request->Reset();
-				return Entry.Request;
+				OnDestroyRequest();
 			}
-		}
+			if (!Waiters.IsEmpty())
+			{
+				if (THttpUniquePtr<IHttpRequest> Request = Client->TryCreateRequest({}))
+				{
+					if (FWaiter* Waiter = Waiters.Pop())
+					{
+						Waiter->Request = MoveTemp(Request);
+						Waiter->Event->Trigger();
+					}
+				}
+			}
+		};
+		Client = ConnectionPool.CreateClient(QueueParams);
 	}
-	if (bUnboundedOverflow || (OverflowLimit > 0))
+
+	THttpUniquePtr<IHttpRequest> CreateRequest(const FHttpRequestParams& Params)
 	{
-		// The use of two operations here (load, then fetch_add) implies that we can exceed the overflow limit because the combined operation
-		// is not atomic.  This is acceptable for our use case.  If we wanted to enforce the hard limit, we could use a loop instead.
-		if (bUnboundedOverflow || (ActiveOverflowRequests.load(std::memory_order_relaxed) < OverflowLimit))
+		if (Params.bIgnoreMaxRequests)
 		{
-			// Create an overflow request (outside of the pre-allocated range of requests)
-			ActiveOverflowRequests.fetch_add(1, std::memory_order_relaxed);
-			return new FHttpRequest(InitData->ServiceUrl, InitData->EffectiveServiceUrl, InitData->AccessToken, InitData->SharedData, true);
+			return Client->TryCreateRequest(Params);
 		}
-	}
-	return nullptr;
-}
-
-FHttpRequest* FHttpRequestPool::WaitForFreeRequest(bool bUnboundedOverflow)
-{
-	TRACE_CPUPROFILER_EVENT_SCOPE(HttpDDC_WaitForFreeRequest);
 
-	FHttpRequest* Request = GetFreeRequest(bUnboundedOverflow);
-	if (Request == nullptr)
-	{
-		// Make it a fair by allowing each thread to register itself in a FIFO
-		// so that the first thread to start waiting is the first one to get a request.
-		FWaiter* Waiter = new FWaiter(this);
-		Waiter->AddRef(); // One ref for the thread that will dequeue
-		Waiter->AddRef(); // One ref for us
-
-		Waiters.enqueue(Waiter);
-
-		while (!Waiter->Wait(Http::Private::WaitIntervalMs))
+		while (THttpUniquePtr<IHttpRequest> Request = Client->TryCreateRequest(Params))
 		{
-			// While waiting, allow us to check if a race occurred and a request has been freed
-			// between the time we checked for free requests and the time we queued ourself as a Waiter.
-			if ((Request = GetFreeRequest(bUnboundedOverflow)) != nullptr)
+			if (FWaiter* Waiter = Waiters.Pop())
+			{
+				Waiter->Request = MoveTemp(Request);
+				Waiter->Event->Trigger();
+			}
+			else
 			{
-				// We abandon the FWaiter, it will be freed by the next dequeue
-				// and if it has a request, it will be queued back to the pool.
-				Waiter->Release();
 				return Request;
 			}
 		}
 
-		Request = Waiter->Request.exchange(nullptr);
-		Request->Reset();
-		Waiter->Release();
-	}
-	check(Request);
-	return Request;
-}
+		FWaiter LocalWaiter;
+		Waiters.Push(&LocalWaiter);
 
-void FHttpRequestPool::ReleaseRequestToPool(FHttpRequest* Request)
-{
-	if ((Request < Requests.GetData()) || (Request >= (Requests.GetData() + Requests.Num())))
-	{
-		// For overflow requests (outside of the pre-allocated range of requests), just delete it immediately
-		delete Request;
-		ActiveOverflowRequests.fetch_sub(1, std::memory_order_relaxed);
-		return;
-	}
-
-	for (FEntry& Entry : Pool)
-	{
-		if (Entry.Request == Request)
+		while (THttpUniquePtr<IHttpRequest> Request = Client->TryCreateRequest(Params))
 		{
-			// If only 1 user is remaining, we can give it to a waiter
-			// instead of releasing it back to the pool.
-			if (Entry.Usage == 1u)
+			if (FWaiter* Waiter = Waiters.Pop())
 			{
-				if (FWaiter* Waiter = Waiters.dequeue())
-				{
-					Waiter->Request = Request;
-					Waiter->Trigger();
-					Waiter->Release();
-					return;
-				}
+				Waiter->Request = MoveTemp(Request);
+				Waiter->Event->Trigger();
+			}
+			if (LocalWaiter.Event->Wait(0))
+			{
+				return MoveTemp(LocalWaiter.Request);
 			}
-			
-			Entry.Usage--;
-			return;
 		}
+
+		TRACE_CPUPROFILER_EVENT_SCOPE(HttpRequestQueue_Wait);
+		LocalWaiter.Event->Wait();
+		return MoveTemp(LocalWaiter.Request);
 	}
 
-	checkNoEntry();
-}
+	THttpUniquePtr<IHttpClient> Client;
+	TLockFreePointerListFIFO<FWaiter, 0> Waiters;
+};
 
-void FHttpRequestPool::MakeRequestShared(FHttpRequest* Request, uint8 Users)
+FHttpRequestQueue::FHttpRequestQueue(IHttpConnectionPool& ConnectionPool, const FHttpClientParams& ClientParams)
+	: Data(MakePimpl<Http::Private::FHttpRequestQueueData>(ConnectionPool, ClientParams))
 {
-	// Overflow requests (outside of the pre-allocated range of requests), cannot be made shared
-	check((Request >= Requests.GetData()) && (Request < (Requests.GetData() + Requests.Num())));
-
-	check(Users != 0);
-	for (FEntry& Entry : Pool)
-	{
-		if (Entry.Request == Request)
-		{
-			Entry.Usage = Users;
-			return;
-		}
-	}
+}
 
-	checkNoEntry();
+THttpUniquePtr<IHttpRequest> FHttpRequestQueue::CreateRequest(const FHttpRequestParams& Params)
+{
+	static_assert(sizeof(Params) == sizeof(bool), "CreateRequest only handles bIgnoreMaxRequests");
+	check(Data);
+	return Data->CreateRequest(Params);
 }
 
 } // UE
diff --git a/Engine/Source/Developer/DerivedDataCache/Private/Http/HttpClient.h b/Engine/Source/Developer/DerivedDataCache/Private/Http/HttpClient.h
--- a/Engine/Source/Developer/DerivedDataCache/Private/Http/HttpClient.h
+++ b/Engine/Source/Developer/DerivedDataCache/Private/Http/HttpClient.h
@@ -3,39 +3,48 @@
 #pragma once
 
 #include "CoreTypes.h"
-#include "Containers/Array.h"
-#include "Containers/ArrayView.h"
-#include "Containers/DepletableMpscQueue.h"
+#include "Containers/ContainersFwd.h"
 #include "Containers/StringFwd.h"
-#include "Containers/UnrealString.h"
-#include "Dom/JsonObject.h"
-#include "Experimental/Async/LazyEvent.h"
-#include "Experimental/Containers/FAAArrayQueue.h"
-#include "HAL/CriticalSection.h"
-#include "HAL/PlatformProcess.h"
-#include "HAL/Thread.h"
-#include "Memory/CompositeBuffer.h"
-#include "Memory/MemoryView.h"
+#include "Containers/StringView.h"
+#include "Memory/MemoryFwd.h"
+#include "Templates/Function.h"
 #include "Templates/PimplPtr.h"
 #include "Templates/RefCounting.h"
-#include "Templates/SharedPointer.h"
 #include "Templates/UniquePtr.h"
-#include <atomic>
-
-struct curl_slist;
-namespace UE { class FHttpRequest; }
-namespace UE::DerivedData { class IRequestOwner; }
-namespace UE::Http::Private { struct FAsyncRequestData; }
-namespace UE::Http::Private { struct FCurlStringListDeleter; }
-namespace UE::Http::Private { struct FHttpRequestStatics; }
-namespace UE::Http::Private { struct FHttpSharedDataInternals; }
-namespace UE::Http::Private { struct FHttpSharedDataStatics; }
-namespace UE::Http::Private { using FCurlStringList = TUniquePtr<curl_slist, FCurlStringListDeleter>; }
+
+namespace UE { class IHttpClient; }
+namespace UE { class IHttpConnectionPool; }
+namespace UE { class IHttpReceiver; }
+namespace UE { class IHttpRequest; }
+namespace UE { class IHttpResponse; }
+namespace UE { struct FHttpClientParams; }
+namespace UE { struct FHttpConnectionPoolParams; }
+namespace UE { struct FHttpRequestParams; }
+namespace UE { struct FHttpResponseStats; }
+namespace UE::Http::Private { struct FHttpRequestQueueData; }
+
+namespace UE::Http::Private
+{
+
+template <typename T>
+struct THttpDestroyer
+{
+	void operator()(void* Object) const
+	{
+		if (Object)
+		{
+			((T*)Object)->Destroy();
+		}
+	}
+};
+
+} // UE::Http::Private
 
 namespace UE
 {
 
-class FHttpRequestPool;
+template <typename Type, typename BaseType = Type>
+using THttpUniquePtr = TUniquePtr<Type, Http::Private::THttpDestroyer<BaseType>>;
 
 enum class EHttpMethod : uint8
 {
@@ -46,10 +55,11 @@ enum class EHttpMethod : uint8
 	Delete,
 };
 
-enum class EHttpMediaType : int16
+FAnsiStringView LexToString(EHttpMethod Method);
+bool TryLexFromString(EHttpMethod& OutMethod, FAnsiStringView View);
+
+enum class EHttpMediaType : uint8
 {
-	// Negative integers reserved for future custom string content types
-	UnspecifiedContentType = 0,
 	Any,
 	Binary,
 	Text,
@@ -62,520 +72,382 @@ enum class EHttpMediaType : int16
 	FormUrlEncoded,
 };
 
-inline FStringView GetHttpMimeType(EHttpMediaType Type)
+FAnsiStringView LexToString(EHttpMediaType MediaType);
+bool TryLexFromString(EHttpMediaType& OutMediaType, FAnsiStringView View);
+
+enum class EHttpTlsLevel : uint8
 {
-	switch (Type)
-	{
-		case EHttpMediaType::UnspecifiedContentType:
-			checkNoEntry();
-			return FStringView();
-		case EHttpMediaType::Any:
-			return TEXTVIEW("*/*");
-		case EHttpMediaType::Binary:
-			return TEXTVIEW("application/octet-stream");
-		case EHttpMediaType::Text:
-			return TEXTVIEW("text/plain");
-		case EHttpMediaType::Json:
-			return TEXTVIEW("application/json");
-		case EHttpMediaType::Yaml:
-			return TEXTVIEW("text/yaml");
-		case EHttpMediaType::CbObject:
-			return TEXTVIEW("application/x-ue-cb");
-		case EHttpMediaType::CbPackage:
-			return TEXTVIEW("application/x-ue-cbpkg");
-		case EHttpMediaType::CbPackageOffer:
-			return TEXTVIEW("application/x-ue-offer");
-		case EHttpMediaType::CompressedBinary:
-			return TEXTVIEW("application/x-ue-comp");
-		case EHttpMediaType::FormUrlEncoded:
-			return TEXTVIEW("application/x-www-form-urlencoded");
-		default:
-			return TEXTVIEW("unknown");
-	}
-}
+	/** Do not use TLS. */
+	None = 0,
+	/** Try to use TLS and fall back to none. */
+	Try,
+	/** Require TLS. */
+	All,
+};
 
-/**
- * Encapsulation for access token shared by all requests.
- */
-class FHttpAccessToken
+enum class EHttpErrorCode : uint8
+{
+	/** Success. */
+	None = 0,
+	/** Unknown. Check IHttpResponse::GetError(). */
+	Unknown,
+	/** Request was canceled. */
+	Canceled,
+	/** Failed to resolve the host. */
+	ResolveHost,
+	/** Failed to establish a connection. */
+	Connect,
+	/** Failed to establish a TLS connection. */
+	TlsConnect,
+	/** Failed to verify the certificate of the peer. */
+	TlsPeerVerification,
+	/** Failed because the timeout period was exceeded. */
+	TimedOut,
+};
+
+class IHttpManager
 {
 public:
-	void SetToken(FStringView Token);
-	inline uint32 GetSerial() const { return Serial.load(std::memory_order_relaxed); }
-	friend FAnsiStringBuilderBase& operator<<(FAnsiStringBuilderBase& Builder, const FHttpAccessToken& Token);
+	static IHttpManager& Get();
 
-private:
-	mutable FRWLock Lock;
-	TArray<ANSICHAR> Header;
-	std::atomic<uint32> Serial;
+	[[nodiscard]] virtual THttpUniquePtr<IHttpConnectionPool> CreateConnectionPool(const FHttpConnectionPoolParams& Params) = 0;
 };
 
-class FHttpSharedData
+class IHttpConnectionPool
 {
 public:
-	FHttpSharedData(uint32 OverrideMaxConnections = 0);
-	~FHttpSharedData();
+	[[nodiscard]] virtual THttpUniquePtr<IHttpClient> CreateClient(const FHttpClientParams& Params) = 0;
 
-private:
-	void ProcessAsyncRequests();
-	void AddRequest(void* Curl);
-	void* GetCurlShare() const;
+protected:
+	friend Http::Private::THttpDestroyer<IHttpConnectionPool>;
+	virtual ~IHttpConnectionPool() = default;
+	virtual void Destroy() = 0;
+};
 
-	friend FHttpRequest;
-	friend Http::Private::FHttpSharedDataStatics;
+struct FHttpConnectionPoolParams final
+{
+	/** Maximum number of concurrent connections created by the pool. Use 0 for the default limit. */
+	uint32 MaxConnections = 0;
 
-	FLazyEvent PendingRequestEvent;
-	FThread AsyncServiceThread;
-	std::atomic<bool> bAsyncThreadStarting = false;
-	std::atomic<bool> bAsyncThreadShutdownRequested = false;
+	/** Minimum number of concurrent connections maintained by the pool for reuse. Use 0 for the default limit. */
+	uint32 MinConnections = 0;
 
-	TPimplPtr<Http::Private::FHttpSharedDataInternals> Internals;
+	/** Allow requests to send async using the pool. Requests will block on execution when this is disabled. */
+	bool bAllowAsync = true;
 };
 
-
-/**
- * Minimal HTTP request type wrapping CURL without the need for managers. This request
- * is written to allow reuse of request objects, in order to allow connections to be reused.
- *
- * CURL has a global library initialization (curl_global_init). We rely on this happening in 
- * the Online/HTTP library which is a dependency on this module.
- */
-class FHttpRequest
+class IHttpClient
 {
 public:
-
 	/**
-	 * Convenience result type interpreted from HTTP response code.
+	 * Try to create a request.
+	 *
+	 * The request is counted against the maximum request count as soon as it is created.
+	 *
+	 * Returns null when the maximum request count has been reached. Retry after destroying a request.
 	 */
-	enum class EResult
-	{
-		Success,
-		Failed,
-		FailedTimeout
-	};
-
-	enum class ECompletionBehavior : uint8
-	{
-		Done,
-		Retry
-	};
+	[[nodiscard]] virtual THttpUniquePtr<IHttpRequest> TryCreateRequest(const FHttpRequestParams& Params) = 0;
 
-	using FResultCode = int;
+protected:
+	friend Http::Private::THttpDestroyer<IHttpClient>;
+	virtual ~IHttpClient() = default;
+	virtual void Destroy() = 0;
+};
 
-	using FOnHttpRequestComplete = TUniqueFunction<ECompletionBehavior(EResult HttpResult, FHttpRequest* Request)>;
+struct FHttpClientParams final
+{
+	/** Invoked whenever a request from this client is destroyed. This may be used to retry creating a request. */
+	TFunction<void()> OnDestroyRequest;
 
-	FHttpRequest(
-		FStringView InDomain,
-		FStringView InEffectiveDomain,
-		const FHttpAccessToken* InAuthorizationToken,
-		FHttpSharedData* InSharedData,
-		bool bInLogErrors);
-	~FHttpRequest();
+	/** Maximum number of concurrent requests created by the client. Use 0 for the default limit. */
+	uint32 MaxRequests = 0;
 
-	/**
-	 * Resets all options on the request except those that should always be set.
-	 */
-	void Reset();
+	/** Minimum number of concurrent requests maintained by the client for reuse. Use 0 for the default limit. */
+	uint32 MinRequests = 0;
 
-	void PrepareToRetry();
+	/** Domain name cache timeout in seconds. Use 0 to disable the cache. */
+	uint32 DnsCacheTimeout = 60;
 
-	/** Gets the domain name for this request */
-	const FString& GetName() const
-	{
-		return Domain;
-	}
+	/** Connection timeout in milliseconds. Use 0 for the default timeout. */
+	uint32 ConnectTimeout = 0;
 
-	/** Gets the domain name for this request */
-	const FString& GetDomain() const
-	{
-		return Domain;
-	}
+	/** Average transfer speed in bytes/s below which a request will abort. Use 0 to disable. */
+	uint32 LowSpeedLimit = 0;
 
-	/** Gets the effective domain name for this request */
-	const FString& GetEffectiveDomain() const
-	{
-		return EffectiveDomain;
-	}
+	/** Time in seconds after which a request will abort if it stays below LowSpeedLimit. Use 0 to disable. */
+	uint32 LowSpeedTime = 0;
 
-	/** Returns the HTTP response code.*/
-	const int64 GetResponseCode() const
-	{
-		return ResponseCode;
-	}
+	/** Level of TLS to use for requests created by the client. */
+	EHttpTlsLevel TlsLevel = EHttpTlsLevel::None;
 
-	/** Returns the number of bytes received this request (headers withstanding). */
-	const size_t GetBytesReceived() const
-	{
-		return BytesReceived;
-	}
+	/** Follow redirects in responses automatically. */
+	bool bFollowRedirects : 1;
+	/** Follow redirects of the corresponding status code without rewriting POST to GET. */
+	bool bFollow301Post : 1;
+	bool bFollow302Post : 1;
+	bool bFollow303Post : 1;
 
-	/** Returns the number of attempts we've made issuing this request (currently tracked for async requests only). */
-	const size_t GetAttempts() const
-	{
-		return Attempts;
-	}
+	/** Verbose logging for requests created by the client. */
+	bool bVerbose : 1;
 
-	/** Returns the number of bytes sent during this request (headers withstanding). */
-	const size_t GetBytesSent() const
+	inline FHttpClientParams()
+		: bFollowRedirects(false)
+		, bFollow301Post(false)
+		, bFollow302Post(false)
+		, bFollow303Post(false)
+		, bVerbose(false)
 	{
-		return BytesSent;
 	}
+};
 
+class IHttpRequest
+{
+public:
 	/**
-	 * Upload buffer using the request, using either "Put" or "Post" verbs.
-	 * @param Uri Url to use.
-	 * @param Buffer Data to upload
-	 * @return Result of the request
+	 * Reset the request to an empty state equivalent to when it was created.
+	 *
+	 * The request must be idle to be reset, which means any associated response must have completed.
 	 */
-	EResult PerformBlockingPut(
-		FStringView Uri,
-		const FCompositeBuffer& Buffer,
-		EHttpMediaType ContentType = EHttpMediaType::UnspecifiedContentType,
-		TConstArrayView<long> ExpectedErrorCodes = {});
-	EResult PerformBlockingPost(FStringView Uri,
-		const FCompositeBuffer& Buffer,
-		EHttpMediaType ContentType = EHttpMediaType::UnspecifiedContentType,
-		EHttpMediaType AcceptType = EHttpMediaType::UnspecifiedContentType,
-		TConstArrayView<long> ExpectedErrorCodes = {});
-	void EnqueueAsyncPut(
-		DerivedData::IRequestOwner& Owner,
-		FHttpRequestPool* Pool,
-		FStringView Uri,
-		const FCompositeBuffer& Buffer,
-		FOnHttpRequestComplete&& OnComplete,
-		EHttpMediaType ContentType = EHttpMediaType::UnspecifiedContentType,
-		TConstArrayView<long> ExpectedErrorCodes = {});
-	void EnqueueAsyncPost(
-		DerivedData::IRequestOwner& Owner,
-		FHttpRequestPool* Pool,
-		FStringView Uri,
-		const FCompositeBuffer& Buffer,
-		FOnHttpRequestComplete&& OnComplete,
-		EHttpMediaType ContentType = EHttpMediaType::UnspecifiedContentType,
-		EHttpMediaType AcceptType = EHttpMediaType::UnspecifiedContentType,
-		TConstArrayView<long> ExpectedErrorCodes = {});
+	virtual void Reset() = 0;
 
-	/**
-	 * Download an url into a buffer using the request.
-	 * @param Uri Url to use.
-	 * @param Buffer Optional buffer where data should be downloaded to. If empty downloaded data will
-	 * be stored in an internal buffer and accessed GetResponse* methods.
-	 * @return Result of the request
-	 */
-	EResult PerformBlockingDownload(
-		FStringView Uri,
-		TArray64<uint8>* Buffer,
-		EHttpMediaType AcceptType = EHttpMediaType::UnspecifiedContentType,
-		TConstArrayView<long> ExpectedErrorCodes = {400});
-	void EnqueueAsyncDownload(
-		DerivedData::IRequestOwner& Owner,
-		FHttpRequestPool* Pool,
-		FStringView Uri,
-		FOnHttpRequestComplete&& OnComplete,
-		EHttpMediaType AcceptType = EHttpMediaType::UnspecifiedContentType,
-		TConstArrayView<long> ExpectedErrorCodes = {400});
+	/** Set the URI used for the request, including the scheme and authority. */
+	virtual void SetUri(FAnsiStringView Uri) = 0;
 
-	/**
-	 * Query an url using the request. Queries can use either "Head" or "Delete" verbs.
-	 * @param Uri Url to use.
-	 * @return Result of the request
-	 */
-	EResult PerformBlockingHead(
-		FStringView Uri,
-		EHttpMediaType AcceptType = EHttpMediaType::UnspecifiedContentType,
-		TConstArrayView<long> ExpectedErrorCodes = {400});
-	EResult PerformBlockingDelete(
-		FStringView Uri,
-		TConstArrayView<long> ExpectedErrorCodes = {400});
-	void EnqueueAsyncHead(
-		DerivedData::IRequestOwner& Owner,
-		FHttpRequestPool* Pool,
-		FStringView Uri,
-		FOnHttpRequestComplete&& OnComplete,
-		EHttpMediaType AcceptType = EHttpMediaType::UnspecifiedContentType,
-		TConstArrayView<long> ExpectedErrorCodes = {400});
-	void EnqueueAsyncDelete(
-		DerivedData::IRequestOwner& Owner,
-		FHttpRequestPool* Pool,
-		FStringView Uri,
-		FOnHttpRequestComplete&& OnComplete,
-		TConstArrayView<long> ExpectedErrorCodes = {400});
+	/** Set the method used for the request. Defaults to GET. */
+	virtual void SetMethod(EHttpMethod Method) = 0;
 
 	/**
-	 * Set a header to send with the request.
+	 * Set the body for PUT and POST requests.
+	 *
+	 * Body must be owned or otherwise remain valid until the request is complete.
 	 */
-	void AddHeader(FStringView Header, FStringView Value);
+	virtual void SetBody(const FCompositeBuffer& Body) = 0;
 
-	/**
-	 * Attempts to find the header from the response. Returns false if header is not present.
-	 */
-	bool GetHeader(FAnsiStringView Header, FString& OutValue) const;
+	/** Set the Content-Type header with a known media type. Use AddHeader() for other types. */
+	void SetContentType(EHttpMediaType Type, FAnsiStringView Param = {});
 
-	/**
-	 * Returns the response buffer. Note that is the request is performed
-	 * with an external buffer as target buffer this string will be empty.
-	 */
-	const TArray64<uint8>& GetResponseBuffer() const
-	{
-		return ResponseBuffer;
-	}
+	/** Add an Accept header with a known media type. Use AddHeader() for other types. */
+	void AddAcceptType(EHttpMediaType Type, float Weight = 1.0);
 
-	FSharedBuffer MoveResponseBufferToShared()
-	{
-		return MakeSharedBufferFromArray(MoveTemp(ResponseBuffer));
-	}
+	/** Add the header. An empty value removes the header. */
+	virtual void AddHeader(FAnsiStringView Name, FAnsiStringView Value) = 0;
 
 	/**
-	 * Returns the response buffer as a string. Note that is the request is performed
-	 * with an external buffer as target buffer this string will be empty.
+	 * Send the request on the client that it was created from, and wait for it to complete before returning.
+	 *
+	 * Request must be idle to send, which means any previous response must be complete.
+	 *
+	 * Receiver must remain valid until OnComplete is called or until it returns another receiver
+	 * to be used for subsequent events.
 	 */
-	FString GetResponseAsString() const
-	{
-		return GetAnsiBufferAsString(ResponseBuffer);
-	}
+	virtual void Send(IHttpReceiver* Receiver, THttpUniquePtr<IHttpResponse>& OutResponse) = 0;
 
 	/**
-	 * Returns the response header as a string.
+	 * Asynchronously send the request on the client that it was created from.
+	 *
+	 * Request must be idle to send, which means any previous response must be complete.
+	 *
+	 * Receiver must remain valid until OnComplete is called or until it returns another receiver
+	 * to be used for subsequent events.
+	 *
+	 * Response must be kept alive until OnComplete is called. Cancel to complete immediately.
 	 */
-	FString GetResponseHeaderAsString()
-	{
-		return GetAnsiBufferAsString(ResponseHeader);
-	}
+	virtual void SendAsync(IHttpReceiver* Receiver, THttpUniquePtr<IHttpResponse>& OutResponse) = 0;
 
-	/**
-	 * Tries to parse the response buffer as a JsonObject. Return empty pointer if 
-	 * parse error occurs.
-	 */
-	TSharedPtr<FJsonObject> GetResponseAsJsonObject() const;
+protected:
+	friend Http::Private::THttpDestroyer<IHttpRequest>;
+	virtual ~IHttpRequest() = default;
+	virtual void Destroy() = 0;
+};
 
-	/**
-	 * Tries to parse the response buffer as a JsonArray. Return empty array if
-	 * parse error occurs.
-	 */
-	TArray<TSharedPtr<FJsonValue>> GetResponseAsJsonArray() const;
+struct FHttpRequestParams final
+{
+	/** Whether to force creation of the request even if it would exceed the configured maximum. */
+	bool bIgnoreMaxRequests = false;
+};
 
-	void CompleteAsync(FResultCode Result);
+class IHttpResponseMonitor
+{
+public:
+	/** Cancels the request associated with this response. */
+	virtual void Cancel() = 0;
 
-	/** Will return true if the response code is considered a success */
-	static bool IsSuccessResponse(long ResponseCode)
-	{
-		// We consider anything in the 1XX or 2XX range a success
-		return ResponseCode >= 100 && ResponseCode < 300;
-	}
+	/** Waits for the request associated with this response to be complete. */
+	virtual void Wait() const = 0;
 
-	static bool AllowAsync();
+	/** Returns true if the request associated with this response is complete, and false otherwise. */
+	[[nodiscard]] virtual bool Poll() const = 0;
 
-private:
-	friend struct Http::Private::FAsyncRequestData;
-	friend struct Http::Private::FHttpRequestStatics;
-
-	void* Curl;
-	FResultCode ResultCode;
-	FHttpSharedData* SharedData;
-	Http::Private::FAsyncRequestData* AsyncData;
-	long ResponseCode;
-	size_t BytesSent;
-	size_t BytesReceived;
-	size_t Attempts;
-	bool bLogErrors;
-
-	FCompositeBuffer ReadCompositeBuffer;
-	TArray64<uint8>* WriteDataBufferPtr;
-	TArray64<uint8>* WriteHeaderBufferPtr;
-	
-	TArray64<uint8> ResponseHeader;
-	TArray64<uint8> ResponseBuffer;
-	TArray<FString> Headers;
-	FString Domain;
-	FString EffectiveDomain;
-	const FHttpAccessToken* AuthorizationToken;
-
-	void AddContentTypeHeader(FStringView Header, EHttpMediaType Type);
-	Http::Private::FCurlStringList PrepareToIssueRequest(FStringView Uri, EHttpMethod Verb, uint64 ContentLength);
+	virtual void AddRef() const = 0;
+	virtual void Release() const = 0;
+};
+
+class IHttpResponse
+{
+public:
+	/** Returns a monitor for the response that can poll, wait, or cancel it. */
+	[[nodiscard]] virtual TRefCountPtr<IHttpResponseMonitor> GetMonitor() = 0;
+
+	/** Cancels the request associated with this response. */
+	virtual void Cancel() = 0;
+
+	/** Waits for the request associated with this response to be complete. */
+	virtual void Wait() const = 0;
+
+	/** Returns true if the request associated with this response is complete, and false otherwise. */
+	[[nodiscard]] virtual bool Poll() const = 0;
+
+	/** Get the URI used for the request associated with this response. Always available. */
+	[[nodiscard]] virtual FAnsiStringView GetUri() const = 0;
+
+	/** Get the method used for the request associated with this response. Always available. */
+	[[nodiscard]] virtual EHttpMethod GetMethod() const = 0;
 
 	/**
-	 * Performs the request, blocking until finished.
-	 * @param Uri Address on the domain to query
-	 * @param Verb HTTP verb to use
-	 * @param ContentLength The number of bytes being uploaded for the body of this request.
-	 * @param ExpectedErrorCodes An array of expected return codes outside of the success range that should NOT be logged as an abnormal/exceptional outcome.
-	 * If unset the response body will be stored in the request.
+	 * Returns the status code of the response. Available once the response status has been received.
+	 *
+	 * The return value is negative until the response status has been received.
+	 * The return value is 0 if an error occurred. See GetError().
 	 */
-	EResult PerformBlocking(FStringView Uri, EHttpMethod Verb, uint64 ContentLength, TConstArrayView<long> ExpectedErrorCodes);
+	[[nodiscard]] virtual int32 GetStatusCode() const = 0;
 
-	void EnqueueAsync(
-		DerivedData::IRequestOwner& Owner,
-		FHttpRequestPool* Pool,
-		FStringView Uri,
-		EHttpMethod Verb,
-		uint64 ContentLength,
-		FOnHttpRequestComplete&& OnComplete,
-		TConstArrayView<long> ExpectedErrorCodes);
+	/** Returns the error code from the HTTP stack. Available when the response is complete. */
+	[[nodiscard]] virtual EHttpErrorCode GetErrorCode() const = 0;
 
-	void LogResult(FResultCode EResult, FStringView Uri, EHttpMethod Verb, TConstArrayView<long> ExpectedErrorCodes) const;
+	/** Returns the optional error message from the HTTP stack. Available when the status code is 0. */
+	[[nodiscard]] virtual FAnsiStringView GetError() const = 0;
 
-	static FString GetAnsiBufferAsString(TConstArrayView64<uint8> Buffer)
-	{
-		// Content is NOT null-terminated; we need to specify lengths here
-		static_assert(sizeof(UTF8CHAR) == sizeof(uint8));
-		return FString(Buffer.Num(), reinterpret_cast<const UTF8CHAR*>(Buffer.GetData()));
-	}
-};
+	/** Returns information about the request such as size, rate, and duration. Available when complete. */
+	[[nodiscard]] virtual const FHttpResponseStats& GetStats() const = 0;
 
-/**
- * Pool that manages a fixed set of requests. Users are required to release requests that have been 
- * acquired. Usable with \ref FScopedRequestPtr which handles this automatically.
- */
-class FHttpRequestPool
-{
-public:
-	FHttpRequestPool(
-		FStringView InServiceUrl,
-		FStringView InEffectiveServiceUrl,
-		const FHttpAccessToken* InAuthorizationToken,
-		FHttpSharedData* InSharedData,
-		uint32 PoolSize,
-		uint32 InOverflowLimit = 0);
-	~FHttpRequestPool();
+	/** Returns every header in the response. Available once OnHeaders() has been invoked. */
+	[[nodiscard]] virtual TConstArrayView<FAnsiStringView> GetAllHeaders() const = 0;
+
+	/** Returns the value of the first header matching the name, or empty if there was no matching header. */
+	[[nodiscard]] FAnsiStringView GetHeader(FAnsiStringView Name) const;
 
 	/**
-	 * Attempts to get a request is free. Once a request has been returned it is
-	 * "owned by the caller and need to release it to the pool when work has been completed.
-	 * @return Usable request instance if one is available, otherwise null.
+	 * Fills the array of header values and returns the total number of headers matching the name.
+	 *
+	 * The return value can be greater than the array size if there are more matching headers than the array can hold.
 	 */
-	FHttpRequest* GetFreeRequest(bool bUnboundedOverflow = false);
+	[[nodiscard]] int32 GetHeaders(FAnsiStringView Name, TArrayView<FAnsiStringView> OutValues) const;
 
-	class FWaiter : public FThreadSafeRefCountedObject
-	{
-	public:
-		std::atomic<FHttpRequest*> Request = nullptr;
+	/** Returns the content type of the response, or Any if there was no matching header. */
+	[[nodiscard]] EHttpMediaType GetContentType() const;
 
-		FWaiter(FHttpRequestPool* InPool)
-			: Event(FPlatformProcess::GetSynchEventFromPool(true))
-			, Pool(InPool)
-		{
-		}
-		
-		bool Wait(uint32 TimeMS)
-		{
-			return Event->Wait(TimeMS);
-		}
+protected:
+	friend Http::Private::THttpDestroyer<IHttpResponse>;
+	virtual ~IHttpResponse() = default;
 
-		void Trigger()
-		{
-			Event->Trigger();
-		}
+	/** Cancel() if not yet invoking OnComplete(), and then Release(). */
+	virtual void Destroy() = 0;
+};
 
-	private:
-		~FWaiter()
-		{
-			FPlatformProcess::ReturnSynchEventToPool(Event);
+/** Appends the URI, method, status code, and optional error from the response to the builder. */
+FStringBuilderBase& operator<<(FStringBuilderBase& Builder, const IHttpResponse& Response);
 
-			if (Request)
-			{
-				Pool->ReleaseRequestToPool(Request.exchange(nullptr));
-			}
-		}
+struct FHttpResponseStats final
+{
+	/** Total size of data sent to the server, in bytes. */
+	uint64 SendSize = 0;
+	/** Total size of data received from the server, in bytes. */
+	uint64 RecvSize = 0;
+	/** Average rate at which data was sent to the server, in bytes/s. */
+	uint64 SendRate = 0;
+	/** Average rate at which data was received from the server, in bytes/s. */
+	uint64 RecvRate = 0;
+
+	/** Time from the start of the request until the name was resolved, in seconds. */
+	double NameResolveTime = 0.0;
+	/** Time from the start of the request until the connection was established, in seconds. */
+	double ConnectTime = 0.0;
+	/** Time from the start of the request until the TLS connection was complete, in seconds. */
+	double TlsConnectTime = 0.0;
+	/** Time from the start of the request until the first byte was received, in seconds. */
+	double StartTransferTime = 0.0;
+	/** Time from the start of the request until response was complete, in seconds. */
+	double TotalTime = 0.0;
+};
 
-		FEvent* Event;
-		FHttpRequestPool* Pool;
-	};
+/**
+ * Interface to receive the headers and body of the response.
+ *
+ * Functions on this interface must be implemented to do the minimal work possible, because their
+ * execution can block other requests from being processed.
+ */
+class IHttpReceiver
+{
+public:
+	virtual ~IHttpReceiver() = default;
 
 	/**
-	 * Block until a request is free. Once a request has been returned it is 
-	 * "owned by the caller and need to release it to the pool when work has been completed.
-	 * @return Usable request instance.
+	 * Invoked when the response has been created, before it begins executing.
+	 *
+	 * Returns the receiver to use for the next part of the response and can return itself.
 	 */
-	FHttpRequest* WaitForFreeRequest(bool bUnboundedOverflow = false);
+	virtual IHttpReceiver* OnCreate(IHttpResponse& Response) { return this; }
 
 	/**
-	 * Release request to the pool.
-	 * @param Request Request that should be freed. Note that any buffer owened by the request can now be reset.
+	 * Invoked when every header has been received for the response.
+	 *
+	 * Returns the receiver to use for the next part of the response and can return itself.
 	 */
-	void ReleaseRequestToPool(FHttpRequest* Request);
+	virtual IHttpReceiver* OnHeaders(IHttpResponse& Response) { return this; }
 
 	/**
-	 * While holding a request, make it shared across many users.
+	 * Invoked when a part of the body has been received.
+	 *
+	 * The Data parameter can be reassigned to point to the tail of the view that is not consumed
+	 * by this receiver, and will be passed to the receiver that is returned.
+	 *
+	 * Returns the receiver to use for the next part of the response and can return itself.
 	 */
-	void MakeRequestShared(FHttpRequest* Request, uint8 Users);
+	virtual IHttpReceiver* OnBody(IHttpResponse& Response, FMemoryView& Data) { return this; }
 
-private:
-
-	struct FEntry
-	{
-		std::atomic<uint8> Usage;
-		FHttpRequest* Request;
-	};
-
-	struct FInitData
-	{
-		FString ServiceUrl;
-		FString EffectiveServiceUrl;
-		const FHttpAccessToken* AccessToken;
-		FHttpSharedData* SharedData;
-
-		FInitData(
-			FStringView InServiceUrl,
-			FStringView InEffectiveServiceUrl,
-			const FHttpAccessToken* InAccessToken,
-			FHttpSharedData* InSharedData)
-			: ServiceUrl(InServiceUrl)
-			, EffectiveServiceUrl(InEffectiveServiceUrl)
-			, AccessToken(InAccessToken)
-			, SharedData(InSharedData)
-		{
-		}
-	};
-
-	TArray<FEntry> Pool;
-	TArray<FHttpRequest> Requests;
-	FAAArrayQueue<FWaiter> Waiters;
-	std::atomic<uint32> ActiveOverflowRequests;
-	TUniquePtr<const FInitData> InitData;
-	const uint32 OverflowLimit;
+	/**
+	 * Invoked when the response is complete. Always invoked exactly once for each response.
+	 *
+	 * Check the status code to detect errors. Can retry using the same request.
+	 *
+	 * It is valid for implementations of this function to delete the response.
+	 *
+	 * Returns the next receiver to handle OnComplete, or null if there are no others.
+	 */
+	virtual IHttpReceiver* OnComplete(IHttpResponse& Response) { return this; }
 };
 
-/**
- * Utility class to manage requesting and releasing requests from the \ref FHttpRequestPool.
- */
-class FScopedHttpPoolRequestPtr
+///////////////////////////////////////////////////////////////////////////////////////////////////////////////////////
+
+class FHttpByteArrayReceiver final : public IHttpReceiver
 {
 public:
-	UE_NONCOPYABLE(FScopedHttpPoolRequestPtr)
+	explicit FHttpByteArrayReceiver(TArray64<uint8>& OutArray, IHttpReceiver* Next = nullptr);
 
-	explicit FScopedHttpPoolRequestPtr(FHttpRequestPool* InPool)
-		: Request(InPool->WaitForFreeRequest())
-		, Pool(InPool)
-	{
-	}
+	IHttpReceiver* OnBody(IHttpResponse& Response, FMemoryView& Data) final;
+	IHttpReceiver* OnComplete(IHttpResponse& Response) final { return Next; }
 
-	~FScopedHttpPoolRequestPtr()
-	{
-		Pool->ReleaseRequestToPool(Request);
-	}
+private:
+	TArray64<uint8>& Array;
+	IHttpReceiver* Next;
+};
 
-	bool IsValid() const
-	{
-		return Request != nullptr;
-	}
+///////////////////////////////////////////////////////////////////////////////////////////////////////////////////////
 
-	inline explicit operator bool() const { return IsValid(); }
+/** Manages a client by letting callers queue to receive a request when one is not immediately available. */
+class FHttpRequestQueue final
+{
+public:
+	/** Create a null request queue. Asserts if CreateRequest is called. */
+	FHttpRequestQueue() = default;
 
-	FHttpRequest* Get() const
-	{
-		check(IsValid());
-		return Request;
-	}
+	FHttpRequestQueue(IHttpConnectionPool& ConnectionPool, const FHttpClientParams& ClientParams);
 
-	FHttpRequest* operator->()
-	{
-		check(IsValid());
-		return Request;
-	}
+	/** Blocks until a request can be created on the underlying client. */
+	[[nodiscard]] THttpUniquePtr<IHttpRequest> CreateRequest(const FHttpRequestParams& Params);
 
 private:
-	FHttpRequest* Request;
-	FHttpRequestPool* Pool;
+	TPimplPtr<Http::Private::FHttpRequestQueueData> Data;
 };
 
 } // UE
diff --git a/Engine/Source/Developer/DerivedDataCache/Private/Http/HttpRequest.cpp b/Engine/Source/Developer/DerivedDataCache/Private/Http/HttpRequest.cpp
deleted file mode 100644
diff --git a/Engine/Source/Developer/DerivedDataCache/Private/HttpCacheStore.cpp b/Engine/Source/Developer/DerivedDataCache/Private/HttpCacheStore.cpp
--- a/Engine/Source/Developer/DerivedDataCache/Private/HttpCacheStore.cpp
+++ b/Engine/Source/Developer/DerivedDataCache/Private/HttpCacheStore.cpp
@@ -13,11 +13,14 @@
 #include "DerivedDataCacheRecord.h"
 #include "DerivedDataCacheUsageStats.h"
 #include "DerivedDataChunk.h"
+#include "DerivedDataRequest.h"
 #include "DerivedDataRequestOwner.h"
 #include "DerivedDataValue.h"
 #include "DesktopPlatformModule.h"
 #include "Dom/JsonObject.h"
+#include "HAL/CriticalSection.h"
 #include "HAL/IConsoleManager.h"
+#include "HAL/PlatformProcess.h"
 #include "Http/HttpClient.h"
 #include "IO/IoHash.h"
 #include "Memory/SharedBuffer.h"
@@ -25,7 +28,9 @@
 #include "Misc/CommandLine.h"
 #include "Misc/ConfigCacheIni.h"
 #include "Misc/FileHelper.h"
+#include "Misc/ScopeExit.h"
 #include "Misc/ScopeLock.h"
+#include "Misc/ScopeRWLock.h"
 #include "Misc/StringBuilder.h"
 #include "ProfilingDebugging/CountersTrace.h"
 #include "ProfilingDebugging/CpuProfilerTrace.h"
@@ -33,6 +38,9 @@
 #include "Serialization/CompactBinaryPackage.h"
 #include "Serialization/CompactBinaryValidation.h"
 #include "Serialization/CompactBinaryWriter.h"
+#include "Serialization/JsonReader.h"
+#include "Serialization/JsonSerializer.h"
+#include "String/Find.h"
 
 #if PLATFORM_MICROSOFT
 #include "Microsoft/WindowsHWrapper.h"
@@ -55,11 +63,18 @@
 namespace UE::DerivedData
 {
 
+static bool bHttpEnableAsync = true;
+static FAutoConsoleVariableRef CVarHttpEnableAsync(
+	TEXT("DDC.Http.EnableAsync"),
+	bHttpEnableAsync,
+	TEXT("If true, asynchronous operations are permitted, otherwise all operations are forced to be synchronous."),
+	ECVF_Default);
+
 static bool bHttpEnableOidc = false;
 static FAutoConsoleVariableRef CVarHttpEnableOidc(
 	TEXT("DDC.Http.EnableOidc"),
 	bHttpEnableOidc,
-	TEXT("If true, Oidc tokens are used, otherwise legacy shared credentials are used."),
+	TEXT("If true, Oidc tokens are used, otherwise shared credentials are used."),
 	ECVF_Default);
 
 TRACE_DECLARE_INT_COUNTER(HttpDDC_Get, TEXT("HttpDDC Get"));
@@ -69,7 +84,6 @@ TRACE_DECLARE_INT_COUNTER(HttpDDC_PutHit, TEXT("HttpDDC Put Hit"));
 TRACE_DECLARE_INT_COUNTER(HttpDDC_BytesReceived, TEXT("HttpDDC Bytes Received"));
 TRACE_DECLARE_INT_COUNTER(HttpDDC_BytesSent, TEXT("HttpDDC Bytes Sent"));
 
-
 template <typename T>
 class TRefCountedUniqueFunction final : public FThreadSafeRefCountedObject
 {
@@ -84,7 +98,6 @@ private:
 	T Function;
 };
 
-
 static bool ShouldAbortForShutdown()
 {
 	return !GIsBuildMachine && FDerivedDataBackend::Get().IsShuttingDown();
@@ -109,6 +122,106 @@ static bool IsValueDataReady(FValue& Value, const ECachePolicy Policy)
 	return false;
 };
 
+static bool TryResolveCanonicalHost(const FAnsiStringView Uri, FAnsiStringBuilderBase& OutUri)
+{
+	// Extract the host from the URI.
+	FAnsiStringView Domain = Uri;
+	if (const int32 SchemeIndex = String::FindFirst(Domain, ANSITEXTVIEW("://")); SchemeIndex != INDEX_NONE)
+	{
+		Domain.RightChopInline(SchemeIndex + ANSITEXTVIEW("://").Len());
+	}
+	if (const int32 SlashIndex = String::FindFirstChar(Domain, '/'); SlashIndex != INDEX_NONE)
+	{
+		Domain.LeftInline(SlashIndex);
+	}
+	if (const int32 AtIndex = String::FindFirstChar(Domain, '@'); AtIndex != INDEX_NONE)
+	{
+		Domain.RightChopInline(AtIndex + 1);
+	}
+	const auto RemovePort = [](FAnsiStringView& Authority)
+	{
+		if (const int32 ColonIndex = String::FindLastChar(Authority, ':'); ColonIndex != INDEX_NONE)
+		{
+			Authority.LeftInline(ColonIndex);
+		}
+	};
+	if (Domain.StartsWith('['))
+	{
+		if (const int32 LastBracketIndex = String::FindLastChar(Domain, ']'); LastBracketIndex != INDEX_NONE)
+		{
+			Domain.MidInline(1, LastBracketIndex - 1);
+		}
+		else
+		{
+			RemovePort(Domain);
+		}
+	}
+	else
+	{
+		RemovePort(Domain);
+	}
+
+	// Append the URI until the end of the domain.
+	const int32 OutUriIndex = OutUri.Len();
+	const int32 DomainIndex = int32(Domain.GetData() - Uri.GetData());
+	const int32 DomainEndIndex = DomainIndex + Domain.Len();
+	OutUri.Append(Uri.Left(DomainEndIndex));
+
+	// Append the URI beyond the end of the domain before returning.
+	ON_SCOPE_EXIT { OutUri.Append(Uri.RightChop(DomainEndIndex)); };
+
+	// Try to resolve the host.
+	::addrinfo* Result = nullptr;
+	::addrinfo Hints{};
+	Hints.ai_flags = AI_CANONNAME;
+	Hints.ai_family = AF_UNSPEC;
+	if (::getaddrinfo(*OutUri + OutUriIndex + DomainIndex, nullptr, &Hints, &Result) == 0)
+	{
+		ON_SCOPE_EXIT { ::freeaddrinfo(Result); };
+		if (Result->ai_canonname)
+		{
+			OutUri.RemoveSuffix(Domain.Len());
+			OutUri.Append(Result->ai_canonname);
+			return true;
+		}
+	}
+	return false;
+}
+
+/**
+ * Encapsulation for access token shared by all requests.
+ */
+class FHttpAccessToken
+{
+public:
+	void SetToken(FStringView Token);
+	inline uint32 GetSerial() const { return Serial.load(std::memory_order_relaxed); }
+	friend FAnsiStringBuilderBase& operator<<(FAnsiStringBuilderBase& Builder, const FHttpAccessToken& Token);
+
+private:
+	mutable FRWLock Lock;
+	TArray<ANSICHAR> Header;
+	std::atomic<uint32> Serial;
+};
+
+void FHttpAccessToken::SetToken(const FStringView Token)
+{
+	FWriteScopeLock WriteLock(Lock);
+	const FAnsiStringView Prefix = ANSITEXTVIEW("Bearer ");
+	const int32 TokenLen = FPlatformString::ConvertedLength<ANSICHAR>(Token.GetData(), Token.Len());
+	Header.Empty(Prefix.Len() + TokenLen);
+	Header.Append(Prefix.GetData(), Prefix.Len());
+	const int32 TokenIndex = Header.AddUninitialized(TokenLen);
+	FPlatformString::Convert(Header.GetData() + TokenIndex, TokenLen, Token.GetData(), Token.Len());
+	Serial.fetch_add(1, std::memory_order_relaxed);
+}
+
+FAnsiStringBuilderBase& operator<<(FAnsiStringBuilderBase& Builder, const FHttpAccessToken& Token)
+{
+	FReadScopeLock ReadLock(Token.Lock);
+	return Builder.Append(Token.Header);
+}
+
 struct FHttpCacheStoreParams
 {
 	FString Host;
@@ -191,7 +304,6 @@ public:
 
 private:
 	FString Domain;
-	FString EffectiveDomain;
 	FString Namespace;
 	FString StructuredNamespace;
 	FString OAuthProvider;
@@ -201,16 +313,19 @@ private:
 	FString OAuthProviderIdentifier;
 	FString OAuthAccessToken;
 
+	FAnsiStringBuilderBase EffectiveDomain;
+
 	FDerivedDataCacheUsageStats UsageStats;
 	FBackendDebugOptions DebugOptions;
-	TUniquePtr<FHttpSharedData> SharedData;
-	TUniquePtr<FHttpRequestPool> GetRequestPools[2];
-	TUniquePtr<FHttpRequestPool> PutRequestPools[2];
-	TUniquePtr<FHttpRequestPool> NonBlockingRequestPools;
+	THttpUniquePtr<IHttpConnectionPool> ConnectionPool;
+	FHttpRequestQueue GetRequestQueues[2];
+	FHttpRequestQueue PutRequestQueues[2];
+	FHttpRequestQueue NonBlockingRequestQueue;
 
 	FCriticalSection AccessCs;
 	TUniquePtr<FHttpAccessToken> Access;
 	FTSTicker::FDelegateHandle RefreshAccessTokenHandle;
+	double RefreshAccessTokenTime = 0.0;
 	uint32 FailedLoginAttempts = 0;
 
 	bool bIsUsable = false;
@@ -218,38 +333,22 @@ private:
 
 	static inline FHttpCacheStore* AnyInstance = nullptr;
 
-	bool IsServiceReady();
-	bool AcquireAccessToken();
+	FHttpClientParams GetDefaultClientParams() const;
+
+	THttpUniquePtr<IHttpResponse> BeginIsServiceReady(IHttpClient& Client, TArray64<uint8>& Body);
+	bool EndIsServiceReady(THttpUniquePtr<IHttpResponse>& Response, TArray64<uint8>& Body);
+	bool AcquireAccessToken(IHttpClient* Client = nullptr);
 	void SetAccessToken(FStringView Token, double RefreshDelay = 0.0);
-	bool ShouldRetryOnError(FHttpRequest::EResult Result, int64 ResponseCode);
-	bool ShouldRetryOnError(int64 ResponseCode) { return ShouldRetryOnError(FHttpRequest::EResult::Success, ResponseCode); }
 
-	enum class OperationCategory
+	enum class EOperationCategory
 	{
 		Get,
 		Put,
 	};
-	template <OperationCategory Category>
-	FHttpRequest* WaitForHttpRequestForOwner(IRequestOwner& Owner, bool bUnboundedOverflow, FHttpRequestPool*& OutPool)
-	{
-		if (!FHttpRequest::AllowAsync())
-		{
-			if (Category == OperationCategory::Get)
-			{
-				OutPool = GetRequestPools[IsInGameThread()].Get();
-			}
-			else
-			{
-				OutPool = PutRequestPools[IsInGameThread()].Get();
-			}
-			return OutPool->WaitForFreeRequest();
-		}
-		else
-		{
-			OutPool = NonBlockingRequestPools.Get();
-			return OutPool->WaitForFreeRequest(bUnboundedOverflow);
-		}
-	}
+
+	class FHttpOperation;
+
+	TUniquePtr<FHttpOperation> WaitForHttpOperation(EOperationCategory Category, bool bUnboundedOverflow);
 
 	struct FGetCacheRecordOnlyResponse
 	{
@@ -307,10 +406,298 @@ private:
 		TConstArrayView<FCacheGetValueRequest> ValueRefs,
 		FOnCacheGetValueComplete&& OnComplete);
 
+	class FHealthCheckOp;
 	class FPutPackageOp;
 	class FGetRecordOp;
 };
 
+//----------------------------------------------------------------------------------------------------------
+// FHttpCacheStore::FHttpOperation
+//----------------------------------------------------------------------------------------------------------
+class FHttpCacheStore::FHttpOperation final
+{
+public:
+	FHttpOperation(const FHttpOperation&) = delete;
+	FHttpOperation& operator=(const FHttpOperation&) = delete;
+
+	explicit FHttpOperation(THttpUniquePtr<IHttpRequest>&& InRequest)
+		: Request(MoveTemp(InRequest))
+	{
+	}
+
+	// Prepare Request
+
+	void SetUri(FAnsiStringView Uri) { Request->SetUri(Uri); }
+	void SetMethod(EHttpMethod Method) { Request->SetMethod(Method); }
+	void AddHeader(FAnsiStringView Name, FAnsiStringView Value) { Request->AddHeader(Name, Value); }
+	void SetBody(const FCompositeBuffer& Body) { Request->SetBody(Body); }
+	void SetContentType(EHttpMediaType Type) { Request->SetContentType(Type); }
+	void AddAcceptType(EHttpMediaType Type) { Request->AddAcceptType(Type); }
+	void SetExpectedErrorCodes(TConstArrayView<int32> Codes) { ExpectedErrorCodes = Codes; }
+
+	// Send Request
+
+	void Send();
+	void SendAsync(IRequestOwner& Owner, TUniqueFunction<void ()>&& OnComplete);
+
+	// Consume Response
+
+	int32 GetStatusCode() const { return Response->GetStatusCode(); }
+	EHttpErrorCode GetErrorCode() const { return Response->GetErrorCode(); }
+	EHttpMediaType GetContentType() const { return Response->GetContentType(); }
+	FAnsiStringView GetHeader(FAnsiStringView Name) const { return Response->GetHeader(Name); }
+	FSharedBuffer GetBody() const { return ResponseBody; }
+	FString GetBodyAsString() const;
+	TSharedPtr<FJsonObject> GetBodyAsJson() const;
+	uint64 GetBytesSent() const { return Response->GetStats().SendSize; }
+	uint64 GetBytesReceived() const { return Response->GetStats().RecvSize; }
+
+	friend FStringBuilderBase& operator<<(FStringBuilderBase& Builder, const FHttpOperation& Operation)
+	{
+		check(Operation.Response);
+		return Builder << *Operation.Response;
+	}
+
+private:
+	class FHttpOperationReceiver;
+	class FAsyncHttpOperationReceiver;
+
+	FSharedBuffer ResponseBody;
+	THttpUniquePtr<IHttpRequest> Request;
+	THttpUniquePtr<IHttpResponse> Response;
+	TArray<int32, TInlineAllocator<4>> ExpectedErrorCodes;
+	uint32 AttemptCount = 0;
+};
+
+class FHttpCacheStore::FHttpOperation::FHttpOperationReceiver final : public IHttpReceiver
+{
+public:
+	FHttpOperationReceiver(const FHttpOperationReceiver&) = delete;
+	FHttpOperationReceiver& operator=(const FHttpOperationReceiver&) = delete;
+
+	explicit FHttpOperationReceiver(FHttpOperation* InOperation, IHttpReceiver* InNext = nullptr)
+		: Operation(InOperation)
+		, Next(InNext)
+		, BodyReceiver(BodyArray, this)
+	{
+	}
+
+	FHttpOperation* GetOperation() const { return Operation; }
+
+private:
+	IHttpReceiver* OnCreate(IHttpResponse& LocalResponse) final
+	{
+		++Operation->AttemptCount;
+		return &BodyReceiver;
+	}
+
+	IHttpReceiver* OnComplete(IHttpResponse& LocalResponse) final
+	{
+		Operation->ResponseBody = MakeSharedBufferFromArray(MoveTemp(BodyArray));
+
+		LogResponse(LocalResponse);
+
+		if (!ShouldRetry(LocalResponse))
+		{
+			Operation->Request.Reset();
+		}
+
+		return Next;
+	}
+
+	bool ShouldRetry(IHttpResponse& LocalResponse) const
+	{
+		if (Operation->AttemptCount >= UE_HTTPDDC_MAX_ATTEMPTS || ShouldAbortForShutdown())
+		{
+			return false;
+		}
+
+		if (LocalResponse.GetErrorCode() == EHttpErrorCode::TimedOut)
+		{
+			return true;
+		}
+
+		// Too many requests, make a new attempt.
+		if (LocalResponse.GetStatusCode() == 429)
+		{
+			return true;
+		}
+
+		return false;
+	}
+
+	void LogResponse(IHttpResponse& LocalResponse) const
+	{
+		if (UE_LOG_ACTIVE(LogDerivedDataCache, Display))
+		{
+			const int32 StatusCode = LocalResponse.GetStatusCode();
+			const bool bVerbose = (StatusCode >= 200 && StatusCode < 300) || Operation->ExpectedErrorCodes.Contains(StatusCode);
+
+			TStringBuilder<80> StatsText;
+			if (!bVerbose || UE_LOG_ACTIVE(LogDerivedDataCache, Verbose))
+			{
+				const FHttpResponseStats& Stats = LocalResponse.GetStats();
+				if (Stats.SendSize)
+				{
+					StatsText << TEXTVIEW("sent ") << Stats.SendSize << TEXTVIEW(" bytes, ");
+				}
+				if (Stats.RecvSize)
+				{
+					StatsText << TEXTVIEW("received ") << Stats.RecvSize << TEXTVIEW(" bytes, ");
+				}
+				StatsText.Appendf(TEXT("%.3f seconds"), Stats.TotalTime);
+			}
+
+			if (bVerbose)
+			{
+				UE_LOG(LogDerivedDataCache, Verbose, TEXT("HTTP: %s (%s)"), *WriteToString<256>(LocalResponse), *StatsText);
+			}
+			else
+			{
+				FString Body = Operation->GetBodyAsString();
+				Body.ReplaceCharInline(TEXT('\r'), TEXT(' '));
+				Body.ReplaceCharInline(TEXT('\n'), TEXT(' '));
+				UE_LOG(LogDerivedDataCache, Display,
+					TEXT("HTTP: %s (%s) %s"), *WriteToString<256>(LocalResponse), *StatsText, *Body);
+			}
+		}
+	}
+
+private:
+	FHttpOperation* Operation;
+	IHttpReceiver* Next;
+	TArray64<uint8> BodyArray;
+	FHttpByteArrayReceiver BodyReceiver{BodyArray, this};
+};
+
+class FHttpCacheStore::FHttpOperation::FAsyncHttpOperationReceiver final : public FRequestBase, public IHttpReceiver
+{
+public:
+	FAsyncHttpOperationReceiver(const FAsyncHttpOperationReceiver&) = delete;
+	FAsyncHttpOperationReceiver& operator=(const FAsyncHttpOperationReceiver&) = delete;
+
+	FAsyncHttpOperationReceiver(FHttpOperation* InOperation, IRequestOwner* InOwner, TUniqueFunction<void ()>&& InOperationComplete)
+		: Owner(InOwner)
+		, BaseReceiver(InOperation, this)
+		, OperationComplete(MoveTemp(InOperationComplete))
+	{
+	}
+
+private:
+	// IRequest Interface
+
+	void SetPriority(EPriority Priority) final {}
+	void Cancel() final { Monitor->Cancel(); }
+	void Wait() final { Monitor->Wait(); }
+
+	// IHttpReceiver Interface
+
+	IHttpReceiver* OnCreate(IHttpResponse& LocalResponse) final
+	{
+		Monitor = LocalResponse.GetMonitor();
+		Owner->Begin(this);
+		return &BaseReceiver;
+	}
+
+	IHttpReceiver* OnComplete(IHttpResponse& LocalResponse) final
+	{
+		Owner->End(this, [Self = this]
+		{
+			FHttpOperation* Operation = Self->BaseReceiver.GetOperation();
+			if (IHttpRequest* LocalRequest = Operation->Request.Get())
+			{
+				// Retry as indicated by the request not being reset.
+				TRefCountPtr<FAsyncHttpOperationReceiver> Receiver = new FAsyncHttpOperationReceiver(Operation, Self->Owner, MoveTemp(Self->OperationComplete));
+				LocalRequest->SendAsync(Receiver, Operation->Response);
+			}
+			else if (Self->OperationComplete)
+			{
+				// Launch a task for the completion function since it can execute arbitrary code.
+				Self->Owner->LaunchTask(TEXT("HttpOperationComplete"), [Self = TRefCountPtr(Self)]
+				{
+					Self->OperationComplete();
+				});
+			}
+		});
+		return nullptr;
+	}
+
+private:
+	IRequestOwner* Owner;
+	FHttpOperationReceiver BaseReceiver;
+	TUniqueFunction<void ()> OperationComplete;
+	TRefCountPtr<IHttpResponseMonitor> Monitor;
+};
+
+void FHttpCacheStore::FHttpOperation::Send()
+{
+	FHttpOperationReceiver Receiver(this);
+	do
+	{
+		Request->Send(&Receiver, Response);
+	}
+	while (Request);
+}
+
+void FHttpCacheStore::FHttpOperation::SendAsync(IRequestOwner& Owner, TUniqueFunction<void ()>&& OnComplete)
+{
+	TRefCountPtr<FAsyncHttpOperationReceiver> Receiver = new FAsyncHttpOperationReceiver(this, &Owner, MoveTemp(OnComplete));
+	Request->SendAsync(Receiver, Response);
+}
+
+FString FHttpCacheStore::FHttpOperation::GetBodyAsString() const
+{
+	static_assert(sizeof(uint8) == sizeof(UTF8CHAR));
+	const int32 Len = IntCastChecked<int32>(ResponseBody.GetSize());
+	return FString(Len, (const UTF8CHAR*)ResponseBody.GetData());
+}
+
+TSharedPtr<FJsonObject> FHttpCacheStore::FHttpOperation::GetBodyAsJson() const
+{
+	TSharedPtr<FJsonObject> JsonObject;
+	TSharedRef<TJsonReader<>> JsonReader = TJsonReaderFactory<>::Create(GetBodyAsString());
+	FJsonSerializer::Deserialize(JsonReader, JsonObject);
+	return JsonObject;
+}
+
+//----------------------------------------------------------------------------------------------------------
+// FHttpCacheStore::FHealthCheckOp
+//----------------------------------------------------------------------------------------------------------
+class FHttpCacheStore::FHealthCheckOp final
+{
+public:
+	FHealthCheckOp(FHttpCacheStore& CacheStore, IHttpClient& Client)
+		: Operation(Client.TryCreateRequest({}))
+		, Owner(EPriority::High)
+		, Domain(*CacheStore.Domain)
+	{
+		Operation.SetUri(WriteToAnsiString<256>(CacheStore.EffectiveDomain, ANSITEXTVIEW("/health/ready")));
+		Operation.SendAsync(Owner, []{});
+	}
+
+	bool IsReady()
+	{
+		Owner.Wait();
+		const FString Body = Operation.GetBodyAsString();
+		if (Operation.GetStatusCode() == 200)
+		{
+			UE_LOG(LogDerivedDataCache, Display, TEXT("%s: HTTP DDC: %s"), Domain, *Body);
+			return true;
+		}
+		else
+		{
+			UE_LOG(LogDerivedDataCache, Warning, TEXT("%s: Unable to reach HTTP DDC at %s. %s"),
+				Domain, *WriteToString<256>(Operation), *Body);
+			return false;
+		}
+	}
+
+private:
+	FHttpOperation Operation;
+	FRequestOwner Owner;
+	const TCHAR* Domain;
+};
+
 //----------------------------------------------------------------------------------------------------------
 // FHttpCacheStore::FPutPackageOp
 //----------------------------------------------------------------------------------------------------------
@@ -398,12 +785,9 @@ private:
 		FOnCachePutPackageComplete&& OnComplete,
 		FCachePutRefResponse&& Response);
 
-	FHttpRequest::ECompletionBehavior OnCompressedBlobUploadComplete(
-		FHttpRequest::EResult HttpResult,
-		FHttpRequest* Request);
+	void OnCompressedBlobUploadComplete(FHttpOperation& Operation);
 
-	void OnPutRefFinalizationComplete(
-		FCachePutRefResponse&& Response);
+	void OnPutRefFinalizationComplete(FCachePutRefResponse&& Response);
 
 	FCachePutPackageResponse MakeResponse(uint64 InBytesSent, EStatus Status)
 	{
@@ -559,29 +943,35 @@ void FHttpCacheStore::FPutPackageOp::PutRefAsync(
 	FString Bucket(Key.Bucket.ToString());
 	Bucket.ToLowerInline();
 
-	TStringBuilder<256> RefsUri;
-	RefsUri << "api/v1/refs/" << CacheStore.StructuredNamespace << "/" << Bucket << "/" << Key.Hash;
+	TAnsiStringBuilder<256> RefsUri;
+	RefsUri << CacheStore.EffectiveDomain << ANSITEXTVIEW("/api/v1/refs/") << CacheStore.StructuredNamespace << '/' << Bucket << '/' << Key.Hash;
 	if (bFinalize)
 	{
-		RefsUri << "/finalize/" << ObjectHash;
+		RefsUri << ANSITEXTVIEW("/finalize/") << ObjectHash;
 	}
 
-	FHttpRequestPool* Pool = nullptr;
-	FHttpRequest* Request = CacheStore.WaitForHttpRequestForOwner<OperationCategory::Put>(Owner, bFinalize /* bUnboundedOverflow */, Pool);
-
-	auto OnHttpRequestComplete = [&CacheStore, &Owner, Name = FSharedString(Name), Key, Object, UserData, bFinalize, OnComplete = MoveTemp(OnComplete)]
-	(FHttpRequest::EResult HttpResult, FHttpRequest* Request)
+	TUniquePtr<FHttpOperation> Operation = CacheStore.WaitForHttpOperation(EOperationCategory::Put, /*bUnboundedOverflow*/ bFinalize);
+	FHttpOperation& LocalOperation = *Operation;
+	LocalOperation.SetUri(RefsUri);
+	if (bFinalize)
+	{
+		LocalOperation.SetMethod(EHttpMethod::Post);
+		LocalOperation.SetContentType(EHttpMediaType::FormUrlEncoded);
+	}
+	else
+	{
+		LocalOperation.SetMethod(EHttpMethod::Put);
+		LocalOperation.SetContentType(EHttpMediaType::CbObject);
+		LocalOperation.AddHeader(ANSITEXTVIEW("X-Jupiter-IoHash"), WriteToAnsiString<48>(ObjectHash));
+		LocalOperation.SetBody(Object.GetBuffer());
+	}
+	LocalOperation.AddAcceptType(EHttpMediaType::Json);
+	LocalOperation.SendAsync(Owner, [Operation = MoveTemp(Operation), &CacheStore, Name, Key, Object, UserData, bFinalize, OnComplete = MoveTemp(OnComplete)]
 	{
 		TRACE_CPUPROFILER_EVENT_SCOPE(HttpDDC_PutRefAsync_OnHttpRequestComplete);
 
-		if (Owner.IsCanceled())
-		{
-			OnComplete({ Name, Key, UserData, Request->GetBytesSent(), {}, EStatus::Canceled });
-			return FHttpRequest::ECompletionBehavior::Done;
-		}
-
-		int64 ResponseCode = Request->GetResponseCode();
-		if (FHttpRequest::IsSuccessResponse(ResponseCode))
+		const int32 StatusCode = Operation->GetStatusCode();
+		if (StatusCode >= 200 && StatusCode <= 204)
 		{
 			TArray<FIoHash> NeededBlobHashes;
 
@@ -599,7 +989,7 @@ void FHttpCacheStore::FPutPackageOp::PutRefAsync(
 					}
 				});
 			}
-			else if (TSharedPtr<FJsonObject> ResponseObject = Request->GetResponseAsJsonObject())
+			else if (TSharedPtr<FJsonObject> ResponseObject = Operation->GetBodyAsJson())
 			{
 				TArray<FString> NeedsArrayStrings;
 				ResponseObject->TryGetStringArrayField(TEXT("needs"), NeedsArrayStrings);
@@ -616,28 +1006,14 @@ void FHttpCacheStore::FPutPackageOp::PutRefAsync(
 				}
 			}
 
-			OnComplete({ Name, Key, UserData, Request->GetBytesSent(), NeededBlobHashes, EStatus::Ok });
-			return FHttpRequest::ECompletionBehavior::Done;
+			OnComplete({ Name, Key, UserData, Operation->GetBytesSent(), NeededBlobHashes, EStatus::Ok });
 		}
-
-		if (!ShouldAbortForShutdown() && CacheStore.ShouldRetryOnError(HttpResult, ResponseCode) && ((Request->GetAttempts()+1) < UE_HTTPDDC_MAX_ATTEMPTS))
+		else
 		{
-			return FHttpRequest::ECompletionBehavior::Retry;
+			const EStatus Status = Operation->GetErrorCode() == EHttpErrorCode::Canceled ? EStatus::Canceled : EStatus::Error;
+			OnComplete({ Name, Key, UserData, Operation->GetBytesSent(), {}, Status });
 		}
-
-		OnComplete({ Name, Key, UserData, Request->GetBytesSent(), {}, EStatus::Error });
-		return FHttpRequest::ECompletionBehavior::Done;
-	};
-
-	if (bFinalize)
-	{
-		Request->EnqueueAsyncPost(Owner, Pool, *RefsUri, FCompositeBuffer(), MoveTemp(OnHttpRequestComplete), EHttpMediaType::FormUrlEncoded);
-	}
-	else
-	{
-		Request->AddHeader(TEXTVIEW("X-Jupiter-IoHash"), WriteToString<48>(ObjectHash));
-		Request->EnqueueAsyncPut(Owner, Pool, *RefsUri, Object.GetBuffer(), MoveTemp(OnHttpRequestComplete), EHttpMediaType::CbObject);
-	}
+	});
 }
 
 void FHttpCacheStore::FPutPackageOp::OnPackagePutRefComplete(
@@ -741,51 +1117,40 @@ void FHttpCacheStore::FPutPackageOp::OnPackagePutRefComplete(
 	FRequestBarrier Barrier(Owner);
 	for (const FCompressedBlobUpload& CompressedBlobUpload : CompressedBlobUploads)
 	{
-		TStringBuilder<256> CompressedBlobsUri;
-		CompressedBlobsUri << "api/v1/compressed-blobs/" << CacheStore.StructuredNamespace << "/" << CompressedBlobUpload.Hash;
-
-		FHttpRequestPool* Pool = nullptr;
-		FHttpRequest* Request = CacheStore.WaitForHttpRequestForOwner<OperationCategory::Put>(Owner, true /* bUnboundedOverflow */, Pool);
-		Request->EnqueueAsyncPut(Owner, Pool, *CompressedBlobsUri, FCompositeBuffer(CompressedBlobUpload.BlobBuffer),
-			[PutPackageOp](FHttpRequest::EResult HttpResult, FHttpRequest* Request)
-			{
-				return PutPackageOp->OnCompressedBlobUploadComplete(HttpResult, Request);
-			},
-			EHttpMediaType::CompressedBinary);
+		TUniquePtr<FHttpOperation> Operation = CacheStore.WaitForHttpOperation(EOperationCategory::Put, /*bUnboundedOverflow*/ true);
+		FHttpOperation& LocalOperation = *Operation;
+		LocalOperation.SetUri(WriteToAnsiString<256>(CacheStore.EffectiveDomain, ANSITEXTVIEW("/api/v1/compressed-blobs/"), CacheStore.StructuredNamespace, '/', CompressedBlobUpload.Hash));
+		LocalOperation.SetMethod(EHttpMethod::Put);
+		LocalOperation.SetContentType(EHttpMediaType::CompressedBinary);
+		LocalOperation.SetBody(FCompositeBuffer(CompressedBlobUpload.BlobBuffer));
+		LocalOperation.SendAsync(Owner, [Operation = MoveTemp(Operation), PutPackageOp]
+		{
+			PutPackageOp->OnCompressedBlobUploadComplete(*Operation);
+		});
 	}
 }
 
-FHttpRequest::ECompletionBehavior FHttpCacheStore::FPutPackageOp::OnCompressedBlobUploadComplete(
-	FHttpRequest::EResult HttpResult,
-	FHttpRequest* Request)
+void FHttpCacheStore::FPutPackageOp::OnCompressedBlobUploadComplete(FHttpOperation& Operation)
 {
-	int64 ResponseCode = Request->GetResponseCode();
-	bool bIsSuccessResponse = FHttpRequest::IsSuccessResponse(ResponseCode);
+	BytesSent.fetch_add(Operation.GetBytesSent(), std::memory_order_relaxed);
 
-	if (!bIsSuccessResponse && !ShouldAbortForShutdown() && !Owner.IsCanceled() && CacheStore.ShouldRetryOnError(HttpResult, ResponseCode) && ((Request->GetAttempts()+1) < UE_HTTPDDC_MAX_ATTEMPTS))
-	{
-		return FHttpRequest::ECompletionBehavior::Retry;
-	}
-
-	BytesSent.fetch_add(Request->GetBytesSent(), std::memory_order_relaxed);
-	if (bIsSuccessResponse)
+	const int32 StatusCode = Operation.GetStatusCode();
+	if (StatusCode >= 200 && StatusCode <= 204)
 	{
 		SuccessfulBlobUploads.fetch_add(1, std::memory_order_relaxed);
 	}
 
 	if (PendingBlobUploads.fetch_sub(1, std::memory_order_relaxed) == 1)
 	{
+		const uint32 LocalSuccessfulBlobUploads = SuccessfulBlobUploads.load(std::memory_order_relaxed);
 		if (Owner.IsCanceled())
 		{
 			OnComplete(MakeResponse(BytesSent.load(std::memory_order_relaxed), EStatus::Canceled));
-			return FHttpRequest::ECompletionBehavior::Done;
 		}
-
-		uint32 LocalSuccessfulBlobUploads = SuccessfulBlobUploads.load(std::memory_order_relaxed);
-		if (LocalSuccessfulBlobUploads == TotalBlobUploads)
+		else if (LocalSuccessfulBlobUploads == TotalBlobUploads)
 		{
 			// Perform finalization
-			PutRefAsync(CacheStore, Owner, Name, Key, PackageObject, PackageObjectHash, UserData, true,
+			PutRefAsync(CacheStore, Owner, Name, Key, PackageObject, PackageObjectHash, UserData, /*bFinalize*/ true,
 				[PutPackageOp = TRefCountPtr<FPutPackageOp>(this)](FCachePutRefResponse&& Response)
 				{
 					return PutPackageOp->OnPutRefFinalizationComplete(MoveTemp(Response));
@@ -793,17 +1158,15 @@ FHttpRequest::ECompletionBehavior FHttpCacheStore::FPutPackageOp::OnCompressedBl
 		}
 		else
 		{
-			uint32 FailedBlobUploads = (uint32)(TotalBlobUploads - LocalSuccessfulBlobUploads);
+			const uint32 FailedBlobUploads = TotalBlobUploads - LocalSuccessfulBlobUploads;
 			UE_LOG(LogDerivedDataCache, Log, TEXT("%s: Failed to put %d/%d blobs for put of %s from '%s'"),
 				*CacheStore.Domain, FailedBlobUploads, TotalBlobUploads, *WriteToString<96>(Key), *Name);
 			OnComplete(MakeResponse(BytesSent.load(std::memory_order_relaxed), EStatus::Error));
 		}
 	}
-	return FHttpRequest::ECompletionBehavior::Done;
 }
 
-void FHttpCacheStore::FPutPackageOp::OnPutRefFinalizationComplete(
-	FCachePutRefResponse&& Response)
+void FHttpCacheStore::FPutPackageOp::OnPutRefFinalizationComplete(FCachePutRefResponse&& Response)
 {
 	BytesSent.fetch_add(Response.BytesSent, std::memory_order_relaxed);
 
@@ -847,83 +1210,68 @@ void FHttpCacheStore::FGetRecordOp::GetDataBatch(
 	}
 
 	FRequestBarrier Barrier(Owner);
-	TRefCountedUniqueFunction<FOnGetCachedDataBatchComplete>* CompletionFunction = new TRefCountedUniqueFunction<FOnGetCachedDataBatchComplete>(MoveTemp(OnComplete));
-	TRefCountPtr<TRefCountedUniqueFunction<FOnGetCachedDataBatchComplete>> BatchOnCompleteRef(CompletionFunction);
+	TRefCountPtr<TRefCountedUniqueFunction<FOnGetCachedDataBatchComplete>> SharedOnComplete(new TRefCountedUniqueFunction<FOnGetCachedDataBatchComplete>(MoveTemp(OnComplete)));
 	for (int32 ValueIndex = 0; ValueIndex < Values.Num(); ++ValueIndex)
 	{
-		const ValueType& Value = Values[ValueIndex];
-		const FIoHash& RawHash = Value.GetRawHash();
-
-		FHttpRequestPool* Pool = nullptr;
-		FHttpRequest* Request = CacheStore.WaitForHttpRequestForOwner<OperationCategory::Get>(Owner, true /* bUnboundedOverflow */, Pool);
-
-		auto OnHttpRequestComplete = [&CacheStore, &Owner, Name = FSharedString(Name), Key = FCacheKey(Key), ValueIndex, Value = Value.RemoveData(), ValueIdGetter, OnCompletePtr = TRefCountPtr<TRefCountedUniqueFunction<FOnGetCachedDataBatchComplete>>(CompletionFunction)]
-		(FHttpRequest::EResult HttpResult, FHttpRequest* Request)
+		const ValueType Value = Values[ValueIndex].RemoveData();
+		TUniquePtr<FHttpOperation> Operation = CacheStore.WaitForHttpOperation(EOperationCategory::Get, /*bUnboundedOverflow*/ true);
+		FHttpOperation& LocalOperation = *Operation;
+		LocalOperation.SetUri(WriteToAnsiString<256>(CacheStore.EffectiveDomain, ANSITEXTVIEW("/api/v1/compressed-blobs/"), CacheStore.StructuredNamespace, '/', Value.GetRawHash()));
+		LocalOperation.SetMethod(EHttpMethod::Get);
+		LocalOperation.AddAcceptType(EHttpMediaType::Any);
+		LocalOperation.SetExpectedErrorCodes({404});
+		LocalOperation.SendAsync(Owner, [Operation = MoveTemp(Operation), &CacheStore, Name, Key, Value, ValueIndex, ValueIdGetter, SharedOnComplete]
 		{
 			TRACE_CPUPROFILER_EVENT_SCOPE(HttpDDC_GetDataBatch_OnHttpRequestComplete);
 
-			int64 ResponseCode = Request->GetResponseCode();
 			bool bHit = false;
 			FCompressedBuffer CompressedBuffer;
-			if (FHttpRequest::IsSuccessResponse(ResponseCode))
+			if (Operation->GetStatusCode() == 200)
 			{
-				FString ReceivedContentType;
-				if (Request->GetHeader("Content-Type", ReceivedContentType))
+				switch (Operation->GetContentType())
 				{
-					if (ReceivedContentType == TEXT("application/x-ue-comp"))
-					{
-						CompressedBuffer = FCompressedBuffer::FromCompressed(Request->MoveResponseBufferToShared());
-						bHit = true;
-					}
-					else if (ReceivedContentType == TEXT("application/octet-stream"))
-					{
-						CompressedBuffer = FValue::Compress(Request->MoveResponseBufferToShared()).GetData();
-						bHit = true;
-					}
-					else
-					{
-						bHit = false;
-					}
+				case EHttpMediaType::Any:
+				case EHttpMediaType::CompressedBinary:
+					CompressedBuffer = FCompressedBuffer::FromCompressed(Operation->GetBody());
+					bHit = true;
+					break;
+				case EHttpMediaType::Binary:
+					CompressedBuffer = FValue::Compress(Operation->GetBody()).GetData();
+					bHit = true;
+					break;
+				default:
+					break;
+				}
+			}
+
+			if (bHit)
+			{
+				if (CompressedBuffer.GetRawHash() == Value.GetRawHash())
+				{
+					SharedOnComplete->GetFunction()({ Name, Key, ValueIndex, Operation->GetBytesReceived(), MoveTemp(CompressedBuffer), EStatus::Ok });
 				}
 				else
 				{
-					CompressedBuffer = FCompressedBuffer::FromCompressed(Request->MoveResponseBufferToShared());
-					bHit = true;
+					UE_LOG(LogDerivedDataCache, Display,
+						TEXT("%s: Cache miss with corrupted value %s with hash %s for %s from '%s'"),
+						*CacheStore.Domain, *ValueIdGetter(Value), *WriteToString<48>(Value.GetRawHash()),
+						*WriteToString<96>(Key), *Name);
+					SharedOnComplete->GetFunction()({ Name, Key, ValueIndex, Operation->GetBytesReceived(), {}, EStatus::Error });
 				}
 			}
-
-			if (!ShouldAbortForShutdown() && !Owner.IsCanceled() && CacheStore.ShouldRetryOnError(HttpResult, ResponseCode) && ((Request->GetAttempts()+1) < UE_HTTPDDC_MAX_ATTEMPTS))
+			else if (Operation->GetErrorCode() == EHttpErrorCode::Canceled)
 			{
-				return FHttpRequest::ECompletionBehavior::Retry;
+				SharedOnComplete->GetFunction()({ Name, Key, ValueIndex, Operation->GetBytesReceived(), {}, EStatus::Canceled });
 			}
-
-			if (!bHit)
+			else
 			{
 				UE_LOG(LogDerivedDataCache, Verbose,
 					TEXT("%s: Cache miss with missing value %s with hash %s for %s from '%s'"),
 					*CacheStore.Domain, *ValueIdGetter(Value), *WriteToString<48>(Value.GetRawHash()), *WriteToString<96>(Key),
 					*Name);
-				OnCompletePtr->GetFunction()({ Name, Key, ValueIndex, Request->GetBytesReceived(), {}, EStatus::Error });
-			}
-			else if (CompressedBuffer.GetRawHash() != Value.GetRawHash())
-			{
-				UE_LOG(LogDerivedDataCache, Display,
-					TEXT("%s: Cache miss with corrupted value %s with hash %s for %s from '%s'"),
-					*CacheStore.Domain, *ValueIdGetter(Value), *WriteToString<48>(Value.GetRawHash()),
-					*WriteToString<96>(Key), *Name);
-				OnCompletePtr->GetFunction()({ Name, Key, ValueIndex, Request->GetBytesReceived(), {}, EStatus::Error });
+				SharedOnComplete->GetFunction()({ Name, Key, ValueIndex, Operation->GetBytesReceived(), {}, EStatus::Error });
 			}
-			else
-			{
-				OnCompletePtr->GetFunction()({ Name, Key, ValueIndex, Request->GetBytesReceived(), MoveTemp(CompressedBuffer), EStatus::Ok });
-			}
-
-			return FHttpRequest::ECompletionBehavior::Done;
-		};
-
-		TStringBuilder<256> CompressedBlobsUri;
-		CompressedBlobsUri << "api/v1/compressed-blobs/" << CacheStore.StructuredNamespace << "/" << RawHash;
-		Request->EnqueueAsyncDownload(Owner, Pool, *CompressedBlobsUri, MoveTemp(OnHttpRequestComplete), EHttpMediaType::Any, { 404 });
+		});
 	}
 }
 
@@ -1048,30 +1396,33 @@ void FHttpCacheStore::FGetRecordOp::DataProbablyExistsBatch(
 		return;
 	}
 
-	FHttpRequestPool* Pool = nullptr;
-	FHttpRequest* Request = CacheStore.WaitForHttpRequestForOwner<OperationCategory::Get>(Owner, true /* bUnboundedOverflow */, Pool);
-
-	TStringBuilder<256> CompressedBlobsUri;
-	CompressedBlobsUri << "api/v1/compressed-blobs/" << CacheStore.StructuredNamespace << "/exists?";
+	TAnsiStringBuilder<256> CompressedBlobsUri;
+	CompressedBlobsUri << CacheStore.EffectiveDomain << ANSITEXTVIEW("/api/v1/compressed-blobs/") << CacheStore.StructuredNamespace << ANSITEXTVIEW("/exists?");
 	bool bFirstItem = true;
 	for (const FValueWithId& Value : Values)
 	{
 		if (!bFirstItem)
 		{
-			CompressedBlobsUri << "&";
+			CompressedBlobsUri << '&';
 		}
-		CompressedBlobsUri << "id=" << Value.GetRawHash();
+		CompressedBlobsUri << ANSITEXTVIEW("id=") << Value.GetRawHash();
 		bFirstItem = false;
 	}
 
-	auto OnHttpRequestComplete = [this, Values = TArray<FValueWithId>(Values), InOnComplete = MoveTemp(InOnComplete)](FHttpRequest::EResult HttpResult, FHttpRequest* Request)
+	TUniquePtr<FHttpOperation> Operation = CacheStore.WaitForHttpOperation(EOperationCategory::Get, /*bUnboundedOverflow*/ true);
+	FHttpOperation& LocalOperation = *Operation;
+	LocalOperation.SetUri(CompressedBlobsUri);
+	LocalOperation.SetMethod(EHttpMethod::Post);
+	LocalOperation.SetContentType(EHttpMediaType::FormUrlEncoded);
+	LocalOperation.AddAcceptType(EHttpMediaType::Json);
+	LocalOperation.SendAsync(Owner, [Operation = MoveTemp(Operation), this, Values = TArray<FValueWithId>(Values), InOnComplete = MoveTemp(InOnComplete)]
 	{
 		TRACE_CPUPROFILER_EVENT_SCOPE(HttpDDC_DataProbablyExistsBatch_OnHttpRequestComplete);
 
-		int64 ResponseCode = Request->GetResponseCode();
-		if (FHttpRequest::IsSuccessResponse(ResponseCode))
+		const int32 StatusCode = Operation->GetStatusCode();
+		if (StatusCode >= 200 && StatusCode <= 204)
 		{
-			if (TSharedPtr<FJsonObject> ResponseObject = Request->GetResponseAsJsonObject())
+			if (TSharedPtr<FJsonObject> ResponseObject = Operation->GetBodyAsJson())
 			{
 				TArray<FString> NeedsArrayStrings;
 				if (ResponseObject->TryGetStringArrayField(TEXT("needs"), NeedsArrayStrings))
@@ -1082,12 +1433,12 @@ void FHttpCacheStore::FGetRecordOp::DataProbablyExistsBatch(
 						{
 							const FValueWithId& Value = Values[ValueIndex];
 							UE_LOG(LogDerivedDataCache, VeryVerbose,
-								TEXT("%s: Cache exists hit for %s with hash %s for %s from '%s'"),
-								*CacheStore.Domain, *WriteToString<16>(Value.GetId()), *WriteToString<48>(Value.GetRawHash()), *WriteToString<96>(Key),
-								*Name);
+								TEXT("%s: Cache exists hit for value %s with hash %s for %s from '%s'"),
+								*CacheStore.Domain, *WriteToString<16>(Value.GetId()),
+								*WriteToString<48>(Value.GetRawHash()), *WriteToString<96>(Key), *Name);
 							InOnComplete({ Name, Key, ValueIndex, EStatus::Ok });
 						}
-						return FHttpRequest::ECompletionBehavior::Done;
+						return;
 					}
 				}
 
@@ -1113,17 +1464,17 @@ void FHttpCacheStore::FGetRecordOp::DataProbablyExistsBatch(
 					if (ResultStatus[ValueIndex])
 					{
 						UE_LOG(LogDerivedDataCache, VeryVerbose,
-							TEXT("%s: Cache exists hit for %s with hash %s for %s from '%s'"),
-							*CacheStore.Domain, *WriteToString<16>(Value.GetId()), *WriteToString<48>(Value.GetRawHash()), *WriteToString<96>(Key),
-							*Name);
+							TEXT("%s: Cache exists hit for value %s with hash %s for %s from '%s'"),
+							*CacheStore.Domain, *WriteToString<32>(Value.GetId()),
+							*WriteToString<48>(Value.GetRawHash()), *WriteToString<96>(Key), *Name);
 						InOnComplete({ Name, Key, ValueIndex, EStatus::Ok });
 					}
 					else
 					{
 						UE_LOG(LogDerivedDataCache, Verbose,
 							TEXT("%s: Cache exists miss with missing value %s with hash %s for %s from '%s'"),
-							*CacheStore.Domain, *WriteToString<16>(Value.GetId()), *WriteToString<48>(Value.GetRawHash()), *WriteToString<96>(Key),
-							*Name);
+							*CacheStore.Domain, *WriteToString<32>(Value.GetId()),
+							*WriteToString<48>(Value.GetRawHash()), *WriteToString<96>(Key), *Name);
 						InOnComplete({ Name, Key, ValueIndex, EStatus::Error });
 					}
 				}
@@ -1133,31 +1484,25 @@ void FHttpCacheStore::FGetRecordOp::DataProbablyExistsBatch(
 				for (int32 ValueIndex = 0; ValueIndex < Values.Num(); ++ValueIndex)
 				{
 					UE_LOG(LogDerivedDataCache, Log,
-						TEXT("%s: Cache exists returned invalid results."),
-						*CacheStore.Domain);
+						TEXT("%s: Cache exists miss with invalid response for value %s for %s from '%s'"),
+						*CacheStore.Domain, *WriteToString<32>(Values[ValueIndex].GetId()),
+						*WriteToString<96>(Key), *Name);
 					InOnComplete({ Name, Key, ValueIndex, EStatus::Error });
 				}
 			}
-
-			return FHttpRequest::ECompletionBehavior::Done;
-		}
-
-		if (!ShouldAbortForShutdown() && !Owner.IsCanceled() && CacheStore.ShouldRetryOnError(HttpResult, ResponseCode) && ((Request->GetAttempts()+1) < UE_HTTPDDC_MAX_ATTEMPTS))
-		{
-			return FHttpRequest::ECompletionBehavior::Retry;
 		}
-
-		for (int32 ValueIndex = 0; ValueIndex < Values.Num(); ++ValueIndex)
+		else
 		{
-			const FValueWithId& Value = Values[ValueIndex];
-			UE_LOG(LogDerivedDataCache, Verbose, TEXT("%s: Cache miss with failed HTTP request for %s from '%s'"),
-				*CacheStore.Domain, *WriteToString<96>(Key), *Name);
-			InOnComplete({Name, Key, ValueIndex, EStatus::Error});
+			for (int32 ValueIndex = 0; ValueIndex < Values.Num(); ++ValueIndex)
+			{
+				UE_LOG(LogDerivedDataCache, Verbose,
+					TEXT("%s: Cache exists miss with failed response for value %s for %s from '%s'"),
+					*CacheStore.Domain, *WriteToString<32>(Values[ValueIndex].GetId()),
+					*WriteToString<96>(Key), *Name);
+				InOnComplete({Name, Key, ValueIndex, EStatus::Error});
+			}
 		}
-		return FHttpRequest::ECompletionBehavior::Done;
-	};
-	FCompositeBuffer DummyBuffer;
-	Request->EnqueueAsyncPost(Owner, Pool, *CompressedBlobsUri, DummyBuffer, MoveTemp(OnHttpRequestComplete), EHttpMediaType::FormUrlEncoded);
+	});
 }
 
 void FHttpCacheStore::FGetRecordOp::FinishDataStep(bool bSuccess, uint64 InBytesReceived)
@@ -1191,7 +1536,6 @@ void FHttpCacheStore::FGetRecordOp::FinishDataStep(bool bSuccess, uint64 InBytes
 
 FHttpCacheStore::FHttpCacheStore(const FHttpCacheStoreParams& Params)
 	: Domain(Params.Host)
-	, EffectiveDomain(Params.Host)
 	, Namespace(Params.Namespace)
 	, StructuredNamespace(Params.StructuredNamespace)
 	, OAuthProvider(Params.OAuthProvider)
@@ -1202,61 +1546,45 @@ FHttpCacheStore::FHttpCacheStore(const FHttpCacheStoreParams& Params)
 	, OAuthAccessToken(Params.OAuthAccessToken)
 	, bReadOnly(Params.bReadOnly)
 {
-	SharedData = MakeUnique<FHttpSharedData>();
-	if (IsServiceReady() && AcquireAccessToken())
+	TRACE_CPUPROFILER_EVENT_SCOPE(HttpDDC_Construct);
+
+	EffectiveDomain.Append(Domain);
+	TAnsiStringBuilder<256> ResolvedDomain;
+	if (Params.bResolveHostCanonicalName && TryResolveCanonicalHost(EffectiveDomain, ResolvedDomain))
 	{
-		FString OriginalDomainPrefix;
-		TAnsiStringBuilder<64> DomainResolveName;
+		// Store the URI with the canonical name to pin to one region when using DNS-based region selection.
+		UE_LOG(LogDerivedDataCache, Display,
+			TEXT("%s: Pinned to %hs based on DNS canonical name."), *Domain, *ResolvedDomain);
+		EffectiveDomain.Reset();
+		EffectiveDomain.Append(ResolvedDomain);
+	}
 
-		if (Domain.StartsWith(TEXT("http://")))
-		{
-			DomainResolveName << Domain.RightChop(7);
-			OriginalDomainPrefix = TEXT("http://");
-		}
-		else if (Domain.StartsWith(TEXT("https://")))
-		{
-			DomainResolveName << Domain.RightChop(8);
-			OriginalDomainPrefix = TEXT("https://");
-		}
-		else
-		{
-			DomainResolveName << Domain;
-		}
+	constexpr uint32 MaxTotalConnections = 8;
+	FHttpConnectionPoolParams ConnectionPoolParams;
+	ConnectionPoolParams.MaxConnections = MaxTotalConnections;
+	ConnectionPoolParams.MinConnections = MaxTotalConnections;
+	ConnectionPool = IHttpManager::Get().CreateConnectionPool(ConnectionPoolParams);
 
-		addrinfo* AddrResult = nullptr;
-		addrinfo AddrHints;
-		FMemory::Memset(&AddrHints, 0, sizeof(AddrHints));
-		AddrHints.ai_flags = AI_CANONNAME;
-		AddrHints.ai_family = AF_UNSPEC;
-		if (Params.bResolveHostCanonicalName && !::getaddrinfo(*DomainResolveName, nullptr, &AddrHints, &AddrResult))
-		{
-			if (AddrResult->ai_canonname)
-			{
-				// Swap the domain with a canonical name from DNS so that if we are using regional redirection, we pin to a region.
-				EffectiveDomain = OriginalDomainPrefix + ANSI_TO_TCHAR(AddrResult->ai_canonname);
+	FHttpClientParams ClientParams = GetDefaultClientParams();
 
-				UE_LOG(LogDerivedDataCache, Display,
-					TEXT("%s: Pinned to %s based on DNS canonical name."),
-					*Domain, *EffectiveDomain);
-			}
-			else
-			{
-				EffectiveDomain = Domain;
-			}
+	THttpUniquePtr<IHttpClient> Client = ConnectionPool->CreateClient(ClientParams);
+	FHealthCheckOp HealthCheck(*this, *Client);
+	if (AcquireAccessToken(Client.Get()) && HealthCheck.IsReady())
+	{
+		ClientParams.MaxRequests = UE_HTTPDDC_GET_REQUEST_POOL_SIZE;
+		ClientParams.MinRequests = UE_HTTPDDC_GET_REQUEST_POOL_SIZE;
+		GetRequestQueues[0] = FHttpRequestQueue(*ConnectionPool, ClientParams);
+		GetRequestQueues[1] = FHttpRequestQueue(*ConnectionPool, ClientParams);
 
-			::freeaddrinfo(AddrResult);
-		}
-		else
-		{
-			EffectiveDomain = Domain;
-		}
+		ClientParams.MaxRequests = UE_HTTPDDC_PUT_REQUEST_POOL_SIZE;
+		ClientParams.MinRequests = UE_HTTPDDC_PUT_REQUEST_POOL_SIZE;
+		PutRequestQueues[0] = FHttpRequestQueue(*ConnectionPool, ClientParams);
+		PutRequestQueues[1] = FHttpRequestQueue(*ConnectionPool, ClientParams);
+
+		ClientParams.MaxRequests = UE_HTTPDDC_NONBLOCKING_REQUEST_POOL_SIZE * 2;
+		ClientParams.MinRequests = UE_HTTPDDC_NONBLOCKING_REQUEST_POOL_SIZE;
+		NonBlockingRequestQueue = FHttpRequestQueue(*ConnectionPool, ClientParams);
 
-		GetRequestPools[0] = MakeUnique<FHttpRequestPool>(*Domain, *EffectiveDomain, Access.Get(), SharedData.Get(), UE_HTTPDDC_GET_REQUEST_POOL_SIZE);
-		GetRequestPools[1] = MakeUnique<FHttpRequestPool>(*Domain, *EffectiveDomain, Access.Get(), SharedData.Get(), UE_HTTPDDC_GET_REQUEST_POOL_SIZE);
-		PutRequestPools[0] = MakeUnique<FHttpRequestPool>(*Domain, *EffectiveDomain, Access.Get(), SharedData.Get(), UE_HTTPDDC_PUT_REQUEST_POOL_SIZE);
-		PutRequestPools[1] = MakeUnique<FHttpRequestPool>(*Domain, *EffectiveDomain, Access.Get(), SharedData.Get(), UE_HTTPDDC_PUT_REQUEST_POOL_SIZE);
-		// Allowing the non-blocking requests to overflow to double their pre-allocated size before we start waiting for one to free up.
-		NonBlockingRequestPools = MakeUnique<FHttpRequestPool>(*Domain, *EffectiveDomain, Access.Get(), SharedData.Get(), UE_HTTPDDC_NONBLOCKING_REQUEST_POOL_SIZE, UE_HTTPDDC_NONBLOCKING_REQUEST_POOL_SIZE);
 		bIsUsable = true;
 	}
 
@@ -1276,31 +1604,26 @@ FHttpCacheStore::~FHttpCacheStore()
 	}
 }
 
-bool FHttpCacheStore::LegacyDebugOptions(FBackendDebugOptions& InOptions)
+FHttpClientParams FHttpCacheStore::GetDefaultClientParams() const
 {
-	DebugOptions = InOptions;
-	return true;
+	FHttpClientParams ClientParams;
+	ClientParams.DnsCacheTimeout = 300;
+	ClientParams.ConnectTimeout = 30 * 1000;
+	ClientParams.LowSpeedLimit = 1024;
+	ClientParams.LowSpeedTime = 30;
+	ClientParams.TlsLevel = EHttpTlsLevel::All;
+	ClientParams.bFollowRedirects = true;
+	ClientParams.bFollow302Post = true;
+	return ClientParams;
 }
 
-bool FHttpCacheStore::IsServiceReady()
+bool FHttpCacheStore::LegacyDebugOptions(FBackendDebugOptions& InOptions)
 {
-	FHttpRequest Request(*Domain, *Domain, nullptr, SharedData.Get(), false);
-	FHttpRequest::EResult Result = Request.PerformBlockingDownload(TEXT("health/ready"), nullptr);
-	
-	if (Result == FHttpRequest::EResult::Success && Request.GetResponseCode() == 200)
-	{
-		UE_LOG(LogDerivedDataCache, Display, TEXT("%s: HTTP DDC service status: %s."), *Request.GetName(), *Request.GetResponseAsString());
-		return true;
-	}
-	else
-	{
-		UE_LOG(LogDerivedDataCache, Warning, TEXT("%s: Unable to reach HTTP DDC service at %s. Status: %d . Response: %s"), *Request.GetName(), *Domain, Request.GetResponseCode(), *Request.GetResponseAsString());
-	}
-
-	return false;
+	DebugOptions = InOptions;
+	return true;
 }
 
-bool FHttpCacheStore::AcquireAccessToken()
+bool FHttpCacheStore::AcquireAccessToken(IHttpClient* Client)
 {
 	if (Domain.StartsWith(TEXT("http://localhost")))
 	{
@@ -1317,6 +1640,8 @@ bool FHttpCacheStore::AcquireAccessToken()
 	ensureMsgf(OAuthProvider.StartsWith(TEXT("http://")) || OAuthProvider.StartsWith(TEXT("https://")),
 		TEXT("%s: OAuth provider %s must be a complete URI including the scheme."), *Domain, *OAuthProvider);
 
+	TRACE_CPUPROFILER_EVENT_SCOPE(HttpDDC_AcquireAccessToken);
+
 	// In case many requests wants to update the token at the same time
 	// get the current serial while we wait to take the CS.
 	const uint32 WantsToUpdateTokenSerial = Access ? Access->GetSerial() : 0;
@@ -1354,27 +1679,32 @@ bool FHttpCacheStore::AcquireAccessToken()
 		return false;
 	}
 
-	const uint32 SchemeEnd = OAuthProvider.Find(TEXT("://")) + 3;
-	const uint32 DomainEnd = OAuthProvider.Find(TEXT("/"), ESearchCase::CaseSensitive, ESearchDir::FromStart, SchemeEnd);
-	FString AuthDomain(DomainEnd, *OAuthProvider);
-	FString Uri(*OAuthProvider + DomainEnd + 1);
+	THttpUniquePtr<IHttpClient> LocalClient;
+	if (!Client)
+	{
+		LocalClient = ConnectionPool->CreateClient(GetDefaultClientParams());
+		Client = LocalClient.Get();
+	}
+
+	FHttpRequestParams RequestParams;
+	RequestParams.bIgnoreMaxRequests = true;
+	FHttpOperation Operation(Client->TryCreateRequest(RequestParams));
+	Operation.SetUri(StringCast<ANSICHAR>(*OAuthProvider));
 
-	FHttpRequest Request(*AuthDomain, *AuthDomain, nullptr, SharedData.Get(), false);
-	FHttpRequest::EResult Result = FHttpRequest::EResult::Success;
 	if (OAuthProvider.StartsWith(TEXT("http://localhost")))
 	{
 		// Simple unauthenticated call to a local endpoint that mimics the result from an OIDC provider.
-		Result = Request.PerformBlockingDownload(*Uri, nullptr);
+		Operation.Send();
 	}
 	else
 	{
-		// Needs client id and secret to authenticate with an actual OIDC provider.
+		// Needs client id and secret to authenticate with the OIDC provider.
 
 		// If contents of the secret string is a file path, resolve and read form data.
 		if (OAuthSecret.StartsWith(TEXT("file://")))
 		{
 			FString FilePath = OAuthSecret.Mid(7, OAuthSecret.Len() - 7);
-				FString SecretFileContents;
+			FString SecretFileContents;
 			if (FFileHelper::LoadFileToString(SecretFileContents, *FilePath))
 			{
 				// Overwrite the filepath with the actual content.
@@ -1394,12 +1724,15 @@ bool FHttpCacheStore::AcquireAccessToken()
 			<< ANSITEXTVIEW("&grant_type=client_credentials")
 			<< ANSITEXTVIEW("&client_secret=") << OAuthSecret;
 
-		Result = Request.PerformBlockingPost(*Uri, FCompositeBuffer(FSharedBuffer::MakeView(MakeMemoryView(OAuthFormData))), EHttpMediaType::FormUrlEncoded);
+		Operation.SetMethod(EHttpMethod::Post);
+		Operation.SetContentType(EHttpMediaType::FormUrlEncoded);
+		Operation.SetBody(FCompositeBuffer(FSharedBuffer::MakeView(MakeMemoryView(OAuthFormData))));
+		Operation.Send();
 	}
 
-	if (Result == FHttpRequest::EResult::Success && Request.GetResponseCode() == 200)
+	if (Operation.GetStatusCode() == 200)
 	{
-		if (TSharedPtr<FJsonObject> ResponseObject = Request.GetResponseAsJsonObject())
+		if (TSharedPtr<FJsonObject> ResponseObject = Operation.GetBodyAsJson())
 		{
 			FString AccessTokenString;
 			double ExpiryTimeSeconds = 0.0;
@@ -1414,7 +1747,7 @@ bool FHttpCacheStore::AcquireAccessToken()
 		}
 	}
 
-	UE_LOG(LogDerivedDataCache, Warning, TEXT("%s: Failed to log in to HTTP services. Server responded with code %d."), *Domain, Request.GetResponseCode());
+	UE_LOG(LogDerivedDataCache, Warning, TEXT("%s: Failed to log in to HTTP services with request %s."), *Domain, *WriteToString<256>(Operation));
 	FailedLoginAttempts++;
 	return false;
 }
@@ -1433,43 +1766,67 @@ void FHttpCacheStore::SetAccessToken(FStringView Token, double RefreshDelay)
 	}
 	Access->SetToken(Token);
 
-	// Schedule a refresh of the token ahead of expiry time (this will not work in commandlets)
-	constexpr float RefreshGracePeriod = 20.0f;
-	if (RefreshDelay > RefreshGracePeriod && !IsRunningCommandlet())
+	constexpr double RefreshGracePeriod = 20.0f;
+	if (RefreshDelay > RefreshGracePeriod)
 	{
-		RefreshAccessTokenHandle = FTSTicker::GetCoreTicker().AddTicker(FTickerDelegate::CreateLambda(
-			[this](float DeltaTime)
-			{
-				AcquireAccessToken();
-				return false;
-			}
-		), RefreshDelay - RefreshGracePeriod);
+		// Schedule a refresh of the token ahead of expiry time (this will not work in commandlets)
+		if (!IsRunningCommandlet())
+		{
+			RefreshAccessTokenHandle = FTSTicker::GetCoreTicker().AddTicker(FTickerDelegate::CreateLambda(
+				[this](float DeltaTime)
+				{
+					AcquireAccessToken();
+					return false;
+				}
+			), RefreshDelay - RefreshGracePeriod);
+		}
+
+		// Schedule a forced refresh of the token when the scheduled refresh is starved or unavailable.
+		RefreshAccessTokenTime = FPlatformTime::Seconds() + RefreshDelay - RefreshGracePeriod * 0.5f;
+	}
+	else
+	{
+		RefreshAccessTokenTime = 0.0;
 	}
 
 	// Reset failed login attempts, the service is indeed alive.
 	FailedLoginAttempts = 0;
 }
 
-bool FHttpCacheStore::ShouldRetryOnError(FHttpRequest::EResult Result, int64 ResponseCode)
+TUniquePtr<FHttpCacheStore::FHttpOperation> FHttpCacheStore::WaitForHttpOperation(EOperationCategory Category, bool bUnboundedOverflow)
 {
-	if (Result == FHttpRequest::EResult::FailedTimeout)
+	if (Access && RefreshAccessTokenTime && RefreshAccessTokenTime < FPlatformTime::Seconds())
 	{
-		return true;
+		AcquireAccessToken();
 	}
 
-	// Access token might have expired, request a new token and try again.
-	if (ResponseCode == 401 && AcquireAccessToken())
+	THttpUniquePtr<IHttpRequest> Request;
+
+	FHttpRequestParams Params;
+	if (FPlatformProcess::SupportsMultithreading() && bHttpEnableAsync)
 	{
-		return true;
+		Params.bIgnoreMaxRequests = bUnboundedOverflow;
+		Request = NonBlockingRequestQueue.CreateRequest(Params);
+	}
+	else
+	{
+		const bool bIsInGameThread = IsInGameThread();
+		if (Category == EOperationCategory::Get)
+		{
+			Request = GetRequestQueues[bIsInGameThread].CreateRequest(Params);
+		}
+		else
+		{
+			Request = PutRequestQueues[bIsInGameThread].CreateRequest(Params);
+		}
 	}
 
-	// Too many requests, make a new attempt
-	if (ResponseCode == 429)
+	if (Access)
 	{
-		return true;
+		Request->AddHeader(ANSITEXTVIEW("Authorization"), WriteToAnsiString<1024>(*Access));
 	}
 
-	return false;
+	return MakeUnique<FHttpOperation>(MoveTemp(Request));
 }
 
 void FHttpCacheStore::GetCacheRecordOnlyAsync(
@@ -1511,54 +1868,45 @@ void FHttpCacheStore::GetCacheRecordOnlyAsync(
 	FString Bucket(Key.Bucket.ToString());
 	Bucket.ToLowerInline();
 
-	TStringBuilder<256> RefsUri;
-	RefsUri << "api/v1/refs/" << StructuredNamespace << "/" << Bucket << "/" << Key.Hash;
-
-	FHttpRequestPool* Pool = nullptr;
-	FHttpRequest* Request = WaitForHttpRequestForOwner<OperationCategory::Get>(Owner, false /* bUnboundedOverflow */, Pool);
-
-	auto OnHttpRequestComplete = [this, &Owner, Name = FSharedString(Name), Key, UserData, OnComplete = MoveTemp(OnComplete)]
-	(FHttpRequest::EResult HttpResult, FHttpRequest* Request)
+	TUniquePtr<FHttpOperation> Operation = WaitForHttpOperation(EOperationCategory::Get, /*bUnboundedOverflow*/ false);
+	FHttpOperation& LocalOperation = *Operation;
+	LocalOperation.SetUri(WriteToAnsiString<256>(EffectiveDomain, ANSITEXTVIEW("/api/v1/refs/"), StructuredNamespace, '/', Bucket, '/', Key.Hash));
+	LocalOperation.SetMethod(EHttpMethod::Get);
+	LocalOperation.AddAcceptType(EHttpMediaType::CbObject);
+	LocalOperation.SetExpectedErrorCodes({404});
+	LocalOperation.SendAsync(Owner, [Operation = MoveTemp(Operation), this, Name, Key, UserData, OnComplete = MoveTemp(OnComplete)]
 	{
 		TRACE_CPUPROFILER_EVENT_SCOPE(HttpDDC_GetCacheRecordOnlyAsync_OnHttpRequestComplete);
 
-		int64 ResponseCode = Request->GetResponseCode();
-		if (FHttpRequest::IsSuccessResponse(ResponseCode))
+		const int32 StatusCode = Operation->GetStatusCode();
+		if (StatusCode >= 200 && StatusCode <= 204)
 		{
-			FSharedBuffer ResponseBuffer = Request->MoveResponseBufferToShared();
+			FSharedBuffer Body = Operation->GetBody();
 
-			if (ValidateCompactBinary(ResponseBuffer, ECbValidateMode::Default) != ECbValidateError::None)
+			if (ValidateCompactBinary(Body, ECbValidateMode::Default) != ECbValidateError::None)
 			{
 				UE_LOG(LogDerivedDataCache, Log, TEXT("%s: Cache miss with invalid package for %s from '%s'"),
 					*Domain, *WriteToString<96>(Key), *Name);
-				OnComplete({ Name, Key, UserData, Request->GetBytesReceived(), {}, EStatus::Error });
-				return FHttpRequest::ECompletionBehavior::Done;
+				OnComplete({ Name, Key, UserData, Operation->GetBytesReceived(), {}, EStatus::Error });
 			}
-
-			FOptionalCacheRecord Record = FCacheRecord::Load(FCbPackage(FCbObject(ResponseBuffer)));
-			if (Record.IsNull())
+			else if (FOptionalCacheRecord Record = FCacheRecord::Load(FCbPackage(FCbObject(Body))); Record.IsNull())
 			{
 				UE_LOG(LogDerivedDataCache, Log, TEXT("%s: Cache miss with record load failure for %s from '%s'"),
 					*Domain, *WriteToString<96>(Key), *Name);
-				OnComplete({ Name, Key, UserData, Request->GetBytesReceived(), {}, EStatus::Error });
-				return FHttpRequest::ECompletionBehavior::Done;
+				OnComplete({ Name, Key, UserData, Operation->GetBytesReceived(), {}, EStatus::Error });
+			}
+			else
+			{
+				OnComplete({ Name, Key, UserData, Operation->GetBytesReceived(), MoveTemp(Record), EStatus::Ok });
 			}
-
-			OnComplete({ Name, Key, UserData, Request->GetBytesReceived(), MoveTemp(Record), EStatus::Ok });
-			return FHttpRequest::ECompletionBehavior::Done;
 		}
-
-		if (!ShouldAbortForShutdown() && !Owner.IsCanceled() && ShouldRetryOnError(HttpResult, ResponseCode) && ((Request->GetAttempts()+1) < UE_HTTPDDC_MAX_ATTEMPTS))
+		else
 		{
-			return FHttpRequest::ECompletionBehavior::Retry;
+			UE_LOG(LogDerivedDataCache, Verbose, TEXT("%s: Cache miss with missing package for %s from '%s'"),
+				*Domain, *WriteToString<96>(Key), *Name);
+			OnComplete({ Name, Key, UserData, Operation->GetBytesReceived(), {}, EStatus::Error });
 		}
-
-		UE_LOG(LogDerivedDataCache, Verbose, TEXT("%s: Cache miss with missing package for %s from '%s'"),
-			*Domain, *WriteToString<96>(Key), *Name);
-		OnComplete({ Name, Key, UserData, Request->GetBytesReceived(), {}, EStatus::Error });
-		return FHttpRequest::ECompletionBehavior::Done;
-	};
-	Request->EnqueueAsyncDownload(Owner, Pool, *RefsUri, MoveTemp(OnHttpRequestComplete), EHttpMediaType::CbObject, { 401, 404 });
+	});
 }
 
 void FHttpCacheStore::PutCacheRecordAsync(
@@ -1710,29 +2058,28 @@ void FHttpCacheStore::GetCacheValueAsync(
 	FString Bucket(Key.Bucket.ToString());
 	Bucket.ToLowerInline();
 
-	TStringBuilder<256> RefsUri;
-	RefsUri << "api/v1/refs/" << StructuredNamespace << "/" << Bucket << "/" << Key.Hash;
-
-	FHttpRequestPool* Pool = nullptr;
-	FHttpRequest* Request = WaitForHttpRequestForOwner<OperationCategory::Get>(Owner, false /* bUnboundedOverflow */, Pool);
+	TUniquePtr<FHttpOperation> Operation = WaitForHttpOperation(EOperationCategory::Get, /*bUnboundedOverflow*/ false);
+	FHttpOperation& LocalOperation = *Operation;
+	LocalOperation.SetUri(WriteToAnsiString<256>(EffectiveDomain, ANSITEXTVIEW("/api/v1/refs/"), StructuredNamespace, '/', Bucket, '/', Key.Hash));
+	LocalOperation.SetMethod(EHttpMethod::Get);
 	if (bSkipData)
 	{
-		Request->AddHeader(TEXT("Accept"), TEXT("application/x-ue-cb"));
+		LocalOperation.AddAcceptType(EHttpMediaType::CbObject);
 	}
 	else
 	{
-		Request->AddHeader(TEXT("Accept"), TEXT("application/x-jupiter-inline"));
+		LocalOperation.AddHeader(ANSITEXTVIEW("Accept"), ANSITEXTVIEW("application/x-jupiter-inline"));
 	}
-
-	auto OnHttpRequestComplete = [this, &Owner, Name, Key, UserData, bSkipData, OnComplete = MoveTemp(OnComplete)] (FHttpRequest::EResult HttpResult, FHttpRequest* Request)
+	LocalOperation.SetExpectedErrorCodes({404});
+	LocalOperation.SendAsync(Owner, [Operation = MoveTemp(Operation), this, Name, Key, UserData, bSkipData, OnComplete = MoveTemp(OnComplete)]
 	{
 		TRACE_CPUPROFILER_EVENT_SCOPE(HttpDDC_GetCacheValueAsync_OnHttpRequestComplete);
 
-		int64 ResponseCode = Request->GetResponseCode();
-		if (FHttpRequest::IsSuccessResponse(ResponseCode))
+		const int32 StatusCode = Operation->GetStatusCode();
+		if (StatusCode >= 200 && StatusCode <= 204)
 		{
 			FValue ResultValue;
-			FSharedBuffer ResponseBuffer = Request->MoveResponseBufferToShared();
+			FSharedBuffer ResponseBuffer = Operation->GetBody();
 
 			if (bSkipData)
 			{
@@ -1741,7 +2088,7 @@ void FHttpCacheStore::GetCacheValueAsync(
 					UE_LOG(LogDerivedDataCache, Display, TEXT("%s: Cache miss with invalid package for %s from '%s'"),
 						*Domain, *WriteToString<96>(Key), *Name);
 					OnComplete({Name, Key, {}, UserData, EStatus::Error});
-					return FHttpRequest::ECompletionBehavior::Done;
+					return;
 				}
 
 				const FCbObjectView Object = FCbObject(ResponseBuffer);
@@ -1752,7 +2099,7 @@ void FHttpCacheStore::GetCacheValueAsync(
 					UE_LOG(LogDerivedDataCache, Display, TEXT("%s: Cache miss with invalid value for %s from '%'"),
 						*Domain, *WriteToString<96>(Key), *Name);
 					OnComplete({Name, Key, {}, UserData, EStatus::Error});
-					return FHttpRequest::ECompletionBehavior::Done;
+					return;
 				}
 				ResultValue = FValue(RawHash, RawSize);
 			}
@@ -1761,8 +2108,7 @@ void FHttpCacheStore::GetCacheValueAsync(
 				FCompressedBuffer CompressedBuffer = FCompressedBuffer::FromCompressed(ResponseBuffer);
 				if (!CompressedBuffer)
 				{
-					FString ReceivedHashStr;
-					if (Request->GetHeader("X-Jupiter-InlinePayloadHash", ReceivedHashStr))
+					if (FAnsiStringView ReceivedHashStr = Operation->GetHeader("X-Jupiter-InlinePayloadHash"); !ReceivedHashStr.IsEmpty())
 					{
 						FIoHash ReceivedHash(ReceivedHashStr);
 						FIoHash ComputedHash = FIoHash::HashBuffer(ResponseBuffer.GetView());
@@ -1778,26 +2124,18 @@ void FHttpCacheStore::GetCacheValueAsync(
 					UE_LOG(LogDerivedDataCache, Display, TEXT("%s: Cache miss with invalid package for %s from '%s'"),
 						*Domain, *WriteToString<96>(Key), *Name);
 					OnComplete({Name, Key, {}, UserData, EStatus::Error});
-					return FHttpRequest::ECompletionBehavior::Done;
+					return;
 				}
 				ResultValue = FValue(CompressedBuffer);
 			}
 			OnComplete({Name, Key, ResultValue, UserData, EStatus::Ok});
-			return FHttpRequest::ECompletionBehavior::Done;
-		}
-
-		if (!ShouldAbortForShutdown() && !Owner.IsCanceled() && ShouldRetryOnError(HttpResult, ResponseCode) && ((Request->GetAttempts()+1) < UE_HTTPDDC_MAX_ATTEMPTS))
-		{
-			return FHttpRequest::ECompletionBehavior::Retry;
+			return;
 		}
 
 		UE_LOG(LogDerivedDataCache, Verbose, TEXT("%s: Cache miss with failed HTTP request for %s from '%s'"),
 			*Domain, *WriteToString<96>(Key), *Name);
 		OnComplete({Name, Key, {}, UserData, EStatus::Error});
-		return FHttpRequest::ECompletionBehavior::Done;
-	};
-
-	Request->EnqueueAsyncDownload(Owner, Pool, *RefsUri, MoveTemp(OnHttpRequestComplete), EHttpMediaType::UnspecifiedContentType, { 401, 404 });
+	});
 }
 
 void FHttpCacheStore::GetCacheRecordAsync(
@@ -1833,8 +2171,6 @@ void FHttpCacheStore::RefCachedDataProbablyExistsBatchAsync(
 		return;
 	}
 
-	TStringBuilder<256> RefsUri;
-	RefsUri << "api/v1/refs/" << StructuredNamespace;
 	FCbWriter RequestWriter;
 	RequestWriter.BeginObject();
 	RequestWriter.BeginArray(ANSITEXTVIEW("ops"));
@@ -1857,33 +2193,32 @@ void FHttpCacheStore::RefCachedDataProbablyExistsBatchAsync(
 	RequestWriter.EndObject();
 	FCbFieldIterator RequestFields = RequestWriter.Save();
 
-
-	FHttpRequestPool* Pool = nullptr;
-	FHttpRequest* Request = WaitForHttpRequestForOwner<OperationCategory::Get>(Owner, false /* bUnboundedOverflow */, Pool);
-
-	Request->AddHeader(TEXT("Accept"), TEXT("application/x-ue-cb"));
-
-	auto OnHttpRequestComplete = [this, &Owner, ValueRefs = TArray<FCacheGetValueRequest>(ValueRefs), OnComplete = MoveTemp(OnComplete)](FHttpRequest::EResult HttpResult, FHttpRequest* Request)
+	TUniquePtr<FHttpOperation> Operation = WaitForHttpOperation(EOperationCategory::Get, /*bUnboundedOverflow*/ false);
+	FHttpOperation& LocalOperation = *Operation;
+	LocalOperation.SetUri(WriteToAnsiString<256>(EffectiveDomain, ANSITEXTVIEW("/api/v1/refs/"), StructuredNamespace));
+	LocalOperation.SetMethod(EHttpMethod::Post);
+	LocalOperation.SetContentType(EHttpMediaType::CbObject);
+	LocalOperation.AddAcceptType(EHttpMediaType::CbObject);
+	LocalOperation.SetBody(FCompositeBuffer(RequestFields.GetOuterBuffer()));
+	LocalOperation.SendAsync(Owner, [Operation = MoveTemp(Operation), this, ValueRefs = TArray<FCacheGetValueRequest>(ValueRefs), OnComplete = MoveTemp(OnComplete)]
 	{
 		TRACE_CPUPROFILER_EVENT_SCOPE(HttpDDC_RefCachedDataProbablyExistsBatchAsync_OnHttpRequestComplete);
 
-		int64 ResponseCode = Request->GetResponseCode();
-		if (FHttpRequest::IsSuccessResponse(ResponseCode))
+		const int32 OverallStatusCode = Operation->GetStatusCode();
+		if (OverallStatusCode >= 200 && OverallStatusCode <= 204)
 		{
-			FMemoryView ResponseView = MakeMemoryView(Request->GetResponseBuffer().GetData(), Request->GetResponseBuffer().Num());
+			FMemoryView ResponseView = Operation->GetBody();
 			if (ValidateCompactBinary(ResponseView, ECbValidateMode::Default) != ECbValidateError::None)
 			{
 				for (const FCacheGetValueRequest& ValueRef : ValueRefs)
 				{
-					UE_LOG(LogDerivedDataCache, Log,
-						TEXT("%s: Cache exists returned invalid results."),
-						*Domain);
-						OnComplete(ValueRef.MakeResponse(EStatus::Error));
+					UE_LOG(LogDerivedDataCache, Log, TEXT("%s: Cache exists returned invalid results."), *Domain);
+					OnComplete(ValueRef.MakeResponse(EStatus::Error));
 				}
-				return FHttpRequest::ECompletionBehavior::Done;
+				return;
 			}
 
-			const FCbObjectView ResponseObject = FCbObjectView(Request->GetResponseBuffer().GetData());
+			const FCbObjectView ResponseObject(ResponseView.GetData());
 
 			FCbArrayView ResultsArrayView = ResponseObject[ANSITEXTVIEW("results")].AsArrayView();
 
@@ -1896,7 +2231,7 @@ void FHttpCacheStore::RefCachedDataProbablyExistsBatchAsync(
 						*Domain, ValueRefs.Num(), ResultsArrayView.Num());
 						OnComplete(ValueRef.MakeResponse(EStatus::Error));
 				}
-				return FHttpRequest::ECompletionBehavior::Done;
+				return;
 			}
 
 			for (FCbFieldView ResultFieldView : ResultsArrayView)
@@ -1915,7 +2250,7 @@ void FHttpCacheStore::RefCachedDataProbablyExistsBatchAsync(
 
 				const FCacheGetValueRequest& ValueRef = ValueRefs[OpId];
 
-				if (!FHttpRequest::IsSuccessResponse(StatusCode))
+				if (StatusCode < 200 || StatusCode > 204)
 				{
 					UE_LOG(LogDerivedDataCache, Verbose, TEXT("%s: Cache miss with unsuccessful response code %d for %s from '%s'"),
 						*Domain, StatusCode, *WriteToString<96>(ValueRef.Key), *ValueRef.Name);
@@ -1943,12 +2278,7 @@ void FHttpCacheStore::RefCachedDataProbablyExistsBatchAsync(
 
 				OnComplete({ValueRef.Name, ValueRef.Key, FValue(RawHash, RawSize), ValueRef.UserData, EStatus::Ok});
 			}
-			return FHttpRequest::ECompletionBehavior::Done;
-		}
-
-		if (!ShouldAbortForShutdown() && !Owner.IsCanceled() && ShouldRetryOnError(HttpResult, ResponseCode) && ((Request->GetAttempts()+1) < UE_HTTPDDC_MAX_ATTEMPTS))
-		{
-			return FHttpRequest::ECompletionBehavior::Retry;
+			return;
 		}
 
 		for (const FCacheGetValueRequest& ValueRef : ValueRefs)
@@ -1957,10 +2287,7 @@ void FHttpCacheStore::RefCachedDataProbablyExistsBatchAsync(
 				*Domain, *WriteToString<96>(ValueRef.Key), *ValueRef.Name);
 			OnComplete(ValueRef.MakeResponse(EStatus::Error));
 		}
-		return FHttpRequest::ECompletionBehavior::Done;
-	};
-
-	Request->EnqueueAsyncPost(Owner, Pool, *RefsUri, FCompositeBuffer(RequestFields.GetOuterBuffer()), MoveTemp(OnHttpRequestComplete), EHttpMediaType::CbObject);
+	});
 }
 
 void FHttpCacheStore::LegacyStats(FDerivedDataCacheStatsNode& OutNode)
diff --git a/Engine/Source/Developer/DerivedDataCache/Private/ZenCacheStore.cpp b/Engine/Source/Developer/DerivedDataCache/Private/ZenCacheStore.cpp
--- a/Engine/Source/Developer/DerivedDataCache/Private/ZenCacheStore.cpp
+++ b/Engine/Source/Developer/DerivedDataCache/Private/ZenCacheStore.cpp
@@ -11,6 +11,7 @@
 #include "DerivedDataCacheRecord.h"
 #include "DerivedDataCacheUsageStats.h"
 #include "DerivedDataChunk.h"
+#include "DerivedDataRequest.h"
 #include "DerivedDataRequestOwner.h"
 #include "Http/HttpClient.h"
 #include "Math/UnrealMathUtility.h"
@@ -21,11 +22,9 @@
 #include "ProfilingDebugging/CpuProfilerTrace.h"
 #include "Serialization/CompactBinary.h"
 #include "Serialization/CompactBinaryPackage.h"
-#include "Serialization/CompactBinarySerialization.h"
-#include "Serialization/CompactBinaryValidation.h"
 #include "Serialization/CompactBinaryWriter.h"
-#include "Serialization/LargeMemoryReader.h"
 #include "Serialization/LargeMemoryWriter.h"
+#include "Serialization/MemoryReader.h"
 #include "Templates/Function.h"
 #include "ZenBackendUtils.h"
 #include "ZenSerialization.h"
@@ -63,10 +62,6 @@ void ForEachBatch(const int32 BatchSize, const int32 TotalCount, T&& Fn)
 	}
 }
 
-//----------------------------------------------------------------------------------------------------------
-// FZenCacheStore
-//----------------------------------------------------------------------------------------------------------
-
 /**
  * Backend for a HTTP based caching service (Zen)
  */
@@ -77,10 +72,10 @@ public:
 	/**
 	 * Creates the cache store client, checks health status and attempts to acquire an access token.
 	 *
-	 * @param ServiceUrl	Base url to the service including schema.
+	 * @param ServiceUrl	Base url to the service including scheme.
 	 * @param Namespace		Namespace to use.
 	 */
-	FZenCacheStore(const TCHAR* ServiceUrl, const TCHAR* Namespace, const TCHAR* StructuredNamespace);
+	FZenCacheStore(const TCHAR* ServiceUrl, const TCHAR* Namespace);
 
 	inline FString GetName() const { return ZenService.GetInstance().GetURL(); }
 
@@ -125,29 +120,13 @@ public:
 private:
 	bool IsServiceReady();
 
-	using FOnRpcComplete = TUniqueFunction<void(FHttpRequest::EResult HttpResult, FHttpRequest* Request, FCbPackage& Response)>;
-	static FHttpRequest::EResult ParseRpcResponse(FHttpRequest::EResult ResultFromPost, FHttpRequest& Request, FCbPackage& OutResponse);
-	static FHttpRequest::EResult PerformBlockingRpc(FHttpRequest& Request,
-		FStringView Uri,
-		FCbObject RequestObject,
-		FCbPackage &OutResponse);
-	static FHttpRequest::EResult PerformBlockingRpc(FHttpRequest& Request,
-		FStringView Uri,
-		const FCbPackage& RequestPackage,
-		FCbPackage& OutResponse);
-	static void EnqueueAsyncRpc(FHttpRequest& Request,
-		UE::DerivedData::IRequestOwner& Owner,
-		FHttpRequestPool* Pool,
-		FStringView Uri,
-		FCbObject RequestObject,
-		FOnRpcComplete&& OnComplete);
-	static void EnqueueAsyncRpc(FHttpRequest& Request,
-		UE::DerivedData::IRequestOwner& Owner,
-		FHttpRequestPool* Pool,
-		FStringView Uri,
-		const FCbPackage& RequestPackage,
-		FOnRpcComplete&& OnComplete);
-
+	static FCompositeBuffer SaveRpcPackage(const FCbPackage& Package);
+	THttpUniquePtr<IHttpRequest> CreateRpcRequest();
+	THttpUniquePtr<IHttpResponse> PerformBlockingRpc(FCbObject RequestObject, FCbPackage& OutResponse);
+	THttpUniquePtr<IHttpResponse> PerformBlockingRpc(const FCbPackage& RequestPackage, FCbPackage& OutResponse);
+	using FOnRpcComplete = TUniqueFunction<void(THttpUniquePtr<IHttpResponse>& HttpResponse, FCbPackage& Response)>;
+	void EnqueueAsyncRpc(IRequestOwner& Owner, FCbObject RequestObject, FOnRpcComplete&& OnComplete);
+	void EnqueueAsyncRpc(IRequestOwner& Owner, const FCbPackage& RequestPackage, FOnRpcComplete&& OnComplete);
 
 	template <typename T, typename... ArgTypes>
 	static TRefCountPtr<T> MakeAsyncOp(ArgTypes&&... Args)
@@ -155,29 +134,31 @@ private:
 		// TODO: This should in-place construct from a pre-allocated memory pool
 		return TRefCountPtr<T>(new T(Forward<ArgTypes>(Args)...));
 	}
+
 private:
-	friend class FPutOp;
-	friend class FGetOp;
-	friend class FPutValueOp;
-	friend class FGetValueOp;
-	friend class FGetChunksOp;
+	class FPutOp;
+	class FGetOp;
+	class FPutValueOp;
+	class FGetValueOp;
+	class FGetChunksOp;
+
+	class FCbPackageReceiver;
+	class FAsyncCbPackageReceiver;
 
 	FString Namespace;
-	FString StructuredNamespace;
 	UE::Zen::FScopeZenService ZenService;
 	mutable FDerivedDataCacheUsageStats UsageStats;
-	TUniquePtr<FHttpSharedData> SharedData;
-	TUniquePtr<FHttpRequestPool> RequestPool;
+	THttpUniquePtr<IHttpConnectionPool> ConnectionPool;
+	FHttpRequestQueue RequestQueue;
 	bool bIsUsable = false;
 	int32 BatchPutMaxBytes = 1024*1024;
 	int32 CacheRecordBatchSize = 8;
 	int32 CacheChunksBatchSize = 8;
-
-	/** Debug Options */
 	FBackendDebugOptions DebugOptions;
+	TAnsiStringBuilder<256> RpcUri;
 };
 
-class FPutOp final : public FThreadSafeRefCountedObject
+class FZenCacheStore::FPutOp final : public FThreadSafeRefCountedObject
 {
 public:
 	FPutOp(FZenCacheStore& InCacheStore,
@@ -206,7 +187,7 @@ public:
 				{
 					ECachePolicy BatchDefaultPolicy = Batch[0].Policy.GetRecordPolicy();
 					BatchWriter << ANSITEXTVIEW("DefaultPolicy") << *WriteToString<128>(BatchDefaultPolicy);
-					BatchWriter.AddString(ANSITEXTVIEW("Namespace"), CacheStore.StructuredNamespace);
+					BatchWriter.AddString(ANSITEXTVIEW("Namespace"), CacheStore.Namespace);
 
 					BatchWriter.BeginArray(ANSITEXTVIEW("Requests"));
 					for (const FCachePutRequest& Request : Batch)
@@ -231,10 +212,10 @@ public:
 			BatchWriter.EndObject();
 			BatchPackage.SetObject(BatchWriter.Save().AsObject());
 
-			auto OnRpcComplete = [this, OpRef = TRefCountPtr<FPutOp>(this), Batch](FHttpRequest::EResult HttpResult, FHttpRequest* HttpRequest, FCbPackage& Response)
+			auto OnRpcComplete = [this, OpRef = TRefCountPtr<FPutOp>(this), Batch](THttpUniquePtr<IHttpResponse>& HttpResponse, FCbPackage& Response)
 			{
 				int32 RequestIndex = 0;
-				if (HttpResult == FHttpRequest::EResult::Success)
+				if (HttpResponse->GetErrorCode() == EHttpErrorCode::None)
 				{
 					const FCbObject& ResponseObj = Response.GetObject();
 					for (FCbField ResponseField : ResponseObj[ANSITEXTVIEW("Result")])
@@ -267,8 +248,7 @@ public:
 					OnMiss(Request);
 				}
 			};
-			FHttpRequest* Request = CacheStore.RequestPool->WaitForFreeRequest();
-			CacheStore.EnqueueAsyncRpc(*Request, Owner, CacheStore.RequestPool.Get(), TEXTVIEW("/z$/$rpc"), BatchPackage, MoveTemp(OnRpcComplete));
+			CacheStore.EnqueueAsyncRpc(Owner, BatchPackage, MoveTemp(OnRpcComplete));
 		}
 	}
 
@@ -317,7 +297,7 @@ private:
 	COOK_STAT(FCookStats::FScopedStatsCounter Timer = CacheStore.UsageStats.TimePut());
 };
 
-class FGetOp final : public FThreadSafeRefCountedObject
+class FZenCacheStore::FGetOp final : public FThreadSafeRefCountedObject
 {
 public:
 	FGetOp(FZenCacheStore& InCacheStore,
@@ -354,7 +334,7 @@ public:
 					{
 						ECachePolicy BatchDefaultPolicy = Batch[0].Policy.GetRecordPolicy();
 						BatchRequest << ANSITEXTVIEW("DefaultPolicy") << *WriteToString<128>(BatchDefaultPolicy);
-						BatchRequest.AddString(ANSITEXTVIEW("Namespace"), CacheStore.StructuredNamespace);
+						BatchRequest.AddString(ANSITEXTVIEW("Namespace"), CacheStore.Namespace);
 
 						BatchRequest.BeginArray(ANSITEXTVIEW("Requests"));
 						for (const FCacheGetRequest& Request : Batch)
@@ -376,10 +356,10 @@ public:
 				BatchRequest.EndObject();
 
 				FGetOp* OriginalOp = this;
-				auto OnRpcComplete = [this, OpRef = TRefCountPtr<FGetOp>(OriginalOp), Batch](FHttpRequest::EResult HttpResult, FHttpRequest* HttpRequest, FCbPackage& Response)
+				auto OnRpcComplete = [this, OpRef = TRefCountPtr<FGetOp>(OriginalOp), Batch](THttpUniquePtr<IHttpResponse>& HttpResponse, FCbPackage& Response)
 				{
 					int32 RequestIndex = 0;
-					if (HttpResult == FHttpRequest::EResult::Success)
+					if (HttpResponse->GetErrorCode() == EHttpErrorCode::None)
 					{
 						const FCbObject& ResponseObj = Response.GetObject();
 						
@@ -418,8 +398,7 @@ public:
 					}
 				};
 
-				FHttpRequest* Request = CacheStore.RequestPool->WaitForFreeRequest();
-				CacheStore.EnqueueAsyncRpc(*Request, Owner, CacheStore.RequestPool.Get(), TEXTVIEW("/z$/$rpc"), BatchRequest.Save().AsObject(), MoveTemp(OnRpcComplete));
+				CacheStore.EnqueueAsyncRpc(Owner, BatchRequest.Save().AsObject(), MoveTemp(OnRpcComplete));
 			});
 	}
 private:
@@ -451,7 +430,7 @@ private:
 	COOK_STAT(FCookStats::FScopedStatsCounter Timer = CacheStore.UsageStats.TimeGet());
 };
 
-class FPutValueOp final : public FThreadSafeRefCountedObject
+class FZenCacheStore::FPutValueOp final : public FThreadSafeRefCountedObject
 {
 public:
 	FPutValueOp(FZenCacheStore& InCacheStore,
@@ -480,7 +459,7 @@ public:
 				{
 					ECachePolicy BatchDefaultPolicy = Batch[0].Policy;
 					BatchWriter << ANSITEXTVIEW("DefaultPolicy") << *WriteToString<128>(BatchDefaultPolicy);
-					BatchWriter.AddString(ANSITEXTVIEW("Namespace"), CacheStore.StructuredNamespace);
+					BatchWriter.AddString(ANSITEXTVIEW("Namespace"), CacheStore.Namespace);
 
 					BatchWriter.BeginArray("Requests");
 					for (const FCachePutValueRequest& Request : Batch)
@@ -508,10 +487,10 @@ public:
 			BatchWriter.EndObject();
 			BatchPackage.SetObject(BatchWriter.Save().AsObject());
 
-			auto OnRpcComplete = [this, OpRef = TRefCountPtr<FPutValueOp>(this), Batch](FHttpRequest::EResult HttpResult, FHttpRequest* HttpRequest, FCbPackage& Response)
+			auto OnRpcComplete = [this, OpRef = TRefCountPtr<FPutValueOp>(this), Batch](THttpUniquePtr<IHttpResponse>& HttpResponse, FCbPackage& Response)
 			{
 				int32 RequestIndex = 0;
-				if (HttpResult == FHttpRequest::EResult::Success)
+				if (HttpResponse->GetErrorCode() == EHttpErrorCode::None)
 				{
 					const FCbObject& ResponseObj = Response.GetObject();
 					for (FCbField ResponseField : ResponseObj[ANSITEXTVIEW("Result")])
@@ -543,8 +522,8 @@ public:
 					OnMiss(Request);
 				}
 			};
-			FHttpRequest* Request = CacheStore.RequestPool->WaitForFreeRequest();
-			CacheStore.EnqueueAsyncRpc(*Request, Owner, CacheStore.RequestPool.Get(), TEXTVIEW("/z$/$rpc"), BatchPackage, MoveTemp(OnRpcComplete));
+
+			CacheStore.EnqueueAsyncRpc(Owner, BatchPackage, MoveTemp(OnRpcComplete));
 		}
 	}
 
@@ -586,7 +565,7 @@ private:
 	COOK_STAT(FCookStats::FScopedStatsCounter Timer = CacheStore.UsageStats.TimePut());
 };
 
-class FGetValueOp final : public FThreadSafeRefCountedObject
+class FZenCacheStore::FGetValueOp final : public FThreadSafeRefCountedObject
 {
 public:
 	FGetValueOp(FZenCacheStore& InCacheStore,
@@ -623,7 +602,7 @@ public:
 				{
 					ECachePolicy BatchDefaultPolicy = Batch[0].Policy;
 					BatchRequest << ANSITEXTVIEW("DefaultPolicy") << *WriteToString<128>(BatchDefaultPolicy);
-					BatchRequest.AddString(ANSITEXTVIEW("Namespace"), CacheStore.StructuredNamespace);
+					BatchRequest.AddString(ANSITEXTVIEW("Namespace"), CacheStore.Namespace);
 
 					BatchRequest.BeginArray("Requests");
 					for (const FCacheGetValueRequest& Request : Batch)
@@ -645,10 +624,10 @@ public:
 			BatchRequest.EndObject();
 
 			FGetValueOp* OriginalOp = this;
-			auto OnRpcComplete = [this, OpRef = TRefCountPtr<FGetValueOp>(OriginalOp), Batch](FHttpRequest::EResult HttpResult, FHttpRequest* HttpRequest, FCbPackage& Response)
+			auto OnRpcComplete = [this, OpRef = TRefCountPtr<FGetValueOp>(OriginalOp), Batch](THttpUniquePtr<IHttpResponse>& HttpResponse, FCbPackage& Response)
 			{
 				int32 RequestIndex = 0;
-				if (HttpResult == FHttpRequest::EResult::Success)
+				if (HttpResponse->GetErrorCode() == EHttpErrorCode::None)
 				{
 					const FCbObject& ResponseObj = Response.GetObject();
 
@@ -699,8 +678,8 @@ public:
 					OnMiss(Request);
 				}
 			};
-			FHttpRequest* Request = CacheStore.RequestPool->WaitForFreeRequest();
-			CacheStore.EnqueueAsyncRpc(*Request, Owner, CacheStore.RequestPool.Get(), TEXTVIEW("/z$/$rpc"), BatchRequest.Save().AsObject(), MoveTemp(OnRpcComplete));
+
+			CacheStore.EnqueueAsyncRpc(Owner, BatchRequest.Save().AsObject(), MoveTemp(OnRpcComplete));
 		});
 	}
 private:
@@ -731,7 +710,7 @@ private:
 	COOK_STAT(FCookStats::FScopedStatsCounter Timer = CacheStore.UsageStats.TimeGet());
 };
 
-class FGetChunksOp final : public FThreadSafeRefCountedObject
+class FZenCacheStore::FGetChunksOp final : public FThreadSafeRefCountedObject
 {
 public:
 	FGetChunksOp(FZenCacheStore& InCacheStore,
@@ -770,7 +749,7 @@ public:
 				{
 					ECachePolicy DefaultPolicy = Batch[0].Policy;
 					BatchRequest << ANSITEXTVIEW("DefaultPolicy") << WriteToString<128>(DefaultPolicy);
-					BatchRequest.AddString(ANSITEXTVIEW("Namespace"), CacheStore.StructuredNamespace);
+					BatchRequest.AddString(ANSITEXTVIEW("Namespace"), CacheStore.Namespace);
 
 					BatchRequest.BeginArray(ANSITEXTVIEW("ChunkRequests"));
 					for (const FCacheGetChunkRequest& Request : Batch)
@@ -809,16 +788,16 @@ public:
 			BatchRequest.EndObject();
 
 			FCbPackage BatchResponse;
-			FHttpRequest::EResult HttpResult = FHttpRequest::EResult::Failed;
+			EHttpErrorCode HttpErrorCode = EHttpErrorCode::Unknown;
 
 			{
 				LLM_SCOPE_BYTAG(UntaggedDDCResult);
-				FScopedHttpPoolRequestPtr Request(CacheStore.RequestPool.Get());
-				HttpResult = CacheStore.PerformBlockingRpc(*Request.Get(), TEXTVIEW("/z$/$rpc"), BatchRequest.Save().AsObject(), BatchResponse);
+				THttpUniquePtr<IHttpResponse> HttpResponse = CacheStore.PerformBlockingRpc(BatchRequest.Save().AsObject(), BatchResponse);
+				HttpErrorCode = HttpResponse->GetErrorCode();
 			}
 
 			int32 RequestIndex = 0;
-			if (HttpResult == FHttpRequest::EResult::Success)
+			if (HttpErrorCode == EHttpErrorCode::None)
 			{
 				const FCbObject& ResponseObj = BatchResponse.GetObject();
 
@@ -904,22 +883,123 @@ private:
 	COOK_STAT(FCookStats::FScopedStatsCounter Timer = CacheStore.UsageStats.TimeGet());
 };
 
+class FZenCacheStore::FCbPackageReceiver final : public IHttpReceiver
+{
+public:
+	FCbPackageReceiver(const FCbPackageReceiver&) = delete;
+	FCbPackageReceiver& operator=(const FCbPackageReceiver&) = delete;
+
+	explicit FCbPackageReceiver(FCbPackage& OutPackage, IHttpReceiver* InNext = nullptr)
+		: Package(OutPackage)
+		, Next(InNext)
+	{
+	}
+
+private:
+	IHttpReceiver* OnCreate(IHttpResponse& Response) final
+	{
+		return &BodyReceiver;
+	}
+
+	IHttpReceiver* OnComplete(IHttpResponse& Response) final
+	{
+		FMemoryReaderView Ar(MakeMemoryView(BodyArray));
+		Package.TryLoad(Ar);
+		return Next;
+	}
+
+private:
+	FCbPackage& Package;
+	IHttpReceiver* Next;
+	TArray64<uint8> BodyArray;
+	FHttpByteArrayReceiver BodyReceiver{BodyArray, this};
+};
+
+class FZenCacheStore::FAsyncCbPackageReceiver final : public FRequestBase, public IHttpReceiver
+{
+public:
+	FAsyncCbPackageReceiver(const FAsyncCbPackageReceiver&) = delete;
+	FAsyncCbPackageReceiver& operator=(const FAsyncCbPackageReceiver&) = delete;
+
+	FAsyncCbPackageReceiver(
+		THttpUniquePtr<IHttpRequest>&& InRequest,
+		IRequestOwner* InOwner,
+		FOnRpcComplete&& InOnRpcComplete)
+		: Request(MoveTemp(InRequest))
+		, Owner(InOwner)
+		, BaseReceiver(Package, this)
+		, OnRpcComplete(MoveTemp(InOnRpcComplete))
+	{
+		Request->SendAsync(this, Response);
+	}
+
+private:
+	// IRequest Interface
+
+	void SetPriority(EPriority Priority) final {}
+	void Cancel() final { Monitor->Cancel(); }
+	void Wait() final { Monitor->Wait(); }
+
+	// IHttpReceiver Interface
+
+	IHttpReceiver* OnCreate(IHttpResponse& LocalResponse) final
+	{
+		Monitor = LocalResponse.GetMonitor();
+		Owner->Begin(this);
+		return &BaseReceiver;
+	}
+
+	IHttpReceiver* OnComplete(IHttpResponse& LocalResponse) final
+	{
+		Request.Reset();
+		Owner->End(this, [Self = this]
+		{
+			if (Self->OnRpcComplete)
+			{
+				// Launch a task for the completion function since it can execute arbitrary code.
+				Self->Owner->LaunchTask(TEXT("ZenHttpComplete"), [Self = TRefCountPtr(Self)]
+				{
+					Self->OnRpcComplete(Self->Response, Self->Package);
+				});
+			}
+		});
+		return nullptr;
+	}
+
+private:
+	THttpUniquePtr<IHttpRequest> Request;
+	THttpUniquePtr<IHttpResponse> Response;
+	TRefCountPtr<IHttpResponseMonitor> Monitor;
+	IRequestOwner* Owner;
+	FCbPackage Package;
+	FCbPackageReceiver BaseReceiver;
+	FOnRpcComplete OnRpcComplete;
+};
+
 FZenCacheStore::FZenCacheStore(
 	const TCHAR* InServiceUrl,
-	const TCHAR* InNamespace,
-	const TCHAR* InStructuredNamespace)
+	const TCHAR* InNamespace)
 	: Namespace(InNamespace)
-	, StructuredNamespace(InStructuredNamespace)
 	, ZenService(InServiceUrl)
 {
-	const uint32 MaxConnections = FMath::Clamp(static_cast<uint32>(FPlatformMisc::NumberOfCoresIncludingHyperthreads()), 8, 64);
-	constexpr uint32 RequestPoolSize = 128;
-	constexpr uint32 RequestPoolOverflowSize = 128;
-	
-	SharedData = MakeUnique<FHttpSharedData>(MaxConnections);
 	if (IsServiceReady())
 	{
-		RequestPool = MakeUnique<FHttpRequestPool>(ZenService.GetInstance().GetURL(), ZenService.GetInstance().GetURL(), nullptr, SharedData.Get(), RequestPoolSize, RequestPoolOverflowSize);
+		RpcUri << ZenService.GetInstance().GetURL() << ANSITEXTVIEW("/z$/$rpc");
+
+		const uint32 MaxConnections = FMath::Clamp(static_cast<uint32>(FPlatformMisc::NumberOfCoresIncludingHyperthreads()), 8, 64);
+		constexpr uint32 RequestPoolSize = 128;
+		constexpr uint32 RequestPoolOverflowSize = 128;
+
+		FHttpConnectionPoolParams ConnectionPoolParams;
+		ConnectionPoolParams.MaxConnections = MaxConnections;
+		ConnectionPoolParams.MinConnections = MaxConnections;
+		ConnectionPool = IHttpManager::Get().CreateConnectionPool(ConnectionPoolParams);
+
+		FHttpClientParams ClientParams;
+		ClientParams.MaxRequests = RequestPoolSize + RequestPoolOverflowSize;
+		ClientParams.MinRequests = RequestPoolSize;
+		RequestQueue = FHttpRequestQueue(*ConnectionPool, ClientParams);
+
 		bIsUsable = true;
 
 		// Issue a request for stats as it will be fetched asynchronously and issuing now makes them available sooner for future callers.
@@ -937,85 +1017,61 @@ bool FZenCacheStore::IsServiceReady()
 	return ZenService.GetInstance().IsServiceReady();
 }
 
-FHttpRequest::EResult FZenCacheStore::ParseRpcResponse(FHttpRequest::EResult ResultFromPost, FHttpRequest& Request, FCbPackage& OutResponse)
+FCompositeBuffer FZenCacheStore::SaveRpcPackage(const FCbPackage& Package)
 {
-	if (ResultFromPost != FHttpRequest::EResult::Success || !FHttpRequest::IsSuccessResponse(Request.GetResponseCode()))
-	{
-		return FHttpRequest::EResult::Failed;
-	}
-
-	const TArray64<uint8>& ResponseBuffer = Request.GetResponseBuffer();
-	if (ResponseBuffer.Num())
-	{
-		FLargeMemoryReader Ar(ResponseBuffer.GetData(), ResponseBuffer.Num());
-		if (!OutResponse.TryLoad(Ar))
-		{
-			return FHttpRequest::EResult::Failed;
-		}
-	}
-
-	return FHttpRequest::EResult::Success;
+	FLargeMemoryWriter Memory;
+	Zen::Http::SaveCbPackage(Package, Memory);
+	uint64 PackageMemorySize = Memory.TotalSize();
+	return FCompositeBuffer(FSharedBuffer::TakeOwnership(Memory.ReleaseOwnership(), PackageMemorySize, FMemory::Free));
 }
 
-FHttpRequest::EResult FZenCacheStore::PerformBlockingRpc(FHttpRequest& Request,
-	FStringView Uri,
-	FCbObject RequestObject,
-	FCbPackage& OutResponse)
+THttpUniquePtr<IHttpRequest> FZenCacheStore::CreateRpcRequest()
 {
-	return ParseRpcResponse(
-		Request.PerformBlockingPost(Uri, RequestObject.GetBuffer(), EHttpMediaType::CbObject, EHttpMediaType::CbPackage), Request, OutResponse);
+	THttpUniquePtr<IHttpRequest> Request = RequestQueue.CreateRequest({});
+	Request->SetUri(RpcUri);
+	Request->SetMethod(EHttpMethod::Post);
+	Request->AddAcceptType(EHttpMediaType::CbPackage);
+	return Request;
 }
 
-FHttpRequest::EResult FZenCacheStore::PerformBlockingRpc(FHttpRequest& Request,
-	FStringView Uri,
-	const FCbPackage& RequestPackage,
-	FCbPackage& OutResponse)
+THttpUniquePtr<IHttpResponse> FZenCacheStore::PerformBlockingRpc(FCbObject RequestObject, FCbPackage& OutResponse)
 {
-	FLargeMemoryWriter PackageMemory;
-	UE::Zen::Http::SaveCbPackage(RequestPackage, PackageMemory);
+	THttpUniquePtr<IHttpRequest> Request = CreateRpcRequest();
+	Request->SetContentType(EHttpMediaType::CbObject);
+	Request->SetBody(RequestObject.GetBuffer().MakeOwned());
+
+	FCbPackageReceiver Receiver(OutResponse);
+	THttpUniquePtr<IHttpResponse> Response;
+	Request->Send(&Receiver, Response);
+	return Response;
+}
 
-	return ParseRpcResponse(
-		Request.PerformBlockingPost(Uri, FCompositeBuffer(FSharedBuffer::MakeView(PackageMemory.GetView())), EHttpMediaType::CbPackage, EHttpMediaType::CbPackage), Request, OutResponse);
+THttpUniquePtr<IHttpResponse> FZenCacheStore::PerformBlockingRpc(const FCbPackage& RequestPackage, FCbPackage& OutResponse)
+{
+	THttpUniquePtr<IHttpRequest> Request = CreateRpcRequest();
+	Request->SetContentType(EHttpMediaType::CbPackage);
+	Request->SetBody(SaveRpcPackage(RequestPackage));
+
+	FCbPackageReceiver Receiver(OutResponse);
+	THttpUniquePtr<IHttpResponse> Response;
+	Request->Send(&Receiver, Response);
+	return Response;
 }
 
-void FZenCacheStore::EnqueueAsyncRpc(FHttpRequest& Request,
-	UE::DerivedData::IRequestOwner& Owner,
-	FHttpRequestPool* Pool,
-	FStringView Uri,
-	FCbObject RequestObject,
-	FOnRpcComplete&& OnComplete)
+void FZenCacheStore::EnqueueAsyncRpc(IRequestOwner& Owner, FCbObject RequestObject, FOnRpcComplete&& OnComplete)
 {
-	auto OnHttpRequestComplete = [OnComplete = MoveTemp(OnComplete)]
-	(FHttpRequest::EResult HttpResult, FHttpRequest* Request)
-	{
-		FCbPackage Response;
-		FHttpRequest::EResult ParsedResponseResult = ParseRpcResponse(HttpResult, *Request, Response);
-		OnComplete(ParsedResponseResult, Request, Response);
-		return FHttpRequest::ECompletionBehavior::Done;
-	};
-	Request.EnqueueAsyncPost(Owner, Pool, Uri, RequestObject.GetBuffer(), MoveTemp(OnHttpRequestComplete), EHttpMediaType::CbObject, EHttpMediaType::CbPackage);
+	THttpUniquePtr<IHttpRequest> Request = CreateRpcRequest();
+	Request->SetContentType(EHttpMediaType::CbObject);
+	Request->SetBody(RequestObject.GetBuffer().MakeOwned());
+	new FAsyncCbPackageReceiver(MoveTemp(Request), &Owner, MoveTemp(OnComplete));
 }
 
-void FZenCacheStore::EnqueueAsyncRpc(FHttpRequest& Request,
-	UE::DerivedData::IRequestOwner& Owner,
-	FHttpRequestPool* Pool,
-	FStringView Uri,
-	const FCbPackage& RequestPackage,
-	FOnRpcComplete&& OnComplete)
+void FZenCacheStore::EnqueueAsyncRpc(IRequestOwner& Owner, const FCbPackage& RequestPackage, FOnRpcComplete&& OnComplete)
 {
-	auto OnHttpRequestComplete = [OnComplete = MoveTemp(OnComplete)]
-	(FHttpRequest::EResult HttpResult, FHttpRequest* Request)
-	{
-		FCbPackage Response;
-		FHttpRequest::EResult ParsedResponseResult = ParseRpcResponse(HttpResult, *Request, Response);
-		OnComplete(ParsedResponseResult, Request, Response);
-		return FHttpRequest::ECompletionBehavior::Done;
-	};
-	FLargeMemoryWriter PackageMemory;
-	UE::Zen::Http::SaveCbPackage(RequestPackage, PackageMemory);
-	uint64 PackageMemorySize = PackageMemory.TotalSize();
-	FSharedBuffer PackageSharedBuffer = FSharedBuffer::TakeOwnership(PackageMemory.ReleaseOwnership(), PackageMemorySize, FMemory::Free);
-	Request.EnqueueAsyncPost(Owner, Pool, Uri, FCompositeBuffer(PackageSharedBuffer), MoveTemp(OnHttpRequestComplete), EHttpMediaType::CbPackage, EHttpMediaType::CbPackage);
+	THttpUniquePtr<IHttpRequest> Request = CreateRpcRequest();
+	Request->SetContentType(EHttpMediaType::CbPackage);
+	Request->SetBody(SaveRpcPackage(RequestPackage));
+	new FAsyncCbPackageReceiver(MoveTemp(Request), &Owner, MoveTemp(OnComplete));
 }
 
 void FZenCacheStore::LegacyStats(FDerivedDataCacheStatsNode& OutNode)
@@ -1144,7 +1200,7 @@ TTuple<ILegacyCacheStore*, ECacheStoreFlags> CreateZenCacheStore(const TCHAR* No
 		UE_LOG(LogDerivedDataCache, Warning, TEXT("%s: Missing required parameter 'StructuredNamespace', falling back to Namespace: '%s'"), NodeName, *StructuredNamespace);
 	}
 
-	TUniquePtr<FZenCacheStore> Backend = MakeUnique<FZenCacheStore>(*ServiceUrl, *Namespace, *StructuredNamespace);
+	TUniquePtr<FZenCacheStore> Backend = MakeUnique<FZenCacheStore>(*ServiceUrl, *StructuredNamespace);
 	if (!Backend->IsUsable())
 	{
 		UE_LOG(LogDerivedDataCache, Warning, TEXT("%s: Failed to contact the service (%s), will not use it."), NodeName, *Backend->GetName());
