DDC: Fixed an error in IHttpClient::GetHeader[s]

#rnx

diff --git a/Engine/Source/Developer/DerivedDataCache/Private/Http/HttpClient.cpp b/Engine/Source/Developer/DerivedDataCache/Private/Http/HttpClient.cpp
--- a/Engine/Source/Developer/DerivedDataCache/Private/Http/HttpClient.cpp
+++ b/Engine/Source/Developer/DerivedDataCache/Private/Http/HttpClient.cpp
@@ -178,7 +178,7 @@ FAnsiStringView IHttpResponse::GetHeader(const FAnsiStringView Name) const
 		if (Header.StartsWith(Name, ESearchCase::IgnoreCase) &&
 			Header.Len() > NameLen && Header.GetData()[NameLen] == ':')
 		{
-			return Header.TrimStartAndEnd();
+			return Header.RightChop(NameLen + 1).TrimStartAndEnd();
 		}
 	}
 	return {};
@@ -197,7 +197,7 @@ int32 IHttpResponse::GetHeaders(FAnsiStringView Name, TArrayView<FAnsiStringView
 		{
 			if (ValuesCount > MatchIndex)
 			{
-				Values[MatchIndex] = Header.TrimStartAndEnd();
+				Values[MatchIndex] = Header.RightChop(NameLen + 1).TrimStartAndEnd();
 			}
 			++MatchIndex;
 		}
