DDC: Added public key pinning support to the Horde Storage client

It would be better to use Certificate Transparency, but the UE SSL module currently expects pinned public keys, and verification using known key hashes is far better than no verification.

#rnx

diff --git a/Engine/Source/Developer/DerivedDataCache/Private/HttpCacheStore.cpp b/Engine/Source/Developer/DerivedDataCache/Private/HttpCacheStore.cpp
--- a/Engine/Source/Developer/DerivedDataCache/Private/HttpCacheStore.cpp
+++ b/Engine/Source/Developer/DerivedDataCache/Private/HttpCacheStore.cpp
@@ -54,6 +54,10 @@
 #include <netdb.h>
 #endif
 
+#if WITH_SSL
+#include "Ssl.h"
+#endif
+
 #define UE_HTTPDDC_GET_REQUEST_POOL_SIZE 48
 #define UE_HTTPDDC_PUT_REQUEST_POOL_SIZE 16
 #define UE_HTTPDDC_NONBLOCKING_REQUEST_POOL_SIZE 128
@@ -122,9 +126,8 @@ static bool IsValueDataReady(FValue& Value, const ECachePolicy Policy)
 	return false;
 };
 
-static bool TryResolveCanonicalHost(const FAnsiStringView Uri, FAnsiStringBuilderBase& OutUri)
+static FAnsiStringView GetDomainFromUri(const FAnsiStringView Uri)
 {
-	// Extract the host from the URI.
 	FAnsiStringView Domain = Uri;
 	if (const int32 SchemeIndex = String::FindFirst(Domain, ANSITEXTVIEW("://")); SchemeIndex != INDEX_NONE)
 	{
@@ -160,8 +163,13 @@ static bool TryResolveCanonicalHost(const FAnsiStringView Uri, FAnsiStringBuilde
 	{
 		RemovePort(Domain);
 	}
+	return Domain;
+}
 
+static bool TryResolveCanonicalHost(const FAnsiStringView Uri, FAnsiStringBuilderBase& OutUri)
+{
 	// Append the URI until the end of the domain.
+	const FAnsiStringView Domain = GetDomainFromUri(Uri);
 	const int32 OutUriIndex = OutUri.Len();
 	const int32 DomainIndex = int32(Domain.GetData() - Uri.GetData());
 	const int32 DomainEndIndex = DomainIndex + Domain.Len();
@@ -225,6 +233,7 @@ FAnsiStringBuilderBase& operator<<(FAnsiStringBuilderBase& Builder, const FHttpA
 struct FHttpCacheStoreParams
 {
 	FString Host;
+	FString HostPinnedPublicKeys;
 	FString Namespace;
 	FString StructuredNamespace;
 	FString OAuthProvider;
@@ -233,6 +242,7 @@ struct FHttpCacheStoreParams
 	FString OAuthScope;
 	FString OAuthProviderIdentifier;
 	FString OAuthAccessToken;
+	FString OAuthPinnedPublicKeys;
 	bool bResolveHostCanonicalName = true;
 	bool bReadOnly = false;
 
@@ -1559,6 +1569,17 @@ FHttpCacheStore::FHttpCacheStore(const FHttpCacheStoreParams& Params)
 		EffectiveDomain.Append(ResolvedDomain);
 	}
 
+#if WITH_SSL
+	if (!Params.HostPinnedPublicKeys.IsEmpty() && EffectiveDomain.ToView().StartsWith(ANSITEXTVIEW("https://")))
+	{
+		FSslModule::Get().GetCertificateManager().SetPinnedPublicKeys(FString(GetDomainFromUri(EffectiveDomain)), Params.HostPinnedPublicKeys);
+	}
+	if (!Params.OAuthPinnedPublicKeys.IsEmpty() && OAuthProvider.StartsWith(TEXT("https://")))
+	{
+		FSslModule::Get().GetCertificateManager().SetPinnedPublicKeys(FString(GetDomainFromUri(WriteToAnsiString<256>(OAuthProvider))), Params.OAuthPinnedPublicKeys);
+	}
+#endif
+
 	constexpr uint32 MaxTotalConnections = 8;
 	FHttpConnectionPoolParams ConnectionPoolParams;
 	ConnectionPoolParams.MaxConnections = MaxTotalConnections;
@@ -2671,6 +2692,8 @@ void FHttpCacheStoreParams::Parse(const TCHAR* NodeName, const TCHAR* Config)
 		}
 	}
 
+	FParse::Value(Config, TEXT("HostPinnedPublicKeys="), HostPinnedPublicKeys);
+
 	FParse::Bool(Config, TEXT("ResolveHostCanonicalName="), bResolveHostCanonicalName);
 
 	// Namespace Params
@@ -2716,6 +2739,8 @@ void FHttpCacheStoreParams::Parse(const TCHAR* NodeName, const TCHAR* Config)
 		}
 	}
 
+	FParse::Value(Config, TEXT("OAuthPinnedPublicKeys="), OAuthPinnedPublicKeys);
+
 	// Cache Params
 
 	FParse::Bool(Config, TEXT("ReadOnly="), bReadOnly);
