DDC: Fixed warnings from PVS-Studio in CurlHttpClient

#rnx

diff --git a/Engine/Source/Developer/DerivedDataCache/Private/Http/CurlHttpClient.cpp b/Engine/Source/Developer/DerivedDataCache/Private/Http/CurlHttpClient.cpp
--- a/Engine/Source/Developer/DerivedDataCache/Private/Http/CurlHttpClient.cpp
+++ b/Engine/Source/Developer/DerivedDataCache/Private/Http/CurlHttpClient.cpp
@@ -950,15 +950,15 @@ FCurlHttpResponse::FCurlHttpResponse(CURL* InCurl, EHttpMethod InMethod, FAnsiSt
 
 bool FCurlHttpResponse::Create()
 {
-	for (TGuardValue GuardReceiverFunction(ReceiverFunction, EHttpReceiverFunction::OnCreate);
-		 IHttpReceiver* NewReceiver = Receiver->OnCreate(*this);
-		 Receiver = NewReceiver)
+	for (TGuardValue GuardReceiverFunction(ReceiverFunction, EHttpReceiverFunction::OnCreate);;)
 	{
+		IHttpReceiver* NewReceiver = Receiver->OnCreate(*this);
 		if (Receiver == NewReceiver)
 		{
 			break;
 		}
 		checkf(NewReceiver, TEXT("Receiver must not be null for %s"), *WriteToString<256>(*this));
+		Receiver = NewReceiver;
 	}
 	return !IsCanceled();
 }
@@ -1018,15 +1018,15 @@ bool FCurlHttpResponse::WriteHeader(FAnsiStringView Header)
 			HeaderViews.Add(MakeStringView(Data, Len - 2));
 			Data += Len;
 		}
-		for (TGuardValue GuardReceiverFunction(ReceiverFunction, EHttpReceiverFunction::OnHeaders);
-			 IHttpReceiver* NewReceiver = Receiver->OnHeaders(*this);
-			 Receiver = NewReceiver)
+		for (TGuardValue GuardReceiverFunction(ReceiverFunction, EHttpReceiverFunction::OnHeaders);;)
 		{
+			IHttpReceiver* NewReceiver = Receiver->OnHeaders(*this);
 			if (Receiver == NewReceiver)
 			{
 				break;
 			}
 			checkf(NewReceiver, TEXT("Receiver must not be null for %s"), *WriteToString<256>(*this));
+			Receiver = NewReceiver;
 		}
 	}
 	else
@@ -1119,14 +1119,14 @@ void FCurlHttpResponse::SetComplete(CURLcode Code)
 	// Hold a reference to safely access Receiver, State, and CompleteEvent.
 	TRefCountPtr<FCurlHttpResponse> Self(this);
 
-	for (TGuardValue GuardReceiverFunction(ReceiverFunction, EHttpReceiverFunction::OnComplete);
-		 IHttpReceiver* NewReceiver = Receiver->OnComplete(*this);
-		 Receiver = NewReceiver)
+	for (TGuardValue GuardReceiverFunction(ReceiverFunction, EHttpReceiverFunction::OnComplete);;)
 	{
-		if (Receiver == NewReceiver)
+		IHttpReceiver* NewReceiver = Receiver->OnComplete(*this);
+		if (Receiver == NewReceiver || !NewReceiver)
 		{
 			break;
 		}
+		Receiver = NewReceiver;
 	}
 
 	CurlHttp::Private::AtomicFetchOr(State, ECurlHttpResponseState::Complete);
