DDC: Added more detail to warnings from the Zen cache client

#rnx

diff --git a/Engine/Source/Developer/DerivedDataCache/Private/ZenCacheStore.cpp b/Engine/Source/Developer/DerivedDataCache/Private/ZenCacheStore.cpp
--- a/Engine/Source/Developer/DerivedDataCache/Private/ZenCacheStore.cpp
+++ b/Engine/Source/Developer/DerivedDataCache/Private/ZenCacheStore.cpp
@@ -240,7 +240,8 @@ public:
 					if (RequestIndex != Batch.Num())
 					{
 						UE_LOG(LogDerivedDataCache, Warning,
-							TEXT("%s: Invalid response received from PutCacheRecords rpc: %d results expected, received %d."), *CacheStore.GetName(), Batch.Num(), RequestIndex);
+							TEXT("%s: Invalid response received from PutCacheRecords RPC: %d results expected, received %d, from %s"),
+							*CacheStore.GetName(), Batch.Num(), RequestIndex, *WriteToString<256>(*HttpResponse));
 					}
 				}
 				for (const FCachePutRequest& Request : Batch.RightChop(RequestIndex))
@@ -389,7 +390,8 @@ public:
 						if (RequestIndex != Batch.Num())
 						{
 							UE_LOG(LogDerivedDataCache, Warning,
-								TEXT("%s: Invalid response received from GetCacheRecords rpc: %d results expected, received %d."), *CacheStore.GetName(), Batch.Num(), RequestIndex);
+								TEXT("%s: Invalid response received from GetCacheRecords RPC: %d results expected, received %d, from %s"),
+								*CacheStore.GetName(), Batch.Num(), RequestIndex, *WriteToString<256>(*HttpResponse));
 						}
 					}
 					for (const FCacheGetRequest& Request : Batch.RightChop(RequestIndex))
@@ -514,7 +516,8 @@ public:
 					if (RequestIndex != Batch.Num())
 					{
 						UE_LOG(LogDerivedDataCache, Warning,
-							TEXT("%s: Invalid response received from PutCacheValues rpc: %d results expected, received %d."), *CacheStore.GetName(), Batch.Num(), RequestIndex);
+							TEXT("%s: Invalid response received from PutCacheValues RPC: %d results expected, received %d, from %s"),
+							*CacheStore.GetName(), Batch.Num(), RequestIndex, *WriteToString<256>(*HttpResponse));
 					}
 				}
 				for (const FCachePutValueRequest& Request : Batch.RightChop(RequestIndex))
@@ -670,7 +673,8 @@ public:
 					if (RequestIndex != Batch.Num())
 					{
 						UE_LOG(LogDerivedDataCache, Warning,
-							TEXT("Invalid response received from GetCacheValues rpc: %d results expected, received %d."), Batch.Num(), RequestIndex);
+							TEXT("%s: Invalid response received from GetCacheValues RPC: %d results expected, received %d from %s"),
+							*CacheStore.GetName(), Batch.Num(), RequestIndex, *WriteToString<256>(*HttpResponse));
 					}
 				}
 				for (const FCacheGetValueRequest& Request : Batch.RightChop(RequestIndex))
