DDC: Changed cache replay to enable compression by default

The -DDC-ReplayCompress param has been replaced by -DDC-ReplayNoCompress.

#rnx

diff --git a/Engine/Source/Developer/DerivedDataCache/Private/DerivedDataCacheReplay.cpp b/Engine/Source/Developer/DerivedDataCache/Private/DerivedDataCacheReplay.cpp
--- a/Engine/Source/Developer/DerivedDataCache/Private/DerivedDataCacheReplay.cpp
+++ b/Engine/Source/Developer/DerivedDataCache/Private/DerivedDataCacheReplay.cpp
@@ -721,7 +721,7 @@ ILegacyCacheStore* TryCreateCacheStoreReplay(ILegacyCacheStore* InnerCache)
 	if (!ReplaySavePath.IsEmpty())
 	{
 		// Create the replay cache store to save requests that pass the filters.
-		const uint64 BlockSize = FParse::Param(CommandLine, TEXT("DDC-ReplayCompress")) ? GCacheReplayCompressionBlockSize : 0;
+		const uint64 BlockSize = FParse::Param(CommandLine, TEXT("DDC-ReplayNoCompress")) ? 0 : GCacheReplayCompressionBlockSize;
 		ReplayTarget = ReplayStore = new FCacheStoreReplay(InnerCache, KeyFilter, MethodFilter, MoveTemp(ReplaySavePath), BlockSize);
 	}
 	else
