DDC: Fixed verbose logging in the file system cache

#rnx

diff --git a/Engine/Source/Developer/DerivedDataCache/Private/FileSystemCacheStore.cpp b/Engine/Source/Developer/DerivedDataCache/Private/FileSystemCacheStore.cpp
--- a/Engine/Source/Developer/DerivedDataCache/Private/FileSystemCacheStore.cpp
+++ b/Engine/Source/Developer/DerivedDataCache/Private/FileSystemCacheStore.cpp
@@ -836,13 +836,13 @@ private:
 	[[nodiscard]] bool IsDeactivatedForPerformance();
 
 private:
-	/** Base path we are storing the cache files in. */
-	FString	CachePath;
-	/** Class of this cache. */
+	/** Base path to store the cache files in. */
+	FString CachePath;
+	/** Speed class of this cache. */
 	EBackendSpeedClass SpeedClass;
 	/** If true, do not attempt to write to this cache. */
 	bool		bReadOnly;
-	/** If true, CachedDataProbablyExists will update the file timestamps. */
+	/** If true, always update file timestamps on access. */
 	bool		bTouch;
 
 	/** Age of file when it should be deleted from DDC cache. */
@@ -1531,8 +1531,6 @@ void FFileSystemCacheStore::GetChunks(
 			{
 				const uint64 RawOffset = FMath::Min(Value.GetRawSize(), Request.RawOffset);
 				RawSize = FMath::Min(Value.GetRawSize() - RawOffset, Request.RawSize);
-				UE_LOG(LogDerivedDataCache, Verbose, TEXT("%s: Cache hit for %s from '%s'"),
-					*CachePath, *WriteToString<96>(Request.Key, '/', Request.Id), *Request.Name);
 				TRACE_COUNTER_INCREMENT(FileSystemDDC_GetHit);
 				TRACE_COUNTER_ADD(FileSystemDDC_BytesRead, !bExistsOnly ? RawSize : 0);
 				COOK_STAT(Timer.AddHit(!bExistsOnly ? RawSize : 0));
@@ -1543,6 +1541,8 @@ void FFileSystemCacheStore::GetChunks(
 				Status = bExistsOnly || Buffer.GetSize() == RawSize ? EStatus::Ok : EStatus::Error;
 			}
 		}
+		UE_CLOG(Status == EStatus::Ok, LogDerivedDataCache, Verbose, TEXT("%s: Cache hit for %s from '%s'"),
+			*CachePath, *WriteToString<96>(Request.Key, '/', Request.Id), *Request.Name);
 		OnComplete({Request.Name, Request.Key, Request.Id, Request.RawOffset,
 			RawSize, Value.GetRawHash(), MoveTemp(Buffer), Request.UserData, Status});
 	}
@@ -1620,7 +1620,7 @@ bool FFileSystemCacheStore::PutCacheRecord(
 						UE_LOG(LogDerivedDataCache, Log,
 							TEXT("%s: Cache put of non-deterministic record will overwrite existing record due to "
 							     "missing value %s with hash %s for %s from '%.*s'"),
-							*CachePath, *WriteToString<16>(Value.GetId()), *WriteToString<48>(Value.GetRawHash()),
+							*CachePath, *WriteToString<32>(Value.GetId()), *WriteToString<48>(Value.GetRawHash()),
 							*WriteToString<96>(Key), Name.Len(), Name.GetData());
 						return false;
 					}
@@ -2232,6 +2232,9 @@ void FFileSystemCacheStore::GetCacheContent(
 		OutArchive = OpenFileRead(Path, Name);
 		if (OutArchive)
 		{
+			UE_LOG(LogDerivedDataCache, VeryVerbose,
+				TEXT("%s: Opened %s from '%.*s'"),
+				*CachePath, *Path, Name.Len(), Name.GetData());
 			Reader.SetSource(*OutArchive);
 			if (Reader.GetRawHash() == RawHash)
 			{
