DDC: Use TSharedRef for shared Horde Storage completion functions

#rnx

diff --git a/Engine/Source/Developer/DerivedDataCache/Private/HttpCacheStore.cpp b/Engine/Source/Developer/DerivedDataCache/Private/HttpCacheStore.cpp
--- a/Engine/Source/Developer/DerivedDataCache/Private/HttpCacheStore.cpp
+++ b/Engine/Source/Developer/DerivedDataCache/Private/HttpCacheStore.cpp
@@ -88,20 +88,6 @@ TRACE_DECLARE_INT_COUNTER(HttpDDC_PutHit, TEXT("HttpDDC Put Hit"));
 TRACE_DECLARE_INT_COUNTER(HttpDDC_BytesReceived, TEXT("HttpDDC Bytes Received"));
 TRACE_DECLARE_INT_COUNTER(HttpDDC_BytesSent, TEXT("HttpDDC Bytes Sent"));
 
-template <typename T>
-class TRefCountedUniqueFunction final : public FThreadSafeRefCountedObject
-{
-public:
-	explicit TRefCountedUniqueFunction(T&& InFunction) : Function(MoveTemp(InFunction))
-	{
-	}
-
-	const T& GetFunction() const { return Function; }
-
-private:
-	T Function;
-};
-
 static bool ShouldAbortForShutdown()
 {
 	return !GIsBuildMachine && FDerivedDataBackend::Get().IsShuttingDown();
@@ -1220,7 +1206,7 @@ void FHttpCacheStore::FGetRecordOp::GetDataBatch(
 	}
 
 	FRequestBarrier Barrier(Owner);
-	TRefCountPtr<TRefCountedUniqueFunction<FOnGetCachedDataBatchComplete>> SharedOnComplete(new TRefCountedUniqueFunction<FOnGetCachedDataBatchComplete>(MoveTemp(OnComplete)));
+	TSharedRef<FOnGetCachedDataBatchComplete> SharedOnComplete = MakeShared<FOnGetCachedDataBatchComplete>(MoveTemp(OnComplete));
 	for (int32 ValueIndex = 0; ValueIndex < Values.Num(); ++ValueIndex)
 	{
 		const ValueType Value = Values[ValueIndex].RemoveData();
@@ -1258,7 +1244,7 @@ void FHttpCacheStore::FGetRecordOp::GetDataBatch(
 			{
 				if (CompressedBuffer.GetRawHash() == Value.GetRawHash())
 				{
-					SharedOnComplete->GetFunction()({ Name, Key, ValueIndex, Operation->GetBytesReceived(), MoveTemp(CompressedBuffer), EStatus::Ok });
+					SharedOnComplete.Get()({ Name, Key, ValueIndex, Operation->GetBytesReceived(), MoveTemp(CompressedBuffer), EStatus::Ok });
 				}
 				else
 				{
@@ -1266,12 +1252,12 @@ void FHttpCacheStore::FGetRecordOp::GetDataBatch(
 						TEXT("%s: Cache miss with corrupted value %s with hash %s for %s from '%s'"),
 						*CacheStore.Domain, *ValueIdGetter(Value), *WriteToString<48>(Value.GetRawHash()),
 						*WriteToString<96>(Key), *Name);
-					SharedOnComplete->GetFunction()({ Name, Key, ValueIndex, Operation->GetBytesReceived(), {}, EStatus::Error });
+					SharedOnComplete.Get()({ Name, Key, ValueIndex, Operation->GetBytesReceived(), {}, EStatus::Error });
 				}
 			}
 			else if (Operation->GetErrorCode() == EHttpErrorCode::Canceled)
 			{
-				SharedOnComplete->GetFunction()({ Name, Key, ValueIndex, Operation->GetBytesReceived(), {}, EStatus::Canceled });
+				SharedOnComplete.Get()({ Name, Key, ValueIndex, Operation->GetBytesReceived(), {}, EStatus::Canceled });
 			}
 			else
 			{
@@ -1279,7 +1265,7 @@ void FHttpCacheStore::FGetRecordOp::GetDataBatch(
 					TEXT("%s: Cache miss with missing value %s with hash %s for %s from '%s'"),
 					*CacheStore.Domain, *ValueIdGetter(Value), *WriteToString<48>(Value.GetRawHash()), *WriteToString<96>(Key),
 					*Name);
-				SharedOnComplete->GetFunction()({ Name, Key, ValueIndex, Operation->GetBytesReceived(), {}, EStatus::Error });
+				SharedOnComplete.Get()({ Name, Key, ValueIndex, Operation->GetBytesReceived(), {}, EStatus::Error });
 			}
 		});
 	}
@@ -2324,18 +2310,18 @@ void FHttpCacheStore::Put(
 {
 	TRACE_CPUPROFILER_EVENT_SCOPE(HttpDDC_Put);
 	FRequestBarrier Barrier(Owner);
-	TRefCountedUniqueFunction<FOnCachePutComplete>* CompletionFunction = new TRefCountedUniqueFunction<FOnCachePutComplete>(MoveTemp(OnComplete));
-	TRefCountPtr<TRefCountedUniqueFunction<FOnCachePutComplete>> BatchOnCompleteRef(CompletionFunction);
+	TSharedRef<FOnCachePutComplete> SharedOnComplete = MakeShared<FOnCachePutComplete>(MoveTemp(OnComplete));
 	for (const FCachePutRequest& Request : Requests)
 	{
-		PutCacheRecordAsync(Owner, Request.Name, Request.Record, Request.Policy, Request.UserData, [COOK_STAT(Timer = UsageStats.TimePut(), ) OnCompletePtr = TRefCountPtr<TRefCountedUniqueFunction<FOnCachePutComplete>>(CompletionFunction)](FCachePutResponse&& Response, uint64 BytesSent) mutable
+		PutCacheRecordAsync(Owner, Request.Name, Request.Record, Request.Policy, Request.UserData,
+			[COOK_STAT(Timer = UsageStats.TimePut(), ) SharedOnComplete](FCachePutResponse&& Response, uint64 BytesSent) mutable
 		{
 			TRACE_COUNTER_ADD(HttpDDC_BytesSent, BytesSent);
 			if (Response.Status == EStatus::Ok)
 			{
 				COOK_STAT(if (BytesSent) { Timer.AddHit(BytesSent); });
 			}
-			OnCompletePtr->GetFunction()(MoveTemp(Response));
+			SharedOnComplete.Get()(MoveTemp(Response));
 		});
 	}
 }
@@ -2347,18 +2333,18 @@ void FHttpCacheStore::Get(
 {
 	TRACE_CPUPROFILER_EVENT_SCOPE(HttpDDC_Get);
 	FRequestBarrier Barrier(Owner);
-	TRefCountedUniqueFunction<FOnCacheGetComplete>* CompletionFunction = new TRefCountedUniqueFunction<FOnCacheGetComplete>(MoveTemp(OnComplete));
-	TRefCountPtr<TRefCountedUniqueFunction<FOnCacheGetComplete>> BatchOnCompleteRef(CompletionFunction);
+	TSharedRef<FOnCacheGetComplete> SharedOnComplete = MakeShared<FOnCacheGetComplete>(MoveTemp(OnComplete));
 	for (const FCacheGetRequest& Request : Requests)
 	{
-		GetCacheRecordAsync(Owner, Request.Name, Request.Key, Request.Policy, Request.UserData, [COOK_STAT(Timer = UsageStats.TimePut(), ) OnCompletePtr = TRefCountPtr<TRefCountedUniqueFunction<FOnCacheGetComplete>>(CompletionFunction)](FCacheGetResponse&& Response, uint64 BytesReceived) mutable
+		GetCacheRecordAsync(Owner, Request.Name, Request.Key, Request.Policy, Request.UserData,
+			[COOK_STAT(Timer = UsageStats.TimePut(), ) SharedOnComplete](FCacheGetResponse&& Response, uint64 BytesReceived) mutable
 		{
 			TRACE_COUNTER_ADD(HttpDDC_BytesReceived, BytesReceived);
 			if (Response.Status == EStatus::Ok)
 			{
 				COOK_STAT(Timer.AddHit(BytesReceived););
 			}
-			OnCompletePtr->GetFunction()(MoveTemp(Response));
+			SharedOnComplete.Get()(MoveTemp(Response));
 		});
 	}
 }
@@ -2370,18 +2356,18 @@ void FHttpCacheStore::PutValue(
 {
 	TRACE_CPUPROFILER_EVENT_SCOPE(HttpDDC_PutValue);
 	FRequestBarrier Barrier(Owner);
-	TRefCountedUniqueFunction<FOnCachePutValueComplete>* CompletionFunction = new TRefCountedUniqueFunction<FOnCachePutValueComplete>(MoveTemp(OnComplete));
-	TRefCountPtr<TRefCountedUniqueFunction<FOnCachePutValueComplete>> BatchOnCompleteRef(CompletionFunction);
+	TSharedRef<FOnCachePutValueComplete> SharedOnComplete = MakeShared<FOnCachePutValueComplete>(MoveTemp(OnComplete));
 	for (const FCachePutValueRequest& Request : Requests)
 	{
-		PutCacheValueAsync(Owner, Request.Name, Request.Key, Request.Value, Request.Policy, Request.UserData, [COOK_STAT(Timer = UsageStats.TimePut(),) OnCompletePtr = TRefCountPtr<TRefCountedUniqueFunction<FOnCachePutValueComplete>>(CompletionFunction)](FCachePutValueResponse&& Response, uint64 BytesSent) mutable
+		PutCacheValueAsync(Owner, Request.Name, Request.Key, Request.Value, Request.Policy, Request.UserData,
+			[COOK_STAT(Timer = UsageStats.TimePut(),) SharedOnComplete](FCachePutValueResponse&& Response, uint64 BytesSent) mutable
 		{
 			TRACE_COUNTER_ADD(HttpDDC_BytesSent, BytesSent);
 			if (Response.Status == EStatus::Ok)
 			{
 				COOK_STAT(if (BytesSent) { Timer.AddHit(BytesSent); });
 			}
-			OnCompletePtr->GetFunction()(MoveTemp(Response));
+			SharedOnComplete.Get()(MoveTemp(Response));
 		});
 	}
 }
@@ -2429,15 +2415,14 @@ void FHttpCacheStore::GetValue(
 	else
 	{
 		FRequestBarrier Barrier(Owner);
-		TRefCountedUniqueFunction<FOnCacheGetValueComplete>* CompletionFunction = new TRefCountedUniqueFunction<FOnCacheGetValueComplete>(MoveTemp(OnComplete));
-		TRefCountPtr<TRefCountedUniqueFunction<FOnCacheGetValueComplete>> BatchOnCompleteRef(CompletionFunction);
+		TSharedRef<FOnCacheGetValueComplete> SharedOnComplete = MakeShared<FOnCacheGetValueComplete>(MoveTemp(OnComplete));
 		int64 HitBytes = 0;
 		for (const FCacheGetValueRequest& Request : Requests)
 		{
 			GetCacheValueAsync(Owner, Request.Name, Request.Key, Request.Policy, Request.UserData,
-			[this, COOK_STAT(StartTime, bIsInGameThread,) Policy = Request.Policy, OnCompletePtr = TRefCountPtr<TRefCountedUniqueFunction<FOnCacheGetValueComplete>>(CompletionFunction)] (FCacheGetValueResponse&& Response)
+			[this, COOK_STAT(StartTime, bIsInGameThread,) Policy = Request.Policy, SharedOnComplete](FCacheGetValueResponse&& Response)
 			{
-				const FOnCacheGetValueComplete& OnComplete = OnCompletePtr->GetFunction();
+				const FOnCacheGetValueComplete& OnComplete = SharedOnComplete.Get();
 				check(OnComplete);
 				if (Response.Status != EStatus::Ok)
 				{
