DDC: Optimized Horde Storage cache bucket transforms

#rnx

diff --git a/Engine/Source/Developer/DerivedDataCache/Private/HttpCacheStore.cpp b/Engine/Source/Developer/DerivedDataCache/Private/HttpCacheStore.cpp
--- a/Engine/Source/Developer/DerivedDataCache/Private/HttpCacheStore.cpp
+++ b/Engine/Source/Developer/DerivedDataCache/Private/HttpCacheStore.cpp
@@ -6,6 +6,7 @@
 
 #if WITH_HTTP_DDC_BACKEND
 
+#include "Algo/Transform.h"
 #include "Compression/CompressedBuffer.h"
 #include "Containers/StringView.h"
 #include "Containers/Ticker.h"
@@ -936,8 +937,8 @@ void FHttpCacheStore::FPutPackageOp::PutRefAsync(
 	bool bFinalize,
 	FOnCachePutRefComplete&& OnComplete)
 {
-	FString Bucket(Key.Bucket.ToString());
-	Bucket.ToLowerInline();
+	TAnsiStringBuilder<64> Bucket;
+	Algo::Transform(Key.Bucket.ToString(), AppendChars(Bucket), FCharAnsi::ToLower);
 
 	TAnsiStringBuilder<256> RefsUri;
 	RefsUri << CacheStore.EffectiveDomain << ANSITEXTVIEW("/api/v1/refs/") << CacheStore.StructuredNamespace << '/' << Bucket << '/' << Key.Hash;
@@ -1872,8 +1873,8 @@ void FHttpCacheStore::GetCacheRecordOnlyAsync(
 		return OnComplete(MakeResponse(0, EStatus::Error));
 	}
 
-	FString Bucket(Key.Bucket.ToString());
-	Bucket.ToLowerInline();
+	TAnsiStringBuilder<64> Bucket;
+	Algo::Transform(Key.Bucket.ToString(), AppendChars(Bucket), FCharAnsi::ToLower);
 
 	TUniquePtr<FHttpOperation> Operation = WaitForHttpOperation(EOperationCategory::Get, /*bUnboundedOverflow*/ false);
 	FHttpOperation& LocalOperation = *Operation;
@@ -1957,9 +1958,6 @@ void FHttpCacheStore::PutCacheRecordAsync(
 	// TODO: Jupiter currently always overwrites.  It doesn't have a "write if not present" feature (for records or attachments),
 	//		 but would require one to implement all policy correctly.
 
-	FString Bucket(Key.Bucket.ToString());
-	Bucket.ToLowerInline();
-
 	FCbPackage Package = Record.Save();
 
 	FPutPackageOp::PutPackage(*this, Owner, Name, Key, MoveTemp(Package), Policy, UserData, [MakeResponse = MoveTemp(MakeResponse), OnComplete = MoveTemp(OnComplete)](FPutPackageOp::FCachePutPackageResponse&& Response)
@@ -2008,9 +2006,6 @@ void FHttpCacheStore::PutCacheValueAsync(
 	// TODO: Jupiter currently always overwrites.  It doesn't have a "write if not present" feature (for records or attachments),
 	//		 but would require one to implement all policy correctly.
 
-	FString Bucket(Key.Bucket.ToString());
-	Bucket.ToLowerInline();
-
 	FCbWriter Writer;
 	Writer.BeginObject();
 	Writer.AddBinaryAttachment("RawHash", Value.GetRawHash());
@@ -2062,8 +2057,8 @@ void FHttpCacheStore::GetCacheValueAsync(
 
 	const bool bSkipData = EnumHasAnyFlags(Policy, ECachePolicy::SkipData);
 
-	FString Bucket(Key.Bucket.ToString());
-	Bucket.ToLowerInline();
+	TAnsiStringBuilder<64> Bucket;
+	Algo::Transform(Key.Bucket.ToString(), AppendChars(Bucket), FCharAnsi::ToLower);
 
 	TUniquePtr<FHttpOperation> Operation = WaitForHttpOperation(EOperationCategory::Get, /*bUnboundedOverflow*/ false);
 	FHttpOperation& LocalOperation = *Operation;
@@ -2187,9 +2182,9 @@ void FHttpCacheStore::RefCachedDataProbablyExistsBatchAsync(
 		RequestWriter.BeginObject();
 		RequestWriter.AddInteger(ANSITEXTVIEW("opId"), OpIndex);
 		RequestWriter.AddString(ANSITEXTVIEW("op"), ANSITEXTVIEW("GET"));
-		FCacheKey Key = ValueRef.Key;
-		FString Bucket(Key.Bucket.ToString());
-		Bucket.ToLowerInline();
+		const FCacheKey& Key = ValueRef.Key;
+		TAnsiStringBuilder<64> Bucket;
+		Algo::Transform(Key.Bucket.ToString(), AppendChars(Bucket), FCharAnsi::ToLower);
 		RequestWriter.AddString(ANSITEXTVIEW("bucket"), Bucket);
 		RequestWriter.AddString(ANSITEXTVIEW("key"), LexToString(Key.Hash));
 		RequestWriter.AddBool(ANSITEXTVIEW("resolveAttachments"), true);
