DDC: Separated the async completion counter from the cache singleton

#rnx

diff --git a/Engine/Source/Developer/DerivedDataCache/Private/DerivedDataBackends.cpp b/Engine/Source/Developer/DerivedDataCache/Private/DerivedDataBackends.cpp
--- a/Engine/Source/Developer/DerivedDataCache/Private/DerivedDataBackends.cpp
+++ b/Engine/Source/Developer/DerivedDataCache/Private/DerivedDataBackends.cpp
@@ -46,6 +46,18 @@ static TAutoConsoleVariable<FString> GDerivedDataCacheGraphName(
 	TEXT("Name of the graph to use for the Derived Data Cache."),
 	ECVF_ReadOnly);
 
+namespace UE::DerivedData::Private
+{
+
+static std::atomic<int32> GAsyncTaskCounter;
+
+int32 AddToAsyncTaskCounter(int32 Addend)
+{
+	return GAsyncTaskCounter.fetch_add(Addend);
+}
+
+} // UE::DerivedData::Private
+
 namespace UE::DerivedData
 {
 
@@ -907,9 +919,9 @@ public:
 			bIsShuttingDown.store(true, std::memory_order_relaxed);
 		}
 
-		while (AsyncCompletionCounter.GetValue())
+		while (const int32 AsyncCompletionCounter = Private::AddToAsyncTaskCounter(0))
 		{
-			check(AsyncCompletionCounter.GetValue() >= 0);
+			check(AsyncCompletionCounter > 0);
 			FPlatformProcess::Sleep(0.1f);
 			if (FPlatformTime::Seconds() - LastPrint > 5.0)
 			{
@@ -988,13 +1000,12 @@ public:
 
 	virtual void AddToAsyncCompletionCounter(int32 Addend) override
 	{
-		AsyncCompletionCounter.Add(Addend);
-		check(AsyncCompletionCounter.GetValue() >= 0);
+		verify(Private::AddToAsyncTaskCounter(Addend) + Addend >= 0);
 	}
 
 	virtual bool AnyAsyncRequestsRemaining() override
 	{
-		return AsyncCompletionCounter.GetValue() > 0;
+		return Private::AddToAsyncTaskCounter(0) > 0;
 	}
 
 	virtual bool IsShuttingDown() override
@@ -1187,7 +1198,6 @@ private:
 
 	static inline FDerivedDataBackendGraph*			StaticGraph;
 
-	FThreadSafeCounter								AsyncCompletionCounter;
 	FString											GraphName;
 	FString											ReadPakFilename;
 	FString											WritePakFilename;
diff --git a/Engine/Source/Developer/DerivedDataCache/Private/DerivedDataCachePrivate.h b/Engine/Source/Developer/DerivedDataCache/Private/DerivedDataCachePrivate.h
--- a/Engine/Source/Developer/DerivedDataCache/Private/DerivedDataCachePrivate.h
+++ b/Engine/Source/Developer/DerivedDataCache/Private/DerivedDataCachePrivate.h
@@ -15,6 +15,9 @@ namespace UE::DerivedData { class IRequestOwner; }
 namespace UE::DerivedData::Private
 {
 
+// Implemented in DerivedDataBackends.cpp
+int32 AddToAsyncTaskCounter(int32 Addend); // returns the previous value
+
 // Implemented in DerivedDataCache.cpp
 ICache* CreateCache(FDerivedDataCacheInterface** OutLegacyCache);
 void LaunchTaskInCacheThreadPool(IRequestOwner& Owner, TUniqueFunction<void ()>&& TaskBody);
diff --git a/Engine/Source/Developer/DerivedDataCache/Private/DerivedDataCacheStoreAsync.cpp b/Engine/Source/Developer/DerivedDataCache/Private/DerivedDataCacheStoreAsync.cpp
--- a/Engine/Source/Developer/DerivedDataCache/Private/DerivedDataCacheStoreAsync.cpp
+++ b/Engine/Source/Developer/DerivedDataCache/Private/DerivedDataCacheStoreAsync.cpp
@@ -154,11 +154,11 @@ void FCacheStoreAsync::Execute(
 				COOK_STAT(Timer.AddHit(0));
 			}
 			OnComplete(MoveTemp(Response));
-			FDerivedDataBackend::Get().AddToAsyncCompletionCounter(-1);
+			Private::AddToAsyncTaskCounter(-1);
 		});
 	};
 
-	FDerivedDataBackend::Get().AddToAsyncCompletionCounter(Requests.Num());
+	Private::AddToAsyncTaskCounter(Requests.Num());
 	if (Owner.GetPriority() == EPriority::Blocking)
 	{
 		return ExecuteWithStats(Requests, Owner, MoveTemp(OnComplete));
@@ -178,7 +178,7 @@ void FCacheStoreAsync::Execute(
 			else
 			{
 				CompleteWithStatus(Requests, MoveTemp(OnComplete), EStatus::Canceled);
-				FDerivedDataBackend::Get().AddToAsyncCompletionCounter(-Requests.Num());
+				Private::AddToAsyncTaskCounter(-Requests.Num());
 			}
 		});
 }
