DDC: Added a mock cache store for testing the cache

#rnx

diff --git a/Engine/Source/Developer/DerivedDataCache/Tests/TestCacheStore.cpp b/Engine/Source/Developer/DerivedDataCache/Tests/TestCacheStore.cpp
new file mode 100644
--- /dev/null
+++ b/Engine/Source/Developer/DerivedDataCache/Tests/TestCacheStore.cpp
@@ -0,0 +1,511 @@
+// Copyright Epic Games, Inc. All Rights Reserved.
+
+#if WITH_LOW_LEVEL_TESTS
+
+#include "TestCacheStore.h"
+
+#include "Algo/BinarySearch.h"
+#include "Containers/Array.h"
+#include "Containers/Map.h"
+#include "DerivedDataRequest.h"
+#include "DerivedDataRequestOwner.h"
+#include "Templates/Function.h"
+
+namespace UE::DerivedData
+{
+
+class FTestCacheStore final : public ITestCacheStore
+{
+public:
+	explicit FTestCacheStore(ECacheStoreFlags Flags, bool bAsync = false);
+	FTestCacheStore(const FTestCacheStore&) = delete;
+	FTestCacheStore& operator=(const FTestCacheStore&) = delete;
+
+	ECacheStoreFlags GetFlags() const final { return CacheStoreFlags; }
+	uint32 GetTotalRequestCount() const final { return TotalRequestCount; }
+	uint32 GetCanceledRequestCount() const final { return CanceledRequestCount; }
+
+	void AddRecord(const FCacheKey& Key, TConstArrayView<FValueWithId> Values, const FCbObject* Meta = nullptr) final;
+	TConstArrayView<FValueWithId> FindRecord(const FCacheKey& Key, FCbObject* OutMeta = nullptr) const final;
+
+	void AddValue(const FCacheKey& Key, const FValue& Value) final;
+	FValue FindValue(const FCacheKey& Key) const final;
+
+	void AddContent(const FCompressedBuffer& Content) final;
+	FCompressedBuffer FindContent(const FIoHash& RawHash, uint64 RawSize) const final;
+
+	TConstArrayView<FCachePutRequest> GetPutRequests() const final { return PutRequests; }
+	TConstArrayView<FCacheGetRequest> GetGetRequests() const final { return GetRequests; }
+	TConstArrayView<FCachePutValueRequest> GetPutValueRequests() const final { return PutValueRequests; }
+	TConstArrayView<FCacheGetValueRequest> GetGetValueRequests() const final { return GetValueRequests; }
+	TConstArrayView<FCacheGetChunkRequest> GetGetChunkRequests() const final { return GetChunkRequests; }
+
+	void ExecuteAsync() final;
+
+private:
+	void Put(
+		TConstArrayView<FCachePutRequest> Requests,
+		IRequestOwner& Owner,
+		FOnCachePutComplete&& OnComplete) final;
+	void Get(
+		TConstArrayView<FCacheGetRequest> Requests,
+		IRequestOwner& Owner,
+		FOnCacheGetComplete&& OnComplete) final;
+	void PutValue(
+		TConstArrayView<FCachePutValueRequest> Requests,
+		IRequestOwner& Owner,
+		FOnCachePutValueComplete&& OnComplete) final;
+	void GetValue(
+		TConstArrayView<FCacheGetValueRequest> Requests,
+		IRequestOwner& Owner,
+		FOnCacheGetValueComplete&& OnComplete) final;
+	void GetChunks(
+		TConstArrayView<FCacheGetChunkRequest> Requests,
+		IRequestOwner& Owner,
+		FOnCacheGetChunkComplete&& OnComplete) final;
+
+	void LegacyStats(FDerivedDataCacheStatsNode& OutNode) final {}
+	bool LegacyDebugOptions(FBackendDebugOptions& Options) final { return false; }
+
+	void ExecuteOrQueue(IRequestOwner& Owner, TUniqueFunction<void (bool bCancel)>&& Function);
+
+	struct FRecord
+	{
+		FCbObject Meta;
+		mutable TArray<FValueWithId> Values;
+	};
+
+	TArray<FCachePutRequest> PutRequests;
+	TArray<FCacheGetRequest> GetRequests;
+	TArray<FCachePutValueRequest> PutValueRequests;
+	TArray<FCacheGetValueRequest> GetValueRequests;
+	TArray<FCacheGetChunkRequest> GetChunkRequests;
+
+	TMap<FIoHash, FCompressedBuffer> ContentMap;
+	TMap<FCacheKey, FRecord> RecordMap;
+	TMap<FCacheKey, FValue> ValueMap;
+
+	class FAsyncRequest;
+	TArray<TRefCountPtr<FAsyncRequest>> Queue;
+
+	uint32 TotalRequestCount = 0;
+	uint32 CanceledRequestCount = 0;
+	ECachePolicy CacheStorePolicy;
+	const ECacheStoreFlags CacheStoreFlags;
+	const bool bAsync;
+};
+
+class FTestCacheStore::FAsyncRequest final : public FRequestBase
+{
+public:
+	FAsyncRequest(IRequestOwner& InOwner, TUniqueFunction<void (bool bCancel)>&& InBody)
+		: Owner(InOwner)
+		, Body(MoveTemp(InBody))
+	{
+		Owner.Begin(this);
+	}
+
+	void SetPriority(EPriority Priority) final
+	{
+	}
+
+	void Cancel() final
+	{
+		if (Body)
+		{
+			auto MovedBody = MoveTemp(Body);
+			Body.Reset();
+			Owner.End(this, MovedBody, true);
+		}
+	}
+
+	void Wait() final
+	{
+		if (Body)
+		{
+			auto MovedBody = MoveTemp(Body);
+			Body.Reset();
+			Owner.End(this, MovedBody, false);
+		}
+	}
+
+private:
+	IRequestOwner& Owner;
+	TUniqueFunction<void (bool bCancel)> Body;
+};
+
+FTestCacheStore::FTestCacheStore(const ECacheStoreFlags Flags, const bool bInAsync)
+	: CacheStorePolicy(ECachePolicy::None)
+	, CacheStoreFlags(Flags)
+	, bAsync(bInAsync)
+{
+	CacheStorePolicy |= EnumHasAllFlags(Flags, ECacheStoreFlags::Query | ECacheStoreFlags::Local) ? ECachePolicy::QueryLocal : ECachePolicy::None;
+	CacheStorePolicy |= EnumHasAllFlags(Flags, ECacheStoreFlags::Store | ECacheStoreFlags::Local) ? ECachePolicy::StoreLocal : ECachePolicy::None;
+	CacheStorePolicy |= EnumHasAllFlags(Flags, ECacheStoreFlags::Query | ECacheStoreFlags::Remote) ? ECachePolicy::QueryRemote : ECachePolicy::None;
+	CacheStorePolicy |= EnumHasAllFlags(Flags, ECacheStoreFlags::Store | ECacheStoreFlags::Remote) ? ECachePolicy::StoreRemote : ECachePolicy::None;
+}
+
+void FTestCacheStore::AddRecord(const FCacheKey& Key, TConstArrayView<FValueWithId> Values, const FCbObject* Meta)
+{
+	FRecord& Record = RecordMap.FindOrAdd(Key);
+
+	if (Meta && *Meta)
+	{
+		Record.Meta = *Meta;
+	}
+
+	Record.Values = Values;
+	Record.Values.Sort();
+
+	for (FValueWithId& Value : Record.Values)
+	{
+		AddContent(Value.GetData());
+	}
+}
+
+TConstArrayView<FValueWithId> FTestCacheStore::FindRecord(const FCacheKey& Key, FCbObject* OutMeta) const
+{
+	if (const FRecord* Record = RecordMap.Find(Key))
+	{
+		if (OutMeta)
+		{
+			*OutMeta = Record->Meta;
+		}
+		for (FValueWithId& Value : Record->Values)
+		{
+			if (!Value.HasData())
+			{
+				if (FCompressedBuffer Content = FindContent(Value.GetRawHash(), Value.GetRawSize()))
+				{
+					Value = FValueWithId(Value.GetId(), MoveTemp(Content));
+				}
+			}
+		}
+		return Record->Values;
+	}
+	return {};
+}
+
+void FTestCacheStore::AddValue(const FCacheKey& Key, const FValue& Value)
+{
+	AddContent(Value.GetData());
+	ValueMap.Emplace(Key, Value);
+}
+
+FValue FTestCacheStore::FindValue(const FCacheKey& Key) const
+{
+	if (const FValue* Value = ValueMap.Find(Key))
+	{
+		return *Value;
+	}
+	return {};
+}
+
+void FTestCacheStore::AddContent(const FCompressedBuffer& Content)
+{
+	if (Content)
+	{
+		ContentMap.Emplace(Content.GetRawHash(), Content);
+	}
+}
+
+FCompressedBuffer FTestCacheStore::FindContent(const FIoHash& RawHash, uint64 RawSize) const
+{
+	if (const FCompressedBuffer* Content = ContentMap.Find(RawHash); Content && Content->GetRawSize() == RawSize)
+	{
+		return *Content;
+	}
+	return {};
+}
+
+void FTestCacheStore::Put(
+	TConstArrayView<FCachePutRequest> Requests,
+	IRequestOwner& Owner,
+	FOnCachePutComplete&& OnComplete)
+{
+	const int32 Index = PutRequests.Num();
+	const int32 Count = Requests.Num();
+	PutRequests.Append(Requests);
+	TotalRequestCount += Count;
+
+	ExecuteOrQueue(Owner, [this, Index, Count, OnComplete = MoveTemp(OnComplete)](bool bCancel)
+	{
+		TConstArrayView<FCachePutRequest> Requests = MakeArrayView(PutRequests).Mid(Index, Count);
+
+		if (bCancel)
+		{
+			CanceledRequestCount += Count;
+			return CompleteWithStatus(Requests, OnComplete, EStatus::Canceled);
+		}
+
+		for (const FCachePutRequest& Request : Requests)
+		{
+			bool bOk = false;
+
+			const ECachePolicy RecordPolicy = Request.Policy.GetRecordPolicy();
+			if (EnumHasAnyFlags(CacheStorePolicy, RecordPolicy & ECachePolicy::Store))
+			{
+				bOk = true;
+
+				const FCacheKey& Key = Request.Record.GetKey();
+
+				const bool bPartial = EnumHasAnyFlags(RecordPolicy, ECachePolicy::PartialRecord);
+				TArray<FValueWithId> Values(Request.Record.GetValues());
+				for (FValueWithId& Value : Values)
+				{
+					const ECachePolicy ValuePolicy = Request.Policy.GetValuePolicy(Value.GetId());
+					if (EnumHasAnyFlags(CacheStorePolicy, ValuePolicy & ECachePolicy::Store))
+					{
+						if (Value.HasData())
+						{
+							AddContent(Value.GetData());
+						}
+						else if (FCompressedBuffer Content = FindContent(Value.GetRawHash(), Value.GetRawSize()))
+						{
+							Value = FValueWithId(Value.GetId(), MoveTemp(Content));
+						}
+						else
+						{
+							bOk &= bPartial;
+						}
+					}
+				}
+
+				bool bExistingOk = false;
+				if (EnumHasAnyFlags(CacheStorePolicy, RecordPolicy & ECachePolicy::Query))
+				{
+					bExistingOk = RecordMap.Contains(Key);
+					for (const FValueWithId& Value : FindRecord(Key))
+					{
+						const ECachePolicy ValuePolicy = Request.Policy.GetValuePolicy(Value.GetId());
+						bExistingOk &= !EnumHasAnyFlags(CacheStorePolicy, ValuePolicy & ECachePolicy::Query) || Value.HasData();
+					}
+				}
+
+				if (bOk && !bExistingOk)
+				{
+					AddRecord(Key, Values, &Request.Record.GetMeta());
+				}
+			}
+
+			OnComplete(Request.MakeResponse(bOk ? EStatus::Ok : EStatus::Error));
+		}
+	});
+}
+
+void FTestCacheStore::Get(
+	TConstArrayView<FCacheGetRequest> Requests,
+	IRequestOwner& Owner,
+	FOnCacheGetComplete&& OnComplete)
+{
+	const int32 Index = GetRequests.Num();
+	const int32 Count = Requests.Num();
+	GetRequests.Append(Requests);
+	TotalRequestCount += Count;
+
+	ExecuteOrQueue(Owner, [this, Index, Count, OnComplete = MoveTemp(OnComplete)](bool bCancel)
+	{
+		TConstArrayView<FCacheGetRequest> Requests = MakeArrayView(GetRequests).Mid(Index, Count);
+
+		if (bCancel)
+		{
+			CanceledRequestCount += Count;
+			return CompleteWithStatus(Requests, OnComplete, EStatus::Canceled);
+		}
+
+		for (const FCacheGetRequest& Request : Requests)
+		{
+			bool bOk = false;
+			FCacheRecordBuilder RecordBuilder(Request.Key);
+
+			const ECachePolicy RecordPolicy = Request.Policy.GetRecordPolicy();
+			if (EnumHasAnyFlags(CacheStorePolicy, RecordPolicy & ECachePolicy::Query))
+			{
+				bOk = true;
+
+				FCbObject Meta;
+				for (const FValueWithId& Value : FindRecord(Request.Key, &Meta))
+				{
+					const ECachePolicy ValuePolicy = Request.Policy.GetValuePolicy(Value.GetId());
+					const bool bSkipData = EnumHasAnyFlags(ValuePolicy, ECachePolicy::SkipData);
+					const bool bQuery = EnumHasAnyFlags(ValuePolicy, ECachePolicy::Query);
+					bOk &= !bQuery || Value.HasData();
+					RecordBuilder.AddValue(bQuery && !bSkipData ? Value : Value.RemoveData());
+				}
+
+				if (!bOk && !EnumHasAnyFlags(RecordPolicy, ECachePolicy::PartialRecord))
+				{
+					RecordBuilder = FCacheRecordBuilder(Request.Key);
+				}
+			}
+
+			OnComplete({Request.Name, RecordBuilder.Build(), Request.UserData, bOk ? EStatus::Ok : EStatus::Error});
+		}
+	});
+}
+
+void FTestCacheStore::PutValue(
+	TConstArrayView<FCachePutValueRequest> Requests,
+	IRequestOwner& Owner,
+	FOnCachePutValueComplete&& OnComplete)
+{
+	const int32 Index = PutValueRequests.Num();
+	const int32 Count = Requests.Num();
+	PutValueRequests.Append(Requests);
+	TotalRequestCount += Count;
+
+	ExecuteOrQueue(Owner, [this, Index, Count, OnComplete = MoveTemp(OnComplete)](bool bCancel)
+	{
+		TConstArrayView<FCachePutValueRequest> Requests = MakeArrayView(PutValueRequests).Mid(Index, Count);
+
+		if (bCancel)
+		{
+			CanceledRequestCount += Count;
+			return CompleteWithStatus(Requests, OnComplete, EStatus::Canceled);
+		}
+
+		for (const FCachePutValueRequest& Request : Requests)
+		{
+			const bool bOk = EnumHasAnyFlags(CacheStorePolicy, Request.Policy & ECachePolicy::Store);
+			if (bOk && (!EnumHasAnyFlags(CacheStorePolicy, Request.Policy & ECachePolicy::Query) || !FindValue(Request.Key).HasData()))
+			{
+				AddValue(Request.Key, Request.Value);
+			}
+			OnComplete(Request.MakeResponse(bOk ? EStatus::Ok : EStatus::Error));
+		}
+	});
+}
+
+void FTestCacheStore::GetValue(
+	TConstArrayView<FCacheGetValueRequest> Requests,
+	IRequestOwner& Owner,
+	FOnCacheGetValueComplete&& OnComplete)
+{
+	const int32 Index = GetValueRequests.Num();
+	const int32 Count = Requests.Num();
+	GetValueRequests.Append(Requests);
+	TotalRequestCount += Count;
+
+	ExecuteOrQueue(Owner, [this, Index, Count, OnComplete = MoveTemp(OnComplete)](bool bCancel)
+	{
+		TConstArrayView<FCacheGetValueRequest> Requests = MakeArrayView(GetValueRequests).Mid(Index, Count);
+
+		if (bCancel)
+		{
+			CanceledRequestCount += Count;
+			return CompleteWithStatus(Requests, OnComplete, EStatus::Canceled);
+		}
+
+		for (const FCacheGetValueRequest& Request : Requests)
+		{
+			bool bOk = false;
+			FValue Value;
+			if (EnumHasAnyFlags(CacheStorePolicy, Request.Policy & ECachePolicy::Query))
+			{
+				Value = FindValue(Request.Key);
+				if (Value.HasData())
+				{
+					bOk = true;
+					if (EnumHasAnyFlags(Request.Policy, ECachePolicy::SkipData))
+					{
+						Value = Value.RemoveData();
+					}
+				}
+			}
+			OnComplete({Request.Name, Request.Key, MoveTemp(Value), Request.UserData, bOk ? EStatus::Ok : EStatus::Error});
+		}
+	});
+}
+
+void FTestCacheStore::GetChunks(
+	TConstArrayView<FCacheGetChunkRequest> Requests,
+	IRequestOwner& Owner,
+	FOnCacheGetChunkComplete&& OnComplete)
+{
+	const int32 Index = GetChunkRequests.Num();
+	const int32 Count = Requests.Num();
+	GetChunkRequests.Append(Requests);
+	TotalRequestCount += Count;
+
+	ExecuteOrQueue(Owner, [this, Index, Count, OnComplete = MoveTemp(OnComplete)](bool bCancel)
+	{
+		TConstArrayView<FCacheGetChunkRequest> Requests = MakeArrayView(GetChunkRequests).Mid(Index, Count);
+
+		if (bCancel)
+		{
+			CanceledRequestCount += Count;
+			return CompleteWithStatus(Requests, OnComplete, EStatus::Canceled);
+		}
+
+		for (const FCacheGetChunkRequest& Request : Requests)
+		{
+			FIoHash RawHash;
+			uint64 RawSize = 0;
+			FSharedBuffer RawData;
+			EStatus Status = EStatus::Error;
+			if (EnumHasAnyFlags(CacheStorePolicy, Request.Policy & ECachePolicy::Query))
+			{
+				FValue Value;
+				if (Request.Id.IsNull())
+				{
+					Value = FindValue(Request.Key);
+				}
+				else
+				{
+					TConstArrayView<FValueWithId> Values = FindRecord(Request.Key);
+					const int32 ValueIndex = Algo::BinarySearchBy(Values, Request.Id, &FValueWithId::GetId);
+					Value = Values.IsValidIndex(ValueIndex) ? Values[ValueIndex] : FValue::Null;
+				}
+				if (!(Value == FValue::Null))
+				{
+					const uint64 RawOffset = FMath::Min(Value.GetRawSize(), Request.RawOffset);
+					RawSize = FMath::Min(Value.GetRawSize() - RawOffset, Request.RawSize);
+					RawHash = Value.GetRawHash();
+					const bool bExistsOnly = !EnumHasAnyFlags(Request.Policy, ECachePolicy::SkipData);
+					if (!bExistsOnly)
+					{
+						FCompressedBufferReader ValueReader(Value.GetData());
+						RawData = ValueReader.Decompress(RawOffset, RawSize);
+					}
+					Status = bExistsOnly || RawData.GetSize() == RawSize ? EStatus::Ok : EStatus::Error;
+				}
+			}
+			OnComplete({Request.Name, Request.Key, Request.Id, Request.RawOffset,
+				RawSize, RawHash, MoveTemp(RawData), Request.UserData, Status});
+		}
+	});
+}
+
+void FTestCacheStore::ExecuteOrQueue(IRequestOwner& Owner, TUniqueFunction<void (bool bCancel)>&& Function)
+{
+	if (bAsync)
+	{
+		Queue.Emplace(new FAsyncRequest(Owner, MoveTemp(Function)));
+	}
+	else
+	{
+		Function(false);
+	}
+}
+
+void FTestCacheStore::ExecuteAsync()
+{
+	while (!Queue.IsEmpty())
+	{
+		TArray<TRefCountPtr<FAsyncRequest>> Requests;
+		Swap(Requests, Queue);
+		for (const TRefCountPtr<FAsyncRequest>& Request : Requests)
+		{
+			Request->Wait();
+		}
+	}
+}
+
+ITestCacheStore* CreateTestCacheStore(const ECacheStoreFlags Flags, const bool bAsync)
+{
+	return new FTestCacheStore(Flags, bAsync);
+}
+
+} // UE::DerivedData
+
+#endif // WITH_LOW_LEVEL_TESTS
diff --git a/Engine/Source/Developer/DerivedDataCache/Tests/TestCacheStore.h b/Engine/Source/Developer/DerivedDataCache/Tests/TestCacheStore.h
new file mode 100644
--- /dev/null
+++ b/Engine/Source/Developer/DerivedDataCache/Tests/TestCacheStore.h
@@ -0,0 +1,41 @@
+// Copyright Epic Games, Inc. All Rights Reserved.
+
+#pragma once
+
+#if WITH_LOW_LEVEL_TESTS
+
+#include "DerivedDataLegacyCacheStore.h"
+
+namespace UE::DerivedData
+{
+
+class ITestCacheStore : public ILegacyCacheStore
+{
+public:
+	virtual ECacheStoreFlags GetFlags() const = 0;
+	virtual uint32 GetTotalRequestCount() const = 0;
+	virtual uint32 GetCanceledRequestCount() const = 0;
+
+	virtual void AddRecord(const FCacheKey& Key, TConstArrayView<FValueWithId> Values, const FCbObject* Meta = nullptr) = 0;
+	virtual TConstArrayView<FValueWithId> FindRecord(const FCacheKey& Key, FCbObject* OutMeta = nullptr) const = 0;
+
+	virtual void AddValue(const FCacheKey& Key, const FValue& Value) = 0;
+	virtual FValue FindValue(const FCacheKey& Key) const = 0;
+
+	virtual void AddContent(const FCompressedBuffer& Content) = 0;
+	virtual FCompressedBuffer FindContent(const FIoHash& RawHash, uint64 RawSize) const = 0;
+
+	virtual TConstArrayView<FCachePutRequest> GetPutRequests() const = 0;
+	virtual TConstArrayView<FCacheGetRequest> GetGetRequests() const = 0;
+	virtual TConstArrayView<FCachePutValueRequest> GetPutValueRequests() const = 0;
+	virtual TConstArrayView<FCacheGetValueRequest> GetGetValueRequests() const = 0;
+	virtual TConstArrayView<FCacheGetChunkRequest> GetGetChunkRequests() const = 0;
+
+	virtual void ExecuteAsync() = 0;
+};
+
+ITestCacheStore* CreateTestCacheStore(ECacheStoreFlags Flags, bool bAsync);
+
+} // UE::DerivedData
+
+#endif WITH_LOW_LEVEL_TESTS
