DDC: Fixed a compile error in the mock cache store

#rnx

diff --git a/Engine/Source/Developer/DerivedDataCache/Tests/TestCacheStore.h b/Engine/Source/Developer/DerivedDataCache/Tests/TestCacheStore.h
--- a/Engine/Source/Developer/DerivedDataCache/Tests/TestCacheStore.h
+++ b/Engine/Source/Developer/DerivedDataCache/Tests/TestCacheStore.h
@@ -38,4 +38,4 @@ ITestCacheStore* CreateTestCacheStore(ECacheStoreFlags Flags, bool bAsync);
 
 } // UE::DerivedData
 
-#endif WITH_LOW_LEVEL_TESTS
+#endif // WITH_LOW_LEVEL_TESTS
