Fix vs2022 compile errors when /Zc:lambda- is removed

#rnx

diff --git a/Engine/Plugins/Tests/RHITests/Source/RHITests/Public/RHITextureTests.h b/Engine/Plugins/Tests/RHITests/Source/RHITests/Public/RHITextureTests.h
--- a/Engine/Plugins/Tests/RHITests/Source/RHITests/Public/RHITextureTests.h
+++ b/Engine/Plugins/Tests/RHITests/Source/RHITests/Public/RHITextureTests.h
@@ -758,7 +758,8 @@ public:
 			}
 		};
 
-		const auto DumpTestData = [](const void* Ptr, int32 Stride, FIntPoint Size, FIntPoint Offset = FIntPoint::ZeroValue)
+		// Temp: FIntPoint::ZeroValue replaced with UE::Math::TIntPoint<int32>::ZeroValue to work around msvc lambda compiler bug
+		const auto DumpTestData = [](const void* Ptr, int32 Stride, FIntPoint Size, FIntPoint Offset = UE::Math::TIntPoint<int32>::ZeroValue)
 		{
 			FString DataString;
 			uint8* Bytes = (uint8*)Ptr;
diff --git a/Engine/Source/Developer/DerivedDataCache/Private/DerivedDataCache.cpp b/Engine/Source/Developer/DerivedDataCache/Private/DerivedDataCache.cpp
--- a/Engine/Source/Developer/DerivedDataCache/Private/DerivedDataCache.cpp
+++ b/Engine/Source/Developer/DerivedDataCache/Private/DerivedDataCache.cpp
@@ -1098,7 +1098,7 @@ public:
 		FOnCachePutComplete&& OnComplete) final
 	{
 		DDC_SCOPE_CYCLE_COUNTER(DDC_Put);
-		return Backend->GetRoot().Put(Requests, Owner, OnComplete ? MoveTemp(OnComplete) : [](auto&&){});
+		return Backend->GetRoot().Put(Requests, Owner, OnComplete ? MoveTemp(OnComplete) : FOnCachePutComplete([](auto&&) {}));
 	}
 
 	void Get(
@@ -1107,7 +1107,7 @@ public:
 		FOnCacheGetComplete&& OnComplete) final
 	{
 		DDC_SCOPE_CYCLE_COUNTER(DDC_Get);
-		return Backend->GetRoot().Get(Requests, Owner, OnComplete ? MoveTemp(OnComplete) : [](auto&&){});
+		return Backend->GetRoot().Get(Requests, Owner, OnComplete ? MoveTemp(OnComplete) : FOnCacheGetComplete([](auto&&) {}));
 	}
 
 	void PutValue(
@@ -1116,7 +1116,7 @@ public:
 		FOnCachePutValueComplete&& OnComplete) final
 	{
 		DDC_SCOPE_CYCLE_COUNTER(DDC_PutValue);
-		return Backend->GetRoot().PutValue(Requests, Owner, OnComplete ? MoveTemp(OnComplete) : [](auto&&){});
+		return Backend->GetRoot().PutValue(Requests, Owner, OnComplete ? MoveTemp(OnComplete) : FOnCachePutValueComplete([](auto&&) {}));
 	}
 
 	void GetValue(
@@ -1125,7 +1125,7 @@ public:
 		FOnCacheGetValueComplete&& OnComplete) final
 	{
 		DDC_SCOPE_CYCLE_COUNTER(DDC_GetValue);
-		return Backend->GetRoot().GetValue(Requests, Owner, OnComplete ? MoveTemp(OnComplete) : [](auto&&){});
+		return Backend->GetRoot().GetValue(Requests, Owner, OnComplete ? MoveTemp(OnComplete) : FOnCacheGetValueComplete([](auto&&) {}));
 	}
 
 	void GetChunks(
@@ -1134,7 +1134,7 @@ public:
 		FOnCacheGetChunkComplete&& OnComplete) final
 	{
 		DDC_SCOPE_CYCLE_COUNTER(DDC_GetChunks);
-		return Backend->GetRoot().GetChunks(Requests, Owner, OnComplete ? MoveTemp(OnComplete) : [](auto&&){});
+		return Backend->GetRoot().GetChunks(Requests, Owner, OnComplete ? MoveTemp(OnComplete) : FOnCacheGetChunkComplete([](auto&&) {}));
 	}
 
 	ICacheStoreMaintainer& GetMaintainer() final
diff --git a/Engine/Source/Runtime/MovieScene/Public/EntitySystem/MovieScenePropertyMetaDataTraits.h b/Engine/Source/Runtime/MovieScene/Public/EntitySystem/MovieScenePropertyMetaDataTraits.h
--- a/Engine/Source/Runtime/MovieScene/Public/EntitySystem/MovieScenePropertyMetaDataTraits.h
+++ b/Engine/Source/Runtime/MovieScene/Public/EntitySystem/MovieScenePropertyMetaDataTraits.h
@@ -51,6 +51,10 @@ struct TPropertyMetaDataComponents<TPropertyMetaData<>>
 	{
 		return TArrayView<const FComponentTypeID>();
 	}
+
+	template<int Index> void GetType() const
+	{
+	}
 };
 
 
