More workarounds for issues with /Zc:lambda in MSVC 14.33

#rnx

diff --git a/Engine/Source/Developer/DerivedDataCache/Private/DerivedDataBuildSession.cpp b/Engine/Source/Developer/DerivedDataCache/Private/DerivedDataBuildSession.cpp
--- a/Engine/Source/Developer/DerivedDataCache/Private/DerivedDataBuildSession.cpp
+++ b/Engine/Source/Developer/DerivedDataCache/Private/DerivedDataBuildSession.cpp
@@ -68,7 +68,7 @@ void FBuildSessionInternal::Build(
 	FOnBuildComplete&& OnComplete)
 {
 	CreateBuildJob({Cache, BuildSystem, Scheduler, InputResolver, Owner}, Key, Policy,
-		OnComplete ? MoveTemp(OnComplete) : [](FBuildCompleteParams&&){});
+		OnComplete ? MoveTemp(OnComplete) : FOnBuildComplete([](FBuildCompleteParams&&){}));
 }
 
 void FBuildSessionInternal::Build(
@@ -79,7 +79,7 @@ void FBuildSessionInternal::Build(
 	FOnBuildComplete&& OnComplete)
 {
 	CreateBuildJob({Cache, BuildSystem, Scheduler, InputResolver, Owner}, Definition, Inputs, Policy,
-		OnComplete ? MoveTemp(OnComplete) : [](FBuildCompleteParams&&){});
+		OnComplete ? MoveTemp(OnComplete) : FOnBuildComplete([](FBuildCompleteParams&&){}));
 }
 
 void FBuildSessionInternal::Build(
@@ -90,7 +90,7 @@ void FBuildSessionInternal::Build(
 	FOnBuildComplete&& OnComplete)
 {
 	CreateBuildJob({Cache, BuildSystem, Scheduler, InputResolver, Owner}, Action, Inputs, Policy,
-		OnComplete ? MoveTemp(OnComplete) : [](FBuildCompleteParams&&){});
+		OnComplete ? MoveTemp(OnComplete) : FOnBuildComplete([](FBuildCompleteParams&&){}));
 }
 
 ///////////////////////////////////////////////////////////////////////////////////////////////////////////////////////
