DDC: Added partial cache record propagation to the hierarchy, with tests

diff --git a/Engine/Source/Developer/DerivedDataCache/Private/DerivedDataCacheStoreHierarchy.cpp b/Engine/Source/Developer/DerivedDataCache/Private/DerivedDataCacheStoreHierarchy.cpp
--- a/Engine/Source/Developer/DerivedDataCache/Private/DerivedDataCacheStoreHierarchy.cpp
+++ b/Engine/Source/Developer/DerivedDataCache/Private/DerivedDataCacheStoreHierarchy.cpp
@@ -2,7 +2,9 @@
 
 #include "Algo/AllOf.h"
 #include "Algo/AnyOf.h"
+#include "Algo/Compare.h"
 #include "Algo/Find.h"
+#include "Algo/MinElement.h"
 #include "Containers/Array.h"
 #include "DerivedDataCache.h"
 #include "DerivedDataCacheStore.h"
@@ -90,7 +92,7 @@ private:
 	static ECachePolicy GetCombinedPolicy(const ECachePolicy Policy) { return Policy; }
 	static ECachePolicy GetCombinedPolicy(const FCacheRecordPolicy& Policy) { return Policy.GetRecordPolicy(); }
 
-	static ECachePolicy AddPolicy(ECachePolicy BasePolicy, ECachePolicy Policy) { return BasePolicy | Policy; }
+	static ECachePolicy AddPolicy(ECachePolicy BasePolicy, ECachePolicy Policy) { return BasePolicy | (Policy & ~FCacheValuePolicy::PolicyMask); }
 	static ECachePolicy RemovePolicy(ECachePolicy BasePolicy, ECachePolicy Policy) { return BasePolicy & ~Policy; }
 	static FCacheRecordPolicy AddPolicy(const FCacheRecordPolicy& BasePolicy, ECachePolicy Policy);
 	static FCacheRecordPolicy RemovePolicy(const FCacheRecordPolicy& BasePolicy, ECachePolicy Policy);
@@ -169,7 +171,9 @@ struct FCacheStoreHierarchy::TBatchParams
 	static PutFunctionType Put();
 	static GetFunctionType Get();
 	static bool HasResponseData(const FGetResponse& Response);
-	static void FilterResponseByRequest(FGetResponse& Response, const FGetRequest& Request);
+	static bool MergeFromResponse(FGetResponse& OutResponse, const FGetResponse& Response);
+	static void ModifyPolicyForResponse(decltype(FGetRequest::Policy)& Policy, const FGetResponse& Response) {}
+	static FGetResponse FilterResponseByRequest(const FGetResponse& Response, const FGetRequest& Request);
 	static FPutRequest MakePutRequest(const FGetResponse& Response, const FGetRequest& Request);
 	static FGetRequest MakeGetRequest(const FPutRequest& Request, uint64 UserData);
 };
@@ -561,6 +565,8 @@ class FCacheStoreHierarchy::TGetBatch final : public FBatchBase, public Params
 	using Params::Put;
 	using Params::Get;
 	using Params::HasResponseData;
+	using Params::MergeFromResponse;
+	using Params::ModifyPolicyForResponse;
 	using Params::FilterResponseByRequest;
 	using Params::MakePutRequest;
 
@@ -597,6 +603,7 @@ private:
 	{
 		FGetRequest Request;
 		FGetResponse Response;
+		uint64 NodeIndex = 0;
 	};
 
 	const FCacheStoreHierarchy& Hierarchy;
@@ -606,7 +613,6 @@ private:
 	IRequestOwner& Owner;
 	FRequestOwner AsyncOwner;
 	FCounterEvent RemainingRequestCount;
-	int32 NodeIndex = 0;
 };
 
 template <typename Params>
@@ -637,13 +643,24 @@ void FCacheStoreHierarchy::TGetBatch<Params>::DispatchRequests()
 	NodeRequests.Reserve(RequestCount);
 	AsyncNodeRequests.Reserve(RequestCount);
 
-	for (const int32 NodeCount = Hierarchy.Nodes.Num(); NodeIndex < NodeCount && !Owner.IsCanceled(); ++NodeIndex)
+	for (const int32 NodeCount = Hierarchy.Nodes.Num(); !Owner.IsCanceled();)
 	{
+		const int32 NodeIndex = Algo::MinElementBy(States, &FState::NodeIndex)->NodeIndex;
+		if (NodeIndex >= NodeCount)
+		{
+			break;
+		}
+
 		const FCacheStoreNode& Node = Hierarchy.Nodes[NodeIndex];
 
 		uint64 StateIndex = 0;
-		for (const FState& State : States)
+		for (FState& State : States)
 		{
+			if (NodeIndex < State.NodeIndex)
+			{
+				continue;
+			}
+
 			const FGetRequest& Request = State.Request;
 			const FGetResponse& Response = State.Response;
 			if (Response.Status == EStatus::Ok)
@@ -651,24 +668,36 @@ void FCacheStoreHierarchy::TGetBatch<Params>::DispatchRequests()
 				if (HasResponseData(Response) && CanStore(GetCombinedPolicy(Request.Policy), Node.CacheFlags))
 				{
 					AsyncNodeRequests.Add(MakePutRequest(Response, Request));
+					++State.NodeIndex;
 				}
 				else if (EnumHasAnyFlags(Node.CacheFlags, ECacheStoreFlags::StopGetStore) && CanQuery(GetCombinedPolicy(Request.Policy), Node.CacheFlags))
 				{
 					NodeRequests.Add({Request.Name, Request.Key, AddPolicy(Request.Policy, ECachePolicy::SkipData), StateIndex});
 				}
+				else
+				{
+					++State.NodeIndex;
+				}
 			}
 			else
 			{
-				if (const ECachePolicy Policy = GetCombinedPolicy(Request.Policy); CanQuery(Policy, Node.CacheFlags))
+				if (const ECachePolicy CombinedPolicy = GetCombinedPolicy(Request.Policy); CanQuery(CombinedPolicy, Node.CacheFlags))
 				{
-					if (CanStoreIfOk(Policy, Node.NodeFlags))
+					auto Policy = Request.Policy;
+					if (CanStoreIfOk(CombinedPolicy, Node.NodeFlags))
 					{
-						NodeRequests.Add({Request.Name, Request.Key, RemovePolicy(Request.Policy, ECachePolicy::SkipData | ECachePolicy::SkipMeta), StateIndex});
+						Policy = RemovePolicy(Policy, ECachePolicy::SkipData | ECachePolicy::SkipMeta);
 					}
-					else
+					if (CanQueryIfError(CombinedPolicy, Node.NodeFlags))
 					{
-						NodeRequests.Add({Request.Name, Request.Key, Request.Policy, StateIndex});
+						Policy = AddPolicy(Policy, ECachePolicy::PartialRecord);
 					}
+					ModifyPolicyForResponse(Policy, Response);
+					NodeRequests.Add({Request.Name, Request.Key, Policy, StateIndex});
+				}
+				else
+				{
+					++State.NodeIndex;
 				}
 			}
 			++StateIndex;
@@ -709,25 +738,34 @@ void FCacheStoreHierarchy::TGetBatch<Params>::DispatchRequests()
 template <typename Params>
 void FCacheStoreHierarchy::TGetBatch<Params>::CompleteRequest(FGetResponse&& Response)
 {
-	FState& State = States[int32(Response.UserData)];
-	const FCacheStoreNode& Node = Hierarchy.Nodes[NodeIndex];
+	ON_SCOPE_EXIT
+	{
+		if (RemainingRequestCount.Signal())
+		{
+			DispatchRequests();
+		}
+	};
 
-	const bool bFirstOk = Response.Status == EStatus::Ok && State.Response.Status == EStatus::Error;
-	const bool bLastQuery = bFirstOk || !CanQueryIfError(GetCombinedPolicy(State.Request.Policy), Node.NodeFlags);
+	FState& State = States[int32(Response.UserData)];
+	const FCacheStoreNode& Node = Hierarchy.Nodes[State.NodeIndex];
+	const EStatus PreviousStatus = State.Response.Status;
+	const int32 PreviousNodeIndex = State.NodeIndex;
 
-	if (State.Response.Status == EStatus::Error)
+	// Failure to merge requires dispatching to the same node to try to recover missing values.
+	if (!MergeFromResponse(State.Response, Response))
 	{
-		Response.UserData = State.Request.UserData;
-		// TODO: Merge values from partial records.
-		State.Response = Response;
+		return;
 	}
 
-	if (bLastQuery && CanStoreIfOk(GetCombinedPolicy(State.Request.Policy), Node.NodeFlags) && HasResponseData(Response))
+	const bool bFirstOk = State.Response.Status == EStatus::Ok && PreviousStatus == EStatus::Error;
+	const bool bLastQuery = bFirstOk || !CanQueryIfError(GetCombinedPolicy(State.Request.Policy), Node.NodeFlags);
+
+	if (bLastQuery && CanStoreIfOk(GetCombinedPolicy(State.Request.Policy), Node.NodeFlags) && HasResponseData(State.Response))
 	{
 		// Store any retrieved values to previous writable nodes if Ok or there are no remaining nodes to query.
-		FPutRequest PutRequest = MakePutRequest(Response, State.Request);
+		FPutRequest PutRequest = MakePutRequest(State.Response, State.Request);
 		PutRequest.Policy = RemovePolicy(PutRequest.Policy, ECachePolicy::Query);
-		for (int32 PutNodeIndex = 0; PutNodeIndex < NodeIndex; ++PutNodeIndex)
+		for (int32 PutNodeIndex = 0; PutNodeIndex < PreviousNodeIndex; ++PutNodeIndex)
 		{
 			const FCacheStoreNode& PutNode = Hierarchy.Nodes[PutNodeIndex];
 			if (CanStore(GetCombinedPolicy(State.Request.Policy), PutNode.CacheFlags))
@@ -740,9 +778,8 @@ void FCacheStoreHierarchy::TGetBatch<Params>::CompleteRequest(FGetResponse&& Res
 
 	if (bFirstOk)
 	{
-		// Values may be fetched to fill previous nodes. Remove values if requested.
-		FilterResponseByRequest(Response, State.Request);
-		OnComplete(MoveTemp(Response));
+		// Values may be fetched to populate earlier nodes. Remove values if requested.
+		OnComplete(FilterResponseByRequest(State.Response, State.Request));
 	}
 
 	if (Response.Status == EStatus::Ok)
@@ -760,11 +797,7 @@ void FCacheStoreHierarchy::TGetBatch<Params>::CompleteRequest(FGetResponse&& Res
 		}
 	}
 
-	if (RemainingRequestCount.Signal())
-	{
-		++NodeIndex;
-		DispatchRequests();
-	}
+	++State.NodeIndex;
 }
 
 ///////////////////////////////////////////////////////////////////////////////////////////////////////////////////////
@@ -788,8 +821,95 @@ bool FCacheStoreHierarchy::FCacheRecordBatchParams::HasResponseData(const FCache
 }
 
 template <>
-void FCacheStoreHierarchy::FCacheRecordBatchParams::FilterResponseByRequest(
-	FCacheGetResponse& Response,
+bool FCacheStoreHierarchy::FCacheRecordBatchParams::MergeFromResponse(
+	FCacheGetResponse& OutResponse,
+	const FCacheGetResponse& Response)
+{
+	if (OutResponse.Status == EStatus::Ok)
+	{
+		return true;
+	}
+
+	const FCacheRecord& NewRecord = Response.Record;
+	const FCacheRecord ExistingRecord = MoveTemp(OutResponse.Record);
+
+	OutResponse = Response;
+
+	if (ExistingRecord.GetValues().IsEmpty())
+	{
+		return true;
+	}
+
+	bool bOk = true;
+	FCacheRecordBuilder Builder(ExistingRecord.GetKey());
+	Builder.SetMeta(CopyTemp(NewRecord.GetMeta() ? NewRecord.GetMeta() : ExistingRecord.GetMeta()));
+	for (const FValueWithId& Value : NewRecord.GetValues())
+	{
+		if (Value.HasData())
+		{
+			Builder.AddValue(Value);
+		}
+		else if (const FValueWithId& ExistingValue = ExistingRecord.GetValue(Value.GetId()); ExistingValue && ExistingValue.HasData())
+		{
+			if (Value == ExistingValue)
+			{
+				// There is a matching existing value with data. Merge it into the new record.
+				Builder.AddValue(ExistingValue);
+			}
+			else
+			{
+				// There is an existing value and it differs. Requery the current node for this value.
+				Builder.AddValue(Value);
+				OutResponse.Status = EStatus::Error;
+				bOk = false;
+			}
+		}
+		else
+		{
+			Builder.AddValue(Value);
+		}
+	}
+	OutResponse.Record = Builder.Build();
+	return bOk;
+}
+
+template <>
+void FCacheStoreHierarchy::FCacheRecordBatchParams::ModifyPolicyForResponse(
+	FCacheRecordPolicy& Policy,
+	const FCacheGetResponse& Response)
+{
+	const TConstArrayView<FValueWithId> Values = Response.Record.GetValues();
+	if (Values.IsEmpty())
+	{
+		return;
+	}
+
+	FCacheRecordPolicyBuilder Builder(Policy.GetBasePolicy());
+
+	// Skip values that are present in the response.
+	for (const FValueWithId& Value : Values)
+	{
+		if (Value.HasData())
+		{
+			Builder.AddValuePolicy(Value.GetId(), ECachePolicy::None);
+		}
+	}
+
+	// Copy any the policy for any other values from the source policy.
+	for (const FCacheValuePolicy& ValuePolicy : Policy.GetValuePolicies())
+	{
+		if (!Response.Record.GetValue(ValuePolicy.Id).HasData())
+		{
+			Builder.AddValuePolicy(ValuePolicy);
+		}
+	}
+
+	Policy = Builder.Build();
+}
+
+template <>
+FCacheGetResponse FCacheStoreHierarchy::FCacheRecordBatchParams::FilterResponseByRequest(
+	const FCacheGetResponse& Response,
 	const FCacheGetRequest& Request)
 {
 	const ECachePolicy RecordPolicy = Request.Policy.GetRecordPolicy();
@@ -813,8 +933,9 @@ void FCacheStoreHierarchy::FCacheRecordBatchParams::FilterResponseByRequest(
 				Builder.AddValue(Value);
 			}
 		}
-		Response.Record = Builder.Build();
+		return {Response.Name, Builder.Build(), Response.UserData, Response.Status};
 	}
+	return Response;
 }
 
 template <>
@@ -876,14 +997,27 @@ bool FCacheStoreHierarchy::FCacheValueBatchParams::HasResponseData(const FCacheG
 }
 
 template <>
-void FCacheStoreHierarchy::FCacheValueBatchParams::FilterResponseByRequest(
-	FCacheGetValueResponse& Response,
+bool FCacheStoreHierarchy::FCacheValueBatchParams::MergeFromResponse(
+	FGetResponse& OutResponse,
+	const FGetResponse& Response)
+{
+	if (OutResponse.Status == EStatus::Error)
+	{
+		OutResponse = Response;
+	}
+	return true;
+}
+
+template <>
+FCacheGetValueResponse FCacheStoreHierarchy::FCacheValueBatchParams::FilterResponseByRequest(
+	const FCacheGetValueResponse& Response,
 	const FCacheGetValueRequest& Request)
 {
 	if (Response.Value.HasData() && EnumHasAnyFlags(Request.Policy, ECachePolicy::SkipData))
 	{
-		Response.Value = Response.Value.RemoveData();
+		return {Response.Name, Response.Key, Response.Value.RemoveData(), Response.UserData, Response.Status};
 	}
+	return Response;
 }
 
 template <>
diff --git a/Engine/Source/Developer/DerivedDataCache/Tests/DerivedDataCacheStoreHierarchyTest.cpp b/Engine/Source/Developer/DerivedDataCache/Tests/DerivedDataCacheStoreHierarchyTest.cpp
new file mode 100644
--- /dev/null
+++ b/Engine/Source/Developer/DerivedDataCache/Tests/DerivedDataCacheStoreHierarchyTest.cpp
@@ -0,0 +1,279 @@
+// Copyright Epic Games, Inc. All Rights Reserved.
+
+#if WITH_LOW_LEVEL_TESTS
+
+#include "Algo/AllOf.h"
+#include "DerivedDataCacheStore.h"
+#include "DerivedDataLegacyCacheStore.h"
+#include "DerivedDataRequestOwner.h"
+#include "TestCacheStore.h"
+#include "TestHarness.h"
+
+namespace UE::DerivedData
+{
+
+class IMemoryCacheStore;
+ILegacyCacheStore* CreateCacheStoreHierarchy(ICacheStoreOwner*& OutOwner, IMemoryCacheStore* MemoryCache);
+
+static FValue CreateHierarchyTestValue(uint32 Value)
+{
+	return FValue::Compress(FSharedBuffer::MakeView(MakeMemoryView<uint32>({Value})));
+}
+
+TEST_CASE("DerivedData::Cache::Hierarchy::PartialRecordPropagation", "[DerivedData]")
+{
+	ICacheStoreOwner* StoreOwner = nullptr;
+	TUniquePtr<ILegacyCacheStore> Hierarchy(CreateCacheStoreHierarchy(StoreOwner, nullptr));
+
+	TUniquePtr<ITestCacheStore> Store0(CreateTestCacheStore(ECacheStoreFlags::Local | ECacheStoreFlags::Query | ECacheStoreFlags::Store, /*bAsync*/ true));
+	StoreOwner->Add(Store0.Get(), Store0->GetFlags());
+	TUniquePtr<ITestCacheStore> Store1(CreateTestCacheStore(ECacheStoreFlags::Local | ECacheStoreFlags::Query | ECacheStoreFlags::Store, /*bAsync*/ true));
+	StoreOwner->Add(Store1.Get(), Store1->GetFlags());
+	TUniquePtr<ITestCacheStore> Store2(CreateTestCacheStore(ECacheStoreFlags::Local | ECacheStoreFlags::Query | ECacheStoreFlags::Store, /*bAsync*/ true));
+	StoreOwner->Add(Store2.Get(), Store2->GetFlags());
+
+	const FValueId Id0 = FValueId::FromName("0");
+	const FValueId Id1 = FValueId::FromName("1");
+	const FValueId Id2 = FValueId::FromName("2");
+	const FValueWithId Value0(Id0, CreateHierarchyTestValue(0));
+	const FValueWithId Value1(Id1, CreateHierarchyTestValue(1));
+	const FValueWithId Value2(Id2, CreateHierarchyTestValue(2));
+	const FCacheKey Key{FCacheBucket("Test"), FIoHash::HashBuffer(MakeMemoryView<uint8>({0, 1, 2}))};
+
+	Store0->AddRecord(Key, {Value0, Value1.RemoveData(), Value2.RemoveData()});
+	Store1->AddRecord(Key, {Value0, Value1, Value2.RemoveData()});
+	Store2->AddRecord(Key, {Value0, Value1, Value2});
+
+	FCacheRecordPolicyBuilder PolicyBuilder(ECachePolicy::Default);
+	PolicyBuilder.AddValuePolicy(Id0, ECachePolicy::Local);
+
+	FRequestOwner Owner(EPriority::Normal);
+	Hierarchy->Get({{{TEXTVIEW("PropagatePartial")}, Key, PolicyBuilder.Build()}}, Owner, [](FCacheGetResponse&& Response)
+	{
+		CHECK(Response.Status == EStatus::Ok);
+		CHECK(Response.Record.GetValues().Num() == 3);
+		CHECK(Algo::AllOf(Response.Record.GetValues(), &FValue::HasData));
+	});
+
+	// Check that every value is requested from Store0.
+	CHECK(Store0->GetTotalRequestCount() == 1);
+	CHECK(Store1->GetTotalRequestCount() == 0);
+	CHECK(Store2->GetTotalRequestCount() == 0);
+	REQUIRE(Store0->GetGetRequests().Num() == 1);
+	const FCacheGetRequest& Request0 = Store0->GetGetRequests()[0];
+	CHECK(Request0.Policy.GetRecordPolicy() == (ECachePolicy::Default | ECachePolicy::PartialRecord));
+	CHECK(Request0.Policy.GetValuePolicy(Id0) == ECachePolicy::Local);
+	CHECK(Request0.Policy.GetValuePolicies().Num() == 1);
+
+	Store0->ExecuteAsync();
+
+	// Check that every value except Value0 is requested from Store1.
+	CHECK(Store0->GetTotalRequestCount() == 1);
+	CHECK(Store1->GetTotalRequestCount() == 1);
+	CHECK(Store2->GetTotalRequestCount() == 0);
+	REQUIRE(Store1->GetGetRequests().Num() == 1);
+	const FCacheGetRequest& Request1 = Store1->GetGetRequests()[0];
+	CHECK(Request1.Policy.GetRecordPolicy() == (ECachePolicy::Default | ECachePolicy::PartialRecord));
+	CHECK(Request1.Policy.GetValuePolicy(Id0) == ECachePolicy::None);
+	CHECK(Request1.Policy.GetValuePolicies().Num() == 1);
+
+	Store1->ExecuteAsync();
+
+	// Check that every value except Value0 and Value1 is requested from Store2.
+	CHECK(Store0->GetTotalRequestCount() == 1);
+	CHECK(Store1->GetTotalRequestCount() == 1);
+	CHECK(Store2->GetTotalRequestCount() == 1);
+	REQUIRE(Store2->GetGetRequests().Num() == 1);
+	const FCacheGetRequest& Request2 = Store2->GetGetRequests()[0];
+	CHECK(Request2.Policy.GetRecordPolicy() == ECachePolicy::Default);
+	CHECK(Request2.Policy.GetValuePolicy(Id0) == ECachePolicy::None);
+	CHECK(Request2.Policy.GetValuePolicy(Id1) == ECachePolicy::None);
+	CHECK(Request2.Policy.GetValuePolicies().Num() == 2);
+
+	Store2->ExecuteAsync();
+
+	// Check that the complete record is put to Store0.
+	CHECK(Store0->GetTotalRequestCount() == 2);
+	REQUIRE(Store0->GetPutRequests().Num() == 1);
+	const FCachePutRequest& PutRequest0 = Store0->GetPutRequests()[0];
+	CHECK(PutRequest0.Policy.GetRecordPolicy() == ECachePolicy::Store);
+	CHECK(PutRequest0.Policy.GetValuePolicy(Id0) == ECachePolicy::StoreLocal);
+	CHECK(PutRequest0.Policy.GetValuePolicies().Num() == 1);
+	CHECK(PutRequest0.Record.GetValues().Num() == 3);
+	CHECK(Algo::AllOf(PutRequest0.Record.GetValues(), &FValue::HasData));
+
+	CHECK(!Store0->FindContent(Value1.GetRawHash(), Value1.GetRawSize()));
+	CHECK(!Store0->FindContent(Value2.GetRawHash(), Value2.GetRawSize()));
+	Store0->ExecuteAsync();
+	CHECK(Store0->FindContent(Value1.GetRawHash(), Value1.GetRawSize()));
+	CHECK(Store0->FindContent(Value2.GetRawHash(), Value2.GetRawSize()));
+
+	// Check that the complete record is put to Store1.
+	CHECK(Store1->GetTotalRequestCount() == 2);
+	REQUIRE(Store1->GetPutRequests().Num() == 1);
+	const FCachePutRequest& PutRequest1 = Store1->GetPutRequests()[0];
+	CHECK(PutRequest1.Policy.GetRecordPolicy() == ECachePolicy::Store);
+	CHECK(PutRequest1.Policy.GetValuePolicy(Id0) == ECachePolicy::StoreLocal);
+	CHECK(PutRequest1.Policy.GetValuePolicies().Num() == 1);
+	CHECK(PutRequest1.Record.GetValues().Num() == 3);
+	CHECK(Algo::AllOf(PutRequest1.Record.GetValues(), &FValue::HasData));
+
+	CHECK(!Store1->FindContent(Value2.GetRawHash(), Value2.GetRawSize()));
+	Store1->ExecuteAsync();
+	CHECK(Store1->FindContent(Value2.GetRawHash(), Value2.GetRawSize()));
+
+	// Wait to ensure that the counts checked below are final.
+	Owner.Wait();
+
+	// Check that the stores are not accessed again.
+	CHECK(Store0->GetTotalRequestCount() == 2);
+	CHECK(Store1->GetTotalRequestCount() == 2);
+	CHECK(Store2->GetTotalRequestCount() == 1);
+
+	StoreOwner->RemoveNotSafe(Store2.Get());
+	StoreOwner->RemoveNotSafe(Store1.Get());
+	StoreOwner->RemoveNotSafe(Store0.Get());
+}
+
+TEST_CASE("DerivedData::Cache::Hierarchy::PartialNonDeterministicRecordPropagation", "[DerivedData]")
+{
+	ICacheStoreOwner* StoreOwner = nullptr;
+	TUniquePtr<ILegacyCacheStore> Hierarchy(CreateCacheStoreHierarchy(StoreOwner, nullptr));
+
+	TUniquePtr<ITestCacheStore> Store0(CreateTestCacheStore(ECacheStoreFlags::Local | ECacheStoreFlags::Query | ECacheStoreFlags::Store, /*bAsync*/ true));
+	StoreOwner->Add(Store0.Get(), Store0->GetFlags());
+	TUniquePtr<ITestCacheStore> Store1(CreateTestCacheStore(ECacheStoreFlags::Local | ECacheStoreFlags::Query | ECacheStoreFlags::Store, /*bAsync*/ true));
+	StoreOwner->Add(Store1.Get(), Store1->GetFlags());
+	TUniquePtr<ITestCacheStore> Store2(CreateTestCacheStore(ECacheStoreFlags::Local | ECacheStoreFlags::Query | ECacheStoreFlags::Store, /*bAsync*/ true));
+	StoreOwner->Add(Store2.Get(), Store2->GetFlags());
+
+	const FValueId Id0 = FValueId::FromName("0");
+	const FValueId Id1 = FValueId::FromName("1");
+	const FValueId Id2 = FValueId::FromName("2");
+	const FValueWithId Value0(Id0, CreateHierarchyTestValue(0));
+	const FValueWithId Value1(Id1, CreateHierarchyTestValue(1));
+	const FValueWithId Value2(Id2, CreateHierarchyTestValue(2));
+	const FValueWithId Value0Diff(Id0, CreateHierarchyTestValue(3));
+	const FCacheKey Key{FCacheBucket("Test"), FIoHash::HashBuffer(MakeMemoryView<uint8>({0, 1, 2, 3}))};
+
+	SECTION("Store is requeried if its response contains differing values that were not queried initially")
+	{
+		Store0->AddRecord(Key, {Value0, Value1.RemoveData(), Value2});
+		Store1->AddRecord(Key, {Value0Diff, Value1, Value2.RemoveData()});
+		Store2->AddRecord(Key, {Value0Diff, Value1, Value2});
+
+		FCacheRecordPolicyBuilder PolicyBuilder(ECachePolicy::Default);
+		PolicyBuilder.AddValuePolicy(Id0, ECachePolicy::Local);
+
+		FRequestOwner Owner(EPriority::Normal);
+		Hierarchy->Get({{{TEXTVIEW("PropagatePartial")}, Key, PolicyBuilder.Build()}}, Owner, [&Id0, &Value0Diff](FCacheGetResponse&& Response)
+		{
+			CHECK(Response.Status == EStatus::Ok);
+			CHECK(Response.Record.GetValues().Num() == 3);
+			CHECK(Algo::AllOf(Response.Record.GetValues(), &FValue::HasData));
+			CHECK(Response.Record.GetValue(Id0) == Value0Diff);
+		});
+
+		CHECK(Store0->GetTotalRequestCount() == 1);
+		CHECK(Store1->GetTotalRequestCount() == 0);
+		CHECK(Store2->GetTotalRequestCount() == 0);
+		Store0->ExecuteAsync();
+
+		// Check that only Value1 is requested from Store1.
+		CHECK(Store0->GetTotalRequestCount() == 1);
+		CHECK(Store1->GetTotalRequestCount() == 1);
+		CHECK(Store2->GetTotalRequestCount() == 0);
+		REQUIRE(Store1->GetGetRequests().Num() == 1);
+		const FCacheGetRequest& Request1 = Store1->GetGetRequests()[0];
+		CHECK(Request1.Policy.GetRecordPolicy() == (ECachePolicy::Default | ECachePolicy::PartialRecord));
+		CHECK(Request1.Policy.GetValuePolicy(Id0) == ECachePolicy::None);
+		CHECK(Request1.Policy.GetValuePolicy(Id2) == ECachePolicy::None);
+		CHECK(Request1.Policy.GetValuePolicies().Num() == 2);
+		Store1->ExecuteAsync();
+
+		// Check that only Value0 is requested from Store1.
+		CHECK(Store0->GetTotalRequestCount() == 1);
+		CHECK(Store1->GetTotalRequestCount() == 2);
+		CHECK(Store2->GetTotalRequestCount() == 0);
+		REQUIRE(Store1->GetGetRequests().Num() == 2);
+		const FCacheGetRequest& Request1Repeat = Store1->GetGetRequests()[1];
+		CHECK(Request1Repeat.Policy.GetRecordPolicy() == (ECachePolicy::Default | ECachePolicy::PartialRecord));
+		CHECK(Request1Repeat.Policy.GetValuePolicy(Id0) == ECachePolicy::Local);
+		CHECK(Request1Repeat.Policy.GetValuePolicy(Id1) == ECachePolicy::None);
+		CHECK(Request1Repeat.Policy.GetValuePolicy(Id2) == ECachePolicy::None);
+		CHECK(Request1Repeat.Policy.GetValuePolicies().Num() == 3);
+		Store1->ExecuteAsync();
+
+		// Wait to ensure that the counts checked below are final.
+		Owner.Wait();
+
+		// Check that the complete record was put to Store0 and Store2.
+		REQUIRE(Store0->GetPutRequests().Num() == 1);
+		REQUIRE(Store2->GetPutRequests().Num() == 1);
+		CHECK(Algo::AllOf(Store0->GetPutRequests()[0].Record.GetValues(), &FValue::HasData));
+		CHECK(Algo::AllOf(Store2->GetPutRequests()[0].Record.GetValues(), &FValue::HasData));
+
+		// Check that the stores are not accessed again.
+		CHECK(Store0->GetTotalRequestCount() == 2);
+		CHECK(Store1->GetTotalRequestCount() == 2);
+		CHECK(Store2->GetTotalRequestCount() == 1);
+	}
+
+	SECTION("Store is not requeried if its response contains differing values that were queried initially")
+	{
+		Store0->AddRecord(Key, {Value0.RemoveData(), Value1.RemoveData(), Value2});
+		Store1->AddRecord(Key, {Value0Diff, Value1, Value2.RemoveData()});
+		Store2->AddRecord(Key, {Value0Diff, Value1, Value2});
+
+		FCacheRecordPolicyBuilder PolicyBuilder(ECachePolicy::Default);
+		PolicyBuilder.AddValuePolicy(Id0, ECachePolicy::Local);
+
+		FRequestOwner Owner(EPriority::Normal);
+		Hierarchy->Get({{{TEXTVIEW("PropagatePartial")}, Key, PolicyBuilder.Build()}}, Owner, [&Id0, &Value0Diff](FCacheGetResponse&& Response)
+		{
+			CHECK(Response.Status == EStatus::Ok);
+			CHECK(Response.Record.GetValues().Num() == 3);
+			CHECK(Algo::AllOf(Response.Record.GetValues(), &FValue::HasData));
+			CHECK(Response.Record.GetValue(Id0) == Value0Diff);
+		});
+
+		CHECK(Store0->GetTotalRequestCount() == 1);
+		CHECK(Store1->GetTotalRequestCount() == 0);
+		CHECK(Store2->GetTotalRequestCount() == 0);
+		Store0->ExecuteAsync();
+
+		// Check that only Value0 and Value1 are requested from Store1.
+		CHECK(Store0->GetTotalRequestCount() == 1);
+		CHECK(Store1->GetTotalRequestCount() == 1);
+		CHECK(Store2->GetTotalRequestCount() == 0);
+		REQUIRE(Store1->GetGetRequests().Num() == 1);
+		const FCacheGetRequest& Request1 = Store1->GetGetRequests()[0];
+		CHECK(Request1.Policy.GetRecordPolicy() == (ECachePolicy::Default | ECachePolicy::PartialRecord));
+		CHECK(Request1.Policy.GetValuePolicy(Id0) == ECachePolicy::Local);
+		CHECK(Request1.Policy.GetValuePolicy(Id2) == ECachePolicy::None);
+		CHECK(Request1.Policy.GetValuePolicies().Num() == 2);
+		Store1->ExecuteAsync();
+
+		// Wait to ensure that the counts checked below are final.
+		Owner.Wait();
+
+		// Check that the complete record was put to Store0 and Store2.
+		REQUIRE(Store0->GetPutRequests().Num() == 1);
+		REQUIRE(Store2->GetPutRequests().Num() == 1);
+		CHECK(Algo::AllOf(Store0->GetPutRequests()[0].Record.GetValues(), &FValue::HasData));
+		CHECK(Algo::AllOf(Store2->GetPutRequests()[0].Record.GetValues(), &FValue::HasData));
+
+		// Check that the stores are not accessed again.
+		CHECK(Store0->GetTotalRequestCount() == 2);
+		CHECK(Store1->GetTotalRequestCount() == 1);
+		CHECK(Store2->GetTotalRequestCount() == 1);
+	}
+
+	StoreOwner->RemoveNotSafe(Store2.Get());
+	StoreOwner->RemoveNotSafe(Store1.Get());
+	StoreOwner->RemoveNotSafe(Store0.Get());
+}
+
+} // UE::DerivedData
+
+#endif // WITH_LOW_LEVEL_TESTS
diff --git a/Engine/Source/Developer/DerivedDataCache/Tests/TestCacheStore.cpp b/Engine/Source/Developer/DerivedDataCache/Tests/TestCacheStore.cpp
--- a/Engine/Source/Developer/DerivedDataCache/Tests/TestCacheStore.cpp
+++ b/Engine/Source/Developer/DerivedDataCache/Tests/TestCacheStore.cpp
@@ -490,14 +490,11 @@ void FTestCacheStore::ExecuteOrQueue(IRequestOwner& Owner, TUniqueFunction<void
 
 void FTestCacheStore::ExecuteAsync()
 {
-	while (!Queue.IsEmpty())
+	TArray<TRefCountPtr<FAsyncRequest>> Requests;
+	Swap(Requests, Queue);
+	for (const TRefCountPtr<FAsyncRequest>& Request : Requests)
 	{
-		TArray<TRefCountPtr<FAsyncRequest>> Requests;
-		Swap(Requests, Queue);
-		for (const TRefCountPtr<FAsyncRequest>& Request : Requests)
-		{
-			Request->Wait();
-		}
+		Request->Wait();
 	}
 }
 
