DDC: Fixed UserData value returned by the cache hierarchy

#rnx

diff --git a/Engine/Source/Developer/DerivedDataCache/Private/DerivedDataCacheStoreHierarchy.cpp b/Engine/Source/Developer/DerivedDataCache/Private/DerivedDataCacheStoreHierarchy.cpp
--- a/Engine/Source/Developer/DerivedDataCache/Private/DerivedDataCacheStoreHierarchy.cpp
+++ b/Engine/Source/Developer/DerivedDataCache/Private/DerivedDataCacheStoreHierarchy.cpp
@@ -747,6 +747,7 @@ void FCacheStoreHierarchy::TGetBatch<Params>::CompleteRequest(FGetResponse&& Res
 	};
 
 	FState& State = States[int32(Response.UserData)];
+	Response.UserData = State.Request.UserData;
 	const FCacheStoreNode& Node = Hierarchy.Nodes[State.NodeIndex];
 	const EStatus PreviousStatus = State.Response.Status;
 	const int32 PreviousNodeIndex = State.NodeIndex;
diff --git a/Engine/Source/Developer/DerivedDataCache/Tests/DerivedDataCacheStoreHierarchyTest.cpp b/Engine/Source/Developer/DerivedDataCache/Tests/DerivedDataCacheStoreHierarchyTest.cpp
--- a/Engine/Source/Developer/DerivedDataCache/Tests/DerivedDataCacheStoreHierarchyTest.cpp
+++ b/Engine/Source/Developer/DerivedDataCache/Tests/DerivedDataCacheStoreHierarchyTest.cpp
@@ -48,8 +48,9 @@ TEST_CASE("DerivedData::Cache::Hierarchy::PartialRecordPropagation", "[DerivedDa
 	PolicyBuilder.AddValuePolicy(Id0, ECachePolicy::Local);
 
 	FRequestOwner Owner(EPriority::Normal);
-	Hierarchy->Get({{{TEXTVIEW("PropagatePartial")}, Key, PolicyBuilder.Build()}}, Owner, [](FCacheGetResponse&& Response)
+	Hierarchy->Get({{{TEXTVIEW("PropagatePartial")}, Key, PolicyBuilder.Build(), 12345}}, Owner, [](FCacheGetResponse&& Response)
 	{
+		CHECK(Response.UserData == 12345);
 		CHECK(Response.Status == EStatus::Ok);
 		CHECK(Response.Record.GetValues().Num() == 3);
 		CHECK(Algo::AllOf(Response.Record.GetValues(), &FValue::HasData));
@@ -166,8 +167,9 @@ TEST_CASE("DerivedData::Cache::Hierarchy::PartialNonDeterministicRecordPropagati
 		PolicyBuilder.AddValuePolicy(Id0, ECachePolicy::Local);
 
 		FRequestOwner Owner(EPriority::Normal);
-		Hierarchy->Get({{{TEXTVIEW("PropagatePartial")}, Key, PolicyBuilder.Build()}}, Owner, [&Id0, &Value0Diff](FCacheGetResponse&& Response)
+		Hierarchy->Get({{{TEXTVIEW("PropagatePartial")}, Key, PolicyBuilder.Build(), 12345}}, Owner, [&Id0, &Value0Diff](FCacheGetResponse&& Response)
 		{
+			CHECK(Response.UserData == 12345);
 			CHECK(Response.Status == EStatus::Ok);
 			CHECK(Response.Record.GetValues().Num() == 3);
 			CHECK(Algo::AllOf(Response.Record.GetValues(), &FValue::HasData));
@@ -229,8 +231,9 @@ TEST_CASE("DerivedData::Cache::Hierarchy::PartialNonDeterministicRecordPropagati
 		PolicyBuilder.AddValuePolicy(Id0, ECachePolicy::Local);
 
 		FRequestOwner Owner(EPriority::Normal);
-		Hierarchy->Get({{{TEXTVIEW("PropagatePartial")}, Key, PolicyBuilder.Build()}}, Owner, [&Id0, &Value0Diff](FCacheGetResponse&& Response)
+		Hierarchy->Get({{{TEXTVIEW("PropagatePartial")}, Key, PolicyBuilder.Build(), 12345}}, Owner, [&Id0, &Value0Diff](FCacheGetResponse&& Response)
 		{
+			CHECK(Response.UserData == 12345);
 			CHECK(Response.Status == EStatus::Ok);
 			CHECK(Response.Record.GetValues().Num() == 3);
 			CHECK(Algo::AllOf(Response.Record.GetValues(), &FValue::HasData));
diff --git a/Engine/Source/Developer/DerivedDataCache/Tests/TestCacheStore.cpp b/Engine/Source/Developer/DerivedDataCache/Tests/TestCacheStore.cpp
--- a/Engine/Source/Developer/DerivedDataCache/Tests/TestCacheStore.cpp
+++ b/Engine/Source/Developer/DerivedDataCache/Tests/TestCacheStore.cpp
@@ -225,7 +225,7 @@ void FTestCacheStore::Put(
 {
 	const int32 Index = PutRequests.Num();
 	const int32 Count = Requests.Num();
-	PutRequests.Append(Requests);
+	PutRequests.Append(Requests.GetData(), Requests.Num());
 	TotalRequestCount += Count;
 
 	ExecuteOrQueue(Owner, [this, Index, Count, OnComplete = MoveTemp(OnComplete)](bool bCancel)
@@ -300,7 +300,7 @@ void FTestCacheStore::Get(
 {
 	const int32 Index = GetRequests.Num();
 	const int32 Count = Requests.Num();
-	GetRequests.Append(Requests);
+	GetRequests.Append(Requests.GetData(), Requests.Num());
 	TotalRequestCount += Count;
 
 	ExecuteOrQueue(Owner, [this, Index, Count, OnComplete = MoveTemp(OnComplete)](bool bCancel)
@@ -351,7 +351,7 @@ void FTestCacheStore::PutValue(
 {
 	const int32 Index = PutValueRequests.Num();
 	const int32 Count = Requests.Num();
-	PutValueRequests.Append(Requests);
+	PutValueRequests.Append(Requests.GetData(), Requests.Num());
 	TotalRequestCount += Count;
 
 	ExecuteOrQueue(Owner, [this, Index, Count, OnComplete = MoveTemp(OnComplete)](bool bCancel)
@@ -383,7 +383,7 @@ void FTestCacheStore::GetValue(
 {
 	const int32 Index = GetValueRequests.Num();
 	const int32 Count = Requests.Num();
-	GetValueRequests.Append(Requests);
+	GetValueRequests.Append(Requests.GetData(), Requests.Num());
 	TotalRequestCount += Count;
 
 	ExecuteOrQueue(Owner, [this, Index, Count, OnComplete = MoveTemp(OnComplete)](bool bCancel)
@@ -424,7 +424,7 @@ void FTestCacheStore::GetChunks(
 {
 	const int32 Index = GetChunkRequests.Num();
 	const int32 Count = Requests.Num();
-	GetChunkRequests.Append(Requests);
+	GetChunkRequests.Append(Requests.GetData(), Requests.Num());
 	TotalRequestCount += Count;
 
 	ExecuteOrQueue(Owner, [this, Index, Count, OnComplete = MoveTemp(OnComplete)](bool bCancel)
