DDC: Reverted unnecessary changes made by IWYU

There are several cases where includes were added that are not directly used (only by reference or pointer), and where forward declarations were reordered in a way that worsens readability.

#rnx

diff --git a/Engine/Source/Developer/DerivedDataCache/Public/DerivedDataBuildAction.h b/Engine/Source/Developer/DerivedDataCache/Public/DerivedDataBuildAction.h
--- a/Engine/Source/Developer/DerivedDataCache/Public/DerivedDataBuildAction.h
+++ b/Engine/Source/Developer/DerivedDataCache/Public/DerivedDataBuildAction.h
@@ -9,7 +9,6 @@
 #include "Misc/ScopeExit.h"
 #include "Templates/Function.h"
 #include "Templates/RefCounting.h"
-#include "Templates/RemoveReference.h"
 #include "Templates/UniquePtr.h"
 #include "Templates/UnrealTemplate.h"
 
@@ -17,12 +16,13 @@
 
 class FCbObject;
 class FCbWriter;
+struct FGuid;
+struct FIoHash;
+
 namespace UE::DerivedData { class FBuildAction; }
 namespace UE::DerivedData { class FBuildActionBuilder; }
 namespace UE::DerivedData { class FOptionalBuildAction; }
 namespace UE::DerivedData { struct FBuildActionKey; }
-struct FGuid;
-struct FIoHash;
 
 namespace UE::DerivedData::Private
 {
diff --git a/Engine/Source/Developer/DerivedDataCache/Public/DerivedDataBuildDefinition.h b/Engine/Source/Developer/DerivedDataCache/Public/DerivedDataBuildDefinition.h
--- a/Engine/Source/Developer/DerivedDataCache/Public/DerivedDataBuildDefinition.h
+++ b/Engine/Source/Developer/DerivedDataCache/Public/DerivedDataBuildDefinition.h
@@ -9,7 +9,6 @@
 #include "Misc/ScopeExit.h"
 #include "Templates/Function.h"
 #include "Templates/RefCounting.h"
-#include "Templates/RemoveReference.h"
 #include "Templates/UniquePtr.h"
 #include "Templates/UnrealTemplate.h"
 
@@ -17,13 +16,14 @@
 
 class FCbObject;
 class FCbWriter;
+struct FGuid;
+struct FIoHash;
+
 namespace UE::DerivedData { class FBuildDefinition; }
 namespace UE::DerivedData { class FBuildDefinitionBuilder; }
 namespace UE::DerivedData { class FOptionalBuildDefinition; }
 namespace UE::DerivedData { struct FBuildKey; }
 namespace UE::DerivedData { struct FBuildValueKey; }
-struct FGuid;
-struct FIoHash;
 
 namespace UE::DerivedData::Private
 {
diff --git a/Engine/Source/Developer/DerivedDataCache/Public/DerivedDataBuildFunctionRegistry.h b/Engine/Source/Developer/DerivedDataCache/Public/DerivedDataBuildFunctionRegistry.h
--- a/Engine/Source/Developer/DerivedDataCache/Public/DerivedDataBuildFunctionRegistry.h
+++ b/Engine/Source/Developer/DerivedDataCache/Public/DerivedDataBuildFunctionRegistry.h
@@ -3,12 +3,12 @@
 #pragma once
 
 #include "Containers/StringFwd.h"
-#include "Templates/Function.h"
 
-namespace UE::DerivedData { class IBuildFunction; }
 struct FGuid;
 template <typename FuncType> class TFunctionRef;
 
+namespace UE::DerivedData { class IBuildFunction; }
+
 namespace UE::DerivedData
 {
 
diff --git a/Engine/Source/Developer/DerivedDataCache/Public/DerivedDataBuildInputs.h b/Engine/Source/Developer/DerivedDataCache/Public/DerivedDataBuildInputs.h
--- a/Engine/Source/Developer/DerivedDataCache/Public/DerivedDataBuildInputs.h
+++ b/Engine/Source/Developer/DerivedDataCache/Public/DerivedDataBuildInputs.h
@@ -9,11 +9,11 @@
 #include "Misc/ScopeExit.h"
 #include "Templates/Function.h"
 #include "Templates/RefCounting.h"
-#include "Templates/RemoveReference.h"
 #include "Templates/UniquePtr.h"
 #include "Templates/UnrealTemplate.h"
 
 class FCompressedBuffer;
+
 namespace UE::DerivedData { class FBuildInputs; }
 namespace UE::DerivedData { class FBuildInputsBuilder; }
 
diff --git a/Engine/Source/Developer/DerivedDataCache/Public/DerivedDataBuildJob.h b/Engine/Source/Developer/DerivedDataCache/Public/DerivedDataBuildJob.h
--- a/Engine/Source/Developer/DerivedDataCache/Public/DerivedDataBuildJob.h
+++ b/Engine/Source/Developer/DerivedDataCache/Public/DerivedDataBuildJob.h
@@ -4,17 +4,16 @@
 
 #include "Containers/StringFwd.h"
 #include "CoreTypes.h"
-#include "DerivedDataBuildOutput.h"
 #include "DerivedDataRequestTypes.h"
 #include "DerivedDataSharedStringFwd.h"
 
+template <typename FuncType> class TUniqueFunction;
+
 namespace UE::DerivedData { class FBuildOutput; }
 namespace UE::DerivedData { class IBuild; }
 namespace UE::DerivedData { class ICache; }
 namespace UE::DerivedData { struct FBuildJobCompleteParams; }
 namespace UE::DerivedData { struct FCacheKey; }
-template <typename FuncType> class TUniqueFunction;
-
 namespace UE::DerivedData { enum class EBuildStatus : uint32; }
 
 namespace UE::DerivedData
diff --git a/Engine/Source/Developer/DerivedDataCache/Public/DerivedDataBuildOutput.h b/Engine/Source/Developer/DerivedDataCache/Public/DerivedDataBuildOutput.h
--- a/Engine/Source/Developer/DerivedDataCache/Public/DerivedDataBuildOutput.h
+++ b/Engine/Source/Developer/DerivedDataCache/Public/DerivedDataBuildOutput.h
@@ -9,9 +9,7 @@
 #include "CoreTypes.h"
 #include "DerivedDataSharedStringFwd.h"
 #include "Misc/ScopeExit.h"
-#include "Serialization/CompactBinary.h"
 #include "Templates/RefCounting.h"
-#include "Templates/RemoveReference.h"
 #include "Templates/UniquePtr.h"
 #include "Templates/UnrealTemplate.h"
 
@@ -19,6 +17,7 @@
 
 class FCbObject;
 class FCbWriter;
+
 namespace UE::DerivedData { class FBuildOutput; }
 namespace UE::DerivedData { class FBuildOutputBuilder; }
 namespace UE::DerivedData { class FCacheRecord; }
diff --git a/Engine/Source/Developer/DerivedDataCache/Public/DerivedDataBuildScheduler.h b/Engine/Source/Developer/DerivedDataCache/Public/DerivedDataBuildScheduler.h
--- a/Engine/Source/Developer/DerivedDataCache/Public/DerivedDataBuildScheduler.h
+++ b/Engine/Source/Developer/DerivedDataCache/Public/DerivedDataBuildScheduler.h
@@ -12,7 +12,6 @@
 namespace UE::DerivedData { class IBuildJobSchedule; }
 namespace UE::DerivedData { class IRequestOwner; }
 namespace UE::DerivedData { struct FBuildSchedulerParams; }
-namespace UE { namespace DerivedData { class IBuildJob; } }
 
 namespace UE::DerivedData
 {
diff --git a/Engine/Source/Developer/DerivedDataCache/Public/DerivedDataBuildSession.h b/Engine/Source/Developer/DerivedDataCache/Public/DerivedDataBuildSession.h
--- a/Engine/Source/Developer/DerivedDataCache/Public/DerivedDataBuildSession.h
+++ b/Engine/Source/Developer/DerivedDataCache/Public/DerivedDataBuildSession.h
@@ -6,7 +6,6 @@
 #include "DerivedDataBuildTypes.h"
 #include "DerivedDataSharedStringFwd.h"
 #include "Templates/Function.h"
-#include "Templates/RemoveReference.h"
 #include "Templates/UniquePtr.h"
 #include "Templates/UnrealTemplate.h"
 
diff --git a/Engine/Source/Developer/DerivedDataCache/Public/DerivedDataCacheKey.h b/Engine/Source/Developer/DerivedDataCache/Public/DerivedDataCacheKey.h
--- a/Engine/Source/Developer/DerivedDataCache/Public/DerivedDataCacheKey.h
+++ b/Engine/Source/Developer/DerivedDataCache/Public/DerivedDataCacheKey.h
@@ -10,8 +10,6 @@
 #include "Misc/CString.h"
 #include "Templates/TypeHash.h"
 
-template <typename CharType> class TStringBuilderBase;
-
 #define UE_API DERIVEDDATACACHE_API
 
 class FCbFieldView;
diff --git a/Engine/Source/Developer/DerivedDataCache/Public/DerivedDataCachePolicy.h b/Engine/Source/Developer/DerivedDataCache/Public/DerivedDataCachePolicy.h
--- a/Engine/Source/Developer/DerivedDataCache/Public/DerivedDataCachePolicy.h
+++ b/Engine/Source/Developer/DerivedDataCache/Public/DerivedDataCachePolicy.h
@@ -9,7 +9,6 @@
 #include "DerivedDataValueId.h"
 #include "Misc/EnumClassFlags.h"
 #include "Templates/RefCounting.h"
-#include "Templates/RemoveReference.h"
 #include "Templates/UnrealTemplate.h"
 
 #define UE_API DERIVEDDATACACHE_API
@@ -17,9 +16,10 @@
 class FCbFieldView;
 class FCbObjectView;
 class FCbWriter;
+template <typename FuncType> class TFunctionRef;
+
 namespace UE::DerivedData::Private { class ICacheRecordPolicyShared; }
 namespace UE::DerivedData { class FOptionalCacheRecordPolicy; }
-template <typename FuncType> class TFunctionRef;
 
 namespace UE::DerivedData
 {
diff --git a/Engine/Source/Developer/DerivedDataCache/Public/DerivedDataCacheRecord.h b/Engine/Source/Developer/DerivedDataCache/Public/DerivedDataCacheRecord.h
--- a/Engine/Source/Developer/DerivedDataCache/Public/DerivedDataCacheRecord.h
+++ b/Engine/Source/Developer/DerivedDataCache/Public/DerivedDataCacheRecord.h
@@ -8,27 +8,22 @@
 #include "HAL/Platform.h"
 #include "Memory/MemoryFwd.h"
 #include "Misc/ScopeExit.h"
-#include "Serialization/CompactBinary.h"
 #include "Templates/Function.h"
 #include "Templates/RefCounting.h"
-#include "Templates/RemoveReference.h"
 #include "Templates/UniquePtr.h"
 #include "Templates/UnrealTemplate.h"
 
-class FCbWriter;
-class FCompositeBuffer;
-class FSharedBuffer;
-
 #define UE_API DERIVEDDATACACHE_API
 
 class FCbObject;
 class FCbPackage;
+class FCbWriter;
+
 namespace UE::DerivedData { class FCacheRecord; }
 namespace UE::DerivedData { class FOptionalCacheRecord; }
 namespace UE::DerivedData { class FValue; }
 namespace UE::DerivedData { class FValueWithId; }
 namespace UE::DerivedData { class IRequestOwner; }
-
 namespace UE::DerivedData { using FOnCacheRecordComplete = TUniqueFunction<void (FCacheRecord&& Record)>; }
 
 ///////////////////////////////////////////////////////////////////////////////////////////////////////////////////////
diff --git a/Engine/Source/Developer/DerivedDataCache/Public/DerivedDataCacheUsageStats.h b/Engine/Source/Developer/DerivedDataCache/Public/DerivedDataCacheUsageStats.h
--- a/Engine/Source/Developer/DerivedDataCache/Public/DerivedDataCacheUsageStats.h
+++ b/Engine/Source/Developer/DerivedDataCache/Public/DerivedDataCacheUsageStats.h
@@ -3,15 +3,12 @@
 
 #include "Containers/Array.h"
 #include "Containers/Map.h"
-#include "Containers/Set.h"
-#include "Containers/SparseArray.h"
 #include "Containers/UnrealString.h"
 #include "CoreMinimal.h"
 #include "HAL/Platform.h"
 #include "ProfilingDebugging/CookStats.h"
 #include "Templates/Function.h"
 #include "Templates/SharedPointer.h"
-#include "Templates/Tuple.h"
 #include "Templates/UnrealTemplate.h"
 
 /**
diff --git a/Engine/Source/Developer/DerivedDataCache/Public/DerivedDataRequestOwner.h b/Engine/Source/Developer/DerivedDataCache/Public/DerivedDataRequestOwner.h
--- a/Engine/Source/Developer/DerivedDataCache/Public/DerivedDataRequestOwner.h
+++ b/Engine/Source/Developer/DerivedDataCache/Public/DerivedDataRequestOwner.h
@@ -3,9 +3,7 @@
 #pragma once
 
 #include "CoreTypes.h"
-#include "DerivedDataRequest.h"
 #include "Misc/EnumClassFlags.h"
-#include "Templates/Function.h"
 #include "Templates/Invoke.h"
 #include "Templates/RefCounting.h"
 #include "Templates/UniquePtr.h"
@@ -13,9 +11,9 @@
 
 #define UE_API DERIVEDDATACACHE_API
 
-namespace UE::DerivedData { class IRequest; }
 template <typename FuncType> class TUniqueFunction;
 
+namespace UE::DerivedData { class IRequest; }
 namespace UE::DerivedData { enum class EPriority : uint8; }
 
 namespace UE::DerivedData
diff --git a/Engine/Source/Developer/DerivedDataCache/Public/DerivedDataSharedString.h b/Engine/Source/Developer/DerivedDataCache/Public/DerivedDataSharedString.h
--- a/Engine/Source/Developer/DerivedDataCache/Public/DerivedDataSharedString.h
+++ b/Engine/Source/Developer/DerivedDataCache/Public/DerivedDataSharedString.h
@@ -6,7 +6,6 @@
 #include "CoreTypes.h"
 #include "DerivedDataSharedStringFwd.h"
 #include "HAL/MemoryBase.h"
-#include "HAL/PlatformCrt.h"
 #include "HAL/UnrealMemory.h"
 #include "Misc/CString.h"
 
diff --git a/Engine/Source/Developer/DerivedDataCache/Public/DerivedDataValue.h b/Engine/Source/Developer/DerivedDataCache/Public/DerivedDataValue.h
--- a/Engine/Source/Developer/DerivedDataCache/Public/DerivedDataValue.h
+++ b/Engine/Source/Developer/DerivedDataCache/Public/DerivedDataValue.h
@@ -8,12 +8,8 @@
 #include "IO/IoHash.h"
 #include "Memory/MemoryFwd.h"
 #include "Misc/AssertionMacros.h"
-#include "Templates/RemoveReference.h"
 #include "Templates/UnrealTemplate.h"
 
-class FCompositeBuffer;
-class FSharedBuffer;
-
 #define UE_API DERIVEDDATACACHE_API
 
 namespace UE::DerivedData
