DDC: Fixed a non-unity error in TestCacheStore

#rnx

diff --git a/Engine/Source/Developer/DerivedDataCache/Tests/TestCacheStore.cpp b/Engine/Source/Developer/DerivedDataCache/Tests/TestCacheStore.cpp
--- a/Engine/Source/Developer/DerivedDataCache/Tests/TestCacheStore.cpp
+++ b/Engine/Source/Developer/DerivedDataCache/Tests/TestCacheStore.cpp
@@ -9,6 +9,7 @@
 #include "Containers/Map.h"
 #include "DerivedDataRequest.h"
 #include "DerivedDataRequestOwner.h"
+#include "Serialization/CompactBinary.h"
 #include "Templates/Function.h"
 
 namespace UE::DerivedData
