Disabling the -DDC=None codepath for non-programs to reduce confusion when people attempt to use it in editor.

diff --git a/Engine/Source/Developer/DerivedDataCache/Private/DerivedDataBackends.cpp b/Engine/Source/Developer/DerivedDataCache/Private/DerivedDataBackends.cpp
--- a/Engine/Source/Developer/DerivedDataCache/Private/DerivedDataBackends.cpp
+++ b/Engine/Source/Developer/DerivedDataCache/Private/DerivedDataBackends.cpp
@@ -129,11 +129,20 @@ public:
 			GraphName = GDerivedDataCacheGraphName.GetValueOnGameThread();
 		}
 
+		// A DDC graph of "None" is used by build worker programs that use the DDC build code paths but avoid use of the DDC cache
+		// code paths. Unfortunately the cache must currently be instantiated as part of initializing the build code, so we have
+		// to disable the cache portion by injecting "-DDC=None" to the commandline args during process startup. This mode is not
+		// compatible with use in the editor/commandlets which is not written to operate with a non-functional cache layer. This
+		// can lead to confusion when people attempt to use "-DDC=None" in the editor expecting it to behave like "-DDC=Cold".
+		// To avoid this confusion (and until we can use the build without the cache) it is restricted to use in programs only
+		// and not the editor.
+#if IS_PROGRAM
 		if (GraphName == TEXT("None"))
 		{
 			UE_LOG(LogDerivedDataCache, Display, TEXT("Requested cache graph of 'None'. Every cache operation will fail."));
 		}
 		else
+#endif
 		{
 			const TCHAR* const RootName = TEXT("Root");
 
