Fix bugs in the System.DerivedDataCache.HttpDerivedDataBackend.CacheStore automated test caused by:

1. Using the wrong structured namespace when instantiating Zen.
2. Assuming that the order that PUT requests will complete is equal to the order that they were in the original request array.

diff --git a/Engine/Source/Developer/DerivedDataCache/Private/Tests/HttpCacheStoreTest.cpp b/Engine/Source/Developer/DerivedDataCache/Private/Tests/HttpCacheStoreTest.cpp
--- a/Engine/Source/Developer/DerivedDataCache/Private/Tests/HttpCacheStoreTest.cpp
+++ b/Engine/Source/Developer/DerivedDataCache/Private/Tests/HttpCacheStoreTest.cpp
@@ -548,10 +548,15 @@ TArray<FCacheRecord> CreateTestCacheRecords(ICacheStore* InTestBackend, uint32 I
 	InTestBackend->Put(PutRequests, Owner, [&CacheRecords, &PutRequests] (FCachePutResponse&& Response)
 	{
 		check(Response.Status == EStatus::Ok);
-		CacheRecords.Add(PutRequests[Response.UserData].Record);
 	});
 	Owner.Wait();
 
+	CacheRecords.Reserve(PutRequests.Num());
+	for (const FCachePutRequest& PutRequest : PutRequests)
+	{
+		CacheRecords.Add(PutRequest.Record);
+	}
+
 	return CacheRecords;
 }
 
@@ -594,10 +599,15 @@ TArray<FValue> CreateTestCacheValues(ICacheStore* InTestBackend, uint32 InNumVal
 	InTestBackend->PutValue(PutValueRequests, Owner, [&Values, &PutValueRequests](FCachePutValueResponse&& Response)
 		{
 			check(Response.Status == EStatus::Ok);
-			Values.Add(PutValueRequests[Response.UserData].Value);
 		});
 	Owner.Wait();
 
+	Values.Reserve(PutValueRequests.Num());
+	for (const FCachePutValueRequest& PutValueRequest : PutValueRequests)
+	{
+		Values.Add(PutValueRequest.Value);
+	}
+
 	return Values;
 }
 
@@ -666,11 +676,11 @@ bool FHttpCacheStoreTest::RunTest(const FString& Parameters)
 
 	FScopeZenService ScopeZenSiblingService(MoveTemp(ZenTestServiceSiblingSettings));
 	TUniquePtr<ILegacyCacheStore> ZenIntermediarySiblingBackend(CreateZenCacheStore(TEXT("TestSibling"),
-		*FString::Printf(TEXT("Host=%s, Namespace=%s, StructuredNamespace=%s"), ScopeZenSiblingService.GetInstance().GetURL(), *TestNamespace, *TestNamespace)).Key);
+		*FString::Printf(TEXT("Host=%s, Namespace=%s, StructuredNamespace=%s"), ScopeZenSiblingService.GetInstance().GetURL(), *TestNamespace, *TestStructuredNamespace)).Key);
 
 	FScopeZenService ScopeZenService(MoveTemp(ZenTestServiceSettings));
 	TUniquePtr<ILegacyCacheStore> ZenIntermediaryBackend(CreateZenCacheStore(TEXT("Test"),
-		*FString::Printf(TEXT("Host=%s, Namespace=%s, StructuredNamespace=%s"), ScopeZenService.GetInstance().GetURL(), *TestNamespace, *TestNamespace)).Key);
+		*FString::Printf(TEXT("Host=%s, Namespace=%s, StructuredNamespace=%s"), ScopeZenService.GetInstance().GetURL(), *TestNamespace, *TestStructuredNamespace)).Key);
 	auto WaitForZenPushToUpstream = [](ILegacyCacheStore* ZenBackend, TConstArrayView<FCacheRecord> Records)
 	{
 		// TODO: Expecting a legitimate means to wait for zen to finish pushing records to its upstream in the future
