DDC: Fixed unsafe type cast warnings

#rnx

diff --git a/Engine/Source/Developer/DerivedDataCache/Private/DerivedDataBuildOutput.cpp b/Engine/Source/Developer/DerivedDataCache/Private/DerivedDataBuildOutput.cpp
--- a/Engine/Source/Developer/DerivedDataCache/Private/DerivedDataBuildOutput.cpp
+++ b/Engine/Source/Developer/DerivedDataCache/Private/DerivedDataBuildOutput.cpp
@@ -267,12 +267,14 @@ bool FBuildOutputInternal::TryLoad()
 
 	if (FCbFieldView MessagesField = OutputView[ANSITEXTVIEW("Messages")])
 	{
-		if (!MessagesField.IsArray())
+		const FCbArrayView MessagesArray = MessagesField.AsArrayView();
+		const uint64 MessageCount = MessagesArray.Num();
+		if (MessagesField.HasError() || MessageCount > MAX_int32)
 		{
 			return false;
 		}
-		Messages.Reserve(MessagesField.AsArrayView().Num());
-		for (FCbFieldView MessageField : MessagesField)
+		Messages.Reserve(int32(MessageCount));
+		for (FCbFieldView MessageField : MessagesArray)
 		{
 			const FUtf8StringView LevelName = MessageField[ANSITEXTVIEW("Level")].AsString();
 			const FUtf8StringView Message = MessageField[ANSITEXTVIEW("Message")].AsString();
@@ -287,12 +289,14 @@ bool FBuildOutputInternal::TryLoad()
 
 	if (FCbFieldView LogsField = OutputView[ANSITEXTVIEW("Logs")])
 	{
-		if (!LogsField.IsArray())
+		const FCbArrayView LogsArray = LogsField.AsArrayView();
+		const uint64 LogCount = LogsArray.Num();
+		if (LogsField.HasError() || LogCount > MAX_int32)
 		{
 			return false;
 		}
-		Logs.Reserve(LogsField.AsArrayView().Num());
-		for (FCbFieldView LogField : LogsField)
+		Logs.Reserve(int32(LogCount));
+		for (FCbFieldView LogField : LogsArray)
 		{
 			const FUtf8StringView LevelName = LogField[ANSITEXTVIEW("Level")].AsString();
 			const FUtf8StringView Category = LogField[ANSITEXTVIEW("Category")].AsString();
diff --git a/Engine/Source/Developer/DerivedDataCache/Private/DerivedDataCacheReplay.cpp b/Engine/Source/Developer/DerivedDataCache/Private/DerivedDataCacheReplay.cpp
--- a/Engine/Source/Developer/DerivedDataCache/Private/DerivedDataCacheReplay.cpp
+++ b/Engine/Source/Developer/DerivedDataCache/Private/DerivedDataCacheReplay.cpp
@@ -380,7 +380,7 @@ bool FCacheReplayReader::FState::DispatchRequests(
 	{
 		TRACE_CPUPROFILER_EVENT_SCOPE(ReplayDDC_Serialize);
 		const FCbArrayView Array = Object[ANSITEXTVIEW("Requests")].AsArrayView();
-		Requests.Reserve(Array.Num());
+		Requests.Reserve(IntCastChecked<int32>(Array.Num()));
 		for (FCbFieldView Field : Array)
 		{
 			RequestType& Request = Requests.AddDefaulted_GetRef();
diff --git a/Engine/Source/Developer/DerivedDataCache/Private/DerivedDataCacheStoreHierarchy.cpp b/Engine/Source/Developer/DerivedDataCache/Private/DerivedDataCacheStoreHierarchy.cpp
--- a/Engine/Source/Developer/DerivedDataCache/Private/DerivedDataCacheStoreHierarchy.cpp
+++ b/Engine/Source/Developer/DerivedDataCache/Private/DerivedDataCacheStoreHierarchy.cpp
@@ -175,7 +175,7 @@ struct FCacheStoreHierarchy::TBatchParams
 	static void ModifyPolicyForResponse(decltype(FGetRequest::Policy)& Policy, const FGetResponse& Response) {}
 	static FGetResponse FilterResponseByRequest(const FGetResponse& Response, const FGetRequest& Request);
 	static FPutRequest MakePutRequest(const FGetResponse& Response, const FGetRequest& Request);
-	static FGetRequest MakeGetRequest(const FPutRequest& Request, uint64 UserData);
+	static FGetRequest MakeGetRequest(const FPutRequest& Request, int32 RequestIndex);
 };
 
 ///////////////////////////////////////////////////////////////////////////////////////////////////////////////////////
@@ -441,7 +441,7 @@ bool FCacheStoreHierarchy::TPutBatch<Params>::DispatchGetRequests()
 	TArray<FGetRequest, TInlineAllocator<1>> NodeRequests;
 	NodeRequests.Reserve(Requests.Num());
 
-	uint64 RequestIndex = 0;
+	int32 RequestIndex = 0;
 	for (const FPutRequest& Request : Requests)
 	{
 		if (!States[RequestIndex].bStop && CanQuery(GetCombinedPolicy(Request.Policy), Node.CacheFlags))
@@ -603,7 +603,7 @@ private:
 	{
 		FGetRequest Request;
 		FGetResponse Response;
-		uint64 NodeIndex = 0;
+		int32 NodeIndex = 0;
 	};
 
 	const FCacheStoreHierarchy& Hierarchy;
@@ -956,9 +956,9 @@ FCachePutRequest FCacheStoreHierarchy::FCacheRecordBatchParams::MakePutRequest(
 template <>
 FCacheGetRequest FCacheStoreHierarchy::FCacheRecordBatchParams::MakeGetRequest(
 	const FCachePutRequest& Request,
-	const uint64 UserData)
+	const int32 RequestIndex)
 {
-	return {Request.Name, Request.Record.GetKey(), AddPolicy(Request.Policy, ECachePolicy::SkipData), UserData};
+	return {Request.Name, Request.Record.GetKey(), AddPolicy(Request.Policy, ECachePolicy::SkipData), uint64(RequestIndex)};
 }
 
 void FCacheStoreHierarchy::Put(
@@ -1032,9 +1032,9 @@ FCachePutValueRequest FCacheStoreHierarchy::FCacheValueBatchParams::MakePutReque
 template <>
 FCacheGetValueRequest FCacheStoreHierarchy::FCacheValueBatchParams::MakeGetRequest(
 	const FCachePutValueRequest& Request,
-	const uint64 UserData)
+	const int32 RequestIndex)
 {
-	return {Request.Name, Request.Key, AddPolicy(Request.Policy, ECachePolicy::SkipData), UserData};
+	return {Request.Name, Request.Key, AddPolicy(Request.Policy, ECachePolicy::SkipData), uint64(RequestIndex)};
 }
 
 void FCacheStoreHierarchy::PutValue(
diff --git a/Engine/Source/Developer/DerivedDataCache/Private/DerivedDataCacheStoreThrottle.cpp b/Engine/Source/Developer/DerivedDataCache/Private/DerivedDataCacheStoreThrottle.cpp
--- a/Engine/Source/Developer/DerivedDataCache/Private/DerivedDataCacheStoreThrottle.cpp
+++ b/Engine/Source/Developer/DerivedDataCache/Private/DerivedDataCacheStoreThrottle.cpp
@@ -169,7 +169,7 @@ private:
 			if (ExpectedTime > ActualTime)
 			{
 				TRACE_CPUPROFILER_EVENT_SCOPE(ThrottlingBandwidth);
-				FPlatformProcess::Sleep(ExpectedTime - ActualTime);
+				FPlatformProcess::Sleep(float(ExpectedTime - ActualTime));
 			}
 		}
 	}
diff --git a/Engine/Source/Developer/DerivedDataCache/Private/FileSystemCacheStore.cpp b/Engine/Source/Developer/DerivedDataCache/Private/FileSystemCacheStore.cpp
--- a/Engine/Source/Developer/DerivedDataCache/Private/FileSystemCacheStore.cpp
+++ b/Engine/Source/Developer/DerivedDataCache/Private/FileSystemCacheStore.cpp
@@ -2354,7 +2354,7 @@ bool FFileSystemCacheStore::SaveFile(
 	}
 
 	const double WriteDuration = FPlatformTime::Seconds() - StartTime;
-	const double WriteSpeed = WriteDuration > 0.001 ? (WriteSize / WriteDuration) / (1024.0 * 1024.0) : 0.0;
+	const double WriteSpeed = WriteDuration > 0.001 ? (double(WriteSize) / WriteDuration) / (1024.0 * 1024.0) : 0.0;
 	UE_LOG(LogDerivedDataCache, VeryVerbose,
 		TEXT("%s: Saved %s from '%.*s' (%" INT64_FMT " bytes, %.02f secs, %.2f MiB/s)"),
 		*CachePath, *Path, DebugName.Len(), DebugName.GetData(), WriteSize, WriteDuration, WriteSpeed);
@@ -2393,7 +2393,7 @@ bool FFileSystemCacheStore::LoadFile(
 	else
 	{
 		const double ReadDuration = FPlatformTime::Seconds() - StartTime;
-		const double ReadSpeed = ReadDuration > 0.001 ? (ReadSize / ReadDuration) / (1024.0 * 1024.0) : 0.0;
+		const double ReadSpeed = ReadDuration > 0.001 ? (double(ReadSize) / ReadDuration) / (1024.0 * 1024.0) : 0.0;
 
 		UE_LOG(LogDerivedDataCache, VeryVerbose,
 			TEXT("%s: Loaded %s from '%.*s' (%" INT64_FMT " bytes, %.02f secs, %.2f MiB/s)"),
diff --git a/Engine/Source/Developer/DerivedDataCache/Private/HttpCacheStore.cpp b/Engine/Source/Developer/DerivedDataCache/Private/HttpCacheStore.cpp
--- a/Engine/Source/Developer/DerivedDataCache/Private/HttpCacheStore.cpp
+++ b/Engine/Source/Developer/DerivedDataCache/Private/HttpCacheStore.cpp
@@ -1786,7 +1786,7 @@ void FHttpCacheStore::SetAccessToken(FStringView Token, double RefreshDelay)
 					AcquireAccessToken();
 					return false;
 				}
-			), RefreshDelay - RefreshGracePeriod);
+			), float(FMath::Min(RefreshDelay - RefreshGracePeriod, MAX_flt)));
 		}
 
 		// Schedule a forced refresh of the token when the scheduled refresh is starved or unavailable.
@@ -1803,7 +1803,7 @@ void FHttpCacheStore::SetAccessToken(FStringView Token, double RefreshDelay)
 
 TUniquePtr<FHttpCacheStore::FHttpOperation> FHttpCacheStore::WaitForHttpOperation(EOperationCategory Category, bool bUnboundedOverflow)
 {
-	if (Access && RefreshAccessTokenTime && RefreshAccessTokenTime < FPlatformTime::Seconds())
+	if (Access && RefreshAccessTokenTime > 0.0 && RefreshAccessTokenTime < FPlatformTime::Seconds())
 	{
 		AcquireAccessToken();
 	}
diff --git a/Engine/Source/Developer/DerivedDataCache/Private/PakFileCacheStore.cpp b/Engine/Source/Developer/DerivedDataCache/Private/PakFileCacheStore.cpp
--- a/Engine/Source/Developer/DerivedDataCache/Private/PakFileCacheStore.cpp
+++ b/Engine/Source/Developer/DerivedDataCache/Private/PakFileCacheStore.cpp
@@ -302,12 +302,12 @@ bool FPakFileCacheStore::GetCachedData(const TCHAR* CacheKey, TArray<uint8>& Out
 		{
 			check(Item->Size);
 			check(!OutData.Num());
-			OutData.AddUninitialized(Item->Size);
+			OutData.AddUninitialized(int32(Item->Size));
 			if (!FileHandle->Read(OutData.GetData(), int64(Item->Size)))
 			{
 				UE_LOG(LogDerivedDataCache, Warning, TEXT("%s: Pak file, bad read."), *CachePath);
 			}
-			else if (uint32 TestCrc = FCrc::MemCrc_DEPRECATED(OutData.GetData(), Item->Size); TestCrc != Item->Crc)
+			else if (uint32 TestCrc = FCrc::MemCrc_DEPRECATED(OutData.GetData(), int32(Item->Size)); TestCrc != Item->Crc)
 			{
 				UE_LOG(LogDerivedDataCache, Warning, TEXT("%s: Pak file, bad crc."), *CachePath);
 			}
diff --git a/Engine/Source/Developer/DerivedDataCache/Private/Tests/HttpCacheStoreTest.cpp b/Engine/Source/Developer/DerivedDataCache/Private/Tests/HttpCacheStoreTest.cpp
--- a/Engine/Source/Developer/DerivedDataCache/Private/Tests/HttpCacheStoreTest.cpp
+++ b/Engine/Source/Developer/DerivedDataCache/Private/Tests/HttpCacheStoreTest.cpp
@@ -74,7 +74,7 @@ protected:
 	{
 		std::atomic<uint64> Requests{ 0 };
 		std::atomic<uint64> MaxLatency{ 0 };
-		std::atomic<uint64> TotalMS{ 0 };
+		std::atomic<uint64> TotalCycles{ 0 };
 		std::atomic<uint64> TotalRequests{ 0 };
 
 		FEvent* StartEvent = FPlatformProcess::GetSynchEventFromPool(true);
@@ -94,11 +94,11 @@ protected:
 
 					while (FPlatformTime::Seconds() < StopTime.load(std::memory_order_relaxed))
 					{
-						uint64 Before = FPlatformTime::Cycles64();
+						const uint64 Before = FPlatformTime::Cycles64();
 						TestFunction();
-						uint64 Delta = FPlatformTime::Cycles64() - Before;
+						const uint64 Delta = FPlatformTime::Cycles64() - Before;
 						Requests++;
-						TotalMS += FPlatformTime::ToMilliseconds64(Delta);
+						TotalCycles += Delta;
 						TotalRequests++;
 
 						// Compare exchange loop until we either succeed to set the maximum value
@@ -143,11 +143,13 @@ protected:
 
 			if (TotalRequests)
 			{
-				UE_LOG(LogHttpDerivedDataBackendTests, Display, TEXT("RPS: %llu, AvgLatency: %.02f ms, MaxLatency: %.02f s"), Requests.exchange(0), double(TotalMS) / TotalRequests, FPlatformTime::ToSeconds(MaxLatency));
+				UE_LOG(LogHttpDerivedDataBackendTests, Display, TEXT("RPS: %" UINT64_FMT ", AvgLatency: %.02f ms, MaxLatency: %.02f s"),
+					Requests.exchange(0), FPlatformTime::ToMilliseconds64(TotalCycles) / double(TotalRequests), FPlatformTime::ToSeconds64(MaxLatency));
 			}
 			else
 			{
-				UE_LOG(LogHttpDerivedDataBackendTests, Display, TEXT("RPS: %llu, AvgLatency: N/A, MaxLatency: %.02f s"), Requests.exchange(0), FPlatformTime::ToSeconds(MaxLatency));
+				UE_LOG(LogHttpDerivedDataBackendTests, Display, TEXT("RPS: %" UINT64_FMT ", AvgLatency: N/A, MaxLatency: %.02f s"),
+					Requests.exchange(0), FPlatformTime::ToSeconds64(MaxLatency));
 			}
 		}
 
@@ -203,7 +205,7 @@ protected:
 						RecordBuilder.AddValue(Value);
 					}
 				}
-				GetOutputs[Response.UserData].Emplace(FGetOutput{ RecordBuilder.Build(), Response.Status });
+				GetOutputs[int32(Response.UserData)].Emplace(FGetOutput{ RecordBuilder.Build(), Response.Status });
 			});
 		RequestOwner.Wait();
 
@@ -249,7 +251,7 @@ protected:
 		FRequestOwner RequestOwner(EPriority::Blocking);
 		TestBackend->GetValue(Requests, RequestOwner, [&GetValueOutputs](FCacheGetValueResponse&& Response)
 			{
-				GetValueOutputs[Response.UserData].Emplace(FGetValueOutput{ Response.Value, Response.Status });
+				GetValueOutputs[int32(Response.UserData)].Emplace(FGetValueOutput{ Response.Value, Response.Status });
 			});
 		RequestOwner.Wait();
 
@@ -297,7 +299,7 @@ protected:
 		FRequestOwner RequestOwner(EPriority::Blocking);
 		TestBackend->GetChunks(Requests, RequestOwner, [&GetOutputs](FCacheGetChunkResponse&& Response)
 			{
-				GetOutputs[Response.UserData].Emplace(FGetChunksOutput { Response.RawData, Response.Status });
+				GetOutputs[int32(Response.UserData)].Emplace(FGetChunksOutput { Response.RawData, Response.Status });
 			});
 		RequestOwner.Wait();
 
diff --git a/Engine/Source/Developer/DerivedDataCache/Private/ZenCacheStore.cpp b/Engine/Source/Developer/DerivedDataCache/Private/ZenCacheStore.cpp
--- a/Engine/Source/Developer/DerivedDataCache/Private/ZenCacheStore.cpp
+++ b/Engine/Source/Developer/DerivedDataCache/Private/ZenCacheStore.cpp
@@ -990,7 +990,7 @@ FZenCacheStore::FZenCacheStore(
 	{
 		RpcUri << ZenService.GetInstance().GetURL() << ANSITEXTVIEW("/z$/$rpc");
 
-		const uint32 MaxConnections = FMath::Clamp(static_cast<uint32>(FPlatformMisc::NumberOfCoresIncludingHyperthreads()), 8, 64);
+		const uint32 MaxConnections = uint32(FMath::Clamp(FPlatformMisc::NumberOfCoresIncludingHyperthreads(), 8, 64));
 		constexpr uint32 RequestPoolSize = 128;
 		constexpr uint32 RequestPoolOverflowSize = 128;
 
