DDC: Fixed the cache hierarchy test to execute async puts

#rnx

diff --git a/Engine/Source/Developer/DerivedDataCache/Tests/DerivedDataCacheStoreHierarchyTest.cpp b/Engine/Source/Developer/DerivedDataCache/Tests/DerivedDataCacheStoreHierarchyTest.cpp
--- a/Engine/Source/Developer/DerivedDataCache/Tests/DerivedDataCacheStoreHierarchyTest.cpp
+++ b/Engine/Source/Developer/DerivedDataCache/Tests/DerivedDataCacheStoreHierarchyTest.cpp
@@ -206,6 +206,10 @@ TEST_CASE("DerivedData::Cache::Hierarchy::PartialNonDeterministicRecordPropagati
 		CHECK(Request1Repeat.Policy.GetValuePolicies().Num() == 3);
 		Store1->ExecuteAsync();
 
+		// Execute stores to Store0 and Store2.
+		Store0->ExecuteAsync();
+		Store2->ExecuteAsync();
+
 		// Wait to ensure that the counts checked below are final.
 		Owner.Wait();
 
@@ -257,6 +261,10 @@ TEST_CASE("DerivedData::Cache::Hierarchy::PartialNonDeterministicRecordPropagati
 		CHECK(Request1.Policy.GetValuePolicies().Num() == 2);
 		Store1->ExecuteAsync();
 
+		// Execute stores to Store0 and Store2.
+		Store0->ExecuteAsync();
+		Store2->ExecuteAsync();
+
 		// Wait to ensure that the counts checked below are final.
 		Owner.Wait();
 
diff --git a/Engine/Source/Developer/DerivedDataCache/Tests/TestCacheStore.cpp b/Engine/Source/Developer/DerivedDataCache/Tests/TestCacheStore.cpp
--- a/Engine/Source/Developer/DerivedDataCache/Tests/TestCacheStore.cpp
+++ b/Engine/Source/Developer/DerivedDataCache/Tests/TestCacheStore.cpp
@@ -19,6 +19,8 @@ class FTestCacheStore final : public ITestCacheStore
 {
 public:
 	explicit FTestCacheStore(ECacheStoreFlags Flags, bool bAsync = false);
+	~FTestCacheStore() final;
+
 	FTestCacheStore(const FTestCacheStore&) = delete;
 	FTestCacheStore& operator=(const FTestCacheStore&) = delete;
 
@@ -146,6 +148,14 @@ FTestCacheStore::FTestCacheStore(const ECacheStoreFlags Flags, const bool bInAsy
 	CacheStorePolicy |= EnumHasAllFlags(Flags, ECacheStoreFlags::Store | ECacheStoreFlags::Remote) ? ECachePolicy::StoreRemote : ECachePolicy::None;
 }
 
+FTestCacheStore::~FTestCacheStore()
+{
+	while (!Queue.IsEmpty())
+	{
+		ExecuteAsync();
+	}
+}
+
 void FTestCacheStore::AddRecord(const FCacheKey& Key, TConstArrayView<FValueWithId> Values, const FCbObject* Meta)
 {
 	FRecord& Record = RecordMap.FindOrAdd(Key);
