Fix implicit conversion
#rnx

diff --git a/Engine/Source/Developer/DerivedDataCache/Private/S3CacheStore.cpp b/Engine/Source/Developer/DerivedDataCache/Private/S3CacheStore.cpp
--- a/Engine/Source/Developer/DerivedDataCache/Private/S3CacheStore.cpp
+++ b/Engine/Source/Developer/DerivedDataCache/Private/S3CacheStore.cpp
@@ -1558,7 +1558,7 @@ bool FS3CacheStore::LoadFile(
 	}
 
 	const double ReadDuration = FPlatformTime::Seconds() - StartTime;
-	const double ReadSpeed = ReadDuration > 0.001 ? (ReadSize / ReadDuration) / (1024.0 * 1024.0) : 0.0;
+	const double ReadSpeed = ReadDuration > 0.001 ? ((double)ReadSize / ReadDuration) / (1024.0 * 1024.0) : 0.0;
 
 	UE_LOG(LogDerivedDataCache, VeryVerbose,
 		TEXT("%s: Loaded %s from '%.*s' (%" INT64_FMT " bytes, %.02f secs, %.2f MiB/s)"),
