DDC: Removed obsolete unstructured namespace parameter from the clients for Horde Storage and Zen

diff --git a/Engine/Source/Developer/DerivedDataCache/Private/HttpCacheStore.cpp b/Engine/Source/Developer/DerivedDataCache/Private/HttpCacheStore.cpp
--- a/Engine/Source/Developer/DerivedDataCache/Private/HttpCacheStore.cpp
+++ b/Engine/Source/Developer/DerivedDataCache/Private/HttpCacheStore.cpp
@@ -222,7 +222,6 @@ struct FHttpCacheStoreParams
 	FString Host;
 	FString HostPinnedPublicKeys;
 	FString Namespace;
-	FString StructuredNamespace;
 	FString OAuthProvider;
 	FString OAuthClientId;
 	FString OAuthSecret;
@@ -291,7 +290,6 @@ public:
 
 	const FString& GetDomain() const { return Domain; }
 	const FString& GetNamespace() const { return Namespace; }
-	const FString& GetStructuredNamespace() const { return StructuredNamespace; }
 	const FString& GetOAuthProvider() const { return OAuthProvider; }
 	const FString& GetOAuthClientId() const { return OAuthClientId; }
 	const FString& GetOAuthSecret() const { return OAuthSecret; }
@@ -302,7 +300,6 @@ public:
 private:
 	FString Domain;
 	FString Namespace;
-	FString StructuredNamespace;
 	FString OAuthProvider;
 	FString OAuthClientId;
 	FString OAuthSecret;
@@ -941,7 +938,7 @@ void FHttpCacheStore::FPutPackageOp::PutRefAsync(
 	Algo::Transform(Key.Bucket.ToString(), AppendChars(Bucket), FCharAnsi::ToLower);
 
 	TAnsiStringBuilder<256> RefsUri;
-	RefsUri << CacheStore.EffectiveDomain << ANSITEXTVIEW("/api/v1/refs/") << CacheStore.StructuredNamespace << '/' << Bucket << '/' << Key.Hash;
+	RefsUri << CacheStore.EffectiveDomain << ANSITEXTVIEW("/api/v1/refs/") << CacheStore.Namespace << '/' << Bucket << '/' << Key.Hash;
 	if (bFinalize)
 	{
 		RefsUri << ANSITEXTVIEW("/finalize/") << ObjectHash;
@@ -1116,7 +1113,7 @@ void FHttpCacheStore::FPutPackageOp::OnPackagePutRefComplete(
 	{
 		TUniquePtr<FHttpOperation> Operation = CacheStore.WaitForHttpOperation(EOperationCategory::Put, /*bUnboundedOverflow*/ true);
 		FHttpOperation& LocalOperation = *Operation;
-		LocalOperation.SetUri(WriteToAnsiString<256>(CacheStore.EffectiveDomain, ANSITEXTVIEW("/api/v1/compressed-blobs/"), CacheStore.StructuredNamespace, '/', CompressedBlobUpload.Hash));
+		LocalOperation.SetUri(WriteToAnsiString<256>(CacheStore.EffectiveDomain, ANSITEXTVIEW("/api/v1/compressed-blobs/"), CacheStore.Namespace, '/', CompressedBlobUpload.Hash));
 		LocalOperation.SetMethod(EHttpMethod::Put);
 		LocalOperation.SetContentType(EHttpMediaType::CompressedBinary);
 		LocalOperation.SetBody(FCompositeBuffer(CompressedBlobUpload.BlobBuffer));
@@ -1213,7 +1210,7 @@ void FHttpCacheStore::FGetRecordOp::GetDataBatch(
 		const ValueType Value = Values[ValueIndex].RemoveData();
 		TUniquePtr<FHttpOperation> Operation = CacheStore.WaitForHttpOperation(EOperationCategory::Get, /*bUnboundedOverflow*/ true);
 		FHttpOperation& LocalOperation = *Operation;
-		LocalOperation.SetUri(WriteToAnsiString<256>(CacheStore.EffectiveDomain, ANSITEXTVIEW("/api/v1/compressed-blobs/"), CacheStore.StructuredNamespace, '/', Value.GetRawHash()));
+		LocalOperation.SetUri(WriteToAnsiString<256>(CacheStore.EffectiveDomain, ANSITEXTVIEW("/api/v1/compressed-blobs/"), CacheStore.Namespace, '/', Value.GetRawHash()));
 		LocalOperation.SetMethod(EHttpMethod::Get);
 		LocalOperation.AddAcceptType(EHttpMediaType::Any);
 		LocalOperation.SetExpectedErrorCodes({404});
@@ -1394,7 +1391,7 @@ void FHttpCacheStore::FGetRecordOp::DataProbablyExistsBatch(
 	}
 
 	TAnsiStringBuilder<256> CompressedBlobsUri;
-	CompressedBlobsUri << CacheStore.EffectiveDomain << ANSITEXTVIEW("/api/v1/compressed-blobs/") << CacheStore.StructuredNamespace << ANSITEXTVIEW("/exists?");
+	CompressedBlobsUri << CacheStore.EffectiveDomain << ANSITEXTVIEW("/api/v1/compressed-blobs/") << CacheStore.Namespace << ANSITEXTVIEW("/exists?");
 	bool bFirstItem = true;
 	for (const FValueWithId& Value : Values)
 	{
@@ -1534,7 +1531,6 @@ void FHttpCacheStore::FGetRecordOp::FinishDataStep(bool bSuccess, uint64 InBytes
 FHttpCacheStore::FHttpCacheStore(const FHttpCacheStoreParams& Params)
 	: Domain(Params.Host)
 	, Namespace(Params.Namespace)
-	, StructuredNamespace(Params.StructuredNamespace)
 	, OAuthProvider(Params.OAuthProvider)
 	, OAuthClientId(Params.OAuthClientId)
 	, OAuthSecret(Params.OAuthSecret)
@@ -1878,7 +1874,7 @@ void FHttpCacheStore::GetCacheRecordOnlyAsync(
 
 	TUniquePtr<FHttpOperation> Operation = WaitForHttpOperation(EOperationCategory::Get, /*bUnboundedOverflow*/ false);
 	FHttpOperation& LocalOperation = *Operation;
-	LocalOperation.SetUri(WriteToAnsiString<256>(EffectiveDomain, ANSITEXTVIEW("/api/v1/refs/"), StructuredNamespace, '/', Bucket, '/', Key.Hash));
+	LocalOperation.SetUri(WriteToAnsiString<256>(EffectiveDomain, ANSITEXTVIEW("/api/v1/refs/"), Namespace, '/', Bucket, '/', Key.Hash));
 	LocalOperation.SetMethod(EHttpMethod::Get);
 	LocalOperation.AddAcceptType(EHttpMediaType::CbObject);
 	LocalOperation.SetExpectedErrorCodes({404});
@@ -2062,7 +2058,7 @@ void FHttpCacheStore::GetCacheValueAsync(
 
 	TUniquePtr<FHttpOperation> Operation = WaitForHttpOperation(EOperationCategory::Get, /*bUnboundedOverflow*/ false);
 	FHttpOperation& LocalOperation = *Operation;
-	LocalOperation.SetUri(WriteToAnsiString<256>(EffectiveDomain, ANSITEXTVIEW("/api/v1/refs/"), StructuredNamespace, '/', Bucket, '/', Key.Hash));
+	LocalOperation.SetUri(WriteToAnsiString<256>(EffectiveDomain, ANSITEXTVIEW("/api/v1/refs/"), Namespace, '/', Bucket, '/', Key.Hash));
 	LocalOperation.SetMethod(EHttpMethod::Get);
 	if (bSkipData)
 	{
@@ -2197,7 +2193,7 @@ void FHttpCacheStore::RefCachedDataProbablyExistsBatchAsync(
 
 	TUniquePtr<FHttpOperation> Operation = WaitForHttpOperation(EOperationCategory::Get, /*bUnboundedOverflow*/ false);
 	FHttpOperation& LocalOperation = *Operation;
-	LocalOperation.SetUri(WriteToAnsiString<256>(EffectiveDomain, ANSITEXTVIEW("/api/v1/refs/"), StructuredNamespace));
+	LocalOperation.SetUri(WriteToAnsiString<256>(EffectiveDomain, ANSITEXTVIEW("/api/v1/refs/"), Namespace));
 	LocalOperation.SetMethod(EHttpMethod::Post);
 	LocalOperation.SetContentType(EHttpMediaType::CbObject);
 	LocalOperation.AddAcceptType(EHttpMediaType::CbObject);
@@ -2678,8 +2674,11 @@ void FHttpCacheStoreParams::Parse(const TCHAR* NodeName, const TCHAR* Config)
 
 	// Namespace Params
 
-	FParse::Value(Config, TEXT("Namespace="), Namespace);
-	FParse::Value(Config, TEXT("StructuredNamespace="), StructuredNamespace);
+	if (Namespace.IsEmpty())
+	{
+		FParse::Value(Config, TEXT("Namespace="), Namespace);
+	}
+	FParse::Value(Config, TEXT("StructuredNamespace="), Namespace);
 
 	// OAuth Params
 
@@ -2754,13 +2753,7 @@ TTuple<ILegacyCacheStore*, ECacheStoreFlags> CreateHttpCacheStore(const TCHAR* N
 	if (Params.Namespace.IsEmpty())
 	{
 		Params.Namespace = FApp::GetProjectName();
-		UE_LOG(LogDerivedDataCache, Warning, TEXT("%s: Missing required parameter 'Namespace', falling back to '%s'"), NodeName, *Params.Namespace);
-	}
-
-	if (Params.StructuredNamespace.IsEmpty())
-	{
-		Params.StructuredNamespace = Params.Namespace;
-		UE_LOG(LogDerivedDataCache, Warning, TEXT("%s: Missing required parameter 'StructuredNamespace', falling back to '%s'"), NodeName, *Params.StructuredNamespace);
+		UE_LOG(LogDerivedDataCache, Warning, TEXT("%s: Missing required parameter 'StructuredNamespace', falling back to '%s'"), NodeName, *Params.Namespace);
 	}
 
 	if (Params.OAuthProvider.IsEmpty())
@@ -2812,8 +2805,7 @@ ILegacyCacheStore* GetAnyHttpCacheStore(
 	FString& OutOAuthScope,
 	FString& OAuthProviderIdentifier,
 	FString& OAuthAccessToken,
-	FString& OutNamespace,
-	FString& OutStructuredNamespace)
+	FString& OutNamespace)
 {
 #if WITH_HTTP_DDC_BACKEND
 	if (FHttpCacheStore* HttpBackend = FHttpCacheStore::GetAny())
@@ -2826,8 +2818,6 @@ ILegacyCacheStore* GetAnyHttpCacheStore(
 		OAuthProviderIdentifier = HttpBackend->GetOAuthProviderIdentifier();
 		OAuthAccessToken = HttpBackend->GetOAuthAccessToken();
 		OutNamespace = HttpBackend->GetNamespace();
-		OutStructuredNamespace = HttpBackend->GetStructuredNamespace();
-
 		return HttpBackend;
 	}
 	return nullptr;
diff --git a/Engine/Source/Developer/DerivedDataCache/Private/Tests/HttpCacheStoreTest.cpp b/Engine/Source/Developer/DerivedDataCache/Private/Tests/HttpCacheStoreTest.cpp
--- a/Engine/Source/Developer/DerivedDataCache/Private/Tests/HttpCacheStoreTest.cpp
+++ b/Engine/Source/Developer/DerivedDataCache/Private/Tests/HttpCacheStoreTest.cpp
@@ -46,8 +46,7 @@ ILegacyCacheStore* GetAnyHttpCacheStore(
 	FString& OutOAuthScope,
 	FString& OAuthProviderIdentifier,
 	FString& OAuthAccessToken,
-	FString& OutNamespace,
-	FString& OutStructuredNamespace);
+	FString& OutNamespace);
 
 TTuple<ILegacyCacheStore*, ECacheStoreFlags> CreateZenCacheStore(const TCHAR* NodeName, const TCHAR* Config);
 
@@ -162,7 +161,7 @@ protected:
 	static ILegacyCacheStore* GetTestBackend()
 	{
 		static ILegacyCacheStore* CachedBackend = GetAnyHttpCacheStore(
-			TestDomain, TestOAuthProvider, TestOAuthClientId, TestOAuthSecret, TestOAuthScope, TestOAuthProviderIdentifier, TestOAuthAccessToken, TestNamespace, TestStructuredNamespace);
+			TestDomain, TestOAuthProvider, TestOAuthClientId, TestOAuthSecret, TestOAuthScope, TestOAuthProviderIdentifier, TestOAuthAccessToken, TestNamespace);
 		return CachedBackend;
 	}
 
@@ -495,7 +494,6 @@ protected:
 	static inline FString TestOAuthProviderIdentifier;
 	static inline FString TestOAuthAccessToken;
 	static inline FString TestNamespace;
-	static inline FString TestStructuredNamespace;
 };
 
 TArray<FCacheRecord> CreateTestCacheRecords(ICacheStore* InTestBackend, uint32 InNumKeys, uint32 InNumValues, FCbObject MetaContents = FCbObject(), const char* BucketName = nullptr)
@@ -625,13 +623,12 @@ bool FHttpCacheStoreTest::RunTest(const FString& Parameters)
 	FServiceSettings ZenUpstreamTestServiceSettings;
 	FServiceAutoLaunchSettings& ZenUpstreamTestAutoLaunchSettings = ZenUpstreamTestServiceSettings.SettingsVariant.Get<FServiceAutoLaunchSettings>();
 	ZenUpstreamTestAutoLaunchSettings.DataPath = FPaths::ConvertRelativePathToFull(FPaths::Combine(FPaths::EngineSavedDir(), "ZenUpstreamUnitTest"));
-	ZenUpstreamTestAutoLaunchSettings.ExtraArgs = FString::Printf(TEXT("--http asio --upstream-jupiter-url \"%s\" --upstream-jupiter-oauth-url \"%s\" --upstream-jupiter-oauth-clientid \"%s\" --upstream-jupiter-oauth-clientsecret \"%s\" --upstream-jupiter-namespace-ddc \"%s\" --upstream-jupiter-namespace \"%s\""),
+	ZenUpstreamTestAutoLaunchSettings.ExtraArgs = FString::Printf(TEXT("--http asio --upstream-jupiter-url \"%s\" --upstream-jupiter-oauth-url \"%s\" --upstream-jupiter-oauth-clientid \"%s\" --upstream-jupiter-oauth-clientsecret \"%s\" --upstream-jupiter-namespace \"%s\""),
 		*TestDomain,
 		*TestOAuthProvider,
 		*TestOAuthClientId,
 		*TestOAuthSecret,
-		*TestNamespace,
-		*TestStructuredNamespace
+		*TestNamespace
 	);
 	ZenUpstreamTestAutoLaunchSettings.DesiredPort = 23337; // Avoid the normal default port
 	ZenUpstreamTestAutoLaunchSettings.bShowConsole = true;
@@ -642,13 +639,12 @@ bool FHttpCacheStoreTest::RunTest(const FString& Parameters)
 	FServiceSettings ZenUpstreamSiblingTestServiceSettings;
 	FServiceAutoLaunchSettings& ZenUpstreamSiblingTestAutoLaunchSettings = ZenUpstreamSiblingTestServiceSettings.SettingsVariant.Get<FServiceAutoLaunchSettings>();
 	ZenUpstreamSiblingTestAutoLaunchSettings.DataPath = FPaths::ConvertRelativePathToFull(FPaths::Combine(FPaths::EngineSavedDir(), "ZenUpstreamSiblingUnitTest"));
-	ZenUpstreamSiblingTestAutoLaunchSettings.ExtraArgs = FString::Printf(TEXT("--http asio --upstream-jupiter-url \"%s\" --upstream-jupiter-oauth-url \"%s\" --upstream-jupiter-oauth-clientid \"%s\" --upstream-jupiter-oauth-clientsecret \"%s\" --upstream-jupiter-namespace-ddc \"%s\" --upstream-jupiter-namespace \"%s\""),
+	ZenUpstreamSiblingTestAutoLaunchSettings.ExtraArgs = FString::Printf(TEXT("--http asio --upstream-jupiter-url \"%s\" --upstream-jupiter-oauth-url \"%s\" --upstream-jupiter-oauth-clientid \"%s\" --upstream-jupiter-oauth-clientsecret \"%s\" --upstream-jupiter-namespace \"%s\""),
 		*TestDomain,
 		*TestOAuthProvider,
 		*TestOAuthClientId,
 		*TestOAuthSecret,
-		*TestNamespace,
-		*TestStructuredNamespace
+		*TestNamespace
 	);
 	ZenUpstreamSiblingTestAutoLaunchSettings.DesiredPort = 23338; // Avoid the normal default port
 	ZenUpstreamSiblingTestAutoLaunchSettings.bShowConsole = true;
@@ -678,11 +674,11 @@ bool FHttpCacheStoreTest::RunTest(const FString& Parameters)
 
 	FScopeZenService ScopeZenSiblingService(MoveTemp(ZenTestServiceSiblingSettings));
 	TUniquePtr<ILegacyCacheStore> ZenIntermediarySiblingBackend(CreateZenCacheStore(TEXT("TestSibling"),
-		*FString::Printf(TEXT("Host=%s, Namespace=%s, StructuredNamespace=%s"), ScopeZenSiblingService.GetInstance().GetURL(), *TestNamespace, *TestStructuredNamespace)).Key);
+		*FString::Printf(TEXT("Host=%s, StructuredNamespace=%s"), ScopeZenSiblingService.GetInstance().GetURL(), *TestNamespace)).Key);
 
 	FScopeZenService ScopeZenService(MoveTemp(ZenTestServiceSettings));
 	TUniquePtr<ILegacyCacheStore> ZenIntermediaryBackend(CreateZenCacheStore(TEXT("Test"),
-		*FString::Printf(TEXT("Host=%s, Namespace=%s, StructuredNamespace=%s"), ScopeZenService.GetInstance().GetURL(), *TestNamespace, *TestStructuredNamespace)).Key);
+		*FString::Printf(TEXT("Host=%s, StructuredNamespace=%s"), ScopeZenService.GetInstance().GetURL(), *TestNamespace)).Key);
 	auto WaitForZenPushToUpstream = [](ILegacyCacheStore* ZenBackend, TConstArrayView<FCacheRecord> Records)
 	{
 		// TODO: Expecting a legitimate means to wait for zen to finish pushing records to its upstream in the future
diff --git a/Engine/Source/Developer/DerivedDataCache/Private/ZenCacheStore.cpp b/Engine/Source/Developer/DerivedDataCache/Private/ZenCacheStore.cpp
--- a/Engine/Source/Developer/DerivedDataCache/Private/ZenCacheStore.cpp
+++ b/Engine/Source/Developer/DerivedDataCache/Private/ZenCacheStore.cpp
@@ -1194,14 +1194,13 @@ TTuple<ILegacyCacheStore*, ECacheStoreFlags> CreateZenCacheStore(const TCHAR* No
 	if (!FParse::Value(Config, TEXT("Namespace="), Namespace))
 	{
 		Namespace = FApp::GetProjectName();
-		UE_LOG(LogDerivedDataCache, Warning, TEXT("%s: Missing required parameter 'Namespace', falling back to '%s'"), NodeName, *Namespace);
 	}
 
 	FString StructuredNamespace;
 	if (!FParse::Value(Config, TEXT("StructuredNamespace="), StructuredNamespace))
 	{
 		StructuredNamespace = Namespace;
-		UE_LOG(LogDerivedDataCache, Warning, TEXT("%s: Missing required parameter 'StructuredNamespace', falling back to Namespace: '%s'"), NodeName, *StructuredNamespace);
+		UE_LOG(LogDerivedDataCache, Warning, TEXT("%s: Missing required parameter 'StructuredNamespace', falling back '%s'"), NodeName, *StructuredNamespace);
 	}
 
 	TUniquePtr<FZenCacheStore> Backend = MakeUnique<FZenCacheStore>(*ServiceUrl, *StructuredNamespace);
