DDC: Added to the summary when file system cache store maintenance completes

#rnx

diff --git a/Engine/Source/Developer/DerivedDataCache/Private/FileSystemCacheStore.cpp b/Engine/Source/Developer/DerivedDataCache/Private/FileSystemCacheStore.cpp
--- a/Engine/Source/Developer/DerivedDataCache/Private/FileSystemCacheStore.cpp
+++ b/Engine/Source/Developer/DerivedDataCache/Private/FileSystemCacheStore.cpp
@@ -252,9 +252,12 @@ private:
 	/** Ignore the scan rate for one maintenance scan. */
 	bool bIgnoreScanRate = false;
 
+	uint32 FileCount = 0;
+	uint32 FolderCount = 0;
 	uint32 ProcessCount = 0;
 	uint32 DeleteCount = 0;
 	uint64 DeleteSize = 0;
+	uint64 ScannedSize = 0;
 
 	double BatchStartTime = 0.0;
 
@@ -355,8 +358,11 @@ void FFileSystemCacheStoreMaintainer::Loop()
 	while (!bExit)
 	{
 		const FDateTime ScanStart = FDateTime::Now();
+		FileCount = 0;
+		FolderCount = 0;
 		DeleteCount = 0;
 		DeleteSize = 0;
+		ScannedSize = 0;
 		IdleEvent.Reset();
 		bIdle = false;
 		Scan();
@@ -366,8 +372,10 @@ void FFileSystemCacheStoreMaintainer::Loop()
 		const FDateTime ScanEnd = FDateTime::Now();
 
 		UE_LOG(LogDerivedDataCache, Log,
-			TEXT("%s: Maintenance finished in %s and deleted %u file(s) with total size %" UINT64_FMT " MiB."),
-			*CachePath, *(ScanEnd - ScanStart).ToString(), DeleteCount, DeleteSize / 1024 / 1024);
+			TEXT("%s: Maintenance finished in %s and deleted %u files with total size %" UINT64_FMT " MiB. "
+				 "Scanned %u files in %u folders with total size %" UINT64_FMT " MiB."),
+			*CachePath, *(ScanEnd - ScanStart).ToString(), DeleteCount, DeleteSize / 1024 / 1024,
+			FileCount, FolderCount, ScannedSize / 1024 / 1024);
 
 		if (bExit || bExitAfterScan)
 		{
@@ -603,6 +611,8 @@ void FFileSystemCacheStoreMaintainer::ResetRoots()
 
 void FFileSystemCacheStoreMaintainer::ProcessDirectory(const TCHAR* const Path)
 {
+	++FolderCount;
+
 	bool bTryDelete = true;
 
 	FileManager.IterateDirectoryStat(Path, [this, &bTryDelete](const TCHAR* const Path, const FFileStatData& Stat) -> bool
@@ -630,6 +640,9 @@ void FFileSystemCacheStoreMaintainer::ProcessFile(const TCHAR* const Path, const
 		return;
 	}
 
+	++FileCount;
+	ScannedSize += Stat.FileSize > 0 ? uint64(Stat.FileSize) : 0;
+
 	const FDateTime Now = FDateTime::UtcNow();
 	if (Stat.ModificationTime + Params.MaxFileAge < Now && Stat.AccessTime + Params.MaxFileAge < Now)
 	{
