DDC: Added functions to convert between derived data priority and queued work priority

#rnx

diff --git a/Engine/Source/Developer/DerivedDataCache/Private/DerivedDataCache.cpp b/Engine/Source/Developer/DerivedDataCache/Private/DerivedDataCache.cpp
--- a/Engine/Source/Developer/DerivedDataCache/Private/DerivedDataCache.cpp
+++ b/Engine/Source/Developer/DerivedDataCache/Private/DerivedDataCache.cpp
@@ -327,21 +327,16 @@ FQueuedThreadPool* GCacheThreadPool;
 class FCacheThreadPoolTaskRequest final : public FRequestBase, private IQueuedWork
 {
 public:
-	inline FCacheThreadPoolTaskRequest(
-		IRequestOwner& InOwner,
-		TUniqueFunction<void ()>&& InTaskBody)
+	inline FCacheThreadPoolTaskRequest(IRequestOwner& InOwner, TUniqueFunction<void ()>&& InTaskBody)
 		: Owner(InOwner)
 		, TaskBody(MoveTemp(InTaskBody))
-	{
-	}
-
-	inline void Start(EPriority Priority)
 	{
 		Owner.Begin(this);
 		DoneEvent.Reset();
-		GCacheThreadPool->AddQueuedWork(this, GetPriority(Priority));
+		GCacheThreadPool->AddQueuedWork(this, ConvertToQueuedWorkPriority(Owner.GetPriority()));
 	}
 
+private:
 	inline void Execute()
 	{
 		FScopeCycleCounter Scope(GetStatId(), /*bAlways*/ true);
@@ -359,7 +354,7 @@ public:
 	{
 		if (GCacheThreadPool->RetractQueuedWork(this))
 		{
-			GCacheThreadPool->AddQueuedWork(this, GetPriority(Priority));
+			GCacheThreadPool->AddQueuedWork(this, ConvertToQueuedWorkPriority(Priority));
 		}
 	}
 
@@ -395,21 +390,6 @@ public:
 		}
 	}
 
-private:
-	static EQueuedWorkPriority GetPriority(EPriority Priority)
-	{
-		switch (Priority)
-		{
-		case EPriority::Blocking: return EQueuedWorkPriority::Blocking;
-		case EPriority::Highest:  return EQueuedWorkPriority::Highest;
-		case EPriority::High:     return EQueuedWorkPriority::High;
-		case EPriority::Normal:   return EQueuedWorkPriority::Normal;
-		case EPriority::Low:      return EQueuedWorkPriority::Low;
-		case EPriority::Lowest:   return EQueuedWorkPriority::Lowest;
-		default: checkNoEntry();  return EQueuedWorkPriority::Normal;
-		}
-	}
-
 	// IQueuedWork Interface
 
 	inline void DoThreadedWork() final { Execute(); }
@@ -430,8 +410,7 @@ void LaunchTaskInCacheThreadPool(IRequestOwner& Owner, TUniqueFunction<void ()>&
 {
 	if (GCacheThreadPool)
 	{
-		FCacheThreadPoolTaskRequest* Request = new FCacheThreadPoolTaskRequest(Owner, MoveTemp(TaskBody));
-		Request->Start(Owner.GetPriority());
+		new FCacheThreadPoolTaskRequest(Owner, MoveTemp(TaskBody));
 	}
 	else
 	{
diff --git a/Engine/Source/Developer/DerivedDataCache/Private/DerivedDataRequestTypes.cpp b/Engine/Source/Developer/DerivedDataCache/Private/DerivedDataRequestTypes.cpp
--- a/Engine/Source/Developer/DerivedDataCache/Private/DerivedDataRequestTypes.cpp
+++ b/Engine/Source/Developer/DerivedDataCache/Private/DerivedDataRequestTypes.cpp
@@ -4,6 +4,7 @@
 
 #include "Containers/StringConv.h"
 #include "Containers/StringView.h"
+#include "Misc/QueuedThreadPool.h"
 #include "Misc/StringBuilder.h"
 #include "Serialization/CompactBinary.h"
 #include "Serialization/CompactBinaryWriter.h"
@@ -124,6 +125,34 @@ bool LoadFromCompactBinary(FCbFieldView Field, EPriority& OutPriority, const EPr
 	return false;
 }
 
+EQueuedWorkPriority ConvertToQueuedWorkPriority(EPriority Priority)
+{
+	switch (Priority)
+	{
+	case EPriority::Blocking: return EQueuedWorkPriority::Blocking;
+	case EPriority::Highest:  return EQueuedWorkPriority::Highest;
+	case EPriority::High:     return EQueuedWorkPriority::High;
+	case EPriority::Normal:   return EQueuedWorkPriority::Normal;
+	case EPriority::Low:      return EQueuedWorkPriority::Low;
+	case EPriority::Lowest:   return EQueuedWorkPriority::Lowest;
+	default: checkNoEntry();  return EQueuedWorkPriority::Normal;
+	}
+}
+
+EPriority ConvertFromQueuedWorkPriority(EQueuedWorkPriority Priority)
+{
+	switch (Priority)
+	{
+	case EQueuedWorkPriority::Blocking: return EPriority::Blocking;
+	case EQueuedWorkPriority::Highest:  return EPriority::Highest;
+	case EQueuedWorkPriority::High:     return EPriority::High;
+	case EQueuedWorkPriority::Normal:   return EPriority::Normal;
+	case EQueuedWorkPriority::Low:      return EPriority::Low;
+	case EQueuedWorkPriority::Lowest:   return EPriority::Lowest;
+	default: checkNoEntry();            return EPriority::Normal;
+	}
+}
+
 FAnsiStringBuilderBase& operator<<(FAnsiStringBuilderBase& Builder, EStatus Status) { return Private::StatusToString(Builder, Status); }
 FWideStringBuilderBase& operator<<(FWideStringBuilderBase& Builder, EStatus Status) { return Private::StatusToString(Builder, Status); }
 FUtf8StringBuilderBase& operator<<(FUtf8StringBuilderBase& Builder, EStatus Status) { return Private::StatusToString(Builder, Status); }
diff --git a/Engine/Source/Developer/DerivedDataCache/Public/DerivedDataRequestTypes.h b/Engine/Source/Developer/DerivedDataCache/Public/DerivedDataRequestTypes.h
--- a/Engine/Source/Developer/DerivedDataCache/Public/DerivedDataRequestTypes.h
+++ b/Engine/Source/Developer/DerivedDataCache/Public/DerivedDataRequestTypes.h
@@ -9,6 +9,7 @@
 
 class FCbFieldView;
 class FCbWriter;
+enum class EQueuedWorkPriority : uint8;
 
 namespace UE::DerivedData
 {
@@ -61,6 +62,11 @@ UE_API bool TryLexFromString(EPriority& OutPriority, FWideStringView String);
 UE_API FCbWriter& operator<<(FCbWriter& Writer, EPriority Priority);
 UE_API bool LoadFromCompactBinary(FCbFieldView Field, EPriority& OutPriority, EPriority Default = EPriority::Normal);
 
+/** Converts to a queued work priority. Asserts on invalid priority values. */
+UE_API EQueuedWorkPriority ConvertToQueuedWorkPriority(EPriority Priority);
+/** Converts from a queued work priority. Asserts on invalid priority values. */
+UE_API EPriority ConvertFromQueuedWorkPriority(EQueuedWorkPriority Priority);
+
 /** Status of a request that has completed. */
 enum class EStatus : uint8
 {
