DDC: Allow both OAuthProvider and OAuthProviderIdentifier to be used in the same configuration for HttpCacheStore

- At least one of OAuthProvider, OAuthProviderIdentifier, or OAuthAccessTokenEnvOverride must be set.
- If only OAuthAccessTokenEnvOverride is set, then the environment variable must be set too.
- If only OAuthProvider is set, then OAuthSecret must be set too.
- If OAuthProvider is set, then OAuthClientId must be set.

diff --git a/Engine/Source/Developer/DerivedDataCache/Private/HttpCacheStore.cpp b/Engine/Source/Developer/DerivedDataCache/Private/HttpCacheStore.cpp
--- a/Engine/Source/Developer/DerivedDataCache/Private/HttpCacheStore.cpp
+++ b/Engine/Source/Developer/DerivedDataCache/Private/HttpCacheStore.cpp
@@ -75,13 +75,6 @@ static FAutoConsoleVariableRef CVarHttpEnableAsync(
 	TEXT("If true, asynchronous operations are permitted, otherwise all operations are forced to be synchronous."),
 	ECVF_Default);
 
-static bool bHttpEnableOidc = false;
-static FAutoConsoleVariableRef CVarHttpEnableOidc(
-	TEXT("DDC.Http.EnableOidc"),
-	bHttpEnableOidc,
-	TEXT("If true, Oidc tokens are used, otherwise shared credentials are used."),
-	ECVF_Default);
-
 TRACE_DECLARE_INT_COUNTER(HttpDDC_Get, TEXT("HttpDDC Get"));
 TRACE_DECLARE_INT_COUNTER(HttpDDC_GetHit, TEXT("HttpDDC Get Hit"));
 TRACE_DECLARE_INT_COUNTER(HttpDDC_Put, TEXT("HttpDDC Put"));
@@ -1641,9 +1634,6 @@ bool FHttpCacheStore::AcquireAccessToken(IHttpClient* Client)
 		return false;
 	}
 
-	ensureMsgf(OAuthProvider.StartsWith(TEXT("http://")) || OAuthProvider.StartsWith(TEXT("https://")),
-		TEXT("%s: OAuth provider %s must be a complete URI including the scheme."), *Domain, *OAuthProvider);
-
 	TRACE_CPUPROFILER_EVENT_SCOPE(HttpDDC_AcquireAccessToken);
 
 	// In case many requests wants to update the token at the same time
@@ -1664,95 +1654,82 @@ bool FHttpCacheStore::AcquireAccessToken(IHttpClient* Client)
 		return true;
 	}
 
-	if (bHttpEnableOidc && !OAuthProviderIdentifier.IsEmpty())
+	if (!OAuthSecret.IsEmpty())
 	{
-		FString AccessTokenString;
-		FDateTime TokenExpiresAt;
-		if (FDesktopPlatformModule::Get()->GetOidcAccessToken(FPaths::RootDir(), FPaths::GetProjectFilePath(), OAuthProviderIdentifier, FApp::IsUnattended(), GWarn, AccessTokenString, TokenExpiresAt))
+		THttpUniquePtr<IHttpClient> LocalClient;
+		if (!Client)
 		{
-			const double ExpiryTimeSeconds = (TokenExpiresAt - FDateTime::UtcNow()).GetTotalSeconds();
-			UE_LOG(LogDerivedDataCache, Display,
-				TEXT("%s: OidcToken: Logged in to HTTP DDC services. Expires at %s which is in %.0f seconds."),
-				*Domain, *TokenExpiresAt.ToString(), ExpiryTimeSeconds);
-			SetAccessToken(AccessTokenString, ExpiryTimeSeconds);
-			return true;
+			LocalClient = ConnectionPool->CreateClient(GetDefaultClientParams());
+			Client = LocalClient.Get();
 		}
 
-		UE_LOG(LogDerivedDataCache, Error, TEXT("%s: OidcToken: Failed to log in to HTTP services."), *Domain);
-		FailedLoginAttempts++;
-		return false;
-	}
+		FHttpRequestParams RequestParams;
+		RequestParams.bIgnoreMaxRequests = true;
+		FHttpOperation Operation(Client->TryCreateRequest(RequestParams));
+		Operation.SetUri(StringCast<ANSICHAR>(*OAuthProvider));
 
-	THttpUniquePtr<IHttpClient> LocalClient;
-	if (!Client)
-	{
-		LocalClient = ConnectionPool->CreateClient(GetDefaultClientParams());
-		Client = LocalClient.Get();
-	}
-
-	FHttpRequestParams RequestParams;
-	RequestParams.bIgnoreMaxRequests = true;
-	FHttpOperation Operation(Client->TryCreateRequest(RequestParams));
-	Operation.SetUri(StringCast<ANSICHAR>(*OAuthProvider));
-
-	if (OAuthProvider.StartsWith(TEXT("http://localhost")))
-	{
-		// Simple unauthenticated call to a local endpoint that mimics the result from an OIDC provider.
-		Operation.Send();
-	}
-	else
-	{
-		// Needs client id and secret to authenticate with the OIDC provider.
+		if (OAuthProvider.StartsWith(TEXT("http://localhost")))
+		{
+			// Simple unauthenticated call to a local endpoint that mimics the result from an OIDC provider.
+			Operation.Send();
+		}
+		else
+		{
+			TUtf8StringBuilder<256> OAuthFormData;
+			OAuthFormData
+				<< ANSITEXTVIEW("client_id=") << OAuthClientId
+				<< ANSITEXTVIEW("&scope=") << OAuthScope
+				<< ANSITEXTVIEW("&grant_type=client_credentials")
+				<< ANSITEXTVIEW("&client_secret=") << OAuthSecret;
+
+			Operation.SetMethod(EHttpMethod::Post);
+			Operation.SetContentType(EHttpMediaType::FormUrlEncoded);
+			Operation.SetBody(FCompositeBuffer(FSharedBuffer::MakeView(MakeMemoryView(OAuthFormData))));
+			Operation.Send();
+		}
 
-		// If contents of the secret string is a file path, resolve and read form data.
-		if (OAuthSecret.StartsWith(TEXT("file://")))
+		if (Operation.GetStatusCode() == 200)
 		{
-			FString FilePath = OAuthSecret.Mid(7, OAuthSecret.Len() - 7);
-			FString SecretFileContents;
-			if (FFileHelper::LoadFileToString(SecretFileContents, *FilePath))
-			{
-				// Overwrite the filepath with the actual content.
-				OAuthSecret = SecretFileContents;
-			}
-			else
+			if (TSharedPtr<FJsonObject> ResponseObject = Operation.GetBodyAsJson())
 			{
-				UE_LOG(LogDerivedDataCache, Warning, TEXT("%s: Failed to read OAuth form data file (%s)."), *Domain, *OAuthSecret);
-				return false;
+				FString AccessTokenString;
+				double ExpiryTimeSeconds = 0.0;
+				if (ResponseObject->TryGetStringField(TEXT("access_token"), AccessTokenString) &&
+					ResponseObject->TryGetNumberField(TEXT("expires_in"), ExpiryTimeSeconds))
+				{
+					UE_LOG(LogDerivedDataCache, Display,
+						TEXT("%s: Logged in to HTTP DDC services. Expires in %.0f seconds."), *Domain, ExpiryTimeSeconds);
+					SetAccessToken(AccessTokenString, ExpiryTimeSeconds);
+					return true;
+				}
 			}
 		}
 
-		TUtf8StringBuilder<256> OAuthFormData;
-		OAuthFormData
-			<< ANSITEXTVIEW("client_id=") << OAuthClientId
-			<< ANSITEXTVIEW("&scope=") << OAuthScope
-			<< ANSITEXTVIEW("&grant_type=client_credentials")
-			<< ANSITEXTVIEW("&client_secret=") << OAuthSecret;
-
-		Operation.SetMethod(EHttpMethod::Post);
-		Operation.SetContentType(EHttpMediaType::FormUrlEncoded);
-		Operation.SetBody(FCompositeBuffer(FSharedBuffer::MakeView(MakeMemoryView(OAuthFormData))));
-		Operation.Send();
+		UE_LOG(LogDerivedDataCache, Warning, TEXT("%s: Failed to log in to HTTP services with request %s."), *Domain, *WriteToString<256>(Operation));
+		FailedLoginAttempts++;
+		return false;
 	}
 
-	if (Operation.GetStatusCode() == 200)
+	if (!OAuthProviderIdentifier.IsEmpty())
 	{
-		if (TSharedPtr<FJsonObject> ResponseObject = Operation.GetBodyAsJson())
+		FString AccessTokenString;
+		FDateTime TokenExpiresAt;
+		if (FDesktopPlatformModule::Get()->GetOidcAccessToken(FPaths::RootDir(), FPaths::GetProjectFilePath(), OAuthProviderIdentifier, FApp::IsUnattended(), GWarn, AccessTokenString, TokenExpiresAt))
 		{
-			FString AccessTokenString;
-			double ExpiryTimeSeconds = 0.0;
-			if (ResponseObject->TryGetStringField(TEXT("access_token"), AccessTokenString) &&
-				ResponseObject->TryGetNumberField(TEXT("expires_in"), ExpiryTimeSeconds))
-			{
-				UE_LOG(LogDerivedDataCache, Display,
-					TEXT("%s: Logged in to HTTP DDC services. Expires in %.0f seconds."), *Domain, ExpiryTimeSeconds);
-				SetAccessToken(AccessTokenString, ExpiryTimeSeconds);
-				return true;
-			}
+			const double ExpiryTimeSeconds = (TokenExpiresAt - FDateTime::UtcNow()).GetTotalSeconds();
+			UE_LOG(LogDerivedDataCache, Display,
+				TEXT("%s: OidcToken: Logged in to HTTP DDC services. Expires at %s which is in %.0f seconds."),
+				*Domain, *TokenExpiresAt.ToString(), ExpiryTimeSeconds);
+			SetAccessToken(AccessTokenString, ExpiryTimeSeconds);
+			return true;
 		}
+
+		UE_LOG(LogDerivedDataCache, Error, TEXT("%s: OidcToken: Failed to log in to HTTP services."), *Domain);
+		FailedLoginAttempts++;
+		return false;
 	}
 
-	UE_LOG(LogDerivedDataCache, Warning, TEXT("%s: Failed to log in to HTTP services with request %s."), *Domain, *WriteToString<256>(Operation));
-	FailedLoginAttempts++;
+	UE_LOG(LogDerivedDataCache, Warning, TEXT("%s: No available configuration to acquire an access token."), *Domain);
 	return false;
 }
 
@@ -2703,6 +2680,18 @@ void FHttpCacheStoreParams::Parse(const TCHAR* NodeName, const TCHAR* Config)
 		}
 	}
 
+	// If the secret is a file path, read the secret from the file.
+	if (OAuthSecret.StartsWith(TEXT("file://")))
+	{
+		TStringBuilder<256> FilePath;
+		FilePath << MakeStringView(OAuthSecret).RightChop(TEXTVIEW("file://").Len());
+		if (!FFileHelper::LoadFileToString(OAuthSecret, *FilePath))
+		{
+			OAuthSecret.Empty();
+			UE_LOG(LogDerivedDataCache, Warning, TEXT("%s: Failed to read OAuth secret file: %s"), NodeName, *FilePath);
+		}
+	}
+
 	FParse::Value(Config, TEXT("OAuthScope="), OAuthScope);
 
 	FParse::Value(Config, TEXT("OAuthProviderIdentifier="), OAuthProviderIdentifier);
@@ -2734,20 +2723,23 @@ namespace UE::DerivedData
 
 TTuple<ILegacyCacheStore*, ECacheStoreFlags> CreateHttpCacheStore(const TCHAR* NodeName, const TCHAR* Config)
 {
-#if WITH_HTTP_DDC_BACKEND
+#if !WITH_HTTP_DDC_BACKEND
+	UE_LOG(LogDerivedDataCache, Warning, TEXT("%s: HTTP cache is not yet supported in the current build configuration."), NodeName);
+#else
 	FHttpCacheStoreParams Params;
 	Params.Parse(NodeName, Config);
 
+	bool bValidParams = true;
+
 	if (Params.Host.IsEmpty())
 	{
 		UE_LOG(LogDerivedDataCache, Error, TEXT("%s: Missing required parameter 'Host'"), NodeName);
-		return MakeTuple(nullptr, ECacheStoreFlags::None);
+		bValidParams = false;
 	}
-
-	if (Params.Host == TEXT("None"))
+	else if (Params.Host == TEXTVIEW("None"))
 	{
 		UE_LOG(LogDerivedDataCache, Log, TEXT("%s: Disabled because Host is set to 'None'"), NodeName);
-		return MakeTuple(nullptr, ECacheStoreFlags::None);
+		bValidParams = false;
 	}
 
 	if (Params.Namespace.IsEmpty())
@@ -2756,25 +2748,46 @@ TTuple<ILegacyCacheStore*, ECacheStoreFlags> CreateHttpCacheStore(const TCHAR* N
 		UE_LOG(LogDerivedDataCache, Warning, TEXT("%s: Missing required parameter 'StructuredNamespace', falling back to '%s'"), NodeName, *Params.Namespace);
 	}
 
-	if (Params.OAuthProvider.IsEmpty())
+	if (!Params.Host.StartsWith(TEXT("http://localhost")))
 	{
-		UE_LOG(LogDerivedDataCache, Error, TEXT("%s: Missing required parameter 'OAuthProvider'"), NodeName);
-		return MakeTuple(nullptr, ECacheStoreFlags::None);
-	}
+		bool bValidOAuthAccessToken = !Params.OAuthAccessToken.IsEmpty();
 
-	// No need for OAuth client id and secret if using a local provider.
-	if (!Params.OAuthProvider.StartsWith(TEXT("http://localhost")))
-	{
-		if (Params.OAuthClientId.IsEmpty())
+		bool bValidOAuthProviderIdentifier = !Params.OAuthProviderIdentifier.IsEmpty();
+
+		bool bValidOAuthProvider = !Params.OAuthProvider.IsEmpty();
+		if (bValidOAuthProvider)
 		{
-			UE_LOG(LogDerivedDataCache, Error, TEXT("%s: Missing required parameter 'OAuthClientId'"), NodeName);
-			return MakeTuple(nullptr, ECacheStoreFlags::None);
+			if (!Params.OAuthProvider.StartsWith(TEXT("http://")) &&
+				!Params.OAuthProvider.StartsWith(TEXT("https://")))
+			{
+				UE_LOG(LogDerivedDataCache, Error, TEXT("%s: OAuth provider '%s' must be a complete URI including the scheme."), NodeName, *Params.OAuthProvider);
+				bValidParams = false;
+			}
+
+			// No need for OAuthClientId and OAuthSecret if using a local provider.
+			if (!Params.OAuthProvider.StartsWith(TEXT("http://localhost")))
+			{
+				if (Params.OAuthClientId.IsEmpty())
+				{
+					UE_LOG(LogDerivedDataCache, Error, TEXT("%s: Missing required parameter 'OAuthClientId'"), NodeName);
+					bValidOAuthProvider = false;
+					bValidParams = false;
+				}
+
+				if (Params.OAuthSecret.IsEmpty())
+				{
+					UE_CLOG(!bValidOAuthAccessToken && !bValidOAuthProviderIdentifier,
+						LogDerivedDataCache, Error, TEXT("%s: Missing required parameter 'OAuthSecret'"), NodeName);
+					bValidOAuthProvider = false;
+				}
+			}
 		}
 
-		if (Params.OAuthSecret.IsEmpty())
+		if (!bValidOAuthAccessToken && !bValidOAuthProviderIdentifier && !bValidOAuthProvider)
 		{
-			UE_LOG(LogDerivedDataCache, Error, TEXT("%s: Missing required parameter 'OAuthSecret'"), NodeName);
-			return MakeTuple(nullptr, ECacheStoreFlags::None);
+			UE_LOG(LogDerivedDataCache, Error, TEXT("%s: At least one OAuth configuration must be provided and valid. "
+				"Options are 'OAuthProvider', 'OAuthProviderIdentifier', and 'OAuthAccessTokenEnvOverride'"), NodeName);
+			bValidParams = false;
 		}
 	}
 
@@ -2783,18 +2796,18 @@ TTuple<ILegacyCacheStore*, ECacheStoreFlags> CreateHttpCacheStore(const TCHAR* N
 		Params.OAuthScope = TEXTVIEW("cache_access");
 	}
 
-	TUniquePtr<FHttpCacheStore> Backend = MakeUnique<FHttpCacheStore>(Params);
-	if (!Backend->IsUsable())
+	if (bValidParams)
 	{
+		if (TUniquePtr<FHttpCacheStore> Store = MakeUnique<FHttpCacheStore>(Params); Store->IsUsable())
+		{
+			const ECacheStoreFlags StoreFlag = (Params.bReadOnly ? ECacheStoreFlags::None : ECacheStoreFlags::Store);
+			return MakeTuple(Store.Release(), ECacheStoreFlags::Remote | ECacheStoreFlags::Query | StoreFlag);
+		}
 		UE_LOG(LogDerivedDataCache, Warning, TEXT("%s: Failed to contact the service (%s), will not use it."), NodeName, *Params.Host);
-		Backend.Reset();
 	}
+#endif
 
-	return MakeTuple(Backend.Release(), ECacheStoreFlags::Remote | ECacheStoreFlags::Query | (Params.bReadOnly ? ECacheStoreFlags::None : ECacheStoreFlags::Store));
-#else
-	UE_LOG(LogDerivedDataCache, Warning, TEXT("%s: HTTP cache is not yet supported in the current build configuration."), NodeName);
 	return MakeTuple(nullptr, ECacheStoreFlags::None);
-#endif
 }
 
 ILegacyCacheStore* GetAnyHttpCacheStore(
@@ -2820,10 +2833,8 @@ ILegacyCacheStore* GetAnyHttpCacheStore(
 		OutNamespace = HttpBackend->GetNamespace();
 		return HttpBackend;
 	}
-	return nullptr;
-#else
-	return nullptr;
 #endif
+	return nullptr;
 }
 
 } // UE::DerivedData
