DDC: Added a check for non-zero size in the legacy cache

#rnx

diff --git a/Engine/Source/Developer/DerivedDataCache/Private/DerivedDataCache.cpp b/Engine/Source/Developer/DerivedDataCache/Private/DerivedDataCache.cpp
--- a/Engine/Source/Developer/DerivedDataCache/Private/DerivedDataCache.cpp
+++ b/Engine/Source/Developer/DerivedDataCache/Private/DerivedDataCache.cpp
@@ -500,7 +500,7 @@ class FDerivedDataCache final
 						[this, &bGetResult](FLegacyCacheGetResponse&& Response)
 						{
 							const uint64 RawSize = Response.Value.GetRawSize();
-							bGetResult = Response.Status == EStatus::Ok && RawSize < MAX_int64;
+							bGetResult = Response.Status == EStatus::Ok && RawSize > 0 && RawSize < MAX_int64;
 							if (bGetResult)
 							{
 								const FCompositeBuffer& RawData = Response.Value.GetRawData();
