CacheStoreTest - automated test spurious failure fix: Record whether the first local put and remote put succeed, and use that to decide bExpectHit for each key that get what we put to local/remote.
#rnx

diff --git a/Engine/Source/Developer/DerivedDataCache/Private/Tests/CacheStoreTest.cpp b/Engine/Source/Developer/DerivedDataCache/Private/Tests/CacheStoreTest.cpp
--- a/Engine/Source/Developer/DerivedDataCache/Private/Tests/CacheStoreTest.cpp
+++ b/Engine/Source/Developer/DerivedDataCache/Private/Tests/CacheStoreTest.cpp
@@ -25,7 +25,6 @@ bool FCacheStoreTest::RunTest(const FString& Parameters)
 
 	FRequestOwner Owner(EPriority::Blocking);
 	FCacheBucket DDCTestBucket(TEXT("DDCTest"));
-	static bool bExpectUpstream = FParse::Param(FCommandLine::Get(), TEXT("CacheStoreTestUpstream"));
 	static bool bExpectWarm = FParse::Param(FCommandLine::Get(), TEXT("CacheStoreTestWarm"));
 
 	// NumKeys = (2 Value vs Record)*(2 SkipData vs Default)*(2 ForceMiss vs Not)*(2 use local)
@@ -70,7 +69,6 @@ bool FCacheStoreTest::RunTest(const FString& Parameters)
 		bool bForceMiss = false;
 		bool bUseLocal = true;
 		bool bUseRemote = true;
-		bool bShouldBeHit = true;
 		bool bReceivedPut = false;
 		bool bReceivedGet = false;
 		bool bReceivedPutValue = false;
@@ -96,7 +94,6 @@ bool FCacheStoreTest::RunTest(const FString& Parameters)
 		KeyData.bForceMiss = (KeyIndex & (1 << 4)) == 0;
 		KeyData.bUseLocal = (KeyIndex & (1 << 5)) == 0;
 		KeyData.bUseRemote = (KeyIndex & (1 << 6)) == 0;
-		KeyData.bShouldBeHit = !KeyData.bForceMiss && (KeyData.bUseLocal || (KeyData.bUseRemote && bExpectUpstream));
 		ECachePolicy SharedPolicy = KeyData.bUseLocal ? ECachePolicy::Local : ECachePolicy::None;
 		SharedPolicy |= KeyData.bUseRemote ? ECachePolicy::Remote : ECachePolicy::None;
 		ECachePolicy PutPolicy = SharedPolicy;
@@ -170,17 +167,34 @@ bool FCacheStoreTest::RunTest(const FString& Parameters)
 		}
 	}
 
+	bool bLocalPutSucceeded = true;
+	bool bRemotePutSucceeded = true;
 	if (!bExpectWarm)
 	{
+		bool bLocalPutSucceededInitialized = false;
+		bool bRemotePutSucceededInitialized = false;
 		{
 			FRequestBarrier Barrier(Owner);
-			Cache.Put(PutRequests, Owner, [this](FCachePutResponse&& Response)
+			Cache.Put(PutRequests, Owner,
+				[this, &bLocalPutSucceededInitialized, &bLocalPutSucceeded, &bRemotePutSucceededInitialized,
+				&bRemotePutSucceeded]
+			(FCachePutResponse&& Response)
 				{
 					FUserData* UserData = reinterpret_cast<FUserData*>(Response.UserData);
 					FKeyData* KeyData = UserData->KeyData;
 					TestTrue(TEXT("Valid UserData in Put Callback"), KeyData != nullptr);
 					if (KeyData)
 					{
+						if (KeyData->bUseLocal && !KeyData->bUseRemote && !bLocalPutSucceededInitialized)
+						{
+							bLocalPutSucceededInitialized = true;
+							bLocalPutSucceeded = Response.Status == UE::DerivedData::EStatus::Ok;
+						}
+						if (!KeyData->bUseLocal && KeyData->bUseRemote && !bRemotePutSucceededInitialized)
+						{
+							bRemotePutSucceededInitialized = true;
+							bRemotePutSucceeded = Response.Status == UE::DerivedData::EStatus::Ok;
+						}
 						KeyData->bReceivedPut = true;
 					}
 				});
@@ -215,7 +229,7 @@ bool FCacheStoreTest::RunTest(const FString& Parameters)
 
 	{
 		FRequestBarrier Barrier(Owner);
-		Cache.Get(GetRequests, Owner, [&ValueIds, this, NumValues](FCacheGetResponse&& Response)
+		Cache.Get(GetRequests, Owner, [&ValueIds, this, NumValues, bLocalPutSucceeded, bRemotePutSucceeded](FCacheGetResponse&& Response)
 			{
 				FUserData* UserData = reinterpret_cast<FUserData*>(Response.UserData);
 				FKeyData* KeyData = UserData->KeyData;
@@ -224,7 +238,10 @@ bool FCacheStoreTest::RunTest(const FString& Parameters)
 				{
 					int32 n = (int32)KeyData->KeyIndex;
 					KeyData->bReceivedGet = true;
-					if (KeyData->bShouldBeHit)
+
+					bool bShouldBeHit = !KeyData->bForceMiss && ((KeyData->bUseLocal && bLocalPutSucceeded) ||
+						(KeyData->bUseRemote && bRemotePutSucceeded));
+					if (bShouldBeHit)
 					{
 						TestEqual(*WriteToString<32>(TEXT("Get "), n, TEXT(" succeeded")), Response.Status, EStatus::Ok);
 					}
@@ -274,7 +291,7 @@ bool FCacheStoreTest::RunTest(const FString& Parameters)
 				}
 			});
 
-		Cache.GetValue(GetValueRequests, Owner, [this](FCacheGetValueResponse&& Response)
+		Cache.GetValue(GetValueRequests, Owner, [this, bLocalPutSucceeded, bRemotePutSucceeded](FCacheGetValueResponse&& Response)
 			{
 				FUserData* UserData = reinterpret_cast<FUserData*>(Response.UserData);
 				FKeyData* KeyData = UserData->KeyData;
@@ -283,7 +300,10 @@ bool FCacheStoreTest::RunTest(const FString& Parameters)
 				{
 					int32 n = KeyData->KeyIndex;
 					KeyData->bReceivedGetValue = true;
-					if (KeyData->bShouldBeHit)
+
+					bool bShouldBeHit = !KeyData->bForceMiss && ((KeyData->bUseLocal && bLocalPutSucceeded) ||
+						(KeyData->bUseRemote && bRemotePutSucceeded));
+					if (bShouldBeHit)
 					{
 						TestEqual(*WriteToString<32>(TEXT("GetValue "), n, TEXT(" succeeded")), Response.Status, EStatus::Ok);
 					}
@@ -317,7 +337,7 @@ bool FCacheStoreTest::RunTest(const FString& Parameters)
 				}
 			});
 
-		Cache.GetChunks(ChunkRequests, Owner, [this](FCacheGetChunkResponse&& Response)
+		Cache.GetChunks(ChunkRequests, Owner, [this, bLocalPutSucceeded, bRemotePutSucceeded](FCacheGetChunkResponse&& Response)
 			{
 				FUserData* UserData = reinterpret_cast<FUserData*>(Response.UserData);
 				FKeyData* KeyData = UserData->KeyData;
@@ -330,7 +350,9 @@ bool FCacheStoreTest::RunTest(const FString& Parameters)
 					Name << TEXT("GetChunks (") << n << TEXT(",") << ValueIndex << TEXT(")");
 
 					KeyData->ReceivedChunk[ValueIndex] = true;
-					if (KeyData->bShouldBeHit)
+					bool bShouldBeHit = !KeyData->bForceMiss && ((KeyData->bUseLocal && bLocalPutSucceeded) ||
+						(KeyData->bUseRemote && bRemotePutSucceeded));
+					if (bShouldBeHit)
 					{
 						TestEqual(*WriteToString<32>(*Name, TEXT(" succeeded")), Response.Status, EStatus::Ok);
 					}
@@ -338,7 +360,7 @@ bool FCacheStoreTest::RunTest(const FString& Parameters)
 					{
 						TestEqual(*WriteToString<32>(*Name, TEXT(" failed as expected")), Response.Status, EStatus::Error);
 					}
-					if (KeyData->bShouldBeHit && Response.Status == EStatus::Ok)
+					if (bShouldBeHit && Response.Status == EStatus::Ok)
 					{
 						FValue& ExpectedValue = KeyData->BufferValues[ValueIndex];
 						TestEqual(*WriteToString<32>(*Name, TEXT(" Hash")),
