Make ZenCacheStore::FGetChunksOp asynchronous using EnqueueAsyncRpc

diff --git a/Engine/Source/Developer/DerivedDataCache/Private/ZenCacheStore.cpp b/Engine/Source/Developer/DerivedDataCache/Private/ZenCacheStore.cpp
--- a/Engine/Source/Developer/DerivedDataCache/Private/ZenCacheStore.cpp
+++ b/Engine/Source/Developer/DerivedDataCache/Private/ZenCacheStore.cpp
@@ -791,75 +791,71 @@ public:
 			}
 			BatchRequest.EndObject();
 
-			FCbPackage BatchResponse;
-			EHttpErrorCode HttpErrorCode = EHttpErrorCode::Unknown;
-
+			FGetChunksOp* OriginalOp = this;
+			auto OnRpcComplete = [this, OpRef = TRefCountPtr<FGetChunksOp>(OriginalOp), Batch](THttpUniquePtr<IHttpResponse>& HttpResponse, FCbPackage& Response)
 			{
-				LLM_SCOPE_BYTAG(UntaggedDDCResult);
-				THttpUniquePtr<IHttpResponse> HttpResponse = CacheStore.PerformBlockingRpc(BatchRequest.Save().AsObject(), BatchResponse);
-				HttpErrorCode = HttpResponse->GetErrorCode();
-			}
-
-			int32 RequestIndex = 0;
-			if (HttpErrorCode == EHttpErrorCode::None)
-			{
-				const FCbObject& ResponseObj = BatchResponse.GetObject();
-
-				for (FCbFieldView ResultView : ResponseObj[ANSITEXTVIEW("Result")])
+				int32 RequestIndex = 0;
+				if (HttpResponse->GetErrorCode() == EHttpErrorCode::None)
 				{
-					if (RequestIndex >= Batch.Num())
-					{
-						++RequestIndex;
-						continue;
-					}
-					const FCacheGetChunkRequest& Request = Batch[RequestIndex++];
-						
-					FIoHash RawHash;
-					bool Succeeded = false;
-					uint64 RawSize = 0;
-					FCbObjectView ResultObject = ResultView.AsObjectView();
-					FSharedBuffer RequestedBytes;
-					if (CacheStore.DebugOptions.ShouldSimulateGetMiss(Request.Key))
-					{
-						UE_LOG(LogDerivedDataCache, Verbose, TEXT("%s: Simulated miss for get of '%s' from '%s'"),
-							*CacheStore.GetName(), *WriteToString<96>(Request.Key, '/', Request.Id), *Request.Name);
-					}
-					else
+					const FCbObject& ResponseObj = Response.GetObject();
+
+					for (FCbFieldView ResultView : ResponseObj[ANSITEXTVIEW("Result")])
 					{
-						FCbFieldView HashView = ResultObject[ANSITEXTVIEW("RawHash")];
-						RawHash = HashView.AsHash();
-						if (!HashView.HasError())
+						if (RequestIndex >= Batch.Num())
 						{
-							if (const FCbAttachment* Attachment = BatchResponse.FindAttachment(HashView.AsHash()))
+							++RequestIndex;
+							continue;
+						}
+						const FCacheGetChunkRequest& Request = Batch[RequestIndex++];
+
+						FIoHash RawHash;
+						bool Succeeded = false;
+						uint64 RawSize = 0;
+						FCbObjectView ResultObject = ResultView.AsObjectView();
+						FSharedBuffer RequestedBytes;
+						if (CacheStore.DebugOptions.ShouldSimulateGetMiss(Request.Key))
+						{
+							UE_LOG(LogDerivedDataCache, Verbose, TEXT("%s: Simulated miss for get of '%s' from '%s'"),
+								*CacheStore.GetName(), *WriteToString<96>(Request.Key, '/', Request.Id), *Request.Name);
+						}
+						else
+						{
+							FCbFieldView HashView = ResultObject[ANSITEXTVIEW("RawHash")];
+							RawHash = HashView.AsHash();
+							if (!HashView.HasError())
 							{
-								FCompressedBuffer CompressedBuffer = Attachment->AsCompressedBinary();
-								if (CompressedBuffer)
+								if (const FCbAttachment* Attachment = Response.FindAttachment(HashView.AsHash()))
 								{
-									TRACE_COUNTER_ADD(ZenDDC_BytesReceived, CompressedBuffer.GetCompressedSize());
-									RequestedBytes = FCompressedBufferReader(CompressedBuffer).Decompress(Request.RawOffset, Request.RawSize);
-									RawSize = RequestedBytes.GetSize();
-									Succeeded = true;
+									FCompressedBuffer CompressedBuffer = Attachment->AsCompressedBinary();
+									if (CompressedBuffer)
+									{
+										TRACE_COUNTER_ADD(ZenDDC_BytesReceived, CompressedBuffer.GetCompressedSize());
+										RequestedBytes = FCompressedBufferReader(CompressedBuffer).Decompress(Request.RawOffset, Request.RawSize);
+										RawSize = RequestedBytes.GetSize();
+										Succeeded = true;
+									}
 								}
-							}
-							else
-							{
-								FCbFieldView RawSizeField = ResultObject[ANSITEXTVIEW("RawSize")];
-								uint64 TotalSize = RawSizeField.AsUInt64();
-								Succeeded = !RawSizeField.HasError();
-								if (Succeeded)
+								else
 								{
-									RawSize = FMath::Min(Request.RawSize, TotalSize - FMath::Min(Request.RawOffset, TotalSize));
+									FCbFieldView RawSizeField = ResultObject[ANSITEXTVIEW("RawSize")];
+									uint64 TotalSize = RawSizeField.AsUInt64();
+									Succeeded = !RawSizeField.HasError();
+									if (Succeeded)
+									{
+										RawSize = FMath::Min(Request.RawSize, TotalSize - FMath::Min(Request.RawOffset, TotalSize));
+									}
 								}
 							}
 						}
+						Succeeded ? OnHit(Request, MoveTemp(RawHash), RawSize, MoveTemp(RequestedBytes)) : OnMiss(Request);
 					}
-					Succeeded ? OnHit(Request, MoveTemp(RawHash), RawSize, MoveTemp(RequestedBytes)) : OnMiss(Request);
 				}
-			}
-			for (const FCacheGetChunkRequest& Request : Batch.RightChop(RequestIndex))
-			{
-				OnMiss(Request);
-			}
+				for (const FCacheGetChunkRequest& Request : Batch.RightChop(RequestIndex))
+				{
+					OnMiss(Request);
+				}
+			};
+			CacheStore.EnqueueAsyncRpc(Owner, BatchRequest.Save().AsObject(), MoveTemp(OnRpcComplete));
 		});
 	}
 private:
