Use Http::TryLoadCbPackage as first option to parse packages from Zen.
This will allow UE to talk to both old and new zen servers as we move to use the new package wire format in Zen responses.

diff --git a/Engine/Source/Developer/DerivedDataCache/Private/ZenCacheStore.cpp b/Engine/Source/Developer/DerivedDataCache/Private/ZenCacheStore.cpp
--- a/Engine/Source/Developer/DerivedDataCache/Private/ZenCacheStore.cpp
+++ b/Engine/Source/Developer/DerivedDataCache/Private/ZenCacheStore.cpp
@@ -903,7 +903,15 @@ private:
 
 	IHttpReceiver* OnComplete(IHttpResponse& Response) final
 	{
-		FMemoryReaderView Ar(MakeMemoryView(BodyArray));
+		FMemoryView MemoryView = MakeMemoryView(BodyArray);
+		{
+			FMemoryReaderView Ar(MemoryView);
+			if (Zen::Http::TryLoadCbPackage(Package, Ar))
+			{
+				return Next;
+			}
+		}
+		FMemoryReaderView Ar(MemoryView);
 		Package.TryLoad(Ar);
 		return Next;
 	}
diff --git a/Engine/Source/Developer/Zen/Internal/ZenServerHttp.cpp b/Engine/Source/Developer/Zen/Internal/ZenServerHttp.cpp
--- a/Engine/Source/Developer/Zen/Internal/ZenServerHttp.cpp
+++ b/Engine/Source/Developer/Zen/Internal/ZenServerHttp.cpp
@@ -182,6 +182,13 @@ namespace UE::Zen {
 
 		if (ResponseBuffer.Num())
 		{
+			{
+				FLargeMemoryReader Ar(ResponseBuffer.GetData(), ResponseBuffer.Num());
+				if (Http::TryLoadCbPackage(OutResponse, Ar))
+				{
+					return Result::Success;
+				}
+			}
 			FLargeMemoryReader Ar(ResponseBuffer.GetData(), ResponseBuffer.Num());
 			if (!OutResponse.TryLoad(Ar))
 			{
