DDC: Fixed untagged allocations by capturing tags in tasks

diff --git a/Engine/Source/Developer/DerivedDataCache/Private/DerivedDataBuild.cpp b/Engine/Source/Developer/DerivedDataCache/Private/DerivedDataBuild.cpp
--- a/Engine/Source/Developer/DerivedDataCache/Private/DerivedDataBuild.cpp
+++ b/Engine/Source/Developer/DerivedDataCache/Private/DerivedDataBuild.cpp
@@ -15,6 +15,7 @@
 #include "DerivedDataBuildTypes.h"
 #include "DerivedDataBuildWorkerRegistry.h"
 #include "DerivedDataCache.h"
+#include "DerivedDataPrivate.h"
 #include "DerivedDataSharedString.h"
 #include "Misc/Guid.h"
 
@@ -128,6 +129,7 @@ private:
 
 IBuild* CreateBuild(ICache& Cache)
 {
+	LLM_SCOPE_BYTAG(DerivedDataBuild);
 	return new FBuild(Cache);
 }
 
diff --git a/Engine/Source/Developer/DerivedDataCache/Private/DerivedDataCache.cpp b/Engine/Source/Developer/DerivedDataCache/Private/DerivedDataCache.cpp
--- a/Engine/Source/Developer/DerivedDataCache/Private/DerivedDataCache.cpp
+++ b/Engine/Source/Developer/DerivedDataCache/Private/DerivedDataCache.cpp
@@ -14,10 +14,12 @@
 #include "DerivedDataCachePrivate.h"
 #include "DerivedDataCacheUsageStats.h"
 #include "DerivedDataPluginInterface.h"
+#include "DerivedDataPrivate.h"
 #include "DerivedDataRequest.h"
 #include "DerivedDataRequestOwner.h"
 #include "Experimental/Async/LazyEvent.h"
 #include "Features/IModularFeatures.h"
+#include "HAL/LowLevelMemTracker.h"
 #include "HAL/ThreadSafeCounter.h"
 #include "Misc/CommandLine.h"
 #include "Misc/CoreMisc.h"
@@ -331,6 +333,7 @@ public:
 		: Owner(InOwner)
 		, TaskBody(MoveTemp(InTaskBody))
 	{
+		LLM_IF_ENABLED(MemTag = FLowLevelMemTracker::Get().GetActiveTagData(ELLMTracker::Default));
 		Owner.Begin(this);
 		DoneEvent.Reset();
 		GCacheThreadPool->AddQueuedWork(this, ConvertToQueuedWorkPriority(Owner.GetPriority()));
@@ -339,6 +342,7 @@ public:
 private:
 	inline void Execute()
 	{
+		LLM_SCOPE(MemTag);
 		FScopeCycleCounter Scope(GetStatId(), /*bAlways*/ true);
 		Owner.End(this, [this]
 		{
@@ -404,6 +408,7 @@ private:
 	IRequestOwner& Owner;
 	TUniqueFunction<void ()> TaskBody;
 	FLazyEvent DoneEvent{EEventMode::ManualReset};
+	LLM(const UE::LLMPrivate::FTagData* MemTag = nullptr);
 };
 
 void LaunchTaskInCacheThreadPool(IRequestOwner& Owner, TUniqueFunction<void ()>&& TaskBody)
@@ -518,7 +523,7 @@ class FDerivedDataCache final
 			if (bGetResult)
 			{
 				
-				if(GVerifyDDC && DataDeriver && DataDeriver->IsDeterministic())
+				if (GVerifyDDC && DataDeriver && DataDeriver->IsDeterministic())
 				{
 					TArray<uint8> CmpData;
 					DataDeriver->Build(CmpData);
@@ -542,13 +547,12 @@ class FDerivedDataCache final
 						}
 					}
 
-					if(!bMatchesInSize || bDifferentMemory)
+					if (!bMatchesInSize || bDifferentMemory)
 					{
 						FString ErrMsg = FString::Printf(TEXT("There is a mismatch between the DDC data and the generated data for plugin (%s) for asset (%s). BytesInDDC:%d, BytesGenerated:%d, bDifferentMemory:%d, offset:%d"), DataDeriver->GetPluginName(), *DataDeriver->GetDebugContextString(), NumInDDC, NumGenerated, bDifferentMemory, DifferentOffset);
 						ensureMsgf(false, TEXT("%s"), *ErrMsg);
 						UE_LOG(LogDerivedDataCache, Error, TEXT("%s"), *ErrMsg );
 					}
-					
 				}
 
 				check(Data.Num());
@@ -1142,6 +1146,7 @@ private:
 
 ICache* CreateCache(FDerivedDataCacheInterface** OutLegacyCache)
 {
+	LLM_SCOPE_BYTAG(DerivedDataCache);
 	FDerivedDataCache* Cache = new FDerivedDataCache;
 	if (OutLegacyCache)
 	{
diff --git a/Engine/Source/Developer/DerivedDataCache/Private/DerivedDataCacheModule.cpp b/Engine/Source/Developer/DerivedDataCache/Private/DerivedDataCacheModule.cpp
--- a/Engine/Source/Developer/DerivedDataCache/Private/DerivedDataCacheModule.cpp
+++ b/Engine/Source/Developer/DerivedDataCache/Private/DerivedDataCacheModule.cpp
@@ -7,6 +7,7 @@
 #include "DerivedDataBuildPrivate.h"
 #include "DerivedDataCache.h"
 #include "DerivedDataCachePrivate.h"
+#include "DerivedDataPrivate.h"
 #include "HAL/CriticalSection.h"
 #include "Misc/ScopeLock.h"
 #include "Modules/ModuleManager.h"
@@ -100,6 +101,10 @@ static FDerivedDataCacheModule* GetModule()
 namespace UE::DerivedData
 {
 
+LLM_DEFINE_TAG(DerivedData);
+LLM_DEFINE_TAG(DerivedDataBuild, "Build", "DerivedData");
+LLM_DEFINE_TAG(DerivedDataCache, "Cache", "DerivedData");
+
 ICache& GetCache()
 {
 	if (ICache* Cache = Private::GDerivedDataCache)
diff --git a/Engine/Source/Developer/DerivedDataCache/Private/DerivedDataPrivate.h b/Engine/Source/Developer/DerivedDataCache/Private/DerivedDataPrivate.h
new file mode 100644
--- /dev/null
+++ b/Engine/Source/Developer/DerivedDataCache/Private/DerivedDataPrivate.h
@@ -0,0 +1,14 @@
+// Copyright Epic Games, Inc. All Rights Reserved.
+
+#pragma once
+
+#include "HAL/LowLevelMemTracker.h"
+
+namespace UE::DerivedData
+{
+
+LLM_DECLARE_TAG(DerivedData);
+LLM_DECLARE_TAG(DerivedDataBuild);
+LLM_DECLARE_TAG(DerivedDataCache);
+
+} // UE::DerivedData
diff --git a/Engine/Source/Developer/DerivedDataCache/Private/DerivedDataRequestOwner.cpp b/Engine/Source/Developer/DerivedDataCache/Private/DerivedDataRequestOwner.cpp
--- a/Engine/Source/Developer/DerivedDataCache/Private/DerivedDataRequestOwner.cpp
+++ b/Engine/Source/Developer/DerivedDataCache/Private/DerivedDataRequestOwner.cpp
@@ -7,6 +7,7 @@
 #include "DerivedDataRequestTypes.h"
 #include "Experimental/Async/LazyEvent.h"
 #include "HAL/CriticalSection.h"
+#include "HAL/LowLevelMemTracker.h"
 #include "Misc/ScopeExit.h"
 #include "Misc/ScopeRWLock.h"
 #include "Tasks/Task.h"
@@ -298,13 +299,19 @@ void IRequestOwner::LaunchTask(const TCHAR* DebugName, TUniqueFunction<void ()>&
 		void Cancel() final { Task.Wait(); }
 		void Wait() final { Task.Wait(); }
 		FTask Task;
+		LLM(const UE::LLMPrivate::FTagData* MemTag = nullptr);
 	};
 
 	Tasks::FTaskEvent TaskEvent(TEXT("LaunchTaskRequest"));
 	FTaskRequest* Request = new FTaskRequest;
+	LLM_IF_ENABLED(Request->MemTag = FLowLevelMemTracker::Get().GetActiveTagData(ELLMTracker::Default));
 	Request->Task = Launch(
 		DebugName,
-		[this, Request, TaskBody = MoveTemp(TaskBody)] { End(Request, TaskBody); },
+		[this, Request, TaskBody = MoveTemp(TaskBody)]
+		{
+			LLM_SCOPE(Request->MemTag);
+			End(Request, TaskBody);
+		},
 		TaskEvent,
 		GetPriority() <= EPriority::Normal ? ETaskPriority::BackgroundNormal : ETaskPriority::BackgroundHigh);
 	Begin(Request);
diff --git a/Engine/Source/Developer/DerivedDataCache/Private/HttpCacheStore.cpp b/Engine/Source/Developer/DerivedDataCache/Private/HttpCacheStore.cpp
--- a/Engine/Source/Developer/DerivedDataCache/Private/HttpCacheStore.cpp
+++ b/Engine/Source/Developer/DerivedDataCache/Private/HttpCacheStore.cpp
@@ -21,6 +21,7 @@
 #include "Dom/JsonObject.h"
 #include "HAL/CriticalSection.h"
 #include "HAL/IConsoleManager.h"
+#include "HAL/LowLevelMemTracker.h"
 #include "HAL/PlatformProcess.h"
 #include "Http/HttpClient.h"
 #include "IO/IoHash.h"
@@ -568,6 +569,7 @@ public:
 		, BaseReceiver(InOperation, this)
 		, OperationComplete(MoveTemp(InOperationComplete))
 	{
+		LLM_IF_ENABLED(MemTag = FLowLevelMemTracker::Get().GetActiveTagData(ELLMTracker::Default));
 	}
 
 private:
@@ -581,6 +583,7 @@ private:
 
 	IHttpReceiver* OnCreate(IHttpResponse& LocalResponse) final
 	{
+		LLM_SCOPE(MemTag);
 		Monitor = LocalResponse.GetMonitor();
 		Owner->Begin(this);
 		return &BaseReceiver;
@@ -588,6 +591,7 @@ private:
 
 	IHttpReceiver* OnComplete(IHttpResponse& LocalResponse) final
 	{
+		LLM_SCOPE(MemTag);
 		Owner->End(this, [Self = this]
 		{
 			FHttpOperation* Operation = Self->BaseReceiver.GetOperation();
@@ -614,6 +618,7 @@ private:
 	FHttpOperationReceiver BaseReceiver;
 	TUniqueFunction<void ()> OperationComplete;
 	TRefCountPtr<IHttpResponseMonitor> Monitor;
+	LLM(const UE::LLMPrivate::FTagData* MemTag = nullptr);
 };
 
 void FHttpCacheStore::FHttpOperation::Send()
diff --git a/Engine/Source/Runtime/Core/Private/ProfilingDebugging/TagTrace.cpp b/Engine/Source/Runtime/Core/Private/ProfilingDebugging/TagTrace.cpp
--- a/Engine/Source/Runtime/Core/Private/ProfilingDebugging/TagTrace.cpp
+++ b/Engine/Source/Runtime/Core/Private/ProfilingDebugging/TagTrace.cpp
@@ -6,6 +6,7 @@
 #include "CoreTypes.h"
 #include "ProfilingDebugging/MemoryTrace.h"
 #include "HAL/LowLevelMemTracker.h"
+#include "HAL/LowLevelMemTrackerPrivate.h"
 #include "Trace/Trace.inl"
 #include "UObject/NameTypes.h"
 
@@ -65,9 +66,10 @@ FMemScope::FMemScope(const FName& InName, bool bShouldActivate /*= true*/)
 ////////////////////////////////////////////////////////////////////////////////
 FMemScope::FMemScope(const UE::LLMPrivate::FTagData* TagData, bool bShouldActivate /*= true*/)
 {
-	// TagData is opaque so we cant really use the input, additionally we 
-	// cannot count on LLM being active. Instead we have inserted an explicit 
-	// Trace scope directly following the LLM scope. 
+	if (UE_TRACE_CHANNELEXPR_IS_ENABLED(MemAllocChannel) & bShouldActivate && TagData)
+	{
+		ActivateScope(MemoryTrace_AnnounceFNameTag(TagData->GetName()));
+	}
 }
 
 ////////////////////////////////////////////////////////////////////////////////
