DDC: Used InheritedContext to implement the reverted inheritance of memory tags

#rnx

diff --git a/Engine/Source/Developer/DerivedDataCache/Private/DerivedDataCache.cpp b/Engine/Source/Developer/DerivedDataCache/Private/DerivedDataCache.cpp
--- a/Engine/Source/Developer/DerivedDataCache/Private/DerivedDataCache.cpp
+++ b/Engine/Source/Developer/DerivedDataCache/Private/DerivedDataCache.cpp
@@ -5,6 +5,7 @@
 
 #include "Algo/AllOf.h"
 #include "Async/AsyncWork.h"
+#include "Async/InheritedContext.h"
 #include "Async/TaskGraphInterfaces.h"
 #include "Containers/Map.h"
 #include "DDCCleanup.h"
@@ -325,13 +326,14 @@ namespace UE::DerivedData::Private
 
 FQueuedThreadPool* GCacheThreadPool;
 
-class FCacheThreadPoolTaskRequest final : public FRequestBase, private IQueuedWork
+class FCacheThreadPoolTaskRequest final : public FRequestBase, private FInheritedContextBase, private IQueuedWork
 {
 public:
 	inline FCacheThreadPoolTaskRequest(IRequestOwner& InOwner, TUniqueFunction<void ()>&& InTaskBody)
 		: Owner(InOwner)
 		, TaskBody(MoveTemp(InTaskBody))
 	{
+		CaptureInheritedContext();
 		Owner.Begin(this);
 		DoneEvent.Reset();
 		GCacheThreadPool->AddQueuedWork(this, ConvertToQueuedWorkPriority(Owner.GetPriority()));
@@ -340,6 +342,7 @@ public:
 private:
 	inline void Execute()
 	{
+		FInheritedContextScope InheritedContextScope = RestoreInheritedContext();
 		FScopeCycleCounter Scope(GetStatId(), /*bAlways*/ true);
 		Owner.End(this, [this]
 		{
diff --git a/Engine/Source/Developer/DerivedDataCache/Private/Http/CurlHttpClient.cpp b/Engine/Source/Developer/DerivedDataCache/Private/Http/CurlHttpClient.cpp
--- a/Engine/Source/Developer/DerivedDataCache/Private/Http/CurlHttpClient.cpp
+++ b/Engine/Source/Developer/DerivedDataCache/Private/Http/CurlHttpClient.cpp
@@ -2,6 +2,7 @@
 
 #include "HttpClient.h"
 
+#include "Async/InheritedContext.h"
 #include "Containers/DepletableMpscQueue.h"
 #include "Containers/LockFreeList.h"
 #include "Containers/StringView.h"
@@ -279,7 +280,7 @@ enum class ECurlHttpResponseState : uint8
 
 ENUM_CLASS_FLAGS(ECurlHttpResponseState);
 
-class FCurlHttpResponse final : public IHttpResponse, public IHttpResponseMonitor
+class FCurlHttpResponse final : public IHttpResponse, public IHttpResponseMonitor, private FInheritedContextBase
 {
 public:
 	FCurlHttpResponse(CURL* Curl, EHttpMethod Method, FAnsiStringView Uri, IHttpReceiver* Receiver);
@@ -933,6 +934,8 @@ FCurlHttpResponse::FCurlHttpResponse(CURL* InCurl, EHttpMethod InMethod, FAnsiSt
 	, Curl(InCurl)
 	, Receiver(InReceiver)
 {
+	CaptureInheritedContext();
+
 	// Release() is called by Destroy().
 	AddRef();
 
@@ -999,6 +1002,8 @@ bool FCurlHttpResponse::WriteHeader(FAnsiStringView Header)
 		return !EnumHasAnyFlags(State.load(std::memory_order_relaxed), ECurlHttpResponseState::Canceled);
 	}
 
+	FInheritedContextScope InheritedContextScope = RestoreInheritedContext();
+
 	// Reset headers between responses and keep only the last.
 	if (!HeaderViews.IsEmpty())
 	{
@@ -1040,6 +1045,8 @@ bool FCurlHttpResponse::WriteHeader(FAnsiStringView Header)
 
 bool FCurlHttpResponse::WriteBody(FMemoryView Body)
 {
+	FInheritedContextScope InheritedContextScope = RestoreInheritedContext();
+
 	ConditionallySetResponseStatus();
 
 	bHasBody = true;
@@ -1116,8 +1123,9 @@ void FCurlHttpResponse::SetComplete(CURLcode Code)
 	Curl = nullptr;
 	Client = nullptr;
 
-	// Hold a reference to safely access Receiver, State, and CompleteEvent.
+	// Hold a reference to safely access Receiver, State, CompleteEvent, and FInheritedContextBase.
 	TRefCountPtr<FCurlHttpResponse> Self(this);
+	FInheritedContextScope InheritedContextScope = RestoreInheritedContext();
 
 	for (TGuardValue GuardReceiverFunction(ReceiverFunction, EHttpReceiverFunction::OnComplete);;)
 	{
