DDC: Fixed SharedString reference counting flagged by TSan

#rnx

diff --git a/Engine/Source/Developer/DerivedDataCache/Public/DerivedDataSharedString.h b/Engine/Source/Developer/DerivedDataCache/Public/DerivedDataSharedString.h
--- a/Engine/Source/Developer/DerivedDataCache/Public/DerivedDataSharedString.h
+++ b/Engine/Source/Developer/DerivedDataCache/Public/DerivedDataSharedString.h
@@ -145,7 +145,7 @@ void TSharedString<CharType>::AddRef(CharType* const MaybeNullChars)
 		static_assert(sizeof(int32) == sizeof(std::atomic<int32>));
 		int32* const Header = reinterpret_cast<int32*>(MaybeNullChars) - 2;
 		std::atomic<int32>& RefCount = reinterpret_cast<std::atomic<int32>&>(Header[0]);
-		RefCount.fetch_add(1, std::memory_order_relaxed);
+		RefCount.fetch_add(1, std::memory_order_acquire);
 	}
 }
 
@@ -157,7 +157,7 @@ void TSharedString<CharType>::Release(CharType* const MaybeNullChars)
 		static_assert(sizeof(int32) == sizeof(std::atomic<int32>));
 		int32* const Header = reinterpret_cast<int32*>(MaybeNullChars) - 2;
 		std::atomic<int32>& RefCount = reinterpret_cast<std::atomic<int32>&>(Header[0]);
-		if (RefCount.fetch_sub(1, std::memory_order_relaxed) == 1)
+		if (RefCount.fetch_sub(1, std::memory_order_release) == 1)
 		{
 			FMemory::Free(Header);
 		}
