DDC: Fixed the default cache verify mode to verify every access

#rnx

diff --git a/Engine/Source/Developer/DerivedDataCache/Private/DerivedDataCacheStoreVerify.cpp b/Engine/Source/Developer/DerivedDataCache/Private/DerivedDataCacheStoreVerify.cpp
--- a/Engine/Source/Developer/DerivedDataCache/Private/DerivedDataCacheStoreVerify.cpp
+++ b/Engine/Source/Developer/DerivedDataCache/Private/DerivedDataCacheStoreVerify.cpp
@@ -36,7 +36,7 @@ public:
 
 		const TCHAR* const CommandLine = FCommandLine::Get();
 
-		const bool bDefaultMatch = FParse::Param(CommandLine, TEXT("-DDC-Verify")) ||
+		const bool bDefaultMatch = FParse::Param(CommandLine, TEXT("DDC-Verify")) ||
 			String::FindFirst(CommandLine, TEXT("-DDC-Verify="), ESearchCase::IgnoreCase) == INDEX_NONE;
 		float DefaultRate = bDefaultMatch ? 100.0f : 0.0f;
 		FParse::Value(CommandLine, TEXT("-DDC-VerifyRate="), DefaultRate);
