Add "Accept" field in Zen RPC request to enable new style RPC responses.
See https://github.com/EpicGames/zen/pull/180 for Zen side change.

diff --git a/Engine/Source/Developer/DerivedDataCache/Private/ZenCacheStore.cpp b/Engine/Source/Developer/DerivedDataCache/Private/ZenCacheStore.cpp
--- a/Engine/Source/Developer/DerivedDataCache/Private/ZenCacheStore.cpp
+++ b/Engine/Source/Developer/DerivedDataCache/Private/ZenCacheStore.cpp
@@ -183,6 +183,8 @@ public:
 			BatchWriter.BeginObject();
 			{
 				BatchWriter << ANSITEXTVIEW("Method") << "PutCacheRecords";
+				BatchWriter.AddInteger(ANSITEXTVIEW("Accept"), Zen::Http::kCbPkgMagic);
+
 				BatchWriter.BeginObject(ANSITEXTVIEW("Params"));
 				{
 					ECachePolicy BatchDefaultPolicy = Batch[0].Policy.GetRecordPolicy();
@@ -331,6 +333,8 @@ public:
 				BatchRequest.BeginObject();
 				{
 					BatchRequest << ANSITEXTVIEW("Method") << ANSITEXTVIEW("GetCacheRecords");
+					BatchRequest.AddInteger(ANSITEXTVIEW("Accept"), Zen::Http::kCbPkgMagic);
+
 					BatchRequest.BeginObject(ANSITEXTVIEW("Params"));
 					{
 						ECachePolicy BatchDefaultPolicy = Batch[0].Policy.GetRecordPolicy();
@@ -457,6 +461,8 @@ public:
 			BatchWriter.BeginObject();
 			{
 				BatchWriter << ANSITEXTVIEW("Method") << ANSITEXTVIEW("PutCacheValues");
+				BatchWriter.AddInteger(ANSITEXTVIEW("Accept"), Zen::Http::kCbPkgMagic);
+
 				BatchWriter.BeginObject(ANSITEXTVIEW("Params"));
 				{
 					ECachePolicy BatchDefaultPolicy = Batch[0].Policy;
@@ -601,6 +607,8 @@ public:
 			BatchRequest.BeginObject();
 			{
 				BatchRequest << ANSITEXTVIEW("Method") << ANSITEXTVIEW("GetCacheValues");
+				BatchRequest.AddInteger(ANSITEXTVIEW("Accept"), Zen::Http::kCbPkgMagic);
+
 				BatchRequest.BeginObject(ANSITEXTVIEW("Params"));
 				{
 					ECachePolicy BatchDefaultPolicy = Batch[0].Policy;
@@ -748,7 +756,8 @@ public:
 			BatchRequest.BeginObject();
 			{
 				BatchRequest << ANSITEXTVIEW("Method") << "GetCacheChunks";
-				BatchRequest.AddInteger(ANSITEXTVIEW("MethodVersion"), 1);
+				BatchRequest.AddInteger(ANSITEXTVIEW("Accept"), Zen::Http::kCbPkgMagic);
+
 				BatchRequest.BeginObject(ANSITEXTVIEW("Params"));
 				{
 					ECachePolicy DefaultPolicy = Batch[0].Policy;
diff --git a/Engine/Source/Developer/Zen/Internal/ZenSerialization.cpp b/Engine/Source/Developer/Zen/Internal/ZenSerialization.cpp
--- a/Engine/Source/Developer/Zen/Internal/ZenSerialization.cpp
+++ b/Engine/Source/Developer/Zen/Internal/ZenSerialization.cpp
@@ -133,8 +133,6 @@ struct CbPackageHeader
 	uint32	Reserved2;
 };
 
-static const uint32 kMagic = 0xaa77aacc;
-
 struct CbAttachmentEntry
 {
 	uint64	AttachmentSize;
@@ -156,7 +154,7 @@ void SaveCbPackage(const FCbPackage& Package, FArchive& Ar)
 	FCompressedBuffer ObjectBuffer = FCompressedBuffer::Compress(Object.GetBuffer(), FOodleDataCompression::ECompressor::NotSet, FOodleDataCompression::ECompressionLevel::None);
 
 	CbPackageHeader Hdr;
-	Hdr.HeaderMagic = kMagic;
+	Hdr.HeaderMagic = kCbPkgMagic;
 	Hdr.AttachmentCount = Attachments.Num();
 	Hdr.Reserved1 = 0;
 	Hdr.Reserved2 = 0;
@@ -248,7 +246,7 @@ bool TryLoadCbPackage(FCbPackage& Package, FArchive& Ar, FCbBufferAllocator Allo
 	CbPackageHeader Hdr;
 	Ar.Serialize(&Hdr, sizeof Hdr);
 
-	if (Hdr.HeaderMagic != kMagic)
+	if (Hdr.HeaderMagic != kCbPkgMagic)
 	{
 		return false;
 	}
diff --git a/Engine/Source/Developer/Zen/Internal/ZenSerialization.h b/Engine/Source/Developer/Zen/Internal/ZenSerialization.h
--- a/Engine/Source/Developer/Zen/Internal/ZenSerialization.h
+++ b/Engine/Source/Developer/Zen/Internal/ZenSerialization.h
@@ -26,6 +26,8 @@ ZEN_API bool TryLoadCbPackage(FCbPackage& Package, FArchive& Ar, FCbBufferAlloca
 
 namespace Http
 {
+static const uint32 kCbPkgMagic = 0xaa77aacc;
+	
 ZEN_API void SaveCbPackage(const FCbPackage& Package, FArchive& Ar);
 ZEN_API bool TryLoadCbPackage(FCbPackage& Package, FArchive& Ar, FCbBufferAllocator Allocator = FUniqueBuffer::Alloc);
 }
