If multiple active FileSystemCacheStores in DDC are configured to have the same path, ignore all but the first one to become active.

diff --git a/Engine/Source/Developer/DerivedDataCache/Private/FileSystemCacheStore.cpp b/Engine/Source/Developer/DerivedDataCache/Private/FileSystemCacheStore.cpp
--- a/Engine/Source/Developer/DerivedDataCache/Private/FileSystemCacheStore.cpp
+++ b/Engine/Source/Developer/DerivedDataCache/Private/FileSystemCacheStore.cpp
@@ -3,6 +3,7 @@
 #include "Algo/Accumulate.h"
 #include "Algo/AllOf.h"
 #include "Algo/Compare.h"
+#include "Algo/Find.h"
 #include "Algo/StableSort.h"
 #include "Algo/Transform.h"
 #include "Async/Async.h"
@@ -761,6 +762,7 @@ public:
 		const TCHAR* Params,
 		const TCHAR* AccessLogPath,
 		ECacheStoreFlags& OutFlags);
+	~FFileSystemCacheStore();
 
 	static bool RunSpeedTest(
 		FString CachePath,
@@ -881,6 +883,9 @@ private:
 
 	TUniquePtr<FFileSystemCacheStoreMaintainerParams> DeactivationDeferredMaintainerParams;
 
+	static inline FRWLock ActiveFileStoresLock;
+	static inline TArray<FFileSystemCacheStore*> ActiveFileStores;
+
 	enum class EPerformanceReEvaluationResult
 	{
 		Invalid = 0,
@@ -923,6 +928,20 @@ FFileSystemCacheStore::FFileSystemCacheStore(
 	FParse::Value(InParams, TEXT("MaxRecordSizeKB="), MaxRecordSizeKB);
 	FParse::Value(InParams, TEXT("MaxValueSizeKB="), MaxValueSizeKB);
 
+	{
+		FReadScopeLock ReadLock(ActiveFileStoresLock);
+		FFileSystemCacheStore** FoundExitingActiveStore = Algo::FindByPredicate(ActiveFileStores, [this](const FFileSystemCacheStore* ActiveStore)
+		{
+			return ActiveStore && FPaths::IsSamePath(ActiveStore->CachePath, CachePath);
+		});
+		if (FoundExitingActiveStore)
+		{
+			UE_LOG(LogDerivedDataCache, Warning, TEXT("Attempted to create multiple cache stores to path %s, ignoring duplicate entry"), *CachePath);
+			OutFlags = ECacheStoreFlags::None;
+			return;
+		}
+	}
+
 	// Flush the cache if requested.
 	bool bFlush = false;
 	if (!bReadOnly && FParse::Bool(InParams, TEXT("Flush="), bFlush) && bFlush)
@@ -1129,9 +1148,20 @@ FFileSystemCacheStore::FFileSystemCacheStore(
 		Flags |= bReadOnly ? ECacheStoreFlags::None : ECacheStoreFlags::Store;
 		Flags |= SpeedClass == EBackendSpeedClass::Local ? ECacheStoreFlags::Local : ECacheStoreFlags::Remote;
 		OutFlags = Flags;
+
+		{
+			FWriteScopeLock WriteLock(ActiveFileStoresLock);
+			ActiveFileStores.Add(this);
+		}
 	}
 }
 
+FFileSystemCacheStore::~FFileSystemCacheStore()
+{
+	FWriteScopeLock WriteLock(ActiveFileStoresLock);
+	ActiveFileStores.Remove(this);
+}
+
 bool FFileSystemCacheStore::RunSpeedTest(
 	FString CachePath,
 	bool bReadOnly,
