DDC: Sorted and removed duplicate includes in DerivedDataCache.cpp

#rnx

diff --git a/Engine/Source/Developer/DerivedDataCache/Private/DerivedDataCache.cpp b/Engine/Source/Developer/DerivedDataCache/Private/DerivedDataCache.cpp
--- a/Engine/Source/Developer/DerivedDataCache/Private/DerivedDataCache.cpp
+++ b/Engine/Source/Developer/DerivedDataCache/Private/DerivedDataCache.cpp
@@ -4,6 +4,8 @@
 #include "DerivedDataCacheInterface.h"
 
 #include "Algo/AllOf.h"
+#include "Analytics.h"
+#include "AnalyticsEventAttribute.h"
 #include "Async/AsyncWork.h"
 #include "Async/InheritedContext.h"
 #include "Async/TaskGraphInterfaces.h"
@@ -18,7 +20,6 @@
 #include "DerivedDataPrivate.h"
 #include "DerivedDataRequest.h"
 #include "DerivedDataRequestOwner.h"
-#include "DerivedDataCacheUsageStats.h"
 #include "Experimental/Async/LazyEvent.h"
 #include "Features/IModularFeatures.h"
 #include "HAL/ThreadSafeCounter.h"
@@ -32,8 +33,6 @@
 #include "Stats/Stats.h"
 #include "Stats/StatsMisc.h"
 #include "ZenServerInterface.h"
-#include "AnalyticsEventAttribute.h"
-#include "Analytics.h"
 #include <atomic>
 
 DEFINE_STAT(STAT_DDC_NumGets);
