Zen: Changed the Analytics dependency from public to private

#rnx

diff --git a/Engine/Source/Developer/DerivedDataCache/DerivedDataCache.Build.cs b/Engine/Source/Developer/DerivedDataCache/DerivedDataCache.Build.cs
--- a/Engine/Source/Developer/DerivedDataCache/DerivedDataCache.Build.cs
+++ b/Engine/Source/Developer/DerivedDataCache/DerivedDataCache.Build.cs
@@ -11,7 +11,7 @@ public class DerivedDataCache : ModuleRules
 		PrivateDependencyModuleNames.Add("Zen");
 
 		// Dependencies for "S3" and "HTTP" backends
-		PrivateDependencyModuleNames.AddRange(new string[] { "SSL", "Json", "Zen" });
+		PrivateDependencyModuleNames.AddRange(new string[] { "SSL", "Json", "Analytics", "Zen" });
 		PrivateIncludePathModuleNames.AddRange(new string[] { "DesktopPlatform", "Zen"});
 		AddEngineThirdPartyPrivateStaticDependencies(Target, "libcurl");
 		AddEngineThirdPartyPrivateStaticDependencies(Target, "OpenSSL");
diff --git a/Engine/Source/Developer/Zen/Zen.Build.cs b/Engine/Source/Developer/Zen/Zen.Build.cs
--- a/Engine/Source/Developer/Zen/Zen.Build.cs
+++ b/Engine/Source/Developer/Zen/Zen.Build.cs
@@ -7,14 +7,8 @@ public class Zen : ModuleRules
 {
 	public Zen(ReadOnlyTargetRules Target) : base(Target)
 	{
-		PrivateDependencyModuleNames.AddRange(new string[] { "Core", "Sockets", "SSL", "Json" });
+		PrivateDependencyModuleNames.AddRange(new string[] { "Core", "Sockets", "SSL", "Json", "Analytics" });
 		PrivateIncludePathModuleNames.Add("DesktopPlatform");
-		
-		PublicDependencyModuleNames.AddRange(
-			new string[] {		
-				"Analytics",
-			}
-		);
 
 		AddEngineThirdPartyPrivateStaticDependencies(Target, "libcurl");
 		AddEngineThirdPartyPrivateStaticDependencies(Target, "OpenSSL");
