DDC: Fixed the async cache replay to be tracked as async work by the cache

#rnx

diff --git a/Engine/Source/Developer/DerivedDataCache/Private/DerivedDataCacheReplay.cpp b/Engine/Source/Developer/DerivedDataCache/Private/DerivedDataCacheReplay.cpp
--- a/Engine/Source/Developer/DerivedDataCache/Private/DerivedDataCacheReplay.cpp
+++ b/Engine/Source/Developer/DerivedDataCache/Private/DerivedDataCacheReplay.cpp
@@ -7,6 +7,7 @@
 #include "Containers/StringView.h"
 #include "DerivedDataCacheKeyFilter.h"
 #include "DerivedDataCacheMethod.h"
+#include "DerivedDataCachePrivate.h"
 #include "DerivedDataLegacyCacheStore.h"
 #include "DerivedDataRequestOwner.h"
 #include "DerivedDataValue.h"
@@ -458,10 +459,12 @@ void FCacheReplayReader::FState::ApplyPolicyTransform(FCacheGetChunkRequest& Req
 
 void FCacheReplayReader::FState::ReadFromFileAsync(const TCHAR* const ReplayPath, const uint64 ScratchSize)
 {
+	Private::AddToAsyncTaskCounter(1);
 	ReadAsyncPipe.Launch(TEXT("CacheReplayReadFromFileAsync"), [this, ReplayPath = FString(ReplayPath), ScratchSize]
 	{
 		TRACE_CPUPROFILER_EVENT_SCOPE(ReplayDDC_ReadFromFileAsync);
 		ReadFromFile(*ReplayPath, ScratchSize);
+		Private::AddToAsyncTaskCounter(-1);
 	});
 }
 
