Addeed specific Unreal Cloud DDC analytics for core.cooking and core.loading events ( Domain, Effective Domain, Cache Hit/MIss for Cache Puts & Gets, FailedOIDC login attempts )

diff --git a/Engine/Source/Developer/DerivedDataCache/Private/HttpCacheStore.cpp b/Engine/Source/Developer/DerivedDataCache/Private/HttpCacheStore.cpp
--- a/Engine/Source/Developer/DerivedDataCache/Private/HttpCacheStore.cpp
+++ b/Engine/Source/Developer/DerivedDataCache/Private/HttpCacheStore.cpp
@@ -33,6 +33,7 @@
 #include "Misc/ScopeLock.h"
 #include "Misc/ScopeRWLock.h"
 #include "Misc/StringBuilder.h"
+#include "ProfilingDebugging/CookStats.h"
 #include "ProfilingDebugging/CountersTrace.h"
 #include "ProfilingDebugging/CpuProfilerTrace.h"
 #include "Serialization/CompactBinary.h"
@@ -42,6 +43,8 @@
 #include "Serialization/JsonReader.h"
 #include "Serialization/JsonSerializer.h"
 #include "String/Find.h"
+#include "Analytics.h"
+#include "AnalyticsEventAttribute.h"
 
 #if PLATFORM_MICROSOFT
 #include "Microsoft/WindowsHWrapper.h"
@@ -314,6 +317,8 @@ private:
 	FTSTicker::FDelegateHandle RefreshAccessTokenHandle;
 	double RefreshAccessTokenTime = 0.0;
 	uint32 FailedLoginAttempts = 0;
+	uint32 LoginAttempts = 0;
+	uint32 InteractiveLoginAttempts = 0;
 
 	bool bIsUsable = false;
 	bool bReadOnly = false;
@@ -326,6 +331,7 @@ private:
 	bool EndIsServiceReady(THttpUniquePtr<IHttpResponse>& Response, TArray64<uint8>& Body);
 	bool AcquireAccessToken(IHttpClient* Client = nullptr);
 	void SetAccessToken(FStringView Token, double RefreshDelay = 0.0);
+	void OnAnalyticsEvent(TArray<FAnalyticsEventAttribute>& Attributes);
 
 	enum class EOperationCategory
 	{
@@ -1584,11 +1590,99 @@ FHttpCacheStore::FHttpCacheStore(const FHttpCacheStoreParams& Params)
 		bIsUsable = true;
 	}
 
+	// Add analytics callbacks for the Loading and Cooking events
+	FAnalytics::Get().GetEventCallback("Core.Loading")->AddRaw(this, &FHttpCacheStore::OnAnalyticsEvent);
+	FAnalytics::Get().GetEventCallback("Core.Cooking")->AddRaw(this, &FHttpCacheStore::OnAnalyticsEvent);
+
 	AnyInstance = this;
 }
 
+void FHttpCacheStore::OnAnalyticsEvent(TArray<FAnalyticsEventAttribute>& Attributes)
+{
+#if ENABLE_COOK_STATS
+
+	FString BaseName = TEXT("CloudDDC.");
+
+	{
+		FString AttrName = BaseName + TEXT(".Domain");
+		Attributes.Emplace(MoveTemp(AttrName), Domain);
+	}
+
+	{
+		FString AttrName = BaseName + TEXT(".EffectiveDomain");
+		Attributes.Emplace(MoveTemp(AttrName), EffectiveDomain.ToString() );
+	}
+
+	{
+		FString AttrName = BaseName + TEXT(".Namespace");
+		Attributes.Emplace(MoveTemp(AttrName), Namespace);
+	}
+
+	{
+		FString AttrName = BaseName + TEXT(".LoginAttempts");
+		Attributes.Emplace(MoveTemp(AttrName), LoginAttempts);
+	}
+
+	{
+		FString AttrName = BaseName + TEXT(".LoginAttempts");
+		Attributes.Emplace(MoveTemp(AttrName), InteractiveLoginAttempts );
+	}
+
+
+	{
+		FString AttrName = BaseName + TEXT(".FailedLoginAttempts");
+		Attributes.Emplace(MoveTemp(AttrName), FailedLoginAttempts);
+	}
+
+	const int64 GetHits = UsageStats.GetStats.GetAccumulatedValueAnyThread(FCookStats::CallStats::EHitOrMiss::Hit, FCookStats::CallStats::EStatType::Counter);
+	const int64 GetMisses = UsageStats.GetStats.GetAccumulatedValueAnyThread(FCookStats::CallStats::EHitOrMiss::Miss, FCookStats::CallStats::EStatType::Counter);
+	const int64 PutHits = UsageStats.PutStats.GetAccumulatedValueAnyThread(FCookStats::CallStats::EHitOrMiss::Hit, FCookStats::CallStats::EStatType::Counter);
+	const int64 PutMisses = UsageStats.PutStats.GetAccumulatedValueAnyThread(FCookStats::CallStats::EHitOrMiss::Miss, FCookStats::CallStats::EStatType::Counter);
+	const int64 TotalGets = GetHits + GetMisses;
+	const int64 TotalPuts = PutHits + PutMisses;
+
+	{
+		FString AttrName = BaseName + TEXT(".GetHits");
+		Attributes.Emplace(MoveTemp(AttrName), GetHits);
+	}
+
+	{
+		FString AttrName = BaseName + TEXT(".GetMisses");
+		Attributes.Emplace(MoveTemp(AttrName), GetMisses);
+	}
+
+	{
+		FString AttrName = BaseName + TEXT(".TotalGets");
+		Attributes.Emplace(MoveTemp(AttrName), TotalGets);
+	}
+
+	{
+		FString AttrName = BaseName + TEXT(".PutHits");
+		Attributes.Emplace(MoveTemp(AttrName), PutHits);
+	}
+
+	{
+		FString AttrName = BaseName + TEXT(".PutMisses");
+		Attributes.Emplace(MoveTemp(AttrName), TotalPuts);
+	}
+
+	{
+		FString AttrName = BaseName + TEXT(".TotalGets");
+		Attributes.Emplace(MoveTemp(AttrName), TotalPuts);
+	}
+
+#endif
+}
+
 FHttpCacheStore::~FHttpCacheStore()
 {
+	// Remove any callbacks we registered if analytics are still available
+	if (FAnalytics::IsAvailable())
+	{
+		FAnalytics::Get().GetEventCallback("Core.Loading")->RemoveAll(this);
+		FAnalytics::Get().GetEventCallback("Core.Cooking")->RemoveAll(this);
+	}
+
 	if (RefreshAccessTokenHandle.IsValid())
 	{
 		FTSTicker::GetCoreTicker().RemoveTicker(RefreshAccessTokenHandle);
@@ -1627,6 +1721,8 @@ bool FHttpCacheStore::AcquireAccessToken(IHttpClient* Client)
 		return true;
 	}
 
+	LoginAttempts++;
+
 	// Avoid spamming this if the service is down.
 	if (FailedLoginAttempts > UE_HTTPDDC_MAX_FAILED_LOGIN_ATTEMPTS)
 	{
@@ -1713,8 +1809,15 @@ bool FHttpCacheStore::AcquireAccessToken(IHttpClient* Client)
 	{
 		FString AccessTokenString;
 		FDateTime TokenExpiresAt;
-		if (FDesktopPlatformModule::Get()->GetOidcAccessToken(FPaths::RootDir(), FPaths::GetProjectFilePath(), OAuthProviderIdentifier, FApp::IsUnattended(), GWarn, AccessTokenString, TokenExpiresAt))
+		bool WasInteractiveLogin;
+
+		if (FDesktopPlatformModule::Get()->GetOidcAccessToken(FPaths::RootDir(), FPaths::GetProjectFilePath(), OAuthProviderIdentifier, FApp::IsUnattended(), GWarn, AccessTokenString, TokenExpiresAt, WasInteractiveLogin))
 		{
+			if (WasInteractiveLogin)
+			{
+				InteractiveLoginAttempts++;
+			}
+
 			const double ExpiryTimeSeconds = (TokenExpiresAt - FDateTime::UtcNow()).GetTotalSeconds();
 			UE_LOG(LogDerivedDataCache, Display,
 				TEXT("%s: OidcToken: Logged in to HTTP DDC services. Expires at %s which is in %.0f seconds."),
diff --git a/Engine/Source/Developer/DesktopPlatform/Private/DesktopPlatformBase.cpp b/Engine/Source/Developer/DesktopPlatform/Private/DesktopPlatformBase.cpp
--- a/Engine/Source/Developer/DesktopPlatform/Private/DesktopPlatformBase.cpp
+++ b/Engine/Source/Developer/DesktopPlatform/Private/DesktopPlatformBase.cpp
@@ -805,7 +805,7 @@ bool FDesktopPlatformBase::IsUnrealBuildToolRunning()
 	return bIsRunning;
 }
 
-bool FDesktopPlatformBase::GetOidcAccessToken(const FString& RootDir, const FString& ProjectFileName, const FString& ProviderIdentifier, bool Unattended, FFeedbackContext* Warn, FString& OutToken, FDateTime& OutTokenExpiresAt)
+bool FDesktopPlatformBase::GetOidcAccessToken(const FString& RootDir, const FString& ProjectFileName, const FString& ProviderIdentifier, bool Unattended, FFeedbackContext* Warn, FString& OutToken, FDateTime& OutTokenExpiresAt, bool& WasInteractiveLogin )
 {
 	FString ResultFilePath = FPaths::CreateTempFilename(*FPaths::ProjectIntermediateDir(), TEXT("oidcToken.json"));
 
@@ -822,11 +822,16 @@ bool FDesktopPlatformBase::GetOidcAccessToken(const FString& RootDir, const FStr
 	FString ProcessStdout;
 	bRes = InvokeOidcTokenToolSync(LOCTEXT("GetOidcAccessToken", "Fetching OIDC Access Token..."), RootDir, UnattendedArguments, Warn, ExitCode, ProcessStdout);
 
+	WasInteractiveLogin = false;
+
 	if (ExitCode == 10)
 	{
 		if (!Unattended)
 		{
 			bRes = GetOidcAccessTokenInteractive(RootDir, Arguments, Warn, ExitCode);
+
+			WasInteractiveLogin = true;
+
 			if (!bRes)
 			{
 				UE_LOG(LogDesktopPlatform, Error, TEXT("Unable to allocate an access token. Interactive login failed, make sure you are assigned access and are able to login in the created browser window. Provider used: '%s'. Ran OidcToken (project file is '%s', exe path is '%s')"), *ProviderIdentifier, *ProjectFileName, *GetOidcTokenExecutableFilename(RootDir));
diff --git a/Engine/Source/Developer/DesktopPlatform/Private/DesktopPlatformBase.h b/Engine/Source/Developer/DesktopPlatform/Private/DesktopPlatformBase.h
--- a/Engine/Source/Developer/DesktopPlatform/Private/DesktopPlatformBase.h
+++ b/Engine/Source/Developer/DesktopPlatform/Private/DesktopPlatformBase.h
@@ -55,7 +55,7 @@ public:
 	virtual bool RunUnrealBuildTool(const FText& Description, const FString& RootDir, const FString& Arguments, FFeedbackContext* Warn) override final;
 	virtual bool IsUnrealBuildToolRunning() override;
 
-    virtual bool GetOidcAccessToken(const FString& RootDir, const FString& ProjectFileName, const FString& ProviderIdentifier, bool Unattended, FFeedbackContext* Warn, FString& OutToken, FDateTime& OutTokenExpiresAt) override;
+    virtual bool GetOidcAccessToken(const FString& RootDir, const FString& ProjectFileName, const FString& ProviderIdentifier, bool Unattended, FFeedbackContext* Warn, FString& OutToken, FDateTime& OutTokenExpiresAt, bool& WasInteractiveLogin) override;
 	virtual bool GetOidcTokenStatus(const FString& RootDir, const FString& ProjectFileName, const FString& ProviderIdentifier, FFeedbackContext* Warn, int& OutStatus) override;
 
 	virtual const TArray<FTargetInfo>& GetTargetsForProject(const FString& ProjectFile) const override;
diff --git a/Engine/Source/Developer/DesktopPlatform/Public/IDesktopPlatform.h b/Engine/Source/Developer/DesktopPlatform/Public/IDesktopPlatform.h
--- a/Engine/Source/Developer/DesktopPlatform/Public/IDesktopPlatform.h
+++ b/Engine/Source/Developer/DesktopPlatform/Public/IDesktopPlatform.h
@@ -469,7 +469,7 @@ public:
 	* @param OutTokenExpiresAt	When the token expires
 	* @return true if the task completed successfully.
 	*/
-	virtual bool GetOidcAccessToken(const FString& RootDir, const FString& ProjectFileName, const FString& ProviderIdentifier, bool Unattended, FFeedbackContext* Warn, FString& OutToken, FDateTime& OutTokenExpiresAt) = 0;
+	virtual bool GetOidcAccessToken(const FString& RootDir, const FString& ProjectFileName, const FString& ProviderIdentifier, bool Unattended, FFeedbackContext* Warn, FString& OutToken, FDateTime& OutTokenExpiresAt, bool& WasInteractiveLogin) = 0;
 
 	/**
 	* Get the status of a access token for use by OIDC compliant Identity Providers
