DDC: Added a DeleteOnly mode to FileSystemCacheStore

This allows a cache to be kept in the hierarchy solely to ensure that it is cleaned up as its contents expire. Maintenance has been improved to delete more empty folders and to count every scanned folder.

diff --git a/Engine/Source/Developer/DerivedDataCache/Private/DerivedDataBackends.cpp b/Engine/Source/Developer/DerivedDataCache/Private/DerivedDataBackends.cpp
--- a/Engine/Source/Developer/DerivedDataCache/Private/DerivedDataBackends.cpp
+++ b/Engine/Source/Developer/DerivedDataCache/Private/DerivedDataBackends.cpp
@@ -735,7 +735,9 @@ public:
 				{
 					bUsingSharedDDC |= bShared;
 					DataCache = MakeTuple(InnerFileSystem, Flags);
-					UE_LOG(LogDerivedDataCache, Log, TEXT("Using %s data cache path %s: %s"), NodeName, *Path, !EnumHasAnyFlags(Flags, ECacheStoreFlags::Store) ? TEXT("ReadOnly") : TEXT("Writable"));
+					UE_LOG(LogDerivedDataCache, Log, TEXT("Using %s data cache path %s: %s"), NodeName, *Path,
+						EnumHasAnyFlags(Flags, ECacheStoreFlags::Store) ? TEXT("Writable") :
+						EnumHasAnyFlags(Flags, ECacheStoreFlags::Query) ? TEXT("ReadOnly") : TEXT("DeleteOnly"));
 					Directories.AddUnique(Path);
 				}
 				else
diff --git a/Engine/Source/Developer/DerivedDataCache/Private/FileSystemCacheStore.cpp b/Engine/Source/Developer/DerivedDataCache/Private/FileSystemCacheStore.cpp
--- a/Engine/Source/Developer/DerivedDataCache/Private/FileSystemCacheStore.cpp
+++ b/Engine/Source/Developer/DerivedDataCache/Private/FileSystemCacheStore.cpp
@@ -256,7 +256,8 @@ private:
 	uint32 FileCount = 0;
 	uint32 FolderCount = 0;
 	uint32 ProcessCount = 0;
-	uint32 DeleteCount = 0;
+	uint32 DeleteFileCount = 0;
+	uint32 DeleteFolderCount = 0;
 	uint64 DeleteSize = 0;
 	uint64 ScannedSize = 0;
 
@@ -361,7 +362,8 @@ void FFileSystemCacheStoreMaintainer::Loop()
 		const FDateTime ScanStart = FDateTime::Now();
 		FileCount = 0;
 		FolderCount = 0;
-		DeleteCount = 0;
+		DeleteFileCount = 0;
+		DeleteFolderCount = 0;
 		DeleteSize = 0;
 		ScannedSize = 0;
 		IdleEvent.Reset();
@@ -373,10 +375,10 @@ void FFileSystemCacheStoreMaintainer::Loop()
 		const FDateTime ScanEnd = FDateTime::Now();
 
 		UE_LOG(LogDerivedDataCache, Log,
-			TEXT("%s: Maintenance finished in %s and deleted %u files with total size %" UINT64_FMT " MiB. "
-				 "Scanned %u files in %u folders with total size %" UINT64_FMT " MiB."),
-			*CachePath, *(ScanEnd - ScanStart).ToString(), DeleteCount, DeleteSize / 1024 / 1024,
-			FileCount, FolderCount, ScannedSize / 1024 / 1024);
+			TEXT("%s: Maintenance finished in %s and deleted %u files with total size %" UINT64_FMT " MiB "
+				 "and %u empty folders. Scanned %u files in %u folders with total size %" UINT64_FMT " MiB."),
+			*CachePath, *(ScanEnd - ScanStart).ToString(), DeleteFileCount, DeleteSize / 1024 / 1024,
+			DeleteFolderCount, FileCount, FolderCount, ScannedSize / 1024 / 1024);
 
 		if (bExit || bExitAfterScan)
 		{
@@ -437,14 +439,23 @@ void FFileSystemCacheStoreMaintainer::CreateBucketRoots()
 {
 	TStringBuilder<256> BucketsPath;
 	FPathViews::Append(BucketsPath, CachePath, GBucketsDirectoryName);
-	FileManager.IterateDirectoryStat(*BucketsPath, [this](const TCHAR* Path, const FFileStatData& Stat) -> bool
+	if (FileManager.DirectoryExists(*BucketsPath))
 	{
-		if (Stat.bIsDirectory)
+		++FolderCount;
+		const int32 StartRootCount = Roots.Num();
+		FileManager.IterateDirectoryStat(*BucketsPath, [this](const TCHAR* Path, const FFileStatData& Stat) -> bool
 		{
-			Roots.Add(MakeUnique<FRoot>(Path, Random));
+			if (Stat.bIsDirectory)
+			{
+				Roots.Add(MakeUnique<FRoot>(Path, Random));
+			}
+			return !bExit;
+		});
+		if (StartRootCount == Roots.Num())
+		{
+			DeleteDirectory(*BucketsPath);
 		}
-		return !bExit;
-	});
+	}
 }
 
 void FFileSystemCacheStoreMaintainer::ScanHashRoot(const uint32 RootIndex)
@@ -498,6 +509,7 @@ void FFileSystemCacheStoreMaintainer::ScanHashRoot(const uint32 RootIndex)
 
 TStaticBitArray<256> FFileSystemCacheStoreMaintainer::ScanHashDirectory(FStringBuilderBase& BasePath)
 {
+	++FolderCount;
 	TStaticBitArray<256> Exists;
 	FileManager.IterateDirectoryStat(*BasePath, [this, &Exists](const TCHAR* Path, const FFileStatData& Stat) -> bool
 	{
@@ -521,6 +533,7 @@ TStaticBitArray<256> FFileSystemCacheStoreMaintainer::ScanHashDirectory(FStringB
 
 TStaticBitArray<10> FFileSystemCacheStoreMaintainer::ScanLegacyDirectory(FStringBuilderBase& BasePath)
 {
+	++FolderCount;
 	TStaticBitArray<10> Exists;
 	FileManager.IterateDirectoryStat(*BasePath, [this, &Exists](const TCHAR* Path, const FFileStatData& Stat) -> bool
 	{
@@ -647,7 +660,7 @@ void FFileSystemCacheStoreMaintainer::ProcessFile(const TCHAR* const Path, const
 	const FDateTime Now = FDateTime::UtcNow();
 	if (Stat.ModificationTime + Params.MaxFileAge < Now && Stat.AccessTime + Params.MaxFileAge < Now)
 	{
-		++DeleteCount;
+		++DeleteFileCount;
 		DeleteSize += Stat.FileSize > 0 ? uint64(Stat.FileSize) : 0;
 		if (FileManager.Delete(Path, /*bRequireExists*/ false, /*bEvenReadOnly*/ false, /*bQuiet*/ true))
 		{
@@ -688,6 +701,7 @@ void FFileSystemCacheStoreMaintainer::DeleteDirectory(const TCHAR* Path)
 {
 	if (FileManager.DeleteDirectory(Path))
 	{
+		++DeleteFolderCount;
 		UE_LOG(LogDerivedDataCache, VeryVerbose, TEXT("%s: Maintenance deleted empty directory %s."), *CachePath, Path);
 	}
 }
@@ -914,6 +928,8 @@ FFileSystemCacheStore::FFileSystemCacheStore(
 	, bDeactivationDeferredClean(false)
 	, DeactivateAtMS(-1.f)
 {
+	bool bDeleteOnly = false;
+
 	// If we find a platform that has more stringent limits, this needs to be rethought.
 	checkf(GMaxCacheRootLen + GMaxCacheKeyLen <= FPlatformMisc::GetMaxPathLength(),
 		TEXT("Not enough room left for cache keys in max path."));
@@ -922,6 +938,7 @@ FFileSystemCacheStore::FFileSystemCacheStore(
 	FPaths::NormalizeFilename(CachePath);
 
 	// Params that override our instance defaults
+	FParse::Bool(InParams, TEXT("DeleteOnly="), bDeleteOnly);
 	FParse::Bool(InParams, TEXT("ReadOnly="), bReadOnly);
 	FParse::Bool(InParams, TEXT("Touch="), bTouch);
 	FParse::Value(InParams, TEXT("UnusedFileAge="), DaysToDeleteUnusedFiles);
@@ -1144,8 +1161,9 @@ FFileSystemCacheStore::FFileSystemCacheStore(
 			AccessLogWriter.Reset(new FAccessLogWriter(InAccessLogPath, CachePath));
 		}
 
-		ECacheStoreFlags Flags = ECacheStoreFlags::Query;
-		Flags |= bReadOnly ? ECacheStoreFlags::None : ECacheStoreFlags::Store;
+		ECacheStoreFlags Flags = ECacheStoreFlags::None;
+		Flags |= bDeleteOnly ? ECacheStoreFlags::None : ECacheStoreFlags::Query;
+		Flags |= bDeleteOnly || bReadOnly ? ECacheStoreFlags::None : ECacheStoreFlags::Store;
 		Flags |= SpeedClass == EBackendSpeedClass::Local ? ECacheStoreFlags::Local : ECacheStoreFlags::Remote;
 		OutFlags = Flags;
 
