DDC: Fixed parameter names related to OIDC to be compliant with the coding standard and simplified analytics formatting

#rnx

diff --git a/Engine/Source/Developer/DerivedDataCache/Private/HttpCacheStore.cpp b/Engine/Source/Developer/DerivedDataCache/Private/HttpCacheStore.cpp
--- a/Engine/Source/Developer/DerivedDataCache/Private/HttpCacheStore.cpp
+++ b/Engine/Source/Developer/DerivedDataCache/Private/HttpCacheStore.cpp
@@ -1600,39 +1600,14 @@ FHttpCacheStore::FHttpCacheStore(const FHttpCacheStoreParams& Params)
 void FHttpCacheStore::OnAnalyticsEvent(TArray<FAnalyticsEventAttribute>& Attributes)
 {
 #if ENABLE_COOK_STATS
+	const FString BaseName(TEXTVIEW("CloudDDC."));
 
-	FString BaseName = TEXT("CloudDDC.");
-
-	{
-		FString AttrName = BaseName + TEXT(".Domain");
-		Attributes.Emplace(MoveTemp(AttrName), Domain);
-	}
-
-	{
-		FString AttrName = BaseName + TEXT(".EffectiveDomain");
-		Attributes.Emplace(MoveTemp(AttrName), EffectiveDomain.ToString() );
-	}
-
-	{
-		FString AttrName = BaseName + TEXT(".Namespace");
-		Attributes.Emplace(MoveTemp(AttrName), Namespace);
-	}
-
-	{
-		FString AttrName = BaseName + TEXT(".LoginAttempts");
-		Attributes.Emplace(MoveTemp(AttrName), LoginAttempts);
-	}
-
-	{
-		FString AttrName = BaseName + TEXT(".LoginAttempts");
-		Attributes.Emplace(MoveTemp(AttrName), InteractiveLoginAttempts );
-	}
-
-
-	{
-		FString AttrName = BaseName + TEXT(".FailedLoginAttempts");
-		Attributes.Emplace(MoveTemp(AttrName), FailedLoginAttempts);
-	}
+	Attributes.Emplace(FString(WriteToString<64>(BaseName, TEXTVIEW("Domain"))), Domain);
+	Attributes.Emplace(FString(WriteToString<64>(BaseName, TEXTVIEW("EffectiveDomain"))), *EffectiveDomain);
+	Attributes.Emplace(FString(WriteToString<64>(BaseName, TEXTVIEW("Namespace"))), Namespace);
+	Attributes.Emplace(FString(WriteToString<64>(BaseName, TEXTVIEW("LoginAttempts"))), LoginAttempts);
+	Attributes.Emplace(FString(WriteToString<64>(BaseName, TEXTVIEW("InteractiveLoginAttempts"))), InteractiveLoginAttempts);
+	Attributes.Emplace(FString(WriteToString<64>(BaseName, TEXTVIEW("FailedLoginAttempts"))), FailedLoginAttempts);
 
 	const int64 GetHits = UsageStats.GetStats.GetAccumulatedValueAnyThread(FCookStats::CallStats::EHitOrMiss::Hit, FCookStats::CallStats::EStatType::Counter);
 	const int64 GetMisses = UsageStats.GetStats.GetAccumulatedValueAnyThread(FCookStats::CallStats::EHitOrMiss::Miss, FCookStats::CallStats::EStatType::Counter);
@@ -1641,36 +1616,12 @@ void FHttpCacheStore::OnAnalyticsEvent(TArray<FAnalyticsEventAttribute>& Attribu
 	const int64 TotalGets = GetHits + GetMisses;
 	const int64 TotalPuts = PutHits + PutMisses;
 
-	{
-		FString AttrName = BaseName + TEXT(".GetHits");
-		Attributes.Emplace(MoveTemp(AttrName), GetHits);
-	}
-
-	{
-		FString AttrName = BaseName + TEXT(".GetMisses");
-		Attributes.Emplace(MoveTemp(AttrName), GetMisses);
-	}
-
-	{
-		FString AttrName = BaseName + TEXT(".TotalGets");
-		Attributes.Emplace(MoveTemp(AttrName), TotalGets);
-	}
-
-	{
-		FString AttrName = BaseName + TEXT(".PutHits");
-		Attributes.Emplace(MoveTemp(AttrName), PutHits);
-	}
-
-	{
-		FString AttrName = BaseName + TEXT(".PutMisses");
-		Attributes.Emplace(MoveTemp(AttrName), TotalPuts);
-	}
-
-	{
-		FString AttrName = BaseName + TEXT(".TotalGets");
-		Attributes.Emplace(MoveTemp(AttrName), TotalPuts);
-	}
-
+	Attributes.Emplace(FString(WriteToString<64>(BaseName, TEXTVIEW("GetHits"))), GetHits);
+	Attributes.Emplace(FString(WriteToString<64>(BaseName, TEXTVIEW("GetMisses"))), GetMisses);
+	Attributes.Emplace(FString(WriteToString<64>(BaseName, TEXTVIEW("TotalGets"))), TotalGets);
+	Attributes.Emplace(FString(WriteToString<64>(BaseName, TEXTVIEW("PutHits"))), PutHits);
+	Attributes.Emplace(FString(WriteToString<64>(BaseName, TEXTVIEW("PutMisses"))), TotalPuts);
+	Attributes.Emplace(FString(WriteToString<64>(BaseName, TEXTVIEW("TotalGets"))), TotalPuts);
 #endif
 }
 
@@ -1809,11 +1760,11 @@ bool FHttpCacheStore::AcquireAccessToken(IHttpClient* Client)
 	{
 		FString AccessTokenString;
 		FDateTime TokenExpiresAt;
-		bool WasInteractiveLogin;
+		bool bWasInteractiveLogin;
 
-		if (FDesktopPlatformModule::Get()->GetOidcAccessToken(FPaths::RootDir(), FPaths::GetProjectFilePath(), OAuthProviderIdentifier, FApp::IsUnattended(), GWarn, AccessTokenString, TokenExpiresAt, WasInteractiveLogin))
+		if (FDesktopPlatformModule::Get()->GetOidcAccessToken(FPaths::RootDir(), FPaths::GetProjectFilePath(), OAuthProviderIdentifier, FApp::IsUnattended(), GWarn, AccessTokenString, TokenExpiresAt, bWasInteractiveLogin))
 		{
-			if (WasInteractiveLogin)
+			if (bWasInteractiveLogin)
 			{
 				InteractiveLoginAttempts++;
 			}
diff --git a/Engine/Source/Developer/DesktopPlatform/Private/DesktopPlatformBase.cpp b/Engine/Source/Developer/DesktopPlatform/Private/DesktopPlatformBase.cpp
--- a/Engine/Source/Developer/DesktopPlatform/Private/DesktopPlatformBase.cpp
+++ b/Engine/Source/Developer/DesktopPlatform/Private/DesktopPlatformBase.cpp
@@ -805,7 +805,7 @@ bool FDesktopPlatformBase::IsUnrealBuildToolRunning()
 	return bIsRunning;
 }
 
-bool FDesktopPlatformBase::GetOidcAccessToken(const FString& RootDir, const FString& ProjectFileName, const FString& ProviderIdentifier, bool Unattended, FFeedbackContext* Warn, FString& OutToken, FDateTime& OutTokenExpiresAt, bool& WasInteractiveLogin )
+bool FDesktopPlatformBase::GetOidcAccessToken(const FString& RootDir, const FString& ProjectFileName, const FString& ProviderIdentifier, bool bUnattended, FFeedbackContext* Warn, FString& OutToken, FDateTime& OutTokenExpiresAt, bool& bOutWasInteractiveLogin)
 {
 	FString ResultFilePath = FPaths::CreateTempFilename(*FPaths::ProjectIntermediateDir(), TEXT("oidcToken.json"));
 
@@ -822,15 +822,15 @@ bool FDesktopPlatformBase::GetOidcAccessToken(const FString& RootDir, const FStr
 	FString ProcessStdout;
 	bRes = InvokeOidcTokenToolSync(LOCTEXT("GetOidcAccessToken", "Fetching OIDC Access Token..."), RootDir, UnattendedArguments, Warn, ExitCode, ProcessStdout);
 
-	WasInteractiveLogin = false;
+	bOutWasInteractiveLogin = false;
 
 	if (ExitCode == 10)
 	{
-		if (!Unattended)
+		if (!bUnattended)
 		{
 			bRes = GetOidcAccessTokenInteractive(RootDir, Arguments, Warn, ExitCode);
 
-			WasInteractiveLogin = true;
+			bOutWasInteractiveLogin = true;
 
 			if (!bRes)
 			{
diff --git a/Engine/Source/Developer/DesktopPlatform/Private/DesktopPlatformBase.h b/Engine/Source/Developer/DesktopPlatform/Private/DesktopPlatformBase.h
--- a/Engine/Source/Developer/DesktopPlatform/Private/DesktopPlatformBase.h
+++ b/Engine/Source/Developer/DesktopPlatform/Private/DesktopPlatformBase.h
@@ -55,7 +55,7 @@ public:
 	virtual bool RunUnrealBuildTool(const FText& Description, const FString& RootDir, const FString& Arguments, FFeedbackContext* Warn) override final;
 	virtual bool IsUnrealBuildToolRunning() override;
 
-    virtual bool GetOidcAccessToken(const FString& RootDir, const FString& ProjectFileName, const FString& ProviderIdentifier, bool Unattended, FFeedbackContext* Warn, FString& OutToken, FDateTime& OutTokenExpiresAt, bool& WasInteractiveLogin) override;
+    virtual bool GetOidcAccessToken(const FString& RootDir, const FString& ProjectFileName, const FString& ProviderIdentifier, bool bUnattended, FFeedbackContext* Warn, FString& OutToken, FDateTime& OutTokenExpiresAt, bool& bOutWasInteractiveLogin) override;
 	virtual bool GetOidcTokenStatus(const FString& RootDir, const FString& ProjectFileName, const FString& ProviderIdentifier, FFeedbackContext* Warn, int& OutStatus) override;
 
 	virtual const TArray<FTargetInfo>& GetTargetsForProject(const FString& ProjectFile) const override;
diff --git a/Engine/Source/Developer/DesktopPlatform/Public/IDesktopPlatform.h b/Engine/Source/Developer/DesktopPlatform/Public/IDesktopPlatform.h
--- a/Engine/Source/Developer/DesktopPlatform/Public/IDesktopPlatform.h
+++ b/Engine/Source/Developer/DesktopPlatform/Public/IDesktopPlatform.h
@@ -463,13 +463,14 @@ public:
 	* @param RootDir			Engine root directory for the project to use.
 	* @param ProjectFileName	Filename of the current project
 	* @param ProviderIdentifier	The identifier for the provider as configured in oidc-token.json
-	* @param Unattended			True to indicate that no user interaction should be assumed
+	* @param bUnattended		True to indicate that no user interaction should be assumed
 	* @param Warn				Feedback context to use for progress updates
 	* @param OutToken			The allocated access token
 	* @param OutTokenExpiresAt	When the token expires
+	* @param bOutWasInteractiveLogin	True if the interactive login flow was used
 	* @return true if the task completed successfully.
 	*/
-	virtual bool GetOidcAccessToken(const FString& RootDir, const FString& ProjectFileName, const FString& ProviderIdentifier, bool Unattended, FFeedbackContext* Warn, FString& OutToken, FDateTime& OutTokenExpiresAt, bool& WasInteractiveLogin) = 0;
+	virtual bool GetOidcAccessToken(const FString& RootDir, const FString& ProjectFileName, const FString& ProviderIdentifier, bool bUnattended, FFeedbackContext* Warn, FString& OutToken, FDateTime& OutTokenExpiresAt, bool& bOutWasInteractiveLogin) = 0;
 
 	/**
 	* Get the status of a access token for use by OIDC compliant Identity Providers
