DDC: Fixed invalid put stats in the Jupiter client

#rnx

diff --git a/Engine/Source/Developer/DerivedDataCache/Private/HttpCacheStore.cpp b/Engine/Source/Developer/DerivedDataCache/Private/HttpCacheStore.cpp
--- a/Engine/Source/Developer/DerivedDataCache/Private/HttpCacheStore.cpp
+++ b/Engine/Source/Developer/DerivedDataCache/Private/HttpCacheStore.cpp
@@ -2358,7 +2358,7 @@ void FHttpCacheStore::Get(
 	for (const FCacheGetRequest& Request : Requests)
 	{
 		GetCacheRecordAsync(Owner, Request.Name, Request.Key, Request.Policy, Request.UserData,
-			[COOK_STAT(Timer = UsageStats.TimePut(), ) SharedOnComplete](FCacheGetResponse&& Response, uint64 BytesReceived) mutable
+			[COOK_STAT(Timer = UsageStats.TimeGet(), ) SharedOnComplete](FCacheGetResponse&& Response, uint64 BytesReceived) mutable
 		{
 			TRACE_COUNTER_ADD(HttpDDC_BytesReceived, BytesReceived);
 			if (Response.Status == EStatus::Ok)
