DDC: Added experimental interface IBuildSchedulerThreadPoolProvider to integrate the build scheduler with the asset compiling managers

diff --git a/Engine/Source/Developer/DerivedDataCache/Private/DerivedDataBuildJob.cpp b/Engine/Source/Developer/DerivedDataCache/Private/DerivedDataBuildJob.cpp
--- a/Engine/Source/Developer/DerivedDataCache/Private/DerivedDataBuildJob.cpp
+++ b/Engine/Source/Developer/DerivedDataCache/Private/DerivedDataBuildJob.cpp
@@ -441,6 +441,7 @@ void FBuildJob::CreateContext()
 
 	// Populate the scheduler params with the information that is available now.
 	FBuildSchedulerParams& SchedulerParams = Schedule->EditParameters();
+	SchedulerParams.TypeName = Context->GetTypeName();
 	SchedulerParams.Key = Action.Get().GetKey();
 	Action.Get().IterateConstants([&SchedulerParams](FUtf8StringView Key, FCbObject&& Value)
 	{
diff --git a/Engine/Source/Developer/DerivedDataCache/Private/DerivedDataBuildJobContext.cpp b/Engine/Source/Developer/DerivedDataCache/Private/DerivedDataBuildJobContext.cpp
--- a/Engine/Source/Developer/DerivedDataCache/Private/DerivedDataBuildJobContext.cpp
+++ b/Engine/Source/Developer/DerivedDataCache/Private/DerivedDataBuildJobContext.cpp
@@ -24,6 +24,7 @@ FBuildJobContext::FBuildJobContext(
 	FBuildOutputBuilder& InOutputBuilder)
 	: Job(InJob)
 	, CacheKey(InCacheKey)
+	, TypeName(InFunction.GetName())
 	, Function(InFunction)
 	, OutputBuilder(InOutputBuilder)
 	, CachePolicyMask(~ECachePolicy::None)
@@ -161,6 +162,13 @@ void FBuildJobContext::EndAsyncBuild()
 	Owner->End(this, [this] { EndBuild(); });
 }
 
+void FBuildJobContext::SetTypeName(const FUtf8SharedString& InTypeName)
+{
+	checkf(!TypeName.IsEmpty(), TEXT("Empty type name not allowed for build of '%s' by %s."),
+		*Job.GetName(), *WriteToString<32>(Job.GetFunction()));
+	TypeName = InTypeName;
+}
+
 void FBuildJobContext::SetCacheBucket(FCacheBucket Bucket)
 {
 	checkf(!Bucket.IsNull(), TEXT("Null cache bucket not allowed for build of '%s' by %s. "
diff --git a/Engine/Source/Developer/DerivedDataCache/Private/DerivedDataBuildJobContext.h b/Engine/Source/Developer/DerivedDataCache/Private/DerivedDataBuildJobContext.h
--- a/Engine/Source/Developer/DerivedDataCache/Private/DerivedDataBuildJobContext.h
+++ b/Engine/Source/Developer/DerivedDataCache/Private/DerivedDataBuildJobContext.h
@@ -37,6 +37,8 @@ public:
 
 	const FSharedString& GetName() const final;
 
+	inline const FUtf8SharedString& GetTypeName() const { return TypeName; }
+
 	inline const FCacheKey& GetCacheKey() const { return CacheKey; }
 
 	inline ECachePolicy GetCachePolicyMask() const final { return CachePolicyMask; }
@@ -67,6 +69,7 @@ private:
 	void BeginAsyncBuild() final;
 	void EndAsyncBuild() final;
 
+	void SetTypeName(const FUtf8SharedString& TypeName) final;
 	void SetCacheBucket(FCacheBucket Bucket) final;
 	void SetCachePolicyMask(ECachePolicy Policy) final;
 	void SetBuildPolicyMask(EBuildPolicy Policy) final;
@@ -80,6 +83,7 @@ private:
 private:
 	IBuildJob& Job;
 	FCacheKey CacheKey;
+	FUtf8SharedString TypeName;
 	const IBuildFunction& Function;
 	FBuildOutputBuilder& OutputBuilder;
 	TMap<FUtf8SharedString, FCbObject> Constants;
diff --git a/Engine/Source/Developer/DerivedDataCache/Private/DerivedDataBuildScheduler.cpp b/Engine/Source/Developer/DerivedDataCache/Private/DerivedDataBuildScheduler.cpp
--- a/Engine/Source/Developer/DerivedDataCache/Private/DerivedDataBuildScheduler.cpp
+++ b/Engine/Source/Developer/DerivedDataCache/Private/DerivedDataBuildScheduler.cpp
@@ -1,193 +1,232 @@
 // Copyright Epic Games, Inc. All Rights Reserved.
 
 #include "DerivedDataBuildScheduler.h"
-#include "Containers/RingBuffer.h"
 #include "DerivedDataBuildFunction.h"
 #include "DerivedDataBuildFunctionFactory.h"
 #include "DerivedDataBuildFunctionRegistry.h"
+#include "DerivedDataBuildPrivate.h"
 #include "DerivedDataRequest.h"
 #include "DerivedDataRequestOwner.h"
 #include "Experimental/Async/LazyEvent.h"
-#include "Misc/Guid.h"
-#include "Misc/ScopeLock.h"
+#include "Experimental/DerivedDataBuildSchedulerThreadPoolProvider.h"
+#include "Experimental/Misc/ExecutionResource.h"
+#include "HAL/CriticalSection.h"
+#include "Misc/IQueuedWork.h"
+#include "Misc/QueuedThreadPool.h"
+#include "Misc/ScopeRWLock.h"
 
 namespace UE::DerivedData::Private
 {
 
 ///////////////////////////////////////////////////////////////////////////////////////////////////////////////////////
 
-static void ScheduleAsyncStep(IBuildJob& Job, IRequestOwner& Owner, const TCHAR* DebugName)
+class FBuildSchedulerMemoryQueue
 {
-	Owner.LaunchTask(DebugName, [&Job] { Job.StepExecution(); });
-}
+public:
+	FBuildSchedulerMemoryQueue();
+	~FBuildSchedulerMemoryQueue();
+
+	void WaitForMemory(
+		const FUtf8SharedString& TypeName,
+		uint64 MemoryEstimate,
+		IRequestOwner& Owner,
+		TUniqueFunction<void ()>&& OnComplete);
+
+private:
+	void OnModularFeatureRegistered(const FName& Type, IModularFeature* ModularFeature);
+	void OnModularFeatureUnregistered(const FName& Type, IModularFeature* ModularFeature);
+
+	void AddProviderNoLock(IBuildSchedulerThreadPoolProvider* Provider);
+	void RemoveProvider(IBuildSchedulerThreadPoolProvider* Provider);
+
+	FQueuedThreadPool* FindThreadPool(const FUtf8SharedString& TypeName) const;
+
+	class FRequest;
+
+private:
+	mutable FRWLock Lock;
+	TMap<FUtf8SharedString, IBuildSchedulerThreadPoolProvider*> Providers;
+};
 
 ///////////////////////////////////////////////////////////////////////////////////////////////////////////////////////
 
-/** Limits simultaneous build jobs to reduce peak memory usage */
-class FMemoryScheduler
+class FBuildSchedulerMemoryQueue::FRequest final : public FRequestBase, public IQueuedWork
 {
 public:
-	FMemoryScheduler();
-	~FMemoryScheduler();
+	FRequest(
+		uint64 MemoryEstimate,
+		IRequestOwner& Owner,
+		FQueuedThreadPool& ThreadPool,
+		TUniqueFunction<void ()>&& OnComplete);
+	~FRequest();
+
+private:
+	TRefCountPtr<IRequest> TryEnd();
+	void Execute();
 
-	void RegisterRunningJob(uint64 MemoryEstimate);
-	void StepAsyncOrQueue(uint64 MemoryEstimate, IBuildJob& Job, IRequestOwner& Owner, const TCHAR* DebugName);
-	void RegisterEndedJob(uint64 MemoryEstimate);
+	void SetPriority(EPriority Priority) final;
+	void Wait() final;
+	void Cancel() final { Wait(); }
+
+	void DoThreadedWork() final { Execute(); }
+	void Abandon() final { Execute(); }
+	int64 GetRequiredMemory() const final { return int64(MemoryEstimate); }
 
 private:
-	/** Handles waiting and cancellation while a job is queued up */
-	class FRequest final : public FRequestBase
-	{
-	public:
-		FRequest(FMemoryScheduler& InScheduler, IBuildJob& InJob, IRequestOwner& InOwner, uint64 InMemoryEstimate);
-		~FRequest() { ensure(!TryClaimEnd()); }
-
-		void SetPriority(EPriority Priority) final {}
-		void Wait() final { Event.Wait(); }
-		void Cancel() final;
-
-		uint64 GetMemoryEstimate() const { return MemoryEstimate; }
-		bool TryClaimEnd() { return !bClaimed.test_and_set(); }
-		void End(const TCHAR* DebugName);
-
-	private:
-		FMemoryScheduler& Scheduler;
-		IBuildJob& Job;
-		IRequestOwner& Owner;
-		const uint64 MemoryEstimate;
-		UE::FLazyEvent Event{EEventMode::ManualReset};
-		std::atomic_flag bClaimed = ATOMIC_FLAG_INIT;
-	};
-
-	const uint64 TotalPhysical;
-	const uint64 AvailablePhysicalAtStartup;
-	const uint64 MaxMemoryUsage;
-
-	FCriticalSection CriticalSection;
-	TRingBuffer<TRefCountPtr<FRequest>> Queue;
-	uint64 TotalScheduledMemory = 0;
-	uint64 TotalScheduledWatermark = 0;
-
-	// @pre CriticalSection locked
-	bool CanRunNow(uint64 MemoryEstimate) const
-	{
-		return TotalScheduledMemory == 0 || TotalScheduledMemory + MemoryEstimate < MaxMemoryUsage;
-	}
+	uint64 MemoryEstimate;
+	IRequestOwner& Owner;
+	FQueuedThreadPool& ThreadPool;
+	TUniqueFunction<void ()> OnComplete;
+	UE::FLazyEvent Event{EEventMode::ManualReset};
+	std::atomic<bool> bClaimed = false;
 };
 
-FMemoryScheduler::FRequest::FRequest(
-	FMemoryScheduler& InScheduler,
-	IBuildJob& InJob,
+FBuildSchedulerMemoryQueue::FRequest::FRequest(
+	const uint64 InMemoryEstimate,
 	IRequestOwner& InOwner,
-	uint64 InMemoryEstimate)
-	: Scheduler(InScheduler)
-	, Job(InJob)
+	FQueuedThreadPool& InThreadPool,
+	TUniqueFunction<void ()>&& InOnComplete)
+	: MemoryEstimate(InMemoryEstimate)
 	, Owner(InOwner)
-	, MemoryEstimate(InMemoryEstimate)
+	, ThreadPool(InThreadPool)
+	, OnComplete(MoveTemp(InOnComplete))
 {
+	AddRef(); // Released in Execute() or Cancel()
 	Owner.Begin(this);
+	ThreadPool.AddQueuedWork(this, ConvertToQueuedWorkPriority(Owner.GetPriority()));
+}
+
+FBuildSchedulerMemoryQueue::FRequest::~FRequest()
+{
+	check(bClaimed.load(std::memory_order_relaxed));
 }
 
-void FMemoryScheduler::FRequest::End(const TCHAR* DebugName)
+TRefCountPtr<IRequest> FBuildSchedulerMemoryQueue::FRequest::TryEnd()
 {
-	Owner.End(this, [this, DebugName]
-	{ 
-		ScheduleAsyncStep(Job, Owner, DebugName);
+	return bClaimed.exchange(true) ? nullptr : Owner.End(this, [this]
+	{
+		Invoke(OnComplete);
 		Event.Trigger();
 	});
 }
 
-void FMemoryScheduler::FRequest::Cancel()
+void FBuildSchedulerMemoryQueue::FRequest::Execute()
+{
+	TryEnd();
+	Release();
+}
+
+void FBuildSchedulerMemoryQueue::FRequest::SetPriority(EPriority Priority)
 {
-	if (TryClaimEnd())
+	if (ThreadPool.RetractQueuedWork(this))
 	{
-		// Add estimated memory to simplify implementation, even though memory won't be allocated.
-		// FBuildJobSchedule::EndJob() will restore the scheduler's available memory.
-		// Might require optimization if lots of queued jobs are cancelled at the same time.
-		Scheduler.RegisterRunningJob(MemoryEstimate);
-		End(TEXT("MemoryQueueCancel"));
+		ThreadPool.AddQueuedWork(this, ConvertToQueuedWorkPriority(Priority));
 	}
 }
 
-FMemoryScheduler::FMemoryScheduler()
-	: TotalPhysical(FPlatformMemory::GetStats().TotalPhysical)
-	, AvailablePhysicalAtStartup(FPlatformMemory::GetStats().AvailablePhysical)
-	, MaxMemoryUsage(TotalPhysical / 8 + AvailablePhysicalAtStartup / 2)
+void FBuildSchedulerMemoryQueue::FRequest::Wait()
 {
-	Queue.Reserve(128);
+	if (TRefCountPtr<IRequest> Self = TryEnd())
+	{
+		if (ThreadPool.RetractQueuedWork(this))
+		{
+			Release();
+		}
+	}
+	else
+	{
+		Event.Wait();
+	}
 }
 
-FMemoryScheduler::~FMemoryScheduler()
+///////////////////////////////////////////////////////////////////////////////////////////////////////////////////////
+
+FBuildSchedulerMemoryQueue::FBuildSchedulerMemoryQueue()
 {
-	ensure(Queue.IsEmpty());
-	ensure(TotalScheduledMemory == 0);
+	IModularFeatures& ModularFeatures = IModularFeatures::Get();
+	for (IBuildSchedulerThreadPoolProvider* Provider : ModularFeatures.GetModularFeatureImplementations<IBuildSchedulerThreadPoolProvider>(IBuildSchedulerThreadPoolProvider::FeatureName))
+	{
+		AddProviderNoLock(Provider);
+	}
+	ModularFeatures.OnModularFeatureRegistered().AddRaw(this, &FBuildSchedulerMemoryQueue::OnModularFeatureRegistered);
+	ModularFeatures.OnModularFeatureUnregistered().AddRaw(this, &FBuildSchedulerMemoryQueue::OnModularFeatureUnregistered);
 }
 
-void FMemoryScheduler::RegisterRunningJob(uint64 MemoryEstimate)
+FBuildSchedulerMemoryQueue::~FBuildSchedulerMemoryQueue()
 {
-	check(MemoryEstimate);
-
-	FScopeLock Lock(&CriticalSection);
-	TotalScheduledMemory += MemoryEstimate;
-	TotalScheduledWatermark = FMath::Max(TotalScheduledWatermark, TotalScheduledMemory);
+	IModularFeatures& ModularFeatures = IModularFeatures::Get();
+	ModularFeatures.OnModularFeatureUnregistered().RemoveAll(this);
+	ModularFeatures.OnModularFeatureRegistered().RemoveAll(this);
 }
 
-void FMemoryScheduler::StepAsyncOrQueue(
-	uint64 MemoryEstimate,
-	IBuildJob& Job,
-	IRequestOwner& Owner,
-	const TCHAR* DebugName)
+void FBuildSchedulerMemoryQueue::OnModularFeatureRegistered(const FName& Type, IModularFeature* ModularFeature)
 {
-	check(MemoryEstimate);
+	if (Type == IBuildSchedulerThreadPoolProvider::FeatureName)
 	{
-		FScopeLock Lock(&CriticalSection);
-
-		if (!CanRunNow(MemoryEstimate))
-		{
-			Queue.Emplace(new FRequest(*this, Job, Owner, MemoryEstimate));
-			return;
-		}
-
-		TotalScheduledMemory += MemoryEstimate;
-		TotalScheduledWatermark = FMath::Max(TotalScheduledWatermark, TotalScheduledMemory);
+		FWriteScopeLock WriteLock(Lock);
+		AddProviderNoLock(static_cast<IBuildSchedulerThreadPoolProvider*>(ModularFeature));
 	}
+}
 
-	ScheduleAsyncStep(Job, Owner, DebugName);
+void FBuildSchedulerMemoryQueue::OnModularFeatureUnregistered(const FName& Type, IModularFeature* ModularFeature)
+{
+	if (Type == IBuildSchedulerThreadPoolProvider::FeatureName)
+	{
+		RemoveProvider(static_cast<IBuildSchedulerThreadPoolProvider*>(ModularFeature));
+	}
 }
 
-void FMemoryScheduler::RegisterEndedJob(uint64 DoneEstimate)
+void FBuildSchedulerMemoryQueue::AddProviderNoLock(IBuildSchedulerThreadPoolProvider* Provider)
 {
-	if (DoneEstimate)
+	
+	const FUtf8SharedString& TypeName = Provider->GetTypeName();
+	const uint32 TypeNameHash = GetTypeHash(TypeName);
+	if (TypeName.IsEmpty())
 	{
-		TArray<TRefCountPtr<FRequest>, TInlineAllocator<16>> Continuations;
+		UE_LOG(LogDerivedDataBuild, Error,
+			TEXT("An empty type name is not allowed in a build scheduler thread pool provider."));
+	}
+	else if (Providers.FindByHash(TypeNameHash, TypeName))
+	{
+		UE_LOG(LogDerivedDataBuild, Error,
+			TEXT("More than one build scheduler thread pool provider has been registered with the type name %s."),
+			*WriteToString<64>(TypeName));
+	}
+	else
+	{
+		Providers.EmplaceByHash(TypeNameHash, TypeName, Provider);
+	}
+}
 
-		{
-			FScopeLock Lock(&CriticalSection);
-
-			TotalScheduledMemory -= DoneEstimate;
-
-			if (Queue.IsEmpty())
-			{
-				return;
-			}
-
-			while (Queue.Num() && CanRunNow(Queue.First()->GetMemoryEstimate()))
-			{
-				if (Queue.First()->TryClaimEnd())
-				{
-					TotalScheduledMemory += Queue.First()->GetMemoryEstimate();
-					Continuations.Add(Queue.First());
-				}
-				Queue.PopFront();
-			}
-
-			TotalScheduledWatermark = FMath::Max(TotalScheduledWatermark, TotalScheduledMemory);
-		}
+void FBuildSchedulerMemoryQueue::RemoveProvider(IBuildSchedulerThreadPoolProvider* Provider)
+{
+	const FUtf8SharedString& TypeName = Provider->GetTypeName();
+	const uint32 TypeNameHash = GetTypeHash(TypeName);
+	FWriteScopeLock WriteLock(Lock);
+	Providers.RemoveByHash(TypeNameHash, TypeName);
+}
 
-		for (const TRefCountPtr<FRequest>& Request : Continuations)
-		{
-			Request->End(TEXT("MemoryQueueContinue"));
-		}
+FQueuedThreadPool* FBuildSchedulerMemoryQueue::FindThreadPool(const FUtf8SharedString& TypeName) const
+{
+	const uint32 TypeNameHash = GetTypeHash(TypeName);
+	FReadScopeLock ReadLock(Lock);
+	IBuildSchedulerThreadPoolProvider* const* Provider = Providers.FindByHash(TypeNameHash, TypeName);
+	return Provider ? (*Provider)->GetThreadPool() : nullptr;
+}
+
+void FBuildSchedulerMemoryQueue::WaitForMemory(
+	const FUtf8SharedString& TypeName,
+	const uint64 MemoryEstimate,
+	IRequestOwner& Owner,
+	TUniqueFunction<void ()>&& OnComplete)
+{
+	if (FQueuedThreadPool* ThreadPool = FindThreadPool(TypeName))
+	{
+		new FRequest(MemoryEstimate, Owner, *ThreadPool, MoveTemp(OnComplete));
+	}
+	else
+	{
+		Invoke(OnComplete);
 	}
 }
 
@@ -196,10 +235,10 @@ void FMemoryScheduler::RegisterEndedJob(uint64 DoneEstimate)
 class FBuildJobSchedule final : public IBuildJobSchedule
 {
 public:
-	FBuildJobSchedule(IBuildJob& InJob, IRequestOwner& InOwner, FMemoryScheduler& InMemoryLimiter)
+	FBuildJobSchedule(IBuildJob& InJob, IRequestOwner& InOwner, FBuildSchedulerMemoryQueue& InMemoryQueue)
 		: Job(InJob)
 		, Owner(InOwner)
-		, MemoryLimiter(InMemoryLimiter)
+		, MemoryQueue(InMemoryQueue)
 	{
 	}
 
@@ -209,11 +248,24 @@ public:
 	void ScheduleCacheStore() final       { StepSync(); }
 	void ScheduleResolveKey() final       { StepAsync(TEXT("ResolveKey")); }
 	void ScheduleResolveInputMeta() final { StepAsync(TEXT("ResolveInputMeta")); }
-	void ScheduleResolveInputData() final { StepAsyncOrQueue(TEXT("ResolveInputData")); }
-	void ScheduleExecuteRemote() final    { StepAsyncOrQueue(TEXT("ExecuteRemote")); }
+	void ScheduleResolveInputData() final
+	{
+		if (Params.MissingRemoteInputsSize)
+		{
+			StepAsync(TEXT("ResolveInputData"));
+		}
+		else
+		{
+			StepAsyncOrQueue(TEXT("ResolveInputData"));
+		}
+	}
+	void ScheduleExecuteRemote() final    { StepAsync(TEXT("ExecuteRemote")); }
 	void ScheduleExecuteLocal() final     { StepAsyncOrQueue(TEXT("ExecuteLocal")); }
 
-	void EndJob() final { MemoryLimiter.RegisterEndedJob(ScheduledMemoryEstimate); }
+	void EndJob() final
+	{
+		ExecutionResources = nullptr;
+	}
 
 private:
 	void StepSync()
@@ -229,7 +281,7 @@ private:
 		}
 		else
 		{
-			ScheduleAsyncStep(Job, Owner, DebugName);
+			Owner.LaunchTask(DebugName, [this] { Job.StepExecution(); });
 		}
 	}
 
@@ -238,33 +290,30 @@ private:
 		check(Params.TotalRequiredMemory >= Params.ResolvedInputsSize);
 		const uint64 CurrentMemoryEstimate = Params.TotalRequiredMemory - Params.ResolvedInputsSize;
 
-		// Only queue for memory once
-		if (ScheduledMemoryEstimate || CurrentMemoryEstimate == 0)
+		// Queue for memory only once, the first time it is needed for local execution.
+		// This will occur either when resolving input data for local execution or before beginning local execution.
+		// NOTE: No attempt is made to reserve memory prior to remote execution, which may become a problem if remote
+		//       execution frequently requires input data to be loaded.
+		if (QueuedMemoryEstimate || CurrentMemoryEstimate == 0)
 		{
-			StepAsync(DebugName);
+			return StepAsync(DebugName);
 		}
-		else
+
+		QueuedMemoryEstimate = CurrentMemoryEstimate;
+		MemoryQueue.WaitForMemory(Params.TypeName, CurrentMemoryEstimate, Owner, [this, DebugName]
 		{
-			ScheduledMemoryEstimate = CurrentMemoryEstimate;
-
-			if (Owner.GetPriority() == EPriority::Blocking)
-			{
-				MemoryLimiter.RegisterRunningJob(ScheduledMemoryEstimate);
-				StepSync();
-			}
-			else
-			{
-				MemoryLimiter.StepAsyncOrQueue(ScheduledMemoryEstimate, Job, Owner, DebugName);
-			}
-		}
+			ExecutionResources = FExecutionResourceContext::Get();
+			StepAsync(DebugName);
+		});
 	}
 
 private:
 	IBuildJob& Job;
 	IRequestOwner& Owner;
 	FBuildSchedulerParams Params;
-	FMemoryScheduler& MemoryLimiter;
-	uint64 ScheduledMemoryEstimate = 0;
+	FBuildSchedulerMemoryQueue& MemoryQueue;
+	TRefCountPtr<IExecutionResource> ExecutionResources;
+	uint64 QueuedMemoryEstimate = 0;
 };
 
 ///////////////////////////////////////////////////////////////////////////////////////////////////////////////////////
@@ -273,10 +322,10 @@ class FBuildScheduler final : public IBuildScheduler
 {
 	TUniquePtr<IBuildJobSchedule> BeginJob(IBuildJob& Job, IRequestOwner& Owner) final
 	{
-		return MakeUnique<FBuildJobSchedule>(Job, Owner, MemoryLimiter);
+		return MakeUnique<FBuildJobSchedule>(Job, Owner, MemoryQueue);
 	}
 
-	FMemoryScheduler MemoryLimiter;
+	FBuildSchedulerMemoryQueue MemoryQueue;
 };
 
 IBuildScheduler* CreateBuildScheduler()
diff --git a/Engine/Source/Developer/DerivedDataCache/Public/DerivedDataBuildFunction.h b/Engine/Source/Developer/DerivedDataCache/Public/DerivedDataBuildFunction.h
--- a/Engine/Source/Developer/DerivedDataCache/Public/DerivedDataBuildFunction.h
+++ b/Engine/Source/Developer/DerivedDataCache/Public/DerivedDataBuildFunction.h
@@ -153,6 +153,9 @@ public:
 	/** Adds the message to the build output. Must be deterministic and based only on inputs. */
 	virtual void AddMessage(FStringView Message) = 0;
 
+	/** Overrides the name of the type for the build. Defaults to the function name. */
+	virtual void SetTypeName(const FUtf8SharedString& TypeName) = 0;
+
 	/** Overrides the cache bucket used when reading or writing this build in the cache. */
 	virtual void SetCacheBucket(FCacheBucket Bucket) = 0;
 
diff --git a/Engine/Source/Developer/DerivedDataCache/Public/DerivedDataBuildScheduler.h b/Engine/Source/Developer/DerivedDataCache/Public/DerivedDataBuildScheduler.h
--- a/Engine/Source/Developer/DerivedDataCache/Public/DerivedDataBuildScheduler.h
+++ b/Engine/Source/Developer/DerivedDataCache/Public/DerivedDataBuildScheduler.h
@@ -5,6 +5,7 @@
 #include "DerivedDataBuildJob.h"
 #include "DerivedDataBuildKey.h"
 #include "DerivedDataRequest.h"
+#include "DerivedDataSharedString.h"
 #include "HAL/Platform.h"
 #include "Templates/RefCounting.h"
 #include "Templates/UniquePtr.h"
@@ -36,6 +37,10 @@ public:
 /** Parameters that describe a build job to the build scheduler. */
 struct FBuildSchedulerParams
 {
+	/** Type name from the build function for the job. Available from the CacheQuery state onward. */
+	FUtf8SharedString TypeName;
+
+	/** Available from the CacheQuery state onward. */
 	FBuildActionKey Key;
 
 	/** Total size of constants and inputs, whether resolved or not. */
diff --git a/Engine/Source/Developer/DerivedDataCache/Public/Experimental/DerivedDataBuildSchedulerThreadPoolProvider.h b/Engine/Source/Developer/DerivedDataCache/Public/Experimental/DerivedDataBuildSchedulerThreadPoolProvider.h
new file mode 100644
--- /dev/null
+++ b/Engine/Source/Developer/DerivedDataCache/Public/Experimental/DerivedDataBuildSchedulerThreadPoolProvider.h
@@ -0,0 +1,33 @@
+// Copyright Epic Games, Inc. All Rights Reserved.
+
+#pragma once
+
+#include "DerivedDataSharedStringFwd.h"
+#include "Features/IModularFeature.h"
+#include "UObject/NameTypes.h"
+
+#define UE_API DERIVEDDATACACHE_API
+
+class FQueuedThreadPool;
+
+namespace UE::DerivedData
+{
+
+/** Provides the thread pool to use when reserving memory, and matches build functions using their type name. */
+class IBuildSchedulerThreadPoolProvider : public IModularFeature
+{
+public:
+	static inline const FLazyName FeatureName{"BuildSchedulerThreadPoolProvider"};
+
+	virtual ~IBuildSchedulerThreadPoolProvider() = default;
+
+	/** Returns the type name that this provider corresponds to. */
+	virtual const FUtf8SharedString& GetTypeName() const = 0;
+
+	/** Returns the thread pool to use when reserving memory for the build. */
+	virtual FQueuedThreadPool* GetThreadPool() const = 0;
+};
+
+} // UE::DerivedData
+
+#undef UE_API
