Print more detailed timing metrics for HTTP operations in the HttpCacheStore.

diff --git a/Engine/Source/Developer/DerivedDataCache/Private/HttpCacheStore.cpp b/Engine/Source/Developer/DerivedDataCache/Private/HttpCacheStore.cpp
--- a/Engine/Source/Developer/DerivedDataCache/Private/HttpCacheStore.cpp
+++ b/Engine/Source/Developer/DerivedDataCache/Private/HttpCacheStore.cpp
@@ -538,7 +538,7 @@ private:
 				{
 					StatsText << TEXTVIEW("received ") << Stats.RecvSize << TEXTVIEW(" bytes, ");
 				}
-				StatsText.Appendf(TEXT("%.3f seconds"), Stats.TotalTime);
+				StatsText.Appendf(TEXT("%.3f seconds %.3f|%.3f|%.3f|%.3f"), Stats.TotalTime, Stats.NameResolveTime, Stats.ConnectTime, Stats.TlsConnectTime, Stats.StartTransferTime);
 			}
 
 			if (bVerbose)
