DDC: Invoke the build completion callback after storing to the cache

This avoids a race condition where a build completes and then another process is unable to retrieve the data from the cache.

#rnx

diff --git a/Engine/Source/Developer/DerivedDataCache/Private/DerivedDataBuildJob.cpp b/Engine/Source/Developer/DerivedDataCache/Private/DerivedDataBuildJob.cpp
--- a/Engine/Source/Developer/DerivedDataCache/Private/DerivedDataBuildJob.cpp
+++ b/Engine/Source/Developer/DerivedDataCache/Private/DerivedDataBuildJob.cpp
@@ -176,10 +176,10 @@ private:
 	void SetOutputNoCheck(FBuildOutput&& Output, EBuildJobState NewState = EBuildJobState::CacheStore);
 
 	/** Terminate the job and send the error to the output complete callback. */
-	void CompleteWithError(FUtf8StringView Error);
+	void CompleteWithError(FUtf8StringView Error) { AdvanceToState(EBuildJobState::Complete, Error); }
 
 	/** Advance to the new state, dispatching to the scheduler and invoking callbacks as appropriate. */
-	void AdvanceToState(EBuildJobState NewState);
+	void AdvanceToState(EBuildJobState NewState, FUtf8StringView NewError = {});
 
 	/** Execute a transition from the old state to the new state. */
 	void ExecuteTransition(EBuildJobState OldState, EBuildJobState NewState);
@@ -242,7 +242,7 @@ private:
 	FBuildWorker* Worker{};
 	/** Worker executor to use for remote execution. Available in [ExecuteRemote, ExecuteRemoteRetryWait]. */
 	IBuildWorkerExecutor* WorkerExecutor{};
-	/** Invoked exactly once when the output is complete or when the job fails. */
+	/** Invoked exactly once when the job is complete. */
 	FOnBuildComplete OnComplete;
 
 	/** Keys for missing inputs. */
@@ -1016,48 +1016,26 @@ void FBuildJob::SetOutputNoCheck(FBuildOutput&& InOutput, EBuildJobState NewStat
 {
 	checkf(Output.IsNull(), TEXT("Job already has an output for build of '%s' by %s."), *Name, *WriteToString<32>(FunctionName));
 	Output = MoveTemp(InOutput);
-
-	if (OnComplete)
-	{
-		const FCacheKey& CacheKey = Context ? Context->GetCacheKey() : FCacheKey::Empty;
-		const EStatus Status = Owner.IsCanceled() ? EStatus::Canceled : Output.Get().HasError() ? EStatus::Error : EStatus::Ok;
-		OnComplete({CacheKey, FBuildOutput(Output.Get()), BuildStatus, Status});
-		OnComplete = nullptr;
-	}
-
 	return AdvanceToState(NewState);
 }
 
 ///////////////////////////////////////////////////////////////////////////////////////////////////////////////////////
 
-void FBuildJob::CompleteWithError(FUtf8StringView Error)
-{
-	if (FWriteScopeLock WriteLock(Lock); Output || NextState == EBuildJobState::Complete)
-	{
-		return;
-	}
-	OutputBuilder.AddLog({ANSITEXTVIEW("LogDerivedDataBuild"), Error, EBuildOutputLogLevel::Error});
-	return SetOutputNoCheck(OutputBuilder.Build(), EBuildJobState::Complete);
-}
-
-void FBuildJob::AdvanceToState(EBuildJobState NewState)
+void FBuildJob::AdvanceToState(EBuildJobState NewState, FUtf8StringView NewError)
 {
 	EBuildJobState OldState = EBuildJobState::Complete;
 
 	if (FWriteScopeLock WriteLock(Lock); NextState < NewState)
 	{
-		// TODO: Improve CompleteWithError to avoid unexpected state transitions.
-		//checkf(NextState < NewState,
-		//	TEXT("Job in state %s is requesting an invalid transition from %s to %s for build of '%s' by %s."),
-		//	LexToString(State), LexToString(NextState), LexToString(NewState), *Name, *WriteToString<32>(FunctionName));
-
-		if (Owner.IsCanceled())
+		if (Owner.IsCanceled() && NewError.IsEmpty())
 		{
 			NewState = EBuildJobState::Complete;
-			if (Output.IsNull())
-			{
-				OutputBuilder.AddLog({ANSITEXTVIEW("LogDerivedDataBuild"), ANSITEXTVIEW("Build was canceled."), EBuildOutputLogLevel::Error});
-			}
+			NewError = ANSITEXTVIEW("Build was canceled.");
+		}
+
+		if (!NewError.IsEmpty() && Output.IsNull())
+		{
+			OutputBuilder.AddLog({ANSITEXTVIEW("LogDerivedDataBuild"), NewError, EBuildOutputLogLevel::Error});
 		}
 
 		NextState = NewState;
@@ -1138,6 +1116,13 @@ void FBuildJob::ExecuteTransition(EBuildJobState OldState, EBuildJobState NewSta
 	}
 	if (NewState == EBuildJobState::Complete)
 	{
+		if (OnComplete)
+		{
+			const FCacheKey& CacheKey = Context ? Context->GetCacheKey() : FCacheKey::Empty;
+			const EStatus Status = Owner.IsCanceled() ? EStatus::Canceled : Output.Get().HasError() ? EStatus::Error : EStatus::Ok;
+			OnComplete({CacheKey, FBuildOutput(Output.Get()), BuildStatus, Status});
+			OnComplete = nullptr;
+		}
 		Output.Reset();
 		Context = nullptr;
 	}
