DDC: Made the cache optional for the build system

#rnx

diff --git a/Engine/Source/Developer/DerivedDataCache/Private/DerivedDataBuild.cpp b/Engine/Source/Developer/DerivedDataCache/Private/DerivedDataBuild.cpp
--- a/Engine/Source/Developer/DerivedDataCache/Private/DerivedDataBuild.cpp
+++ b/Engine/Source/Developer/DerivedDataCache/Private/DerivedDataBuild.cpp
@@ -72,7 +72,7 @@ DEFINE_LOG_CATEGORY(LogDerivedDataBuild);
 class FBuild final : public IBuild
 {
 public:
-	explicit FBuild(ICache& InCache)
+	explicit FBuild(ICache* InCache)
 		: Cache(InCache)
 	{
 	}
@@ -118,7 +118,7 @@ public:
 	}
 
 private:
-	ICache& Cache;
+	ICache* Cache;
 	TUniquePtr<IBuildFunctionRegistry> FunctionRegistry{CreateBuildFunctionRegistry()};
 	TUniquePtr<IBuildWorkerRegistry> WorkerRegistry{CreateBuildWorkerRegistry()};
 	TUniquePtr<IBuildScheduler> DefaultScheduler{CreateBuildScheduler()};
@@ -127,7 +127,7 @@ private:
 
 ///////////////////////////////////////////////////////////////////////////////////////////////////////////////////////
 
-IBuild* CreateBuild(ICache& Cache)
+IBuild* CreateBuild(ICache* Cache)
 {
 	LLM_SCOPE_BYTAG(DerivedDataBuild);
 	return new FBuild(Cache);
diff --git a/Engine/Source/Developer/DerivedDataCache/Private/DerivedDataBuildJob.cpp b/Engine/Source/Developer/DerivedDataCache/Private/DerivedDataBuildJob.cpp
--- a/Engine/Source/Developer/DerivedDataCache/Private/DerivedDataBuildJob.cpp
+++ b/Engine/Source/Developer/DerivedDataCache/Private/DerivedDataBuildJob.cpp
@@ -124,7 +124,7 @@ public:
 	inline const FSharedString& GetName() const final { return Name; }
 	inline const FUtf8SharedString& GetFunction() const final { return FunctionName; }
 
-	inline ICache& GetCache() const final { return Cache; }
+	inline ICache* GetCache() const final { return Cache; }
 	inline IBuild& GetBuild() const final { return BuildSystem; }
 
 	void StepExecution() final;
@@ -248,7 +248,7 @@ private:
 	/** Keys for missing inputs. */
 	TArray<FUtf8StringView> MissingInputs;
 
-	ICache& Cache;
+	ICache* Cache{};
 	IBuild& BuildSystem;
 
 	/** True if AdvanceToState is executing. */
@@ -514,7 +514,7 @@ static FCacheRecordPolicy MakeCacheRecordQueryPolicy(const FBuildPolicy& BuildPo
 
 void FBuildJob::EnterCacheQuery()
 {
-	if (!Context ||
+	if (!Cache || !Context ||
 		!EnumHasAnyFlags(Context->GetCachePolicyMask(), ECachePolicy::Query) ||
 		!EnumHasAnyFlags(Context->GetBuildPolicyMask() & BuildPolicy.GetCombinedPolicy(), EBuildPolicy::CacheQuery))
 	{
@@ -526,7 +526,7 @@ void FBuildJob::BeginCacheQuery()
 {
 	TRACE_CPUPROFILER_EVENT_SCOPE(FBuildJob::CacheQuery);
 	EnumAddFlags(BuildStatus, EBuildStatus::CacheQuery);
-	Cache.Get({{Name, Context->GetCacheKey(), MakeCacheRecordQueryPolicy(BuildPolicy, *Context)}}, Owner,
+	Cache->Get({{Name, Context->GetCacheKey(), MakeCacheRecordQueryPolicy(BuildPolicy, *Context)}}, Owner,
 		[this](FCacheGetResponse&& Response) { EndCacheQuery(MoveTemp(Response)); });
 }
 
@@ -575,7 +575,7 @@ static ECachePolicy MakeCacheStorePolicy(EBuildPolicy BuildPolicy, const FBuildJ
 
 void FBuildJob::EnterCacheStore()
 {
-	if (!Context ||
+	if (!Cache || !Context ||
 		!EnumHasAnyFlags(Context->GetCachePolicyMask(), ECachePolicy::Store) ||
 		!EnumHasAnyFlags(Context->GetBuildPolicyMask() & BuildPolicy.GetCombinedPolicy(), EBuildPolicy::CacheStoreOnBuild) ||
 		EnumHasAnyFlags(BuildStatus, EBuildStatus::CacheQueryHit) ||
@@ -592,7 +592,7 @@ void FBuildJob::BeginCacheStore()
 	EnumAddFlags(BuildStatus, EBuildStatus::CacheStore);
 	FCacheRecordBuilder RecordBuilder(Context->GetCacheKey());
 	Output.Get().Save(RecordBuilder);
-	Cache.Put({{Name, RecordBuilder.Build(), MakeCacheStorePolicy(BuildPolicy.GetCombinedPolicy(), *Context)}}, Owner,
+	Cache->Put({{Name, RecordBuilder.Build(), MakeCacheStorePolicy(BuildPolicy.GetCombinedPolicy(), *Context)}}, Owner,
 		[this](FCachePutResponse&& Response) { EndCacheStore(MoveTemp(Response)); });
 }
 
diff --git a/Engine/Source/Developer/DerivedDataCache/Private/DerivedDataBuildPrivate.h b/Engine/Source/Developer/DerivedDataCache/Private/DerivedDataBuildPrivate.h
--- a/Engine/Source/Developer/DerivedDataCache/Private/DerivedDataBuildPrivate.h
+++ b/Engine/Source/Developer/DerivedDataCache/Private/DerivedDataBuildPrivate.h
@@ -36,7 +36,7 @@ namespace UE::DerivedData::Private
 DECLARE_LOG_CATEGORY_EXTERN(LogDerivedDataBuild, Log, All);
 
 // Implemented in DerivedDataBuild.cpp
-IBuild* CreateBuild(ICache& Cache);
+IBuild* CreateBuild(ICache* Cache);
 
 // Implemented in DerivedDataBuildFunctionRegistry.cpp
 IBuildFunctionRegistry* CreateBuildFunctionRegistry();
@@ -65,7 +65,7 @@ FValueId GetBuildOutputValueId();
 // Implemented in DerivedDataBuildSession.cpp
 FBuildSession CreateBuildSession(
 	const FSharedString& Name,
-	ICache& Cache,
+	ICache* Cache,
 	IBuild& BuildSystem,
 	IBuildScheduler& Scheduler,
 	IBuildInputResolver* InputResolver);
@@ -73,7 +73,7 @@ FBuildSession CreateBuildSession(
 // Implemented in DerivedDataBuildJob.cpp
 struct FBuildJobCreateParams
 {
-	ICache& Cache;
+	ICache* Cache{};
 	IBuild& BuildSystem;
 	IBuildScheduler& Scheduler;
 	IBuildInputResolver* InputResolver{};
diff --git a/Engine/Source/Developer/DerivedDataCache/Private/DerivedDataBuildSession.cpp b/Engine/Source/Developer/DerivedDataCache/Private/DerivedDataBuildSession.cpp
--- a/Engine/Source/Developer/DerivedDataCache/Private/DerivedDataBuildSession.cpp
+++ b/Engine/Source/Developer/DerivedDataCache/Private/DerivedDataBuildSession.cpp
@@ -18,7 +18,7 @@ class FBuildSessionInternal final : public IBuildSessionInternal
 public:
 	FBuildSessionInternal(
 		const FSharedString& InName,
-		ICache& InCache,
+		ICache* InCache,
 		IBuild& InBuildSystem,
 		IBuildScheduler& InScheduler,
 		IBuildInputResolver* InInputResolver)
@@ -53,7 +53,7 @@ public:
 		FOnBuildComplete&& OnComplete) final;
 
 	FSharedString Name;
-	ICache& Cache;
+	ICache* Cache;
 	IBuild& BuildSystem;
 	IBuildScheduler& Scheduler;
 	IBuildInputResolver* InputResolver;
@@ -102,7 +102,7 @@ FBuildSession CreateBuildSession(IBuildSessionInternal* Session)
 
 FBuildSession CreateBuildSession(
 	const FSharedString& Name,
-	ICache& Cache,
+	ICache* Cache,
 	IBuild& BuildSystem,
 	IBuildScheduler& Scheduler,
 	IBuildInputResolver* InputResolver)
diff --git a/Engine/Source/Developer/DerivedDataCache/Private/DerivedDataCacheModule.cpp b/Engine/Source/Developer/DerivedDataCache/Private/DerivedDataCacheModule.cpp
--- a/Engine/Source/Developer/DerivedDataCache/Private/DerivedDataCacheModule.cpp
+++ b/Engine/Source/Developer/DerivedDataCache/Private/DerivedDataCacheModule.cpp
@@ -57,7 +57,7 @@ public:
 		FScopeLock Lock(&CreateLock);
 		if (!GDerivedDataBuild)
 		{
-			GDerivedDataBuild = CreateBuild(*GDerivedDataCache);
+			GDerivedDataBuild = CreateBuild(GDerivedDataCache);
 			check(GDerivedDataBuild);
 		}
 	}
diff --git a/Engine/Source/Developer/DerivedDataCache/Public/DerivedDataBuildJob.h b/Engine/Source/Developer/DerivedDataCache/Public/DerivedDataBuildJob.h
--- a/Engine/Source/Developer/DerivedDataCache/Public/DerivedDataBuildJob.h
+++ b/Engine/Source/Developer/DerivedDataCache/Public/DerivedDataBuildJob.h
@@ -40,7 +40,7 @@ public:
 	virtual const FUtf8SharedString& GetFunction() const = 0;
 
 	/** Returns the cache associated with this job. */
-	virtual ICache& GetCache() const = 0;
+	virtual ICache* GetCache() const = 0;
 	/** Returns the build system associated with this job. */
 	virtual IBuild& GetBuild() const = 0;
 
