DDC: Remove any trailing slash from the Jupiter URI

#rnx

diff --git a/Engine/Source/Developer/DerivedDataCache/Private/HttpCacheStore.cpp b/Engine/Source/Developer/DerivedDataCache/Private/HttpCacheStore.cpp
--- a/Engine/Source/Developer/DerivedDataCache/Private/HttpCacheStore.cpp
+++ b/Engine/Source/Developer/DerivedDataCache/Private/HttpCacheStore.cpp
@@ -1539,6 +1539,9 @@ FHttpCacheStore::FHttpCacheStore(const FHttpCacheStoreParams& Params)
 {
 	TRACE_CPUPROFILER_EVENT_SCOPE(HttpDDC_Construct);
 
+	// Remove any trailing / because constructing a URI will add one.
+	while (Domain.RemoveFromEnd(TEXT("/")));
+
 	EffectiveDomain.Append(Domain);
 	TAnsiStringBuilder<256> ResolvedDomain;
 	if (Params.bResolveHostCanonicalName && TryResolveCanonicalHost(EffectiveDomain, ResolvedDomain))
