HttpCacheStore: When a record is missing data and we return EStatus::Error and an empty record, fix the FCacheRecord passed into OnComplete to not have metadata or any values. The FCacheStoreHierarchy client of FHttpCacheStore assumes that metadata or values are only present on a successful record, and if it is present, incorrectly Puts the record as a valid record into the other cachestores in the hierarchy. This then breaks clients of DDC, because the record is invalid since it is missing some values.
This change fixes the warning messages when using EditorDomain: LogEditorDomain: Warning: Package %s received an invalid CacheRecord...
#rnx

diff --git a/Engine/Source/Developer/DerivedDataCache/Private/HttpCacheStore.cpp b/Engine/Source/Developer/DerivedDataCache/Private/HttpCacheStore.cpp
--- a/Engine/Source/Developer/DerivedDataCache/Private/HttpCacheStore.cpp
+++ b/Engine/Source/Developer/DerivedDataCache/Private/HttpCacheStore.cpp
@@ -29,6 +29,7 @@
 #include "Misc/CommandLine.h"
 #include "Misc/ConfigCacheIni.h"
 #include "Misc/FileHelper.h"
+#include "Misc/Optional.h"
 #include "Misc/ScopeExit.h"
 #include "Misc/ScopeLock.h"
 #include "Misc/ScopeRWLock.h"
@@ -1507,8 +1508,9 @@ void FHttpCacheStore::FGetRecordOp::FinishDataStep(bool bSuccess, uint64 InBytes
 
 	if (PendingOperations.fetch_sub(1, std::memory_order_acq_rel) == 1)
 	{
-		EStatus Status = EStatus::Error;
+		EStatus Status;
 		uint32 LocalSuccessfulOperations = SuccessfulOperations.load(std::memory_order_relaxed);
+		TOptional<FCacheRecord> Record;
 		if (LocalSuccessfulOperations == TotalOperations)
 		{
 			for (int32 Index = 0; Index < RequiredHeads.Num(); ++Index)
@@ -1520,9 +1522,15 @@ void FHttpCacheStore::FGetRecordOp::FinishDataStep(bool bSuccess, uint64 InBytes
 			{
 				RecordBuilder.AddValue(FValueWithId(RequiredGets[Index].GetId(), FetchedBuffers[Index]));
 			}
+			Record.Emplace(RecordBuilder.Build());
 			Status = EStatus::Ok;
 		}
-		OnComplete({Name, RecordBuilder.Build(), UserData, Status}, BytesReceived.load(std::memory_order_relaxed));
+		else
+		{
+			Record.Emplace(FCacheRecordBuilder(Key).Build());
+			Status = EStatus::Error;
+		}
+		OnComplete({Name, MoveTemp(Record.GetValue()), UserData, Status}, BytesReceived.load(std::memory_order_relaxed));
 	}
 }
 
