DDC: Added a missing read lock in the cache hierarchy

#rnx

diff --git a/Engine/Source/Developer/DerivedDataCache/Private/DerivedDataCacheStoreHierarchy.cpp b/Engine/Source/Developer/DerivedDataCache/Private/DerivedDataCacheStoreHierarchy.cpp
--- a/Engine/Source/Developer/DerivedDataCache/Private/DerivedDataCacheStoreHierarchy.cpp
+++ b/Engine/Source/Developer/DerivedDataCache/Private/DerivedDataCacheStoreHierarchy.cpp
@@ -748,7 +748,6 @@ void FCacheStoreHierarchy::TGetBatch<Params>::CompleteRequest(FGetResponse&& Res
 
 	FState& State = States[int32(Response.UserData)];
 	Response.UserData = State.Request.UserData;
-	const FCacheStoreNode& Node = Hierarchy.Nodes[State.NodeIndex];
 	const EStatus PreviousStatus = State.Response.Status;
 	const int32 PreviousNodeIndex = State.NodeIndex;
 
@@ -758,6 +757,9 @@ void FCacheStoreHierarchy::TGetBatch<Params>::CompleteRequest(FGetResponse&& Res
 		return;
 	}
 
+	FReadScopeLock Lock(Hierarchy.NodesLock);
+
+	const FCacheStoreNode& Node = Hierarchy.Nodes[State.NodeIndex];
 	const bool bFirstOk = State.Response.Status == EStatus::Ok && PreviousStatus == EStatus::Error;
 	const bool bLastQuery = bFirstOk || !CanQueryIfError(GetCombinedPolicy(State.Request.Policy), Node.NodeFlags);
 
