DDC: Allow the build system to be created without a cache

diff --git a/Engine/Source/Developer/DerivedDataCache/Private/DerivedDataCacheModule.cpp b/Engine/Source/Developer/DerivedDataCache/Private/DerivedDataCacheModule.cpp
--- a/Engine/Source/Developer/DerivedDataCache/Private/DerivedDataCacheModule.cpp
+++ b/Engine/Source/Developer/DerivedDataCache/Private/DerivedDataCacheModule.cpp
@@ -53,7 +53,6 @@ public:
 
 	void CreateBuildOnce()
 	{
-		CreateCacheOnce();
 		FScopeLock Lock(&CreateLock);
 		if (!GDerivedDataBuild)
 		{
