DDC: Added LaunchTaskInThreadPool as a replacement for the two existing thread pool task requests

diff --git a/Engine/Source/Developer/DerivedDataCache/Private/DerivedDataBuildScheduler.cpp b/Engine/Source/Developer/DerivedDataCache/Private/DerivedDataBuildScheduler.cpp
--- a/Engine/Source/Developer/DerivedDataCache/Private/DerivedDataBuildScheduler.cpp
+++ b/Engine/Source/Developer/DerivedDataCache/Private/DerivedDataBuildScheduler.cpp
@@ -5,14 +5,11 @@
 #include "DerivedDataBuildFunctionFactory.h"
 #include "DerivedDataBuildFunctionRegistry.h"
 #include "DerivedDataBuildPrivate.h"
-#include "DerivedDataRequest.h"
 #include "DerivedDataRequestOwner.h"
-#include "Experimental/Async/LazyEvent.h"
+#include "DerivedDataThreadPoolTask.h"
 #include "Experimental/DerivedDataBuildSchedulerThreadPoolProvider.h"
 #include "Experimental/Misc/ExecutionResource.h"
 #include "HAL/CriticalSection.h"
-#include "Misc/IQueuedWork.h"
-#include "Misc/QueuedThreadPool.h"
 #include "Misc/ScopeRWLock.h"
 
 namespace UE::DerivedData::Private
@@ -41,8 +38,6 @@ private:
 
 	FQueuedThreadPool* FindThreadPool(const FUtf8SharedString& TypeName) const;
 
-	class FRequest;
-
 private:
 	mutable FRWLock Lock;
 	TMap<FUtf8SharedString, IBuildSchedulerThreadPoolProvider*> Providers;
@@ -50,97 +45,6 @@ private:
 
 ///////////////////////////////////////////////////////////////////////////////////////////////////////////////////////
 
-class FBuildSchedulerMemoryQueue::FRequest final : public FRequestBase, public IQueuedWork
-{
-public:
-	FRequest(
-		uint64 MemoryEstimate,
-		IRequestOwner& Owner,
-		FQueuedThreadPool& ThreadPool,
-		TUniqueFunction<void ()>&& OnComplete);
-	~FRequest();
-
-private:
-	TRefCountPtr<IRequest> TryEnd();
-	void Execute();
-
-	void SetPriority(EPriority Priority) final;
-	void Wait() final;
-	void Cancel() final { Wait(); }
-
-	void DoThreadedWork() final { Execute(); }
-	void Abandon() final { Execute(); }
-	int64 GetRequiredMemory() const final { return int64(MemoryEstimate); }
-
-private:
-	uint64 MemoryEstimate;
-	IRequestOwner& Owner;
-	FQueuedThreadPool& ThreadPool;
-	TUniqueFunction<void ()> OnComplete;
-	UE::FLazyEvent Event{EEventMode::ManualReset};
-	std::atomic<bool> bClaimed = false;
-};
-
-FBuildSchedulerMemoryQueue::FRequest::FRequest(
-	const uint64 InMemoryEstimate,
-	IRequestOwner& InOwner,
-	FQueuedThreadPool& InThreadPool,
-	TUniqueFunction<void ()>&& InOnComplete)
-	: MemoryEstimate(InMemoryEstimate)
-	, Owner(InOwner)
-	, ThreadPool(InThreadPool)
-	, OnComplete(MoveTemp(InOnComplete))
-{
-	AddRef(); // Released in Execute() or Cancel()
-	Owner.Begin(this);
-	ThreadPool.AddQueuedWork(this, ConvertToQueuedWorkPriority(Owner.GetPriority()));
-}
-
-FBuildSchedulerMemoryQueue::FRequest::~FRequest()
-{
-	check(bClaimed.load(std::memory_order_relaxed));
-}
-
-TRefCountPtr<IRequest> FBuildSchedulerMemoryQueue::FRequest::TryEnd()
-{
-	return bClaimed.exchange(true) ? nullptr : Owner.End(this, [this]
-	{
-		Invoke(OnComplete);
-		Event.Trigger();
-	});
-}
-
-void FBuildSchedulerMemoryQueue::FRequest::Execute()
-{
-	TryEnd();
-	Release();
-}
-
-void FBuildSchedulerMemoryQueue::FRequest::SetPriority(EPriority Priority)
-{
-	if (ThreadPool.RetractQueuedWork(this))
-	{
-		ThreadPool.AddQueuedWork(this, ConvertToQueuedWorkPriority(Priority));
-	}
-}
-
-void FBuildSchedulerMemoryQueue::FRequest::Wait()
-{
-	if (TRefCountPtr<IRequest> Self = TryEnd())
-	{
-		if (ThreadPool.RetractQueuedWork(this))
-		{
-			Release();
-		}
-	}
-	else
-	{
-		Event.Wait();
-	}
-}
-
-///////////////////////////////////////////////////////////////////////////////////////////////////////////////////////
-
 FBuildSchedulerMemoryQueue::FBuildSchedulerMemoryQueue()
 {
 	IModularFeatures& ModularFeatures = IModularFeatures::Get();
@@ -222,7 +126,7 @@ void FBuildSchedulerMemoryQueue::WaitForMemory(
 {
 	if (FQueuedThreadPool* ThreadPool = FindThreadPool(TypeName))
 	{
-		new FRequest(MemoryEstimate, Owner, *ThreadPool, MoveTemp(OnComplete));
+		LaunchTaskInThreadPool(MemoryEstimate, Owner, *ThreadPool, MoveTemp(OnComplete));
 	}
 	else
 	{
@@ -246,8 +150,10 @@ public:
 
 	void ScheduleCacheQuery() final       { StepSync(); }
 	void ScheduleCacheStore() final       { StepSync(); }
+
 	void ScheduleResolveKey() final       { StepAsync(TEXT("ResolveKey")); }
 	void ScheduleResolveInputMeta() final { StepAsync(TEXT("ResolveInputMeta")); }
+
 	void ScheduleResolveInputData() final
 	{
 		if (Params.MissingRemoteInputsSize)
@@ -259,6 +165,7 @@ public:
 			StepAsyncOrQueue(TEXT("ResolveInputData"));
 		}
 	}
+
 	void ScheduleExecuteRemote() final    { StepAsync(TEXT("ExecuteRemote")); }
 	void ScheduleExecuteLocal() final     { StepAsyncOrQueue(TEXT("ExecuteLocal")); }
 
@@ -271,6 +178,7 @@ private:
 	void StepSync()
 	{
 		Job.StepExecution();
+		// DO NOT ACCESS THIS AGAIN PAST THIS POINT!
 	}
 
 	void StepAsync(const TCHAR* DebugName)
@@ -283,6 +191,7 @@ private:
 		{
 			Owner.LaunchTask(DebugName, [this] { Job.StepExecution(); });
 		}
+		// DO NOT ACCESS THIS AGAIN PAST THIS POINT!
 	}
 
 	void StepAsyncOrQueue(const TCHAR* DebugName)
@@ -305,6 +214,7 @@ private:
 			ExecutionResources = FExecutionResourceContext::Get();
 			StepAsync(DebugName);
 		});
+		// DO NOT ACCESS THIS AGAIN PAST THIS POINT!
 	}
 
 private:
diff --git a/Engine/Source/Developer/DerivedDataCache/Private/DerivedDataCache.cpp b/Engine/Source/Developer/DerivedDataCache/Private/DerivedDataCache.cpp
--- a/Engine/Source/Developer/DerivedDataCache/Private/DerivedDataCache.cpp
+++ b/Engine/Source/Developer/DerivedDataCache/Private/DerivedDataCache.cpp
@@ -20,6 +20,7 @@
 #include "DerivedDataPrivate.h"
 #include "DerivedDataRequest.h"
 #include "DerivedDataRequestOwner.h"
+#include "DerivedDataThreadPoolTask.h"
 #include "Experimental/Async/LazyEvent.h"
 #include "Features/IModularFeatures.h"
 #include "HAL/ThreadSafeCounter.h"
@@ -328,95 +329,11 @@ namespace UE::DerivedData::Private
 
 FQueuedThreadPool* GCacheThreadPool;
 
-class FCacheThreadPoolTaskRequest final : public FRequestBase, private FInheritedContextBase, private IQueuedWork
-{
-public:
-	inline FCacheThreadPoolTaskRequest(IRequestOwner& InOwner, TUniqueFunction<void ()>&& InTaskBody)
-		: Owner(InOwner)
-		, TaskBody(MoveTemp(InTaskBody))
-	{
-		CaptureInheritedContext();
-		Owner.Begin(this);
-		DoneEvent.Reset();
-		GCacheThreadPool->AddQueuedWork(this, ConvertToQueuedWorkPriority(Owner.GetPriority()));
-	}
-
-private:
-	inline void Execute()
-	{
-		FInheritedContextScope InheritedContextScope = RestoreInheritedContext();
-		FScopeCycleCounter Scope(GetStatId(), /*bAlways*/ true);
-		Owner.End(this, [this]
-		{
-			TaskBody();
-			DoneEvent.Trigger();
-		});
-		// DO NOT ACCESS ANY MEMBERS PAST THIS POINT!
-	}
-
-	// IRequest Interface
-
-	inline void SetPriority(EPriority Priority) final
-	{
-		if (GCacheThreadPool->RetractQueuedWork(this))
-		{
-			GCacheThreadPool->AddQueuedWork(this, ConvertToQueuedWorkPriority(Priority));
-		}
-	}
-
-	inline void Cancel() final
-	{
-		if (!DoneEvent.Wait(0))
-		{
-			if (GCacheThreadPool->RetractQueuedWork(this))
-			{
-				Abandon();
-			}
-			else
-			{
-				FScopeCycleCounter Scope(GetStatId());
-				DoneEvent.Wait();
-			}
-		}
-	}
-
-	inline void Wait() final
-	{
-		if (!DoneEvent.Wait(0))
-		{
-			if (GCacheThreadPool->RetractQueuedWork(this))
-			{
-				DoThreadedWork();
-			}
-			else
-			{
-				FScopeCycleCounter Scope(GetStatId());
-				DoneEvent.Wait();
-			}
-		}
-	}
-
-	// IQueuedWork Interface
-
-	inline void DoThreadedWork() final { Execute(); }
-	inline void Abandon() final { Execute(); }
-
-	inline TStatId GetStatId() const
-	{
-		RETURN_QUICK_DECLARE_CYCLE_STAT(FCacheThreadPoolTaskRequest, STATGROUP_ThreadPoolAsyncTasks);
-	}
-
-private:
-	IRequestOwner& Owner;
-	TUniqueFunction<void ()> TaskBody;
-	FLazyEvent DoneEvent{EEventMode::ManualReset};
-};
-
 void LaunchTaskInCacheThreadPool(IRequestOwner& Owner, TUniqueFunction<void ()>&& TaskBody)
 {
 	if (GCacheThreadPool)
 	{
-		new FCacheThreadPoolTaskRequest(Owner, MoveTemp(TaskBody));
+		LaunchTaskInThreadPool(Owner, *GCacheThreadPool, MoveTemp(TaskBody));
 	}
 	else
 	{
diff --git a/Engine/Source/Developer/DerivedDataCache/Private/DerivedDataThreadPoolTask.cpp b/Engine/Source/Developer/DerivedDataCache/Private/DerivedDataThreadPoolTask.cpp
new file mode 100644
--- /dev/null
+++ b/Engine/Source/Developer/DerivedDataCache/Private/DerivedDataThreadPoolTask.cpp
@@ -0,0 +1,136 @@
+// Copyright Epic Games, Inc. All Rights Reserved.
+
+#include "DerivedDataThreadPoolTask.h"
+
+#include "Async/InheritedContext.h"
+#include "DerivedDataRequest.h"
+#include "DerivedDataRequestOwner.h"
+#include "DerivedDataRequestTypes.h"
+#include "Experimental/Async/LazyEvent.h"
+#include "Misc/IQueuedWork.h"
+#include "Misc/QueuedThreadPool.h"
+#include "Stats/Stats.h"
+
+namespace UE::DerivedData
+{
+
+class FThreadPoolTaskRequest final : public FRequestBase, private FInheritedContextBase, private IQueuedWork
+{
+public:
+	FThreadPoolTaskRequest(
+		uint64 MemoryEstimate,
+		IRequestOwner& Owner,
+		FQueuedThreadPool& ThreadPool,
+		TUniqueFunction<void ()>&& TaskBody);
+	~FThreadPoolTaskRequest();
+
+private:
+	TRefCountPtr<IRequest> TryEnd();
+	void Execute();
+
+	void SetPriority(EPriority Priority) final;
+	void Cancel() final { Wait(); }
+	void Wait() final;
+
+	void DoThreadedWork() final { Execute(); }
+	void Abandon() final { Execute(); }
+	int64 GetRequiredMemory() const final { return int64(MemoryEstimate); }
+
+	TStatId GetStatId() const
+	{
+		RETURN_QUICK_DECLARE_CYCLE_STAT(FThreadPoolTaskRequest, STATGROUP_ThreadPoolAsyncTasks);
+	}
+
+private:
+	uint64 MemoryEstimate;
+	IRequestOwner& Owner;
+	FQueuedThreadPool& ThreadPool;
+	TUniqueFunction<void ()> TaskBody;
+	FLazyEvent Event{EEventMode::ManualReset};
+	std::atomic<bool> bClaimed = false;
+};
+
+FThreadPoolTaskRequest::FThreadPoolTaskRequest(
+	const uint64 InMemoryEstimate,
+	IRequestOwner& InOwner,
+	FQueuedThreadPool& InThreadPool,
+	TUniqueFunction<void ()>&& InTaskBody)
+	: MemoryEstimate(InMemoryEstimate)
+	, Owner(InOwner)
+	, ThreadPool(InThreadPool)
+	, TaskBody(MoveTemp(InTaskBody))
+{
+	check(MemoryEstimate <= MAX_int64);
+	CaptureInheritedContext();
+	AddRef(); // Released in Execute() or Cancel()
+	Owner.Begin(this);
+	ThreadPool.AddQueuedWork(this, ConvertToQueuedWorkPriority(Owner.GetPriority()));
+}
+
+FThreadPoolTaskRequest::~FThreadPoolTaskRequest()
+{
+	check(bClaimed.load(std::memory_order_relaxed));
+}
+
+TRefCountPtr<IRequest> FThreadPoolTaskRequest::TryEnd()
+{
+	return bClaimed.exchange(true) ? nullptr : Owner.End(this, [this]
+	{
+		FInheritedContextScope InheritedContextScope = RestoreInheritedContext();
+		FScopeCycleCounter Scope(GetStatId(), /*bAlways*/ true);
+		TaskBody();
+		Event.Trigger();
+	});
+}
+
+void FThreadPoolTaskRequest::Execute()
+{
+	TryEnd();
+	Release();
+	// DO NOT ACCESS ANY MEMBERS PAST THIS POINT!
+}
+
+void FThreadPoolTaskRequest::SetPriority(EPriority Priority)
+{
+	if (ThreadPool.RetractQueuedWork(this))
+	{
+		ThreadPool.AddQueuedWork(this, ConvertToQueuedWorkPriority(Priority));
+	}
+}
+
+void FThreadPoolTaskRequest::Wait()
+{
+	if (TRefCountPtr<IRequest> Self = TryEnd())
+	{
+		if (ThreadPool.RetractQueuedWork(this))
+		{
+			Release();
+		}
+	}
+	else
+	{
+		FScopeCycleCounter Scope(GetStatId());
+		Event.Wait();
+	}
+}
+
+void LaunchTaskInThreadPool(
+	uint64 MemoryEstimate,
+	IRequestOwner& Owner,
+	FQueuedThreadPool& ThreadPool,
+	TUniqueFunction<void ()>&& TaskBody)
+{
+	// The request is reference-counted and will be deleted when complete.
+	new FThreadPoolTaskRequest(MemoryEstimate, Owner, ThreadPool, MoveTemp(TaskBody));
+}
+
+void LaunchTaskInThreadPool(
+	IRequestOwner& Owner,
+	FQueuedThreadPool& ThreadPool,
+	TUniqueFunction<void ()>&& TaskBody)
+{
+	// The request is reference-counted and will be deleted when complete.
+	new FThreadPoolTaskRequest(0, Owner, ThreadPool, MoveTemp(TaskBody));
+}
+
+} // UE::DerivedData
diff --git a/Engine/Source/Developer/DerivedDataCache/Public/DerivedDataThreadPoolTask.h b/Engine/Source/Developer/DerivedDataCache/Public/DerivedDataThreadPoolTask.h
new file mode 100644
--- /dev/null
+++ b/Engine/Source/Developer/DerivedDataCache/Public/DerivedDataThreadPoolTask.h
@@ -0,0 +1,44 @@
+// Copyright Epic Games, Inc. All Rights Reserved.
+
+#pragma once
+
+#include "CoreTypes.h"
+
+#define UE_API DERIVEDDATACACHE_API
+
+class FQueuedThreadPool;
+
+template <typename FuncType> class TUniqueFunction;
+
+namespace UE::DerivedData { class IRequestOwner; }
+
+namespace UE::DerivedData
+{
+
+/**
+ * Launches a task in the thread pool that tracks the priority of the request owner.
+ *
+ * Waiting on or canceling the task before it has started executing in the thread pool will cause
+ * it to be retracted and executed on the thread that asked to wait or cancel.
+ * 
+ * TaskBody is guaranteed to return before the request is considered complete.
+ * TaskBody is responsible for checking for cancellation by checking Owner.IsCanceled().
+ *
+ * @param MemoryEstimate   The estimated peak memory required to execute TaskBody.
+ * @param Owner            The owner of the request that will execute TaskBody.
+ * @param ThreadPool       The thread pool within which to execute the task, if not retracted.
+ * @param TaskBody         The task body that will be unconditionally executed exactly once.
+ */
+UE_API void LaunchTaskInThreadPool(
+	uint64 MemoryEstimate,
+	IRequestOwner& Owner,
+	FQueuedThreadPool& ThreadPool,
+	TUniqueFunction<void ()>&& TaskBody);
+UE_API void LaunchTaskInThreadPool(
+	IRequestOwner& Owner,
+	FQueuedThreadPool& ThreadPool,
+	TUniqueFunction<void ()>&& TaskBody);
+
+} // UE::DerivedData
+
+#undef UE_API
