DDC: Added new build scheduler modular features to queue by type and reserve memory

diff --git a/Engine/Source/Developer/DerivedDataCache/Private/DerivedDataBuildScheduler.cpp b/Engine/Source/Developer/DerivedDataCache/Private/DerivedDataBuildScheduler.cpp
--- a/Engine/Source/Developer/DerivedDataCache/Private/DerivedDataBuildScheduler.cpp
+++ b/Engine/Source/Developer/DerivedDataCache/Private/DerivedDataBuildScheduler.cpp
@@ -5,9 +5,9 @@
 #include "DerivedDataBuildFunctionFactory.h"
 #include "DerivedDataBuildFunctionRegistry.h"
 #include "DerivedDataBuildPrivate.h"
+#include "DerivedDataBuildSchedulerQueue.h"
 #include "DerivedDataRequestOwner.h"
 #include "DerivedDataThreadPoolTask.h"
-#include "Experimental/DerivedDataBuildSchedulerThreadPoolProvider.h"
 #include "Experimental/Misc/ExecutionResource.h"
 #include "HAL/CriticalSection.h"
 #include "Misc/ScopeRWLock.h"
@@ -17,138 +17,215 @@ namespace UE::DerivedData::Private
 
 ///////////////////////////////////////////////////////////////////////////////////////////////////////////////////////
 
-class FBuildSchedulerMemoryQueue
+class FBuildSchedulerTypeQueue
 {
 public:
-	FBuildSchedulerMemoryQueue();
-	~FBuildSchedulerMemoryQueue();
+	FBuildSchedulerTypeQueue();
+	~FBuildSchedulerTypeQueue();
 
-	void WaitForMemory(
-		const FUtf8SharedString& TypeName,
-		uint64 MemoryEstimate,
-		IRequestOwner& Owner,
-		TUniqueFunction<void ()>&& OnComplete);
+	void Queue(const FUtf8SharedString& TypeName, IRequestOwner& Owner, TUniqueFunction<void ()>&& OnComplete);
 
 private:
 	void OnModularFeatureRegistered(const FName& Type, IModularFeature* ModularFeature);
 	void OnModularFeatureUnregistered(const FName& Type, IModularFeature* ModularFeature);
 
-	void AddProviderNoLock(IBuildSchedulerThreadPoolProvider* Provider);
-	void RemoveProvider(IBuildSchedulerThreadPoolProvider* Provider);
-
-	FQueuedThreadPool* FindThreadPool(const FUtf8SharedString& TypeName) const;
+	void AddQueueNoLock(IBuildSchedulerTypeQueue* Queue);
+	void RemoveQueue(IBuildSchedulerTypeQueue* Queue);
 
 private:
 	mutable FRWLock Lock;
-	TMap<FUtf8SharedString, IBuildSchedulerThreadPoolProvider*> Providers;
+	TMap<FUtf8SharedString, IBuildSchedulerTypeQueue*> Queues;
 };
 
 ///////////////////////////////////////////////////////////////////////////////////////////////////////////////////////
 
-FBuildSchedulerMemoryQueue::FBuildSchedulerMemoryQueue()
+FBuildSchedulerTypeQueue::FBuildSchedulerTypeQueue()
 {
 	IModularFeatures& ModularFeatures = IModularFeatures::Get();
-	for (IBuildSchedulerThreadPoolProvider* Provider : ModularFeatures.GetModularFeatureImplementations<IBuildSchedulerThreadPoolProvider>(IBuildSchedulerThreadPoolProvider::FeatureName))
+	for (IBuildSchedulerTypeQueue* Queue : ModularFeatures.GetModularFeatureImplementations<IBuildSchedulerTypeQueue>(IBuildSchedulerTypeQueue::FeatureName))
 	{
-		AddProviderNoLock(Provider);
+		AddQueueNoLock(Queue);
 	}
-	ModularFeatures.OnModularFeatureRegistered().AddRaw(this, &FBuildSchedulerMemoryQueue::OnModularFeatureRegistered);
-	ModularFeatures.OnModularFeatureUnregistered().AddRaw(this, &FBuildSchedulerMemoryQueue::OnModularFeatureUnregistered);
+	ModularFeatures.OnModularFeatureRegistered().AddRaw(this, &FBuildSchedulerTypeQueue::OnModularFeatureRegistered);
+	ModularFeatures.OnModularFeatureUnregistered().AddRaw(this, &FBuildSchedulerTypeQueue::OnModularFeatureUnregistered);
 }
 
-FBuildSchedulerMemoryQueue::~FBuildSchedulerMemoryQueue()
+FBuildSchedulerTypeQueue::~FBuildSchedulerTypeQueue()
 {
 	IModularFeatures& ModularFeatures = IModularFeatures::Get();
 	ModularFeatures.OnModularFeatureUnregistered().RemoveAll(this);
 	ModularFeatures.OnModularFeatureRegistered().RemoveAll(this);
 }
 
-void FBuildSchedulerMemoryQueue::OnModularFeatureRegistered(const FName& Type, IModularFeature* ModularFeature)
+void FBuildSchedulerTypeQueue::OnModularFeatureRegistered(const FName& Type, IModularFeature* ModularFeature)
 {
-	if (Type == IBuildSchedulerThreadPoolProvider::FeatureName)
+	if (Type == IBuildSchedulerTypeQueue::FeatureName)
 	{
 		FWriteScopeLock WriteLock(Lock);
-		AddProviderNoLock(static_cast<IBuildSchedulerThreadPoolProvider*>(ModularFeature));
+		AddQueueNoLock(static_cast<IBuildSchedulerTypeQueue*>(ModularFeature));
 	}
 }
 
-void FBuildSchedulerMemoryQueue::OnModularFeatureUnregistered(const FName& Type, IModularFeature* ModularFeature)
+void FBuildSchedulerTypeQueue::OnModularFeatureUnregistered(const FName& Type, IModularFeature* ModularFeature)
 {
-	if (Type == IBuildSchedulerThreadPoolProvider::FeatureName)
+	if (Type == IBuildSchedulerTypeQueue::FeatureName)
 	{
-		RemoveProvider(static_cast<IBuildSchedulerThreadPoolProvider*>(ModularFeature));
+		RemoveQueue(static_cast<IBuildSchedulerTypeQueue*>(ModularFeature));
 	}
 }
 
-void FBuildSchedulerMemoryQueue::AddProviderNoLock(IBuildSchedulerThreadPoolProvider* Provider)
+void FBuildSchedulerTypeQueue::AddQueueNoLock(IBuildSchedulerTypeQueue* Queue)
 {
-	
-	const FUtf8SharedString& TypeName = Provider->GetTypeName();
+	const FUtf8SharedString& TypeName = Queue->GetTypeName();
 	const uint32 TypeNameHash = GetTypeHash(TypeName);
 	if (TypeName.IsEmpty())
 	{
 		UE_LOG(LogDerivedDataBuild, Error,
-			TEXT("An empty type name is not allowed in a build scheduler thread pool provider."));
+			TEXT("An empty type name is not allowed in a build scheduler type queue."));
 	}
-	else if (Providers.FindByHash(TypeNameHash, TypeName))
+	else if (Queues.FindByHash(TypeNameHash, TypeName))
 	{
 		UE_LOG(LogDerivedDataBuild, Error,
-			TEXT("More than one build scheduler thread pool provider has been registered with the type name %s."),
+			TEXT("More than one build scheduler type queue has been registered with the type name %s."),
 			*WriteToString<64>(TypeName));
 	}
 	else
 	{
-		Providers.EmplaceByHash(TypeNameHash, TypeName, Provider);
+		Queues.EmplaceByHash(TypeNameHash, TypeName, Queue);
 	}
 }
 
-void FBuildSchedulerMemoryQueue::RemoveProvider(IBuildSchedulerThreadPoolProvider* Provider)
+void FBuildSchedulerTypeQueue::RemoveQueue(IBuildSchedulerTypeQueue* Queue)
 {
-	const FUtf8SharedString& TypeName = Provider->GetTypeName();
+	const FUtf8SharedString& TypeName = Queue->GetTypeName();
 	const uint32 TypeNameHash = GetTypeHash(TypeName);
 	FWriteScopeLock WriteLock(Lock);
-	Providers.RemoveByHash(TypeNameHash, TypeName);
+	Queues.RemoveByHash(TypeNameHash, TypeName);
 }
 
-FQueuedThreadPool* FBuildSchedulerMemoryQueue::FindThreadPool(const FUtf8SharedString& TypeName) const
-{
-	const uint32 TypeNameHash = GetTypeHash(TypeName);
-	FReadScopeLock ReadLock(Lock);
-	IBuildSchedulerThreadPoolProvider* const* Provider = Providers.FindByHash(TypeNameHash, TypeName);
-	return Provider ? (*Provider)->GetThreadPool() : nullptr;
-}
-
-void FBuildSchedulerMemoryQueue::WaitForMemory(
+void FBuildSchedulerTypeQueue::Queue(
 	const FUtf8SharedString& TypeName,
-	const uint64 MemoryEstimate,
 	IRequestOwner& Owner,
 	TUniqueFunction<void ()>&& OnComplete)
 {
-	if (FQueuedThreadPool* ThreadPool = FindThreadPool(TypeName))
+	if (FReadScopeLock ReadLock(Lock); IBuildSchedulerTypeQueue* Queue = Queues.FindRef(TypeName))
+	{
+		return Queue->Queue(Owner, MoveTemp(OnComplete));
+	}
+
+	OnComplete();
+}
+
+///////////////////////////////////////////////////////////////////////////////////////////////////////////////////////
+
+class FBuildSchedulerMemoryQueue
+{
+public:
+	FBuildSchedulerMemoryQueue();
+	~FBuildSchedulerMemoryQueue();
+
+	void Reserve(uint64 Memory, IRequestOwner& Owner, TUniqueFunction<void ()>&& OnComplete);
+
+private:
+	void OnModularFeatureRegistered(const FName& Type, IModularFeature* ModularFeature);
+	void OnModularFeatureUnregistered(const FName& Type, IModularFeature* ModularFeature);
+
+	void AddQueueNoLock(IBuildSchedulerMemoryQueue* Queue);
+	void RemoveQueue(IBuildSchedulerMemoryQueue* Queue);
+
+private:
+	mutable FRWLock Lock;
+	IBuildSchedulerMemoryQueue* GlobalQueue = nullptr;
+};
+
+///////////////////////////////////////////////////////////////////////////////////////////////////////////////////////
+
+FBuildSchedulerMemoryQueue::FBuildSchedulerMemoryQueue()
+{
+	IModularFeatures& ModularFeatures = IModularFeatures::Get();
+	for (IBuildSchedulerMemoryQueue* Queue : ModularFeatures.GetModularFeatureImplementations<IBuildSchedulerMemoryQueue>(IBuildSchedulerMemoryQueue::FeatureName))
+	{
+		AddQueueNoLock(Queue);
+	}
+	ModularFeatures.OnModularFeatureRegistered().AddRaw(this, &FBuildSchedulerMemoryQueue::OnModularFeatureRegistered);
+	ModularFeatures.OnModularFeatureUnregistered().AddRaw(this, &FBuildSchedulerMemoryQueue::OnModularFeatureUnregistered);
+}
+
+FBuildSchedulerMemoryQueue::~FBuildSchedulerMemoryQueue()
+{
+	IModularFeatures& ModularFeatures = IModularFeatures::Get();
+	ModularFeatures.OnModularFeatureUnregistered().RemoveAll(this);
+	ModularFeatures.OnModularFeatureRegistered().RemoveAll(this);
+}
+
+void FBuildSchedulerMemoryQueue::OnModularFeatureRegistered(const FName& Type, IModularFeature* ModularFeature)
+{
+	if (Type == IBuildSchedulerMemoryQueue::FeatureName)
+	{
+		FWriteScopeLock WriteLock(Lock);
+		AddQueueNoLock(static_cast<IBuildSchedulerMemoryQueue*>(ModularFeature));
+	}
+}
+
+void FBuildSchedulerMemoryQueue::OnModularFeatureUnregistered(const FName& Type, IModularFeature* ModularFeature)
+{
+	if (Type == IBuildSchedulerMemoryQueue::FeatureName)
+	{
+		RemoveQueue(static_cast<IBuildSchedulerMemoryQueue*>(ModularFeature));
+	}
+}
+
+void FBuildSchedulerMemoryQueue::AddQueueNoLock(IBuildSchedulerMemoryQueue* Queue)
+{
+	if (GlobalQueue)
 	{
-		LaunchTaskInThreadPool(MemoryEstimate, Owner, *ThreadPool, MoveTemp(OnComplete));
+		UE_LOG(LogDerivedDataBuild, Error, TEXT("More than one build scheduler memory queue has been registered."));
 	}
 	else
 	{
-		Invoke(OnComplete);
+		GlobalQueue = Queue;
 	}
 }
 
+void FBuildSchedulerMemoryQueue::RemoveQueue(IBuildSchedulerMemoryQueue* Queue)
+{
+	FWriteScopeLock WriteLock(Lock);
+	if (GlobalQueue == Queue)
+	{
+		GlobalQueue = nullptr;
+	}
+}
+
+void FBuildSchedulerMemoryQueue::Reserve(const uint64 Memory, IRequestOwner& Owner, TUniqueFunction<void ()>&& OnComplete)
+{
+	if (FReadScopeLock ReadLock(Lock); IBuildSchedulerMemoryQueue* Queue = GlobalQueue)
+	{
+		return Queue->Reserve(Memory, Owner, MoveTemp(OnComplete));
+	}
+
+	OnComplete();
+}
+
 ///////////////////////////////////////////////////////////////////////////////////////////////////////////////////////
 
 class FBuildJobSchedule final : public IBuildJobSchedule
 {
 public:
-	FBuildJobSchedule(IBuildJob& InJob, IRequestOwner& InOwner, FBuildSchedulerMemoryQueue& InMemoryQueue)
+	FBuildJobSchedule(
+		IBuildJob& InJob,
+		IRequestOwner& InOwner,
+		FBuildSchedulerTypeQueue& InTypeQueue,
+		FBuildSchedulerMemoryQueue& InMemoryQueue)
 		: Job(InJob)
 		, Owner(InOwner)
+		, TypeQueue(InTypeQueue)
 		, MemoryQueue(InMemoryQueue)
 	{
 	}
 
 	FBuildSchedulerParams& EditParameters() final { return Params; }
 
-	void ScheduleCacheQuery() final       { StepSync(); }
+	void ScheduleCacheQuery() final       { QueueForType([this] { StepSync(); }); }
 	void ScheduleCacheStore() final       { StepSync(); }
 
 	void ScheduleResolveKey() final       { StepAsync(TEXT("ResolveKey")); }
@@ -158,20 +235,20 @@ public:
 	{
 		if (Params.MissingRemoteInputsSize)
 		{
-			StepAsync(TEXT("ResolveInputData"));
+			QueueForType([this] { StepAsync(TEXT("ResolveInputData")); });
 		}
 		else
 		{
-			StepAsyncOrQueue(TEXT("ResolveInputData"));
+			QueueForTypeThenMemory([this] { StepAsync(TEXT("ResolveInputData")); });
 		}
 	}
 
-	void ScheduleExecuteRemote() final    { StepAsync(TEXT("ExecuteRemote")); }
-	void ScheduleExecuteLocal() final     { StepAsyncOrQueue(TEXT("ExecuteLocal")); }
+	void ScheduleExecuteRemote() final    { QueueForType([this] { StepAsync(TEXT("ExecuteRemote")); }); }
+	void ScheduleExecuteLocal() final     { QueueForTypeThenMemory([this] { StepAsync(TEXT("ExecuteLocal")); }); }
 
 	void EndJob() final
 	{
-		ExecutionResources = nullptr;
+		MemoryExecutionResources = nullptr;
 	}
 
 private:
@@ -194,36 +271,60 @@ private:
 		// DO NOT ACCESS THIS AGAIN PAST THIS POINT!
 	}
 
-	void StepAsyncOrQueue(const TCHAR* DebugName)
+	void QueueForType(TUniqueFunction<void ()>&& OnComplete)
+	{
+		if (bQueued)
+		{
+			return OnComplete();
+		}
+
+		bQueued = true;
+		TypeQueue.Queue(Params.TypeName, Owner, [this, OnComplete = MoveTemp(OnComplete)]
+		{ 
+			TypeExecutionResources = FExecutionResourceContext::Get();
+			OnComplete();
+		});
+		// DO NOT ACCESS THIS AGAIN PAST THIS POINT!
+	}
+
+	void QueueForMemory(TUniqueFunction<void ()>&& OnComplete)
 	{
 		check(Params.TotalRequiredMemory >= Params.ResolvedInputsSize);
-		const uint64 CurrentMemoryEstimate = Params.TotalRequiredMemory - Params.ResolvedInputsSize;
+		const uint64 CurrentMemory = Params.TotalRequiredMemory - Params.ResolvedInputsSize;
 
-		// Queue for memory only once, the first time it is needed for local execution.
+		// Reserve memory only once, the first time it is needed for local execution.
 		// This will occur either when resolving input data for local execution or before beginning local execution.
 		// NOTE: No attempt is made to reserve memory prior to remote execution, which may become a problem if remote
 		//       execution frequently requires input data to be loaded.
-		if (QueuedMemoryEstimate || CurrentMemoryEstimate == 0)
+		if (ReservedMemory || CurrentMemory == 0)
 		{
-			return StepAsync(DebugName);
+			return OnComplete();
 		}
 
-		QueuedMemoryEstimate = CurrentMemoryEstimate;
-		MemoryQueue.WaitForMemory(Params.TypeName, CurrentMemoryEstimate, Owner, [this, DebugName]
+		ReservedMemory = CurrentMemory;
+		MemoryQueue.Reserve(CurrentMemory, Owner, [this, OnComplete = MoveTemp(OnComplete)]
 		{
-			ExecutionResources = FExecutionResourceContext::Get();
-			StepAsync(DebugName);
+			MemoryExecutionResources = FExecutionResourceContext::Get();
+			OnComplete();
 		});
 		// DO NOT ACCESS THIS AGAIN PAST THIS POINT!
 	}
 
+	void QueueForTypeThenMemory(TUniqueFunction<void ()>&& OnComplete)
+	{
+		QueueForType([this, OnComplete = MoveTemp(OnComplete)]() mutable { QueueForMemory(MoveTemp(OnComplete)); });
+	}
+
 private:
 	IBuildJob& Job;
 	IRequestOwner& Owner;
 	FBuildSchedulerParams Params;
+	FBuildSchedulerTypeQueue& TypeQueue;
 	FBuildSchedulerMemoryQueue& MemoryQueue;
-	TRefCountPtr<IExecutionResource> ExecutionResources;
-	uint64 QueuedMemoryEstimate = 0;
+	TRefCountPtr<IExecutionResource> TypeExecutionResources;
+	TRefCountPtr<IExecutionResource> MemoryExecutionResources;
+	uint64 ReservedMemory = 0;
+	bool bQueued = false;
 };
 
 ///////////////////////////////////////////////////////////////////////////////////////////////////////////////////////
@@ -232,9 +333,10 @@ class FBuildScheduler final : public IBuildScheduler
 {
 	TUniquePtr<IBuildJobSchedule> BeginJob(IBuildJob& Job, IRequestOwner& Owner) final
 	{
-		return MakeUnique<FBuildJobSchedule>(Job, Owner, MemoryQueue);
+		return MakeUnique<FBuildJobSchedule>(Job, Owner, TypeQueue, MemoryQueue);
 	}
 
+	FBuildSchedulerTypeQueue TypeQueue;
 	FBuildSchedulerMemoryQueue MemoryQueue;
 };
 
diff --git a/Engine/Source/Developer/DerivedDataCache/Private/DerivedDataCache.cpp b/Engine/Source/Developer/DerivedDataCache/Private/DerivedDataCache.cpp
--- a/Engine/Source/Developer/DerivedDataCache/Private/DerivedDataCache.cpp
+++ b/Engine/Source/Developer/DerivedDataCache/Private/DerivedDataCache.cpp
@@ -331,14 +331,7 @@ FQueuedThreadPool* GCacheThreadPool;
 
 void LaunchTaskInCacheThreadPool(IRequestOwner& Owner, TUniqueFunction<void ()>&& TaskBody)
 {
-	if (GCacheThreadPool)
-	{
-		LaunchTaskInThreadPool(Owner, *GCacheThreadPool, MoveTemp(TaskBody));
-	}
-	else
-	{
-		TaskBody();
-	}
+	LaunchTaskInThreadPool(Owner, GCacheThreadPool, MoveTemp(TaskBody));
 }
 
 /**
diff --git a/Engine/Source/Developer/DerivedDataCache/Private/DerivedDataThreadPoolTask.cpp b/Engine/Source/Developer/DerivedDataCache/Private/DerivedDataThreadPoolTask.cpp
--- a/Engine/Source/Developer/DerivedDataCache/Private/DerivedDataThreadPoolTask.cpp
+++ b/Engine/Source/Developer/DerivedDataCache/Private/DerivedDataThreadPoolTask.cpp
@@ -117,20 +117,26 @@ void FThreadPoolTaskRequest::Wait()
 void LaunchTaskInThreadPool(
 	uint64 MemoryEstimate,
 	IRequestOwner& Owner,
-	FQueuedThreadPool& ThreadPool,
+	FQueuedThreadPool* ThreadPool,
 	TUniqueFunction<void ()>&& TaskBody)
 {
-	// The request is reference-counted and will be deleted when complete.
-	new FThreadPoolTaskRequest(MemoryEstimate, Owner, ThreadPool, MoveTemp(TaskBody));
+	if (ThreadPool)
+	{
+		// The request is reference-counted and will be deleted when complete.
+		new FThreadPoolTaskRequest(MemoryEstimate, Owner, *ThreadPool, MoveTemp(TaskBody));
+	}
+	else
+	{
+		TaskBody();
+	}
 }
 
 void LaunchTaskInThreadPool(
 	IRequestOwner& Owner,
-	FQueuedThreadPool& ThreadPool,
+	FQueuedThreadPool* ThreadPool,
 	TUniqueFunction<void ()>&& TaskBody)
 {
-	// The request is reference-counted and will be deleted when complete.
-	new FThreadPoolTaskRequest(0, Owner, ThreadPool, MoveTemp(TaskBody));
+	LaunchTaskInThreadPool(0, Owner, ThreadPool, MoveTemp(TaskBody));
 }
 
 } // UE::DerivedData
diff --git a/Engine/Source/Developer/DerivedDataCache/Public/DerivedDataBuildSchedulerQueue.h b/Engine/Source/Developer/DerivedDataCache/Public/DerivedDataBuildSchedulerQueue.h
new file mode 100644
--- /dev/null
+++ b/Engine/Source/Developer/DerivedDataCache/Public/DerivedDataBuildSchedulerQueue.h
@@ -0,0 +1,53 @@
+// Copyright Epic Games, Inc. All Rights Reserved.
+
+#pragma once
+
+#include "DerivedDataSharedStringFwd.h"
+#include "Features/IModularFeature.h"
+#include "UObject/NameTypes.h"
+
+template <typename FuncType> class TUniqueFunction;
+
+namespace UE::DerivedData { class IRequestOwner; }
+
+namespace UE::DerivedData
+{
+
+/**
+ * Allows execution of a build to be queued in a type-specific way.
+ *
+ * At most one instance of this feature may be registered for each type name.
+ * Builds are queued before they query the cache, or equivalent if the cache query is skipped.
+ * The completion callback will capture the FExecutionResourceContext until the end of the build.
+ */
+class IBuildSchedulerTypeQueue : public IModularFeature
+{
+public:
+	static inline const FLazyName FeatureName{"BuildSchedulerTypeQueue"};
+
+	virtual ~IBuildSchedulerTypeQueue() = default;
+
+	/** Returns the type name that this provider corresponds to. */
+	virtual const FUtf8SharedString& GetTypeName() const = 0;
+
+	virtual void Queue(IRequestOwner& Owner, TUniqueFunction<void ()>&& OnComplete) = 0;
+};
+
+/**
+ * Allows execution of a build to be queued until there is sufficient memory available.
+ *
+ * At most one instance of this feature may be registered.
+ * Builds are queued before they load input data for local execution of the build.
+ * The completion callback will capture the FExecutionResourceContext until the end of the build.
+ */
+class IBuildSchedulerMemoryQueue : public IModularFeature
+{
+public:
+	static inline const FLazyName FeatureName{"BuildSchedulerMemoryQueue"};
+
+	virtual ~IBuildSchedulerMemoryQueue() = default;
+
+	virtual void Reserve(uint64 Memory, IRequestOwner& Owner, TUniqueFunction<void ()>&& OnComplete) = 0;
+};
+
+} // UE::DerivedData
diff --git a/Engine/Source/Developer/DerivedDataCache/Public/DerivedDataThreadPoolTask.h b/Engine/Source/Developer/DerivedDataCache/Public/DerivedDataThreadPoolTask.h
--- a/Engine/Source/Developer/DerivedDataCache/Public/DerivedDataThreadPoolTask.h
+++ b/Engine/Source/Developer/DerivedDataCache/Public/DerivedDataThreadPoolTask.h
@@ -20,7 +20,9 @@ namespace UE::DerivedData
  *
  * Waiting on or canceling the task before it has started executing in the thread pool will cause
  * it to be retracted and executed on the thread that asked to wait or cancel.
- * 
+ *
+ * Passing a null thread pool will execute the task immediately on the calling thread.
+ *
  * TaskBody is guaranteed to return before the request is considered complete.
  * TaskBody is responsible for checking for cancellation by checking Owner.IsCanceled().
  *
@@ -32,11 +34,11 @@ namespace UE::DerivedData
 UE_API void LaunchTaskInThreadPool(
 	uint64 MemoryEstimate,
 	IRequestOwner& Owner,
-	FQueuedThreadPool& ThreadPool,
+	FQueuedThreadPool* ThreadPool,
 	TUniqueFunction<void ()>&& TaskBody);
 UE_API void LaunchTaskInThreadPool(
 	IRequestOwner& Owner,
-	FQueuedThreadPool& ThreadPool,
+	FQueuedThreadPool* ThreadPool,
 	TUniqueFunction<void ()>&& TaskBody);
 
 } // UE::DerivedData
