Fixes for UE Cloud DDC client:
- Print CbObject responses to the log/console as JSON, not raw binary
- Reduce the thresholds for connect timeout and slow transfer rate time window
- Instead of one request pool of size 128*2 requests, split GET and PUT operations into their own limited request pools and have the PUT request pool (size 20) be set smaller than the GET one (size 128)
- Stop using unbounded overflow for chained sequences of HTTP requests to avoid slow/failed requests due to many upload operations happening concurrently during high upload quantity scenarios (eg: shader version bump, editor domain version bump)

Testing:
- Cook of CitySample with shader version bumped to ensure there are no failed requests due to slow transfer
- Launch of CitySample editor with no local cache to observe throughput when all DDC data has to be fetched from upstream

diff --git a/Engine/Source/Developer/DerivedDataCache/Private/HttpCacheStore.cpp b/Engine/Source/Developer/DerivedDataCache/Private/HttpCacheStore.cpp
--- a/Engine/Source/Developer/DerivedDataCache/Private/HttpCacheStore.cpp
+++ b/Engine/Source/Developer/DerivedDataCache/Private/HttpCacheStore.cpp
@@ -39,6 +39,7 @@
 #include "ProfilingDebugging/CpuProfilerTrace.h"
 #include "Serialization/CompactBinary.h"
 #include "Serialization/CompactBinaryPackage.h"
+#include "Serialization/CompactBinarySerialization.h"
 #include "Serialization/CompactBinaryValidation.h"
 #include "Serialization/CompactBinaryWriter.h"
 #include "Serialization/JsonReader.h"
@@ -65,7 +66,8 @@
 
 #define UE_HTTPDDC_GET_REQUEST_POOL_SIZE 48
 #define UE_HTTPDDC_PUT_REQUEST_POOL_SIZE 16
-#define UE_HTTPDDC_NONBLOCKING_REQUEST_POOL_SIZE 128
+#define UE_HTTPDDC_NONBLOCKING_GET_REQUEST_POOL_SIZE 128
+#define UE_HTTPDDC_NONBLOCKING_PUT_REQUEST_POOL_SIZE 24
 #define UE_HTTPDDC_MAX_FAILED_LOGIN_ATTEMPTS 16
 #define UE_HTTPDDC_MAX_ATTEMPTS 4
 
@@ -311,7 +313,8 @@ private:
 	THttpUniquePtr<IHttpConnectionPool> ConnectionPool;
 	FHttpRequestQueue GetRequestQueues[2];
 	FHttpRequestQueue PutRequestQueues[2];
-	FHttpRequestQueue NonBlockingRequestQueue;
+	FHttpRequestQueue NonBlockingGetRequestQueue;
+	FHttpRequestQueue NonBlockingPutRequestQueue;
 
 	FCriticalSection AccessCs;
 	TUniquePtr<FHttpAccessToken> Access;
@@ -342,7 +345,7 @@ private:
 
 	class FHttpOperation;
 
-	TUniquePtr<FHttpOperation> WaitForHttpOperation(EOperationCategory Category, bool bUnboundedOverflow);
+	TUniquePtr<FHttpOperation> WaitForHttpOperation(EOperationCategory Category);
 
 	struct FGetCacheRecordOnlyResponse
 	{
@@ -642,6 +645,16 @@ FString FHttpCacheStore::FHttpOperation::GetBodyAsString() const
 {
 	static_assert(sizeof(uint8) == sizeof(UTF8CHAR));
 	const int32 Len = IntCastChecked<int32>(ResponseBody.GetSize());
+	if (GetContentType() == EHttpMediaType::CbObject)
+	{
+		if (ValidateCompactBinary(ResponseBody, ECbValidateMode::Default) != ECbValidateError::None)
+		{
+			TUtf8StringBuilder<1024> JsonStringBuilder;
+			const FCbObject ResponseObject(ResponseBody);
+			CompactBinaryToCompactJson(ResponseObject, JsonStringBuilder);
+			return JsonStringBuilder.ToString();
+		}
+	}
 	return FString(Len, (const UTF8CHAR*)ResponseBody.GetData());
 }
 
@@ -943,7 +956,7 @@ void FHttpCacheStore::FPutPackageOp::PutRefAsync(
 		RefsUri << ANSITEXTVIEW("/finalize/") << ObjectHash;
 	}
 
-	TUniquePtr<FHttpOperation> Operation = CacheStore.WaitForHttpOperation(EOperationCategory::Put, /*bUnboundedOverflow*/ bFinalize);
+	TUniquePtr<FHttpOperation> Operation = CacheStore.WaitForHttpOperation(EOperationCategory::Put);
 	FHttpOperation& LocalOperation = *Operation;
 	LocalOperation.SetUri(RefsUri);
 	if (bFinalize)
@@ -1110,7 +1123,7 @@ void FHttpCacheStore::FPutPackageOp::OnPackagePutRefComplete(
 	FRequestBarrier Barrier(Owner);
 	for (const FCompressedBlobUpload& CompressedBlobUpload : CompressedBlobUploads)
 	{
-		TUniquePtr<FHttpOperation> Operation = CacheStore.WaitForHttpOperation(EOperationCategory::Put, /*bUnboundedOverflow*/ true);
+		TUniquePtr<FHttpOperation> Operation = CacheStore.WaitForHttpOperation(EOperationCategory::Put);
 		FHttpOperation& LocalOperation = *Operation;
 		LocalOperation.SetUri(WriteToAnsiString<256>(CacheStore.EffectiveDomain, ANSITEXTVIEW("/api/v1/compressed-blobs/"), CacheStore.Namespace, '/', CompressedBlobUpload.Hash));
 		LocalOperation.SetMethod(EHttpMethod::Put);
@@ -1207,7 +1220,7 @@ void FHttpCacheStore::FGetRecordOp::GetDataBatch(
 	for (int32 ValueIndex = 0; ValueIndex < Values.Num(); ++ValueIndex)
 	{
 		const ValueType Value = Values[ValueIndex].RemoveData();
-		TUniquePtr<FHttpOperation> Operation = CacheStore.WaitForHttpOperation(EOperationCategory::Get, /*bUnboundedOverflow*/ true);
+		TUniquePtr<FHttpOperation> Operation = CacheStore.WaitForHttpOperation(EOperationCategory::Get);
 		FHttpOperation& LocalOperation = *Operation;
 		LocalOperation.SetUri(WriteToAnsiString<256>(CacheStore.EffectiveDomain, ANSITEXTVIEW("/api/v1/compressed-blobs/"), CacheStore.Namespace, '/', Value.GetRawHash()));
 		LocalOperation.SetMethod(EHttpMethod::Get);
@@ -1402,7 +1415,7 @@ void FHttpCacheStore::FGetRecordOp::DataProbablyExistsBatch(
 		bFirstItem = false;
 	}
 
-	TUniquePtr<FHttpOperation> Operation = CacheStore.WaitForHttpOperation(EOperationCategory::Get, /*bUnboundedOverflow*/ true);
+	TUniquePtr<FHttpOperation> Operation = CacheStore.WaitForHttpOperation(EOperationCategory::Get);
 	FHttpOperation& LocalOperation = *Operation;
 	LocalOperation.SetUri(CompressedBlobsUri);
 	LocalOperation.SetMethod(EHttpMethod::Post);
@@ -1594,9 +1607,13 @@ FHttpCacheStore::FHttpCacheStore(const FHttpCacheStoreParams& Params)
 		PutRequestQueues[0] = FHttpRequestQueue(*ConnectionPool, ClientParams);
 		PutRequestQueues[1] = FHttpRequestQueue(*ConnectionPool, ClientParams);
 
-		ClientParams.MaxRequests = UE_HTTPDDC_NONBLOCKING_REQUEST_POOL_SIZE * 2;
-		ClientParams.MinRequests = UE_HTTPDDC_NONBLOCKING_REQUEST_POOL_SIZE;
-		NonBlockingRequestQueue = FHttpRequestQueue(*ConnectionPool, ClientParams);
+		ClientParams.MaxRequests = UE_HTTPDDC_NONBLOCKING_GET_REQUEST_POOL_SIZE;
+		ClientParams.MinRequests = UE_HTTPDDC_NONBLOCKING_GET_REQUEST_POOL_SIZE;
+		NonBlockingGetRequestQueue = FHttpRequestQueue(*ConnectionPool, ClientParams);
+
+		ClientParams.MaxRequests = UE_HTTPDDC_NONBLOCKING_PUT_REQUEST_POOL_SIZE;
+		ClientParams.MinRequests = UE_HTTPDDC_NONBLOCKING_PUT_REQUEST_POOL_SIZE;
+		NonBlockingPutRequestQueue = FHttpRequestQueue(*ConnectionPool, ClientParams);
 
 		bIsUsable = true;
 	}
@@ -1659,10 +1676,10 @@ FHttpCacheStore::~FHttpCacheStore()
 FHttpClientParams FHttpCacheStore::GetDefaultClientParams() const
 {
 	FHttpClientParams ClientParams;
-	ClientParams.DnsCacheTimeout = 300;
-	ClientParams.ConnectTimeout = 30 * 1000;
+	ClientParams.DnsCacheTimeout = 15;
+	ClientParams.ConnectTimeout = 3 * 1000;
 	ClientParams.LowSpeedLimit = 1024;
-	ClientParams.LowSpeedTime = 30;
+	ClientParams.LowSpeedTime = 10;
 	ClientParams.TlsLevel = EHttpTlsLevel::All;
 	ClientParams.bFollowRedirects = true;
 	ClientParams.bFollow302Post = true;
@@ -1838,7 +1855,7 @@ void FHttpCacheStore::SetAccessToken(FStringView Token, double RefreshDelay)
 	FailedLoginAttempts = 0;
 }
 
-TUniquePtr<FHttpCacheStore::FHttpOperation> FHttpCacheStore::WaitForHttpOperation(EOperationCategory Category, bool bUnboundedOverflow)
+TUniquePtr<FHttpCacheStore::FHttpOperation> FHttpCacheStore::WaitForHttpOperation(EOperationCategory Category)
 {
 	if (Access && RefreshAccessTokenTime > 0.0 && RefreshAccessTokenTime < FPlatformTime::Seconds())
 	{
@@ -1850,8 +1867,14 @@ TUniquePtr<FHttpCacheStore::FHttpOperation> FHttpCacheStore::WaitForHttpOperatio
 	FHttpRequestParams Params;
 	if (FPlatformProcess::SupportsMultithreading() && bHttpEnableAsync)
 	{
-		Params.bIgnoreMaxRequests = bUnboundedOverflow;
-		Request = NonBlockingRequestQueue.CreateRequest(Params);
+		if (Category == EOperationCategory::Get)
+		{
+			Request = NonBlockingGetRequestQueue.CreateRequest(Params);
+		}
+		else
+		{
+			Request = NonBlockingPutRequestQueue.CreateRequest(Params);
+		}
 	}
 	else
 	{
@@ -1913,7 +1936,7 @@ void FHttpCacheStore::GetCacheRecordOnlyAsync(
 	TAnsiStringBuilder<64> Bucket;
 	Algo::Transform(Key.Bucket.ToString(), AppendChars(Bucket), FCharAnsi::ToLower);
 
-	TUniquePtr<FHttpOperation> Operation = WaitForHttpOperation(EOperationCategory::Get, /*bUnboundedOverflow*/ false);
+	TUniquePtr<FHttpOperation> Operation = WaitForHttpOperation(EOperationCategory::Get);
 	FHttpOperation& LocalOperation = *Operation;
 	LocalOperation.SetUri(WriteToAnsiString<256>(EffectiveDomain, ANSITEXTVIEW("/api/v1/refs/"), Namespace, '/', Bucket, '/', Key.Hash));
 	LocalOperation.SetMethod(EHttpMethod::Get);
@@ -2097,7 +2120,7 @@ void FHttpCacheStore::GetCacheValueAsync(
 	TAnsiStringBuilder<64> Bucket;
 	Algo::Transform(Key.Bucket.ToString(), AppendChars(Bucket), FCharAnsi::ToLower);
 
-	TUniquePtr<FHttpOperation> Operation = WaitForHttpOperation(EOperationCategory::Get, /*bUnboundedOverflow*/ false);
+	TUniquePtr<FHttpOperation> Operation = WaitForHttpOperation(EOperationCategory::Get);
 	FHttpOperation& LocalOperation = *Operation;
 	LocalOperation.SetUri(WriteToAnsiString<256>(EffectiveDomain, ANSITEXTVIEW("/api/v1/refs/"), Namespace, '/', Bucket, '/', Key.Hash));
 	LocalOperation.SetMethod(EHttpMethod::Get);
@@ -2232,7 +2255,7 @@ void FHttpCacheStore::RefCachedDataProbablyExistsBatchAsync(
 	RequestWriter.EndObject();
 	FCbFieldIterator RequestFields = RequestWriter.Save();
 
-	TUniquePtr<FHttpOperation> Operation = WaitForHttpOperation(EOperationCategory::Get, /*bUnboundedOverflow*/ false);
+	TUniquePtr<FHttpOperation> Operation = WaitForHttpOperation(EOperationCategory::Get);
 	FHttpOperation& LocalOperation = *Operation;
 	LocalOperation.SetUri(WriteToAnsiString<256>(EffectiveDomain, ANSITEXTVIEW("/api/v1/refs/"), Namespace));
 	LocalOperation.SetMethod(EHttpMethod::Post);
