Ensure that Zen as a local service can be configured to follow the same commandline overrides as a filesystem cache store.  Also ensure that if the configured override mechanisms for cache store paths yield a path of "None" that we disable the Zen DDC cache.  These are in support of using Zen as a drop-in replacement for an existing filesystem DDC store.

diff --git a/Engine/Config/BaseEngine.ini b/Engine/Config/BaseEngine.ini
--- a/Engine/Config/BaseEngine.ini
+++ b/Engine/Config/BaseEngine.ini
@@ -2206,6 +2206,7 @@ AutoLaunch=true
 [Zen.AutoLaunch]
 DataPath=%APPSETTINGSDIR%Zen/Data
 LocalDataCachePathEnvOverride=UE-LocalDataCachePath
+LocalDataCachePathCommandLineOverride=LocalDataCachePath
 LocalDataCachePathEditorOverrideSetting=LocalDerivedDataCache
 ExtraArgs=--http asio --gc-cache-duration-seconds 2937600 --gc-interval-seconds 21600
 
diff --git a/Engine/Source/Developer/DerivedDataCache/Private/ZenCacheStore.cpp b/Engine/Source/Developer/DerivedDataCache/Private/ZenCacheStore.cpp
--- a/Engine/Source/Developer/DerivedDataCache/Private/ZenCacheStore.cpp
+++ b/Engine/Source/Developer/DerivedDataCache/Private/ZenCacheStore.cpp
@@ -122,8 +122,6 @@ private:
 
 	static FCompositeBuffer SaveRpcPackage(const FCbPackage& Package);
 	THttpUniquePtr<IHttpRequest> CreateRpcRequest();
-	THttpUniquePtr<IHttpResponse> PerformBlockingRpc(FCbObject RequestObject, FCbPackage& OutResponse);
-	THttpUniquePtr<IHttpResponse> PerformBlockingRpc(const FCbPackage& RequestPackage, FCbPackage& OutResponse);
 	using FOnRpcComplete = TUniqueFunction<void(THttpUniquePtr<IHttpResponse>& HttpResponse, FCbPackage& Response)>;
 	void EnqueueAsyncRpc(IRequestOwner& Owner, FCbObject RequestObject, FOnRpcComplete&& OnComplete);
 	void EnqueueAsyncRpc(IRequestOwner& Owner, const FCbPackage& RequestPackage, FOnRpcComplete&& OnComplete);
@@ -1065,30 +1063,6 @@ THttpUniquePtr<IHttpRequest> FZenCacheStore::CreateRpcRequest()
 	return Request;
 }
 
-THttpUniquePtr<IHttpResponse> FZenCacheStore::PerformBlockingRpc(FCbObject RequestObject, FCbPackage& OutResponse)
-{
-	THttpUniquePtr<IHttpRequest> Request = CreateRpcRequest();
-	Request->SetContentType(EHttpMediaType::CbObject);
-	Request->SetBody(RequestObject.GetBuffer().MakeOwned());
-
-	FCbPackageReceiver Receiver(OutResponse);
-	THttpUniquePtr<IHttpResponse> Response;
-	Request->Send(&Receiver, Response);
-	return Response;
-}
-
-THttpUniquePtr<IHttpResponse> FZenCacheStore::PerformBlockingRpc(const FCbPackage& RequestPackage, FCbPackage& OutResponse)
-{
-	THttpUniquePtr<IHttpRequest> Request = CreateRpcRequest();
-	Request->SetContentType(EHttpMediaType::CbPackage);
-	Request->SetBody(SaveRpcPackage(RequestPackage));
-
-	FCbPackageReceiver Receiver(OutResponse);
-	THttpUniquePtr<IHttpResponse> Response;
-	Request->Send(&Receiver, Response);
-	return Response;
-}
-
 void FZenCacheStore::EnqueueAsyncRpc(IRequestOwner& Owner, FCbObject RequestObject, FOnRpcComplete&& OnComplete)
 {
 	THttpUniquePtr<IHttpRequest> Request = CreateRpcRequest();
@@ -1224,6 +1198,16 @@ TTuple<ILegacyCacheStore*, ECacheStoreFlags> CreateZenCacheStore(const TCHAR* No
 		UE_LOG(LogDerivedDataCache, Warning, TEXT("%s: Missing required parameter 'Namespace', falling back to '%s'"), NodeName, *Namespace);
 	}
 
+	FString CachePathOverride;
+	if (UE::Zen::Private::IsLocalAutoLaunched(ServiceUrl) && UE::Zen::Private::GetLocalDataCachePathOverride(CachePathOverride))
+	{
+		if (CachePathOverride == TEXT("None"))
+		{
+			UE_LOG( LogDerivedDataCache, Log, TEXT("Disabling %s data cache - path set to 'None'."), NodeName );
+			return MakeTuple<ILegacyCacheStore*, ECacheStoreFlags>(nullptr, ECacheStoreFlags::None);
+		}
+	}
+
 	TUniquePtr<FZenCacheStore> Backend = MakeUnique<FZenCacheStore>(*ServiceUrl, *Namespace);
 	if (!Backend->IsUsable())
 	{
diff --git a/Engine/Source/Developer/Zen/Private/ZenServerInterface.cpp b/Engine/Source/Developer/Zen/Private/ZenServerInterface.cpp
--- a/Engine/Source/Developer/Zen/Private/ZenServerInterface.cpp
+++ b/Engine/Source/Developer/Zen/Private/ZenServerInterface.cpp
@@ -184,6 +184,23 @@ DetermineLocalDataCachePath(const TCHAR* ConfigSection, FString& DataPath)
 		}
 	}
 
+	FString DataPathCommandLineOverride;
+	if (GConfig->GetString(ConfigSection, TEXT("LocalDataCachePathCommandLineOverride"), DataPathCommandLineOverride, GEngineIni))
+	{
+		FString DataPathCommandLineOverrideValue;
+		if (FParse::Value(FCommandLine::Get(), *(DataPathCommandLineOverride + TEXT("=")), DataPathCommandLineOverrideValue))
+		{
+			DataPath = DataPathCommandLineOverrideValue;
+			UE_LOG(LogZenServiceInstance, Log, TEXT("Found command line override %s=%s"), *DataPathCommandLineOverride, *DataPath);
+		}
+	}
+
+	// Paths starting with a '?' are looked up from config
+	if (DataPath.StartsWith(TEXT("?")) && !GConfig->GetString(TEXT("DerivedDataCacheSettings"), *DataPath + 1, DataPath, GEngineIni))
+	{
+		DataPath.Empty();
+	}
+
 	FString DataPathEditorOverrideSetting;
 	if (GConfig->GetString(ConfigSection, TEXT("LocalDataCachePathEditorOverrideSetting"), DataPathEditorOverrideSetting, GEngineIni))
 	{
@@ -734,6 +751,31 @@ DetermineCmdLineWithoutTransientComponents(const FServiceAutoLaunchSettings& InS
 	return Parms;
 }
 
+bool
+Private::IsLocalAutoLaunched(FStringView InstanceURL)
+{
+	if (!InstanceURL.IsEmpty() && !InstanceURL.Equals(TEXT("<DefaultInstance>")))
+	{
+		FString TempURL(InstanceURL);
+		return IsLocalHost(TempURL);
+	}
+	return true;
+}
+
+bool
+Private::GetLocalDataCachePathOverride(FString& OutDataPath)
+{
+	const TCHAR* AutoLaunchConfigSection = TEXT("Zen.AutoLaunch");
+	FString DataPath;
+	DetermineLocalDataCachePath(AutoLaunchConfigSection, DataPath);
+	if (DataPath.IsEmpty())
+	{
+		return false;
+	}
+	OutDataPath = DataPath;
+	return true;
+}
+
 bool
 TryGetLocalServiceRunContext(FZenLocalServiceRunContext& OutContext)
 {
diff --git a/Engine/Source/Developer/Zen/Public/Experimental/ZenServerInterface.h b/Engine/Source/Developer/Zen/Public/Experimental/ZenServerInterface.h
--- a/Engine/Source/Developer/Zen/Public/Experimental/ZenServerInterface.h
+++ b/Engine/Source/Developer/Zen/Public/Experimental/ZenServerInterface.h
@@ -67,6 +67,12 @@ private:
 namespace UE::Zen
 {
 
+namespace Private
+{
+	UE_API bool IsLocalAutoLaunched(FStringView InstanceURL);
+	UE_API bool GetLocalDataCachePathOverride(FString& OutDataPath);
+}
+
 class FZenServiceInstance;
 class FZenLocalServiceRunContext;
 
