Add support for Zen DDC sandboxes which are separate instances of zenserver with separate data paths and can be optionally flushed from run-to-run.  It is used to implement Warm and Cold DDC graphs using Zen sandboxes for projects/platforms that use Zen for local DDC by default, and to provide new ZenColdDDC and ZenWarmDDC as new graphs that have Cold/Warm behavior for projects that don't use Zen for local DDC by default.

diff --git a/Engine/Config/BaseEngine.ini b/Engine/Config/BaseEngine.ini
--- a/Engine/Config/BaseEngine.ini
+++ b/Engine/Config/BaseEngine.ini
@@ -2305,6 +2305,18 @@ Cloud=(Type=Cloud, ServerID=Default, EnvHostOverride=UE-CloudDataCacheHost, Comm
 Pak=(Type=ReadPak, Filename="%GAMEDIR%DerivedDataCache/DDC.ddp")
 EnginePak=(Type=ReadPak, Filename="%ENGINEDIR%DerivedDataCache/DDC.ddp")
 
+[ZenColdDDC]
+Root=(Type=KeyLength, Length=120, Inner=AsyncPut)
+AsyncPut=(Type=AsyncPut, Inner=Hierarchy)
+Hierarchy=(Type=Hierarchical, Inner=ZenLocal)
+ZenLocal=(Type=Zen, Namespace="ue.ddc", Sandbox=Cold, Flush=true)
+
+[ZenWarmDDC]
+Root=(Type=KeyLength, Length=120, Inner=AsyncPut)
+AsyncPut=(Type=AsyncPut, Inner=Hierarchy)
+Hierarchy=(Type=Hierarchical, Inner=ZenLocal)
+ZenLocal=(Type=Zen, Namespace="ue.ddc", Sandbox=Cold)
+
 [DerivedDataBackendGraph_Fill_Seattle]
 Root=(Type=KeyLength, Length=120, Inner=AsyncPut)
 AsyncPut=(Type=AsyncPut, Inner=Hierarchy)
diff --git a/Engine/Source/Developer/DerivedDataCache/Private/ZenCacheStore.cpp b/Engine/Source/Developer/DerivedDataCache/Private/ZenCacheStore.cpp
--- a/Engine/Source/Developer/DerivedDataCache/Private/ZenCacheStore.cpp
+++ b/Engine/Source/Developer/DerivedDataCache/Private/ZenCacheStore.cpp
@@ -14,6 +14,7 @@
 #include "DerivedDataRequest.h"
 #include "DerivedDataRequestOwner.h"
 #include "Experimental/ZenStatistics.h"
+#include "HAL/FileManager.h"
 #include "Http/HttpClient.h"
 #include "Math/UnrealMathUtility.h"
 #include "Misc/App.h"
@@ -77,6 +78,9 @@ public:
 	 */
 	FZenCacheStore(const TCHAR* ServiceUrl, const TCHAR* Namespace);
 
+
+	FZenCacheStore(UE::Zen::FServiceSettings&& InSettings, const TCHAR* Namespace);
+
 	inline FString GetName() const { return ZenService.GetInstance().GetURL(); }
 
 	/**
@@ -118,6 +122,8 @@ public:
 	bool LegacyDebugOptions(FBackendDebugOptions& Options) final;
 
 private:
+	void Initialize(const TCHAR* Namespace);
+
 	bool IsServiceReady();
 
 	static FCompositeBuffer SaveRpcPackage(const FCbPackage& Package);
@@ -1033,9 +1039,23 @@ private:
 FZenCacheStore::FZenCacheStore(
 	const TCHAR* InServiceUrl,
 	const TCHAR* InNamespace)
-	: Namespace(InNamespace)
-	, ZenService(InServiceUrl)
+	: ZenService(InServiceUrl)
+{
+	Initialize(InNamespace);
+}
+
+FZenCacheStore::FZenCacheStore(
+	UE::Zen::FServiceSettings&& InSettings,
+	const TCHAR* InNamespace)
+	: ZenService(MoveTemp(InSettings))
+{
+	Initialize(InNamespace);
+}
+
+void FZenCacheStore::Initialize(
+	const TCHAR* InNamespace)
 {
+	Namespace = InNamespace;
 	if (IsServiceReady())
 	{
 		RpcUri << ZenService.GetInstance().GetURL() << ANSITEXTVIEW("/z$/$rpc");
@@ -1116,7 +1136,7 @@ void FZenCacheStore::LegacyStats(FDerivedDataCacheStatsNode& OutNode)
 	using EHitOrMiss = FCookStats::CallStats::EHitOrMiss;
 	using ECacheStatType = FCookStats::CallStats::EStatType;
 
-	Zen::GetDefaultServiceInstance().GetStats(ZenStats);
+	ZenService.GetInstance().GetStats(ZenStats);
 
 	const int64 RemotePutSize = int64(ZenStats.UpstreamStats.TotalUploadedMB * 1024 * 1024);
 	const int64 RemoteGetSize = int64(ZenStats.UpstreamStats.TotalDownloadedMB * 1024 * 1024);
@@ -1229,12 +1249,73 @@ TTuple<ILegacyCacheStore*, ECacheStoreFlags> CreateZenCacheStore(const TCHAR* No
 	{
 		if (CachePathOverride == TEXT("None"))
 		{
-			UE_LOG( LogDerivedDataCache, Log, TEXT("Disabling %s data cache - path set to 'None'."), NodeName );
+			UE_LOG(LogDerivedDataCache, Log, TEXT("Disabling %s data cache - path set to 'None'."), NodeName);
 			return MakeTuple<ILegacyCacheStore*, ECacheStoreFlags>(nullptr, ECacheStoreFlags::None);
 		}
 	}
 
-	TUniquePtr<FZenCacheStore> Backend = MakeUnique<FZenCacheStore>(*ServiceUrl, *Namespace);
+	TUniquePtr<FZenCacheStore> Backend;
+
+	FString Sandbox;
+	FParse::Value(Config, TEXT("Sandbox="), Sandbox);
+
+	bool bFlush = false;
+	FParse::Bool(Config, TEXT("Flush="), bFlush);
+
+	if (!Sandbox.IsEmpty())
+	{
+		Zen::FServiceSettings DefaultServiceSettings;
+		DefaultServiceSettings.ReadFromConfig();
+
+		if (!DefaultServiceSettings.IsAutoLaunch())
+		{
+			UE_LOG(LogDerivedDataCache, Warning, TEXT("%s: Attempting to use a sandbox when there is no default autolaunch configured to interhit settings from.  Cache will be disabled."), NodeName);
+			return MakeTuple<ILegacyCacheStore*, ECacheStoreFlags>(nullptr, ECacheStoreFlags::None);
+		}
+
+		// Make a unique local instance (not the default local instance) of ZenServer
+		Zen::FServiceSettings ServiceSettings;
+		ServiceSettings.SettingsVariant.Emplace<Zen::FServiceAutoLaunchSettings>();
+		Zen::FServiceAutoLaunchSettings& AutoLaunchSettings = ServiceSettings.SettingsVariant.Get<Zen::FServiceAutoLaunchSettings>();
+
+		const Zen::FServiceAutoLaunchSettings& DefaultAutoLaunchSettings = DefaultServiceSettings.SettingsVariant.Get<Zen::FServiceAutoLaunchSettings>();
+		AutoLaunchSettings = DefaultAutoLaunchSettings;
+		// Default as one more than the default port to not collide.  Multiple sandboxes will share a desired port, but will get differing effective ports.
+		AutoLaunchSettings.DesiredPort++;
+
+		FPaths::NormalizeDirectoryName(AutoLaunchSettings.DataPath);
+		AutoLaunchSettings.DataPath += TEXT("_");
+		AutoLaunchSettings.DataPath += Sandbox;
+
+		// The unique local instances will always limit process lifetime for now to avoid accumulating many of them
+		AutoLaunchSettings.bLimitProcessLifetime = true;
+
+		// Flush the cache if requested.
+		if (bFlush)
+		{
+			bool bStopped = true;
+			if (UE::Zen::IsLocalServiceRunning(*AutoLaunchSettings.DataPath))
+			{
+				bStopped = UE::Zen::StopLocalService(*AutoLaunchSettings.DataPath);
+			}
+
+			if (bStopped)
+			{
+				IFileManager::Get().DeleteDirectory(*(AutoLaunchSettings.DataPath / TEXT("")), /*bRequireExists*/ false, /*bTree*/ true);
+			}
+			else
+			{
+				UE_LOG(LogDerivedDataCache, Warning, TEXT("%s: Zen DDC could not be flushed due to an existing instance not shutting down when requested."), NodeName);
+			}
+		}
+
+		Backend = MakeUnique<FZenCacheStore>(MoveTemp(ServiceSettings), *Namespace);
+	}
+	else
+	{
+		Backend = MakeUnique<FZenCacheStore>(*ServiceUrl, *Namespace);
+	}
+
 	if (!Backend->IsUsable())
 	{
 		UE_LOG(LogDerivedDataCache, Warning, TEXT("%s: Failed to contact the service (%s), will not use it."), NodeName, *Backend->GetName());
diff --git a/Engine/Source/Developer/Zen/Private/ZenServerInterface.cpp b/Engine/Source/Developer/Zen/Private/ZenServerInterface.cpp
--- a/Engine/Source/Developer/Zen/Private/ZenServerInterface.cpp
+++ b/Engine/Source/Developer/Zen/Private/ZenServerInterface.cpp
@@ -363,9 +363,6 @@ FServiceSettings::ReadFromConfig()
 			ReadUInt16FromConfig(AutoLaunchConfigSection, TEXT("DesiredPort"), AutoLaunchSettings.DesiredPort, GEngineIni);
 			GConfig->GetBool(AutoLaunchConfigSection, TEXT("ShowConsole"), AutoLaunchSettings.bShowConsole, GEngineIni);
 			GConfig->GetBool(AutoLaunchConfigSection, TEXT("LimitProcessLifetime"), AutoLaunchSettings.bLimitProcessLifetime, GEngineIni);
-
-			// Ensure that the zen data path is inherited by subprocesses
-			FPlatformMisc::SetEnvironmentVar(TEXT("UE-ZenSubprocessDataPath"), *AutoLaunchSettings.DataPath);
 		}
 	}
 	else
@@ -1062,6 +1059,11 @@ FZenServiceInstance::FZenServiceInstance(FStringView InstanceURL)
 	if (InstanceURL.IsEmpty())
 	{
 		Settings.ReadFromConfig();
+		if (Settings.IsAutoLaunch())
+		{
+			// Ensure that the zen data path is inherited by subprocesses
+			FPlatformMisc::SetEnvironmentVar(TEXT("UE-ZenSubprocessDataPath"), *Settings.SettingsVariant.Get<FServiceAutoLaunchSettings>().DataPath);
+		}
 	}
 	else
 	{
