HttpCacheStore: Fix backwards if statement for whether the compactbinary is valid in GetBodyAsString

diff --git a/Engine/Source/Developer/DerivedDataCache/Private/HttpCacheStore.cpp b/Engine/Source/Developer/DerivedDataCache/Private/HttpCacheStore.cpp
--- a/Engine/Source/Developer/DerivedDataCache/Private/HttpCacheStore.cpp
+++ b/Engine/Source/Developer/DerivedDataCache/Private/HttpCacheStore.cpp
@@ -643,7 +643,7 @@ FString FHttpCacheStore::FHttpOperation::GetBodyAsString() const
 	const int32 Len = IntCastChecked<int32>(ResponseBody.GetSize());
 	if (GetContentType() == EHttpMediaType::CbObject)
 	{
-		if (ValidateCompactBinary(ResponseBody, ECbValidateMode::Default) != ECbValidateError::None)
+		if (ValidateCompactBinary(ResponseBody, ECbValidateMode::Default) == ECbValidateError::None)
 		{
 			TUtf8StringBuilder<1024> JsonStringBuilder;
 			const FCbObject ResponseObject(ResponseBody);
