Allow ReadOnly mode for a DDC FileSystemCacheStore to have maintenance performed if Clean=true or DeleteUnused=true.  When ReadOnly, they are now both false by default, but usable if explicitly set to true.  When not ReadOnly, DeleteUnused remains true by default.

diff --git a/Engine/Source/Developer/DerivedDataCache/Private/FileSystemCacheStore.cpp b/Engine/Source/Developer/DerivedDataCache/Private/FileSystemCacheStore.cpp
--- a/Engine/Source/Developer/DerivedDataCache/Private/FileSystemCacheStore.cpp
+++ b/Engine/Source/Developer/DerivedDataCache/Private/FileSystemCacheStore.cpp
@@ -1077,81 +1077,79 @@ FFileSystemCacheStore::FFileSystemCacheStore(
 			//bReadOnly = true;
 		}
 
-		if (!bReadOnly)
+
+		if (FString(FCommandLine::Get()).Contains(TEXT("Run=DerivedDataCache")))
 		{
-			if (FString(FCommandLine::Get()).Contains(TEXT("Run=DerivedDataCache")))
-			{
-				bTouch = true; // we always touch files when running the DDC commandlet
-			}
+			bTouch = true; // we always touch files when running the DDC commandlet
+		}
 
-			// The command line (-ddctouch) enables touch on all filesystem backends if specified.
-			bTouch = bTouch || FParse::Param(FCommandLine::Get(), TEXT("DDCTOUCH"));
+		// The command line (-ddctouch) enables touch on all filesystem backends if specified.
+		bTouch = bTouch || FParse::Param(FCommandLine::Get(), TEXT("DDCTOUCH"));
 
-			if (bTouch)
-			{
-				UE_LOG(LogDerivedDataCache, Display, TEXT("Files in %s will be touched."), *CachePath);
-			}
+		if (bTouch)
+		{
+			UE_LOG(LogDerivedDataCache, Display, TEXT("Files in %s will be touched."), *CachePath);
+		}
+		
+		bool bClean = false;
+		bool bDeleteUnused = !bReadOnly;
+		FParse::Bool(InParams, TEXT("Clean="), bClean);
+		FParse::Bool(InParams, TEXT("DeleteUnused="), bDeleteUnused);
+		bDeleteUnused = bDeleteUnused && !FParse::Param(FCommandLine::Get(), TEXT("NODDCCLEANUP"));
 
-			bool bClean = false;
-			bool bDeleteUnused = true;
-			FParse::Bool(InParams, TEXT("Clean="), bClean);
-			FParse::Bool(InParams, TEXT("DeleteUnused="), bDeleteUnused);
-			bDeleteUnused = bDeleteUnused && !FParse::Param(FCommandLine::Get(), TEXT("NODDCCLEANUP"));
+		if (bClean && bLocalDeactivatedForPerformance)
+		{
+			bDeactivationDeferredClean = true;
+		}
 
-			if (bClean && bLocalDeactivatedForPerformance)
+		if (bClean || bDeleteUnused)
+		{
+			FFileSystemCacheStoreMaintainerParams* MaintainerParams;
+			FFileSystemCacheStoreMaintainerParams LocalMaintainerParams;
+			if (bLocalDeactivatedForPerformance)
 			{
-				bDeactivationDeferredClean = true;
+				DeactivationDeferredMaintainerParams = MakeUnique<FFileSystemCacheStoreMaintainerParams>();
+				MaintainerParams = DeactivationDeferredMaintainerParams.Get();
 			}
-
-			if (bClean || bDeleteUnused)
+			else
 			{
-				FFileSystemCacheStoreMaintainerParams* MaintainerParams;
-				FFileSystemCacheStoreMaintainerParams LocalMaintainerParams;
-				if (bLocalDeactivatedForPerformance)
-				{
-					DeactivationDeferredMaintainerParams = MakeUnique<FFileSystemCacheStoreMaintainerParams>();
-					MaintainerParams = DeactivationDeferredMaintainerParams.Get();
-				}
-				else
-				{
-					MaintainerParams = &LocalMaintainerParams;
-				}
-				MaintainerParams->MaxFileAge = FTimespan::FromDays(DaysToDeleteUnusedFiles);
-				if (bDeleteUnused)
+				MaintainerParams = &LocalMaintainerParams;
+			}
+			MaintainerParams->MaxFileAge = FTimespan::FromDays(DaysToDeleteUnusedFiles);
+			if (bDeleteUnused)
+			{
+				if (!FParse::Value(InParams, TEXT("MaxFileChecksPerSec="), MaintainerParams->MaxScanRate))
 				{
-					if (!FParse::Value(InParams, TEXT("MaxFileChecksPerSec="), MaintainerParams->MaxScanRate))
+					int32 MaxFileScanRate;
+					if (GConfig->GetInt(TEXT("DDCCleanup"), TEXT("MaxFileChecksPerSec"), MaxFileScanRate, GEngineIni))
 					{
-						int32 MaxFileScanRate;
-						if (GConfig->GetInt(TEXT("DDCCleanup"), TEXT("MaxFileChecksPerSec"), MaxFileScanRate, GEngineIni))
-						{
-							MaintainerParams->MaxScanRate = uint32(MaxFileScanRate);
-						}
+						MaintainerParams->MaxScanRate = uint32(MaxFileScanRate);
 					}
-					FParse::Value(InParams, TEXT("FoldersToClean="), MaintainerParams->MaxDirectoryScanCount);
-				}
-				else
-				{
-					MaintainerParams->ScanFrequency = FTimespan::MaxValue();
-				}
-				double TimeToWaitAfterInit;
-				if (bClean)
-				{
-					MaintainerParams->TimeToWaitAfterInit = FTimespan::Zero();
-				}
-				else if (GConfig->GetDouble(TEXT("DDCCleanup"), TEXT("TimeToWaitAfterInit"), TimeToWaitAfterInit, GEngineIni))
-				{
-					MaintainerParams->TimeToWaitAfterInit = FTimespan::FromSeconds(TimeToWaitAfterInit);
 				}
+				FParse::Value(InParams, TEXT("FoldersToClean="), MaintainerParams->MaxDirectoryScanCount);
+			}
+			else
+			{
+				MaintainerParams->ScanFrequency = FTimespan::MaxValue();
+			}
+			double TimeToWaitAfterInit;
+			if (bClean)
+			{
+				MaintainerParams->TimeToWaitAfterInit = FTimespan::Zero();
+			}
+			else if (GConfig->GetDouble(TEXT("DDCCleanup"), TEXT("TimeToWaitAfterInit"), TimeToWaitAfterInit, GEngineIni))
+			{
+				MaintainerParams->TimeToWaitAfterInit = FTimespan::FromSeconds(TimeToWaitAfterInit);
+			}
 
-				if (!bLocalDeactivatedForPerformance)
-				{
-					Maintainer = MakeUnique<FFileSystemCacheStoreMaintainer>(*MaintainerParams, CachePath);
+			if (!bLocalDeactivatedForPerformance)
+			{
+				Maintainer = MakeUnique<FFileSystemCacheStoreMaintainer>(*MaintainerParams, CachePath);
 
-					if (bClean)
-					{
-						Maintainer->BoostPriority();
-						Maintainer->WaitForIdle();
-					}
+				if (bClean)
+				{
+					Maintainer->BoostPriority();
+					Maintainer->WaitForIdle();
 				}
 			}
 		}
