DDC: Fixed the file system cache EditorOverrideSetting not handling paths with \ in them

diff --git a/Engine/Source/Developer/DerivedDataCache/Private/DerivedDataBackends.cpp b/Engine/Source/Developer/DerivedDataCache/Private/DerivedDataBackends.cpp
--- a/Engine/Source/Developer/DerivedDataCache/Private/DerivedDataBackends.cpp
+++ b/Engine/Source/Developer/DerivedDataCache/Private/DerivedDataBackends.cpp
@@ -689,7 +689,8 @@ public:
 				FString SettingPath;
 				if(FParse::Value(*Setting, TEXT("Path="), SettingPath))
 				{
-					SettingPath = SettingPath.TrimQuotes();
+					SettingPath.TrimQuotesInline();
+					SettingPath.ReplaceEscapedCharWithCharInline();
 					if(SettingPath.Len() > 0)
 					{
 						Path = SettingPath;
