ZenCacheStore - ignore "None" as an override value for a local DDC path when using sandboxes (in line with past behavior for non-default local caches like Warm/Cold).

diff --git a/Engine/Source/Developer/DerivedDataCache/Private/ZenCacheStore.cpp b/Engine/Source/Developer/DerivedDataCache/Private/ZenCacheStore.cpp
--- a/Engine/Source/Developer/DerivedDataCache/Private/ZenCacheStore.cpp
+++ b/Engine/Source/Developer/DerivedDataCache/Private/ZenCacheStore.cpp
@@ -1244,8 +1244,13 @@ TTuple<ILegacyCacheStore*, ECacheStoreFlags> CreateZenCacheStore(const TCHAR* No
 		UE_LOG(LogDerivedDataCache, Warning, TEXT("%s: Missing required parameter 'Namespace', falling back to '%s'"), NodeName, *Namespace);
 	}
 
+	FString Sandbox;
+	FParse::Value(Config, TEXT("Sandbox="), Sandbox);
+	bool bHasSandbox = !Sandbox.IsEmpty();
+	bool bUseLocalDataCachePathOverrides = !bHasSandbox;
+
 	FString CachePathOverride;
-	if (UE::Zen::Private::IsLocalAutoLaunched(ServiceUrl) && UE::Zen::Private::GetLocalDataCachePathOverride(CachePathOverride))
+	if (bUseLocalDataCachePathOverrides && UE::Zen::Private::IsLocalAutoLaunched(ServiceUrl) && UE::Zen::Private::GetLocalDataCachePathOverride(CachePathOverride))
 	{
 		if (CachePathOverride == TEXT("None"))
 		{
@@ -1256,13 +1261,10 @@ TTuple<ILegacyCacheStore*, ECacheStoreFlags> CreateZenCacheStore(const TCHAR* No
 
 	TUniquePtr<FZenCacheStore> Backend;
 
-	FString Sandbox;
-	FParse::Value(Config, TEXT("Sandbox="), Sandbox);
-
 	bool bFlush = false;
 	FParse::Bool(Config, TEXT("Flush="), bFlush);
 
-	if (!Sandbox.IsEmpty())
+	if (bHasSandbox)
 	{
 		Zen::FServiceSettings DefaultServiceSettings;
 		DefaultServiceSettings.ReadFromConfig();
