Exclude Zen from HttpCacheStore test on non-windows platforms as it is not ready for use on non-windows platforms.

diff --git a/Engine/Source/Developer/DerivedDataCache/Private/Tests/HttpCacheStoreTest.cpp b/Engine/Source/Developer/DerivedDataCache/Private/Tests/HttpCacheStoreTest.cpp
--- a/Engine/Source/Developer/DerivedDataCache/Private/Tests/HttpCacheStoreTest.cpp
+++ b/Engine/Source/Developer/DerivedDataCache/Private/Tests/HttpCacheStoreTest.cpp
@@ -19,6 +19,7 @@
 // to function against all backends.
 
 #if WITH_DEV_AUTOMATION_TESTS && WITH_HTTP_DDC_BACKEND
+#define UE_HTTPCACHESTORETEST_USE_ZEN (PLATFORM_WINDOWS && UE_WITH_ZEN)
 
 DEFINE_LOG_CATEGORY_STATIC(LogHttpDerivedDataBackendTests, Log, All);
 
@@ -618,7 +619,7 @@ bool FHttpCacheStoreTest::RunTest(const FString& Parameters)
 	using namespace UE::DerivedData;
 	ILegacyCacheStore* TestBackend = GetTestBackend();
 
-#if UE_WITH_ZEN
+#if UE_HTTPCACHESTORETEST_USE_ZEN
 	using namespace UE::Zen;
 	FServiceSettings ZenUpstreamTestServiceSettings;
 	FServiceAutoLaunchSettings& ZenUpstreamTestAutoLaunchSettings = ZenUpstreamTestServiceSettings.SettingsVariant.Get<FServiceAutoLaunchSettings>();
@@ -689,7 +690,7 @@ bool FHttpCacheStoreTest::RunTest(const FString& Parameters)
 		// TODO: Expecting a legitimate means to wait for zen to finish pushing records to its upstream in the future
 		FPlatformProcess::Sleep(1.0f);
 	};
-#endif // UE_WITH_ZEN
+#endif // UE_HTTPCACHESTORETEST_USE_ZEN
 
 	const uint32 RecordsInBatch = 3;
 	const uint32 ValuesInBatch = RecordsInBatch;
@@ -700,7 +701,7 @@ bool FHttpCacheStoreTest::RunTest(const FString& Parameters)
 		TArray<FCacheRecord> RecievedRecordsSkipMeta = GetAndValidateRecordsAndChunks(TEXT("SimpleValueSkipMeta"), PutRecords, ECachePolicy::Default | ECachePolicy::SkipMeta);
 		TArray<FCacheRecord> RecievedRecordsSkipData = GetAndValidateRecordsAndChunks(TEXT("SimpleValueSkipData"), PutRecords, ECachePolicy::Default | ECachePolicy::SkipData);
 
-#if UE_WITH_ZEN
+#if UE_HTTPCACHESTORETEST_USE_ZEN
 		if (ZenIntermediaryBackend)
 		{
 			TArray<FCacheRecord> PutRecordsZen = CreateTestCacheRecords(ZenIntermediaryBackend.Get(), RecordsInBatch, 1, FCbObject(), "AutoTestDummyZen");
@@ -709,7 +710,7 @@ bool FHttpCacheStoreTest::RunTest(const FString& Parameters)
 			ValidateRecords(TEXT("SimpleValueSkipMetaZenAndDirect"), GetAndValidateRecords(TEXT("SimpleValueSkipMetaZen"), PutRecordsZen, ECachePolicy::Default | ECachePolicy::SkipMeta), RecievedRecordsSkipMeta, ECachePolicy::Default | ECachePolicy::SkipMeta);
 			ValidateRecords(TEXT("SimpleValueSkipDataZenAndDirect"), GetAndValidateRecords(TEXT("SimpleValueSkipDataZen"), PutRecordsZen, ECachePolicy::Default | ECachePolicy::SkipData), RecievedRecordsSkipData, ECachePolicy::Default | ECachePolicy::SkipData);
 		}
-#endif // UE_WITH_ZEN
+#endif // UE_HTTPCACHESTORETEST_USE_ZEN
 	}
 
 	{
@@ -724,7 +725,7 @@ bool FHttpCacheStoreTest::RunTest(const FString& Parameters)
 		TArray<FCacheRecord> RecievedRecordsSkipMeta = GetAndValidateRecordsAndChunks(TEXT("SimpleValueWithMetaSkipMeta"), PutRecords, ECachePolicy::Default | ECachePolicy::SkipMeta);
 		TArray<FCacheRecord> RecievedRecordsSkipData = GetAndValidateRecordsAndChunks(TEXT("SimpleValueWithMetaSkipData"), PutRecords, ECachePolicy::Default | ECachePolicy::SkipData);
 
-#if UE_WITH_ZEN
+#if UE_HTTPCACHESTORETEST_USE_ZEN
 		if (ZenIntermediaryBackend)
 		{
 			TArray<FCacheRecord> PutRecordsZen = CreateTestCacheRecords(ZenIntermediaryBackend.Get(), RecordsInBatch, 1, MetaObject, "AutoTestDummyZen");
@@ -733,7 +734,7 @@ bool FHttpCacheStoreTest::RunTest(const FString& Parameters)
 			ValidateRecords(TEXT("SimpleValueWithMetaSkipMetaZenAndDirect"), GetAndValidateRecords(TEXT("SimpleValueWithMetaSkipMetaZen"), PutRecordsZen, ECachePolicy::Default | ECachePolicy::SkipMeta), RecievedRecordsSkipMeta, ECachePolicy::Default | ECachePolicy::SkipMeta);
 			ValidateRecords(TEXT("SimpleValueWithMetaSkipDataZenAndDirect"), GetAndValidateRecords(TEXT("SimpleValueWithMetaSkipDataZen"), PutRecordsZen, ECachePolicy::Default | ECachePolicy::SkipData), RecievedRecordsSkipData, ECachePolicy::Default | ECachePolicy::SkipData);
 		}
-#endif // UE_WITH_ZEN
+#endif // UE_HTTPCACHESTORETEST_USE_ZEN
 	}
 
 	{
@@ -742,7 +743,7 @@ bool FHttpCacheStoreTest::RunTest(const FString& Parameters)
 		TArray<FCacheRecord> RecievedRecordsSkipMeta = GetAndValidateRecordsAndChunks(TEXT("MultiValueSkipMeta"), PutRecords, ECachePolicy::Default | ECachePolicy::SkipMeta);
 		TArray<FCacheRecord> RecievedRecordsSkipData = GetAndValidateRecordsAndChunks(TEXT("MultiValueSkipData"), PutRecords, ECachePolicy::Default | ECachePolicy::SkipData);
 
-#if UE_WITH_ZEN
+#if UE_HTTPCACHESTORETEST_USE_ZEN
 		if (ZenIntermediaryBackend)
 		{
 			TArray<FCacheRecord> PutRecordsZen = CreateTestCacheRecords(ZenIntermediaryBackend.Get(), RecordsInBatch, 5, FCbObject(), "AutoTestDummyZen");
@@ -751,7 +752,7 @@ bool FHttpCacheStoreTest::RunTest(const FString& Parameters)
 			ValidateRecords(TEXT("MultiValueSkipMetaZenAndDirect"), GetAndValidateRecords(TEXT("MultiValueSkipMetaZen"), PutRecordsZen, ECachePolicy::Default | ECachePolicy::SkipMeta), RecievedRecordsSkipMeta, ECachePolicy::Default | ECachePolicy::SkipMeta);
 			ValidateRecords(TEXT("MultiValueSkipDataZenAndDirect"), GetAndValidateRecords(TEXT("MultiValueSkipDataZen"), PutRecordsZen, ECachePolicy::Default | ECachePolicy::SkipData), RecievedRecordsSkipData, ECachePolicy::Default | ECachePolicy::SkipData);
 		}
-#endif // UE_WITH_ZEN
+#endif // UE_HTTPCACHESTORETEST_USE_ZEN
 	}
 
 	{
@@ -759,7 +760,7 @@ bool FHttpCacheStoreTest::RunTest(const FString& Parameters)
 		TArray<FValue> ReceivedValues = GetAndValidateValues(TEXT("SimpleValue"), PutValues, ECachePolicy::Default);
 		TArray<FValue> ReceivedValuesSkipData = GetAndValidateValues(TEXT("SimpleValueSkipData"), PutValues, ECachePolicy::Default | ECachePolicy::SkipData);
 
-#if UE_WITH_ZEN
+#if UE_HTTPCACHESTORETEST_USE_ZEN
 		if (ZenIntermediaryBackend)
 		{
 			TArray<FValue> PutValuesZen = CreateTestCacheValues(ZenIntermediaryBackend.Get(), ValuesInBatch);
@@ -772,7 +773,7 @@ bool FHttpCacheStoreTest::RunTest(const FString& Parameters)
 			GetAndValidateValues(TEXT("SimpleValueZen"), PutValues, ECachePolicy::Default, ZenIntermediarySiblingBackend.Get());
 			GetAndValidateValues(TEXT("SimpleValueSkipDataZen"), PutValues, ECachePolicy::Default | ECachePolicy::SkipData, ZenIntermediarySiblingBackend.Get());
 		}
-#endif // UE_WITH_ZEN
+#endif // UE_HTTPCACHESTORETEST_USE_ZEN
 	}
 
 	return true;
