Fix RemoveTicker Hang in HttpCacheStore.

diff --git a/Engine/Source/Developer/DerivedDataCache/Private/HttpCacheStore.cpp b/Engine/Source/Developer/DerivedDataCache/Private/HttpCacheStore.cpp
--- a/Engine/Source/Developer/DerivedDataCache/Private/HttpCacheStore.cpp
+++ b/Engine/Source/Developer/DerivedDataCache/Private/HttpCacheStore.cpp
@@ -332,7 +332,7 @@ private:
 	THttpUniquePtr<IHttpResponse> BeginIsServiceReady(IHttpClient& Client, TArray64<uint8>& Body);
 	bool EndIsServiceReady(THttpUniquePtr<IHttpResponse>& Response, TArray64<uint8>& Body);
 	bool AcquireAccessToken(IHttpClient* Client = nullptr);
-	void SetAccessToken(FStringView Token, double RefreshDelay = 0.0);
+	void SetAccessTokenAndUnlock(FScopeLock &Lock, FStringView Token, double RefreshDelay = 0.0);
 	
 	enum class EOperationCategory
 	{
@@ -1682,7 +1682,7 @@ bool FHttpCacheStore::AcquireAccessToken(IHttpClient* Client)
 
 	if (!OAuthAccessToken.IsEmpty())
 	{
-		SetAccessToken(OAuthAccessToken);
+		SetAccessTokenAndUnlock(Lock, OAuthAccessToken);
 		return true;
 	}
 
@@ -1731,7 +1731,7 @@ bool FHttpCacheStore::AcquireAccessToken(IHttpClient* Client)
 				{
 					UE_LOG(LogDerivedDataCache, Display,
 						TEXT("%s: Logged in to HTTP DDC services. Expires in %.0f seconds."), *Domain, ExpiryTimeSeconds);
-					SetAccessToken(AccessTokenString, ExpiryTimeSeconds);
+					SetAccessTokenAndUnlock(Lock, AccessTokenString, ExpiryTimeSeconds);
 					return true;
 				}
 			}
@@ -1759,7 +1759,7 @@ bool FHttpCacheStore::AcquireAccessToken(IHttpClient* Client)
 			UE_LOG(LogDerivedDataCache, Display,
 				TEXT("%s: OidcToken: Logged in to HTTP DDC services. Expires at %s which is in %.0f seconds."),
 				*Domain, *TokenExpiresAt.ToString(), ExpiryTimeSeconds);
-			SetAccessToken(AccessTokenString, ExpiryTimeSeconds);
+			SetAccessTokenAndUnlock(Lock, AccessTokenString, ExpiryTimeSeconds);
 			return true;
 		}
 
@@ -1772,13 +1772,11 @@ bool FHttpCacheStore::AcquireAccessToken(IHttpClient* Client)
 	return false;
 }
 
-void FHttpCacheStore::SetAccessToken(FStringView Token, double RefreshDelay)
+void FHttpCacheStore::SetAccessTokenAndUnlock(FScopeLock& Lock, FStringView Token, double RefreshDelay)
 {
-	if (RefreshAccessTokenHandle.IsValid())
-	{
-		FTSTicker::GetCoreTicker().RemoveTicker(RefreshAccessTokenHandle);
-		RefreshAccessTokenHandle.Reset();
-	}
+	// Cache the expired refresh handle.
+	FTSTicker::FDelegateHandle ExpiredRefreshAccessTokenHandle = MoveTemp(RefreshAccessTokenHandle);
+	RefreshAccessTokenHandle.Reset();
 
 	if (!Access)
 	{
@@ -1811,6 +1809,15 @@ void FHttpCacheStore::SetAccessToken(FStringView Token, double RefreshDelay)
 
 	// Reset failed login attempts, the service is indeed alive.
 	FailedLoginAttempts = 0;
+
+	// Unlock the critical section before attempting to remove the expired refresh handle.
+	// The associated ticker delegate could already be executing, which could cause a
+	// hang in RemoveTicker when the critical section is locked.
+	Lock.Unlock();
+	if (ExpiredRefreshAccessTokenHandle.IsValid())
+	{
+		FTSTicker::GetCoreTicker().RemoveTicker(MoveTemp(ExpiredRefreshAccessTokenHandle));
+	}
 }
 
 TUniquePtr<FHttpCacheStore::FHttpOperation> FHttpCacheStore::WaitForHttpOperation(EOperationCategory Category)
