Fixed the Cloud, Local and Shared Percentages that exceed 100% when more than one backend is used. These were flawed, Root Get/Hit count cannot be used for the percentages for each backend in this way.

diff --git a/Engine/Source/Developer/DerivedDataCache/Private/DerivedDataCache.cpp b/Engine/Source/Developer/DerivedDataCache/Private/DerivedDataCache.cpp
--- a/Engine/Source/Developer/DerivedDataCache/Private/DerivedDataCache.cpp
+++ b/Engine/Source/Developer/DerivedDataCache/Private/DerivedDataCache.cpp
@@ -96,51 +96,65 @@ namespace UE::DerivedData::CookStats
 			const TSharedRef<const FDerivedDataCacheStatsNode>* ZenRemoteNode = Nodes.FindByPredicate([](TSharedRef<const FDerivedDataCacheStatsNode> Node) { return (Node->GetCacheType() == TEXT("Zen") || Node->GetCacheType() == TEXT("Horde")) && !Node->IsLocal(); });
 
 			const FDerivedDataCacheUsageStats& RootStats = RootNode->UsageStats.CreateConstIterator().Value();
-			const int64 TotalGetHits = RootStats.GetStats.GetAccumulatedValueAnyThread(FCookStats::CallStats::EHitOrMiss::Hit, FCookStats::CallStats::EStatType::Counter);
-			const int64 TotalGetMisses = RootStats.GetStats.GetAccumulatedValueAnyThread(FCookStats::CallStats::EHitOrMiss::Miss, FCookStats::CallStats::EStatType::Counter);
-			const int64 TotalGets = TotalGetHits + TotalGetMisses;
-
-			int64 LocalHits = 0;
+			
+			int64 LocalGetHits = 0;
+			int64 LocalGetMisses = 0;
 			FDerivedDataCacheSpeedStats LocalSpeedStats;
 			if (LocalNode)
 			{
 				const FDerivedDataCacheUsageStats& UsageStats = (*LocalNode)->UsageStats.CreateConstIterator().Value();
-				LocalHits += UsageStats.GetStats.GetAccumulatedValueAnyThread(FCookStats::CallStats::EHitOrMiss::Hit, FCookStats::CallStats::EStatType::Counter);
+				LocalGetHits += UsageStats.GetStats.GetAccumulatedValueAnyThread(FCookStats::CallStats::EHitOrMiss::Hit, FCookStats::CallStats::EStatType::Counter);
+				LocalGetMisses += UsageStats.GetStats.GetAccumulatedValueAnyThread(FCookStats::CallStats::EHitOrMiss::Miss, FCookStats::CallStats::EStatType::Counter);
 				LocalSpeedStats = (*LocalNode)->SpeedStats;
 			}
 			if (ZenLocalNode)
 			{
 				const FDerivedDataCacheUsageStats& UsageStats = (*ZenLocalNode)->UsageStats.CreateConstIterator().Value();
-				LocalHits += UsageStats.GetStats.GetAccumulatedValueAnyThread(FCookStats::CallStats::EHitOrMiss::Hit, FCookStats::CallStats::EStatType::Counter);
+				LocalGetHits += UsageStats.GetStats.GetAccumulatedValueAnyThread(FCookStats::CallStats::EHitOrMiss::Hit, FCookStats::CallStats::EStatType::Counter);
+				LocalGetMisses += UsageStats.GetStats.GetAccumulatedValueAnyThread(FCookStats::CallStats::EHitOrMiss::Miss, FCookStats::CallStats::EStatType::Counter);
 				LocalSpeedStats = (*ZenLocalNode)->SpeedStats;
 			}
-			int64 SharedHits = 0;
+			const int64 LocalGetTotal = LocalGetHits + LocalGetMisses;
+
+			int64 SharedGetHits = 0;
+			int64 SharedGetMisses = 0;
 			FDerivedDataCacheSpeedStats SharedSpeedStats;
 			if (SharedNode)
 			{
 				// The shared DDC is only queried if the local one misses (or there isn't one). So it's hit rate is technically 
 				const FDerivedDataCacheUsageStats& UsageStats = (*SharedNode)->UsageStats.CreateConstIterator().Value();
-				SharedHits += UsageStats.GetStats.GetAccumulatedValueAnyThread(FCookStats::CallStats::EHitOrMiss::Hit, FCookStats::CallStats::EStatType::Counter);
+				SharedGetHits += UsageStats.GetStats.GetAccumulatedValueAnyThread(FCookStats::CallStats::EHitOrMiss::Hit, FCookStats::CallStats::EStatType::Counter);
+				SharedGetMisses += UsageStats.GetStats.GetAccumulatedValueAnyThread(FCookStats::CallStats::EHitOrMiss::Miss, FCookStats::CallStats::EStatType::Counter);
 				SharedSpeedStats = (*SharedNode)->SpeedStats;
 			}
 			if (ZenRemoteNode)
 			{
 				const FDerivedDataCacheUsageStats& UsageStats = (*ZenRemoteNode)->UsageStats.CreateConstIterator().Value();
-				SharedHits += UsageStats.GetStats.GetAccumulatedValueAnyThread(FCookStats::CallStats::EHitOrMiss::Hit, FCookStats::CallStats::EStatType::Counter);
+				SharedGetHits += UsageStats.GetStats.GetAccumulatedValueAnyThread(FCookStats::CallStats::EHitOrMiss::Hit, FCookStats::CallStats::EStatType::Counter);
+				SharedGetMisses += UsageStats.GetStats.GetAccumulatedValueAnyThread(FCookStats::CallStats::EHitOrMiss::Miss, FCookStats::CallStats::EStatType::Counter);
 				SharedSpeedStats = (*ZenRemoteNode)->SpeedStats;
 			}
-			int64 CloudHits = 0;
+			const int64 SharedGetTotal = SharedGetHits + SharedGetMisses;
+
+			int64 CloudGetHits = 0;
+			int64 CloudGetMisses = 0;
 			FDerivedDataCacheSpeedStats CloudSpeedStats;
 			if (CloudNode)
 			{
 				const FDerivedDataCacheUsageStats& UsageStats = (*CloudNode)->UsageStats.CreateConstIterator().Value();
-				CloudHits += UsageStats.GetStats.GetAccumulatedValueAnyThread(FCookStats::CallStats::EHitOrMiss::Hit, FCookStats::CallStats::EStatType::Counter);
+				CloudGetHits += UsageStats.GetStats.GetAccumulatedValueAnyThread(FCookStats::CallStats::EHitOrMiss::Hit, FCookStats::CallStats::EStatType::Counter);
+				CloudGetMisses += UsageStats.GetStats.GetAccumulatedValueAnyThread(FCookStats::CallStats::EHitOrMiss::Miss, FCookStats::CallStats::EStatType::Counter);
 				CloudSpeedStats = (*CloudNode)->SpeedStats;
 			}
+			const int64 CloudGetTotal = CloudGetHits + CloudGetMisses;
+
+			const int64 RootGetHits = RootStats.GetStats.GetAccumulatedValueAnyThread(FCookStats::CallStats::EHitOrMiss::Hit, FCookStats::CallStats::EStatType::Counter);
+			const int64 RootGetMisses = RootStats.GetStats.GetAccumulatedValueAnyThread(FCookStats::CallStats::EHitOrMiss::Miss, FCookStats::CallStats::EStatType::Counter);
+			const int64 RootGetTotal = RootGetHits + RootGetMisses;
 
-			const int64 TotalPutHits = RootStats.PutStats.GetAccumulatedValueAnyThread(FCookStats::CallStats::EHitOrMiss::Hit, FCookStats::CallStats::EStatType::Counter);
-			const int64 TotalPutMisses = RootStats.PutStats.GetAccumulatedValueAnyThread(FCookStats::CallStats::EHitOrMiss::Miss, FCookStats::CallStats::EStatType::Counter);
-			const int64 TotalPuts = TotalPutHits + TotalPutMisses;
+			const int64 RootPutHits = RootStats.PutStats.GetAccumulatedValueAnyThread(FCookStats::CallStats::EHitOrMiss::Hit, FCookStats::CallStats::EStatType::Counter);
+			const int64 RootPutMisses = RootStats.PutStats.GetAccumulatedValueAnyThread(FCookStats::CallStats::EHitOrMiss::Miss, FCookStats::CallStats::EStatType::Counter);
+			const int64 RootPutTotal = RootPutHits + RootPutMisses;
 
 			AddStat(TEXT("DDC.Summary"), FCookStatsManager::CreateKeyValueArray(
 				TEXT("BackEnd"), FDerivedDataBackend::Get().GetGraphName(),
@@ -148,18 +162,25 @@ namespace UE::DerivedData::CookStats
 				TEXT("HasSharedCache"), SharedNode || ZenRemoteNode,
 				TEXT("HasCloudCache"), !!CloudNode,
 				TEXT("HasZenCache"), ZenLocalNode || ZenRemoteNode,
-				TEXT("TotalGetHits"), TotalGetHits,
-				TEXT("TotalGets"), TotalGets,
-				TEXT("TotalGetHitPct"), SafeDivide(TotalGetHits, TotalGets),
-				TEXT("LocalGetHitPct"), SafeDivide(LocalHits, TotalGets),
-				TEXT("SharedGetHitPct"), SafeDivide(SharedHits, TotalGets),
-				TEXT("CloudGetHitPct"), SafeDivide(CloudHits, TotalGets),
-				TEXT("OtherGetHitPct"), SafeDivide((TotalGetHits - LocalHits - SharedHits - CloudHits), TotalGets),
-				TEXT("GetMissPct"), SafeDivide(TotalGetMisses, TotalGets),
-				TEXT("TotalPutHits"), TotalPutHits,
-				TEXT("TotalPuts"), TotalPuts,
-				TEXT("TotalPutHitPct"), SafeDivide(TotalPutHits, TotalPuts),
-				TEXT("PutMissPct"), SafeDivide(TotalPutMisses, TotalPuts),
+				TEXT("TotalGetHits"), RootGetHits,
+				TEXT("TotalGetMisses"), RootGetMisses,
+				TEXT("TotalGets"), RootGetTotal,
+				TEXT("TotalGetHitPct"), SafeDivide(RootGetHits, RootGetTotal),
+				TEXT("GetMissPct"), SafeDivide(RootGetMisses, RootGetTotal),
+				TEXT("TotalPutHits"), RootPutHits,
+				TEXT("TotalPutMisses"), RootPutMisses,
+				TEXT("TotalPuts"), RootPutTotal,	
+				TEXT("TotalPutHitPct"), SafeDivide(RootPutHits, RootPutTotal),	
+				TEXT("PutMissPct"), SafeDivide(RootPutMisses, RootPutTotal),
+				TEXT("LocalGetHits"), LocalGetHits,
+				TEXT("LocalGetTotal"), LocalGetTotal,
+				TEXT("LocalGetHitPct"), SafeDivide(LocalGetHits, LocalGetTotal),
+				TEXT("SharedGetHits"), SharedGetHits,
+				TEXT("SharedGetTotal"), SharedGetTotal,
+				TEXT("SharedGetHitPct"), SafeDivide(SharedGetHits, SharedGetTotal),
+				TEXT("CloudGetHits"), CloudGetHits,
+				TEXT("CloudGetTotal"), CloudGetTotal,
+				TEXT("CloudGetHitPct"), SafeDivide(CloudGetHits, CloudGetTotal),
 				TEXT("LocalLatency"), LocalSpeedStats.LatencyMS,
 				TEXT("LocalReadSpeed"), LocalSpeedStats.ReadSpeedMBs,
 				TEXT("LocalWriteSpeed"), LocalSpeedStats.WriteSpeedMBs,
