Add support for a Zen cache store Host to be None (implying that it is disabled) and to have both commandline and environment overrides (similar to the HTTP cache store).

diff --git a/Engine/Source/Developer/DerivedDataCache/Private/ZenCacheStore.cpp b/Engine/Source/Developer/DerivedDataCache/Private/ZenCacheStore.cpp
--- a/Engine/Source/Developer/DerivedDataCache/Private/ZenCacheStore.cpp
+++ b/Engine/Source/Developer/DerivedDataCache/Private/ZenCacheStore.cpp
@@ -1237,6 +1237,31 @@ TTuple<ILegacyCacheStore*, ECacheStoreFlags> CreateZenCacheStore(const TCHAR* No
 	FString ServiceUrl;
 	FParse::Value(Config, TEXT("Host="), ServiceUrl);
 
+	FString OverrideName;
+	if (FParse::Value(Config, TEXT("EnvHostOverride="), OverrideName))
+	{
+		FString ServiceUrlEnv = FPlatformMisc::GetEnvironmentVariable(*OverrideName);
+		if (!ServiceUrlEnv.IsEmpty())
+		{
+			ServiceUrl = ServiceUrlEnv;
+			UE_LOG(LogDerivedDataCache, Log, TEXT("%s: Found environment override for Host %s=%s"), NodeName, *OverrideName, *ServiceUrl);
+		}
+	}
+
+	if (FParse::Value(Config, TEXT("CommandLineHostOverride="), OverrideName))
+	{
+		if (FParse::Value(FCommandLine::Get(), *(OverrideName + TEXT("=")), ServiceUrl))
+		{
+			UE_LOG(LogDerivedDataCache, Log, TEXT("%s: Found command line override for Host %s=%s"), NodeName, *OverrideName, *ServiceUrl);
+		}
+	}
+
+	if (ServiceUrl == TEXT("None"))
+	{
+		UE_LOG(LogDerivedDataCache, Log, TEXT("Disabling %s data cache - host set to 'None'."), NodeName);
+		return MakeTuple<ILegacyCacheStore*, ECacheStoreFlags>(nullptr, ECacheStoreFlags::None);
+	}
+
 	FString Namespace;
 	if (!FParse::Value(Config, TEXT("StructuredNamespace="), Namespace) && !FParse::Value(Config, TEXT("Namespace="), Namespace))
 	{
