DDC: Fixed a race condition in the cache hierarchy

The hierarchy now uses a request barrier after the first synchronous request completion. Async requests could occasionally complete fast enough that DispatchRequests would be left to handle the completion. In that scenario, there would not be a barrier in place when beginning any subsequent requests, which would lead to a failed assertion in the request owner.

diff --git a/Engine/Source/Developer/DerivedDataCache/Private/DerivedDataCacheStoreHierarchy.cpp b/Engine/Source/Developer/DerivedDataCache/Private/DerivedDataCacheStoreHierarchy.cpp
--- a/Engine/Source/Developer/DerivedDataCache/Private/DerivedDataCacheStoreHierarchy.cpp
+++ b/Engine/Source/Developer/DerivedDataCache/Private/DerivedDataCacheStoreHierarchy.cpp
@@ -79,6 +79,7 @@ private:
 	void UpdateNodeFlags();
 
 	class FCounterEvent;
+	class FDynamicRequestBarrier;
 
 	class FBatchBase;
 	template <typename Params> class TPutBatch;
@@ -135,6 +136,35 @@ private:
 	std::atomic<int32> Count{0};
 };
 
+class FCacheStoreHierarchy::FDynamicRequestBarrier
+{
+public:
+	FDynamicRequestBarrier() = default;
+	FDynamicRequestBarrier(const FDynamicRequestBarrier&) = delete;
+	FDynamicRequestBarrier& operator=(const FDynamicRequestBarrier&) = delete;
+
+	void Begin(IRequestOwner& NewOwner)
+	{
+		if (!Owner)
+		{
+			NewOwner.BeginBarrier(ERequestBarrierFlags::Priority);
+			Owner = &NewOwner;
+		}
+		check(Owner == &NewOwner);
+	}
+
+	~FDynamicRequestBarrier()
+	{
+		if (Owner)
+		{
+			Owner->EndBarrier(ERequestBarrierFlags::Priority);
+		}
+	}
+
+private:
+	IRequestOwner* Owner = nullptr;
+};
+
 class FCacheStoreHierarchy::FBatchBase
 {
 public:
@@ -401,6 +431,7 @@ template <typename Params>
 void FCacheStoreHierarchy::TPutBatch<Params>::DispatchRequests()
 {
 	FReadScopeLock Lock(Hierarchy.NodesLock);
+	FDynamicRequestBarrier Barrier;
 
 	for (const int32 NodeCount = Hierarchy.Nodes.Num(); NodePutIndex < NodeCount && !BatchOwner.IsCanceled(); ++NodePutIndex)
 	{
@@ -408,6 +439,7 @@ void FCacheStoreHierarchy::TPutBatch<Params>::DispatchRequests()
 		{
 			return;
 		}
+		Barrier.Begin(BatchOwner);
 	}
 
 	int32 RequestIndex = 0;
@@ -635,6 +667,7 @@ template <typename Params>
 void FCacheStoreHierarchy::TGetBatch<Params>::DispatchRequests()
 {
 	FReadScopeLock Lock(Hierarchy.NodesLock);
+	FDynamicRequestBarrier Barrier;
 
 	TArray<FGetRequest, TInlineAllocator<8>> NodeRequests;
 	TArray<FPutRequest, TInlineAllocator<8>> AsyncNodeRequests;
@@ -705,7 +738,7 @@ void FCacheStoreHierarchy::TGetBatch<Params>::DispatchRequests()
 
 		if (!AsyncNodeRequests.IsEmpty())
 		{
-			FRequestBarrier Barrier(AsyncOwner);
+			FRequestBarrier AsyncBarrier(AsyncOwner);
 			Invoke(Put(), Node.AsyncCache, AsyncNodeRequests, AsyncOwner, [](auto&&){});
 			AsyncNodeRequests.Reset();
 		}
@@ -724,6 +757,8 @@ void FCacheStoreHierarchy::TGetBatch<Params>::DispatchRequests()
 				return;
 			}
 		}
+
+		Barrier.Begin(Owner);
 	}
 
 	for (FState& State : States)
@@ -773,7 +808,7 @@ void FCacheStoreHierarchy::TGetBatch<Params>::CompleteRequest(FGetResponse&& Res
 			const FCacheStoreNode& PutNode = Hierarchy.Nodes[PutNodeIndex];
 			if (CanStore(GetCombinedPolicy(State.Request.Policy), PutNode.CacheFlags))
 			{
-				FRequestBarrier Barrier(AsyncOwner);
+				FRequestBarrier AsyncBarrier(AsyncOwner);
 				Invoke(Put(), PutNode.AsyncCache, MakeArrayView(&PutRequest, 1), AsyncOwner, [](auto&&){});
 			}
 		}
@@ -1119,6 +1154,7 @@ void FCacheStoreHierarchy::FGetChunksBatch::Begin(
 void FCacheStoreHierarchy::FGetChunksBatch::DispatchRequests()
 {
 	FReadScopeLock Lock(Hierarchy.NodesLock);
+	FDynamicRequestBarrier Barrier;
 
 	TArray<FCacheGetChunkRequest, TInlineAllocator<8>> NodeRequests;
 	NodeRequests.Reserve(States.Num());
@@ -1152,6 +1188,8 @@ void FCacheStoreHierarchy::FGetChunksBatch::DispatchRequests()
 				return;
 			}
 		}
+
+		Barrier.Begin(Owner);
 	}
 
 	for (const FState& State : States)
diff --git a/Engine/Source/Developer/DerivedDataCache/Tests/DerivedDataCacheStoreHierarchyTest.cpp b/Engine/Source/Developer/DerivedDataCache/Tests/DerivedDataCacheStoreHierarchyTest.cpp
--- a/Engine/Source/Developer/DerivedDataCache/Tests/DerivedDataCacheStoreHierarchyTest.cpp
+++ b/Engine/Source/Developer/DerivedDataCache/Tests/DerivedDataCacheStoreHierarchyTest.cpp
@@ -25,11 +25,11 @@ TEST_CASE("DerivedData::Cache::Hierarchy::PartialRecordPropagation", "[DerivedDa
 	ICacheStoreOwner* StoreOwner = nullptr;
 	TUniquePtr<ILegacyCacheStore> Hierarchy(CreateCacheStoreHierarchy(StoreOwner, nullptr));
 
-	TUniquePtr<ITestCacheStore> Store0(CreateTestCacheStore(ECacheStoreFlags::Local | ECacheStoreFlags::Query | ECacheStoreFlags::Store, /*bAsync*/ true));
+	TUniquePtr<ITestCacheStore> Store0(CreateTestCacheStore(ECacheStoreFlags::Local | ECacheStoreFlags::Query | ECacheStoreFlags::Store, ETestCacheStoreFlags::Async));
 	StoreOwner->Add(Store0.Get(), Store0->GetFlags());
-	TUniquePtr<ITestCacheStore> Store1(CreateTestCacheStore(ECacheStoreFlags::Local | ECacheStoreFlags::Query | ECacheStoreFlags::Store, /*bAsync*/ true));
+	TUniquePtr<ITestCacheStore> Store1(CreateTestCacheStore(ECacheStoreFlags::Local | ECacheStoreFlags::Query | ECacheStoreFlags::Store, ETestCacheStoreFlags::Async));
 	StoreOwner->Add(Store1.Get(), Store1->GetFlags());
-	TUniquePtr<ITestCacheStore> Store2(CreateTestCacheStore(ECacheStoreFlags::Local | ECacheStoreFlags::Query | ECacheStoreFlags::Store, /*bAsync*/ true));
+	TUniquePtr<ITestCacheStore> Store2(CreateTestCacheStore(ECacheStoreFlags::Local | ECacheStoreFlags::Query | ECacheStoreFlags::Store, ETestCacheStoreFlags::Async));
 	StoreOwner->Add(Store2.Get(), Store2->GetFlags());
 
 	const FValueId Id0 = FValueId::FromName("0");
@@ -141,11 +141,11 @@ TEST_CASE("DerivedData::Cache::Hierarchy::PartialNonDeterministicRecordPropagati
 	ICacheStoreOwner* StoreOwner = nullptr;
 	TUniquePtr<ILegacyCacheStore> Hierarchy(CreateCacheStoreHierarchy(StoreOwner, nullptr));
 
-	TUniquePtr<ITestCacheStore> Store0(CreateTestCacheStore(ECacheStoreFlags::Local | ECacheStoreFlags::Query | ECacheStoreFlags::Store, /*bAsync*/ true));
+	TUniquePtr<ITestCacheStore> Store0(CreateTestCacheStore(ECacheStoreFlags::Local | ECacheStoreFlags::Query | ECacheStoreFlags::Store, ETestCacheStoreFlags::Async));
 	StoreOwner->Add(Store0.Get(), Store0->GetFlags());
-	TUniquePtr<ITestCacheStore> Store1(CreateTestCacheStore(ECacheStoreFlags::Local | ECacheStoreFlags::Query | ECacheStoreFlags::Store, /*bAsync*/ true));
+	TUniquePtr<ITestCacheStore> Store1(CreateTestCacheStore(ECacheStoreFlags::Local | ECacheStoreFlags::Query | ECacheStoreFlags::Store, ETestCacheStoreFlags::Async));
 	StoreOwner->Add(Store1.Get(), Store1->GetFlags());
-	TUniquePtr<ITestCacheStore> Store2(CreateTestCacheStore(ECacheStoreFlags::Local | ECacheStoreFlags::Query | ECacheStoreFlags::Store, /*bAsync*/ true));
+	TUniquePtr<ITestCacheStore> Store2(CreateTestCacheStore(ECacheStoreFlags::Local | ECacheStoreFlags::Query | ECacheStoreFlags::Store, ETestCacheStoreFlags::Async));
 	StoreOwner->Add(Store2.Get(), Store2->GetFlags());
 
 	const FValueId Id0 = FValueId::FromName("0");
@@ -285,6 +285,38 @@ TEST_CASE("DerivedData::Cache::Hierarchy::PartialNonDeterministicRecordPropagati
 	StoreOwner->RemoveNotSafe(Store0.Get());
 }
 
+TEST_CASE("DerivedData::Cache::Hierarchy::BlockingRequestOwner", "[DerivedData]")
+{
+	ICacheStoreOwner* StoreOwner = nullptr;
+	TUniquePtr<ILegacyCacheStore> Hierarchy(CreateCacheStoreHierarchy(StoreOwner, nullptr));
+
+	TUniquePtr<ITestCacheStore> Store0(CreateTestCacheStore(ECacheStoreFlags::Local | ECacheStoreFlags::Query | ECacheStoreFlags::Store, ETestCacheStoreFlags::Async | ETestCacheStoreFlags::Wait));
+	StoreOwner->Add(Store0.Get(), Store0->GetFlags());
+	TUniquePtr<ITestCacheStore> Store1(CreateTestCacheStore(ECacheStoreFlags::Local | ECacheStoreFlags::Query | ECacheStoreFlags::Store, ETestCacheStoreFlags::Async));
+	StoreOwner->Add(Store1.Get(), Store1->GetFlags());
+
+	const FCacheKey Key{FCacheBucket("Test"), FIoHash::HashBuffer(MakeMemoryView<uint8>({0, 1, 2, 3}))};
+
+	SECTION("Owner has correct request barriers in place when requests are created within a blocking operation")
+	{
+		FRequestOwner Owner(EPriority::Blocking);
+		Hierarchy->GetValue({{{TEXTVIEW("Get")}, Key, ECachePolicy::Default, 12345}}, Owner, [](FCacheGetValueResponse&& Response)
+		{
+			CHECK(Response.UserData == 12345);
+			CHECK(Response.Status == EStatus::Error);
+		});
+
+		CHECK(Store0->GetTotalRequestCount() == 1);
+		CHECK(Store1->GetTotalRequestCount() == 1);
+		Store1->ExecuteAsync();
+
+		Owner.Wait();
+	}
+
+	StoreOwner->RemoveNotSafe(Store1.Get());
+	StoreOwner->RemoveNotSafe(Store0.Get());
+}
+
 } // UE::DerivedData
 
 #endif // WITH_LOW_LEVEL_TESTS
diff --git a/Engine/Source/Developer/DerivedDataCache/Tests/TestCacheStore.cpp b/Engine/Source/Developer/DerivedDataCache/Tests/TestCacheStore.cpp
--- a/Engine/Source/Developer/DerivedDataCache/Tests/TestCacheStore.cpp
+++ b/Engine/Source/Developer/DerivedDataCache/Tests/TestCacheStore.cpp
@@ -18,7 +18,7 @@ namespace UE::DerivedData
 class FTestCacheStore final : public ITestCacheStore
 {
 public:
-	explicit FTestCacheStore(ECacheStoreFlags Flags, bool bAsync = false);
+	FTestCacheStore(ECacheStoreFlags Flags, ETestCacheStoreFlags TestFlags);
 	~FTestCacheStore() final;
 
 	FTestCacheStore(const FTestCacheStore&) = delete;
@@ -95,7 +95,7 @@ private:
 	uint32 CanceledRequestCount = 0;
 	ECachePolicy CacheStorePolicy;
 	const ECacheStoreFlags CacheStoreFlags;
-	const bool bAsync;
+	const ETestCacheStoreFlags TestFlags;
 };
 
 class FTestCacheStore::FAsyncRequest final : public FRequestBase
@@ -137,10 +137,10 @@ private:
 	TUniqueFunction<void (bool bCancel)> Body;
 };
 
-FTestCacheStore::FTestCacheStore(const ECacheStoreFlags Flags, const bool bInAsync)
+FTestCacheStore::FTestCacheStore(const ECacheStoreFlags Flags, const ETestCacheStoreFlags InTestFlags)
 	: CacheStorePolicy(ECachePolicy::None)
 	, CacheStoreFlags(Flags)
-	, bAsync(bInAsync)
+	, TestFlags(InTestFlags)
 {
 	CacheStorePolicy |= EnumHasAllFlags(Flags, ECacheStoreFlags::Query | ECacheStoreFlags::Local) ? ECachePolicy::QueryLocal : ECachePolicy::None;
 	CacheStorePolicy |= EnumHasAllFlags(Flags, ECacheStoreFlags::Store | ECacheStoreFlags::Local) ? ECachePolicy::StoreLocal : ECachePolicy::None;
@@ -489,9 +489,17 @@ void FTestCacheStore::GetChunks(
 
 void FTestCacheStore::ExecuteOrQueue(IRequestOwner& Owner, TUniqueFunction<void (bool bCancel)>&& Function)
 {
-	if (bAsync)
+	if (EnumHasAnyFlags(TestFlags, ETestCacheStoreFlags::Async))
 	{
-		Queue.Emplace(new FAsyncRequest(Owner, MoveTemp(Function)));
+		TRefCountPtr<FAsyncRequest> Request = new FAsyncRequest(Owner, MoveTemp(Function));
+		if (EnumHasAnyFlags(TestFlags, ETestCacheStoreFlags::Wait))
+		{
+			Request->Wait();
+		}
+		else
+		{
+			Queue.Emplace(MoveTemp(Request));
+		}
 	}
 	else
 	{
@@ -509,9 +517,9 @@ void FTestCacheStore::ExecuteAsync()
 	}
 }
 
-ITestCacheStore* CreateTestCacheStore(const ECacheStoreFlags Flags, const bool bAsync)
+ITestCacheStore* CreateTestCacheStore(const ECacheStoreFlags Flags, const ETestCacheStoreFlags TestFlags)
 {
-	return new FTestCacheStore(Flags, bAsync);
+	return new FTestCacheStore(Flags, TestFlags);
 }
 
 } // UE::DerivedData
diff --git a/Engine/Source/Developer/DerivedDataCache/Tests/TestCacheStore.h b/Engine/Source/Developer/DerivedDataCache/Tests/TestCacheStore.h
--- a/Engine/Source/Developer/DerivedDataCache/Tests/TestCacheStore.h
+++ b/Engine/Source/Developer/DerivedDataCache/Tests/TestCacheStore.h
@@ -5,6 +5,7 @@
 #if WITH_LOW_LEVEL_TESTS
 
 #include "DerivedDataLegacyCacheStore.h"
+#include "Misc/EnumClassFlags.h"
 
 namespace UE::DerivedData
 {
@@ -34,7 +35,18 @@ public:
 	virtual void ExecuteAsync() = 0;
 };
 
-ITestCacheStore* CreateTestCacheStore(ECacheStoreFlags Flags, bool bAsync);
+enum class ETestCacheStoreFlags : uint32
+{
+	None            = 0,
+	/** Async execution of requests. */
+	Async           = 1 << 0,
+	/** Waits on async execution of requests. */
+	Wait            = 1 << 1,
+};
+
+ENUM_CLASS_FLAGS(ETestCacheStoreFlags);
+
+ITestCacheStore* CreateTestCacheStore(ECacheStoreFlags Flags, ETestCacheStoreFlags TestFlags);
 
 } // UE::DerivedData
 
