DDC: Added -ExportBuildOutput to include outputs in exported builds

diff --git a/Engine/Source/Developer/DerivedDataCache/Private/DerivedDataBuildJob.cpp b/Engine/Source/Developer/DerivedDataCache/Private/DerivedDataBuildJob.cpp
--- a/Engine/Source/Developer/DerivedDataCache/Private/DerivedDataBuildJob.cpp
+++ b/Engine/Source/Developer/DerivedDataCache/Private/DerivedDataBuildJob.cpp
@@ -35,6 +35,7 @@
 #include "Misc/StringBuilder.h"
 #include "ProfilingDebugging/CpuProfilerTrace.h"
 #include "Serialization/CompactBinary.h"
+#include "Serialization/CompactBinarySerialization.h"
 #include "Serialization/CompactBinaryWriter.h"
 #include "String/ParseTokens.h"
 #include "Templates/Function.h"
@@ -187,9 +188,15 @@ private:
 	/** Execute the new state. */
 	void ExecuteState(EBuildJobState NewState);
 
+	/** Appends the export path to the builder. */
+	void AppendExportPath(FStringBuilderBase& OutPath) const;
+
 	/** Exports the action and inputs for this build to disk. */
 	void ExportBuild() const;
 
+	/** Exports the outputs for this build to disk. ExportBuild() must be called before this. */
+	void ExportBuildOutput() const;
+
 	/** Determine whether the action and inputs for this build should be exported to disk. */
 	bool ShouldExportBuild() const;
 
@@ -1082,6 +1089,9 @@ void FBuildJob::AdvanceToState(EBuildJobState NewState, FUtf8StringView NewError
 
 void FBuildJob::ExecuteTransition(EBuildJobState OldState, EBuildJobState NewState)
 {
+	UE_LOG(LogTemp, Display, TEXT("Transition %s -> %s for build of '%s' by %s in job %p."),
+		LexToString(OldState), LexToString(NewState), *Name, *WriteToString<32>(FunctionName), this);
+
 	if (OldState == EBuildJobState::NotStarted)
 	{
 		BeginJob();
@@ -1116,6 +1126,10 @@ void FBuildJob::ExecuteTransition(EBuildJobState OldState, EBuildJobState NewSta
 	}
 	if (NewState == EBuildJobState::Complete)
 	{
+		if (EnumHasAnyFlags(BuildStatus, EBuildStatus::BuildExport))
+		{
+			ExportBuildOutput();
+		}
 		if (OnComplete)
 		{
 			const FCacheKey& CacheKey = Context ? Context->GetCacheKey() : FCacheKey::Empty;
@@ -1167,22 +1181,29 @@ void FBuildJob::ExecuteState(EBuildJobState NewState)
 	}
 }
 
-void FBuildJob::ExportBuild() const
+void FBuildJob::AppendExportPath(FStringBuilderBase& OutPath) const
 {
 	// Export to <SavedDir>/DerivedDataBuildExport/<Bucket>[/<Function>]/<Action>
-	TStringBuilder<256> ExportPath;
 	const FCacheKey& Key = Context->GetCacheKey();
-	FPathViews::Append(ExportPath, FPaths::ProjectSavedDir(), TEXT("DerivedDataBuildExport"), Key.Bucket);
+	FPathViews::Append(OutPath, FPaths::ProjectSavedDir(), TEXT("DerivedDataBuildExport"), Key.Bucket);
 	if (!Key.Bucket.ToString().Equals(FunctionName))
 	{
-		FPathViews::Append(ExportPath, FunctionName);
+		FPathViews::Append(OutPath, FunctionName);
 	}
-	FPathViews::Append(ExportPath, Key.Hash);
+	FPathViews::Append(OutPath, Key.Hash);
+}
+
+void FBuildJob::ExportBuild() const
+{
+	constexpr uint64 MaxMetaConstantSize = 16 * 1024;
+
+	TStringBuilder<256> ExportPath;
+	AppendExportPath(ExportPath);
 	int32 ExportRootLen = ExportPath.Len();
 
 	TUtf8StringBuilder<512> Meta;
 	Meta << "Name: " << Name << LINE_TERMINATOR_ANSI;
-	Meta << "Cache: " << Key << LINE_TERMINATOR_ANSI;
+	Meta << "Cache: " << Context->GetCacheKey() << LINE_TERMINATOR_ANSI;
 	Meta << "Function: " << FunctionName << LINE_TERMINATOR_ANSI;
 	Meta << "FunctionVersion: " << Action.Get().GetFunctionVersion() << LINE_TERMINATOR_ANSI;
 	Meta << "BuildSystemVersion: " << Action.Get().GetBuildSystemVersion() << LINE_TERMINATOR_ANSI;
@@ -1194,6 +1215,12 @@ void FBuildJob::ExportBuild() const
 			Meta << "  " << Key << ":" LINE_TERMINATOR_ANSI;
 			Meta << "    RawHash: " << FIoHash(Value.GetHash()) << LINE_TERMINATOR_ANSI;
 			Meta << "    RawSize: " << Value.GetSize() << LINE_TERMINATOR_ANSI;
+			if (Value.GetSize() <= MaxMetaConstantSize)
+			{
+				Meta << "    Value: ";
+				CompactBinaryToCompactJson(Value, Meta);
+				Meta << LINE_TERMINATOR_ANSI;
+			}
 		});
 	}
 	if (Action.Get().HasInputs())
@@ -1225,7 +1252,6 @@ void FBuildJob::ExportBuild() const
 
 	if (Inputs)
 	{
-		Meta << "Inputs:" << LINE_TERMINATOR_ANSI;
 		FPathViews::Append(ExportPath, TEXT("Inputs"));
 		ExportRootLen = ExportPath.Len();
 		Inputs.Get().IterateInputs([&ExportPath, ExportRootLen](FUtf8StringView Key, const FCompressedBuffer& Buffer)
@@ -1240,6 +1266,93 @@ void FBuildJob::ExportBuild() const
 	}
 }
 
+void FBuildJob::ExportBuildOutput() const
+{
+	static bool bExportOutput = FParse::Param(FCommandLine::Get(), TEXT("ExportBuildOutput"));
+	if (!bExportOutput || !Output)
+	{
+		return;
+	}
+
+	const FBuildOutput& LocalOutput = Output.Get();
+
+	TStringBuilder<256> ExportPath;
+	AppendExportPath(ExportPath);
+	int32 ExportRootLen = ExportPath.Len();
+
+	TUtf8StringBuilder<512> Meta;
+	if (!Output.Get().GetValues().IsEmpty())
+	{
+		Meta << "Outputs:" << LINE_TERMINATOR_ANSI;
+		for (const FValueWithId& Value : Output.Get().GetValues())
+		{
+			Meta << "  " << Value.GetId() << ":" LINE_TERMINATOR_ANSI;
+			Meta << "    RawHash: " << Value.GetRawHash() << LINE_TERMINATOR_ANSI;
+			Meta << "    RawSize: " << Value.GetRawSize() << LINE_TERMINATOR_ANSI;
+		}
+	}
+	if (!Output.Get().GetMessages().IsEmpty())
+	{
+		TCbWriter<256> MessageWriter;
+		Meta << "Messages:" << LINE_TERMINATOR_ANSI;
+		for (const FBuildOutputMessage& Message : Output.Get().GetMessages())
+		{
+			Meta << "  - Level: " << Message.Level << LINE_TERMINATOR_ANSI;
+			Meta << "    Message: ";
+			MessageWriter.AddString(Message.Message);
+			CompactBinaryToCompactJson(MessageWriter.Save(), Meta);
+			MessageWriter.Reset();
+			Meta << LINE_TERMINATOR_ANSI;
+		}
+	}
+	if (!Output.Get().GetLogs().IsEmpty())
+	{
+		TCbWriter<256> LogWriter;
+		Meta << "Logs:" << LINE_TERMINATOR_ANSI;
+		for (const FBuildOutputLog& Log : Output.Get().GetLogs())
+		{
+			Meta << "  - Category: " << Log.Category << LINE_TERMINATOR_ANSI;
+			Meta << "    Level: " << Log.Level << LINE_TERMINATOR_ANSI;
+			Meta << "    Message: ";
+			LogWriter.AddString(Log.Message);
+			CompactBinaryToCompactJson(LogWriter.Save(), Meta);
+			LogWriter.Reset();
+			Meta << LINE_TERMINATOR_ANSI;
+		}
+	}
+
+	FPathViews::Append(ExportPath, TEXT("Meta.yaml"));
+	if (TUniquePtr<FArchive> Ar{IFileManager::Get().CreateFileWriter(*ExportPath, FILEWRITE_Append)})
+	{
+		Ar->Serialize(Meta.GetData(), Meta.Len());
+	}
+	ExportPath.RemoveSuffix(ExportPath.Len() - ExportRootLen);
+
+	FPathViews::Append(ExportPath, TEXT("Build.output"));
+	if (TUniquePtr<FArchive> Ar{IFileManager::Get().CreateFileWriter(*ExportPath)})
+	{
+		FCbWriter Writer;
+		Output.Get().Save(Writer);
+		Writer.Save(*Ar);
+	}
+	ExportPath.RemoveSuffix(ExportPath.Len() - ExportRootLen);
+
+	FPathViews::Append(ExportPath, TEXT("Outputs"));
+	ExportRootLen = ExportPath.Len();
+	for (const FValueWithId& Value : Output.Get().GetValues())
+	{
+		if (Value.GetData())
+		{
+			FPathViews::Append(ExportPath, FIoHash(Value.GetRawHash()));
+			if (TUniquePtr<FArchive> Ar{IFileManager::Get().CreateFileWriter(*ExportPath)})
+			{
+				Value.GetData().Save(*Ar);
+			}
+			ExportPath.RemoveSuffix(ExportPath.Len() - ExportRootLen);
+		}
+	}
+}
+
 bool FBuildJob::ShouldExportBuild() const
 {
 	static bool bExportAll;
diff --git a/Engine/Source/Developer/DerivedDataCache/Private/DerivedDataBuildOutput.cpp b/Engine/Source/Developer/DerivedDataCache/Private/DerivedDataBuildOutput.cpp
--- a/Engine/Source/Developer/DerivedDataCache/Private/DerivedDataBuildOutput.cpp
+++ b/Engine/Source/Developer/DerivedDataCache/Private/DerivedDataBuildOutput.cpp
@@ -454,6 +454,16 @@ FBuildOutputBuilder CreateBuildOutput(const FSharedString& Name, const FUtf8Shar
 namespace UE::DerivedData
 {
 
+FUtf8StringBuilderBase& operator<<(FUtf8StringBuilderBase& Builder, EBuildOutputMessageLevel Level)
+{
+	return Builder << Private::LexToString(Level);
+}
+
+FUtf8StringBuilderBase& operator<<(FUtf8StringBuilderBase& Builder, EBuildOutputLogLevel Level)
+{
+	return Builder << Private::LexToString(Level);
+}
+
 FOptionalBuildOutput FBuildOutput::Load(
 	const FSharedString& Name,
 	const FUtf8SharedString& Function,
diff --git a/Engine/Source/Developer/DerivedDataCache/Public/DerivedDataBuildOutput.h b/Engine/Source/Developer/DerivedDataCache/Public/DerivedDataBuildOutput.h
--- a/Engine/Source/Developer/DerivedDataCache/Public/DerivedDataBuildOutput.h
+++ b/Engine/Source/Developer/DerivedDataCache/Public/DerivedDataBuildOutput.h
@@ -79,6 +79,8 @@ enum class EBuildOutputMessageLevel : uint8
 	Display,
 };
 
+UE_API FUtf8StringBuilderBase& operator<<(FUtf8StringBuilderBase& Builder, EBuildOutputMessageLevel Level);
+
 /** A build output message is diagnostic output from a build function and must be deterministic. */
 struct FBuildOutputMessage
 {
@@ -92,6 +94,8 @@ enum class EBuildOutputLogLevel : uint8
 	Warning,
 };
 
+UE_API FUtf8StringBuilderBase& operator<<(FUtf8StringBuilderBase& Builder, EBuildOutputLogLevel Level);
+
 /**
  * A build output log is a log message captured from a build function.
  *
