Fix build breakage from CL23917590.

diff --git a/Engine/Source/Developer/DerivedDataCache/Private/DerivedDataBuildJob.cpp b/Engine/Source/Developer/DerivedDataCache/Private/DerivedDataBuildJob.cpp
--- a/Engine/Source/Developer/DerivedDataCache/Private/DerivedDataBuildJob.cpp
+++ b/Engine/Source/Developer/DerivedDataCache/Private/DerivedDataBuildJob.cpp
@@ -1195,7 +1195,7 @@ void FBuildJob::AppendExportPath(FStringBuilderBase& OutPath) const
 
 void FBuildJob::ExportBuild() const
 {
-	constexpr uint64 MaxMetaConstantSize = 16 * 1024;
+	static constexpr uint64 MaxMetaConstantSize = 16 * 1024;
 
 	TStringBuilder<256> ExportPath;
 	AppendExportPath(ExportPath);
