Added comment to the DerivedDataCache build file explaining the need for developers to  add a DesktopPlatform dependency if they need UE Cloud DDC interactive login to work for their custom targets.

diff --git a/Engine/Source/Developer/DerivedDataCache/DerivedDataCache.Build.cs b/Engine/Source/Developer/DerivedDataCache/DerivedDataCache.Build.cs
--- a/Engine/Source/Developer/DerivedDataCache/DerivedDataCache.Build.cs
+++ b/Engine/Source/Developer/DerivedDataCache/DerivedDataCache.Build.cs
@@ -11,6 +11,13 @@ public class DerivedDataCache : ModuleRules
 		PrivateDependencyModuleNames.Add("Zen");
 		
 		// Dependencies for "S3" and "HTTP" backends
+		// NOTE: DesktopPlatform is needed for UE Cloud DDC interactive login.  It is set as an private include
+		//		 dependency here, but not as a full dependency.  This will compile and link, but the absence of
+		//		 DesktopPlatform module at runtime will prevent interactive login from working.  Adding DesktopPlatform
+		//		 would bloat the size of multiple users of the DDC module who don't need UE Cloud DDC interactive login,
+		//		 so we don't add it as a full dependency.  Build targets (particularly standalone programs) that use the
+		//		 DDC module, should add a dependency on DesktopPlatform if they require interactive login.  This will
+		//		 change in the future as the DDC module dependencies are refactored in the future.
 		PrivateDependencyModuleNames.AddRange(new string[] { "SSL", "Json", "Zen" });
 		PrivateIncludePathModuleNames.AddRange(new string[] { "DesktopPlatform", "Zen", "Analytics" });
 		AddEngineThirdPartyPrivateStaticDependencies(Target, "libcurl");
