Split up the DDC cache analytics per back end so we can more easily reason about specific abck end perf

diff --git a/Engine/Source/Developer/DerivedDataCache/Private/DerivedDataCache.cpp b/Engine/Source/Developer/DerivedDataCache/Private/DerivedDataCache.cpp
--- a/Engine/Source/Developer/DerivedDataCache/Private/DerivedDataCache.cpp
+++ b/Engine/Source/Developer/DerivedDataCache/Private/DerivedDataCache.cpp
@@ -107,15 +107,20 @@ namespace UE::DerivedData::CookStats
 				LocalGetMisses += UsageStats.GetStats.GetAccumulatedValueAnyThread(FCookStats::CallStats::EHitOrMiss::Miss, FCookStats::CallStats::EStatType::Counter);
 				LocalSpeedStats = (*LocalNode)->SpeedStats;
 			}
+			const int64 LocalGetTotal = LocalGetHits + LocalGetMisses;
+
+			int64 ZenLocalGetHits = 0;
+			int64 ZenLocalGetMisses = 0;
+			FDerivedDataCacheSpeedStats ZenLocalSpeedStats;
 			if (ZenLocalNode)
 			{
 				const FDerivedDataCacheUsageStats& UsageStats = (*ZenLocalNode)->UsageStats.CreateConstIterator().Value();
-				LocalGetHits += UsageStats.GetStats.GetAccumulatedValueAnyThread(FCookStats::CallStats::EHitOrMiss::Hit, FCookStats::CallStats::EStatType::Counter);
-				LocalGetMisses += UsageStats.GetStats.GetAccumulatedValueAnyThread(FCookStats::CallStats::EHitOrMiss::Miss, FCookStats::CallStats::EStatType::Counter);
-				LocalSpeedStats = (*ZenLocalNode)->SpeedStats;
-			}
-			const int64 LocalGetTotal = LocalGetHits + LocalGetMisses;
-
+				ZenLocalGetHits += UsageStats.GetStats.GetAccumulatedValueAnyThread(FCookStats::CallStats::EHitOrMiss::Hit, FCookStats::CallStats::EStatType::Counter);
+				ZenLocalGetMisses += UsageStats.GetStats.GetAccumulatedValueAnyThread(FCookStats::CallStats::EHitOrMiss::Miss, FCookStats::CallStats::EStatType::Counter);
+				ZenLocalSpeedStats = (*ZenLocalNode)->SpeedStats;
+			}			
+			const int64 ZenLocalGetTotal = ZenLocalGetHits + ZenLocalGetMisses;
+			
 			int64 SharedGetHits = 0;
 			int64 SharedGetMisses = 0;
 			FDerivedDataCacheSpeedStats SharedSpeedStats;
@@ -127,14 +132,19 @@ namespace UE::DerivedData::CookStats
 				SharedGetMisses += UsageStats.GetStats.GetAccumulatedValueAnyThread(FCookStats::CallStats::EHitOrMiss::Miss, FCookStats::CallStats::EStatType::Counter);
 				SharedSpeedStats = (*SharedNode)->SpeedStats;
 			}
+			const int64 SharedGetTotal = SharedGetHits + SharedGetMisses;
+
+			int64 ZenRemoteGetHits = 0;
+			int64 ZenRemoteGetMisses = 0;
+			FDerivedDataCacheSpeedStats ZenRemoteSpeedStats;
 			if (ZenRemoteNode)
 			{
 				const FDerivedDataCacheUsageStats& UsageStats = (*ZenRemoteNode)->UsageStats.CreateConstIterator().Value();
-				SharedGetHits += UsageStats.GetStats.GetAccumulatedValueAnyThread(FCookStats::CallStats::EHitOrMiss::Hit, FCookStats::CallStats::EStatType::Counter);
-				SharedGetMisses += UsageStats.GetStats.GetAccumulatedValueAnyThread(FCookStats::CallStats::EHitOrMiss::Miss, FCookStats::CallStats::EStatType::Counter);
-				SharedSpeedStats = (*ZenRemoteNode)->SpeedStats;
+				ZenRemoteGetHits += UsageStats.GetStats.GetAccumulatedValueAnyThread(FCookStats::CallStats::EHitOrMiss::Hit, FCookStats::CallStats::EStatType::Counter);
+				ZenRemoteGetMisses += UsageStats.GetStats.GetAccumulatedValueAnyThread(FCookStats::CallStats::EHitOrMiss::Miss, FCookStats::CallStats::EStatType::Counter);
+				ZenRemoteSpeedStats = (*ZenRemoteNode)->SpeedStats;
 			}
-			const int64 SharedGetTotal = SharedGetHits + SharedGetMisses;
+			const int64 ZenRemoteGetTotal = ZenRemoteGetHits + ZenRemoteGetMisses;
 
 			int64 CloudGetHits = 0;
 			int64 CloudGetMisses = 0;
@@ -178,6 +188,12 @@ namespace UE::DerivedData::CookStats
 				TEXT("SharedGetHits"), SharedGetHits,
 				TEXT("SharedGetTotal"), SharedGetTotal,
 				TEXT("SharedGetHitPct"), SafeDivide(SharedGetHits, SharedGetTotal),
+				TEXT("ZenLocalGetHits"), ZenLocalGetHits,
+				TEXT("ZenLocalGetTotal"), ZenLocalGetTotal,
+				TEXT("ZenLocalGetHitPct"), SafeDivide(ZenLocalGetHits, ZenLocalGetTotal),
+				TEXT("ZenRemoteGetHits"), ZenRemoteGetHits,
+				TEXT("ZenRemoteGetTotal"), ZenRemoteGetTotal,
+				TEXT("ZenRemoteGetHitPct"), SafeDivide(ZenRemoteGetHits, ZenRemoteGetTotal),
 				TEXT("CloudGetHits"), CloudGetHits,
 				TEXT("CloudGetTotal"), CloudGetTotal,
 				TEXT("CloudGetHitPct"), SafeDivide(CloudGetHits, CloudGetTotal),
