Only report local stats for Zen DDC.
ZenCacheStore now only tracks what the client sees, regardless if the zenserver instance gets a hit from a configured upstream.
As we don't currently have any upstreams configured in the zen server instance this should give a more accurate number for hit/miss stats that matches the file cache.

diff --git a/Engine/Source/Developer/DerivedDataCache/Private/ZenCacheStore.cpp b/Engine/Source/Developer/DerivedDataCache/Private/ZenCacheStore.cpp
--- a/Engine/Source/Developer/DerivedDataCache/Private/ZenCacheStore.cpp
+++ b/Engine/Source/Developer/DerivedDataCache/Private/ZenCacheStore.cpp
@@ -31,14 +31,14 @@
 #include "ZenBackendUtils.h"
 #include "ZenSerialization.h"
 
-TRACE_DECLARE_INT_COUNTER(ZenDDC_Get,			TEXT("ZenDDC Get"));
-TRACE_DECLARE_INT_COUNTER(ZenDDC_GetHit,		TEXT("ZenDDC Get Hit"));
-TRACE_DECLARE_INT_COUNTER(ZenDDC_Put,			TEXT("ZenDDC Put"));
-TRACE_DECLARE_INT_COUNTER(ZenDDC_PutHit,		TEXT("ZenDDC Put Hit"));
+TRACE_DECLARE_INT_COUNTER(ZenDDC_Get, TEXT("ZenDDC Get"));
+TRACE_DECLARE_INT_COUNTER(ZenDDC_GetHit, TEXT("ZenDDC Get Hit"));
+TRACE_DECLARE_INT_COUNTER(ZenDDC_Put, TEXT("ZenDDC Put"));
+TRACE_DECLARE_INT_COUNTER(ZenDDC_PutHit, TEXT("ZenDDC Put Hit"));
 TRACE_DECLARE_INT_COUNTER(ZenDDC_BytesReceived, TEXT("ZenDDC Bytes Received"));
-TRACE_DECLARE_INT_COUNTER(ZenDDC_BytesSent,		TEXT("ZenDDC Bytes Sent"));
-TRACE_DECLARE_INT_COUNTER(ZenDDC_CacheRecordRequestCount, TEXT("ZenDDC CacheRecord Request Count"));
-TRACE_DECLARE_INT_COUNTER(ZenDDC_ChunkRequestCount, TEXT("ZenDDC Chunk Request Count"));
+TRACE_DECLARE_INT_COUNTER(ZenDDC_BytesSent, TEXT("ZenDDC Bytes Sent"));
+TRACE_DECLARE_INT_COUNTER(ZenDDC_CacheRecordRequestCountInFlight, TEXT("ZenDDC CacheRecord Request Count"));
+TRACE_DECLARE_INT_COUNTER(ZenDDC_ChunkRequestCountInFlight, TEXT("ZenDDC Chunk Request Count"));
 
 namespace UE::DerivedData
 {
@@ -176,6 +176,12 @@ public:
 		, Batches(Requests, [this](const FCachePutRequest& NextRequest) {return BatchGroupingFilter(NextRequest);})
 		, OnComplete(MoveTemp(InOnComplete))
 	{
+		TRACE_COUNTER_ADD(ZenDDC_Put, (int64)Requests.Num());
+		COOK_STAT(Timers.Reserve(Requests.Num()));
+		COOK_STAT(for (const FCachePutRequest& Request : Requests)
+		{
+			Timers.Add(Request.Record.GetKey(), CacheStore.UsageStats.TimePut());
+		})
 	}
 
 	void IssueRequests()
@@ -283,7 +289,15 @@ private:
 		const FCacheKey& Key = Request.Record.GetKey();
 		UE_LOG(LogDerivedDataCache, Verbose, TEXT("%s: Cache Put complete for %s from '%s'"),
 			*CacheStore.GetName(), *WriteToString<96>(Key), *Request.Name);
-		COOK_STAT(Timer.AddHit(Private::GetCacheRecordCompressedSize(Request.Record)));
+		int64 SentSize = 0;
+		for (const FValueWithId& Value : Request.Record.GetValues())
+		{
+			SentSize += Value.GetData().GetCompressed().GetSize();
+		}
+		TRACE_COUNTER_ADD(ZenDDC_BytesSent, SentSize);
+		TRACE_COUNTER_INCREMENT(ZenDDC_PutHit);
+		COOK_STAT(Timers[Key].AddHit(SentSize));
+
 		OnComplete(Request.MakeResponse(EStatus::Ok));
 	};
 
@@ -292,7 +306,12 @@ private:
 		const FCacheKey& Key = Request.Record.GetKey();
 		UE_LOG(LogDerivedDataCache, Verbose, TEXT("%s: Cache Put miss for '%s' from '%s'"),
 			*CacheStore.GetName(), *WriteToString<96>(Key), *Request.Name);
-		COOK_STAT(Timer.AddMiss());
+		int64 SentSize = 0;
+		for (const FValueWithId& Value : Request.Record.GetValues())
+		{
+			SentSize += Value.GetData().GetCompressed().GetSize();
+		}
+		COOK_STAT(Timers[Key].AddMiss(SentSize));
 		OnComplete(Request.MakeResponse(EStatus::Error));
 	};
 
@@ -302,7 +321,7 @@ private:
 	uint64 BatchSize = 0;
 	TBatchView<const FCachePutRequest> Batches;
 	FOnCachePutComplete OnComplete;
-	COOK_STAT(FCookStats::FScopedStatsCounter Timer = CacheStore.UsageStats.TimePut());
+	COOK_STAT(TMap<FCacheKey, FCookStats::FScopedStatsCounter> Timers);
 };
 
 class FZenCacheStore::FGetOp final : public FThreadSafeRefCountedObject
@@ -318,12 +337,17 @@ public:
 		, OnComplete(MoveTemp(InOnComplete))
 	{
 		TRACE_COUNTER_ADD(ZenDDC_Get, (int64)Requests.Num());
-		TRACE_COUNTER_ADD(ZenDDC_CacheRecordRequestCount, int64(Requests.Num()));
+		TRACE_COUNTER_ADD(ZenDDC_CacheRecordRequestCountInFlight, int64(Requests.Num()));
+		COOK_STAT(Timers.Reserve(Requests.Num()));
+		COOK_STAT(for (const FCacheGetRequest& Request : Requests)
+		{
+			Timers.Add(Request.Key, CacheStore.UsageStats.TimeGet());
+		})
 	}
 
 	virtual ~FGetOp()
 	{
-		TRACE_COUNTER_SUBTRACT(ZenDDC_CacheRecordRequestCount, int64(Requests.Num()));
+		TRACE_COUNTER_SUBTRACT(ZenDDC_CacheRecordRequestCountInFlight, int64(Requests.Num()));
 	}
 
 	void IssueRequests()
@@ -422,10 +446,10 @@ private:
 		UE_LOG(LogDerivedDataCache, Verbose, TEXT("%s: Cache hit for '%s' from '%s'"),
 			*CacheStore.GetName(), *WriteToString<96>(Request.Key), *Request.Name);
 
-		TRACE_COUNTER_ADD(ZenDDC_GetHit, int64(1));
+		TRACE_COUNTER_INCREMENT(ZenDDC_GetHit);
 		int64 ReceivedSize = Private::GetCacheRecordCompressedSize(Record);
 		TRACE_COUNTER_ADD(ZenDDC_BytesReceived, ReceivedSize);
-		COOK_STAT(Timer.AddHit(ReceivedSize));
+		COOK_STAT(Timers[Request.Key].AddHit(ReceivedSize));
 
 		OnComplete({Request.Name, MoveTemp(Record), Request.UserData, EStatus::Ok});
 	};
@@ -435,6 +459,7 @@ private:
 		UE_LOG(LogDerivedDataCache, Verbose, TEXT("%s: Cache miss for '%s' from '%s'"),
 			*CacheStore.GetName(), *WriteToString<96>(Request.Key), *Request.Name);
 
+		COOK_STAT(Timers[Request.Key].AddMiss());
 		OnComplete(Request.MakeResponse(EStatus::Error));
 	};
 
@@ -442,7 +467,7 @@ private:
 	IRequestOwner& Owner;
 	const TArray<FCacheGetRequest, TInlineAllocator<1>> Requests;
 	FOnCacheGetComplete OnComplete;
-	COOK_STAT(FCookStats::FScopedStatsCounter Timer = CacheStore.UsageStats.TimeGet());
+	COOK_STAT(TMap<FCacheKey, FCookStats::FScopedStatsCounter> Timers);
 };
 
 class FZenCacheStore::FPutValueOp final : public FThreadSafeRefCountedObject
@@ -458,6 +483,12 @@ public:
 		, Batches(Requests, [this](const FCachePutValueRequest& NextRequest) {return BatchGroupingFilter(NextRequest);})
 		, OnComplete(MoveTemp(InOnComplete))
 	{
+		TRACE_COUNTER_ADD(ZenDDC_Put, (int64)Requests.Num());
+		COOK_STAT(Timers.Reserve(Requests.Num()));
+		COOK_STAT(for (const FCachePutValueRequest& Request : Requests)
+		{
+			Timers.Add(Request.Key, CacheStore.UsageStats.TimePut());
+		})
 	}
 
 	void IssueRequests()
@@ -562,7 +593,10 @@ private:
 	{
 		UE_LOG(LogDerivedDataCache, Verbose, TEXT("%s: Cache PutValue complete for %s from '%s'"),
 			*CacheStore.GetName(), *WriteToString<96>(Request.Key), *Request.Name);
-		COOK_STAT(Timer.AddHit(Request.Value.GetData().GetCompressedSize()));
+		TRACE_COUNTER_INCREMENT(ZenDDC_PutHit);
+		int64 SentSize = Request.Value.GetData().GetCompressedSize();
+		TRACE_COUNTER_ADD(ZenDDC_BytesSent, SentSize);
+		COOK_STAT(Timers[Request.Key].AddHit(SentSize));
 		OnComplete(Request.MakeResponse(EStatus::Ok));
 	};
 
@@ -570,7 +604,8 @@ private:
 	{
 		UE_LOG(LogDerivedDataCache, Verbose, TEXT("%s: Cache PutValue miss for '%s' from '%s'"),
 			*CacheStore.GetName(), *WriteToString<96>(Request.Key), *Request.Name);
-		COOK_STAT(Timer.AddMiss());
+
+		COOK_STAT(Timers[Request.Key].AddMiss());
 		OnComplete(Request.MakeResponse(EStatus::Error));
 	};
 
@@ -580,7 +615,7 @@ private:
 	uint64 BatchSize = 0;
 	TBatchView<const FCachePutValueRequest> Batches;
 	FOnCachePutValueComplete OnComplete;
-	COOK_STAT(FCookStats::FScopedStatsCounter Timer = CacheStore.UsageStats.TimePut());
+	COOK_STAT(TMap<FCacheKey, FCookStats::FScopedStatsCounter> Timers);
 };
 
 class FZenCacheStore::FGetValueOp final : public FThreadSafeRefCountedObject
@@ -596,12 +631,17 @@ public:
 		, OnComplete(MoveTemp(InOnComplete))
 	{
 		TRACE_COUNTER_ADD(ZenDDC_Get, (int64)Requests.Num());
-		TRACE_COUNTER_ADD(ZenDDC_CacheRecordRequestCount, int64(Requests.Num()));
+		TRACE_COUNTER_ADD(ZenDDC_CacheRecordRequestCountInFlight, int64(Requests.Num()));
+		COOK_STAT(Timers.Reserve(Requests.Num()));
+		COOK_STAT(for (const FCacheGetValueRequest& Request : Requests)
+		{
+			Timers.Add(Request.Key, CacheStore.UsageStats.TimeGet());
+		})
 	}
 
 	virtual ~FGetValueOp()
 	{
-		TRACE_COUNTER_SUBTRACT(ZenDDC_CacheRecordRequestCount, int64(Requests.Num()));
+		TRACE_COUNTER_SUBTRACT(ZenDDC_CacheRecordRequestCountInFlight, int64(Requests.Num()));
 	}
 
 	void IssueRequests()
@@ -712,10 +752,10 @@ private:
 	{
 		UE_LOG(LogDerivedDataCache, Verbose, TEXT("%s: Cache hit for '%s' from '%s'"),
 			*CacheStore.GetName(), *WriteToString<96>(Request.Key), *Request.Name);
-					TRACE_COUNTER_ADD(ZenDDC_GetHit, int64(1));
+		TRACE_COUNTER_INCREMENT(ZenDDC_GetHit);
 		int64 ReceivedSize = Value.GetData().GetCompressedSize();
-					TRACE_COUNTER_ADD(ZenDDC_BytesReceived, ReceivedSize);
-					COOK_STAT(Timer.AddHit(ReceivedSize));
+		TRACE_COUNTER_ADD(ZenDDC_BytesReceived, ReceivedSize);
+		COOK_STAT(Timers[Request.Key].AddHit(ReceivedSize));
 
 		OnComplete({Request.Name, Request.Key, MoveTemp(Value), Request.UserData, EStatus::Ok});
 	};
@@ -725,6 +765,7 @@ private:
 		UE_LOG(LogDerivedDataCache, Verbose, TEXT("%s: Cache miss for '%s' from '%s'"),
 			*CacheStore.GetName(), *WriteToString<96>(Request.Key), *Request.Name);
 
+		COOK_STAT(Timers[Request.Key].AddMiss());
 		OnComplete(Request.MakeResponse(EStatus::Error));
 	};
 
@@ -732,7 +773,7 @@ private:
 	IRequestOwner& Owner;
 	const TArray<FCacheGetValueRequest, TInlineAllocator<1>> Requests;
 	FOnCacheGetValueComplete OnComplete;
-	COOK_STAT(FCookStats::FScopedStatsCounter Timer = CacheStore.UsageStats.TimeGet());
+	COOK_STAT(TMap<FCacheKey, FCookStats::FScopedStatsCounter> Timers);
 };
 
 class FZenCacheStore::FGetChunksOp final : public FThreadSafeRefCountedObject
@@ -747,14 +788,19 @@ public:
 		, Requests(InRequests)
 		, OnComplete(MoveTemp(InOnComplete))
 	{
-		TRACE_COUNTER_ADD(ZenDDC_ChunkRequestCount, int64(Requests.Num()));
+		TRACE_COUNTER_ADD(ZenDDC_ChunkRequestCountInFlight, int64(Requests.Num()));
 		TRACE_COUNTER_ADD(ZenDDC_Get, int64(Requests.Num()));
 		Requests.StableSort(TChunkLess());
+		COOK_STAT(Timers.Reserve(Requests.Num()));
+		COOK_STAT(for (const FCacheGetChunkRequest& Request : Requests)
+		{
+			Timers.Add(Request.Key, CacheStore.UsageStats.TimeGet());
+		})
 	}
 
 	virtual ~FGetChunksOp()
 	{
-		TRACE_COUNTER_SUBTRACT(ZenDDC_ChunkRequestCount, int64(Requests.Num()));
+		TRACE_COUNTER_SUBTRACT(ZenDDC_ChunkRequestCountInFlight, int64(Requests.Num()));
 	}
 
 	void IssueRequests()
@@ -889,8 +935,10 @@ private:
 	{
 		UE_LOG(LogDerivedDataCache, Verbose, TEXT("%s: CacheChunk hit for '%s' from '%s'"),
 			*CacheStore.GetName(), *WriteToString<96>(Request.Key, '/', Request.Id), *Request.Name);
-		TRACE_COUNTER_ADD(ZenDDC_GetHit, int64(1));
-		COOK_STAT(Timer.AddHit(RawSize));
+		TRACE_COUNTER_INCREMENT(ZenDDC_GetHit);
+		int64 ReceivedSize = RequestedBytes.GetSize();
+		TRACE_COUNTER_ADD(ZenDDC_BytesReceived, ReceivedSize);
+		COOK_STAT(Timers[Request.Key].AddHit(RawSize));
 		OnComplete({Request.Name, Request.Key, Request.Id, Request.RawOffset,
 			RawSize, MoveTemp(RawHash), MoveTemp(RequestedBytes), Request.UserData, EStatus::Ok});
 	};
@@ -899,6 +947,8 @@ private:
 	{
 		UE_LOG(LogDerivedDataCache, Verbose, TEXT("%s: CacheChunk miss with missing value '%s' for '%s' from '%s'"),
 			*CacheStore.GetName(), *WriteToString<16>(Request.Id), *WriteToString<96>(Request.Key), *Request.Name);
+
+		COOK_STAT(Timers[Request.Key].AddMiss());
 		OnComplete(Request.MakeResponse(EStatus::Error));
 	};
 
@@ -906,7 +956,7 @@ private:
 	IRequestOwner& Owner;
 	TArray<FCacheGetChunkRequest, TInlineAllocator<1>> Requests;
 	FOnCacheGetChunkComplete OnComplete;
-	COOK_STAT(FCookStats::FScopedStatsCounter Timer = CacheStore.UsageStats.TimeGet());
+	COOK_STAT(TMap<FCacheKey, FCookStats::FScopedStatsCounter> Timers);
 };
 
 class FZenCacheStore::FCbPackageReceiver final : public IHttpReceiver
@@ -1127,6 +1177,11 @@ void FZenCacheStore::EnqueueAsyncRpc(IRequestOwner& Owner, const FCbPackage& Req
 
 void FZenCacheStore::LegacyStats(FDerivedDataCacheStatsNode& OutNode)
 {
+	EDerivedDataCacheStatus CacheStatus = EDerivedDataCacheStatus::None;
+	OutNode = { TEXT("Zen"), ZenService.GetInstance().GetURL(), /*bIsLocal*/ ZenService.GetInstance().IsServiceRunningLocally()};
+	OutNode.UsageStats.Add(TEXT(""), UsageStats);
+#if 0
+	// DE: 20230213 We might need to revisit this if we change so Zen handles upstream again
 	Zen::FZenStats ZenStats;
 
 	FDerivedDataCacheUsageStats LocalStats;
@@ -1174,6 +1229,7 @@ void FZenCacheStore::LegacyStats(FDerivedDataCacheStatsNode& OutNode)
 	OutNode = {TEXT("Zen Group"), TEXT(""), /*bIsLocal*/ true};
 	OutNode.Children.Add(LocalNode);
 	OutNode.Children.Add(RemoteNode);
+#endif // 0
 }
 
 bool FZenCacheStore::LegacyDebugOptions(FBackendDebugOptions& InOptions)
