LaunchTaskInThreadPool API change needs to be properly deprecated
add pass-through for old args

diff --git a/Engine/Source/Developer/DerivedDataCache/Private/DerivedDataThreadPoolTask.cpp b/Engine/Source/Developer/DerivedDataCache/Private/DerivedDataThreadPoolTask.cpp
--- a/Engine/Source/Developer/DerivedDataCache/Private/DerivedDataThreadPoolTask.cpp
+++ b/Engine/Source/Developer/DerivedDataCache/Private/DerivedDataThreadPoolTask.cpp
@@ -144,6 +144,15 @@ void LaunchTaskInThreadPool(
 	}
 }
 
+void LaunchTaskInThreadPool(
+	uint64 MemoryEstimate,
+	IRequestOwner& Owner,
+	FQueuedThreadPool* ThreadPool,
+	TUniqueFunction<void ()>&& TaskBody)
+{
+	LaunchTaskInThreadPool(MemoryEstimate,TEXT("LaunchTaskInThreadPool"), Owner, ThreadPool, MoveTemp(TaskBody));
+}
+
 void LaunchTaskInThreadPool(
 	IRequestOwner& Owner,
 	FQueuedThreadPool* ThreadPool,
diff --git a/Engine/Source/Developer/DerivedDataCache/Public/DerivedDataThreadPoolTask.h b/Engine/Source/Developer/DerivedDataCache/Public/DerivedDataThreadPoolTask.h
--- a/Engine/Source/Developer/DerivedDataCache/Public/DerivedDataThreadPoolTask.h
+++ b/Engine/Source/Developer/DerivedDataCache/Public/DerivedDataThreadPoolTask.h
@@ -27,6 +27,7 @@ namespace UE::DerivedData
  * TaskBody is responsible for checking for cancellation by checking Owner.IsCanceled().
  *
  * @param MemoryEstimate   The estimated peak memory required to execute TaskBody.
+ * @param DebugName		   Name of the job, for debugging/logging purposes only
  * @param Owner            The owner of the request that will execute TaskBody.
  * @param ThreadPool       The thread pool within which to execute the task, if not retracted.
  * @param TaskBody         The task body that will be unconditionally executed exactly once.
@@ -37,7 +38,15 @@ UE_API void LaunchTaskInThreadPool(
 	IRequestOwner& Owner,
 	FQueuedThreadPool* ThreadPool,
 	TUniqueFunction<void ()>&& TaskBody);
+	
+UE_DEPRECATED(5.3, "Use LaunchTaskInThreadPool with DebugName")
+UE_API void LaunchTaskInThreadPool(
+	uint64 MemoryEstimate,
+	IRequestOwner& Owner,
+	FQueuedThreadPool* ThreadPool,
+	TUniqueFunction<void ()>&& TaskBody);
 
+// with no MemoryEstimate arg, assumes 0 (none) not -1 (unknown/max)
 UE_API void LaunchTaskInThreadPool(
 	IRequestOwner& Owner,
 	FQueuedThreadPool* ThreadPool,
