DDC: Added -DDC-VerifyFix to overwrite existing records and values that differ from newly generated records and values

diff --git a/Engine/Source/Developer/DerivedDataCache/Private/DerivedDataCacheStoreVerify.cpp b/Engine/Source/Developer/DerivedDataCache/Private/DerivedDataCacheStoreVerify.cpp
--- a/Engine/Source/Developer/DerivedDataCache/Private/DerivedDataCacheStoreVerify.cpp
+++ b/Engine/Source/Developer/DerivedDataCache/Private/DerivedDataCacheStoreVerify.cpp
@@ -62,6 +62,10 @@ public:
 			UE_LOG(LogDerivedDataCache, Display,
 				TEXT("Verify: Using salt -DDC-VerifySalt=%u to filter cache keys to verify."), Filter.GetSalt());
 		}
+
+		bPutOnError = bPutOnError || FParse::Param(CommandLine, TEXT("DDC-VerifyFix"));
+		UE_CLOG(bPutOnError, LogDerivedDataCache, Display,
+			TEXT("Verify: Any record or value that differs will be overwritten."));
 	}
 
 	void Put(
