DDC: Initialized bGetResult in case the callback is somehow missed

diff --git a/Engine/Source/Developer/DerivedDataCache/Private/DerivedDataCache.cpp b/Engine/Source/Developer/DerivedDataCache/Private/DerivedDataCache.cpp
--- a/Engine/Source/Developer/DerivedDataCache/Private/DerivedDataCache.cpp
+++ b/Engine/Source/Developer/DerivedDataCache/Private/DerivedDataCache.cpp
@@ -436,7 +436,7 @@ class FDerivedDataCache final
 			TRACE_CPUPROFILER_EVENT_SCOPE(DDC_DoWork);
 
 			const int64 NumBeforeDDC = Data.Num();
-			bool bGetResult;
+			bool bGetResult = false;
 			{
 				TRACE_CPUPROFILER_EVENT_SCOPE(DDC_Get);
 
