DDC: Added logging if legacy requests fail to copy the buffer

Includes a speculative workaround for check(Data.Num()) failing.

#rnx

diff --git a/Engine/Source/Developer/DerivedDataCache/Private/DerivedDataCache.cpp b/Engine/Source/Developer/DerivedDataCache/Private/DerivedDataCache.cpp
--- a/Engine/Source/Developer/DerivedDataCache/Private/DerivedDataCache.cpp
+++ b/Engine/Source/Developer/DerivedDataCache/Private/DerivedDataCache.cpp
@@ -461,15 +461,18 @@ class FDerivedDataCache final
 								{
 									Data.Append(static_cast<const uint8*>(Segment.GetData()), int64(Segment.GetSize()));
 								}
+								UE_CLOG(Data.Num() != int64(RawSize), LogDerivedDataCache, Display,
+									TEXT("Copied %" INT64_FMT " bytes when %" INT64_FMT " bytes were expected for %s from '%s'"),
+									Data.Num(), int64(RawSize), *CacheKey, *Response.Name);
 							}
 						});
 					BlockingOwner.Wait();
 				}
 				INC_FLOAT_STAT_BY(STAT_DDC_SyncGetTime, bSynchronousForStats ? (float)ThisTime : 0.0f);
 			}
-			if (bGetResult)
+
+			if (bGetResult && Data.Num())
 			{
-				
 				if (GVerifyDDC && DataDeriver && DataDeriver->IsDeterministic())
 				{
 					TArray<uint8> CmpData;
