Fix for DDC and Cook analytics event attributes not being transmitted as numbers to the backend. This allows us to correctly aggregate DDC/Cook values in Grafana.

diff --git a/Engine/Source/Developer/DerivedDataCache/Private/DerivedDataCache.cpp b/Engine/Source/Developer/DerivedDataCache/Private/DerivedDataCache.cpp
--- a/Engine/Source/Developer/DerivedDataCache/Private/DerivedDataCache.cpp
+++ b/Engine/Source/Developer/DerivedDataCache/Private/DerivedDataCache.cpp
@@ -1051,7 +1051,15 @@ public:
 		for (const FDerivedDataCacheSummaryStat& Stat : SummaryStats.Stats)
 		{
 			FString FormattedAttrName = "DDC.Summary." + Stat.Key;
-			Attributes.Emplace(FormattedAttrName, Stat.Value);
+
+			if (Stat.Value.IsNumeric())
+			{
+				Attributes.Emplace(FormattedAttrName, FCString::Atof(*Stat.Value));
+			}
+			else
+			{
+				Attributes.Emplace(FormattedAttrName, Stat.Value);
+			}
 		}
 
 		TSharedRef<FDerivedDataCacheStatsNode> RootNode = Backend->GatherUsageStats();
@@ -1059,7 +1067,14 @@ public:
 			{
 				for (const FCookStatsManager::StringKeyValue& Stat : Node->CustomStats)
 				{
-					Attributes.Emplace(Stat.Key, Stat.Value);
+					if (Stat.Value.IsNumeric())
+					{
+						Attributes.Emplace(Stat.Key, FCString::Atof(*Stat.Value));
+					}
+					else
+					{
+						Attributes.Emplace(Stat.Key, Stat.Value);
+					}
 				}
 			});
 #endif
diff --git a/Engine/Source/Editor/UnrealEd/Private/Cooker/CookProfiling.cpp b/Engine/Source/Editor/UnrealEd/Private/Cooker/CookProfiling.cpp
--- a/Engine/Source/Editor/UnrealEd/Private/Cooker/CookProfiling.cpp
+++ b/Engine/Source/Editor/UnrealEd/Private/Cooker/CookProfiling.cpp
@@ -913,7 +913,15 @@ void LogCookStats(ECookMode::Type CookMode)
 			for (const auto& Attr : StatAttributes)
 			{
 				FString FormattedAttrName = StatName + "." + Attr.Key;
-				Attributes.Emplace(FormattedAttrName, Attr.Value);
+
+				if (Attr.Value.IsNumeric())
+				{
+					Attributes.Emplace(FormattedAttrName, FCString::Atof(*Attr.Value));
+				}
+				else
+				{
+					Attributes.Emplace(FormattedAttrName, Attr.Value);
+				}
 			}
 		};
 
