Remove debug logging

diff --git a/Engine/Source/Developer/DerivedDataCache/Private/DerivedDataBuildJob.cpp b/Engine/Source/Developer/DerivedDataCache/Private/DerivedDataBuildJob.cpp
--- a/Engine/Source/Developer/DerivedDataCache/Private/DerivedDataBuildJob.cpp
+++ b/Engine/Source/Developer/DerivedDataCache/Private/DerivedDataBuildJob.cpp
@@ -1089,9 +1089,6 @@ void FBuildJob::AdvanceToState(EBuildJobState NewState, FUtf8StringView NewError
 
 void FBuildJob::ExecuteTransition(EBuildJobState OldState, EBuildJobState NewState)
 {
-	UE_LOG(LogTemp, Display, TEXT("Transition %s -> %s for build of '%s' by %s in job %p."),
-		LexToString(OldState), LexToString(NewState), *Name, *WriteToString<32>(FunctionName), this);
-
 	if (OldState == EBuildJobState::NotStarted)
 	{
 		BeginJob();
