DDC: Fixed Clang static analysis warnings

#rnx

diff --git a/Engine/Source/Developer/DerivedDataCache/Private/DerivedDataBuildJob.cpp b/Engine/Source/Developer/DerivedDataCache/Private/DerivedDataBuildJob.cpp
--- a/Engine/Source/Developer/DerivedDataCache/Private/DerivedDataBuildJob.cpp
+++ b/Engine/Source/Developer/DerivedDataCache/Private/DerivedDataBuildJob.cpp
@@ -174,7 +174,6 @@ private:
 	void SetAction(FBuildAction&& Action);
 	void SetInputs(FBuildInputs&& Inputs);
 	void SetOutput(const FBuildOutput& Output) final;
-	void SetOutputNoCheck(FBuildOutput&& Output, EBuildJobState NewState = EBuildJobState::CacheStore);
 
 	/** Terminate the job and send the error to the output complete callback. */
 	void CompleteWithError(FUtf8StringView Error) { AdvanceToState(EBuildJobState::Complete, Error); }
@@ -547,8 +546,9 @@ void FBuildJob::EndCacheQuery(FCacheGetResponse&& Response)
 	{
 		if (FOptionalBuildOutput CacheOutput = FBuildOutput::Load(Name, FunctionName, Response.Record))
 		{
+			Output = MoveTemp(CacheOutput).Get();
 			EnumAddFlags(BuildStatus, EBuildStatus::CacheQueryHit);
-			return SetOutputNoCheck(MoveTemp(CacheOutput).Get());
+			return AdvanceToState(EBuildJobState::Complete);
 		}
 	}
 	if (!EnumHasAnyFlags(BuildPolicy.GetCombinedPolicy() & Context->GetBuildPolicyMask(), EBuildPolicy::Build))
@@ -582,10 +582,13 @@ static ECachePolicy MakeCacheStorePolicy(EBuildPolicy BuildPolicy, const FBuildJ
 
 void FBuildJob::EnterCacheStore()
 {
+	checkf(!EnumHasAnyFlags(BuildStatus, EBuildStatus::CacheQueryHit),
+		TEXT("Job is not expected to store to the cache after a cache query hit for build of '%s' by %s."),
+		*Name, *WriteToString<32>(FunctionName));
+
 	if (!Cache || !Context ||
 		!EnumHasAnyFlags(Context->GetCachePolicyMask(), ECachePolicy::Store) ||
 		!EnumHasAnyFlags(Context->GetBuildPolicyMask() & BuildPolicy.GetCombinedPolicy(), EBuildPolicy::CacheStoreOnBuild) ||
-		EnumHasAnyFlags(BuildStatus, EBuildStatus::CacheQueryHit) ||
 		Output.Get().HasError() ||
 		Output.Get().HasLogs())
 	{
@@ -845,8 +848,9 @@ void FBuildJob::EndExecuteRemote(FBuildWorkerActionCompleteParams&& Params)
 	}
 	if (Params.Output)
 	{
+		Output = MoveTemp(Params.Output).Get();
 		EnumAddFlags(BuildStatus, EBuildStatus::BuildRemote);
-		return SetOutputNoCheck(MoveTemp(Params.Output).Get());
+		return AdvanceToState(EBuildJobState::CacheStore);
 	}
 	else
 	{
@@ -934,8 +938,9 @@ void FBuildJob::BeginExecuteLocal()
 
 void FBuildJob::EndExecuteLocal()
 {
+	Output = OutputBuilder.Build();
 	EnumAddFlags(BuildStatus, EBuildStatus::BuildLocal);
-	return SetOutputNoCheck(OutputBuilder.Build());
+	return AdvanceToState(EBuildJobState::CacheStore);
 }
 
 ///////////////////////////////////////////////////////////////////////////////////////////////////////////////////////
@@ -1016,14 +1021,9 @@ void FBuildJob::SetOutput(const FBuildOutput& InOutput)
 		State == EBuildJobState::ExecuteLocal,
 		TEXT("Job is not expecting an output in state %s for build of '%s' by %s."),
 		LexToString(State), *Name, *WriteToString<32>(FunctionName));
-	return SetOutputNoCheck(FBuildOutput(InOutput));
-}
-
-void FBuildJob::SetOutputNoCheck(FBuildOutput&& InOutput, EBuildJobState NewState)
-{
 	checkf(Output.IsNull(), TEXT("Job already has an output for build of '%s' by %s."), *Name, *WriteToString<32>(FunctionName));
-	Output = MoveTemp(InOutput);
-	return AdvanceToState(NewState);
+	Output = InOutput;
+	return AdvanceToState(EBuildJobState::CacheStore);
 }
 
 ///////////////////////////////////////////////////////////////////////////////////////////////////////////////////////
@@ -1117,9 +1117,9 @@ void FBuildJob::ExecuteTransition(EBuildJobState OldState, EBuildJobState NewSta
 		}
 		Schedule->EditParameters().ResolvedInputsSize = 0;
 	}
-	if (OldState <= EBuildJobState::ExecuteLocalWait && EBuildJobState::ExecuteLocalWait < NewState && !Output)
+	if (OldState < EBuildJobState::CacheStore && EBuildJobState::CacheStore <= NewState && !Output)
 	{
-		SetOutputNoCheck(OutputBuilder.Build(), NewState);
+		Output = OutputBuilder.Build();
 	}
 	if (NewState == EBuildJobState::Complete)
 	{
diff --git a/Engine/Source/Developer/DerivedDataCache/Private/PakFileCacheStore.cpp b/Engine/Source/Developer/DerivedDataCache/Private/PakFileCacheStore.cpp
--- a/Engine/Source/Developer/DerivedDataCache/Private/PakFileCacheStore.cpp
+++ b/Engine/Source/Developer/DerivedDataCache/Private/PakFileCacheStore.cpp
@@ -511,9 +511,9 @@ bool FPakFileCacheStore::LoadCache(const TCHAR* InFilename)
 		while (Loader.Tell() < SizeIndex)
 		{
 			FString Key;
-			int64 Offset;
-			int64 Size;
-			uint32 Crc;
+			int64 Offset = 0;
+			int64 Size = 0;
+			uint32 Crc = 0;
 			Loader << Key;
 			Loader << Offset;
 			Loader << Size;
