DDC: Removed deprecated code

#rnx

diff --git a/Engine/Source/Developer/DerivedDataCache/Private/DerivedDataCache.cpp b/Engine/Source/Developer/DerivedDataCache/Private/DerivedDataCache.cpp
--- a/Engine/Source/Developer/DerivedDataCache/Private/DerivedDataCache.cpp
+++ b/Engine/Source/Developer/DerivedDataCache/Private/DerivedDataCache.cpp
@@ -62,9 +62,9 @@ namespace UE::DerivedData::CookStats
 	// See https://developercommunity.visualstudio.com/content/problem/576913/c6244-regression-in-new-lambda-processorpermissive.html
 	static void AddCookStats(FCookStatsManager::AddStatFuncRef AddStat)
 	{
-		PRAGMA_DISABLE_DEPRECATION_WARNINGS;
+		PRAGMA_DISABLE_DEPRECATION_WARNINGS
 		TSharedRef<FDerivedDataCacheStatsNode> RootNode = GetDerivedDataCacheRef().GatherUsageStats();
-		PRAGMA_ENABLE_DEPRECATION_WARNINGS;
+		PRAGMA_ENABLE_DEPRECATION_WARNINGS
 
 		{
 			const FString StatName(TEXT("DDC.Usage"));
@@ -388,14 +388,6 @@ class FDerivedDataCache final
 	class FBuildAsyncWorker : public FNonAbandonableTask
 	{
 	public:
-		enum EWorkerState : uint32
-		{
-			WorkerStateNone			= 0,
-			WorkerStateRunning		= 1 << 0,
-			WorkerStateFinished		= 1 << 1,
-			WorkerStateDestroyed	= 1 << 2,
-		};
-
 		/** 
 		 * Constructor for async task 
 		 * @param	InDataDeriver	plugin to produce cache key and in the event of a miss, return the data.
@@ -412,27 +404,9 @@ class FDerivedDataCache final
 		{
 		}
 
-		virtual ~FBuildAsyncWorker()
-		{
-			// Record that the task is destroyed and check that it was not running or destroyed previously.
-			{
-				const uint32 PreviousState = WorkerState.fetch_or(WorkerStateDestroyed, std::memory_order_relaxed);
-				checkf(!(PreviousState & WorkerStateRunning), TEXT("Destroying DDC worker that is still running! Key: %s"), *CacheKey);
-				checkf(!(PreviousState & WorkerStateDestroyed), TEXT("Destroying DDC worker that has been destroyed previously! Key: %s"), *CacheKey);
-			}
-		}
-
 		/** Async worker that checks the cache backend and if that fails, calls the deriver to build the data and then puts the results to the cache **/
 		void DoWork()
 		{
-			// Record that the task is running and check that it was not running, finished, or destroyed previously.
-			{
-				const uint32 PreviousState = WorkerState.fetch_or(WorkerStateRunning, std::memory_order_relaxed);
-				checkf(!(PreviousState & WorkerStateRunning), TEXT("Starting DDC worker that is already running! Key: %s"), *CacheKey);
-				checkf(!(PreviousState & WorkerStateFinished), TEXT("Starting DDC worker that is already finished! Key: %s"), *CacheKey);
-				checkf(!(PreviousState & WorkerStateDestroyed), TEXT("Starting DDC worker that has been destroyed! Key: %s"), *CacheKey);
-			}
-
 			TRACE_CPUPROFILER_EVENT_SCOPE(DDC_DoWork);
 
 			const int64 NumBeforeDDC = Data.Num();
@@ -554,14 +528,6 @@ class FDerivedDataCache final
 				Data.Empty();
 			}
 			Backend->AddToAsyncCompletionCounter(-1);
-
-			// Record that the task is finished and check that it was running and not finished or destroyed previously.
-			{
-				const uint32 PreviousState = WorkerState.fetch_xor(WorkerStateRunning | WorkerStateFinished, std::memory_order_relaxed);
-				checkf((PreviousState & WorkerStateRunning), TEXT("Finishing DDC worker that was not running! Key: %s"), *CacheKey);
-				checkf(!(PreviousState & WorkerStateFinished), TEXT("Finishing DDC worker that is already finished! Key: %s"), *CacheKey);
-				checkf(!(PreviousState & WorkerStateDestroyed), TEXT("Finishing DDC worker that has been destroyed! Key: %s"), *CacheKey);
-			}
 		}
 
 		FORCEINLINE TStatId GetStatId() const
@@ -569,7 +535,6 @@ class FDerivedDataCache final
 			RETURN_QUICK_DECLARE_CYCLE_STAT(FBuildAsyncWorker, STATGROUP_ThreadPoolAsyncTasks);
 		}
 
-		std::atomic<uint32>				WorkerState{WorkerStateNone};
 		/** true in the case of a cache hit, otherwise the result of the deriver build call **/
 		bool							bSuccess;
 		/** true if we should record the timing **/
@@ -943,13 +908,6 @@ public:
 		Backend->GetDirectories(OutResults);
 	}
 
-	PRAGMA_DISABLE_DEPRECATION_WARNINGS
-	virtual IDDCCleanup* GetCleanup() const override
-	{
-		return const_cast<IDDCCleanup*>(static_cast<const IDDCCleanup*>(this));
-	}
-	PRAGMA_ENABLE_DEPRECATION_WARNINGS
-
 	virtual bool IsFinished() const override
 	{
 		return IsIdle();
diff --git a/Engine/Source/Developer/DerivedDataCache/Private/DerivedDataCacheModule.cpp b/Engine/Source/Developer/DerivedDataCache/Private/DerivedDataCacheModule.cpp
--- a/Engine/Source/Developer/DerivedDataCache/Private/DerivedDataCacheModule.cpp
+++ b/Engine/Source/Developer/DerivedDataCache/Private/DerivedDataCacheModule.cpp
@@ -19,16 +19,9 @@ static FDerivedDataCacheInterface* GDerivedDataLegacyCache;
 static ICache* GDerivedDataCache;
 static IBuild* GDerivedDataBuild;
 
-PRAGMA_DISABLE_DEPRECATION_WARNINGS
-
 class FDerivedDataCacheModule final : public IDerivedDataCacheModule
 {
 public:
-	FDerivedDataCacheInterface& GetDDC() final
-	{
-		return **CreateOrGetCache();
-	}
-
 	FDerivedDataCacheInterface* const* CreateOrGetCache() final
 	{
 		CreateCacheOnce();
@@ -81,8 +74,6 @@ private:
 	FCriticalSection CreateLock;
 };
 
-PRAGMA_ENABLE_DEPRECATION_WARNINGS
-
 static FDerivedDataCacheModule* GetModule()
 {
 	if (!FPlatformProperties::RequiresCookedData())
diff --git a/Engine/Source/Developer/DerivedDataCache/Private/DerivedDataCachePolicy.cpp b/Engine/Source/Developer/DerivedDataCache/Private/DerivedDataCachePolicy.cpp
--- a/Engine/Source/Developer/DerivedDataCache/Private/DerivedDataCachePolicy.cpp
+++ b/Engine/Source/Developer/DerivedDataCache/Private/DerivedDataCachePolicy.cpp
@@ -231,21 +231,6 @@ FCacheRecordPolicy FCacheRecordPolicy::Transform(const TFunctionRef<ECachePolicy
 	return Builder.Build();
 }
 
-void FCacheRecordPolicy::Save(FCbWriter& Writer) const
-{
-	Writer << *this;
-}
-
-FOptionalCacheRecordPolicy FCacheRecordPolicy::Load(const FCbObjectView Object)
-{
-	FCacheRecordPolicy Policy;
-	if (LoadFromCompactBinary(Object.AsFieldView(), Policy))
-	{
-		return Policy;
-	}
-	return {};
-}
-
 void FCacheRecordPolicyBuilder::AddValuePolicy(const FCacheValuePolicy& Value)
 {
 	checkf(!EnumHasAnyFlags(Value.Policy, ~Value.PolicyMask),
diff --git a/Engine/Source/Developer/DerivedDataCache/Public/DerivedDataCacheInterface.h b/Engine/Source/Developer/DerivedDataCache/Public/DerivedDataCacheInterface.h
--- a/Engine/Source/Developer/DerivedDataCache/Public/DerivedDataCacheInterface.h
+++ b/Engine/Source/Developer/DerivedDataCache/Public/DerivedDataCacheInterface.h
@@ -139,30 +139,6 @@ public:
 	// Low Level Interface
 	//--------------------
 
-	/** 
-	 * Synchronously checks the cache and if the item is present, it returns the cached results, otherwise it returns false
-	 * @param	CacheKey	Key to identify the data
-	 * @return	true if the data was retrieved from the cache
-	**/
-	UE_DEPRECATED(4.25, "'GetSynchronous' without a DebugContext is deprecated. Please provide a DebugContext!")
-	bool GetSynchronous(const TCHAR* CacheKey, TArray<uint8>& OutData) { return GetSynchronous(CacheKey, OutData, TEXT("Unknown Context")); }
-
-	/** 
-	 * Starts the async process of checking the cache and if the item is present, retrieving the cached results
-	 * @param	CacheKey		Key to identify the data
-	 * @return	a handle that can be used for PollAsynchronousCompletion, WaitAsynchronousCompletion and GetAsynchronousResults
-	**/
-	UE_DEPRECATED(4.25, "'GetAsynchronous' without a DebugContext is deprecated. Please provide a DebugContext!")
-	uint32 GetAsynchronous(const TCHAR* CacheKey) { return GetAsynchronous(CacheKey, TEXT("Unknown Context")); }
-
-	/** 
-	 * Puts data into the cache. This is fire-and-forget and typically asynchronous.
-	 * @param	CacheKey	Key to identify the data
-	 * @param	Data		Data to put in the cache under this key
-	**/
-	UE_DEPRECATED(4.25, "'Put' without a DebugContext is deprecated. Please provide a DebugContext!")
-	void Put(const TCHAR* CacheKey, TArray<uint8>& Data, bool bPutEvenIfExists = false) { Put(CacheKey, Data, TEXT("Unknown Context"), bPutEvenIfExists); }
-
 	/** 
 	 * Synchronously checks the cache and if the item is present, it returns the cached results, otherwise it returns false.
 	 *
@@ -283,12 +259,6 @@ public:
 	 */
 	bool IsDefaultGraph() const { return TStringView<TCHAR>(GetGraphName()).Equals(GetDefaultGraphName()); }
 
-	/**
-	 * Retrieve the interface to the background cache cleanup.
-	 */
-	UE_DEPRECATED(5.0, "This has been replaced by UE::DerivedData::GetCache().GetMaintainer().")
-	virtual IDDCCleanup* GetCleanup() const = 0;
-
 	//--------------------
 	// UsageStats Interface
 	//--------------------
diff --git a/Engine/Source/Developer/DerivedDataCache/Public/DerivedDataCacheModule.h b/Engine/Source/Developer/DerivedDataCache/Public/DerivedDataCacheModule.h
--- a/Engine/Source/Developer/DerivedDataCache/Public/DerivedDataCacheModule.h
+++ b/Engine/Source/Developer/DerivedDataCache/Public/DerivedDataCacheModule.h
@@ -13,10 +13,6 @@ class FDerivedDataCacheInterface;
 class IDerivedDataCacheModule : public IModuleInterface
 {
 public:
-	/** Return the DDC interface **/
-	UE_DEPRECATED(4.27, "GetDDC has been replaced by CreateOrGetCache.")
-	virtual FDerivedDataCacheInterface& GetDDC() = 0;
-
 	/**
 	 * Returns the cache, which is created by the first call to this function.
 	 *
diff --git a/Engine/Source/Developer/DerivedDataCache/Public/DerivedDataCachePolicy.h b/Engine/Source/Developer/DerivedDataCache/Public/DerivedDataCachePolicy.h
--- a/Engine/Source/Developer/DerivedDataCache/Public/DerivedDataCachePolicy.h
+++ b/Engine/Source/Developer/DerivedDataCache/Public/DerivedDataCachePolicy.h
@@ -110,10 +110,6 @@ UE_API FUtf8StringBuilderBase& operator<<(FUtf8StringBuilderBase& Builder, ECach
 UE_API bool TryLexFromString(ECachePolicy& OutPolicy, FUtf8StringView String);
 UE_API bool TryLexFromString(ECachePolicy& OutPolicy, FWideStringView String);
 
-UE_DEPRECATED(5.1, "Replace ParseCachePolicy with TryLexFromString.") UE_API ECachePolicy ParseCachePolicy(FAnsiStringView Text);
-UE_DEPRECATED(5.1, "Replace ParseCachePolicy with TryLexFromString.") UE_API ECachePolicy ParseCachePolicy(FWideStringView Text);
-UE_DEPRECATED(5.1, "Replace ParseCachePolicy with TryLexFromString.") UE_API ECachePolicy ParseCachePolicy(FUtf8StringView Text);
-
 UE_API FCbWriter& operator<<(FCbWriter& Writer, ECachePolicy Policy);
 UE_API bool LoadFromCompactBinary(FCbFieldView Field, ECachePolicy& OutPolicy, ECachePolicy Default = ECachePolicy::Default);
 
@@ -190,14 +186,6 @@ public:
 	/** Returns a copy of this policy transformed by an operation. */
 	UE_API FCacheRecordPolicy Transform(TFunctionRef<ECachePolicy (ECachePolicy)> Op) const;
 
-	/** Saves the cache record policy to a compact binary object. */
-	UE_DEPRECATED(5.1, "Replace Policy.Save(Writer) with Writer << Policy.")
-	UE_API void Save(FCbWriter& Writer) const;
-
-	/** Loads a cache record policy from an object. */
-	UE_DEPRECATED(5.1, "Replace Load(Object) with LoadFromCompactBinary(Field, Policy).")
-	UE_API static FOptionalCacheRecordPolicy Load(FCbObjectView Object);
-
 private:
 	friend class FCacheRecordPolicyBuilder;
 	friend class FOptionalCacheRecordPolicy;
