Fix C++20 compiler error. Add new equality operator for FCacheBucketOwner to avoid ambiguity when reversing arguments.

bool FCacheBucketOwner::operator==(const FCacheBucketOwner& Other) is new
bool FCacheBucketOwner::operator==(const FCacheBucket& Other) is what it would find ambiguous

diff --git a/Engine/Source/Developer/DerivedDataCache/Private/DerivedDataCacheKey.cpp b/Engine/Source/Developer/DerivedDataCache/Private/DerivedDataCacheKey.cpp
--- a/Engine/Source/Developer/DerivedDataCache/Private/DerivedDataCacheKey.cpp
+++ b/Engine/Source/Developer/DerivedDataCache/Private/DerivedDataCacheKey.cpp
@@ -29,6 +29,7 @@ public:
 	FCacheBucketOwner& operator=(const FCacheBucketOwner&) = delete;
 
 	using FCacheBucket::operator==;
+	inline bool operator==(const FCacheBucketOwner& Other) const { return FCacheBucket::operator==(Other); }
 	inline bool operator==(FUtf8StringView Bucket) const { return ToString() == Bucket; }
 };
 
