Change HttpCacheStoreTest to pass access token rather than OAuth details to Zen intermediary instances.  Also ensure that we don't overlap keys for Meta and NoMeta test data.

diff --git a/Engine/Source/Developer/DerivedDataCache/Private/HttpCacheStore.cpp b/Engine/Source/Developer/DerivedDataCache/Private/HttpCacheStore.cpp
--- a/Engine/Source/Developer/DerivedDataCache/Private/HttpCacheStore.cpp
+++ b/Engine/Source/Developer/DerivedDataCache/Private/HttpCacheStore.cpp
@@ -287,12 +287,17 @@ public:
 
 	const FString& GetDomain() const { return Domain; }
 	const FString& GetNamespace() const { return Namespace; }
-	const FString& GetOAuthProvider() const { return OAuthProvider; }
-	const FString& GetOAuthClientId() const { return OAuthClientId; }
-	const FString& GetOAuthSecret() const { return OAuthSecret; }
-	const FString& GetOAuthScope() const { return OAuthScope; }
-	const FString& GetOAuthProviderIdentifier() const { return OAuthProviderIdentifier; }
-	const FString& GetOAuthAccessToken() const { return OAuthAccessToken; }
+	const FString GetAccessToken() const
+	{
+		TAnsiStringBuilder<128> AccessTokenBuilder;
+
+		if (Access.IsValid())
+		{
+			AccessTokenBuilder << *Access;
+		}
+
+		return FString(ANSI_TO_TCHAR(AccessTokenBuilder.ToString()));
+	}
 
 private:
 	FString Domain;
@@ -2893,24 +2898,14 @@ TTuple<ILegacyCacheStore*, ECacheStoreFlags> CreateHttpCacheStore(const TCHAR* N
 
 ILegacyCacheStore* GetAnyHttpCacheStore(
 	FString& OutDomain,
-	FString& OutOAuthProvider,
-	FString& OutOAuthClientId,
-	FString& OutOAuthSecret,
-	FString& OutOAuthScope,
-	FString& OAuthProviderIdentifier,
-	FString& OAuthAccessToken,
+	FString& OutAccessToken,
 	FString& OutNamespace)
 {
 #if WITH_HTTP_DDC_BACKEND
 	if (FHttpCacheStore* HttpBackend = FHttpCacheStore::GetAny())
 	{
 		OutDomain = HttpBackend->GetDomain();
-		OutOAuthProvider = HttpBackend->GetOAuthProvider();
-		OutOAuthClientId = HttpBackend->GetOAuthClientId();
-		OutOAuthSecret = HttpBackend->GetOAuthSecret();
-		OutOAuthScope = HttpBackend->GetOAuthScope();
-		OAuthProviderIdentifier = HttpBackend->GetOAuthProviderIdentifier();
-		OAuthAccessToken = HttpBackend->GetOAuthAccessToken();
+		OutAccessToken = HttpBackend->GetAccessToken();
 		OutNamespace = HttpBackend->GetNamespace();
 		return HttpBackend;
 	}
diff --git a/Engine/Source/Developer/DerivedDataCache/Private/Tests/HttpCacheStoreTest.cpp b/Engine/Source/Developer/DerivedDataCache/Private/Tests/HttpCacheStoreTest.cpp
--- a/Engine/Source/Developer/DerivedDataCache/Private/Tests/HttpCacheStoreTest.cpp
+++ b/Engine/Source/Developer/DerivedDataCache/Private/Tests/HttpCacheStoreTest.cpp
@@ -41,12 +41,7 @@ namespace UE::DerivedData
 
 ILegacyCacheStore* GetAnyHttpCacheStore(
 	FString& OutDomain,
-	FString& OutOAuthProvider,
-	FString& OutOAuthClientId,
-	FString& OutOAuthSecret,
-	FString& OutOAuthScope,
-	FString& OAuthProviderIdentifier,
-	FString& OAuthAccessToken,
+	FString& OutAccessToken,
 	FString& OutNamespace);
 
 TTuple<ILegacyCacheStore*, ECacheStoreFlags> CreateZenCacheStore(const TCHAR* NodeName, const TCHAR* Config);
@@ -162,7 +157,7 @@ protected:
 	static ILegacyCacheStore* GetTestBackend()
 	{
 		static ILegacyCacheStore* CachedBackend = GetAnyHttpCacheStore(
-			TestDomain, TestOAuthProvider, TestOAuthClientId, TestOAuthSecret, TestOAuthScope, TestOAuthProviderIdentifier, TestOAuthAccessToken, TestNamespace);
+			TestDomain, TestAccessToken, TestNamespace);
 		return CachedBackend;
 	}
 
@@ -488,12 +483,7 @@ protected:
 
 protected:
 	static inline FString TestDomain;
-	static inline FString TestOAuthProvider;
-	static inline FString TestOAuthClientId;
-	static inline FString TestOAuthSecret;
-	static inline FString TestOAuthScope;
-	static inline FString TestOAuthProviderIdentifier;
-	static inline FString TestOAuthAccessToken;
+	static inline FString TestAccessToken;
 	static inline FString TestNamespace;
 };
 
@@ -505,10 +495,12 @@ TArray<FCacheRecord> CreateTestCacheRecords(ICacheStore* InTestBackend, uint32 I
 	TArray<FCacheRecord> CacheRecords;
 	TArray<FCachePutRequest> PutRequests;
 	PutRequests.Reserve(InNumKeys);
+	const uint32 KeySalt = 1;
 
 	for (uint32 KeyIndex = 0; KeyIndex < InNumKeys; ++KeyIndex)
 	{
 		FIoHashBuilder HashBuilder;
+		HashBuilder.Update(&KeySalt,sizeof(KeySalt));
 
 		TArray<FSharedBuffer> Values;
 		for (uint32 ValueIndex = 0; ValueIndex < InNumValues; ++ValueIndex)
@@ -525,6 +517,11 @@ TArray<FCacheRecord> CreateTestCacheRecords(ICacheStore* InTestBackend, uint32 I
 			HashBuilder.Update(Values.Last().GetView());
 		}
 
+		if (MetaContents)
+		{
+			MetaContents.AppendHash(HashBuilder);
+		}
+
 		FCacheKey Key;
 		Key.Bucket = TestCacheBucket;
 		Key.Hash = HashBuilder.Finalize();
@@ -624,11 +621,9 @@ bool FHttpCacheStoreTest::RunTest(const FString& Parameters)
 	FServiceSettings ZenUpstreamTestServiceSettings;
 	FServiceAutoLaunchSettings& ZenUpstreamTestAutoLaunchSettings = ZenUpstreamTestServiceSettings.SettingsVariant.Get<FServiceAutoLaunchSettings>();
 	ZenUpstreamTestAutoLaunchSettings.DataPath = FPaths::ConvertRelativePathToFull(FPaths::Combine(FPaths::EngineSavedDir(), "ZenUpstreamUnitTest"));
-	ZenUpstreamTestAutoLaunchSettings.ExtraArgs = FString::Printf(TEXT("--http asio --upstream-jupiter-url \"%s\" --upstream-jupiter-oauth-url \"%s\" --upstream-jupiter-oauth-clientid \"%s\" --upstream-jupiter-oauth-clientsecret \"%s\" --upstream-jupiter-namespace \"%s\""),
+	ZenUpstreamTestAutoLaunchSettings.ExtraArgs = FString::Printf(TEXT("--http asio --upstream-jupiter-url \"%s\" --upstream-jupiter-token \"%s\" --upstream-jupiter-namespace \"%s\""),
 		*TestDomain,
-		*TestOAuthProvider,
-		*TestOAuthClientId,
-		*TestOAuthSecret,
+		*TestAccessToken,
 		*TestNamespace
 	);
 	ZenUpstreamTestAutoLaunchSettings.DesiredPort = 23337; // Avoid the normal default port
@@ -640,11 +635,9 @@ bool FHttpCacheStoreTest::RunTest(const FString& Parameters)
 	FServiceSettings ZenUpstreamSiblingTestServiceSettings;
 	FServiceAutoLaunchSettings& ZenUpstreamSiblingTestAutoLaunchSettings = ZenUpstreamSiblingTestServiceSettings.SettingsVariant.Get<FServiceAutoLaunchSettings>();
 	ZenUpstreamSiblingTestAutoLaunchSettings.DataPath = FPaths::ConvertRelativePathToFull(FPaths::Combine(FPaths::EngineSavedDir(), "ZenUpstreamSiblingUnitTest"));
-	ZenUpstreamSiblingTestAutoLaunchSettings.ExtraArgs = FString::Printf(TEXT("--http asio --upstream-jupiter-url \"%s\" --upstream-jupiter-oauth-url \"%s\" --upstream-jupiter-oauth-clientid \"%s\" --upstream-jupiter-oauth-clientsecret \"%s\" --upstream-jupiter-namespace \"%s\""),
+	ZenUpstreamSiblingTestAutoLaunchSettings.ExtraArgs = FString::Printf(TEXT("--http asio --upstream-jupiter-url \"%s\" --upstream-jupiter-token \"%s\" --upstream-jupiter-namespace \"%s\""),
 		*TestDomain,
-		*TestOAuthProvider,
-		*TestOAuthClientId,
-		*TestOAuthSecret,
+		*TestAccessToken,
 		*TestNamespace
 	);
 	ZenUpstreamSiblingTestAutoLaunchSettings.DesiredPort = 23338; // Avoid the normal default port
