DDC: Restored the use of ManualResetEvent by HttpClient

#rnx

diff --git a/Engine/Source/Developer/DerivedDataCache/Private/Http/CurlHttpClient.cpp b/Engine/Source/Developer/DerivedDataCache/Private/Http/CurlHttpClient.cpp
--- a/Engine/Source/Developer/DerivedDataCache/Private/Http/CurlHttpClient.cpp
+++ b/Engine/Source/Developer/DerivedDataCache/Private/Http/CurlHttpClient.cpp
@@ -3,10 +3,10 @@
 #include "HttpClient.h"
 
 #include "Async/InheritedContext.h"
+#include "Async/ManualResetEvent.h"
 #include "Containers/ConsumeAllMpmcQueue.h"
 #include "Containers/LockFreeList.h"
 #include "Containers/StringView.h"
-#include "HAL/Event.h"
 #include "HAL/PlatformProcess.h"
 #include "HAL/Thread.h"
 #include "Memory/CompositeBuffer.h"
@@ -294,7 +294,7 @@ public:
 	TRefCountPtr<IHttpResponseMonitor> GetMonitor() final { return this; }
 
 	void Cancel() final;
-	void Wait() const final { CompleteEvent->Wait(); }
+	void Wait() const final { CompleteEvent.Wait(); }
 	bool Poll() const final { return EnumHasAnyFlags(State.load(), ECurlHttpResponseState::Complete); }
 	bool IsCanceled() const { return EnumHasAnyFlags(State.load(std::memory_order_relaxed), ECurlHttpResponseState::Canceled); }
 
@@ -340,7 +340,7 @@ private:
 	FAnsiStringView Error;
 	ANSICHAR ErrorBuffer[CURL_ERROR_SIZE]{};
 	FHttpResponseStats Stats;
-	FEventRef CompleteEvent{EEventMode::ManualReset};
+	mutable FManualResetEvent CompleteEvent;
 
 	enum class EHttpReceiverFunction : uint8
 	{
@@ -1138,7 +1138,7 @@ void FCurlHttpResponse::SetComplete(CURLcode Code)
 	}
 
 	CurlHttp::Private::AtomicFetchOr(State, ECurlHttpResponseState::Complete);
-	CompleteEvent->Trigger();
+	CompleteEvent.Notify();
 
 	if (UE_LOG_ACTIVE(LogHttp, Verbose))
 	{
@@ -1184,7 +1184,7 @@ void FCurlHttpResponse::Cancel()
 	// Waiting within a receiver function would prevent the response from completing.
 	if (LocalReceiverFunction == EHttpReceiverFunction::None)
 	{
-		CompleteEvent->Wait();
+		CompleteEvent.Wait();
 	}
 }
 
diff --git a/Engine/Source/Developer/DerivedDataCache/Private/Http/HttpClient.cpp b/Engine/Source/Developer/DerivedDataCache/Private/Http/HttpClient.cpp
--- a/Engine/Source/Developer/DerivedDataCache/Private/Http/HttpClient.cpp
+++ b/Engine/Source/Developer/DerivedDataCache/Private/Http/HttpClient.cpp
@@ -2,8 +2,8 @@
 
 #include "HttpClient.h"
 
+#include "Async/ManualResetEvent.h"
 #include "Containers/LockFreeList.h"
-#include "HAL/Event.h"
 #include "Memory/MemoryView.h"
 #include "Misc/AsciiSet.h"
 #include "ProfilingDebugging/CpuProfilerTrace.h"
@@ -271,7 +271,7 @@ struct Http::Private::FHttpRequestQueueData
 	struct FWaiter
 	{
 		THttpUniquePtr<IHttpRequest> Request;
-		FEventRef Event{EEventMode::ManualReset};
+		FManualResetEvent Event;
 	};
 
 	FHttpRequestQueueData(IHttpConnectionPool& ConnectionPool, const FHttpClientParams& ClientParams)
@@ -290,7 +290,7 @@ struct Http::Private::FHttpRequestQueueData
 					if (FWaiter* Waiter = Waiters.Pop())
 					{
 						Waiter->Request = MoveTemp(Request);
-						Waiter->Event->Trigger();
+						Waiter->Event.Notify();
 					}
 				}
 			}
@@ -302,7 +302,9 @@ struct Http::Private::FHttpRequestQueueData
 	{
 		if (Params.bIgnoreMaxRequests)
 		{
-			return Client->TryCreateRequest(Params);
+			THttpUniquePtr<IHttpRequest> Request = Client->TryCreateRequest(Params);
+			checkf(Request, TEXT("IHttpClient::TryCreateRequest returned null in spite of bIgnoreMaxRequests."));
+			return Request;
 		}
 
 		while (THttpUniquePtr<IHttpRequest> Request = Client->TryCreateRequest(Params))
@@ -310,7 +312,7 @@ struct Http::Private::FHttpRequestQueueData
 			if (FWaiter* Waiter = Waiters.Pop())
 			{
 				Waiter->Request = MoveTemp(Request);
-				Waiter->Event->Trigger();
+				Waiter->Event.Notify();
 			}
 			else
 			{
@@ -326,16 +328,18 @@ struct Http::Private::FHttpRequestQueueData
 			if (FWaiter* Waiter = Waiters.Pop())
 			{
 				Waiter->Request = MoveTemp(Request);
-				Waiter->Event->Trigger();
+				Waiter->Event.Notify();
 			}
-			if (LocalWaiter.Event->Wait(0))
+			if (LocalWaiter.Event.IsNotified())
 			{
+				checkf(LocalWaiter.Request, TEXT("CreateRequest returning null after IsNotified()."));
 				return MoveTemp(LocalWaiter.Request);
 			}
 		}
 
 		TRACE_CPUPROFILER_EVENT_SCOPE(HttpRequestQueue_Wait);
-		LocalWaiter.Event->Wait();
+		LocalWaiter.Event.Wait();
+		checkf(LocalWaiter.Request, TEXT("CreateRequest returning null after Wait()."));
 		return MoveTemp(LocalWaiter.Request);
 	}
 
