Formatting Changes to Cook Profiling Telemetry for Grafana. Replaced "." with "_". Bumped SchemaVersion of Cook Telemetry to 2.
Format changes will not have any impact on existing telemetry or analytics dashboards

diff --git a/Engine/Source/Developer/DerivedDataCache/Private/DerivedDataCache.cpp b/Engine/Source/Developer/DerivedDataCache/Private/DerivedDataCache.cpp
--- a/Engine/Source/Developer/DerivedDataCache/Private/DerivedDataCache.cpp
+++ b/Engine/Source/Developer/DerivedDataCache/Private/DerivedDataCache.cpp
@@ -1007,7 +1007,7 @@ public:
 		// Append to the attributes
 		for (const FDerivedDataCacheSummaryStat& Stat : SummaryStats.Stats)
 		{
-			FString FormattedAttrName = "DDC_Summary_" + Stat.Key;
+			FString FormattedAttrName = "DDC_Summary_" + Stat.Key.Replace(TEXT("."), TEXT("_"));
 
 			if (Stat.Value.IsNumeric())
 			{
@@ -1024,14 +1024,18 @@ public:
 			{
 				for (const FCookStatsManager::StringKeyValue& Stat : Node->CustomStats)
 				{
+					FString FormattedAttrName = Stat.Key.Replace(TEXT("."), TEXT("_"));
+
 					if (Stat.Value.IsNumeric())
 					{
-						Attributes.Emplace(Stat.Key, FCString::Atof(*Stat.Value));
+						Attributes.Emplace(FormattedAttrName, FCString::Atof(*Stat.Value));
 					}
 					else
 					{
-						Attributes.Emplace(Stat.Key, Stat.Value);
+						Attributes.Emplace(FormattedAttrName, Stat.Value);
 					}
+
+
 				}
 			});
 #endif
diff --git a/Engine/Source/Editor/UnrealEd/Private/Cooker/CookProfiling.cpp b/Engine/Source/Editor/UnrealEd/Private/Cooker/CookProfiling.cpp
--- a/Engine/Source/Editor/UnrealEd/Private/Cooker/CookProfiling.cpp
+++ b/Engine/Source/Editor/UnrealEd/Private/Cooker/CookProfiling.cpp
@@ -908,7 +908,7 @@ void LogCookStats(ECookMode::Type CookMode)
 
 	if (FStudioAnalytics::IsAvailable() && IsCookByTheBookMode(CookMode))
 	{
-		const int SchemaVersion = 1;
+		const int SchemaVersion = 2;
 
 		// convert filtered stats directly to an analytics event
 		TArray<FAnalyticsEventAttribute> Attributes;
@@ -920,7 +920,7 @@ void LogCookStats(ECookMode::Type CookMode)
 		{
 			for (const auto& Attr : StatAttributes)
 			{
-				FString FormattedAttrName = StatName + "_" + Attr.Key;
+				FString FormattedAttrName = (StatName + "_" + Attr.Key).Replace(TEXT("."), TEXT("_"));
 
 				if (Attr.Value.IsNumeric())
 				{
