Change Zen DDC cache store to allow:
- Local storage flag to be inferred by whether the service is run locally or set explicitly in config.
- Remote storage flag to be inferred by whether the service is NOT run locally or set explicitly in config.
- ReadOnly storage flag to be set explicitly in config (defaults to true).

diff --git a/Engine/Source/Developer/DerivedDataCache/Private/ZenCacheStore.cpp b/Engine/Source/Developer/DerivedDataCache/Private/ZenCacheStore.cpp
--- a/Engine/Source/Developer/DerivedDataCache/Private/ZenCacheStore.cpp
+++ b/Engine/Source/Developer/DerivedDataCache/Private/ZenCacheStore.cpp
@@ -121,6 +121,8 @@ public:
 	void LegacyStats(FDerivedDataCacheStatsNode& OutNode) final;
 	bool LegacyDebugOptions(FBackendDebugOptions& Options) final;
 
+	const Zen::FZenServiceInstance& GetServiceInstance() const { return ZenService.GetInstance(); }
+
 private:
 	void Initialize(const TCHAR* Namespace);
 
@@ -1408,7 +1410,24 @@ TTuple<ILegacyCacheStore*, ECacheStoreFlags> CreateZenCacheStore(const TCHAR* No
 		Backend.Reset();
 	}
 
-	return MakeTuple(Backend.Release(), ECacheStoreFlags::Local | ECacheStoreFlags::Remote | ECacheStoreFlags::Query | ECacheStoreFlags::Store);
+	bool bReadOnly = false;
+	FParse::Bool(Config, TEXT("ReadOnly="), bReadOnly);
+
+	// Default to locally launched service getting the Local cache store flag.  Can be overridden by explicit value in config.
+	const bool bServiceRunningLocally = Backend->GetServiceInstance().IsServiceRunningLocally();
+	bool bLocal = bServiceRunningLocally;
+	FParse::Bool(Config, TEXT("Local="), bLocal);
+
+	// Default to non-locally launched service getting the Remote cache store flag.  Can be overridden by explicit value in config.
+	// In the future this could be extended to allow the Remote flag by default (even for locally launched instances) if they have upstreams configured.
+	bool bRemote = !bServiceRunningLocally;
+	FParse::Bool(Config, TEXT("Remote="), bRemote);
+
+	ECacheStoreFlags Flags = ECacheStoreFlags::Query;
+	Flags |= bReadOnly ? ECacheStoreFlags::None : ECacheStoreFlags::Store;
+	Flags |= bLocal ? ECacheStoreFlags::Local : ECacheStoreFlags::None;
+	Flags |= bRemote ? ECacheStoreFlags::Remote : ECacheStoreFlags::None;
+	return MakeTuple(Backend.Release(), Flags);
 }
 
 } // namespace UE::DerivedData
