When a deactivation latency limit is defined, make the initial speed test executed by filesystem DDC time limited so we don't spend an inordinate amount of time speed testing if we're just going to deactivate it due to high latency.

diff --git a/Engine/Source/Developer/DerivedDataCache/Private/FileSystemCacheStore.cpp b/Engine/Source/Developer/DerivedDataCache/Private/FileSystemCacheStore.cpp
--- a/Engine/Source/Developer/DerivedDataCache/Private/FileSystemCacheStore.cpp
+++ b/Engine/Source/Developer/DerivedDataCache/Private/FileSystemCacheStore.cpp
@@ -783,7 +783,9 @@ public:
 		bool bSeekTimeOnly,
 		double& OutSeekTimeMS,
 		double& OutReadSpeedMBs,
-		double& OutWriteSpeedMBs);
+		double& OutWriteSpeedMBs,
+		std::atomic<uint32>* NumLatencyTestsCompleted,
+		std::atomic<bool>* AbandonRequest);
 
 	// ICacheStore Interface
 
@@ -863,6 +865,8 @@ private:
 	[[nodiscard]] bool FileExists(FStringBuilderBase& Path) const;
 	[[nodiscard]] bool IsDeactivatedForPerformance();
 
+	bool RunInitialSpeedTest();
+
 private:
 	/** Base path to store the cache files in. */
 	FString CachePath;
@@ -975,6 +979,15 @@ FFileSystemCacheStore::FFileSystemCacheStore(
 	FParse::Value(InParams, TEXT("ConsiderSlowAt="), ConsiderSlowAtMS);
 	FParse::Value(InParams, TEXT("DeactivateAt="), DeactivateAtMS);
 
+	const FString AbsoluteCachePath = IFileManager::Get().ConvertToAbsolutePathForExternalAppForRead(*CachePath);
+	if (AbsoluteCachePath.Len() >= GMaxCacheRootLen)
+	{
+		const FText ErrorMessage = FText::Format(NSLOCTEXT("DerivedDataCache", "PathTooLong", "Cache path {0} is longer than {1} characters... please adjust [DerivedDataBackendGraph] paths to be shorter (this leaves more room for cache keys)."),
+			FText::FromString(AbsoluteCachePath), FText::AsNumber(GMaxCacheRootLen));
+		FMessageDialog::Open(EAppMsgType::Ok, ErrorMessage);
+		UE_LOG(LogDerivedDataCache, Fatal, TEXT("%s"), *ErrorMessage.ToString());
+	}
+
 	// can skip the speed test so everything acts as local (e.g. 4.25 and earlier behavior).
 	bool SkipSpeedTest = !WITH_EDITOR || FParse::Param(FCommandLine::Get(), TEXT("DDCSkipSpeedTest"));
 	if (SkipSpeedTest)
@@ -986,17 +999,12 @@ FFileSystemCacheStore::FFileSystemCacheStore(
 	}
 
 	if (!SkipSpeedTest &&
-		!RunSpeedTest(CachePath,
-			bReadOnly,
-			false /* bSeekTimeOnly */,
-			SpeedStats.LatencyMS,
-			SpeedStats.ReadSpeedMBs,
-			SpeedStats.WriteSpeedMBs))
+		!RunInitialSpeedTest())
 	{
 		LastPerformanceEvaluationTicks.store(FDateTime::UtcNow().GetTicks(), std::memory_order_relaxed);
 
 		OutFlags = ECacheStoreFlags::None;
-		UE_LOG(LogDerivedDataCache, Warning, TEXT("No read or write access to %s"), *CachePath);
+		UE_LOG(LogDerivedDataCache, Warning, TEXT("No read or write access to %s, or abandoned due to slow progress"), *CachePath);
 	}
 	else
 	{
@@ -1183,7 +1191,9 @@ bool FFileSystemCacheStore::RunSpeedTest(
 	bool bSeekTimeOnly,
 	double& OutSeekTimeMS,
 	double& OutReadSpeedMBs,
-	double& OutWriteSpeedMBs)
+	double& OutWriteSpeedMBs,
+	std::atomic<uint32>* NumLatencyTestsCompleted,
+	std::atomic<bool>* AbandonRequest)
 {
 	SCOPED_BOOT_TIMING("RunSpeedTest");
 	UE_SCOPED_ENGINE_ACTIVITY("Running IO speed test");
@@ -1204,15 +1214,6 @@ bool FFileSystemCacheStore::RunSpeedTest(
 	int TotalDataRead = 0;
 	int TotalDataWritten = 0;
 
-	const FString AbsoluteCachePath = IFileManager::Get().ConvertToAbsolutePathForExternalAppForRead(*CachePath);
-	if (AbsoluteCachePath.Len() >= GMaxCacheRootLen)
-	{
-		const FText ErrorMessage = FText::Format(NSLOCTEXT("DerivedDataCache", "PathTooLong", "Cache path {0} is longer than {1} characters... please adjust [DerivedDataBackendGraph] paths to be shorter (this leaves more room for cache keys)."),
-			FText::FromString(AbsoluteCachePath), FText::AsNumber(GMaxCacheRootLen));
-		FMessageDialog::Open(EAppMsgType::Ok, ErrorMessage);
-		UE_LOG(LogDerivedDataCache, Fatal, TEXT("%s"), *ErrorMessage.ToString());
-	}
-
 	TArray<FString> Paths;
 	TArray<FString> MissingFiles;
 
@@ -1242,6 +1243,18 @@ bool FFileSystemCacheStore::RunSpeedTest(
 	{
 		FFileStatData StatData = IFileManager::Get().GetStatData(*KV.Key);
 
+		if (NumLatencyTestsCompleted)
+		{
+			NumLatencyTestsCompleted->fetch_add(1, std::memory_order_release);
+		}
+
+		if (AbandonRequest && AbandonRequest->load(std::memory_order_relaxed))
+		{
+			const double TotalTestTime = FPlatformTime::Seconds() - StatStartTime;
+			UE_LOG(LogDerivedDataCache, Log, TEXT("Speed tests for %s abandoned on request after %.02f seconds"), *CachePath, TotalTestTime);
+			return false;
+		}
+
 		if (!StatData.bIsValid || StatData.FileSize != KV.Value)
 		{
 			MissingFiles.Add(KV.Key);
@@ -1288,6 +1301,13 @@ bool FFileSystemCacheStore::RunSpeedTest(
 					break;
 				}
 			}
+
+			if (AbandonRequest && AbandonRequest->load(std::memory_order_relaxed))
+			{
+				const double TotalTestTime = FPlatformTime::Seconds() - StatStartTime;
+				UE_LOG(LogDerivedDataCache, Log, TEXT("Speed tests for %s abandoned on request after %.02f seconds"), *CachePath, TotalTestTime);
+				return false;
+			}
 		}
 	}
 
@@ -1317,6 +1337,13 @@ bool FFileSystemCacheStore::RunSpeedTest(
 			}
 
 			TotalDataRead += TempData.Num();
+			
+			if (AbandonRequest && AbandonRequest->load(std::memory_order_relaxed))
+			{
+				const double TotalTestTime = FPlatformTime::Seconds() - StatStartTime;
+				UE_LOG(LogDerivedDataCache, Log, TEXT("Speed tests for %s abandoned on request after %.02f seconds"), *CachePath, TotalTestTime);
+				return false;
+			}
 		}
 
 		TotalReadTime = FPlatformTime::Seconds() - ReadStartTime;
@@ -1359,6 +1386,13 @@ bool FFileSystemCacheStore::RunSpeedTest(
 			}
 
 			TotalDataWritten += TempData.Num();
+			
+			if (AbandonRequest && AbandonRequest->load(std::memory_order_relaxed))
+			{
+				const double TotalTestTime = FPlatformTime::Seconds() - StatStartTime;
+				UE_LOG(LogDerivedDataCache, Log, TEXT("Speed tests for %s abandoned on request after %.02f seconds"), *CachePath, TotalTestTime);
+				return false;
+			}
 		}
 
 		TotalWriteTime = FPlatformTime::Seconds() - WriteStartTime;
@@ -1366,6 +1400,14 @@ bool FFileSystemCacheStore::RunSpeedTest(
 		UE_LOG(LogDerivedDataCache, Verbose, TEXT("write tests %s on %s and took %.02f seconds"),
 			bWriteTestPassed ? TEXT("passed") : TEXT("failed"), *CachePath, TotalReadTime)
 
+			
+		if (AbandonRequest && AbandonRequest->load(std::memory_order_relaxed))
+		{
+			const double TotalTestTime = FPlatformTime::Seconds() - StatStartTime;
+			UE_LOG(LogDerivedDataCache, Log, TEXT("Speed tests for %s abandoned on request after %.02f seconds"), *CachePath, TotalTestTime);
+			return false;
+		}
+		
 		// remove the custom path but do it async as this can be slow on remote drives
 		AsyncTask(ENamedThreads::AnyThread, [CustomPath]() {
 			IFileManager::Get().DeleteDirectory(*CustomPath, false, true);
@@ -2620,7 +2662,9 @@ bool FFileSystemCacheStore::IsDeactivatedForPerformance()
 							true /* bSeekTimeOnly */,
 							LocalSpeedStats.LatencyMS,
 							LocalSpeedStats.ReadSpeedMBs,
-							LocalSpeedStats.WriteSpeedMBs);
+							LocalSpeedStats.WriteSpeedMBs,
+							nullptr,
+							nullptr);
 
 						if (LocalSpeedStats.LatencyMS >= DeactivateAtMS)
 						{
@@ -2635,6 +2679,62 @@ bool FFileSystemCacheStore::IsDeactivatedForPerformance()
 	return true;
 }
 
+bool FFileSystemCacheStore::RunInitialSpeedTest()
+{
+	struct FSpeedTestState : public FThreadSafeRefCountedObject
+	{
+		FDerivedDataCacheSpeedStats SpeedStats;
+		std::atomic<uint32> NumLatencyTestsCompleted = 0;
+		std::atomic<bool> AbandonRequest = false;
+		bool bResult = false;
+		FManualResetEvent CompletionEvent;
+	};
+
+	TRefCountPtr<FSpeedTestState> SpeedTestState = new FSpeedTestState();
+	Tasks::Launch(TEXT("FFileSystemCacheStore::InitialEvaluation"),
+		[CachePath = this->CachePath, bReadOnly = this->bReadOnly, SpeedTestState]()
+		{
+			SpeedTestState->bResult = RunSpeedTest(CachePath,
+				bReadOnly,
+				false /* bSeekTimeOnly */,
+				SpeedTestState->SpeedStats.LatencyMS,
+				SpeedTestState->SpeedStats.ReadSpeedMBs,
+				SpeedTestState->SpeedStats.WriteSpeedMBs,
+				&SpeedTestState->NumLatencyTestsCompleted,
+				&SpeedTestState->AbandonRequest);
+			SpeedTestState->CompletionEvent.Notify();
+		});
+
+	if (!GIsBuildMachine && FPlatformProcess::SupportsMultithreading() && (DeactivateAtMS > 0.f))
+	{
+		if (SpeedTestState->CompletionEvent.WaitFor(FMonotonicTimeSpan::FromMilliseconds(DeactivateAtMS * 2.f)))
+		{
+			// If the task completed in the initial wait period, return the result
+			SpeedStats = SpeedTestState->SpeedStats;
+			return SpeedTestState->bResult;
+		}
+		else
+		{
+			// If the task did not complete the initial wait period, evaluate if we're progressing fast enough to keep waiting
+			// or we should abandon it and supply generic "bad" speed test results.
+			if (SpeedTestState->NumLatencyTestsCompleted.load(std::memory_order_acquire) < 2)
+			{
+				SpeedTestState->AbandonRequest.store(true, std::memory_order_relaxed);
+				SpeedStats.ReadSpeedMBs = 0.0;
+				SpeedStats.WriteSpeedMBs = 0.0;
+				SpeedStats.LatencyMS = 999.0;
+				UE_LOG(LogDerivedDataCache, Log, TEXT("Skipping speed test to %s due to slow test progress. Assuming poor performance"), *CachePath);
+				return true;
+			}
+		}
+	}
+
+	// Wait indefinitely for completion
+	SpeedTestState->CompletionEvent.Wait();
+	SpeedStats = SpeedTestState->SpeedStats;
+	return SpeedTestState->bResult;
+}
+
 ILegacyCacheStore* CreateFileSystemCacheStore(
 	const TCHAR* CachePath,
 	const TCHAR* Params,
