DDC: Fixes related to the use of OIDC functionality from DesktopPlatform

- Fixed a crash when HttpCacheStore is used without DesktopPlatform.
- Fixed GetOidcAccessToken and GetOidcTokenStatus failing for targets with no Intermediate directory.
- Added DesktopPlatform to DerivedDataTool to allow it to use HttpCacheStore with OIDC.

diff --git a/Engine/Source/Developer/DerivedDataCache/Private/HttpCacheStore.cpp b/Engine/Source/Developer/DerivedDataCache/Private/HttpCacheStore.cpp
--- a/Engine/Source/Developer/DerivedDataCache/Private/HttpCacheStore.cpp
+++ b/Engine/Source/Developer/DerivedDataCache/Private/HttpCacheStore.cpp
@@ -1751,9 +1751,10 @@ bool FHttpCacheStore::AcquireAccessToken(IHttpClient* Client)
 	{
 		FString AccessTokenString;
 		FDateTime TokenExpiresAt;
-		bool bWasInteractiveLogin;
+		bool bWasInteractiveLogin = false;
 
-		if (FDesktopPlatformModule::Get()->GetOidcAccessToken(FPaths::RootDir(), FPaths::GetProjectFilePath(), OAuthProviderIdentifier, FApp::IsUnattended(), GWarn, AccessTokenString, TokenExpiresAt, bWasInteractiveLogin))
+		IDesktopPlatform* DesktopPlatform = FDesktopPlatformModule::TryGet();
+		if (DesktopPlatform && DesktopPlatform->GetOidcAccessToken(FPaths::RootDir(), FPaths::GetProjectFilePath(), OAuthProviderIdentifier, FApp::IsUnattended(), GWarn, AccessTokenString, TokenExpiresAt, bWasInteractiveLogin))
 		{
 			if (bWasInteractiveLogin)
 			{
@@ -1767,13 +1768,22 @@ bool FHttpCacheStore::AcquireAccessToken(IHttpClient* Client)
 			SetAccessTokenAndUnlock(Lock, AccessTokenString, ExpiryTimeSeconds);
 			return true;
 		}
-
-		UE_LOG(LogDerivedDataCache, Warning, TEXT("%s: OidcToken: Failed to log in to HTTP services."), *Domain);
-		FailedLoginAttempts++;
-		return false;
+		else if (DesktopPlatform)
+		{
+			UE_LOG(LogDerivedDataCache, Warning, TEXT("%s: OidcToken: Failed to log in to HTTP services."), *Domain);
+			FailedLoginAttempts++;
+			return false;
+		}
+		else
+		{
+			UE_LOG(LogDerivedDataCache, Warning, TEXT("%s: OidcToken: Use of OAuthProviderIdentifier requires that the target depend on DesktopPlatform."), *Domain);
+			FailedLoginAttempts++;
+			return false;
+		}
 	}
 
 	UE_LOG(LogDerivedDataCache, Warning, TEXT("%s: No available configuration to acquire an access token."), *Domain);
+	FailedLoginAttempts++;
 	return false;
 }
 
diff --git a/Engine/Source/Developer/DesktopPlatform/Private/DesktopPlatformBase.cpp b/Engine/Source/Developer/DesktopPlatform/Private/DesktopPlatformBase.cpp
--- a/Engine/Source/Developer/DesktopPlatform/Private/DesktopPlatformBase.cpp
+++ b/Engine/Source/Developer/DesktopPlatform/Private/DesktopPlatformBase.cpp
@@ -808,6 +808,7 @@ bool FDesktopPlatformBase::IsUnrealBuildToolRunning()
 
 bool FDesktopPlatformBase::GetOidcAccessToken(const FString& RootDir, const FString& ProjectFileName, const FString& ProviderIdentifier, bool bUnattended, FFeedbackContext* Warn, FString& OutToken, FDateTime& OutTokenExpiresAt, bool& bOutWasInteractiveLogin)
 {
+	IFileManager::Get().MakeDirectory(*FPaths::ProjectIntermediateDir(), /*bTree*/ true);
 	FString ResultFilePath = FPaths::CreateTempFilename(*FPaths::ProjectIntermediateDir(), TEXT("oidcToken.json"));
 	ON_SCOPE_EXIT
 	{
@@ -891,6 +892,7 @@ bool FDesktopPlatformBase::GetOidcAccessToken(const FString& RootDir, const FStr
 
 bool FDesktopPlatformBase::GetOidcTokenStatus(const FString& RootDir, const FString& ProjectFileName, const FString& ProviderIdentifier, FFeedbackContext* Warn, int& OutStatus)
 {
+	IFileManager::Get().MakeDirectory(*FPaths::ProjectIntermediateDir(), /*bTree*/ true);
 	FString ResultFilePath = FPaths::CreateTempFilename(*FPaths::ProjectIntermediateDir(), TEXT("oidcToken-status.json"));
 
 	FString Arguments = TEXT(" ");
diff --git a/Engine/Source/Developer/DesktopPlatform/Public/DesktopPlatformModule.h b/Engine/Source/Developer/DesktopPlatform/Public/DesktopPlatformModule.h
--- a/Engine/Source/Developer/DesktopPlatform/Public/DesktopPlatformModule.h
+++ b/Engine/Source/Developer/DesktopPlatform/Public/DesktopPlatformModule.h
@@ -21,6 +21,15 @@ public:
 		return DesktopPlatformModule.GetSingleton();
 	}
 
+	static IDesktopPlatform* TryGet()
+	{
+		if (FDesktopPlatformModule* DesktopPlatformModule = FModuleManager::Get().LoadModulePtr<FDesktopPlatformModule>("DesktopPlatform"))
+		{
+			return DesktopPlatformModule->GetSingleton();
+		}
+		return nullptr;
+	}
+
 private:
 	virtual IDesktopPlatform* GetSingleton() const { return DesktopPlatform; }
 
diff --git a/Engine/Source/Programs/DerivedDataTool/DerivedDataTool.Build.cs b/Engine/Source/Programs/DerivedDataTool/DerivedDataTool.Build.cs
--- a/Engine/Source/Programs/DerivedDataTool/DerivedDataTool.Build.cs
+++ b/Engine/Source/Programs/DerivedDataTool/DerivedDataTool.Build.cs
@@ -7,8 +7,10 @@ public class DerivedDataTool : ModuleRules
 	public DerivedDataTool(ReadOnlyTargetRules Target) : base(Target)
 	{
 		PublicIncludePathModuleNames.Add("Launch");
+		PrivateDependencyModuleNames.Add("ApplicationCore");
 		PrivateDependencyModuleNames.Add("Core");
-		PrivateDependencyModuleNames.Add("Projects");
 		PrivateDependencyModuleNames.Add("DerivedDataCache");
+		PrivateDependencyModuleNames.Add("DesktopPlatform");
+		PrivateDependencyModuleNames.Add("Projects");
 	}
 }
diff --git a/Engine/Source/Programs/DerivedDataTool/DerivedDataTool.Target.cs b/Engine/Source/Programs/DerivedDataTool/DerivedDataTool.Target.cs
--- a/Engine/Source/Programs/DerivedDataTool/DerivedDataTool.Target.cs
+++ b/Engine/Source/Programs/DerivedDataTool/DerivedDataTool.Target.cs
@@ -13,7 +13,7 @@ public class DerivedDataToolTarget : TargetRules
 		LaunchModuleName = "DerivedDataTool";
 
 		bBuildDeveloperTools = false;
-		bCompileAgainstApplicationCore = false;
+		bCompileAgainstApplicationCore = true;
 		bCompileAgainstCoreUObject = false;
 		bCompileAgainstEngine = false;
 		bCompileICU = false;
