DDC: Added EHttpVersion to HttpClient

#rnx

diff --git a/Engine/Source/Developer/DerivedDataCache/Private/Http/CurlHttpClient.cpp b/Engine/Source/Developer/DerivedDataCache/Private/Http/CurlHttpClient.cpp
--- a/Engine/Source/Developer/DerivedDataCache/Private/Http/CurlHttpClient.cpp
+++ b/Engine/Source/Developer/DerivedDataCache/Private/Http/CurlHttpClient.cpp
@@ -220,6 +220,7 @@ private:
 	~FCurlHttpClient() final;
 	void Destroy() final { delete this; }
 
+	static long ConvertVersion(EHttpVersion Version);
 	static long ConvertTlsLevel(EHttpTlsLevel Level);
 
 	FCurlHttpConnectionPool& ConnectionPool;
@@ -627,7 +628,7 @@ void FCurlHttpClient::SetDefaultOptions(CURL* Curl, FCurlHttpHeaders& Headers) c
 {
 	ConnectionPool.SetDefaultOptions(Curl, Headers);
 
-	curl_easy_setopt(Curl, CURLOPT_HTTP_VERSION, CURL_HTTP_VERSION_2_0);
+	curl_easy_setopt(Curl, CURLOPT_HTTP_VERSION, ConvertVersion(Params.Version));
 
 	curl_easy_setopt(Curl, CURLOPT_DNS_CACHE_TIMEOUT, FMath::Min<long>(Params.DnsCacheTimeout, MAX_int32));
 	curl_easy_setopt(Curl, CURLOPT_CONNECTTIMEOUT_MS, long(Params.ConnectTimeout));
@@ -642,6 +643,19 @@ void FCurlHttpClient::SetDefaultOptions(CURL* Curl, FCurlHttpHeaders& Headers) c
 	curl_easy_setopt(Curl, CURLOPT_VERBOSE, long(Params.bVerbose));
 }
 
+long FCurlHttpClient::ConvertVersion(const EHttpVersion Version)
+{
+	switch (Version)
+	{
+	case EHttpVersion::None:   return CURL_HTTP_VERSION_NONE;
+	case EHttpVersion::V1_0:   return CURL_HTTP_VERSION_1_0;
+	case EHttpVersion::V1_1:   return CURL_HTTP_VERSION_1_1;
+	case EHttpVersion::V2:     return CURL_HTTP_VERSION_2_0;
+	case EHttpVersion::V2Only: return CURL_HTTP_VERSION_2_PRIOR_KNOWLEDGE;
+	default: checkNoEntry();   return CURL_HTTP_VERSION_NONE;
+	}
+}
+
 long FCurlHttpClient::ConvertTlsLevel(const EHttpTlsLevel Level)
 {
 	switch (Level)
diff --git a/Engine/Source/Developer/DerivedDataCache/Private/Http/HttpClient.h b/Engine/Source/Developer/DerivedDataCache/Private/Http/HttpClient.h
--- a/Engine/Source/Developer/DerivedDataCache/Private/Http/HttpClient.h
+++ b/Engine/Source/Developer/DerivedDataCache/Private/Http/HttpClient.h
@@ -75,6 +75,20 @@ enum class EHttpMediaType : uint8
 FAnsiStringView LexToString(EHttpMediaType MediaType);
 bool TryLexFromString(EHttpMediaType& OutMediaType, FAnsiStringView View);
 
+enum class EHttpVersion : uint8
+{
+	/** Determine version automatically. */
+	None = 0,
+	/** Use HTTP/1.0 */
+	V1_0,
+	/** Use HTTP/1.1 */
+	V1_1,
+	/** Try HTTP/2 with fallback to HTTP/1.1 */
+	V2,
+	/** Use HTTP/2 */
+	V2Only,
+};
+
 enum class EHttpTlsLevel : uint8
 {
 	/** Do not use TLS. */
@@ -177,6 +191,9 @@ struct FHttpClientParams final
 	/** Time in seconds after which a request will abort if it stays below LowSpeedLimit. Use 0 to disable. */
 	uint32 LowSpeedTime = 0;
 
+	/** HTTP version to use for requests created by the client. */
+	EHttpVersion Version = EHttpVersion::None;
+
 	/** Level of TLS to use for requests created by the client. */
 	EHttpTlsLevel TlsLevel = EHttpTlsLevel::None;
 
diff --git a/Engine/Source/Developer/DerivedDataCache/Private/HttpCacheStore.cpp b/Engine/Source/Developer/DerivedDataCache/Private/HttpCacheStore.cpp
--- a/Engine/Source/Developer/DerivedDataCache/Private/HttpCacheStore.cpp
+++ b/Engine/Source/Developer/DerivedDataCache/Private/HttpCacheStore.cpp
@@ -1643,6 +1643,7 @@ FHttpClientParams FHttpCacheStore::GetDefaultClientParams() const
 	ClientParams.ConnectTimeout = 3 * 1000;
 	ClientParams.LowSpeedLimit = 1024;
 	ClientParams.LowSpeedTime = 10;
+	ClientParams.Version = EHttpVersion::V2;
 	ClientParams.TlsLevel = EHttpTlsLevel::All;
 	ClientParams.bFollowRedirects = true;
 	ClientParams.bFollow302Post = true;
diff --git a/Engine/Source/Developer/DerivedDataCache/Private/ZenCacheStore.cpp b/Engine/Source/Developer/DerivedDataCache/Private/ZenCacheStore.cpp
--- a/Engine/Source/Developer/DerivedDataCache/Private/ZenCacheStore.cpp
+++ b/Engine/Source/Developer/DerivedDataCache/Private/ZenCacheStore.cpp
@@ -1125,6 +1125,7 @@ void FZenCacheStore::Initialize(
 		ConnectionPool = IHttpManager::Get().CreateConnectionPool(ConnectionPoolParams);
 
 		FHttpClientParams ClientParams;
+		ClientParams.Version = EHttpVersion::V2;
 		ClientParams.MaxRequests = RequestPoolSize + RequestPoolOverflowSize;
 		ClientParams.MinRequests = RequestPoolSize;
 		RequestQueue = FHttpRequestQueue(*ConnectionPool, ClientParams);
