DDC: Fixed merging and filtering of partial records in the cache hierarchy

diff --git a/Engine/Source/Developer/DerivedDataCache/Private/DerivedDataCacheStoreHierarchy.cpp b/Engine/Source/Developer/DerivedDataCache/Private/DerivedDataCacheStoreHierarchy.cpp
--- a/Engine/Source/Developer/DerivedDataCache/Private/DerivedDataCacheStoreHierarchy.cpp
+++ b/Engine/Source/Developer/DerivedDataCache/Private/DerivedDataCacheStoreHierarchy.cpp
@@ -765,7 +765,7 @@ void FCacheStoreHierarchy::TGetBatch<Params>::DispatchRequests()
 	{
 		if (State.Response.Status != EStatus::Ok)
 		{
-			OnComplete(MoveTemp(State.Response));
+			OnComplete(FilterResponseByRequest(State.Response, State.Request));
 		}
 	}
 }
@@ -868,6 +868,11 @@ bool FCacheStoreHierarchy::FCacheRecordBatchParams::MergeFromResponse(
 		return true;
 	}
 
+	if (Response.Record.GetValues().IsEmpty() && !OutResponse.Record.GetValues().IsEmpty())
+	{
+		return true;
+	}
+
 	const FCacheRecord& NewRecord = Response.Record;
 	const FCacheRecord ExistingRecord = MoveTemp(OutResponse.Record);
 
@@ -951,6 +956,11 @@ FCacheGetResponse FCacheStoreHierarchy::FCacheRecordBatchParams::FilterResponseB
 	const FCacheGetRequest& Request)
 {
 	const ECachePolicy RecordPolicy = Request.Policy.GetRecordPolicy();
+	if (Response.Status != EStatus::Ok && !EnumHasAnyFlags(RecordPolicy, ECachePolicy::PartialRecord))
+	{
+		FCacheRecordBuilder Builder(Response.Record.GetKey());
+		return {Response.Name, Builder.Build(), Response.UserData, Response.Status};
+	}
 	const bool bMightSkipData = EnumHasAnyFlags(RecordPolicy, ECachePolicy::SkipData) || !Request.Policy.IsUniform();
 	if ((bMightSkipData && Algo::AnyOf(Response.Record.GetValues(), &FValue::HasData)) ||
 		(EnumHasAnyFlags(RecordPolicy, ECachePolicy::SkipMeta) && Response.Record.GetMeta()))
