DDC: Removed logging of the OAuth secret from HttpCacheStore

#rnx

diff --git a/Engine/Source/Developer/DerivedDataCache/Private/HttpCacheStore.cpp b/Engine/Source/Developer/DerivedDataCache/Private/HttpCacheStore.cpp
--- a/Engine/Source/Developer/DerivedDataCache/Private/HttpCacheStore.cpp
+++ b/Engine/Source/Developer/DerivedDataCache/Private/HttpCacheStore.cpp
@@ -2773,7 +2773,7 @@ void FHttpCacheStoreParams::Parse(const TCHAR* NodeName, const TCHAR* Config)
 	{
 		if (FParse::Value(FCommandLine::Get(), *(OverrideName + TEXT("=")), OAuthSecret))
 		{
-			UE_LOG(LogDerivedDataCache, Log, TEXT("%s: Found command line override for OAuthSecret %s=%s"), NodeName, *OverrideName, *OAuthSecret);
+			UE_LOG(LogDerivedDataCache, Log, TEXT("%s: Found command line override for OAuthSecret %s={SECRET}"), NodeName, *OverrideName);
 		}
 	}
 
