DDC: Added EnvOAuthSecretOverride parameter to HttpCacheStore

diff --git a/Engine/Source/Developer/DerivedDataCache/Private/HttpCacheStore.cpp b/Engine/Source/Developer/DerivedDataCache/Private/HttpCacheStore.cpp
--- a/Engine/Source/Developer/DerivedDataCache/Private/HttpCacheStore.cpp
+++ b/Engine/Source/Developer/DerivedDataCache/Private/HttpCacheStore.cpp
@@ -2767,8 +2767,17 @@ void FHttpCacheStoreParams::Parse(const TCHAR* NodeName, const TCHAR* Config)
 	}
 
 	FParse::Value(Config, TEXT("OAuthClientId="), OAuthClientId);
-	FParse::Value(Config, TEXT("OAuthSecret="), OAuthSecret);
 
+	FParse::Value(Config, TEXT("OAuthSecret="), OAuthSecret);
+	if (FParse::Value(Config, TEXT("EnvOAuthSecretOverride="), OverrideName))
+	{
+		FString OAuthSecretEnv = FPlatformMisc::GetEnvironmentVariable(*OverrideName);
+		if (!OAuthSecretEnv.IsEmpty())
+		{
+			OAuthSecret = OAuthSecretEnv;
+			UE_LOG(LogDerivedDataCache, Log, TEXT("%s: Found environment override for OAuthSecret %s={SECRET}"), NodeName, *OverrideName);
+		}
+	}
 	if (FParse::Value(Config, TEXT("CommandLineOAuthSecretOverride="), OverrideName))
 	{
 		if (FParse::Value(FCommandLine::Get(), *(OverrideName + TEXT("=")), OAuthSecret))
