Check status code for responses received from zenserver DDC requests.

diff --git a/Engine/Source/Developer/DerivedDataCache/Private/ZenCacheStore.cpp b/Engine/Source/Developer/DerivedDataCache/Private/ZenCacheStore.cpp
--- a/Engine/Source/Developer/DerivedDataCache/Private/ZenCacheStore.cpp
+++ b/Engine/Source/Developer/DerivedDataCache/Private/ZenCacheStore.cpp
@@ -230,7 +230,7 @@ public:
 			auto OnRpcComplete = [this, OpRef = TRefCountPtr<FPutOp>(this), Batch](THttpUniquePtr<IHttpResponse>& HttpResponse, FCbPackage& Response)
 			{
 				int32 RequestIndex = 0;
-				if (HttpResponse->GetErrorCode() == EHttpErrorCode::None)
+				if (HttpResponse->GetErrorCode() == EHttpErrorCode::None && HttpResponse->GetStatusCode() >= 200 && HttpResponse->GetStatusCode() <= 299)
 				{
 					const FCbObject& ResponseObj = Response.GetObject();
 					for (FCbField ResponseField : ResponseObj[ANSITEXTVIEW("Result")])
@@ -259,6 +259,13 @@ public:
 							*CacheStore.GetName(), Batch.Num(), RequestIndex, *WriteToString<256>(*HttpResponse));
 					}
 				}
+				else if (HttpResponse->GetStatusCode() != 404)
+				{
+					UE_LOG(LogDerivedDataCache, Warning,
+						TEXT("%s: Error response received from PutCacheRecords RPC: from %s"),
+						*CacheStore.GetName(), *WriteToString<256>(*HttpResponse));
+				}
+
 				for (const FCachePutRequest& Request : Batch.RightChop(RequestIndex))
 				{
 					OnMiss(Request);
@@ -400,7 +407,7 @@ public:
 				auto OnRpcComplete = [this, OpRef = TRefCountPtr<FGetOp>(OriginalOp), Batch](THttpUniquePtr<IHttpResponse>& HttpResponse, FCbPackage& Response)
 				{
 					int32 RequestIndex = 0;
-					if (HttpResponse->GetErrorCode() == EHttpErrorCode::None)
+					if (HttpResponse->GetErrorCode() == EHttpErrorCode::None && HttpResponse->GetStatusCode() >= 200 && HttpResponse->GetStatusCode() <= 299)
 					{
 						const FCbObject& ResponseObj = Response.GetObject();
 						
@@ -434,6 +441,13 @@ public:
 								*CacheStore.GetName(), Batch.Num(), RequestIndex, *WriteToString<256>(*HttpResponse));
 						}
 					}
+					else if (HttpResponse->GetStatusCode() != 404)
+					{
+						UE_LOG(LogDerivedDataCache, Warning,
+							TEXT("%s: Error response received from GetCacheRecords RPC: from %s"),
+							*CacheStore.GetName(), *WriteToString<256>(*HttpResponse));
+					}
+
 					for (const FCacheGetRequest& Request : Batch.RightChop(RequestIndex))
 					{
 						OnMiss(Request);
@@ -541,7 +555,7 @@ public:
 			auto OnRpcComplete = [this, OpRef = TRefCountPtr<FPutValueOp>(this), Batch](THttpUniquePtr<IHttpResponse>& HttpResponse, FCbPackage& Response)
 			{
 				int32 RequestIndex = 0;
-				if (HttpResponse->GetErrorCode() == EHttpErrorCode::None)
+				if (HttpResponse->GetErrorCode() == EHttpErrorCode::None && HttpResponse->GetStatusCode() >= 200 && HttpResponse->GetStatusCode() <= 299)
 				{
 					const FCbObject& ResponseObj = Response.GetObject();
 					for (FCbField ResponseField : ResponseObj[ANSITEXTVIEW("Result")])
@@ -569,6 +583,13 @@ public:
 							*CacheStore.GetName(), Batch.Num(), RequestIndex, *WriteToString<256>(*HttpResponse));
 					}
 				}
+				else if (HttpResponse->GetStatusCode() != 404)
+				{
+					UE_LOG(LogDerivedDataCache, Warning,
+						TEXT("%s: Error response received from PutCacheValues RPC: from %s"),
+						*CacheStore.GetName(), *WriteToString<256>(*HttpResponse));
+				}
+
 				for (const FCachePutValueRequest& Request : Batch.RightChop(RequestIndex))
 				{
 					OnMiss(Request);
@@ -695,7 +716,7 @@ public:
 			auto OnRpcComplete = [this, OpRef = TRefCountPtr<FGetValueOp>(OriginalOp), Batch](THttpUniquePtr<IHttpResponse>& HttpResponse, FCbPackage& Response)
 			{
 				int32 RequestIndex = 0;
-				if (HttpResponse->GetErrorCode() == EHttpErrorCode::None)
+				if (HttpResponse->GetErrorCode() == EHttpErrorCode::None && HttpResponse->GetStatusCode() >= 200 && HttpResponse->GetStatusCode() <= 299)
 				{
 					const FCbObject& ResponseObj = Response.GetObject();
 
@@ -742,6 +763,13 @@ public:
 							*CacheStore.GetName(), Batch.Num(), RequestIndex, *WriteToString<256>(*HttpResponse));
 					}
 				}
+				else if (HttpResponse->GetStatusCode() != 404)
+				{
+					UE_LOG(LogDerivedDataCache, Warning,
+						TEXT("%s: Error response received from GetCacheValues RPC: from %s"),
+						*CacheStore.GetName(), *WriteToString<256>(*HttpResponse));
+				}
+
 				for (const FCacheGetValueRequest& Request : Batch.RightChop(RequestIndex))
 				{
 					OnMiss(Request);
@@ -872,7 +900,7 @@ public:
 			auto OnRpcComplete = [this, OpRef = TRefCountPtr<FGetChunksOp>(OriginalOp), Batch](THttpUniquePtr<IHttpResponse>& HttpResponse, FCbPackage& Response)
 			{
 				int32 RequestIndex = 0;
-				if (HttpResponse->GetErrorCode() == EHttpErrorCode::None)
+				if (HttpResponse->GetErrorCode() == EHttpErrorCode::None && HttpResponse->GetStatusCode() >= 200 && HttpResponse->GetStatusCode() <= 299)
 				{
 					const FCbObject& ResponseObj = Response.GetObject();
 
@@ -927,6 +955,13 @@ public:
 						Succeeded ? OnHit(Request, MoveTemp(RawHash), RawSize, MoveTemp(RequestedBytes)) : OnMiss(Request);
 					}
 				}
+				else if (HttpResponse->GetStatusCode() != 404)
+				{
+					UE_LOG(LogDerivedDataCache, Warning,
+						TEXT("%s: Error response received from GetChunks RPC: from %s"),
+						*CacheStore.GetName(), *WriteToString<256>(*HttpResponse));
+				}
+
 				for (const FCacheGetChunkRequest& Request : Batch.RightChop(RequestIndex))
 				{
 					OnMiss(Request);
