DDC: Moved file system cache store configuration into the file system cache store

#rnx

diff --git a/Engine/Source/Developer/DerivedDataCache/Private/DerivedDataBackends.cpp b/Engine/Source/Developer/DerivedDataCache/Private/DerivedDataBackends.cpp
--- a/Engine/Source/Developer/DerivedDataCache/Private/DerivedDataBackends.cpp
+++ b/Engine/Source/Developer/DerivedDataCache/Private/DerivedDataBackends.cpp
@@ -65,7 +65,7 @@ ILegacyCacheStore* CreateCacheStoreAsync(ILegacyCacheStore* InnerBackend, ECache
 ILegacyCacheStore* CreateCacheStoreHierarchy(ICacheStoreOwner*& OutOwner, IMemoryCacheStore* MemoryCache);
 ILegacyCacheStore* CreateCacheStoreThrottle(ILegacyCacheStore* InnerCache, uint32 LatencyMS, uint32 MaxBytesPerSecond);
 ILegacyCacheStore* CreateCacheStoreVerify(ILegacyCacheStore* InnerCache, bool bPutOnError);
-ILegacyCacheStore* CreateFileSystemCacheStore(const TCHAR* CacheDirectory, const TCHAR* Params, const TCHAR* AccessLogFileName, ECacheStoreFlags& OutFlags);
+ILegacyCacheStore* CreateFileSystemCacheStore(const TCHAR* NodeName, const TCHAR* Config, ECacheStoreFlags& OutFlags, FString& OutPath);
 TTuple<ILegacyCacheStore*, ECacheStoreFlags> CreateHttpCacheStore(const TCHAR* NodeName, const TCHAR* Config);
 IMemoryCacheStore* CreateMemoryCacheStore(const TCHAR* Name, int64 MaxCacheSize, bool bCanBeDisabled);
 IPakFileCacheStore* CreatePakFileCacheStore(const TCHAR* Filename, bool bWriting, bool bCompressed);
@@ -576,6 +576,16 @@ public:
 	 */
 	FParsedNode ParseHierarchyNode(const TCHAR* NodeName, const TCHAR* Entry, const FString& IniFilename, const TCHAR* IniSection, FParsedNodeMap& InParsedNodes)
 	{
+		if (Hierarchy)
+		{
+			UE_LOG(LogDerivedDataCache, Warning, TEXT("Node %s is disabled because there may be only one hierarchy node. "
+				"Confirm there is only one hierarchy in the cache graph and that it is inside of any async node."), NodeName);
+			return MakeTuple(nullptr, ECacheStoreFlags::None);
+		}
+
+		ILegacyCacheStore* HierarchyStore = CreateCacheStoreHierarchy(Hierarchy, GetMemoryCache());
+		CreatedNodes.AddUnique(HierarchyStore);
+
 		const TCHAR* InnerMatch = TEXT("Inner=");
 		const int32 InnerMatchLength = FCString::Strlen(InnerMatch);
 
@@ -603,15 +613,6 @@ public:
 			UE_LOG(LogDerivedDataCache, Warning, TEXT("Hierarchical cache %s has no inner backends and will not be created."), NodeName);
 			return MakeTuple(nullptr, ECacheStoreFlags::None);
 		}
-
-		if (Hierarchy)
-		{
-			UE_LOG(LogDerivedDataCache, Warning, TEXT("Node %s is disabled because there may be only one hierarchy node. "
-				"Confirm there is only one hierarchy in the cache graph and that it is inside of any async node."), NodeName);
-			return MakeTuple(nullptr, ECacheStoreFlags::None);
-		}
-
-		ILegacyCacheStore* HierarchyStore = CreateCacheStoreHierarchy(Hierarchy, GetMemoryCache());
 		ECacheStoreFlags Flags = ECacheStoreFlags::None;
 		for (const FParsedNode& Node : InnerNodes)
 		{
@@ -628,137 +629,17 @@ public:
 	 * @param Entry Node definition.
 	 * @return Filesystem data cache backend interface instance or nullptr if unsuccessful
 	 */
-	FParsedNode ParseDataCache( const TCHAR* NodeName, const TCHAR* Entry )
+	FParsedNode ParseDataCache(const TCHAR* NodeName, const TCHAR* Config)
 	{
-		FParsedNode DataCache(nullptr, ECacheStoreFlags::None);
-
-		// Parse Path by default, it may be overwritten by EnvPathOverride
 		FString Path;
-		FParse::Value( Entry, TEXT("Path="), Path );
-
-		// Check the EnvPathOverride environment variable to allow persistent overriding of data cache path, eg for offsite workers.
-		FString EnvPathOverride;
-		if( FParse::Value( Entry, TEXT("EnvPathOverride="), EnvPathOverride ) )
-		{
-			FString FilesystemCachePathEnv = FPlatformMisc::GetEnvironmentVariable( *EnvPathOverride );
-			if( FilesystemCachePathEnv.Len() > 0 )
-			{
-				Path = FilesystemCachePathEnv;
-				UE_LOG( LogDerivedDataCache, Log, TEXT("Found environment variable %s=%s"), *EnvPathOverride, *Path );
-			}
-		}
-
-		if (!EnvPathOverride.IsEmpty())
-		{
-			FString DDCPath;
-			if (FPlatformMisc::GetStoredValue(TEXT("Epic Games"), TEXT("GlobalDataCachePath"), *EnvPathOverride, DDCPath))
-			{
-				if (DDCPath.Len() > 0)
-				{
-					Path = DDCPath;
-					UE_LOG( LogDerivedDataCache, Log, TEXT("Found registry key GlobalDataCachePath %s=%s"), *EnvPathOverride, *Path );
-				}
-			}
-		}
-
-		// Check the CommandLineOverride argument to allow redirecting in build scripts
-		FString CommandLineOverride;
-		if( FParse::Value( Entry, TEXT("CommandLineOverride="), CommandLineOverride ) )
+		ECacheStoreFlags Flags;
+		if (ILegacyCacheStore* Store = CreateFileSystemCacheStore(NodeName, Config, Flags, Path))
 		{
-			FString Value;
-			if (FParse::Value(FCommandLine::Get(), *(CommandLineOverride + TEXT("=")), Value))
-			{
-				Path = Value;
-				UE_LOG(LogDerivedDataCache, Log, TEXT("Found command line override %s=%s"), *CommandLineOverride, *Path);
-			}
+			bUsingSharedDDC |= NodeName == TEXTVIEW("Shared");
+			Directories.AddUnique(Path);
+			return MakeTuple(Store, Flags);
 		}
-
-		// Paths starting with a '?' are looked up from config
-		if (Path.StartsWith(TEXT("?")) && !GConfig->GetString(TEXT("DerivedDataCacheSettings"), *Path + 1, Path, GEngineIni))
-		{
-			Path.Empty();
-		}
-
-		// Allow the user to override it from the editor
-		FString EditorOverrideSetting;
-		if(FParse::Value(Entry, TEXT("EditorOverrideSetting="), EditorOverrideSetting))
-		{
-			FString Setting = GConfig->GetStr(TEXT("/Script/UnrealEd.EditorSettings"), *EditorOverrideSetting, GEditorSettingsIni);
-			if(Setting.Len() > 0)
-			{
-				FString SettingPath;
-				if(FParse::Value(*Setting, TEXT("Path="), SettingPath))
-				{
-					SettingPath.TrimQuotesInline();
-					SettingPath.ReplaceEscapedCharWithCharInline();
-					if(SettingPath.Len() > 0)
-					{
-						Path = SettingPath;
-					}
-				}
-			}
-		}
-
-		if( !Path.Len() )
-		{
-			UE_LOG( LogDerivedDataCache, Log, TEXT("%s data cache path not found in *engine.ini, will not use an %s cache."), NodeName, NodeName );
-		}
-		else if( Path == TEXT("None") )
-		{
-			UE_LOG( LogDerivedDataCache, Log, TEXT("Disabling %s data cache - path set to 'None'."), NodeName );
-		}
-		else
-		{
-			// Try to set up the shared drive, allow user to correct any issues that may exist.
-			bool RetryOnFailure = false;
-			do
-			{
-				RetryOnFailure = false;
-
-				// Don't create the file system if shared data cache directory is not mounted
-				bool bShared = FCString::Stricmp(NodeName, TEXT("Shared")) == 0;
-				
-				// parameters we read here from the ini file
-				FString WriteAccessLog;
-				bool bPromptIfMissing = false;
-
-				FParse::Value(Entry, TEXT("WriteAccessLog="), WriteAccessLog);
-				FParse::Bool(Entry, TEXT("PromptIfMissing="), bPromptIfMissing);
-
-				ILegacyCacheStore* InnerFileSystem = nullptr;
-				ECacheStoreFlags Flags;
-				if (!bShared || IFileManager::Get().DirectoryExists(*Path))
-				{
-					InnerFileSystem = CreateFileSystemCacheStore(*Path, Entry, *WriteAccessLog, Flags);
-				}
-
-				if (InnerFileSystem)
-				{
-					bUsingSharedDDC |= bShared;
-					DataCache = MakeTuple(InnerFileSystem, Flags);
-					UE_LOG(LogDerivedDataCache, Log, TEXT("Using %s data cache path %s: %s"), NodeName, *Path,
-						EnumHasAnyFlags(Flags, ECacheStoreFlags::Store) ? TEXT("Writable") :
-						EnumHasAnyFlags(Flags, ECacheStoreFlags::Query) ? TEXT("ReadOnly") : TEXT("DeleteOnly"));
-					Directories.AddUnique(Path);
-				}
-				else
-				{
-					FString Message = FString::Printf(TEXT("%s data cache path (%s) is unavailable so cache will be disabled."), NodeName, *Path);
-					
-					UE_LOG(LogDerivedDataCache, Warning, TEXT("%s"), *Message);
-
-					// Give the user a chance to retry incase they need to connect a network drive or something.
-					if (bPromptIfMissing && !FApp::IsUnattended() && !IS_PROGRAM)
-					{
-						Message += FString::Printf(TEXT("\n\nRetry connection to %s?"), *Path);
-						EAppReturnType::Type MessageReturn = FPlatformMisc::MessageBoxExt(EAppMsgType::YesNo, *Message, TEXT("Could not access DDC"));
-						RetryOnFailure = MessageReturn == EAppReturnType::Yes;
-					}
-				}
-			} while (RetryOnFailure);
-		}
-
-		return DataCache;
+		return MakeTuple(nullptr, ECacheStoreFlags::None);
 	}
 
 	/**
diff --git a/Engine/Source/Developer/DerivedDataCache/Private/FileSystemCacheStore.cpp b/Engine/Source/Developer/DerivedDataCache/Private/FileSystemCacheStore.cpp
--- a/Engine/Source/Developer/DerivedDataCache/Private/FileSystemCacheStore.cpp
+++ b/Engine/Source/Developer/DerivedDataCache/Private/FileSystemCacheStore.cpp
@@ -8,6 +8,8 @@
 #include "Algo/Transform.h"
 #include "Async/Async.h"
 #include "Async/ManualResetEvent.h"
+#include "Async/Mutex.h"
+#include "Async/UniqueLock.h"
 #include "Containers/StaticBitArray.h"
 #include "DerivedDataBackendInterface.h"
 #include "DerivedDataCacheInterface.h"
@@ -24,6 +26,7 @@
 #include "HAL/Thread.h"
 #include "Hash/xxhash.h"
 #include "HashingArchiveProxy.h"
+#include "Misc/App.h"
 #include "Misc/CommandLine.h"
 #include "Misc/ConfigCacheIni.h"
 #include "Misc/CoreMisc.h"
@@ -767,13 +770,70 @@ void FAccessLogWriter::AppendPath(const FStringView Path)
 
 ///////////////////////////////////////////////////////////////////////////////////////////////////////////////////////
 
+struct FFileSystemCacheStoreParams
+{
+	/** Root path under which cache files are stored. */
+	FString CachePath;
+
+	/** Path to the optional access log that records every accessed file. */
+	FString AccessLogPath;
+
+	/** Maximum total size of compressed data stored within a record package with multiple attachments. */
+	uint64 MaxRecordSizeKB = 256;
+	/** Maximum total size of compressed data stored within a value package, or a record package with one attachment. */
+	uint64 MaxValueSizeKB = 1024;
+
+	/** Maintenance will delete files older than this. */
+	double MaxFileAgeInDays = 15.0;
+	/** Maintenance will wait this long after initialization before starting. */
+	double MaintenanceDelayInSeconds = 60.0;
+	/** Limits the number of paths scanned by maintenance in one second. */
+	uint32 MaxScanRate = MAX_uint32;
+	/** Limits the number of directories scanned by maintenance in each cache bucket or content root. */
+	uint32 MaxDirectoryScanCount = MAX_uint32;
+
+	/** Latency lower than this is considered local. */
+	float ConsiderFastAtMs = 10;
+	/** Latency lower than this is considered ok, and anything higher is considered slow. */
+	float ConsiderSlowAtMs = 50;
+	/** Latency higher than this will deactivate the cache store for performance reasons. */
+	float DeactivateAtMs = -1.0f;
+
+	/** If true, skip the speed test to measure latency on startup; */
+	bool bSkipSpeedTest = !WITH_EDITOR;
+	/** If true, display a retry prompt in attended sessions when the cache store is missing. */
+	bool bPromptIfMissing = false;
+	/** If true, files older than the max file age will be deleted during maintenance. */
+	bool bDeleteUnused = false;
+	/** If true, do not write or read any files in this cache store, only delete expired files. */
+	bool bDeleteOnly = false;
+	/** If true, do not write any files in this cache store. */
+	bool bReadOnly = false;
+	/** If true, block on maintenance of this cache store on startup. */
+	bool bClean = false;
+	/** If true, delete everything in this cache store on startup. */
+	bool bFlush = false;
+	/** If true, always update file timestamps on access, even when read-only. */
+	bool bTouch = false;
+
+	void Parse(const TCHAR* NodeName, const TCHAR* Config);
+};
+
+///////////////////////////////////////////////////////////////////////////////////////////////////////////////////////
+
 class FFileSystemCacheStore final : public ILegacyCacheStore
 {
 public:
+	static ILegacyCacheStore* TryCreate(
+		const TCHAR* NodeName,
+		const FFileSystemCacheStoreParams& Params,
+		ECacheStoreFlags& OutFlags,
+		FString& OutPath);
+
+private:
 	FFileSystemCacheStore(
-		const TCHAR* CachePath,
-		const TCHAR* Params,
-		const TCHAR* AccessLogPath,
+		const FFileSystemCacheStoreParams& Params,
+	const FDerivedDataCacheSpeedStats& SpeedStats,
 		ECacheStoreFlags& OutFlags);
 	~FFileSystemCacheStore();
 
@@ -865,20 +925,20 @@ private:
 	[[nodiscard]] bool FileExists(FStringBuilderBase& Path) const;
 	[[nodiscard]] bool IsDeactivatedForPerformance();
 
-	bool RunInitialSpeedTest();
+	static bool RunInitialSpeedTest(const FFileSystemCacheStoreParams& Params, FDerivedDataCacheSpeedStats& OutSpeedStats);
 
 private:
 	/** Base path to store the cache files in. */
 	FString CachePath;
 	/** Speed class of this cache. */
-	EBackendSpeedClass SpeedClass;
+	EBackendSpeedClass SpeedClass = EBackendSpeedClass::Unknown;
 	/** If true, do not attempt to write to this cache. */
 	bool		bReadOnly;
 	/** If true, always update file timestamps on access. */
 	bool		bTouch;
 
 	/** Age of file when it should be deleted from DDC cache. */
-	double		DaysToDeleteUnusedFiles;
+	double		MaxFileAgeInDays = 15.0;
 
 	/** Maximum total size of compressed data stored within a record package with multiple attachments. */
 	uint64		MaxRecordSizeKB = 256;
@@ -900,8 +960,8 @@ private:
 
 	TUniquePtr<FFileSystemCacheStoreMaintainerParams> DeactivationDeferredMaintainerParams;
 
-	static inline FRWLock ActiveFileStoresLock;
-	static inline TArray<FFileSystemCacheStore*> ActiveFileStores;
+	inline static FMutex ActiveStoresMutex;
+	inline static TArray<FFileSystemCacheStore*> ActiveStores;
 
 	enum class EPerformanceReEvaluationResult
 	{
@@ -912,277 +972,256 @@ private:
 	FRWLock PerformanceReEvaluationTaskLock;
 	Tasks::TTask<std::atomic<EPerformanceReEvaluationResult>> PerformanceReEvaluationTask;
 	std::atomic<int64> LastPerformanceEvaluationTicks;
-	std::atomic<bool> bDeactivedForPerformance;
-	bool bDeactivationDeferredClean;
-	float DeactivateAtMS;
+	std::atomic<bool> bDeactivedForPerformance = false;
+	bool bDeactivationDeferredClean = false;
+	float DeactivateAtMs;
 };
 
-FFileSystemCacheStore::FFileSystemCacheStore(
-	const TCHAR* const InCachePath,
-	const TCHAR* const InParams,
-	const TCHAR* const InAccessLogPath,
-	ECacheStoreFlags& OutFlags)
-	: CachePath(InCachePath)
-	, SpeedClass(EBackendSpeedClass::Unknown)
-	, bReadOnly(false)
-	, bTouch(false)
-	, DaysToDeleteUnusedFiles(15.0)
-	, bDeactivedForPerformance(false)
-	, bDeactivationDeferredClean(false)
-	, DeactivateAtMS(-1.f)
+ILegacyCacheStore* FFileSystemCacheStore::TryCreate(
+	const TCHAR* const NodeName,
+	const FFileSystemCacheStoreParams& Params,
+	ECacheStoreFlags& OutFlags,
+	FString& OutPath)
 {
-	bool bDeleteOnly = false;
+	OutFlags = ECacheStoreFlags::None;
 
 	// If we find a platform that has more stringent limits, this needs to be rethought.
 	checkf(GMaxCacheRootLen + GMaxCacheKeyLen <= FPlatformMisc::GetMaxPathLength(),
 		TEXT("Not enough room left for cache keys in max path."));
 
-	check(CachePath.Len());
-	FPaths::NormalizeFilename(CachePath);
+	check(!Params.CachePath.IsEmpty());
+	if (Params.CachePath.IsEmpty())
+	{
+		UE_LOG(LogDerivedDataCache, Log, TEXT("%s: Disabled because no path is configured."), NodeName);
+		return nullptr;
+	}
+	else if (Params.CachePath == TEXTVIEW("None"))
+	{
+		UE_LOG(LogDerivedDataCache, Log, TEXT("%s: Disabled because the path is configured to 'None'"), NodeName);
+		return nullptr;
+	}
 
-	// Params that override our instance defaults
-	FParse::Bool(InParams, TEXT("DeleteOnly="), bDeleteOnly);
-	FParse::Bool(InParams, TEXT("ReadOnly="), bReadOnly);
-	FParse::Bool(InParams, TEXT("Touch="), bTouch);
-	FParse::Value(InParams, TEXT("UnusedFileAge="), DaysToDeleteUnusedFiles);
-	FParse::Value(InParams, TEXT("MaxRecordSizeKB="), MaxRecordSizeKB);
-	FParse::Value(InParams, TEXT("MaxValueSizeKB="), MaxValueSizeKB);
+	const auto HasSameCachePath = [&Params](const FFileSystemCacheStore* ActiveStore)
+	{
+		return FPaths::IsSamePath(Params.CachePath, ActiveStore->CachePath);
+	};
 
+	bool bRetryOnFailure;
+	do
 	{
-		FReadScopeLock ReadLock(ActiveFileStoresLock);
-		FFileSystemCacheStore** FoundExitingActiveStore = Algo::FindByPredicate(ActiveFileStores, [this](const FFileSystemCacheStore* ActiveStore)
+		bRetryOnFailure = false;
+
+		if (TUniqueLock Lock(ActiveStoresMutex); Algo::FindByPredicate(ActiveStores, HasSameCachePath))
 		{
-			return ActiveStore && FPaths::IsSamePath(ActiveStore->CachePath, CachePath);
-		});
-		if (FoundExitingActiveStore)
+			UE_LOG(LogDerivedDataCache, Warning, TEXT("%s: Process has an existing cache store at path %s, and the duplicate is being ignored."), NodeName, *Params.CachePath);
+			return nullptr;
+		}
+
+		// Skip creation of the cache store if the shared cache directory does not exist.
+		bool bShared = NodeName == TEXTVIEW("Shared");
+		if (!bShared || IFileManager::Get().DirectoryExists(*Params.CachePath))
 		{
-			UE_LOG(LogDerivedDataCache, Warning, TEXT("Attempted to create multiple cache stores to path %s, ignoring duplicate entry"), *CachePath);
-			OutFlags = ECacheStoreFlags::None;
-			return;
+			FDerivedDataCacheSpeedStats LocalSpeedStats;
+			LocalSpeedStats.ReadSpeedMBs = 999;
+			LocalSpeedStats.WriteSpeedMBs = 999;
+			LocalSpeedStats.LatencyMS = 0;
+			if (Params.bSkipSpeedTest || RunInitialSpeedTest(Params, LocalSpeedStats))
+			{
+				UE_CLOG(Params.bSkipSpeedTest, LogDerivedDataCache, Log, TEXT("%s: Skipping speed test at path %s and assuming local performance."), NodeName, *Params.CachePath);
+
+				FFileSystemCacheStore* CacheStore = new FFileSystemCacheStore(Params, LocalSpeedStats, OutFlags);
+				{
+					TUniqueLock Lock(ActiveStoresMutex);
+					ActiveStores.Add(CacheStore);
+				}
+				UE_LOG(LogDerivedDataCache, Log, TEXT("%s: Using data cache path %s: %s"), NodeName, *Params.CachePath,
+					EnumHasAnyFlags(OutFlags, ECacheStoreFlags::Store) ? TEXT("Writable") :
+					EnumHasAnyFlags(OutFlags, ECacheStoreFlags::Query) ? TEXT("ReadOnly") : TEXT("DeleteOnly"));
+				OutPath = CacheStore->CachePath;
+				return CacheStore;
+			}
+
+			UE_LOG(LogDerivedDataCache, Warning, TEXT("%s: No read or write access to %s, or the speed test was abandoned due to slow progress."), NodeName, *Params.CachePath);
 		}
-	}
 
-	// Flush the cache if requested.
-	bool bFlush = false;
-	if (!bReadOnly && FParse::Bool(InParams, TEXT("Flush="), bFlush) && bFlush)
-	{
-		TRACE_CPUPROFILER_EVENT_SCOPE(FileSystemDDC_Flush);
-		IFileManager::Get().DeleteDirectory(*(CachePath / TEXT("")), /*bRequireExists*/ false, /*bTree*/ true);
+		// Give the user a chance to retry in case they need to take steps like connecting a network drive.
+		if (Params.bPromptIfMissing && !FApp::IsUnattended() && !IS_PROGRAM)
+		{
+			TStringBuilder<512> Message(InPlace, NodeName, TEXTVIEW(" cache store path "), Params.CachePath,
+				TEXTVIEW(" is not available and this cache store will be disabled.\n\nRetry connection to "), Params.CachePath, TEXTVIEW("?"));
+			bRetryOnFailure = FPlatformMisc::MessageBoxExt(EAppMsgType::YesNo, *Message, TEXT("Failed to access Derived Data Cache")) == EAppReturnType::Yes;
+		}
 	}
+	while (bRetryOnFailure);
 
-	// check latency and speed. Read values should always be valid
-	
-	/* Speeds faster than this are considered local*/
-	const float ConsiderFastAtMS = 10;
-	/* Speeds faster than this are ok. Everything else is slow. This value can be overridden in the ini file */
-	float ConsiderSlowAtMS = 50;
-	FParse::Value(InParams, TEXT("ConsiderSlowAt="), ConsiderSlowAtMS);
-	FParse::Value(InParams, TEXT("DeactivateAt="), DeactivateAtMS);
-
-	const FString AbsoluteCachePath = IFileManager::Get().ConvertToAbsolutePathForExternalAppForRead(*CachePath);
-	if (AbsoluteCachePath.Len() >= GMaxCacheRootLen)
-	{
-		const FText ErrorMessage = FText::Format(NSLOCTEXT("DerivedDataCache", "PathTooLong", "Cache path {0} is longer than {1} characters... please adjust [DerivedDataBackendGraph] paths to be shorter (this leaves more room for cache keys)."),
-			FText::FromString(AbsoluteCachePath), FText::AsNumber(GMaxCacheRootLen));
-		FMessageDialog::Open(EAppMsgType::Ok, ErrorMessage);
-		UE_LOG(LogDerivedDataCache, Fatal, TEXT("%s"), *ErrorMessage.ToString());
-	}
+	UE_LOG(LogDerivedDataCache, Warning, TEXT("%s: Path %s is not available and this cache store will be disabled."), NodeName, *Params.CachePath);
+	return nullptr;
+}
+
+FFileSystemCacheStore::FFileSystemCacheStore(
+	const FFileSystemCacheStoreParams& Params,
+	const FDerivedDataCacheSpeedStats& InSpeedStats,
+	ECacheStoreFlags& OutFlags)
+	: CachePath(Params.CachePath)
+	, bReadOnly(Params.bReadOnly)
+	, bTouch(Params.bTouch)
+	, MaxFileAgeInDays(Params.MaxFileAgeInDays)
+	, SpeedStats(InSpeedStats)
+	, DeactivateAtMs(Params.DeactivateAtMs)
+{
+	LastPerformanceEvaluationTicks.store(FDateTime::UtcNow().GetTicks(), std::memory_order_relaxed);
+
+	bool bReadTestPassed = SpeedStats.ReadSpeedMBs > 0.0;
+	bool bWriteTestPassed = SpeedStats.WriteSpeedMBs > 0.0;
+
+	// if we failed writes mark this as read only
+	bReadOnly = bReadOnly || !bWriteTestPassed;
 
-	// can skip the speed test so everything acts as local (e.g. 4.25 and earlier behavior).
-	bool SkipSpeedTest = !WITH_EDITOR || FParse::Param(FCommandLine::Get(), TEXT("DDCSkipSpeedTest"));
-	if (SkipSpeedTest)
+	const bool bLocalDeactivatedForPerformance = (Params.DeactivateAtMs > 0.f) && (SpeedStats.LatencyMS >= Params.DeactivateAtMs);
+	bDeactivedForPerformance.store(bLocalDeactivatedForPerformance, std::memory_order_relaxed);
+
+	// classify and report on these times
+	if (SpeedStats.LatencyMS < 1)
 	{
-		SpeedStats.ReadSpeedMBs = 999;
-		SpeedStats.WriteSpeedMBs = 999;
-		SpeedStats.LatencyMS = 0;
-		UE_LOG(LogDerivedDataCache, Log, TEXT("Skipping speed test to %s. Assuming local performance"), *CachePath);
+		SpeedClass = EBackendSpeedClass::Local;
 	}
-
-	if (!SkipSpeedTest &&
-		!RunInitialSpeedTest())
+	else if (SpeedStats.LatencyMS <= Params.ConsiderFastAtMs)
 	{
-		LastPerformanceEvaluationTicks.store(FDateTime::UtcNow().GetTicks(), std::memory_order_relaxed);
-
-		OutFlags = ECacheStoreFlags::None;
-		UE_LOG(LogDerivedDataCache, Warning, TEXT("No read or write access to %s, or abandoned due to slow progress"), *CachePath);
+		SpeedClass = EBackendSpeedClass::Fast;
+	}
+	else if (SpeedStats.LatencyMS >= Params.ConsiderSlowAtMs)
+	{
+		SpeedClass = EBackendSpeedClass::Slow;
 	}
 	else
 	{
-		LastPerformanceEvaluationTicks.store(FDateTime::UtcNow().GetTicks(), std::memory_order_relaxed);
-
-		bool bReadTestPassed = SpeedStats.ReadSpeedMBs > 0.0;
-		bool bWriteTestPassed = SpeedStats.WriteSpeedMBs > 0.0;
-
-		// if we failed writes mark this as read only
-		bReadOnly = bReadOnly || !bWriteTestPassed;
+		SpeedClass = EBackendSpeedClass::Ok;
+	}
 
-		const bool bLocalDeactivatedForPerformance = (DeactivateAtMS > 0.f) && (SpeedStats.LatencyMS >= DeactivateAtMS);
-		bDeactivedForPerformance.store(bLocalDeactivatedForPerformance, std::memory_order_relaxed);
+	UE_LOG(LogDerivedDataCache, Display,
+		TEXT("%s: Performance: Latency=%.02fms. RandomReadSpeed=%.02fMBs, RandomWriteSpeed=%.02fMBs. Assigned SpeedClass '%s'"),
+		*CachePath, SpeedStats.LatencyMS, SpeedStats.ReadSpeedMBs, SpeedStats.WriteSpeedMBs, LexToString(SpeedClass));
 
-		// classify and report on these times
-		if (SpeedStats.LatencyMS < 1)
+	if (bLocalDeactivatedForPerformance)
+	{
+		if (GIsBuildMachine)
 		{
-			SpeedClass = EBackendSpeedClass::Local;
+			UE_LOG(LogDerivedDataCache, Display,
+				TEXT("%s: Performance does not meet minimum criteria. "
+					"It will be deactivated until performance measurements improve. "
+					"If this is consistent, consider disabling this cache store through "
+					"environment variables or other configuration."),
+				*CachePath);
 		}
-		else if (SpeedStats.LatencyMS <= ConsiderFastAtMS)
+		else
 		{
-			SpeedClass = EBackendSpeedClass::Fast;
+			UE_LOG(LogDerivedDataCache, Warning,
+				TEXT("%s: Performance does not meet minimum criteria. "
+					"It will be deactivated until performance measurements improve. "
+					"If this is consistent, consider disabling this cache store through "
+					"environment variables or other configuration."),
+				*CachePath);
 		}
-		else if (SpeedStats.LatencyMS >= ConsiderSlowAtMS)
+	}
+
+	if (SpeedClass <= EBackendSpeedClass::Slow && !bReadOnly)
+	{
+		if (GIsBuildMachine)
 		{
-			SpeedClass = EBackendSpeedClass::Slow;
+			UE_LOG(LogDerivedDataCache, Display, TEXT("%s: Access is slow. "
+				"'Touch' will be disabled and queries/writes will be limited."), *CachePath);
 		}
 		else
 		{
-			SpeedClass = EBackendSpeedClass::Ok;
+			UE_LOG(LogDerivedDataCache, Warning, TEXT("%s: Access is slow. "
+				"'Touch' will be disabled and queries/writes will be limited."), *CachePath);
 		}
+		bTouch = false;
+		//bReadOnly = true;
+	}
 
-		UE_LOG(LogDerivedDataCache, Display,
-			TEXT("Performance to %s: Latency=%.02fms. RandomReadSpeed=%.02fMBs, RandomWriteSpeed=%.02fMBs. "
-				 "Assigned SpeedClass '%s'"),
-			*CachePath,
-			SpeedStats.LatencyMS,
-			SpeedStats.ReadSpeedMBs,
-			SpeedStats.WriteSpeedMBs,
-			LexToString(SpeedClass));
+	if (FString(FCommandLine::Get()).Contains(TEXT("Run=DerivedDataCache")))
+	{
+		bTouch = true; // we always touch files when running the DDC commandlet
+	}
+
+	if (bTouch)
+	{
+		UE_LOG(LogDerivedDataCache, Display, TEXT("%s: Files will be touched when accessed."), *CachePath);
+	}
+
+	// Flush the cache if requested.
+	if (!bReadOnly && Params.bFlush)
+	{
+		TRACE_CPUPROFILER_EVENT_SCOPE(FileSystemDDC_Flush);
+		IFileManager::Get().DeleteDirectory(*(CachePath / TEXT("")), /*bRequireExists*/ false, /*bTree*/ true);
+	}
+
+	if (Params.bClean && bLocalDeactivatedForPerformance)
+	{
+		bDeactivationDeferredClean = true;
+	}
 
+	if (Params.bClean || Params.bDeleteUnused)
+	{
+		FFileSystemCacheStoreMaintainerParams* MaintainerParams;
+		FFileSystemCacheStoreMaintainerParams LocalMaintainerParams;
 		if (bLocalDeactivatedForPerformance)
 		{
-			if (GIsBuildMachine)
-			{
-				UE_LOG(LogDerivedDataCache, Display,
-					TEXT("%s: Performance does not meet minimum criteria. "
-						 "It will be deactivated until performance measurements improve. "
-						 "If this is consistent, consider disabling this cache store through "
-						 "environment variables or other configuration."),
-					*CachePath);
-			}
-			else
-			{
-				UE_LOG(LogDerivedDataCache, Warning,
-					TEXT("%s: Performance does not meet minimum criteria. "
-						 "It will be deactivated until performance measurements improve. "
-						 "If this is consistent, consider disabling this cache store through "
-						 "environment variables or other configuration."),
-					*CachePath);
-			}
+			DeactivationDeferredMaintainerParams = MakeUnique<FFileSystemCacheStoreMaintainerParams>();
+			MaintainerParams = DeactivationDeferredMaintainerParams.Get();
 		}
-
-		if (SpeedClass <= EBackendSpeedClass::Slow && !bReadOnly)
+		else
 		{
-			if (GIsBuildMachine)
-			{
-				UE_LOG(LogDerivedDataCache, Display, TEXT("Access to %s appears to be slow. "
-					"'Touch' will be disabled and queries/writes will be limited."), *CachePath);
-			}
-			else
-			{
-				UE_LOG(LogDerivedDataCache, Warning, TEXT("Access to %s appears to be slow. "
-					"'Touch' will be disabled and queries/writes will be limited."), *CachePath);
-			}
-			bTouch = false;
-			//bReadOnly = true;
+			MaintainerParams = &LocalMaintainerParams;
 		}
-
-
-		if (FString(FCommandLine::Get()).Contains(TEXT("Run=DerivedDataCache")))
+		MaintainerParams->MaxFileAge = FTimespan::FromDays(MaxFileAgeInDays);
+		if (Params.bDeleteUnused)
 		{
-			bTouch = true; // we always touch files when running the DDC commandlet
+			MaintainerParams->MaxScanRate = Params.MaxScanRate;
+			MaintainerParams->MaxDirectoryScanCount = Params.MaxDirectoryScanCount;
 		}
-
-		// The command line (-ddctouch) enables touch on all filesystem backends if specified.
-		bTouch = bTouch || FParse::Param(FCommandLine::Get(), TEXT("DDCTOUCH"));
-
-		if (bTouch)
+		else
 		{
-			UE_LOG(LogDerivedDataCache, Display, TEXT("Files in %s will be touched."), *CachePath);
+			MaintainerParams->ScanFrequency = FTimespan::MaxValue();
 		}
-		
-		bool bClean = false;
-		bool bDeleteUnused = !bReadOnly;
-		FParse::Bool(InParams, TEXT("Clean="), bClean);
-		FParse::Bool(InParams, TEXT("DeleteUnused="), bDeleteUnused);
-		bDeleteUnused = bDeleteUnused && !FParse::Param(FCommandLine::Get(), TEXT("NODDCCLEANUP"));
-
-		if (bClean && bLocalDeactivatedForPerformance)
+		if (Params.bClean)
+		{
+			MaintainerParams->TimeToWaitAfterInit = FTimespan::Zero();
+		}
+		else
 		{
-			bDeactivationDeferredClean = true;
+			MaintainerParams->TimeToWaitAfterInit = FTimespan::FromSeconds(Params.MaintenanceDelayInSeconds);
 		}
 
-		if (bClean || bDeleteUnused)
+		if (!bLocalDeactivatedForPerformance)
 		{
-			FFileSystemCacheStoreMaintainerParams* MaintainerParams;
-			FFileSystemCacheStoreMaintainerParams LocalMaintainerParams;
-			if (bLocalDeactivatedForPerformance)
-			{
-				DeactivationDeferredMaintainerParams = MakeUnique<FFileSystemCacheStoreMaintainerParams>();
-				MaintainerParams = DeactivationDeferredMaintainerParams.Get();
-			}
-			else
-			{
-				MaintainerParams = &LocalMaintainerParams;
-			}
-			MaintainerParams->MaxFileAge = FTimespan::FromDays(DaysToDeleteUnusedFiles);
-			if (bDeleteUnused)
-			{
-				if (!FParse::Value(InParams, TEXT("MaxFileChecksPerSec="), MaintainerParams->MaxScanRate))
-				{
-					int32 MaxFileScanRate;
-					if (GConfig->GetInt(TEXT("DDCCleanup"), TEXT("MaxFileChecksPerSec"), MaxFileScanRate, GEngineIni))
-					{
-						MaintainerParams->MaxScanRate = uint32(MaxFileScanRate);
-					}
-				}
-				FParse::Value(InParams, TEXT("FoldersToClean="), MaintainerParams->MaxDirectoryScanCount);
-			}
-			else
-			{
-				MaintainerParams->ScanFrequency = FTimespan::MaxValue();
-			}
-			double TimeToWaitAfterInit;
-			if (bClean)
-			{
-				MaintainerParams->TimeToWaitAfterInit = FTimespan::Zero();
-			}
-			else if (GConfig->GetDouble(TEXT("DDCCleanup"), TEXT("TimeToWaitAfterInit"), TimeToWaitAfterInit, GEngineIni))
-			{
-				MaintainerParams->TimeToWaitAfterInit = FTimespan::FromSeconds(TimeToWaitAfterInit);
-			}
+			Maintainer = MakeUnique<FFileSystemCacheStoreMaintainer>(*MaintainerParams, CachePath);
 
-			if (!bLocalDeactivatedForPerformance)
+			if (Params.bClean)
 			{
-				Maintainer = MakeUnique<FFileSystemCacheStoreMaintainer>(*MaintainerParams, CachePath);
-
-				if (bClean)
-				{
-					Maintainer->BoostPriority();
-					Maintainer->WaitForIdle();
-				}
+				Maintainer->BoostPriority();
+				Maintainer->WaitForIdle();
 			}
 		}
+	}
 
-		if (InAccessLogPath && *InAccessLogPath)
-		{
-			AccessLogWriter.Reset(new FAccessLogWriter(InAccessLogPath, CachePath));
-		}
-
-		ECacheStoreFlags Flags = ECacheStoreFlags::None;
-		Flags |= bDeleteOnly ? ECacheStoreFlags::None : ECacheStoreFlags::Query;
-		Flags |= bDeleteOnly || bReadOnly ? ECacheStoreFlags::None : ECacheStoreFlags::Store;
-		Flags |= SpeedClass == EBackendSpeedClass::Local ? ECacheStoreFlags::Local : ECacheStoreFlags::Remote;
-		OutFlags = Flags;
-
-		{
-			FWriteScopeLock WriteLock(ActiveFileStoresLock);
-			ActiveFileStores.Add(this);
-		}
+	if (!Params.AccessLogPath.IsEmpty())
+	{
+		AccessLogWriter.Reset(new FAccessLogWriter(*Params.AccessLogPath, CachePath));
 	}
+
+	ECacheStoreFlags Flags = ECacheStoreFlags::None;
+	Flags |= Params.bDeleteOnly ? ECacheStoreFlags::None : ECacheStoreFlags::Query;
+	Flags |= Params.bDeleteOnly || bReadOnly ? ECacheStoreFlags::None : ECacheStoreFlags::Store;
+	Flags |= SpeedClass == EBackendSpeedClass::Local ? ECacheStoreFlags::Local : ECacheStoreFlags::Remote;
+	OutFlags = Flags;
 }
 
 FFileSystemCacheStore::~FFileSystemCacheStore()
 {
-	FWriteScopeLock WriteLock(ActiveFileStoresLock);
-	ActiveFileStores.Remove(this);
+	TUniqueLock Lock(ActiveStoresMutex);
+	ActiveStores.Remove(this);
 }
 
 bool FFileSystemCacheStore::RunSpeedTest(
@@ -1251,7 +1290,7 @@ bool FFileSystemCacheStore::RunSpeedTest(
 		if (AbandonRequest && AbandonRequest->load(std::memory_order_relaxed))
 		{
 			const double TotalTestTime = FPlatformTime::Seconds() - StatStartTime;
-			UE_LOG(LogDerivedDataCache, Log, TEXT("Speed tests for %s abandoned on request after %.02f seconds"), *CachePath, TotalTestTime);
+			UE_LOG(LogDerivedDataCache, Log, TEXT("%s: Speed tests abandoned on request after %.02f seconds."), *CachePath, TotalTestTime);
 			return false;
 		}
 
@@ -1267,7 +1306,7 @@ bool FFileSystemCacheStore::RunSpeedTest(
 	// calculate seek time here
 	OutSeekTimeMS = (TotalSeekTime / TestFileEntries.Num()) * 1000;
 
-	UE_LOG(LogDerivedDataCache, Verbose, TEXT("Stat tests to %s took %.02f seconds"), *CachePath, TotalSeekTime);
+	UE_LOG(LogDerivedDataCache, Verbose, TEXT("%s: Stat tests took %.02f seconds."), *CachePath, TotalSeekTime);
 
 	if (bSeekTimeOnly)
 	{
@@ -1294,8 +1333,8 @@ bool FFileSystemCacheStore::RunSpeedTest(
 					TCHAR ErrorBuffer[1024];
 					FPlatformMisc::GetSystemErrorMessage(ErrorBuffer, 1024, ErrorCode);
 					UE_LOG(LogDerivedDataCache, Warning,
-						TEXT("Fail to create %s, derived data cache to this directory will be read only. "
-						     "WriteError: %u (%s)"), *File, ErrorCode, ErrorBuffer);
+						TEXT("%s: Failed to create %s, this cache store will be read-only. "
+							"WriteError: %u (%s)"), *CachePath, *File, ErrorCode, ErrorBuffer);
 					bTestDataExists = false;
 					bWriteTestPassed = false;
 					break;
@@ -1305,7 +1344,7 @@ bool FFileSystemCacheStore::RunSpeedTest(
 			if (AbandonRequest && AbandonRequest->load(std::memory_order_relaxed))
 			{
 				const double TotalTestTime = FPlatformTime::Seconds() - StatStartTime;
-				UE_LOG(LogDerivedDataCache, Log, TEXT("Speed tests for %s abandoned on request after %.02f seconds"), *CachePath, TotalTestTime);
+				UE_LOG(LogDerivedDataCache, Log, TEXT("%s: Speed tests abandoned on request after %.02f seconds."), *CachePath, TotalTestTime);
 				return false;
 			}
 		}
@@ -1330,8 +1369,8 @@ bool FFileSystemCacheStore::RunSpeedTest(
 				TCHAR ErrorBuffer[1024];
 				FPlatformMisc::GetSystemErrorMessage(ErrorBuffer, 1024, ErrorCode);
 				UE_LOG(LogDerivedDataCache, Warning,
-					TEXT("Fail to read from %s, derived data cache will be disabled. ReadError: %u (%s)"),
-					*FilePath, ErrorCode, ErrorBuffer);
+					TEXT("%s: Failed to read from %s, this cache store will be disabled. ReadError: %u (%s)"),
+					*CachePath, *FilePath, ErrorCode, ErrorBuffer);
 				bReadTestPassed = false;
 				break;
 			}
@@ -1341,15 +1380,15 @@ bool FFileSystemCacheStore::RunSpeedTest(
 			if (AbandonRequest && AbandonRequest->load(std::memory_order_relaxed))
 			{
 				const double TotalTestTime = FPlatformTime::Seconds() - StatStartTime;
-				UE_LOG(LogDerivedDataCache, Log, TEXT("Speed tests for %s abandoned on request after %.02f seconds"), *CachePath, TotalTestTime);
+				UE_LOG(LogDerivedDataCache, Log, TEXT("%s: Speed tests abandoned on request after %.02f seconds."), *CachePath, TotalTestTime);
 				return false;
 			}
 		}
 
 		TotalReadTime = FPlatformTime::Seconds() - ReadStartTime;
 
-		UE_LOG(LogDerivedDataCache, Verbose, TEXT("Read tests %s on %s and took %.02f seconds"),
-			bReadTestPassed ? TEXT("passed") : TEXT("failed"), *CachePath, TotalReadTime);
+		UE_LOG(LogDerivedDataCache, Verbose, TEXT("%s: Read tests %s and took %.02f seconds."),
+			*CachePath, bReadTestPassed ? TEXT("passed") : TEXT("failed"), TotalReadTime);
 	}
 
 	// do write tests if or read tests passed and our seeks were below the cut-off
@@ -1379,8 +1418,8 @@ bool FFileSystemCacheStore::RunSpeedTest(
 				TCHAR ErrorBuffer[1024];
 				FPlatformMisc::GetSystemErrorMessage(ErrorBuffer, 1024, ErrorCode);
 				UE_LOG(LogDerivedDataCache, Warning,
-					TEXT("Fail to write to %s, derived data cache will be disabled. ReadError: %u (%s)"),
-					*FilePath, ErrorCode, ErrorBuffer);
+					TEXT("%s: Failed to write to %s, this cache store will be disabled. ReadError: %u (%s)"),
+					*CachePath, *FilePath, ErrorCode, ErrorBuffer);
 				bWriteTestPassed = false;
 				break;
 			}
@@ -1390,33 +1429,34 @@ bool FFileSystemCacheStore::RunSpeedTest(
 			if (AbandonRequest && AbandonRequest->load(std::memory_order_relaxed))
 			{
 				const double TotalTestTime = FPlatformTime::Seconds() - StatStartTime;
-				UE_LOG(LogDerivedDataCache, Log, TEXT("Speed tests for %s abandoned on request after %.02f seconds"), *CachePath, TotalTestTime);
+				UE_LOG(LogDerivedDataCache, Log, TEXT("%s: Speed tests abandoned on request after %.02f seconds."), *CachePath, TotalTestTime);
 				return false;
 			}
 		}
 
 		TotalWriteTime = FPlatformTime::Seconds() - WriteStartTime;
 
-		UE_LOG(LogDerivedDataCache, Verbose, TEXT("write tests %s on %s and took %.02f seconds"),
-			bWriteTestPassed ? TEXT("passed") : TEXT("failed"), *CachePath, TotalReadTime)
+		UE_LOG(LogDerivedDataCache, Verbose, TEXT("%s: Write tests %s and took %.02f seconds."),
+			*CachePath, bWriteTestPassed ? TEXT("passed") : TEXT("failed"), TotalReadTime)
 
 			
 		if (AbandonRequest && AbandonRequest->load(std::memory_order_relaxed))
 		{
 			const double TotalTestTime = FPlatformTime::Seconds() - StatStartTime;
-			UE_LOG(LogDerivedDataCache, Log, TEXT("Speed tests for %s abandoned on request after %.02f seconds"), *CachePath, TotalTestTime);
+			UE_LOG(LogDerivedDataCache, Log, TEXT("%s: Speed tests abandoned on request after %.02f seconds."), *CachePath, TotalTestTime);
 			return false;
 		}
 		
 		// remove the custom path but do it async as this can be slow on remote drives
-		AsyncTask(ENamedThreads::AnyThread, [CustomPath]() {
+		AsyncTask(ENamedThreads::AnyThread, [CustomPath]
+		{
 			IFileManager::Get().DeleteDirectory(*CustomPath, false, true);
 		});
 	}
 
 	const double TotalTestTime = FPlatformTime::Seconds() - StatStartTime;
 
-	UE_LOG(LogDerivedDataCache, Log, TEXT("Speed tests for %s took %.02f seconds"), *CachePath, TotalTestTime);
+	UE_LOG(LogDerivedDataCache, Log, TEXT("%s: Speed tests took %.02f seconds."), *CachePath, TotalTestTime);
 
 	// check latency and speed. Read values should always be valid
 	OutReadSpeedMBs = (bReadTestPassed ? (TotalDataRead / TotalReadTime) : 0) / (1024 * 1024);
@@ -2557,7 +2597,7 @@ bool FFileSystemCacheStore::FileExists(FStringBuilderBase& Path) const
 	{
 		return false;
 	}
-	if (bTouch || (!bReadOnly && (FDateTime::UtcNow() - TimeStamp).GetTotalDays() > (DaysToDeleteUnusedFiles / 4)))
+	if (bTouch || (!bReadOnly && (FDateTime::UtcNow() - TimeStamp).GetTotalDays() > (MaxFileAgeInDays / 4)))
 	{
 		IFileManager::Get().SetTimeStamp(*Path, FDateTime::UtcNow());
 	}
@@ -2566,7 +2606,7 @@ bool FFileSystemCacheStore::FileExists(FStringBuilderBase& Path) const
 
 bool FFileSystemCacheStore::IsDeactivatedForPerformance()
 {
-	if ((DeactivateAtMS <= 0.f) || !bDeactivedForPerformance.load(std::memory_order_relaxed))
+	if ((DeactivateAtMs <= 0.f) || !bDeactivedForPerformance.load(std::memory_order_relaxed))
 	{
 		return false;
 	}
@@ -2647,10 +2687,10 @@ bool FFileSystemCacheStore::IsDeactivatedForPerformance()
 			if (TimespanSinceLastPerfEval > FTimespan::FromMinutes(1))
 			{
 				PerformanceReEvaluationTask = Tasks::Launch(TEXT("FFileSystemCacheStore::ReEvaluatePerformance"),
-					[CachePath = this->CachePath, DeactivateAtMS = this->DeactivateAtMS]() ->
+					[CachePath = this->CachePath, DeactivateAtMs = this->DeactivateAtMs]() ->
 						std::atomic<EPerformanceReEvaluationResult>
 					{
-						check(DeactivateAtMS > 0.f);
+						check(DeactivateAtMs > 0.f);
 
 						FDerivedDataCacheSpeedStats LocalSpeedStats;
 						LocalSpeedStats.ReadSpeedMBs = 999;
@@ -2666,7 +2706,7 @@ bool FFileSystemCacheStore::IsDeactivatedForPerformance()
 							nullptr,
 							nullptr);
 
-						if (LocalSpeedStats.LatencyMS >= DeactivateAtMS)
+						if (LocalSpeedStats.LatencyMS >= DeactivateAtMs)
 						{
 							return EPerformanceReEvaluationResult::PerformanceDeactivate;
 						}
@@ -2679,7 +2719,7 @@ bool FFileSystemCacheStore::IsDeactivatedForPerformance()
 	return true;
 }
 
-bool FFileSystemCacheStore::RunInitialSpeedTest()
+bool FFileSystemCacheStore::RunInitialSpeedTest(const FFileSystemCacheStoreParams& Params, FDerivedDataCacheSpeedStats& OutSpeedStats)
 {
 	struct FSpeedTestState : public FThreadSafeRefCountedObject
 	{
@@ -2692,7 +2732,7 @@ bool FFileSystemCacheStore::RunInitialSpeedTest()
 
 	TRefCountPtr<FSpeedTestState> SpeedTestState = new FSpeedTestState();
 	Tasks::Launch(TEXT("FFileSystemCacheStore::InitialEvaluation"),
-		[CachePath = this->CachePath, bReadOnly = this->bReadOnly, SpeedTestState]()
+		[CachePath = Params.CachePath, bReadOnly = Params.bReadOnly, SpeedTestState]()
 		{
 			SpeedTestState->bResult = RunSpeedTest(CachePath,
 				bReadOnly,
@@ -2705,12 +2745,12 @@ bool FFileSystemCacheStore::RunInitialSpeedTest()
 			SpeedTestState->CompletionEvent.Notify();
 		});
 
-	if (!GIsBuildMachine && FPlatformProcess::SupportsMultithreading() && (DeactivateAtMS > 0.f))
+	if (!GIsBuildMachine && FPlatformProcess::SupportsMultithreading() && (Params.DeactivateAtMs > 0.f))
 	{
-		if (SpeedTestState->CompletionEvent.WaitFor(FMonotonicTimeSpan::FromMilliseconds(DeactivateAtMS * 2.f)))
+		if (SpeedTestState->CompletionEvent.WaitFor(FMonotonicTimeSpan::FromMilliseconds(Params.DeactivateAtMs * 2.f)))
 		{
 			// If the task completed in the initial wait period, return the result
-			SpeedStats = SpeedTestState->SpeedStats;
+			OutSpeedStats = SpeedTestState->SpeedStats;
 			return SpeedTestState->bResult;
 		}
 		else
@@ -2720,10 +2760,10 @@ bool FFileSystemCacheStore::RunInitialSpeedTest()
 			if (SpeedTestState->NumLatencyTestsCompleted.load(std::memory_order_acquire) < 2)
 			{
 				SpeedTestState->AbandonRequest.store(true, std::memory_order_relaxed);
-				SpeedStats.ReadSpeedMBs = 0.0;
-				SpeedStats.WriteSpeedMBs = 0.0;
-				SpeedStats.LatencyMS = 999.0;
-				UE_LOG(LogDerivedDataCache, Log, TEXT("Skipping speed test to %s due to slow test progress. Assuming poor performance"), *CachePath);
+				OutSpeedStats.ReadSpeedMBs = 0.0;
+				OutSpeedStats.WriteSpeedMBs = 0.0;
+				OutSpeedStats.LatencyMS = 999.0;
+				UE_LOG(LogDerivedDataCache, Log, TEXT("%s: Skipping speed test due to slow test progress. Assuming poor performance."), *Params.CachePath);
 				return true;
 			}
 		}
@@ -2731,22 +2771,117 @@ bool FFileSystemCacheStore::RunInitialSpeedTest()
 
 	// Wait indefinitely for completion
 	SpeedTestState->CompletionEvent.Wait();
-	SpeedStats = SpeedTestState->SpeedStats;
+	OutSpeedStats = SpeedTestState->SpeedStats;
 	return SpeedTestState->bResult;
 }
 
-ILegacyCacheStore* CreateFileSystemCacheStore(
-	const TCHAR* CachePath,
-	const TCHAR* Params,
-	const TCHAR* AccessLogPath,
-	ECacheStoreFlags& OutFlags)
+void FFileSystemCacheStoreParams::Parse(const TCHAR* NodeName, const TCHAR* Config)
 {
-	TUniquePtr<FFileSystemCacheStore> Store = MakeUnique<FFileSystemCacheStore>(CachePath, Params, AccessLogPath, OutFlags);
-	if (OutFlags == ECacheStoreFlags::None)
+	FString Key;
+
+	// Path Params
+
+	FParse::Value(Config, TEXT("Path="), CachePath);
+
+	if (FParse::Value(Config, TEXT("EnvPathOverride="), Key))
+	{
+		if (FString Value = FPlatformMisc::GetEnvironmentVariable(*Key); !Value.IsEmpty())
+		{
+			CachePath = MoveTemp(Value);
+			UE_LOG(LogDerivedDataCache, Log, TEXT("%s: Found environment variable %s=%s"), NodeName, *Key, *CachePath);
+		}
+		if (FString Value; FPlatformMisc::GetStoredValue(TEXT("Epic Games"), TEXT("GlobalDataCachePath"), *Key, Value) && !Value.IsEmpty())
+		{
+			CachePath = MoveTemp(Value);
+			UE_LOG(LogDerivedDataCache, Log, TEXT("%s: Found registry key GlobalDataCachePath %s=%s"), NodeName, *Key, *CachePath);
+		}
+	}
+
+	if (FParse::Value(Config, TEXT("CommandLineOverride="), Key) &&
+		FParse::Value(FCommandLine::Get(), *(Key + TEXT("=")), CachePath))
+	{
+		UE_LOG(LogDerivedDataCache, Log, TEXT("%s: Found command line override %s=%s"), NodeName, *Key, *CachePath);
+	}
+
+	if (FParse::Value(Config, TEXT("EditorOverrideSetting="), Key))
+	{
+		if (FString Setting; GConfig->GetString(TEXT("/Script/UnrealEd.EditorSettings"), *Key, Setting, GEditorSettingsIni) && !Setting.IsEmpty())
+		{
+			if (FString Value; FParse::Value(*Setting, TEXT("Path="), Value))
+			{
+				Value.TrimQuotesInline();
+				Value.ReplaceEscapedCharWithCharInline();
+				if (!Value.IsEmpty())
+				{
+					CachePath = MoveTemp(Value);
+					UE_LOG(LogDerivedDataCache, Log, TEXT("%s: Found editor settings override %s=%s"), NodeName, *Key, *CachePath);
+				}
+			}
+		}
+	}
+
+	// Paths that start with a '?' are config keys.
+	if (CachePath.StartsWith(TEXT("?")) && !GConfig->GetString(TEXT("DerivedDataCacheSettings"), *CachePath + 1, CachePath, GEngineIni))
+	{
+		CachePath.Empty();
+	}
+
+	FPaths::NormalizeFilename(CachePath);
+
+	if (const FString AbsoluteCachePath = IFileManager::Get().ConvertToAbsolutePathForExternalAppForRead(*CachePath); AbsoluteCachePath.Len() >= GMaxCacheRootLen)
+	{
+		const FText ErrorMessage = FText::Format(NSLOCTEXT("DerivedDataCache", "PathTooLong", "Cache path {0} is longer than {1} characters. Shorten the path to leave more characters for cache keys."),
+			FText::FromString(AbsoluteCachePath), FText::AsNumber(GMaxCacheRootLen));
+		FMessageDialog::Open(EAppMsgType::Ok, ErrorMessage);
+		UE_LOG(LogDerivedDataCache, Fatal, TEXT("%s"), *ErrorMessage.ToString());
+	}
+
+	// Other Params
+
+	FParse::Value(Config, TEXT("WriteAccessLog="), AccessLogPath);
+	FParse::Value(Config, TEXT("MaxRecordSizeKB="), MaxRecordSizeKB);
+	FParse::Value(Config, TEXT("MaxValueSizeKB="), MaxValueSizeKB);
+	FParse::Value(Config, TEXT("UnusedFileAge="), MaxFileAgeInDays);
+	FParse::Value(Config, TEXT("ConsiderSlowAt="), ConsiderSlowAtMs);
+	FParse::Value(Config, TEXT("DeactivateAt="), DeactivateAtMs);
+	FParse::Bool(Config, TEXT("PromptIfMissing="), bPromptIfMissing);
+	FParse::Bool(Config, TEXT("DeleteOnly="), bDeleteOnly);
+	FParse::Bool(Config, TEXT("ReadOnly="), bReadOnly);
+	FParse::Bool(Config, TEXT("Clean="), bClean);
+	FParse::Bool(Config, TEXT("Flush="), bFlush);
+
+	FParse::Bool(Config, TEXT("Touch="), bTouch);
+	bTouch = bTouch || FParse::Param(FCommandLine::Get(), TEXT("DDCTOUCH"));
+
+	bSkipSpeedTest = bSkipSpeedTest || FParse::Param(FCommandLine::Get(), TEXT("DDCSkipSpeedTest"));
+
+	bDeleteUnused = !bReadOnly;
+	FParse::Bool(Config, TEXT("DeleteUnused="), bDeleteUnused);
+	bDeleteUnused = bDeleteUnused && !FParse::Param(FCommandLine::Get(), TEXT("NoDDCCleanup"));
+
+	if (!FParse::Value(Config, TEXT("MaxFileChecksPerSec="), MaxScanRate))
 	{
-		Store.Reset();
+		int32 MaxFileScanRate;
+		if (GConfig->GetInt(TEXT("DDCCleanup"), TEXT("MaxFileChecksPerSec"), MaxFileScanRate, GEngineIni))
+		{
+			MaxScanRate = uint32(MaxFileScanRate);
+		}
 	}
-	return Store.Release();
+
+	FParse::Value(Config, TEXT("FoldersToClean="), MaxDirectoryScanCount);
+
+	GConfig->GetDouble(TEXT("DDCCleanup"), TEXT("TimeToWaitAfterInit"), MaintenanceDelayInSeconds, GEngineIni);
+}
+
+ILegacyCacheStore* CreateFileSystemCacheStore(
+	const TCHAR* NodeName,
+	const TCHAR* Config,
+	ECacheStoreFlags& OutFlags,
+	FString& OutPath)
+{
+	FFileSystemCacheStoreParams Params;
+	Params.Parse(NodeName, Config);
+	return FFileSystemCacheStore::TryCreate(NodeName, Params, OutFlags, OutPath);
 }
 
 } // UE::DerivedData
