DDC: Removed an erroneous check from the file system cache store

#rnx

diff --git a/Engine/Source/Developer/DerivedDataCache/Private/FileSystemCacheStore.cpp b/Engine/Source/Developer/DerivedDataCache/Private/FileSystemCacheStore.cpp
--- a/Engine/Source/Developer/DerivedDataCache/Private/FileSystemCacheStore.cpp
+++ b/Engine/Source/Developer/DerivedDataCache/Private/FileSystemCacheStore.cpp
@@ -989,7 +989,6 @@ ILegacyCacheStore* FFileSystemCacheStore::TryCreate(
 	checkf(GMaxCacheRootLen + GMaxCacheKeyLen <= FPlatformMisc::GetMaxPathLength(),
 		TEXT("Not enough room left for cache keys in max path."));
 
-	check(!Params.CachePath.IsEmpty());
 	if (Params.CachePath.IsEmpty())
 	{
 		UE_LOG(LogDerivedDataCache, Log, TEXT("%s: Disabled because no path is configured."), NodeName);
