Don't continue initialization of ZenCacheStore of backend is deemed unusable (unable to connect)

diff --git a/Engine/Source/Developer/DerivedDataCache/Private/ZenCacheStore.cpp b/Engine/Source/Developer/DerivedDataCache/Private/ZenCacheStore.cpp
--- a/Engine/Source/Developer/DerivedDataCache/Private/ZenCacheStore.cpp
+++ b/Engine/Source/Developer/DerivedDataCache/Private/ZenCacheStore.cpp
@@ -1444,6 +1444,7 @@ TTuple<ILegacyCacheStore*, ECacheStoreFlags> CreateZenCacheStore(const TCHAR* No
 	{
 		UE_LOG(LogDerivedDataCache, Warning, TEXT("%s: Failed to contact the service (%s), will not use it."), NodeName, *Backend->GetName());
 		Backend.Reset();
+		return MakeTuple(nullptr, ECacheStoreFlags::None);
 	}
 
 	bool bReadOnly = false;
