Fix crash when using -ddc=Cold on projects using ZenServer as Local DDC and Multi-Process Cook simultaneously.

diff --git a/Engine/Source/Developer/DerivedDataCache/Private/ZenCacheStore.cpp b/Engine/Source/Developer/DerivedDataCache/Private/ZenCacheStore.cpp
--- a/Engine/Source/Developer/DerivedDataCache/Private/ZenCacheStore.cpp
+++ b/Engine/Source/Developer/DerivedDataCache/Private/ZenCacheStore.cpp
@@ -1415,7 +1415,9 @@ TTuple<ILegacyCacheStore*, ECacheStoreFlags> CreateZenCacheStore(const TCHAR* No
 		AutoLaunchSettings.bLimitProcessLifetime = true;
 
 		// Flush the cache if requested.
-		if (bFlush)
+		uint32 MultiprocessId = 0;
+		FParse::Value(FCommandLine::Get(), TEXT("-MultiprocessId="), MultiprocessId);
+		if (bFlush && (MultiprocessId == 0))
 		{
 			bool bStopped = true;
 			if (UE::Zen::IsLocalServiceRunning(*AutoLaunchSettings.DataPath))
