DDC: Started toward exposing the cache store owner to the cache stores and flattening the hierarchy

- Construct an empty hierarchy first and populate it as the graph is parsed.
- Record the need for async and verify nodes, and create them only at the end.
- Extract parsing of debug options including throttling into the cache store owner.
- Expose the cache store owner to the cache store parsing/creation functions.
- Remove the unused speed class debug option.

#rnx

diff --git a/Engine/Source/Developer/DerivedDataCache/Private/DerivedDataBackendInterface.h b/Engine/Source/Developer/DerivedDataCache/Private/DerivedDataBackendInterface.h
--- a/Engine/Source/Developer/DerivedDataCache/Private/DerivedDataBackendInterface.h
+++ b/Engine/Source/Developer/DerivedDataCache/Private/DerivedDataBackendInterface.h
@@ -64,9 +64,6 @@ enum class EBackendSpeedClass
 /** Debug options that can be applied to backends to simulate different behavior */
 struct FBackendDebugOptions
 {
-	/** Apply behavior of this speed class */
-	EBackendSpeedClass	SpeedClass;
-
 	/** Filter to control the keys for which to simulate a miss. */
 	FCacheKeyFilter SimulateMissFilter;
 
diff --git a/Engine/Source/Developer/DerivedDataCache/Private/DerivedDataBackends.cpp b/Engine/Source/Developer/DerivedDataCache/Private/DerivedDataBackends.cpp
--- a/Engine/Source/Developer/DerivedDataCache/Private/DerivedDataBackends.cpp
+++ b/Engine/Source/Developer/DerivedDataCache/Private/DerivedDataBackends.cpp
@@ -61,7 +61,7 @@ int32 AddToAsyncTaskCounter(int32 Addend)
 namespace UE::DerivedData
 {
 
-ILegacyCacheStore* CreateCacheStoreAsync(ILegacyCacheStore* InnerBackend, ECacheStoreFlags InnerFlags, IMemoryCacheStore* MemoryCache);
+ILegacyCacheStore* CreateCacheStoreAsync(ILegacyCacheStore* InnerBackend, IMemoryCacheStore* MemoryCache);
 ILegacyCacheStore* CreateCacheStoreHierarchy(ICacheStoreOwner*& OutOwner, IMemoryCacheStore* MemoryCache);
 ILegacyCacheStore* CreateCacheStoreThrottle(ILegacyCacheStore* InnerCache, uint32 LatencyMS, uint32 MaxBytesPerSecond);
 ILegacyCacheStore* CreateCacheStoreVerify(ILegacyCacheStore* InnerCache, bool bPutOnError);
@@ -77,7 +77,7 @@ ILegacyCacheStore* TryCreateCacheStoreReplay(ILegacyCacheStore* InnerCache);
  * This class is used to create a singleton that represents the derived data cache hierarchy and all of the wrappers necessary
  * ideally this would be data driven and the backends would be plugins...
  */
-class FDerivedDataBackendGraph final : public FDerivedDataBackend
+class FDerivedDataBackendGraph final : public FDerivedDataBackend, ICacheStoreOwner
 {
 public:
 	using FParsedNode = TPair<ILegacyCacheStore*, ECacheStoreFlags>;
@@ -91,8 +91,6 @@ public:
 		, MemoryCache(nullptr)
 		, BootCache(nullptr)
 		, WritePakCache(nullptr)
-		, AsyncNode(nullptr)
-		, VerifyNode(nullptr)
 		, Hierarchy(nullptr)
 		, bUsingSharedDDC(false)
 		, bIsShuttingDown(false)
@@ -120,7 +118,7 @@ public:
 		RootCache = nullptr;
 		FParsedNodeMap ParsedNodes;
 
-		FParsedNode RootNode(nullptr, ECacheStoreFlags::None);
+		ILegacyCacheStore* RootNode = nullptr;
 
 		// Create the graph using ini settings. The string "default" forwards creation to use the default graph.
 
@@ -148,8 +146,10 @@ public:
 
 			if (!GraphName.IsEmpty() && GraphName != TEXT("Default"))
 			{
-				RootNode = ParseNode(RootName, GEngineIni, *GraphName, ParsedNodes);
-				if (!RootNode.Key)
+				RootNode = CreateCacheStoreHierarchy(Hierarchy, GetMemoryCache());
+				CreatedNodes.AddUnique(RootNode);
+
+				if (!ParseNode(RootName, GEngineIni, *GraphName, ParsedNodes))
 				{
 					// Destroy any cache stores that have been created.
 					ParsedNodes.Empty();
@@ -157,22 +157,27 @@ public:
 					MemoryCache = nullptr;
 					BootCache = nullptr;
 					WritePakCache = nullptr;
-					AsyncNode = nullptr;
-					VerifyNode = nullptr;
 					Hierarchy = nullptr;
 					ReadPakCache.Empty();
 					Directories.Empty();
+					bAsyncFound = false;
+					bBootFound = false;
+					bHierarchyFound = false;
+					bVerifyFound = false;
+					bVerifyFix = false;
 					UE_LOG(LogDerivedDataCache, Warning,
 						TEXT("Unable to create cache graph '%s'. Reverting to the default graph."), *GraphName);
 				}
 			}
 
-			if (!RootNode.Key)
+			if (!Hierarchy)
 			{
 				// Try to use the default graph.
 				GraphName = FApp::IsEngineInstalled() ? TEXT("InstalledDerivedDataBackendGraph") : TEXT("DerivedDataBackendGraph");
-				RootNode = ParseNode(RootName, GEngineIni, *GraphName, ParsedNodes);
-				if (!RootNode.Key)
+				RootNode = CreateCacheStoreHierarchy(Hierarchy, GetMemoryCache());
+				CreatedNodes.AddUnique(RootNode);
+
+				if (!ParseNode(RootName, GEngineIni, *GraphName, ParsedNodes))
 				{
 					FString Entry;
 					if (!GConfig->DoesSectionExist(*GraphName, GEngineIni))
@@ -197,44 +202,30 @@ public:
 			}
 		}
 
-		// Hierarchy must exist in the graph.
-		if (!Hierarchy)
-		{
-			ILegacyCacheStore* HierarchyStore = CreateCacheStoreHierarchy(Hierarchy, GetMemoryCache());
-			if (RootNode.Key)
-			{
-				Hierarchy->Add(RootNode.Key, RootNode.Value);
-			}
-			CreatedNodes.AddUnique(HierarchyStore);
-			RootNode.Key = HierarchyStore;
-		}
-
 		// Async must exist in the graph.
-		if (!AsyncNode)
 		{
-			AsyncNode = CreateCacheStoreAsync(RootNode.Key, RootNode.Value, GetMemoryCache());
-			CreatedNodes.AddUnique(AsyncNode);
-			RootNode.Key = AsyncNode;
+			RootNode = CreateCacheStoreAsync(RootNode, GetMemoryCache());
+			CreatedNodes.AddUnique(RootNode);
 		}
 
 		// Create a Verify node when using -DDC-Verify[=Type1[@Rate2][+Type2[@Rate2]...]].
-		if (!VerifyNode)
+		if (bVerifyFound)
 		{
 			FString VerifyArg;
 			if (FParse::Value(FCommandLine::Get(), TEXT("-DDC-Verify="), VerifyArg) ||
 				FParse::Param(FCommandLine::Get(), TEXT("DDC-Verify")))
 			{
-				VerifyNode = CreateCacheStoreVerify(RootNode.Key, /*bPutOnError*/ false);
-				CreatedNodes.AddUnique(VerifyNode);
-				RootNode.Key = VerifyNode;
+				IFileManager::Get().DeleteDirectory(*(FPaths::ProjectSavedDir() / TEXT("VerifyDDC/")), /*bRequireExists*/ false, /*bTree*/ true);
+				RootNode = CreateCacheStoreVerify(RootNode, /*bPutOnError*/ bVerifyFix);
+				CreatedNodes.AddUnique(RootNode);
 			}
 		}
 
 		// Create a Replay node when requested on the command line.
-		if (ILegacyCacheStore* ReplayNode = TryCreateCacheStoreReplay(RootNode.Key))
+		if (ILegacyCacheStore* ReplayNode = TryCreateCacheStoreReplay(RootNode))
 		{
+			RootNode = ReplayNode;
 			CreatedNodes.AddUnique(ReplayNode);
-			RootNode.Key = ReplayNode;
 		}
 
 		if (MaxKeyLength == 0)
@@ -242,7 +233,7 @@ public:
 			MaxKeyLength = MAX_BACKEND_KEY_LENGTH;
 		}
 
-		RootCache = RootNode.Key;
+		RootCache = RootNode;
 	}
 
 	/**
@@ -264,15 +255,15 @@ public:
 	 * @param InParsedNodes Map of parsed nodes and their names to be able to find already parsed nodes
 	 * @return Derived data backend interface instance created from ini settings
 	 */
-	FParsedNode ParseNode(const FString& NodeName, const FString& IniFilename, const TCHAR* IniSection, FParsedNodeMap& InParsedNodes)
+	bool ParseNode(const FString& NodeName, const FString& IniFilename, const TCHAR* IniSection, FParsedNodeMap& InParsedNodes)
 	{
 		if (const FParsedNode* ParsedNode = InParsedNodes.Find(NodeName))
 		{
 			UE_LOG(LogDerivedDataCache, Warning, TEXT("Node %s was referenced more than once in the graph. Nodes may not be shared."), *NodeName);
-			return *ParsedNode;
+			return false;
 		}
 
-		FParsedNode ParsedNode(nullptr, ECacheStoreFlags::None);
+		ICacheStoreOwner* NodeOwner = this;
 		FString Entry;
 		if (GConfig->GetString(IniSection, *NodeName, Entry, IniFilename))
 		{
@@ -280,134 +271,74 @@ public:
 			Entry.RemoveFromStart(TEXT("("));
 			Entry.RemoveFromEnd(TEXT(")"));
 
+			TGuardValue NodeNameGuard(ActiveNodeName, NodeName);
+			TGuardValue NodeConfigGuard(ActiveNodeConfig, Entry);
+
 			FString	NodeType;
 			if (FParse::Value(*Entry, TEXT("Type="), NodeType))
 			{
 				if (NodeType == TEXT("FileSystem"))
 				{
-					ParsedNode = ParseDataCache(*NodeName, *Entry);
+					return ParseDataCache(*NodeName, *Entry);
 				}
 				else if (NodeType == TEXT("Boot"))
 				{
-					UE_LOG(LogDerivedDataCache, Display, TEXT("Boot nodes are deprecated. Please remove the Boot node from the cache graph."));
-					if (BootCache == nullptr)
-					{
-						BootCache = ParseBootCache(*NodeName, *Entry);
-						ParsedNode = MakeTuple(BootCache, ECacheStoreFlags::Local | ECacheStoreFlags::Query | ECacheStoreFlags::Store);
-					}
-					else
-					{
-						UE_LOG(LogDerivedDataCache, Warning, TEXT("Unable to create %s Boot cache because only one Boot node is supported."), *NodeName);
-					}
+					return ParseBootCache(*NodeName, *Entry);
 				}
 				else if (NodeType == TEXT("Memory"))
 				{
-					ParsedNode = MakeTuple(ParseMemoryCache(*NodeName, *Entry), ECacheStoreFlags::Local | ECacheStoreFlags::Query | ECacheStoreFlags::Store);
+					return ParseMemoryCache(*NodeName, *Entry);
 				}
 				else if (NodeType == TEXT("Hierarchical"))
 				{
-					ParsedNode = ParseHierarchyNode(*NodeName, *Entry, IniFilename, IniSection, InParsedNodes);
+					return ParseHierarchyNode(*NodeName, *Entry, IniFilename, IniSection, InParsedNodes);
 				}
 				else if (NodeType == TEXT("KeyLength"))
 				{
-					if (MaxKeyLength == 0)
-					{
-						ParsedNode = ParseKeyLength(*NodeName, *Entry, IniFilename, IniSection, InParsedNodes);
-					}
-					else
-					{
-						UE_LOG(LogDerivedDataCache, Warning, TEXT("Unable to create %s KeyLength node because only one KeyLength node is supported."), *NodeName);
-					}
+					return ParseKeyLength(*NodeName, *Entry, IniFilename, IniSection, InParsedNodes);
 				}
 				else if (NodeType == TEXT("AsyncPut"))
 				{
-					if (AsyncNode == nullptr)
-					{
-						ParsedNode = ParseAsyncNode(*NodeName, *Entry, IniFilename, IniSection, InParsedNodes);
-						AsyncNode = ParsedNode.Key;
-					}
-					else
-					{
-						UE_LOG(LogDerivedDataCache, Warning, TEXT("Unable to create %s AsyncPut because only one AsyncPut node is supported."), *NodeName);
-					}
+					return ParseAsyncNode(*NodeName, *Entry, IniFilename, IniSection, InParsedNodes);
 				}
 				else if (NodeType == TEXT("Verify"))
 				{
-					if (VerifyNode == nullptr)
-					{
-						ParsedNode = ParseVerify(*NodeName, *Entry, IniFilename, IniSection, InParsedNodes);
-						VerifyNode = ParsedNode.Key;
-					}
-					else
-					{
-						UE_LOG(LogDerivedDataCache, Warning, TEXT("Unable to create %s Verify because only one Verify node is supported."), *NodeName);
-					}
+					return ParseVerify(*NodeName, *Entry, IniFilename, IniSection, InParsedNodes);
 				}
 				else if (NodeType == TEXT("ReadPak"))
 				{
-					ParsedNode = MakeTuple(ParsePak(*NodeName, *Entry, /*bWriting*/ false), ECacheStoreFlags::Local | ECacheStoreFlags::Query | ECacheStoreFlags::StopStore);
+					return ParsePak(*NodeName, *Entry, /*bWriting*/ false, ECacheStoreFlags::Local | ECacheStoreFlags::Query | ECacheStoreFlags::StopStore);
 				}
 				else if (NodeType == TEXT("WritePak"))
 				{
-					ParsedNode = MakeTuple(ParsePak(*NodeName, *Entry, /*bWriting*/ true), ECacheStoreFlags::Local | ECacheStoreFlags::Store);
+					return ParsePak(*NodeName, *Entry, /*bWriting*/ true, ECacheStoreFlags::Local | ECacheStoreFlags::Store);
 				}
 				else if (NodeType == TEXT("S3"))
 				{
-					ParsedNode = MakeTuple(ParseS3Cache(*NodeName, *Entry), ECacheStoreFlags::Local | ECacheStoreFlags::Query | ECacheStoreFlags::StopStore);
+					return ParseS3Cache(*NodeName, *Entry);
 				}
 				else if (NodeType == TEXT("Cloud") || NodeType == TEXT("Http"))
 				{
-					ParsedNode = CreateHttpCacheStore(*NodeName, *Entry);
+					TTuple<ILegacyCacheStore*, ECacheStoreFlags> Node = CreateHttpCacheStore(*NodeName, *Entry);
+					if (Node.Key)
+					{
+						NodeOwner->Add(Node.Key, Node.Value);
+						return true;
+					}
 				}
 				else if (NodeType == TEXT("Zen"))
 				{
-					ParsedNode = CreateZenCacheStore(*NodeName, *Entry);
-				}
-				
-				if (ParsedNode.Key)
-				{
-					// Add a throttling layer if parameters are found
-					uint32 LatencyMS = 0;
-					FParse::Value(*Entry, TEXT("LatencyMS="), LatencyMS);
-
-					uint32 MaxBytesPerSecond = 0;
-					FParse::Value(*Entry, TEXT("MaxBytesPerSecond="), MaxBytesPerSecond);
-
-					if (LatencyMS != 0 || MaxBytesPerSecond != 0)
+					TTuple<ILegacyCacheStore*, ECacheStoreFlags> Node = CreateZenCacheStore(*NodeName, *Entry);
+					if (Node.Key)
 					{
-						CreatedNodes.AddUnique(ParsedNode.Key);
-						ParsedNode = MakeTuple(CreateCacheStoreThrottle(ParsedNode.Key, LatencyMS, MaxBytesPerSecond), ParsedNode.Value);
+						NodeOwner->Add(Node.Key, Node.Value);
+						return true;
 					}
 				}
 			}
 		}
 
-		if (ParsedNode.Key)
-		{
-			// Store this node so that we don't require any order of adding nodes
-			InParsedNodes.Add(NodeName, ParsedNode);
-			// Keep references to all created nodes.
-			CreatedNodes.AddUnique(ParsedNode.Key);
-
-			// Parse any debug options for this node. E.g. -DDC-<Name>-MissRate
-			FBackendDebugOptions DebugOptions;
-			if (FBackendDebugOptions::ParseFromTokens(DebugOptions, *NodeName, FCommandLine::Get()))
-			{
-				if (!ParsedNode.Key->LegacyDebugOptions(DebugOptions))
-				{
-					UE_LOG(LogDerivedDataCache, Warning, TEXT("%s: Node is ignoring one or more -DDC-%s-<Option> debug options."), *NodeName, *NodeName);
-				}
-			}
-			else if (bHasDefaultDebugOptions)
-			{
-				if (!ParsedNode.Key->LegacyDebugOptions(DefaultDebugOptions))
-				{
-					UE_LOG(LogDerivedDataCache, Warning, TEXT("%s: Node is ignoring one or more -DDC-All-<Option> debug options."), *NodeName);
-				}
-			}
-		}
-
-		return ParsedNode;
+		return false;
 	}
 
 	/**
@@ -418,45 +349,51 @@ public:
 	 * @param bWriting true to create pak interface for writing
 	 * @return Pak file data backend interface instance or nullptr if unsuccessful
 	 */
-	ILegacyCacheStore* ParsePak( const TCHAR* NodeName, const TCHAR* Entry, const bool bWriting )
+	bool ParsePak(const TCHAR* NodeName, const TCHAR* Entry, const bool bWriting, const ECacheStoreFlags Flags)
 	{
 		ILegacyCacheStore* PakNode = nullptr;
 		FString PakFilename;
 		FParse::Value( Entry, TEXT("Filename="), PakFilename );
 		bool bCompressed = GetParsedBool(Entry, TEXT("Compressed="));
 
-		if( !PakFilename.Len() )
+		if (PakFilename.IsEmpty())
 		{
-			UE_LOG( LogDerivedDataCache, Log, TEXT("FDerivedDataBackendGraph: %s pak cache Filename not found in *engine.ini, will not use a pak cache."), NodeName );
+			UE_LOG(LogDerivedDataCache, Log, TEXT("FDerivedDataBackendGraph: %s pak cache Filename not found in *engine.ini, will not use a pak cache."), NodeName);
+			return false;
 		}
-		else
+
+		if (bWriting)
 		{
-			// now add the pak read cache (if any) to the front of the cache hierarchy
-			if ( bWriting )
+			if (WritePakCache)
 			{
-				FGuid Temp = FGuid::NewGuid();
-				ReadPakFilename = PakFilename;
-				WritePakFilename = PakFilename + TEXT(".") + Temp.ToString();
-				WritePakCache = CreatePakFileCacheStore(*WritePakFilename, /*bWriting*/ true, bCompressed);
-				PakNode = WritePakCache;
+				UE_LOG(LogDerivedDataCache, Warning, TEXT("Unable to create %s pak cache because only one pak cache write node is supported."), *NodeName);
+				return false;
 			}
-			else
+
+			FGuid Temp = FGuid::NewGuid();
+			ReadPakFilename = PakFilename;
+			WritePakFilename = PakFilename + TEXT(".") + Temp.ToString();
+			WritePakCache = CreatePakFileCacheStore(*WritePakFilename, /*bWriting*/ true, bCompressed);
+			PakNode = WritePakCache;
+			Add(WritePakCache, Flags);
+			return true;
+		}
+		else
+		{
+			bool bReadPak = FPlatformFileManager::Get().GetPlatformFile().FileExists(*PakFilename);
+			if (!bReadPak)
 			{
-				bool bReadPak = FPlatformFileManager::Get().GetPlatformFile().FileExists( *PakFilename );
-				if( bReadPak )
-				{
-					IPakFileCacheStore* ReadPak = CreatePakFileCacheStore(*PakFilename, /*bWriting*/ false, bCompressed);
-					ReadPakFilename = PakFilename;
-					PakNode = ReadPak;
-					ReadPakCache.Add(ReadPak);
-				}
-				else
-				{
-					UE_LOG( LogDerivedDataCache, Log, TEXT("FDerivedDataBackendGraph: %s pak cache file %s not found, will not use a pak cache."), NodeName, *PakFilename );
-				}
+				UE_LOG(LogDerivedDataCache, Log, TEXT("FDerivedDataBackendGraph: %s pak cache file %s not found, will not use a pak cache."), NodeName, *PakFilename);
+				return false;
 			}
+
+			IPakFileCacheStore* ReadPak = CreatePakFileCacheStore(*PakFilename, /*bWriting*/ false, bCompressed);
+			ReadPakFilename = PakFilename;
+			PakNode = ReadPak;
+			ReadPakCache.Add(ReadPak);
+			Add(ReadPak, Flags);
+			return true;
 		}
-		return PakNode;
 	}
 
 	/**
@@ -469,28 +406,25 @@ public:
 	 * @param InParsedNodes map of nodes and their names which have already been parsed
 	 * @return Verify wrapper backend interface instance or nullptr if unsuccessful
 	 */
-	FParsedNode ParseVerify(const TCHAR* NodeName, const TCHAR* Entry, const FString& IniFilename, const TCHAR* IniSection, FParsedNodeMap& InParsedNodes)
+	bool ParseVerify(const TCHAR* NodeName, const TCHAR* Entry, const FString& IniFilename, const TCHAR* IniSection, FParsedNodeMap& InParsedNodes)
 	{
-		FParsedNode InnerNode(nullptr, ECacheStoreFlags::None);
-		FString InnerName;
-		if (FParse::Value(Entry, TEXT("Inner="), InnerName))
+		if (bVerifyFound)
 		{
-			InnerNode = ParseNode(InnerName, IniFilename, IniSection, InParsedNodes);
+			UE_LOG(LogDerivedDataCache, Warning, TEXT("Unable to create %s Verify because only one Verify node is supported."), *NodeName);
+			return false;
 		}
 
-		if (InnerNode.Key)
-		{
-			IFileManager::Get().DeleteDirectory(*(FPaths::ProjectSavedDir() / TEXT("VerifyDDC/")), /*bRequireExists*/ false, /*bTree*/ true);
+		bVerifyFound = true;
+		bVerifyFix = GetParsedBool(Entry, TEXT("Fix="));
 
-			const bool bFix = GetParsedBool(Entry, TEXT("Fix="));
-			InnerNode = MakeTuple(CreateCacheStoreVerify(InnerNode.Key, bFix), InnerNode.Value);
-		}
-		else
+		FString InnerName;
+		if (FParse::Value(Entry, TEXT("Inner="), InnerName) && ParseNode(InnerName, IniFilename, IniSection, InParsedNodes))
 		{
-			UE_LOG(LogDerivedDataCache, Warning, TEXT("Unable to find inner node %s for Verify node %s. Verify node will not be created."), *InnerName, NodeName);
+			return true;
 		}
 
-		return InnerNode;
+		UE_LOG(LogDerivedDataCache, Warning, TEXT("Unable to find inner node %s for Verify node %s. Verify node will not be created."), *InnerName, NodeName);
+		return false;
 	}
 
 	/**
@@ -503,34 +437,23 @@ public:
 	 * @param InParsedNodes map of nodes and their names which have already been parsed
 	 * @return AsyncPut wrapper backend interface instance or nullptr if unsuccessful
 	 */
-	FParsedNode ParseAsyncNode(const TCHAR* NodeName, const TCHAR* Entry, const FString& IniFilename, const TCHAR* IniSection, FParsedNodeMap& InParsedNodes)
+	bool ParseAsyncNode(const TCHAR* NodeName, const TCHAR* Entry, const FString& IniFilename, const TCHAR* IniSection, FParsedNodeMap& InParsedNodes)
 	{
-		FParsedNode InnerNode(nullptr, ECacheStoreFlags::None);
-		FString InnerName;
-		if (FParse::Value(Entry, TEXT("Inner="), InnerName))
+		if (bAsyncFound)
 		{
-			InnerNode = ParseNode(InnerName, IniFilename, IniSection, InParsedNodes);
+			UE_LOG(LogDerivedDataCache, Warning, TEXT("Unable to create %s AsyncPut because only one AsyncPut node is supported."), *NodeName);
+			return false;
 		}
 
-		if (InnerNode.Key)
-		{
-			if (!Hierarchy)
-			{
-				ILegacyCacheStore* HierarchyStore = CreateCacheStoreHierarchy(Hierarchy, GetMemoryCache());
-				Hierarchy->Add(InnerNode.Key, InnerNode.Value);
-				CreatedNodes.AddUnique(HierarchyStore);
-				InnerNode.Key = HierarchyStore;
-			}
-
-			ILegacyCacheStore* AsyncStore = CreateCacheStoreAsync(InnerNode.Key, InnerNode.Value, GetMemoryCache());
-			InnerNode = MakeTuple(AsyncStore, InnerNode.Value);
-		}
-		else
+		bAsyncFound = true;
+		FString InnerName;
+		if (FParse::Value(Entry, TEXT("Inner="), InnerName) && ParseNode(InnerName, IniFilename, IniSection, InParsedNodes))
 		{
-			UE_LOG(LogDerivedDataCache, Warning, TEXT("Unable to find inner node %s for AsyncPut node %s. AsyncPut node will not be created."), *InnerName, NodeName);
+			return true;
 		}
 
-		return InnerNode;
+		UE_LOG(LogDerivedDataCache, Warning, TEXT("Unable to find inner node %s for AsyncPut node %s. AsyncPut node will not be created."), *InnerName, NodeName);
+		return false;
 	}
 
 	/**
@@ -543,33 +466,25 @@ public:
 	 * @param InParsedNodes map of nodes and their names which have already been parsed
 	 * @return KeyLength wrapper backend interface instance or nullptr if unsuccessful
 	 */
-	FParsedNode ParseKeyLength(const TCHAR* NodeName, const TCHAR* Entry, const FString& IniFilename, const TCHAR* IniSection, FParsedNodeMap& InParsedNodes)
+	bool ParseKeyLength(const TCHAR* NodeName, const TCHAR* Entry, const FString& IniFilename, const TCHAR* IniSection, FParsedNodeMap& InParsedNodes)
 	{
 		if (MaxKeyLength)
 		{
 			UE_LOG(LogDerivedDataCache, Warning, TEXT("Node %s is disabled because there may be only one key length node."), NodeName);
-			return MakeTuple(nullptr, ECacheStoreFlags::None);
+			return false;
 		}
 
-		FParsedNode InnerNode(nullptr, ECacheStoreFlags::None);
 		FString InnerName;
-		if (FParse::Value(Entry, TEXT("Inner="), InnerName))
-		{
-			InnerNode = ParseNode(InnerName, IniFilename, IniSection, InParsedNodes);
-		}
-
-		if (InnerNode.Key)
+		if (FParse::Value(Entry, TEXT("Inner="), InnerName) && ParseNode(InnerName, IniFilename, IniSection, InParsedNodes))
 		{
 			int32 KeyLength = MAX_BACKEND_KEY_LENGTH;
 			FParse::Value(Entry, TEXT("Length="), KeyLength);
 			MaxKeyLength = FMath::Clamp(KeyLength, 0, MAX_BACKEND_KEY_LENGTH);
-		}
-		else
-		{
-			UE_LOG(LogDerivedDataCache, Warning, TEXT("Unable to find inner node %s for KeyLength node %s. KeyLength node will not be created."), *InnerName, NodeName);
+			return true;
 		}
 
-		return InnerNode;
+		UE_LOG(LogDerivedDataCache, Warning, TEXT("Unable to find inner node %s for KeyLength node %s. KeyLength node will not be created."), *InnerName, NodeName);
+		return false;
 	}
 
 	/**
@@ -582,29 +497,26 @@ public:
 	 * @param InParsedNodes map of nodes and their names which have already been parsed
 	 * @return Hierarchical backend interface instance or nullptr if unsuccessful
 	 */
-	FParsedNode ParseHierarchyNode(const TCHAR* NodeName, const TCHAR* Entry, const FString& IniFilename, const TCHAR* IniSection, FParsedNodeMap& InParsedNodes)
+	bool ParseHierarchyNode(const TCHAR* NodeName, const TCHAR* Entry, const FString& IniFilename, const TCHAR* IniSection, FParsedNodeMap& InParsedNodes)
 	{
-		if (Hierarchy)
+		if (bHierarchyFound)
 		{
 			UE_LOG(LogDerivedDataCache, Warning, TEXT("Node %s is disabled because there may be only one hierarchy node. "
 				"Confirm there is only one hierarchy in the cache graph and that it is inside of any async node."), NodeName);
-			return MakeTuple(nullptr, ECacheStoreFlags::None);
+			return false;
 		}
 
-		ILegacyCacheStore* HierarchyStore = CreateCacheStoreHierarchy(Hierarchy, GetMemoryCache());
-		CreatedNodes.AddUnique(HierarchyStore);
-
 		const TCHAR* InnerMatch = TEXT("Inner=");
 		const int32 InnerMatchLength = FCString::Strlen(InnerMatch);
 
-		TArray<FParsedNode> InnerNodes;
+		bHierarchyFound = true;
+		bool bParsed = false;
 		FString InnerName;
 		while (FParse::Value(Entry, InnerMatch, InnerName))
 		{
-			FParsedNode InnerNode = ParseNode(InnerName, IniFilename, IniSection, InParsedNodes);
-			if (InnerNode.Key)
+			if (ParseNode(InnerName, IniFilename, IniSection, InParsedNodes))
 			{
-				InnerNodes.Add(InnerNode);
+				bParsed = true;
 			}
 			else
 			{
@@ -616,18 +528,13 @@ public:
 			Entry += InnerMatchLength;
 		}
 
-		if (InnerNodes.IsEmpty())
+		if (!bParsed)
 		{
 			UE_LOG(LogDerivedDataCache, Warning, TEXT("Hierarchical cache %s has no inner backends and will not be created."), NodeName);
-			return MakeTuple(nullptr, ECacheStoreFlags::None);
+			return false;
 		}
-		ECacheStoreFlags Flags = ECacheStoreFlags::None;
-		for (const FParsedNode& Node : InnerNodes)
-		{
-			Hierarchy->Add(Node.Key, Node.Value);
-			Flags |= Node.Value;
-		}
-		return MakeTuple(HierarchyStore, Flags & ~ECacheStoreFlags::StopStore);
+
+		return true;
 	}
 
 	/**
@@ -637,7 +544,7 @@ public:
 	 * @param Entry Node definition.
 	 * @return Filesystem data cache backend interface instance or nullptr if unsuccessful
 	 */
-	FParsedNode ParseDataCache(const TCHAR* NodeName, const TCHAR* Config)
+	bool ParseDataCache(const TCHAR* NodeName, const TCHAR* Config)
 	{
 		FString Path;
 		ECacheStoreFlags Flags;
@@ -645,28 +552,29 @@ public:
 		{
 			bUsingSharedDDC |= NodeName == TEXTVIEW("Shared");
 			Directories.AddUnique(Path);
-			return MakeTuple(Store, Flags);
+			Add(Store, Flags);
+			return true;
 		}
-		return MakeTuple(nullptr, ECacheStoreFlags::None);
+		return false;
 	}
 
 	/**
 	 * Creates an S3 data cache interface.
 	 */
-	ILegacyCacheStore* ParseS3Cache(const TCHAR* NodeName, const TCHAR* Entry)
+	bool ParseS3Cache(const TCHAR* NodeName, const TCHAR* Entry)
 	{
 		FString ManifestPath;
 		if (!FParse::Value(Entry, TEXT("Manifest="), ManifestPath))
 		{
 			UE_LOG(LogDerivedDataCache, Error, TEXT("Node %s does not specify 'Manifest'."), NodeName);
-			return nullptr;
+			return false;
 		}
 
 		FString BaseUrl;
 		if (!FParse::Value(Entry, TEXT("BaseUrl="), BaseUrl))
 		{
 			UE_LOG(LogDerivedDataCache, Error, TEXT("Node %s does not specify 'BaseUrl'."), NodeName);
-			return nullptr;
+			return false;
 		}
 
 		FString CanaryObjectKey;
@@ -676,7 +584,7 @@ public:
 		if (!FParse::Value(Entry, TEXT("Region="), Region))
 		{
 			UE_LOG(LogDerivedDataCache, Error, TEXT("Node %s does not specify 'Region'."), NodeName);
-			return nullptr;
+			return false;
 		}
 
 		// Check the EnvPathOverride environment variable to allow persistent overriding of data cache path, eg for offsite workers.
@@ -690,7 +598,7 @@ public:
 				if (FilesystemCachePathEnv == TEXT("None"))
 				{
 					UE_LOG(LogDerivedDataCache, Log, TEXT("Node %s disabled due to %s=None"), NodeName, *EnvPathOverride);
-					return nullptr;
+					return false;
 				}
 				else
 				{
@@ -706,7 +614,7 @@ public:
 				{
 					if ( DDCPath.Len() > 0 )
 					{
-						CachePath = DDCPath;			
+						CachePath = DDCPath;
 						UE_LOG( LogDerivedDataCache, Log, TEXT("Found registry key GlobalDataCachePath %s=%s"), *EnvPathOverride, *CachePath );
 					}
 				}
@@ -715,11 +623,12 @@ public:
 
 		if (ILegacyCacheStore* Backend = CreateS3CacheStore(*ManifestPath, *BaseUrl, *Region, *CanaryObjectKey, *CachePath))
 		{
-			return Backend;
+			Add(Backend, ECacheStoreFlags::Local | ECacheStoreFlags::Query | ECacheStoreFlags::StopStore);
+			return true;
 		}
 
 		UE_LOG(LogDerivedDataCache, Log, TEXT("S3 backend is not supported on the current platform."));
-		return nullptr;
+		return false;
 	}
 
 	/**
@@ -730,9 +639,16 @@ public:
 	 * @param OutFilename filename specified for the cache
 	 * @return Boot data cache backend interface instance or nullptr if unsuccessful
 	 */
-	IMemoryCacheStore* ParseBootCache(const TCHAR* NodeName, const TCHAR* Entry)
+	bool ParseBootCache(const TCHAR* NodeName, const TCHAR* Entry)
 	{
-		IMemoryCacheStore* Cache = nullptr;
+		UE_LOG(LogDerivedDataCache, Display, TEXT("Boot nodes are deprecated. Please remove the Boot node from the cache graph."));
+		if (bBootFound)
+		{
+			UE_LOG(LogDerivedDataCache, Warning, TEXT("Unable to create %s Boot cache because only one Boot node is supported."), *NodeName);
+			return false;
+		}
+
+		bBootFound = true;
 
 		// Only allow boot cache with the editor. We don't want other tools and utilities (eg. SCW) writing to the same file.
 #if WITH_EDITOR
@@ -746,10 +662,16 @@ public:
 		MaxCacheSize = FMath::Min(MaxCacheSize, MaxSupportedCacheSize);
 
 		UE_LOG(LogDerivedDataCache, Display, TEXT("%s: Max Cache Size: %d MB"), NodeName, MaxCacheSize);
-		Cache = CreateMemoryCacheStore(TEXT("Boot"), MaxCacheSize * 1024 * 1024, /*bCanBeDisabled*/ true);
+		BootCache = CreateMemoryCacheStore(TEXT("Boot"), MaxCacheSize * 1024 * 1024, /*bCanBeDisabled*/ true);
 #endif
 
-		return Cache;
+		if (BootCache)
+		{
+			Add(BootCache, ECacheStoreFlags::Local | ECacheStoreFlags::Query | ECacheStoreFlags::Store);
+			return true;
+		}
+
+		return false;
 	}
 
 	/**
@@ -759,7 +681,7 @@ public:
 	 * @param Entry Node definition.
 	 * @return Memory data cache backend interface instance or nullptr if unsuccessful
 	 */
-	IMemoryCacheStore* ParseMemoryCache( const TCHAR* NodeName, const TCHAR* Entry )
+	bool ParseMemoryCache(const TCHAR* NodeName, const TCHAR* Entry)
 	{
 		FString Filename;
 		FParse::Value(Entry, TEXT("Filename="), Filename);
@@ -768,7 +690,12 @@ public:
 		{
 			UE_LOG(LogDerivedDataCache, Display, TEXT("Memory nodes that load from a file are deprecated. Please remove the filename from the cache configuration."));
 		}
-		return Cache;
+		if (Cache)
+		{
+			Add(Cache, ECacheStoreFlags::Local | ECacheStoreFlags::Query | ECacheStoreFlags::Store);
+			return true;
+		}
+		return false;
 	}
 
 	IMemoryCacheStore* GetMemoryCache()
@@ -980,6 +907,78 @@ public:
 	}
 
 private:
+	void Add(ILegacyCacheStore* CacheStore, ECacheStoreFlags Flags) final
+	{
+		check(Hierarchy);
+
+		CreatedNodes.AddUnique(CacheStore);
+
+		// Parse any debug options for this node. E.g. -DDC-<Name>-MissRate
+		FBackendDebugOptions DebugOptions;
+		if (FBackendDebugOptions::ParseFromTokens(DebugOptions, *ActiveNodeName, FCommandLine::Get()))
+		{
+			if (!CacheStore->LegacyDebugOptions(DebugOptions))
+			{
+				UE_LOG(LogDerivedDataCache, Warning, TEXT("%s: Node is ignoring one or more -DDC-%s-<Option> debug options."), *ActiveNodeName, *ActiveNodeName);
+			}
+		}
+		else if (bHasDefaultDebugOptions)
+		{
+			if (!CacheStore->LegacyDebugOptions(DefaultDebugOptions))
+			{
+				UE_LOG(LogDerivedDataCache, Warning, TEXT("%s: Node is ignoring one or more -DDC-All-<Option> debug options."), *ActiveNodeName);
+			}
+		}
+
+		// Add a throttling layer if parameters are found
+		{
+			uint32 LatencyMS = 0;
+			FParse::Value(*ActiveNodeConfig, TEXT("LatencyMS="), LatencyMS);
+
+			uint32 MaxBytesPerSecond = 0;
+			FParse::Value(*ActiveNodeConfig, TEXT("MaxBytesPerSecond="), MaxBytesPerSecond);
+
+			if (LatencyMS != 0 || MaxBytesPerSecond != 0)
+			{
+				ILegacyCacheStore* ThrottleNode = CreateCacheStoreThrottle(CacheStore, LatencyMS, MaxBytesPerSecond);
+				ThrottleNodes.Add(CacheStore, ThrottleNode);
+				CreatedNodes.Add(ThrottleNode);
+				CacheStore = ThrottleNode;
+			}
+		}
+
+		Hierarchy->Add(CacheStore, Flags);
+	}
+
+	void SetFlags(ILegacyCacheStore* CacheStore, ECacheStoreFlags Flags) final
+	{
+		check(Hierarchy);
+		if (ILegacyCacheStore* ThrottleNode = ThrottleNodes.FindRef(CacheStore))
+		{
+			Hierarchy->SetFlags(ThrottleNode, Flags);
+		}
+		else
+		{
+			Hierarchy->SetFlags(CacheStore, Flags);
+		}
+	}
+
+	void RemoveNotSafe(ILegacyCacheStore* CacheStore) final
+	{
+		check(Hierarchy);
+		CreatedNodes.Remove(CacheStore);
+		if (ILegacyCacheStore* ThrottleNode = ThrottleNodes.FindRef(CacheStore))
+		{
+			Hierarchy->RemoveNotSafe(ThrottleNode);
+			CreatedNodes.Remove(ThrottleNode);
+			ThrottleNodes.Remove(CacheStore);
+			delete ThrottleNode;
+		}
+		else
+		{
+			Hierarchy->RemoveNotSafe(CacheStore);
+		}
+	}
 
 	/** Delete all created backends in the reversed order they were created. */
 	void DestroyCreatedBackends()
@@ -1106,15 +1105,14 @@ private:
 	/** Root of the graph */
 	ILegacyCacheStore*					RootCache;
 
-	/** References to all created backed interfaces */
+	/** References to all created backends */
 	TArray< ILegacyCacheStore* > CreatedNodes;
+	TMap<ILegacyCacheStore*, ILegacyCacheStore*> ThrottleNodes;
 
 	/** Instances of backend interfaces which exist in only one copy */
 	IMemoryCacheStore* MemoryCache;
-	IMemoryCacheStore*	BootCache;
+	IMemoryCacheStore* BootCache;
 	IPakFileCacheStore* WritePakCache;
-	ILegacyCacheStore*	AsyncNode;
-	ILegacyCacheStore*	VerifyNode;
 	ICacheStoreOwner* Hierarchy;
 	/** Support for multiple read only pak files. */
 	TArray<IPakFileCacheStore*> ReadPakCache;
@@ -1141,6 +1139,15 @@ private:
 
 	FAutoConsoleCommand LoadReplayCommand;
 	TArray<FCacheReplayReader> Replays;
+
+	FString ActiveNodeName;
+	FString ActiveNodeConfig;
+
+	bool bAsyncFound = false;
+	bool bBootFound = false;
+	bool bHierarchyFound = false;
+	bool bVerifyFound = false;
+	bool bVerifyFix = false;
 };
 
 } // UE::DerivedData
@@ -1174,7 +1181,6 @@ struct Private::FBackendDebugMissState
 };
 
 FBackendDebugOptions::FBackendDebugOptions()
-	: SpeedClass(EBackendSpeedClass::Unknown)
 {
 	SimulateMissState.Get() = MakePimpl<Private::FBackendDebugMissState>();
 }
@@ -1194,15 +1200,6 @@ bool FBackendDebugOptions::ParseFromTokens(FBackendDebugOptions& OutOptions, con
 
 	const int32 PrefixLen = Prefix.Len();
 
-	// Look for -DDC-Local-Speed=, -DDC-Shared-Speed=, etc.
-	Prefix.Append(TEXTVIEW("Speed="));
-	FString SpeedClass;
-	if (FParse::Value(InInputTokens, *Prefix, SpeedClass))
-	{
-		LexFromString(OutOptions.SpeedClass, *SpeedClass);
-	}
-	Prefix.RemoveAt(PrefixLen, Prefix.Len() - PrefixLen);
-
 	// Look for -DDC-Local-MissRate=, -DDC-Shared-MissRate=, etc.
 	float MissRate = 0.0f;
 	Prefix.Append(TEXTVIEW("MissRate="));
diff --git a/Engine/Source/Developer/DerivedDataCache/Private/DerivedDataCacheStoreAsync.cpp b/Engine/Source/Developer/DerivedDataCache/Private/DerivedDataCacheStoreAsync.cpp
--- a/Engine/Source/Developer/DerivedDataCache/Private/DerivedDataCacheStoreAsync.cpp
+++ b/Engine/Source/Developer/DerivedDataCache/Private/DerivedDataCacheStoreAsync.cpp
@@ -25,7 +25,7 @@ namespace UE::DerivedData
 class FCacheStoreAsync : public ILegacyCacheStore
 {
 public:
-	FCacheStoreAsync(ILegacyCacheStore* InnerCache, ECacheStoreFlags InnerFlags, IMemoryCacheStore* MemoryCache);
+	FCacheStoreAsync(ILegacyCacheStore* InnerCache, IMemoryCacheStore* MemoryCache);
 
 	virtual void Put(
 		TConstArrayView<FCachePutRequest> Requests,
@@ -114,27 +114,18 @@ private:
 	ILegacyCacheStore* InnerCache;
 	IMemoryCacheStore* MemoryCache;
 	FDerivedDataCacheUsageStats UsageStats;
-	ECacheStoreFlags InnerFlags;
 };
 
-FCacheStoreAsync::FCacheStoreAsync(ILegacyCacheStore* InInnerCache, ECacheStoreFlags InInnerFlags, IMemoryCacheStore* InMemoryCache)
+FCacheStoreAsync::FCacheStoreAsync(ILegacyCacheStore* InInnerCache, IMemoryCacheStore* InMemoryCache)
 	: InnerCache(InInnerCache)
 	, MemoryCache(InMemoryCache)
-	, InnerFlags(InInnerFlags)
 {
 	check(InnerCache);
 }
 
 void FCacheStoreAsync::LegacyStats(FDerivedDataCacheStatsNode& OutNode)
 {
-	OutNode = {TEXT("Async"), TEXT(""), EnumHasAnyFlags(InnerFlags, ECacheStoreFlags::Local)};
-	OutNode.UsageStats.Add(TEXT(""), UsageStats);
-
-	InnerCache->LegacyStats(OutNode.Children.Add_GetRef(MakeShared<FDerivedDataCacheStatsNode>()).Get());
-	if (MemoryCache)
-	{
-		MemoryCache->LegacyStats(OutNode.Children.Add_GetRef(MakeShared<FDerivedDataCacheStatsNode>()).Get());
-	}
+	InnerCache->LegacyStats(OutNode);
 }
 
 template <typename RequestType, typename OnCompleteType, typename OnExecuteType>
@@ -183,9 +174,9 @@ void FCacheStoreAsync::Execute(
 		});
 }
 
-ILegacyCacheStore* CreateCacheStoreAsync(ILegacyCacheStore* InnerCache, ECacheStoreFlags InnerFlags, IMemoryCacheStore* MemoryCache)
+ILegacyCacheStore* CreateCacheStoreAsync(ILegacyCacheStore* InnerCache, IMemoryCacheStore* MemoryCache)
 {
-	return new FCacheStoreAsync(InnerCache, InnerFlags, MemoryCache);
+	return new FCacheStoreAsync(InnerCache, MemoryCache);
 }
 
 } // UE::DerivedData
diff --git a/Engine/Source/Developer/DerivedDataCache/Private/DerivedDataCacheStoreHierarchy.cpp b/Engine/Source/Developer/DerivedDataCache/Private/DerivedDataCacheStoreHierarchy.cpp
--- a/Engine/Source/Developer/DerivedDataCache/Private/DerivedDataCacheStoreHierarchy.cpp
+++ b/Engine/Source/Developer/DerivedDataCache/Private/DerivedDataCacheStoreHierarchy.cpp
@@ -25,7 +25,7 @@
 namespace UE::DerivedData
 {
 
-ILegacyCacheStore* CreateCacheStoreAsync(ILegacyCacheStore* InnerCache, ECacheStoreFlags InnerFlags, IMemoryCacheStore* MemoryCache);
+ILegacyCacheStore* CreateCacheStoreAsync(ILegacyCacheStore* InnerCache, IMemoryCacheStore* MemoryCache);
 
 ///////////////////////////////////////////////////////////////////////////////////////////////////////////////////////
 
@@ -247,7 +247,7 @@ void FCacheStoreHierarchy::Add(ILegacyCacheStore* CacheStore, ECacheStoreFlags F
 	FWriteScopeLock Lock(NodesLock);
 	checkf(!Algo::FindBy(Nodes, CacheStore, &FCacheStoreNode::Cache),
 		TEXT("Attempting to add a cache store that was previously registered to the hierarchy."));
-	TUniquePtr<ILegacyCacheStore> AsyncCacheStore(CreateCacheStoreAsync(CacheStore, Flags, MemoryCache));
+	TUniquePtr<ILegacyCacheStore> AsyncCacheStore(CreateCacheStoreAsync(CacheStore, MemoryCache));
 	Nodes.Add({CacheStore, Flags, {}, MoveTemp(AsyncCacheStore)});
 	UpdateNodeFlags();
 }
