DDC: Reverted the change to FCacheStoreAsync::LegacyStats() because it breaks assumptions in reporting of cook stats

#rnx

diff --git a/Engine/Source/Developer/DerivedDataCache/Private/DerivedDataCacheStoreAsync.cpp b/Engine/Source/Developer/DerivedDataCache/Private/DerivedDataCacheStoreAsync.cpp
--- a/Engine/Source/Developer/DerivedDataCache/Private/DerivedDataCacheStoreAsync.cpp
+++ b/Engine/Source/Developer/DerivedDataCache/Private/DerivedDataCacheStoreAsync.cpp
@@ -125,7 +125,14 @@ FCacheStoreAsync::FCacheStoreAsync(ILegacyCacheStore* InInnerCache, IMemoryCache
 
 void FCacheStoreAsync::LegacyStats(FDerivedDataCacheStatsNode& OutNode)
 {
-	InnerCache->LegacyStats(OutNode);
+	OutNode = {TEXT("Async"), TEXT(""), /*bIsLocal*/ true};
+	OutNode.UsageStats.Add(TEXT(""), UsageStats);
+
+	InnerCache->LegacyStats(OutNode.Children.Add_GetRef(MakeShared<FDerivedDataCacheStatsNode>()).Get());
+	if (MemoryCache)
+	{
+		MemoryCache->LegacyStats(OutNode.Children.Add_GetRef(MakeShared<FDerivedDataCacheStatsNode>()).Get());
+	}
 }
 
 template <typename RequestType, typename OnCompleteType, typename OnExecuteType>
