DDC: Added ICacheStoreStats and related functionality for cache stores to report stats

#rnx

diff --git a/Engine/Source/Developer/DerivedDataCache/Private/DerivedDataBackends.cpp b/Engine/Source/Developer/DerivedDataCache/Private/DerivedDataBackends.cpp
--- a/Engine/Source/Developer/DerivedDataCache/Private/DerivedDataBackends.cpp
+++ b/Engine/Source/Developer/DerivedDataCache/Private/DerivedDataBackends.cpp
@@ -65,7 +65,7 @@ ILegacyCacheStore* CreateCacheStoreAsync(ILegacyCacheStore* InnerBackend, IMemor
 ILegacyCacheStore* CreateCacheStoreHierarchy(ICacheStoreOwner*& OutOwner, IMemoryCacheStore* MemoryCache);
 ILegacyCacheStore* CreateCacheStoreThrottle(ILegacyCacheStore* InnerCache, uint32 LatencyMS, uint32 MaxBytesPerSecond);
 ILegacyCacheStore* CreateCacheStoreVerify(ILegacyCacheStore* InnerCache, bool bPutOnError);
-ILegacyCacheStore* CreateFileSystemCacheStore(const TCHAR* NodeName, const TCHAR* Config, ECacheStoreFlags& OutFlags, FString& OutPath);
+ILegacyCacheStore* CreateFileSystemCacheStore(const TCHAR* Name, const TCHAR* Config, ICacheStoreOwner& Owner, FString& OutPath);
 TTuple<ILegacyCacheStore*, ECacheStoreFlags> CreateHttpCacheStore(const TCHAR* NodeName, const TCHAR* Config);
 IMemoryCacheStore* CreateMemoryCacheStore(const TCHAR* Name, int64 MaxCacheSize, bool bCanBeDisabled);
 IPakFileCacheStore* CreatePakFileCacheStore(const TCHAR* Filename, bool bWriting, bool bCompressed);
@@ -547,12 +547,10 @@ public:
 	bool ParseDataCache(const TCHAR* NodeName, const TCHAR* Config)
 	{
 		FString Path;
-		ECacheStoreFlags Flags;
-		if (ILegacyCacheStore* Store = CreateFileSystemCacheStore(NodeName, Config, Flags, Path))
+		if (ILegacyCacheStore* Store = CreateFileSystemCacheStore(NodeName, Config, *this, Path))
 		{
 			bUsingSharedDDC |= NodeName == TEXTVIEW("Shared");
 			Directories.AddUnique(Path);
-			Add(Store, Flags);
 			return true;
 		}
 		return false;
@@ -980,6 +978,18 @@ private:
 		}
 	}
 
+	ICacheStoreStats* CreateStats(ILegacyCacheStore* CacheStore, ECacheStoreFlags Flags, FStringView Type, FStringView Name, FStringView Path) final
+	{
+		check(Hierarchy);
+		return Hierarchy->CreateStats(CacheStore, Flags, Type, Name, Path);
+	}
+
+	void DestroyStats(ICacheStoreStats* Stats) final
+	{
+		check(Hierarchy);
+		Hierarchy->DestroyStats(Stats);
+	}
+
 	/** Delete all created backends in the reversed order they were created. */
 	void DestroyCreatedBackends()
 	{
diff --git a/Engine/Source/Developer/DerivedDataCache/Private/DerivedDataCacheStore.cpp b/Engine/Source/Developer/DerivedDataCache/Private/DerivedDataCacheStore.cpp
new file mode 100644
--- /dev/null
+++ b/Engine/Source/Developer/DerivedDataCache/Private/DerivedDataCacheStore.cpp
@@ -0,0 +1,67 @@
+// Copyright Epic Games, Inc. All Rights Reserved.
+
+#include "DerivedDataCacheStore.h"
+
+#include "Async/UniqueLock.h"
+#include "CoreGlobals.h"
+#include "DerivedDataCacheRecord.h"
+#include "DerivedDataValue.h"
+#include "Serialization/CompactBinary.h"
+
+namespace UE::DerivedData
+{
+
+FCacheStoreRequestTimer::FCacheStoreRequestTimer(FCacheStoreRequestStats& OutStats)
+	: Stats(OutStats)
+	, StartTime(FMonotonicTimePoint::Now())
+{
+}
+
+FCacheStoreRequestTimer::~FCacheStoreRequestTimer()
+{
+	const FMonotonicTimePoint EndTime = FMonotonicTimePoint::Now();
+	const bool bIsInGameThread = IsInGameThread();
+	TUniqueLock Lock(Stats.Mutex);
+	(bIsInGameThread ? Stats.MainThreadTime : Stats.OtherThreadTime) = EndTime - StartTime;
+	Stats.StartTime = FMath::Min(Stats.StartTime, StartTime);
+	Stats.EndTime = FMath::Max(Stats.EndTime, EndTime);
+}
+
+void FCacheStoreRequestStats::AddLatency(const FMonotonicTimeSpan InLatency)
+{
+	if (Latency > InLatency)
+	{
+		Latency = InLatency;
+	}
+}
+
+void FCacheStoreRequestStats::AddLogicalRead(const FCacheRecord& Record)
+{
+	if (const FCbObject& Meta = Record.GetMeta())
+	{
+		LogicalReadSize += Meta.GetSize();
+	}
+
+	for (const FValueWithId& Value : Record.GetValues())
+	{
+		AddLogicalRead(Value);
+	}
+}
+
+void FCacheStoreRequestStats::AddLogicalRead(const FValue& Value)
+{
+	if (Value.HasData())
+	{
+		LogicalReadSize += Value.GetRawSize();
+	}
+}
+
+void FCacheStoreRequestStats::AddLogicalWrite(const FValue& Value)
+{
+	if (Value.HasData())
+	{
+		LogicalWriteSize += Value.GetRawSize();
+	}
+}
+
+} // UE::DerivedData
diff --git a/Engine/Source/Developer/DerivedDataCache/Private/DerivedDataCacheStore.h b/Engine/Source/Developer/DerivedDataCache/Private/DerivedDataCacheStore.h
--- a/Engine/Source/Developer/DerivedDataCache/Private/DerivedDataCacheStore.h
+++ b/Engine/Source/Developer/DerivedDataCache/Private/DerivedDataCacheStore.h
@@ -2,9 +2,20 @@
 
 #pragma once
 
+#include "Async/Mutex.h"
 #include "DerivedDataCache.h"
+#include "Misc/MonotonicTime.h"
 
+class FText;
+
+namespace UE::DerivedData { class FCacheRecord; }
+namespace UE::DerivedData { class FCacheStoreRequestTimer; }
+namespace UE::DerivedData { class FValue; }
+namespace UE::DerivedData { class ICacheStoreStats; }
 namespace UE::DerivedData { class ILegacyCacheStore; }
+namespace UE::DerivedData { struct FCacheStoreRequestStats; }
+namespace UE::DerivedData { enum class ECacheStoreRequestOp : uint8; }
+namespace UE::DerivedData { enum class ECacheStoreRequestType : uint8; }
 namespace UE::DerivedData { enum class EStatus : uint8; }
 
 namespace UE::DerivedData
@@ -90,6 +101,32 @@ public:
 		TConstArrayView<FCacheGetChunkRequest> Requests,
 		IRequestOwner& Owner,
 		FOnCacheGetChunkComplete&& OnComplete) = 0;
+
+	using ERequestOp = ECacheStoreRequestOp;
+	using ERequestType = ECacheStoreRequestType;
+	using FRequestStats = FCacheStoreRequestStats;
+	using FRequestTimer = FCacheStoreRequestTimer;
+};
+
+/**
+ * The type of data accessed by a request.
+ */
+enum class ECacheStoreRequestType : uint8
+{
+	None = 0,
+	Record,
+	Value,
+};
+
+/**
+ * The operation performed by the request.
+ */
+enum class ECacheStoreRequestOp : uint8
+{
+	None = 0,
+	Put,
+	Get,
+	GetChunk,
 };
 
 enum class ECacheStoreFlags : uint32
@@ -122,6 +159,90 @@ public:
 	virtual void Add(ILegacyCacheStore* CacheStore, ECacheStoreFlags Flags) = 0;
 	virtual void SetFlags(ILegacyCacheStore* CacheStore, ECacheStoreFlags Flags) = 0;
 	virtual void RemoveNotSafe(ILegacyCacheStore* CacheStore) = 0;
+
+	virtual ICacheStoreStats* CreateStats(ILegacyCacheStore* CacheStore, ECacheStoreFlags Flags, FStringView Type, FStringView Name, FStringView Path = {}) = 0;
+	virtual void DestroyStats(ICacheStoreStats* Stats) = 0;
+};
+
+enum class ECacheStoreStatusCode : uint8
+{
+	None = 0,
+	Warning,
+	Error,
+};
+
+/**
+ * Stats that represent how one cache store processed one request.
+ */
+struct FCacheStoreRequestStats
+{
+	/** Bucket that contains the data being accessed by this request. */
+	FCacheBucket Bucket;
+	/** Size of the raw data and metadata that was read from the cache store. */
+	uint64 LogicalReadSize = 0;
+	/** Size of the raw data and metadata that was written to the cache store. */
+	uint64 LogicalWriteSize = 0;
+	/** Size of the compressed data and anything else that was read from the cache store to satisfy this request. */
+	uint64 PhysicalReadSize = 0;
+	/** Size of the compressed data and anything else that was written to the cache store to satisfy this request. */
+	uint64 PhysicalWriteSize = 0;
+	/** Minimum latency to the cache store of any of the individual accesses needed to satisfy this request. */
+	FMonotonicTimeSpan Latency = FMonotonicTimeSpan::Infinity();
+	/** Approximate CPU time on the main thread that was needed to satisfy this request. */
+	FMonotonicTimeSpan MainThreadTime;
+	/** Approximate CPU time off the main thread that was needed to satisfy this request. */
+	FMonotonicTimeSpan OtherThreadTime;
+	/** Wall time at which this the cache store started processing of this request. */
+	FMonotonicTimePoint StartTime = FMonotonicTimePoint::Infinity();
+	/** Wall time at which this the cache store ended processing of this request. */
+	FMonotonicTimePoint EndTime = FMonotonicTimePoint() - FMonotonicTimeSpan::Infinity();
+	/** The type of data that was accessed by this request. */
+	ECacheStoreRequestType Type = ECacheStoreRequestType::None;
+	/** The operation that was performed by this request. */
+	ECacheStoreRequestOp Op = ECacheStoreRequestOp::None;
+	/** The status with which this request completed. */
+	EStatus Status = EStatus::Ok;
+	/** Mutex that can be used by cache stores that process a request on multiple threads. */
+	mutable FMutex Mutex;
+
+	void AddLatency(FMonotonicTimeSpan Latency);
+
+	/** Record a logical read of the metadata and values in the record. */
+	void AddLogicalRead(const FCacheRecord& Record);
+	/** Record a logical read of the value. Recorded as 0 if the value has no data. */
+	void AddLogicalRead(const FValue& Value);
+	/** Record a logical write of the value. Recorded as 0 if the value has no data. */
+	void AddLogicalWrite(const FValue& Value);
+};
+
+/** Updates StartTime, EndTime, MainThreadTime, OtherThreadTime based on its lifetime. */
+class FCacheStoreRequestTimer
+{
+public:
+	explicit FCacheStoreRequestTimer(FCacheStoreRequestStats& Stats);
+	~FCacheStoreRequestTimer();
+
+	FCacheStoreRequestTimer(const FCacheStoreRequestTimer&) = delete;
+	FCacheStoreRequestTimer& operator=(const FCacheStoreRequestTimer&) = delete;
+
+private:
+	FCacheStoreRequestStats& Stats;
+	FMonotonicTimePoint StartTime;
+};
+
+/**
+ * Interface to report the status of a cache store and stats for each request that it processes.
+ */
+class ICacheStoreStats
+{
+public:
+	virtual ~ICacheStoreStats() = default;
+	virtual FStringView GetType() const = 0;
+	virtual FStringView GetName() const = 0;
+	virtual FStringView GetPath() const = 0;
+	virtual void SetFlags(ECacheStoreFlags Flags) = 0;
+	virtual void SetStatus(ECacheStoreStatusCode StatusCode, const FText& Status) = 0;
+	virtual void AddRequest(const FCacheStoreRequestStats& Stats) = 0;
 };
 
 template <typename RequestRangeType, typename OnCompleteType>
diff --git a/Engine/Source/Developer/DerivedDataCache/Private/DerivedDataCacheStoreHierarchy.cpp b/Engine/Source/Developer/DerivedDataCache/Private/DerivedDataCacheStoreHierarchy.cpp
--- a/Engine/Source/Developer/DerivedDataCache/Private/DerivedDataCacheStoreHierarchy.cpp
+++ b/Engine/Source/Developer/DerivedDataCache/Private/DerivedDataCacheStoreHierarchy.cpp
@@ -1,11 +1,15 @@
 // Copyright Epic Games, Inc. All Rights Reserved.
 
+#include "Algo/Accumulate.h"
 #include "Algo/AllOf.h"
 #include "Algo/AnyOf.h"
 #include "Algo/Compare.h"
 #include "Algo/Find.h"
 #include "Algo/MinElement.h"
+#include "Async/Mutex.h"
+#include "Async/UniqueLock.h"
 #include "Containers/Array.h"
+#include "Containers/StaticArray.h"
 #include "DerivedDataCache.h"
 #include "DerivedDataCacheStore.h"
 #include "DerivedDataCacheUsageStats.h"
@@ -13,9 +17,11 @@
 #include "DerivedDataRequest.h"
 #include "DerivedDataRequestOwner.h"
 #include "HAL/CriticalSection.h"
+#include "Logging/StructuredLog.h"
 #include "MemoryCacheStore.h"
 #include "Misc/EnumClassFlags.h"
 #include "Misc/ScopeRWLock.h"
+#include "ProfilingDebugging/CookStats.h"
 #include "Serialization/CompactBinary.h"
 #include "Templates/Invoke.h"
 #include "Templates/RefCounting.h"
@@ -29,6 +35,123 @@ ILegacyCacheStore* CreateCacheStoreAsync(ILegacyCacheStore* InnerCache, IMemoryC
 
 ///////////////////////////////////////////////////////////////////////////////////////////////////////////////////////
 
+class FCacheStoreStats final : public ICacheStoreStats
+{
+public:
+	FCacheStoreStats(ECacheStoreFlags Flags, FStringView Type, FStringView Name, FStringView Path);
+
+	FString Type;
+	FString Name;
+	FString Path;
+	FText Status;
+	ECacheStoreFlags Flags = ECacheStoreFlags::None;
+	ECacheStoreStatusCode StatusCode = ECacheStoreStatusCode::None;
+	mutable FMutex Mutex;
+
+#if ENABLE_COOK_STATS
+	FCookStats::CallStats GetStats;
+	FCookStats::CallStats PutStats;
+#endif
+
+	struct FRateStat
+	{
+		FMonotonicTimePoint StartTime;
+		FMonotonicTimePoint EndTime;
+		uint64 Amount = 0;
+	};
+
+	constexpr static int32 SpeedStatsMax = 16;
+
+	TStaticArray<FMonotonicTimeSpan, SpeedStatsMax> LatencyHistory;
+	int32 LatencyIndex = 0;
+	int32 LatencyCount = 0;
+
+	TStaticArray<FRateStat, SpeedStatsMax> PhysicalReadSizeHistory;
+	int32 PhysicalReadSizeIndex = 0;
+	int32 PhysicalReadSizeCount = 0;
+
+	TStaticArray<FRateStat, SpeedStatsMax> PhysicalWriteSizeHistory;
+	int32 PhysicalWriteSizeIndex = 0;
+	int32 PhysicalWriteSizeCount = 0;
+
+private:
+	FStringView GetType() const final { return Type; }
+	FStringView GetName() const final { return Name; }
+	FStringView GetPath() const final { return Path; }
+	void SetFlags(ECacheStoreFlags Flags) final;
+	void SetStatus(ECacheStoreStatusCode StatusCode, const FText& Status) final;
+	void AddRequest(const FCacheStoreRequestStats& Stats) final;
+};
+
+FCacheStoreStats::FCacheStoreStats(ECacheStoreFlags InFlags, FStringView InType, FStringView InName, FStringView InPath)
+	: Type(InType)
+	, Name(InName)
+	, Path(InPath)
+	, Flags(InFlags)
+{
+}
+
+void FCacheStoreStats::SetFlags(ECacheStoreFlags InFlags)
+{
+	TUniqueLock Lock(Mutex);
+	Flags = InFlags;
+}
+
+void FCacheStoreStats::SetStatus(ECacheStoreStatusCode InStatusCode, const FText& InStatus)
+{
+	TUniqueLock Lock(Mutex);
+	Status = InStatus;
+	StatusCode = InStatusCode;
+}
+
+void FCacheStoreStats::AddRequest(const FCacheStoreRequestStats& Stats)
+{
+#if ENABLE_COOK_STATS
+	using FCallStats = FCookStats::CallStats;
+	using EHitOrMiss = FCallStats::EHitOrMiss;
+	using EStatType = FCallStats::EStatType;
+	const bool bIsInGameThread = IsInGameThread();
+	const bool bIsGet = Stats.Op != ECacheStoreRequestOp::Put;
+#endif
+
+	TUniqueLock Lock(Mutex);
+
+	if (!Stats.Latency.IsInfinity())
+	{
+		LatencyHistory[LatencyIndex++] = Stats.Latency;
+		LatencyCount = FMath::Max(LatencyCount, LatencyIndex);
+		LatencyIndex %= LatencyHistory.Num();
+	}
+
+	if (Stats.PhysicalReadSize)
+	{
+		FRateStat& CurrentPhysicalReadSize = PhysicalReadSizeHistory[PhysicalReadSizeIndex++];
+		PhysicalReadSizeCount = FMath::Max(PhysicalReadSizeCount, PhysicalReadSizeIndex);
+		PhysicalReadSizeIndex %= PhysicalReadSizeHistory.Num();
+		CurrentPhysicalReadSize.StartTime = Stats.StartTime;
+		CurrentPhysicalReadSize.EndTime = Stats.EndTime;
+		CurrentPhysicalReadSize.Amount = Stats.PhysicalReadSize;
+	}
+
+	if (Stats.PhysicalWriteSize)
+	{
+		FRateStat& CurrentPhysicalWriteSize = PhysicalWriteSizeHistory[PhysicalWriteSizeIndex++];
+		PhysicalWriteSizeCount = FMath::Max(PhysicalWriteSizeCount, PhysicalWriteSizeIndex);
+		PhysicalWriteSizeIndex %= PhysicalWriteSizeHistory.Num();
+		CurrentPhysicalWriteSize.StartTime = Stats.StartTime;
+		CurrentPhysicalWriteSize.EndTime = Stats.EndTime;
+		CurrentPhysicalWriteSize.Amount = Stats.PhysicalWriteSize;
+	}
+
+#if ENABLE_COOK_STATS
+	FCookStats::CallStats& CallStats = bIsGet ? GetStats : PutStats;
+	CallStats.Accumulate(Stats.Status == EStatus::Ok ? EHitOrMiss::Hit : EHitOrMiss::Miss, EStatType::Counter, 1, bIsInGameThread);
+	CallStats.Accumulate(Stats.Status == EStatus::Ok ? EHitOrMiss::Hit : EHitOrMiss::Miss, EStatType::Cycles, int64(Stats.MainThreadTime.ToSeconds() / FPlatformTime::GetSecondsPerCycle64()), /*bIsInGameThread*/ true);
+	CallStats.Accumulate(Stats.Status == EStatus::Ok ? EHitOrMiss::Hit : EHitOrMiss::Miss, EStatType::Cycles, int64(Stats.OtherThreadTime.ToSeconds() / FPlatformTime::GetSecondsPerCycle64()), /*bIsInGameThread*/ false);
+	CallStats.Accumulate(Stats.Status == EStatus::Ok ? EHitOrMiss::Hit : EHitOrMiss::Miss, EStatType::Bytes, bIsGet ? Stats.PhysicalReadSize : Stats.PhysicalWriteSize, bIsInGameThread);
+#endif
+}
+
 class FCacheStoreHierarchy final : public ILegacyCacheStore, public ICacheStoreOwner
 {
 public:
@@ -36,11 +159,12 @@ public:
 	~FCacheStoreHierarchy() final = default;
 
 	void Add(ILegacyCacheStore* CacheStore, ECacheStoreFlags Flags) final;
-
 	void SetFlags(ILegacyCacheStore* CacheStore, ECacheStoreFlags Flags) final;
-
 	void RemoveNotSafe(ILegacyCacheStore* CacheStore) final;
 
+	ICacheStoreStats* CreateStats(ILegacyCacheStore* CacheStore, ECacheStoreFlags Flags, FStringView Type, FStringView Name, FStringView Path) final;
+	void DestroyStats(ICacheStoreStats* Stats) final;
+
 	void Put(
 		TConstArrayView<FCachePutRequest> Requests,
 		IRequestOwner& Owner,
@@ -109,6 +233,7 @@ private:
 		ECacheStoreFlags CacheFlags{};
 		ECacheStoreNodeFlags NodeFlags{};
 		TUniquePtr<ILegacyCacheStore> AsyncCache;
+		TArray<FCacheStoreStats*> CacheStats;
 	};
 
 	mutable FRWLock NodesLock;
@@ -305,6 +430,40 @@ void FCacheStoreHierarchy::UpdateNodeFlags()
 	CombinedNodeFlags = StoreFlags | QueryFlags;
 }
 
+ICacheStoreStats* FCacheStoreHierarchy::CreateStats(ILegacyCacheStore* CacheStore, ECacheStoreFlags Flags, FStringView Type, FStringView Name, FStringView Path)
+{
+	FWriteScopeLock Lock(NodesLock);
+	for (FCacheStoreNode& Node : Nodes)
+	{
+		if (Node.Cache == CacheStore)
+		{
+			Node.CacheStats.Emplace(new FCacheStoreStats(Flags, Type, Name, Path));
+			return Node.CacheStats.Last();
+		}
+	}
+	for (;;)
+	{
+		UE_LOGFMT(LogDerivedDataCache, Fatal,
+			"Failed to find {Type} cache store '{Name}' with path '{Path}' when creating stats.", Type, Name, Path);
+	}
+}
+
+void FCacheStoreHierarchy::DestroyStats(ICacheStoreStats* Stats)
+{
+	FWriteScopeLock Lock(NodesLock);
+	for (FCacheStoreNode& Node : Nodes)
+	{
+		if (Node.CacheStats.RemoveSingle((FCacheStoreStats*)Stats))
+		{
+			delete Stats;
+			return;
+		}
+	}
+	UE_LOGFMT(LogDerivedDataCache, Fatal,
+		"Failed to find {Type} cache store '{Name}' with path '{Path}' when destroying stats.",
+		Stats->GetType(), Stats->GetName(), Stats->GetPath());
+}
+
 FCacheRecordPolicy FCacheStoreHierarchy::AddPolicy(const FCacheRecordPolicy& BasePolicy, ECachePolicy Policy)
 {
 	return BasePolicy.Transform([Policy](ECachePolicy P) { return P | Policy; });
@@ -1239,13 +1398,109 @@ void FCacheStoreHierarchy::GetChunks(
 
 ///////////////////////////////////////////////////////////////////////////////////////////////////////////////////////
 
+static void ConvertToLegacyStats(FDerivedDataCacheStatsNode& OutNode, const FCacheStoreStats& Stats)
+{
+	EDerivedDataCacheStatus StatusCode;
+	switch (Stats.StatusCode)
+	{
+	default:
+		checkNoEntry();
+		[[fallthrough]];
+	case ECacheStoreStatusCode::None:
+		StatusCode = EDerivedDataCacheStatus::None;
+		break;
+	case ECacheStoreStatusCode::Warning:
+		StatusCode = EDerivedDataCacheStatus::Warning;
+		break;
+	case ECacheStoreStatusCode::Error:
+		StatusCode = EDerivedDataCacheStatus::Error;
+		break;
+	}
+
+	TUniqueLock StatsLock(Stats.Mutex);
+	OutNode = FDerivedDataCacheStatsNode(Stats.Type, Stats.Path, EnumHasAnyFlags(Stats.Flags, ECacheStoreFlags::Local), StatusCode, *Stats.Status.ToString());
+
+#if ENABLE_COOK_STATS
+	FDerivedDataCacheUsageStats& UsageStats = OutNode.UsageStats.FindOrAdd({});
+	UsageStats.GetStats = Stats.GetStats;
+	UsageStats.PutStats = Stats.PutStats;
+#endif
+
+	const auto MeasureActiveWallTime = [](TConstArrayView<FCacheStoreStats::FRateStat> History) -> FMonotonicTimeSpan
+	{
+		TSortedMap<FMonotonicTimePoint, FMonotonicTimePoint, TInlineAllocator<FCacheStoreStats::SpeedStatsMax>> WallTimes;
+		for (const FCacheStoreStats::FRateStat& RateStat : History)
+		{
+			WallTimes.Add(RateStat.StartTime, RateStat.EndTime);
+		}
+
+		FMonotonicTimeSpan Duration;
+		FMonotonicTimePoint ActiveStartTime = FMonotonicTimePoint::Infinity();
+		FMonotonicTimePoint ActiveEndTime;
+		for (const TTuple<FMonotonicTimePoint, FMonotonicTimePoint>& WallTime : WallTimes)
+		{
+			if (ActiveStartTime.IsInfinity())
+			{
+				ActiveStartTime = WallTime.Key;
+				ActiveEndTime = WallTime.Value;
+			}
+			if (ActiveEndTime < WallTime.Key)
+			{
+				Duration += ActiveEndTime - ActiveStartTime;
+				ActiveStartTime = WallTime.Key;
+			}
+			if (ActiveEndTime < WallTime.Value)
+			{
+				ActiveEndTime = WallTime.Value;
+			}
+		}
+		Duration += ActiveEndTime - ActiveStartTime;
+		return Duration;
+	};
+
+	FDerivedDataCacheSpeedStats& SpeedStats = OutNode.SpeedStats;
+
+	if (TConstArrayView<FMonotonicTimeSpan> LatencyHistory(&Stats.LatencyHistory[0], Stats.LatencyCount); !LatencyHistory.IsEmpty())
+	{
+		SpeedStats.LatencyMS = Algo::Accumulate(LatencyHistory, FMonotonicTimeSpan::Zero()).ToMilliseconds() / LatencyHistory.Num();
+	}
+
+	if (TConstArrayView<FCacheStoreStats::FRateStat> PhysicalReadSizeHistory(&Stats.PhysicalReadSizeHistory[0], Stats.PhysicalReadSizeCount); !PhysicalReadSizeHistory.IsEmpty())
+	{
+		const uint64 PhysicalReadSize = Algo::TransformAccumulate(PhysicalReadSizeHistory, &FCacheStoreStats::FRateStat::Amount, uint64(0));
+		if (const FMonotonicTimeSpan WallTime = MeasureActiveWallTime(PhysicalReadSizeHistory); !WallTime.IsZero())
+		{
+			SpeedStats.ReadSpeedMBs = double(PhysicalReadSize) / WallTime.ToSeconds() / 1024.0 / 1024.0;
+		}
+	}
+
+	if (TConstArrayView<FCacheStoreStats::FRateStat> PhysicalWriteSizeHistory(&Stats.PhysicalWriteSizeHistory[0], Stats.PhysicalWriteSizeCount); !PhysicalWriteSizeHistory.IsEmpty())
+	{
+		const uint64 PhysicalWriteSize = Algo::TransformAccumulate(PhysicalWriteSizeHistory, &FCacheStoreStats::FRateStat::Amount, uint64(0));
+		if (const FMonotonicTimeSpan WallTime = MeasureActiveWallTime(PhysicalWriteSizeHistory); !WallTime.IsZero())
+		{
+			SpeedStats.WriteSpeedMBs = double(PhysicalWriteSize) / WallTime.ToSeconds() / 1024.0 / 1024.0;
+		}
+	}
+}
+
 void FCacheStoreHierarchy::LegacyStats(FDerivedDataCacheStatsNode& OutNode)
 {
 	FReadScopeLock Lock(NodesLock);
 	OutNode.Children.Reserve(Nodes.Num());
 	for (const FCacheStoreNode& Node : Nodes)
 	{
-		Node.Cache->LegacyStats(OutNode.Children.Add_GetRef(MakeShared<FDerivedDataCacheStatsNode>()).Get());
+		if (Node.CacheStats.IsEmpty())
+		{
+			Node.Cache->LegacyStats(OutNode.Children.Add_GetRef(MakeShared<FDerivedDataCacheStatsNode>()).Get());
+		}
+		else
+		{
+			for (const FCacheStoreStats* Stats : Node.CacheStats)
+			{
+				ConvertToLegacyStats(OutNode.Children.Add_GetRef(MakeShared<FDerivedDataCacheStatsNode>()).Get(), *Stats);
+			}
+		}
 	}
 }
 
diff --git a/Engine/Source/Developer/DerivedDataCache/Private/FileSystemCacheStore.cpp b/Engine/Source/Developer/DerivedDataCache/Private/FileSystemCacheStore.cpp
--- a/Engine/Source/Developer/DerivedDataCache/Private/FileSystemCacheStore.cpp
+++ b/Engine/Source/Developer/DerivedDataCache/Private/FileSystemCacheStore.cpp
@@ -772,6 +772,9 @@ void FAccessLogWriter::AppendPath(const FStringView Path)
 
 struct FFileSystemCacheStoreParams
 {
+	/** Name of the cache, e.g., Local, Shared. */
+	FString CacheName;
+
 	/** Root path under which cache files are stored. */
 	FString CachePath;
 
@@ -816,7 +819,7 @@ struct FFileSystemCacheStoreParams
 	/** If true, always update file timestamps on access, even when read-only. */
 	bool bTouch = false;
 
-	void Parse(const TCHAR* NodeName, const TCHAR* Config);
+	void Parse(const TCHAR* Name, const TCHAR* Config);
 };
 
 ///////////////////////////////////////////////////////////////////////////////////////////////////////////////////////
@@ -825,16 +828,15 @@ class FFileSystemCacheStore final : public ILegacyCacheStore
 {
 public:
 	static ILegacyCacheStore* TryCreate(
-		const TCHAR* NodeName,
 		const FFileSystemCacheStoreParams& Params,
-		ECacheStoreFlags& OutFlags,
+		ICacheStoreOwner& Owner,
 		FString& OutPath);
 
 private:
 	FFileSystemCacheStore(
+		ICacheStoreOwner& Owner,
 		const FFileSystemCacheStoreParams& Params,
-	const FDerivedDataCacheSpeedStats& SpeedStats,
-		ECacheStoreFlags& OutFlags);
+		const FDerivedDataCacheSpeedStats& SpeedStats);
 	~FFileSystemCacheStore();
 
 	static bool RunSpeedTest(
@@ -876,33 +878,36 @@ private:
 	bool LegacyDebugOptions(FBackendDebugOptions& Options) final;
 
 private:
-	[[nodiscard]] bool PutCacheRecord(FStringView Name, const FCacheRecord& Record, const FCacheRecordPolicy& Policy, uint64& OutWriteSize);
+	[[nodiscard]] bool PutCacheRecord(FStringView Name, const FCacheRecord& Record, const FCacheRecordPolicy& Policy, FRequestStats& Stats);
 
 	[[nodiscard]] FOptionalCacheRecord GetCacheRecordOnly(
 		FStringView Name,
 		const FCacheKey& Key,
-		const FCacheRecordPolicy& Policy);
+		const FCacheRecordPolicy& Policy,
+		FRequestStats& Stats);
 	[[nodiscard]] FOptionalCacheRecord GetCacheRecord(
 		FStringView Name,
 		const FCacheKey& Key,
 		const FCacheRecordPolicy& Policy,
-		EStatus& OutStatus);
+		EStatus& OutStatus,
+		FRequestStats& Stats);
 
-	[[nodiscard]] bool PutCacheValue(FStringView Name, const FCacheKey& Key, const FValue& Value, ECachePolicy Policy, uint64& OutWriteSize);
+	[[nodiscard]] bool PutCacheValue(FStringView Name, const FCacheKey& Key, const FValue& Value, ECachePolicy Policy, FRequestStats& Stats);
 
-	[[nodiscard]] bool GetCacheValueOnly(FStringView Name, const FCacheKey& Key, ECachePolicy Policy, FValue& OutValue);
-	[[nodiscard]] bool GetCacheValue(FStringView Name, const FCacheKey& Key, ECachePolicy Policy, FValue& OutValue);
+	[[nodiscard]] bool GetCacheValueOnly(FStringView Name, const FCacheKey& Key, ECachePolicy Policy, FValue& OutValue, FRequestStats& Stats);
+	[[nodiscard]] bool GetCacheValue(FStringView Name, const FCacheKey& Key, ECachePolicy Policy, FValue& OutValue, FRequestStats& Stats);
 
-	[[nodiscard]] bool PutCacheContent(FStringView Name, const FCompressedBuffer& Content, uint64& OutWriteSize) const;
+	[[nodiscard]] bool PutCacheContent(FStringView Name, const FCompressedBuffer& Content, FRequestStats& Stats) const;
 
-	[[nodiscard]] bool GetCacheContentExists(const FCacheKey& Key, const FIoHash& RawHash) const;
+	[[nodiscard]] bool GetCacheContentExists(const FCacheKey& Key, const FIoHash& RawHash, FRequestStats& Stats) const;
 	[[nodiscard]] bool GetCacheContent(
 		FStringView Name,
 		const FCacheKey& Key,
 		const FValueId& Id,
 		const FValue& Value,
 		ECachePolicy Policy,
-		FValue& OutValue) const;
+		FValue& OutValue,
+		FRequestStats& Stats) const;
 	void GetCacheContent(
 		FStringView Name,
 		const FCacheKey& Key,
@@ -910,26 +915,31 @@ private:
 		const FValue& Value,
 		ECachePolicy Policy,
 		FCompressedBufferReader& Reader,
-		TUniquePtr<FArchive>& OutArchive) const;
+		TUniquePtr<FArchive>& OutArchive,
+		FRequestStats& Stats) const;
 
 	void BuildCachePackagePath(const FCacheKey& CacheKey, FStringBuilderBase& Path) const;
 	void BuildCacheContentPath(const FIoHash& RawHash, FStringBuilderBase& Path) const;
 
-	[[nodiscard]] bool SaveFileWithHash(FStringBuilderBase& Path, FStringView DebugName, TFunctionRef<void (FArchive&)> WriteFunction, bool bReplaceExisting = false) const;
-	[[nodiscard]] bool LoadFileWithHash(FStringBuilderBase& Path, FStringView DebugName, TFunctionRef<void (FArchive&)> ReadFunction) const;
-	[[nodiscard]] bool SaveFile(FStringBuilderBase& Path, FStringView DebugName, TFunctionRef<void (FArchive&)> WriteFunction, bool bReplaceExisting = false) const;
-	[[nodiscard]] bool LoadFile(FStringBuilderBase& Path, FStringView DebugName, TFunctionRef<void (FArchive&)> ReadFunction) const;
-	[[nodiscard]] TUniquePtr<FArchive> OpenFileWrite(FStringBuilderBase& Path, FStringView DebugName) const;
-	[[nodiscard]] TUniquePtr<FArchive> OpenFileRead(FStringBuilderBase& Path, FStringView DebugName) const;
+	[[nodiscard]] bool SaveFileWithHash(FStringBuilderBase& Path, FStringView DebugName, FRequestStats& Stats, TFunctionRef<void (FArchive&)> WriteFunction, bool bReplaceExisting = false) const;
+	[[nodiscard]] bool LoadFileWithHash(FStringBuilderBase& Path, FStringView DebugName, FRequestStats& Stats, TFunctionRef<void (FArchive&)> ReadFunction) const;
+	[[nodiscard]] bool SaveFile(FStringBuilderBase& Path, FStringView DebugName, FRequestStats& Stats, TFunctionRef<void (FArchive&)> WriteFunction, bool bReplaceExisting = false) const;
+	[[nodiscard]] bool LoadFile(FStringBuilderBase& Path, FStringView DebugName, FRequestStats& Stats, TFunctionRef<void (FArchive&)> ReadFunction) const;
+	[[nodiscard]] TUniquePtr<FArchive> OpenFileWrite(FStringBuilderBase& Path, FStringView DebugName, FRequestStats& Stats) const;
+	[[nodiscard]] TUniquePtr<FArchive> OpenFileRead(FStringBuilderBase& Path, FStringView DebugName, FRequestStats& Stats) const;
+
+	[[nodiscard]] bool FileExists(FStringBuilderBase& Path, FRequestStats& Stats) const;
 
-	[[nodiscard]] bool FileExists(FStringBuilderBase& Path) const;
 	[[nodiscard]] bool IsDeactivatedForPerformance();
+	void UpdateStatus();
 
 	static bool RunInitialSpeedTest(const FFileSystemCacheStoreParams& Params, FDerivedDataCacheSpeedStats& OutSpeedStats);
 
 private:
-	/** Base path to store the cache files in. */
 	FString CachePath;
+	ICacheStoreOwner& StoreOwner;
+	ICacheStoreStats* StoreStats = nullptr;
+
 	/** Speed class of this cache. */
 	EBackendSpeedClass SpeedClass = EBackendSpeedClass::Unknown;
 	/** If true, do not attempt to write to this cache. */
@@ -951,8 +961,6 @@ private:
 	/** Debug Options */
 	FBackendDebugOptions DebugOptions;
 
-	FDerivedDataCacheUsageStats UsageStats;
-
 	/** Speed stats */
 	FDerivedDataCacheSpeedStats SpeedStats;
 
@@ -978,25 +986,22 @@ private:
 };
 
 ILegacyCacheStore* FFileSystemCacheStore::TryCreate(
-	const TCHAR* const NodeName,
 	const FFileSystemCacheStoreParams& Params,
-	ECacheStoreFlags& OutFlags,
+	ICacheStoreOwner& Owner,
 	FString& OutPath)
 {
-	OutFlags = ECacheStoreFlags::None;
-
 	// If we find a platform that has more stringent limits, this needs to be rethought.
 	checkf(GMaxCacheRootLen + GMaxCacheKeyLen <= FPlatformMisc::GetMaxPathLength(),
 		TEXT("Not enough room left for cache keys in max path."));
 
 	if (Params.CachePath.IsEmpty())
 	{
-		UE_LOG(LogDerivedDataCache, Log, TEXT("%s: Disabled because no path is configured."), NodeName);
+		UE_LOG(LogDerivedDataCache, Log, TEXT("%s: Disabled because no path is configured."), *Params.CacheName);
 		return nullptr;
 	}
 	else if (Params.CachePath == TEXTVIEW("None"))
 	{
-		UE_LOG(LogDerivedDataCache, Log, TEXT("%s: Disabled because the path is configured to 'None'"), NodeName);
+		UE_LOG(LogDerivedDataCache, Log, TEXT("%s: Disabled because the path is configured to 'None'"), *Params.CacheName);
 		return nullptr;
 	}
 
@@ -1012,12 +1017,12 @@ ILegacyCacheStore* FFileSystemCacheStore::TryCreate(
 
 		if (TUniqueLock Lock(ActiveStoresMutex); Algo::FindByPredicate(ActiveStores, HasSameCachePath))
 		{
-			UE_LOG(LogDerivedDataCache, Warning, TEXT("%s: Process has an existing cache store at path %s, and the duplicate is being ignored."), NodeName, *Params.CachePath);
+			UE_LOG(LogDerivedDataCache, Warning, TEXT("%s: Process has an existing cache store at path %s, and the duplicate is being ignored."), *Params.CacheName, *Params.CachePath);
 			return nullptr;
 		}
 
 		// Skip creation of the cache store if the shared cache directory does not exist.
-		bool bShared = NodeName == TEXTVIEW("Shared");
+		bool bShared = Params.CacheName == TEXTVIEW("Shared");
 		if (!bShared || IFileManager::Get().DirectoryExists(*Params.CachePath))
 		{
 			FDerivedDataCacheSpeedStats LocalSpeedStats;
@@ -1026,42 +1031,40 @@ ILegacyCacheStore* FFileSystemCacheStore::TryCreate(
 			LocalSpeedStats.LatencyMS = 0;
 			if (Params.bSkipSpeedTest || RunInitialSpeedTest(Params, LocalSpeedStats))
 			{
-				UE_CLOG(Params.bSkipSpeedTest, LogDerivedDataCache, Log, TEXT("%s: Skipping speed test at path %s and assuming local performance."), NodeName, *Params.CachePath);
+				UE_CLOG(Params.bSkipSpeedTest, LogDerivedDataCache, Log, TEXT("%s: Skipping speed test at path %s and assuming local performance."), *Params.CacheName, *Params.CachePath);
 
-				FFileSystemCacheStore* CacheStore = new FFileSystemCacheStore(Params, LocalSpeedStats, OutFlags);
+				FFileSystemCacheStore* CacheStore = new FFileSystemCacheStore(Owner, Params, LocalSpeedStats);
 				{
 					TUniqueLock Lock(ActiveStoresMutex);
 					ActiveStores.Add(CacheStore);
 				}
-				UE_LOG(LogDerivedDataCache, Log, TEXT("%s: Using data cache path %s: %s"), NodeName, *Params.CachePath,
-					EnumHasAnyFlags(OutFlags, ECacheStoreFlags::Store) ? TEXT("Writable") :
-					EnumHasAnyFlags(OutFlags, ECacheStoreFlags::Query) ? TEXT("ReadOnly") : TEXT("DeleteOnly"));
 				OutPath = CacheStore->CachePath;
 				return CacheStore;
 			}
 
-			UE_LOG(LogDerivedDataCache, Warning, TEXT("%s: No read or write access to %s, or the speed test was abandoned due to slow progress."), NodeName, *Params.CachePath);
+			UE_LOG(LogDerivedDataCache, Warning, TEXT("%s: No read or write access to %s, or the speed test was abandoned due to slow progress."), *Params.CacheName, *Params.CachePath);
 		}
 
 		// Give the user a chance to retry in case they need to take steps like connecting a network drive.
 		if (Params.bPromptIfMissing && !FApp::IsUnattended() && !IS_PROGRAM)
 		{
-			TStringBuilder<512> Message(InPlace, NodeName, TEXTVIEW(" cache store path "), Params.CachePath,
+			TStringBuilder<512> Message(InPlace, Params.CacheName, TEXTVIEW(" cache store path "), Params.CachePath,
 				TEXTVIEW(" is not available and this cache store will be disabled.\n\nRetry connection to "), Params.CachePath, TEXTVIEW("?"));
 			bRetryOnFailure = FPlatformMisc::MessageBoxExt(EAppMsgType::YesNo, *Message, TEXT("Failed to access Derived Data Cache")) == EAppReturnType::Yes;
 		}
 	}
 	while (bRetryOnFailure);
 
-	UE_LOG(LogDerivedDataCache, Warning, TEXT("%s: Path %s is not available and this cache store will be disabled."), NodeName, *Params.CachePath);
+	UE_LOG(LogDerivedDataCache, Warning, TEXT("%s: Path %s is not available and this cache store will be disabled."), *Params.CacheName, *Params.CachePath);
 	return nullptr;
 }
 
 FFileSystemCacheStore::FFileSystemCacheStore(
+	ICacheStoreOwner& Owner,
 	const FFileSystemCacheStoreParams& Params,
-	const FDerivedDataCacheSpeedStats& InSpeedStats,
-	ECacheStoreFlags& OutFlags)
+	const FDerivedDataCacheSpeedStats& InSpeedStats)
 	: CachePath(Params.CachePath)
+	, StoreOwner(Owner)
 	, bReadOnly(Params.bReadOnly)
 	, bTouch(Params.bTouch)
 	, MaxFileAgeInDays(Params.MaxFileAgeInDays)
@@ -1214,11 +1217,21 @@ FFileSystemCacheStore::FFileSystemCacheStore(
 	Flags |= Params.bDeleteOnly ? ECacheStoreFlags::None : ECacheStoreFlags::Query;
 	Flags |= Params.bDeleteOnly || bReadOnly ? ECacheStoreFlags::None : ECacheStoreFlags::Store;
 	Flags |= SpeedClass == EBackendSpeedClass::Local ? ECacheStoreFlags::Local : ECacheStoreFlags::Remote;
-	OutFlags = Flags;
+
+	StoreOwner.Add(this, Flags);
+	StoreStats = StoreOwner.CreateStats(this, Flags, TEXTVIEW("File System"), Params.CacheName, CachePath);
+
+	UpdateStatus();
+
+	UE_LOG(LogDerivedDataCache, Log, TEXT("%s: Using data cache path %s: %s"), *Params.CacheName, *CachePath,
+		EnumHasAnyFlags(Flags, ECacheStoreFlags::Store) ? TEXT("Writable") :
+		EnumHasAnyFlags(Flags, ECacheStoreFlags::Query) ? TEXT("ReadOnly") : TEXT("DeleteOnly"));
 }
 
 FFileSystemCacheStore::~FFileSystemCacheStore()
 {
+	StoreOwner.DestroyStats(StoreStats);
+
 	TUniqueLock Lock(ActiveStoresMutex);
 	ActiveStores.Remove(this);
 }
@@ -1466,17 +1479,7 @@ bool FFileSystemCacheStore::RunSpeedTest(
 
 void FFileSystemCacheStore::LegacyStats(FDerivedDataCacheStatsNode& OutNode)
 {
-	EDerivedDataCacheStatus CacheStatus = EDerivedDataCacheStatus::None;
-	FText CacheStatusText;
-	if (IsDeactivatedForPerformance())
-	{
-		CacheStatus = EDerivedDataCacheStatus::Deactivation;
-		CacheStatusText = NSLOCTEXT("DerivedDataCache", "DeactivatedForPerformance", "Deactivated for performance");
-	}
-	const FString& CacheStatusTextString = CacheStatusText.ToString();
-	OutNode = {TEXT("File System"), *CachePath, SpeedClass == EBackendSpeedClass::Local, CacheStatus, CacheStatusTextString.IsEmpty() ? nullptr : *CacheStatusTextString};
-	OutNode.UsageStats.Add(TEXT(""), UsageStats);
-	OutNode.SpeedStats = SpeedStats;
+	checkNoEntry();
 }
 
 bool FFileSystemCacheStore::LegacyDebugOptions(FBackendDebugOptions& InOptions)
@@ -1493,25 +1496,28 @@ void FFileSystemCacheStore::Put(
 	for (const FCachePutRequest& Request : Requests)
 	{
 		bool bOk;
+		FRequestStats RequestStats;
+		RequestStats.Bucket = Request.Record.GetKey().Bucket;
+		RequestStats.Type = ERequestType::Record;
+		RequestStats.Op = ERequestOp::Put;
 		{
 			const FCacheRecord& Record = Request.Record;
 			TRACE_CPUPROFILER_EVENT_SCOPE(FileSystemDDC_Put);
 			TRACE_COUNTER_INCREMENT(FileSystemDDC_Put);
-			COOK_STAT(auto Timer = UsageStats.TimePut());
+			FRequestTimer RequestTimer(RequestStats);
 			uint64 WriteSize = 0;
-			bOk = PutCacheRecord(Request.Name, Record, Request.Policy, WriteSize);
+			bOk = PutCacheRecord(Request.Name, Record, Request.Policy, RequestStats);
 			if (bOk)
 			{
 				UE_LOG(LogDerivedDataCache, Verbose, TEXT("%s: Cache put complete for %s from '%s'"),
 					*CachePath, *WriteToString<96>(Record.GetKey()), *Request.Name);
 				TRACE_COUNTER_INCREMENT(FileSystemDDC_PutHit);
-				TRACE_COUNTER_ADD(FileSystemDDC_BytesWritten, WriteSize);
-				if (WriteSize)
-				{
-					COOK_STAT(Timer.AddHit(WriteSize));
-				}
+				TRACE_COUNTER_ADD(FileSystemDDC_BytesRead, RequestStats.PhysicalReadSize);
+				TRACE_COUNTER_ADD(FileSystemDDC_BytesWritten, RequestStats.PhysicalWriteSize);
 			}
 		}
+		RequestStats.Status = bOk ? EStatus::Ok : EStatus::Error;
+		StoreStats->AddRequest(RequestStats);
 		OnComplete(Request.MakeResponse(bOk ? EStatus::Ok : EStatus::Error));
 	}
 }
@@ -1525,23 +1531,30 @@ void FFileSystemCacheStore::Get(
 	{
 		EStatus Status = EStatus::Error;
 		FOptionalCacheRecord Record;
+		FRequestStats RequestStats;
+		RequestStats.Bucket = Request.Key.Bucket;
+		RequestStats.Type = ERequestType::Record;
+		RequestStats.Op = ERequestOp::Get;
 		{
 			TRACE_CPUPROFILER_EVENT_SCOPE(FileSystemDDC_Get);
 			TRACE_COUNTER_INCREMENT(FileSystemDDC_Get);
-			COOK_STAT(auto Timer = UsageStats.TimeGet());
-			if ((Record = GetCacheRecord(Request.Name, Request.Key, Request.Policy, Status)))
+			FRequestTimer RequestTimer(RequestStats);
+			if ((Record = GetCacheRecord(Request.Name, Request.Key, Request.Policy, Status, RequestStats)))
 			{
 				UE_LOG(LogDerivedDataCache, Verbose, TEXT("%s: Cache hit for %s from '%s'"),
 					*CachePath, *WriteToString<96>(Request.Key), *Request.Name);
 				TRACE_COUNTER_INCREMENT(FileSystemDDC_GetHit);
-				TRACE_COUNTER_ADD(FileSystemDDC_BytesRead, Private::GetCacheRecordCompressedSize(Record.Get()));
-				COOK_STAT(Timer.AddHit(Private::GetCacheRecordCompressedSize(Record.Get())));
+				TRACE_COUNTER_ADD(FileSystemDDC_BytesRead, RequestStats.PhysicalReadSize);
+				TRACE_COUNTER_ADD(FileSystemDDC_BytesWritten, RequestStats.PhysicalWriteSize);
 			}
 			else
 			{
 				Record = FCacheRecordBuilder(Request.Key).Build();
 			}
 		}
+		RequestStats.AddLogicalRead(Record.Get());
+		RequestStats.Status = Status;
+		StoreStats->AddRequest(RequestStats);
 		OnComplete({Request.Name, MoveTemp(Record).Get(), Request.UserData, Status});
 	}
 }
@@ -1554,24 +1567,27 @@ void FFileSystemCacheStore::PutValue(
 	for (const FCachePutValueRequest& Request : Requests)
 	{
 		bool bOk;
+		FRequestStats RequestStats;
+		RequestStats.Bucket = Request.Key.Bucket;
+		RequestStats.Type = ERequestType::Value;
+		RequestStats.Op = ERequestOp::Put;
 		{
 			TRACE_CPUPROFILER_EVENT_SCOPE(FileSystemDDC_PutValue);
 			TRACE_COUNTER_INCREMENT(FileSystemDDC_Put);
-			COOK_STAT(auto Timer = UsageStats.TimePut());
+			FRequestTimer RequestTimer(RequestStats);
 			uint64 WriteSize = 0;
-			bOk = PutCacheValue(Request.Name, Request.Key, Request.Value, Request.Policy, WriteSize);
+			bOk = PutCacheValue(Request.Name, Request.Key, Request.Value, Request.Policy, RequestStats);
 			if (bOk)
 			{
 				UE_LOG(LogDerivedDataCache, Verbose, TEXT("%s: Cache put complete for %s from '%s'"),
 					*CachePath, *WriteToString<96>(Request.Key), *Request.Name);
 				TRACE_COUNTER_INCREMENT(FileSystemDDC_PutHit);
-				TRACE_COUNTER_ADD(FileSystemDDC_BytesWritten, WriteSize);
-				if (WriteSize)
-				{
-					COOK_STAT(Timer.AddHit(WriteSize));
-				}
+				TRACE_COUNTER_ADD(FileSystemDDC_BytesRead, RequestStats.PhysicalReadSize);
+				TRACE_COUNTER_ADD(FileSystemDDC_BytesWritten, RequestStats.PhysicalWriteSize);
 			}
 		}
+		RequestStats.Status = bOk ? EStatus::Ok : EStatus::Error;
+		StoreStats->AddRequest(RequestStats);
 		OnComplete(Request.MakeResponse(bOk ? EStatus::Ok : EStatus::Error));
 	}
 }
@@ -1585,20 +1601,27 @@ void FFileSystemCacheStore::GetValue(
 	{
 		bool bOk;
 		FValue Value;
+		FRequestStats RequestStats;
+		RequestStats.Bucket = Request.Key.Bucket;
+		RequestStats.Type = ERequestType::Value;
+		RequestStats.Op = ERequestOp::Get;
 		{
+			FRequestTimer RequestTimer(RequestStats);
 			TRACE_CPUPROFILER_EVENT_SCOPE(FileSystemDDC_GetValue);
 			TRACE_COUNTER_INCREMENT(FileSystemDDC_Get);
-			COOK_STAT(auto Timer = UsageStats.TimeGet());
-			bOk = GetCacheValue(Request.Name, Request.Key, Request.Policy, Value);
+			bOk = GetCacheValue(Request.Name, Request.Key, Request.Policy, Value, RequestStats);
 			if (bOk)
 			{
 				UE_LOG(LogDerivedDataCache, Verbose, TEXT("%s: Cache hit for %s from '%s'"),
 					*CachePath, *WriteToString<96>(Request.Key), *Request.Name);
 				TRACE_COUNTER_INCREMENT(FileSystemDDC_GetHit);
-				TRACE_COUNTER_ADD(FileSystemDDC_BytesRead, Value.GetData().GetCompressedSize());
-				COOK_STAT(Timer.AddHit(Value.GetData().GetCompressedSize()));
+				TRACE_COUNTER_ADD(FileSystemDDC_BytesRead, RequestStats.PhysicalReadSize);
+				TRACE_COUNTER_ADD(FileSystemDDC_BytesWritten, RequestStats.PhysicalWriteSize);
 			}
 		}
+		RequestStats.AddLogicalRead(Value);
+		RequestStats.Status = bOk ? EStatus::Ok : EStatus::Error;
+		StoreStats->AddRequest(RequestStats);
 		OnComplete({Request.Name, Request.Key, Value, Request.UserData, bOk ? EStatus::Ok : EStatus::Error});
 	}
 }
@@ -1623,11 +1646,15 @@ void FFileSystemCacheStore::GetChunks(
 		EStatus Status = EStatus::Error;
 		FSharedBuffer Buffer;
 		uint64 RawSize = 0;
+		FRequestStats RequestStats;
+		RequestStats.Bucket = Request.Key.Bucket;
+		RequestStats.Type = Request.Id.IsNull() ? ERequestType::Value : ERequestType::Record;
+		RequestStats.Op = ERequestOp::GetChunk;
 		{
 			TRACE_CPUPROFILER_EVENT_SCOPE(FileSystemDDC_GetChunks);
 			TRACE_COUNTER_INCREMENT(FileSystemDDC_Get);
 			const bool bExistsOnly = EnumHasAnyFlags(Request.Policy, ECachePolicy::SkipData);
-			COOK_STAT(auto Timer = bExistsOnly ? UsageStats.TimeProbablyExists() : UsageStats.TimeGet());
+			FRequestTimer RequestTimer(RequestStats);
 			if (!(bHasValue && ValueKey == Request.Key && ValueId == Request.Id) || ValueReader.HasSource() < !bExistsOnly)
 			{
 				ValueReader.ResetSource();
@@ -1643,7 +1670,7 @@ void FFileSystemCacheStore::GetChunks(
 						FCacheRecordPolicyBuilder PolicyBuilder(ECachePolicy::None);
 						PolicyBuilder.AddValuePolicy(Request.Id, Request.Policy);
 						Record.Reset();
-						Record = GetCacheRecordOnly(Request.Name, Request.Key, PolicyBuilder.Build());
+						Record = GetCacheRecordOnly(Request.Name, Request.Key, PolicyBuilder.Build(), RequestStats);
 					}
 					if (Record)
 					{
@@ -1653,17 +1680,17 @@ void FFileSystemCacheStore::GetChunks(
 							Value = ValueWithId;
 							ValueId = Request.Id;
 							ValueKey = Request.Key;
-							GetCacheContent(Request.Name, Request.Key, ValueId, Value, Request.Policy, ValueReader, ValueAr);
+							GetCacheContent(Request.Name, Request.Key, ValueId, Value, Request.Policy, ValueReader, ValueAr, RequestStats);
 						}
 					}
 				}
 				else
 				{
 					ValueKey = Request.Key;
-					bHasValue = GetCacheValueOnly(Request.Name, Request.Key, Request.Policy, Value);
+					bHasValue = GetCacheValueOnly(Request.Name, Request.Key, Request.Policy, Value, RequestStats);
 					if (bHasValue)
 					{
-						GetCacheContent(Request.Name, Request.Key, Request.Id, Value, Request.Policy, ValueReader, ValueAr);
+						GetCacheContent(Request.Name, Request.Key, Request.Id, Value, Request.Policy, ValueReader, ValueAr, RequestStats);
 					}
 				}
 			}
@@ -1672,15 +1699,18 @@ void FFileSystemCacheStore::GetChunks(
 				const uint64 RawOffset = FMath::Min(Value.GetRawSize(), Request.RawOffset);
 				RawSize = FMath::Min(Value.GetRawSize() - RawOffset, Request.RawSize);
 				TRACE_COUNTER_INCREMENT(FileSystemDDC_GetHit);
-				TRACE_COUNTER_ADD(FileSystemDDC_BytesRead, !bExistsOnly ? RawSize : 0);
-				COOK_STAT(Timer.AddHit(!bExistsOnly ? RawSize : 0));
+				TRACE_COUNTER_ADD(FileSystemDDC_BytesRead, RequestStats.PhysicalReadSize);
+				TRACE_COUNTER_ADD(FileSystemDDC_BytesWritten, RequestStats.PhysicalWriteSize);
 				if (!bExistsOnly)
 				{
+					RequestStats.LogicalReadSize += RawSize;
 					Buffer = ValueReader.Decompress(RawOffset, RawSize);
 				}
 				Status = bExistsOnly || Buffer.GetSize() == RawSize ? EStatus::Ok : EStatus::Error;
 			}
 		}
+		RequestStats.Status = Status;
+		StoreStats->AddRequest(RequestStats);
 		UE_CLOG(Status == EStatus::Ok, LogDerivedDataCache, Verbose, TEXT("%s: Cache hit for %s from '%s'"),
 			*CachePath, *WriteToString<96>(Request.Key, '/', Request.Id), *Request.Name);
 		OnComplete({Request.Name, Request.Key, Request.Id, Request.RawOffset,
@@ -1692,7 +1722,7 @@ bool FFileSystemCacheStore::PutCacheRecord(
 	const FStringView Name,
 	const FCacheRecord& Record,
 	const FCacheRecordPolicy& Policy,
-	uint64& OutWriteSize)
+	FRequestStats& Stats)
 {
 	const bool bLocalDeactivatedForPerformance = IsDeactivatedForPerformance();
 	if (bLocalDeactivatedForPerformance || bReadOnly)
@@ -1738,7 +1768,7 @@ bool FFileSystemCacheStore::PutCacheRecord(
 	{
 		// Load the existing package to take its attachments into account.
 		// Save the new package if there is no existing package or it fails to load.
-		bSavePackage |= !LoadFileWithHash(Path, Name, [&ExistingPackage](FArchive& Ar) { ExistingPackage.TryLoad(Ar); });
+		bSavePackage |= !LoadFileWithHash(Path, Name, Stats, [&ExistingPackage](FArchive& Ar) { ExistingPackage.TryLoad(Ar); });
 		if (!bSavePackage)
 		{
 			// Save the new package if the existing package is invalid.
@@ -1753,9 +1783,9 @@ bool FFileSystemCacheStore::PutCacheRecord(
 				// Content differs between the existing record and the new record.
 				UE_LOG(LogDerivedDataCache, Display, TEXT("%s: Cache put found non-deterministic record for %s from '%.*s'"),
 					*CachePath, *WriteToString<96>(Key), Name.Len(), Name.GetData());
-				const auto HasValueContent = [this, Name, &Key](const FValueWithId& Value) -> bool
+				const auto HasValueContent = [this, Name, &Key, &Stats](const FValueWithId& Value) -> bool
 				{
-					if (!Value.HasData() && !GetCacheContentExists(Key, Value.GetRawHash()))
+					if (!Value.HasData() && !GetCacheContentExists(Key, Value.GetRawHash(), Stats))
 					{
 						UE_LOG(LogDerivedDataCache, Log,
 							TEXT("%s: Cache put of non-deterministic record will overwrite existing record due to "
@@ -1833,23 +1863,26 @@ bool FFileSystemCacheStore::PutCacheRecord(
 	// Save the external content to storage.
 	for (FCompressedBuffer& Content : ExternalContent)
 	{
-		uint64 WriteSize = 0;
-		if (!PutCacheContent(Name, Content, OutWriteSize))
+		if (!PutCacheContent(Name, Content, Stats))
 		{
 			return false;
 		}
-		OutWriteSize += WriteSize;
 	}
 
 	// Save the record package to storage.
-	const auto WritePackage = [&Package, &OutWriteSize](FArchive& Ar)
-	{
-		Package.Save(Ar);
-		OutWriteSize += uint64(Ar.TotalSize());
-	};
-	if (bSavePackage && !SaveFileWithHash(Path, Name, WritePackage, bReplaceExisting))
+	const auto WritePackage = [&Package](FArchive& Ar) { Package.Save(Ar); };
+	if (bSavePackage)
 	{
-		return false;
+		if (!SaveFileWithHash(Path, Name, Stats, WritePackage, bReplaceExisting))
+		{
+			return false;
+		}
+		if (const FCbObject& Meta = Record.GetMeta())
+		{
+			Stats.LogicalWriteSize += Meta.GetSize();
+		}
+		Stats.LogicalWriteSize += Algo::TransformAccumulate(Package.GetAttachments(),
+			[](const FCbAttachment& Attachment) { return Attachment.AsCompressedBinary().GetRawSize(); }, uint64(0));
 	}
 
 	if (AccessLogWriter)
@@ -1863,7 +1896,8 @@ bool FFileSystemCacheStore::PutCacheRecord(
 FOptionalCacheRecord FFileSystemCacheStore::GetCacheRecordOnly(
 	const FStringView Name,
 	const FCacheKey& Key,
-	const FCacheRecordPolicy& Policy)
+	const FCacheRecordPolicy& Policy,
+	FRequestStats& Stats)
 {
 	// Skip the request if querying the cache is disabled.
 	const ECachePolicy QueryFlag = SpeedClass == EBackendSpeedClass::Local ? ECachePolicy::QueryLocal : ECachePolicy::QueryRemote;
@@ -1903,7 +1937,7 @@ FOptionalCacheRecord FFileSystemCacheStore::GetCacheRecordOnly(
 	FOptionalCacheRecord Record;
 	{
 		FCbPackage Package;
-		if (!LoadFileWithHash(Path, Name, [&Package](FArchive& Ar) { Package.TryLoad(Ar); }))
+		if (!LoadFileWithHash(Path, Name, Stats, [&Package](FArchive& Ar) { Package.TryLoad(Ar); }))
 		{
 			UE_LOG(LogDerivedDataCache, Verbose, TEXT("%s: Cache miss with missing package for %s from '%.*s'"),
 				*CachePath, *WriteToString<96>(Key), Name.Len(), Name.GetData());
@@ -1939,9 +1973,10 @@ FOptionalCacheRecord FFileSystemCacheStore::GetCacheRecord(
 	const FStringView Name,
 	const FCacheKey& Key,
 	const FCacheRecordPolicy& Policy,
-	EStatus& OutStatus)
+	EStatus& OutStatus,
+	FRequestStats& Stats)
 {
-	FOptionalCacheRecord Record = GetCacheRecordOnly(Name, Key, Policy);
+	FOptionalCacheRecord Record = GetCacheRecordOnly(Name, Key, Policy, Stats);
 	if (Record.IsNull())
 	{
 		OutStatus = EStatus::Error;
@@ -1963,7 +1998,7 @@ FOptionalCacheRecord FFileSystemCacheStore::GetCacheRecord(
 		const FValueId& Id = Value.GetId();
 		const ECachePolicy ValuePolicy = Policy.GetValuePolicy(Id);
 		FValue Content;
-		if (GetCacheContent(Name, Key, Id, Value, ValuePolicy, Content))
+		if (GetCacheContent(Name, Key, Id, Value, ValuePolicy, Content, Stats))
 		{
 			RecordBuilder.AddValue(Id, MoveTemp(Content));
 		}
@@ -1987,7 +2022,7 @@ bool FFileSystemCacheStore::PutCacheValue(
 	const FCacheKey& Key,
 	const FValue& Value,
 	const ECachePolicy Policy,
-	uint64& OutWriteSize)
+	FRequestStats& Stats)
 {
 	const bool bLocalDeactivatedForPerformance = IsDeactivatedForPerformance();
 	if (bLocalDeactivatedForPerformance || bReadOnly)
@@ -2031,7 +2066,7 @@ bool FFileSystemCacheStore::PutCacheValue(
 	{
 		// Load the existing package to take its attachments into account.
 		// Save the new package if there is no existing package or it fails to load.
-		bSavePackage |= !LoadFileWithHash(Path, Name, [&ExistingPackage](FArchive& Ar) { ExistingPackage.TryLoad(Ar); });
+		bSavePackage |= !LoadFileWithHash(Path, Name, Stats, [&ExistingPackage](FArchive& Ar) { ExistingPackage.TryLoad(Ar); });
 		if (!bSavePackage)
 		{
 			const FCbObjectView Object = ExistingPackage.GetObject();
@@ -2051,7 +2086,7 @@ bool FFileSystemCacheStore::PutCacheValue(
 					"with new hash %s and existing hash %s for %s from '%.*s'"),
 					*CachePath, *WriteToString<48>(Value.GetRawHash()), *WriteToString<48>(RawHash),
 					*WriteToString<96>(Key), Name.Len(), Name.GetData());
-				if (!ExistingPackage.FindAttachment(RawHash) && !GetCacheContentExists(Key, RawHash))
+				if (!ExistingPackage.FindAttachment(RawHash) && !GetCacheContentExists(Key, RawHash, Stats))
 				{
 					// Save the new package because the existing package differs and is missing content.
 					UE_LOG(LogDerivedDataCache, Log,
@@ -2067,6 +2102,7 @@ bool FFileSystemCacheStore::PutCacheValue(
 
 	// Save the value to a package and save the data to external content depending on its size.
 	FCbPackage Package;
+	uint64 LogicalPackageSize = 0;
 	TArray<FCompressedBuffer, TInlineAllocator<1>> ExternalContent;
 	if (ExistingPackage && !bSavePackage)
 	{
@@ -2087,7 +2123,7 @@ bool FFileSystemCacheStore::PutCacheValue(
 		if (!Value.HasData())
 		{
 			// Verify that the content exists in storage.
-			if (!GetCacheContentExists(Key, Value.GetRawHash()))
+			if (!GetCacheContentExists(Key, Value.GetRawHash(), Stats))
 			{
 				UE_LOG(LogDerivedDataCache, Verbose, TEXT("%s: Failed due to missing data for put of %s from '%.*s'"),
 					*CachePath, *WriteToString<96>(Key), Name.Len(), Name.GetData());
@@ -2097,6 +2133,7 @@ bool FFileSystemCacheStore::PutCacheValue(
 		else if (Value.GetData().GetCompressedSize() <= MaxValueSizeKB * 1024)
 		{
 			// Store the content in the package.
+			LogicalPackageSize += Value.GetRawSize();
 			Package.AddAttachment(FCbAttachment(Value.GetData()));
 		}
 		else
@@ -2108,23 +2145,21 @@ bool FFileSystemCacheStore::PutCacheValue(
 	// Save the external content to storage.
 	for (FCompressedBuffer& Content : ExternalContent)
 	{
-		uint64 WriteSize = 0;
-		if (!PutCacheContent(Name, Content, OutWriteSize))
+		if (!PutCacheContent(Name, Content, Stats))
 		{
 			return false;
 		}
-		OutWriteSize += WriteSize;
 	}
 
 	// Save the value package to storage.
-	const auto WritePackage = [&Package, &OutWriteSize](FArchive& Ar)
-	{
-		Package.Save(Ar);
-		OutWriteSize += uint64(Ar.TotalSize());
-	};
-	if (bSavePackage && !SaveFileWithHash(Path, Name, WritePackage, bReplaceExisting))
+	const auto WritePackage = [&Package](FArchive& Ar) { Package.Save(Ar); };
+	if (bSavePackage)
 	{
-		return false;
+		if (!SaveFileWithHash(Path, Name, Stats, WritePackage, bReplaceExisting))
+		{
+			return false;
+		}
+		Stats.LogicalWriteSize += LogicalPackageSize;
 	}
 
 	if (AccessLogWriter)
@@ -2139,7 +2174,8 @@ bool FFileSystemCacheStore::GetCacheValueOnly(
 	const FStringView Name,
 	const FCacheKey& Key,
 	const ECachePolicy Policy,
-	FValue& OutValue)
+	FValue& OutValue,
+	FRequestStats& Stats)
 {
 	// Skip the request if querying the cache is disabled.
 	const ECachePolicy QueryFlag = SpeedClass == EBackendSpeedClass::Local ? ECachePolicy::QueryLocal : ECachePolicy::QueryRemote;
@@ -2177,7 +2213,7 @@ bool FFileSystemCacheStore::GetCacheValueOnly(
 	};
 
 	FCbPackage Package;
-	if (!LoadFileWithHash(Path, Name, [&Package](FArchive& Ar) { Package.TryLoad(Ar); }))
+	if (!LoadFileWithHash(Path, Name, Stats, [&Package](FArchive& Ar) { Package.TryLoad(Ar); }))
 	{
 		UE_LOG(LogDerivedDataCache, Verbose, TEXT("%s: Cache miss with missing package for %s from '%.*s'"),
 			*CachePath, *WriteToString<96>(Key), Name.Len(), Name.GetData());
@@ -2231,23 +2267,28 @@ bool FFileSystemCacheStore::GetCacheValue(
 	const FStringView Name,
 	const FCacheKey& Key,
 	const ECachePolicy Policy,
-	FValue& OutValue)
+	FValue& OutValue,
+	FRequestStats& Stats)
 {
-	return GetCacheValueOnly(Name, Key, Policy, OutValue) && GetCacheContent(Name, Key, {}, OutValue, Policy, OutValue);
+	return GetCacheValueOnly(Name, Key, Policy, OutValue, Stats) && GetCacheContent(Name, Key, {}, OutValue, Policy, OutValue, Stats);
 }
 
 bool FFileSystemCacheStore::PutCacheContent(
 	const FStringView Name,
 	const FCompressedBuffer& Content,
-	uint64& OutWriteSize) const
+	FRequestStats& Stats) const
 {
 	const FIoHash& RawHash = Content.GetRawHash();
 	TStringBuilder<256> Path;
 	BuildCacheContentPath(RawHash, Path);
-	const auto WriteContent = [&](FArchive& Ar) { Content.Save(Ar); OutWriteSize += uint64(Ar.TotalSize()); };
-	if (!FileExists(Path) && !SaveFileWithHash(Path, Name, WriteContent))
+	const auto WriteContent = [&Content](FArchive& Ar) { Content.Save(Ar); };
+	if (!FileExists(Path, Stats))
 	{
-		return false;
+		if (!SaveFileWithHash(Path, Name, Stats, WriteContent))
+		{
+			return false;
+		}
+		Stats.LogicalWriteSize += Content.GetRawSize();
 	}
 	if (AccessLogWriter)
 	{
@@ -2256,11 +2297,11 @@ bool FFileSystemCacheStore::PutCacheContent(
 	return true;
 }
 
-bool FFileSystemCacheStore::GetCacheContentExists(const FCacheKey& Key, const FIoHash& RawHash) const
+bool FFileSystemCacheStore::GetCacheContentExists(const FCacheKey& Key, const FIoHash& RawHash, FRequestStats& Stats) const
 {
 	TStringBuilder<256> Path;
 	BuildCacheContentPath(RawHash, Path);
-	return FileExists(Path);
+	return FileExists(Path, Stats);
 }
 
 bool FFileSystemCacheStore::GetCacheContent(
@@ -2269,7 +2310,8 @@ bool FFileSystemCacheStore::GetCacheContent(
 	const FValueId& Id,
 	const FValue& Value,
 	const ECachePolicy Policy,
-	FValue& OutValue) const
+	FValue& OutValue,
+	FRequestStats& Stats) const
 {
 	if (!EnumHasAnyFlags(Policy, ECachePolicy::Query))
 	{
@@ -2289,7 +2331,7 @@ bool FFileSystemCacheStore::GetCacheContent(
 	BuildCacheContentPath(RawHash, Path);
 	if (EnumHasAnyFlags(Policy, ECachePolicy::SkipData))
 	{
-		if (FileExists(Path))
+		if (FileExists(Path, Stats))
 		{
 			if (AccessLogWriter)
 			{
@@ -2302,7 +2344,7 @@ bool FFileSystemCacheStore::GetCacheContent(
 	else
 	{
 		FCompressedBuffer CompressedBuffer;
-		if (LoadFileWithHash(Path, Name, [&CompressedBuffer](FArchive& Ar) { CompressedBuffer = FCompressedBuffer::Load(Ar); }))
+		if (LoadFileWithHash(Path, Name, Stats, [&CompressedBuffer](FArchive& Ar) { CompressedBuffer = FCompressedBuffer::Load(Ar); }))
 		{
 			if (CompressedBuffer.GetRawHash() == RawHash)
 			{
@@ -2335,8 +2377,29 @@ void FFileSystemCacheStore::GetCacheContent(
 	const FValue& Value,
 	const ECachePolicy Policy,
 	FCompressedBufferReader& Reader,
-	TUniquePtr<FArchive>& OutArchive) const
+	TUniquePtr<FArchive>& OutArchive,
+	FRequestStats& Stats) const
 {
+	class FStatsArchive final : TUniquePtr<FArchive>, public FArchiveProxy
+	{
+	public:
+		FStatsArchive(FArchive& InArchive, FRequestStats& InStats)
+			: TUniquePtr<FArchive>(&InArchive)
+			, FArchiveProxy(InArchive)
+			, Stats(InStats)
+		{
+		}
+
+		void Serialize(void* V, int64 Length) final
+		{
+			Stats.PhysicalReadSize += uint64(Length);
+			FArchiveProxy::Serialize(V, Length);
+		}
+
+	private:
+		FRequestStats& Stats;
+	};
+
 	if (!EnumHasAnyFlags(Policy, ECachePolicy::Query))
 	{
 		return;
@@ -2358,7 +2421,7 @@ void FFileSystemCacheStore::GetCacheContent(
 	BuildCacheContentPath(RawHash, Path);
 	if (EnumHasAllFlags(Policy, ECachePolicy::SkipData))
 	{
-		if (FileExists(Path))
+		if (FileExists(Path, Stats))
 		{
 			if (AccessLogWriter)
 			{
@@ -2369,9 +2432,10 @@ void FFileSystemCacheStore::GetCacheContent(
 	}
 	else
 	{
-		OutArchive = OpenFileRead(Path, Name);
+		OutArchive = OpenFileRead(Path, Name, Stats);
 		if (OutArchive)
 		{
+			OutArchive.Reset(new FStatsArchive(*OutArchive.Release(), Stats));
 			UE_LOG(LogDerivedDataCache, VeryVerbose,
 				TEXT("%s: Opened %s from '%.*s'"),
 				*CachePath, *Path, Name.Len(), Name.GetData());
@@ -2415,10 +2479,11 @@ void FFileSystemCacheStore::BuildCacheContentPath(const FIoHash& RawHash, FStrin
 bool FFileSystemCacheStore::SaveFileWithHash(
 	FStringBuilderBase& Path,
 	const FStringView DebugName,
+	FRequestStats& Stats,
 	const TFunctionRef<void (FArchive&)> WriteFunction,
 	const bool bReplaceExisting) const
 {
-	return SaveFile(Path, DebugName, [&WriteFunction](FArchive& Ar)
+	return SaveFile(Path, DebugName, Stats, [&WriteFunction](FArchive& Ar)
 	{
 		THashingArchiveProxy<FBlake3> HashAr(Ar);
 		WriteFunction(HashAr);
@@ -2430,9 +2495,10 @@ bool FFileSystemCacheStore::SaveFileWithHash(
 bool FFileSystemCacheStore::LoadFileWithHash(
 	FStringBuilderBase& Path,
 	const FStringView DebugName,
+	FRequestStats& Stats,
 	const TFunctionRef<void (FArchive& Ar)> ReadFunction) const
 {
-	return LoadFile(Path, DebugName, [this, &Path, &DebugName, &ReadFunction](FArchive& Ar)
+	return LoadFile(Path, DebugName, Stats, [this, &Path, &DebugName, &ReadFunction](FArchive& Ar)
 	{
 		THashingArchiveProxy<FBlake3> HashAr(Ar);
 		ReadFunction(HashAr);
@@ -2454,6 +2520,7 @@ bool FFileSystemCacheStore::LoadFileWithHash(
 bool FFileSystemCacheStore::SaveFile(
 	FStringBuilderBase& Path,
 	const FStringView DebugName,
+	FRequestStats& Stats,
 	const TFunctionRef<void (FArchive&)> WriteFunction,
 	const bool bReplaceExisting) const
 {
@@ -2462,7 +2529,7 @@ bool FFileSystemCacheStore::SaveFile(
 	TStringBuilder<256> TempPath;
 	TempPath << FPathViews::GetPath(Path) << TEXT("/Temp.") << FGuid::NewGuid();
 
-	TUniquePtr<FArchive> Ar = OpenFileWrite(TempPath, DebugName);
+	TUniquePtr<FArchive> Ar = OpenFileWrite(TempPath, DebugName, Stats);
 	if (!Ar)
 	{
 		UE_LOG(LogDerivedDataCache, Warning,
@@ -2499,17 +2566,23 @@ bool FFileSystemCacheStore::SaveFile(
 		TEXT("%s: Saved %s from '%.*s' (%" INT64_FMT " bytes, %.02f secs, %.2f MiB/s)"),
 		*CachePath, *Path, DebugName.Len(), DebugName.GetData(), WriteSize, WriteDuration, WriteSpeed);
 
+	if (WriteSize > 0)
+	{
+		Stats.PhysicalWriteSize += uint64(WriteSize);
+	}
+
 	return true;
 }
 
 bool FFileSystemCacheStore::LoadFile(
 	FStringBuilderBase& Path,
 	const FStringView DebugName,
+	FRequestStats& Stats,
 	const TFunctionRef<void (FArchive& Ar)> ReadFunction) const
 {
 	const double StartTime = FPlatformTime::Seconds();
 
-	TUniquePtr<FArchive> Ar = OpenFileRead(Path, DebugName);
+	TUniquePtr<FArchive> Ar = OpenFileRead(Path, DebugName, Stats);
 	if (!Ar)
 	{
 		return false;
@@ -2548,11 +2621,19 @@ bool FFileSystemCacheStore::LoadFile(
 		}
 	}
 
+	if (ReadSize > 0)
+	{
+		Stats.PhysicalReadSize += uint64(ReadSize);
+	}
+
 	return !bError && ReadSize > 0;
 }
 
-TUniquePtr<FArchive> FFileSystemCacheStore::OpenFileWrite(FStringBuilderBase& Path, const FStringView DebugName) const
+TUniquePtr<FArchive> FFileSystemCacheStore::OpenFileWrite(FStringBuilderBase& Path, const FStringView DebugName, FRequestStats& Stats) const
 {
+	const FMonotonicTimePoint StartTime = FMonotonicTimePoint::Now();
+	ON_SCOPE_EXIT { Stats.AddLatency(FMonotonicTimePoint::Now() - StartTime); };
+
 	// Retry to handle a race where the directory is deleted while the file is being created.
 	constexpr int32 MaxAttemptCount = 3;
 	for (int32 AttemptCount = 0; AttemptCount < MaxAttemptCount; ++AttemptCount)
@@ -2565,7 +2646,7 @@ TUniquePtr<FArchive> FFileSystemCacheStore::OpenFileWrite(FStringBuilderBase& Pa
 	return nullptr;
 }
 
-TUniquePtr<FArchive> FFileSystemCacheStore::OpenFileRead(FStringBuilderBase& Path, const FStringView DebugName) const
+TUniquePtr<FArchive> FFileSystemCacheStore::OpenFileRead(FStringBuilderBase& Path, const FStringView DebugName, FRequestStats& Stats) const
 {
 	// Checking for existence may update the modification time to avoid the file being evicted from the cache.
 	// Reduce Game Thread overhead by executing the update on a worker thread if the path implies higher latency.
@@ -2574,24 +2655,29 @@ TUniquePtr<FArchive> FFileSystemCacheStore::OpenFileRead(FStringBuilderBase& Pat
 		FRequestOwner AsyncOwner(EPriority::Normal);
 		Private::LaunchTaskInCacheThreadPool(AsyncOwner, [this, Path = MakeShared<TStringBuilder<256>>(InPlace, Path)]() mutable
 		{
-			(void)FileExists(*Path);
+			FRequestStats UnusedStats;
+			(void)FileExists(*Path, UnusedStats);
 		});
 		AsyncOwner.KeepAlive();
 	}
 	else
 	{
-		if (!FileExists(Path))
+		if (!FileExists(Path, Stats))
 		{
 			return nullptr;
 		}
 	}
 
+	const FMonotonicTimePoint StartTime = FMonotonicTimePoint::Now();
+	ON_SCOPE_EXIT { Stats.AddLatency(FMonotonicTimePoint::Now() - StartTime); };
 	return TUniquePtr<FArchive>(IFileManager::Get().CreateFileReader(*Path, FILEREAD_Silent));
 }
 
-bool FFileSystemCacheStore::FileExists(FStringBuilderBase& Path) const
+bool FFileSystemCacheStore::FileExists(FStringBuilderBase& Path, FRequestStats& Stats) const
 {
+	const FMonotonicTimePoint StartTime = FMonotonicTimePoint::Now();
 	const FDateTime TimeStamp = IFileManager::Get().GetTimeStamp(*Path);
+	Stats.AddLatency(FMonotonicTimePoint::Now() - StartTime);
 	if (TimeStamp == FDateTime::MinValue())
 	{
 		return false;
@@ -2650,6 +2736,7 @@ bool FFileSystemCacheStore::IsDeactivatedForPerformance()
 					}
 
 					bDeactivedForPerformance.store(bLocalDeactivatedForPerformance, std::memory_order_relaxed);
+					UpdateStatus();
 					return bLocalDeactivatedForPerformance;
 				}
 			}
@@ -2718,6 +2805,18 @@ bool FFileSystemCacheStore::IsDeactivatedForPerformance()
 	return true;
 }
 
+void FFileSystemCacheStore::UpdateStatus()
+{
+	if (bDeactivedForPerformance.load(std::memory_order_relaxed))
+	{
+		StoreStats->SetStatus(ECacheStoreStatusCode::Warning, NSLOCTEXT("DerivedDataCache", "DeactivatedForPerformance", "Deactivated for performance"));
+	}
+	else
+	{
+		StoreStats->SetStatus(ECacheStoreStatusCode::None, {});
+	}
+}
+
 bool FFileSystemCacheStore::RunInitialSpeedTest(const FFileSystemCacheStoreParams& Params, FDerivedDataCacheSpeedStats& OutSpeedStats)
 {
 	struct FSpeedTestState : public FThreadSafeRefCountedObject
@@ -2774,8 +2873,10 @@ bool FFileSystemCacheStore::RunInitialSpeedTest(const FFileSystemCacheStoreParam
 	return SpeedTestState->bResult;
 }
 
-void FFileSystemCacheStoreParams::Parse(const TCHAR* NodeName, const TCHAR* Config)
+void FFileSystemCacheStoreParams::Parse(const TCHAR* Name, const TCHAR* Config)
 {
+	CacheName = Name;
+
 	FString Key;
 
 	// Path Params
@@ -2787,19 +2888,19 @@ void FFileSystemCacheStoreParams::Parse(const TCHAR* NodeName, const TCHAR* Conf
 		if (FString Value = FPlatformMisc::GetEnvironmentVariable(*Key); !Value.IsEmpty())
 		{
 			CachePath = MoveTemp(Value);
-			UE_LOG(LogDerivedDataCache, Log, TEXT("%s: Found environment variable %s=%s"), NodeName, *Key, *CachePath);
+			UE_LOG(LogDerivedDataCache, Log, TEXT("%s: Found environment variable %s=%s"), Name, *Key, *CachePath);
 		}
 		if (FString Value; FPlatformMisc::GetStoredValue(TEXT("Epic Games"), TEXT("GlobalDataCachePath"), *Key, Value) && !Value.IsEmpty())
 		{
 			CachePath = MoveTemp(Value);
-			UE_LOG(LogDerivedDataCache, Log, TEXT("%s: Found registry key GlobalDataCachePath %s=%s"), NodeName, *Key, *CachePath);
+			UE_LOG(LogDerivedDataCache, Log, TEXT("%s: Found registry key GlobalDataCachePath %s=%s"), Name, *Key, *CachePath);
 		}
 	}
 
 	if (FParse::Value(Config, TEXT("CommandLineOverride="), Key) &&
 		FParse::Value(FCommandLine::Get(), *(Key + TEXT("=")), CachePath))
 	{
-		UE_LOG(LogDerivedDataCache, Log, TEXT("%s: Found command line override %s=%s"), NodeName, *Key, *CachePath);
+		UE_LOG(LogDerivedDataCache, Log, TEXT("%s: Found command line override %s=%s"), Name, *Key, *CachePath);
 	}
 
 	if (FParse::Value(Config, TEXT("EditorOverrideSetting="), Key))
@@ -2813,7 +2914,7 @@ void FFileSystemCacheStoreParams::Parse(const TCHAR* NodeName, const TCHAR* Conf
 				if (!Value.IsEmpty())
 				{
 					CachePath = MoveTemp(Value);
-					UE_LOG(LogDerivedDataCache, Log, TEXT("%s: Found editor settings override %s=%s"), NodeName, *Key, *CachePath);
+					UE_LOG(LogDerivedDataCache, Log, TEXT("%s: Found editor settings override %s=%s"), Name, *Key, *CachePath);
 				}
 			}
 		}
@@ -2873,14 +2974,14 @@ void FFileSystemCacheStoreParams::Parse(const TCHAR* NodeName, const TCHAR* Conf
 }
 
 ILegacyCacheStore* CreateFileSystemCacheStore(
-	const TCHAR* NodeName,
+	const TCHAR* Name,
 	const TCHAR* Config,
-	ECacheStoreFlags& OutFlags,
+	ICacheStoreOwner& Owner,
 	FString& OutPath)
 {
 	FFileSystemCacheStoreParams Params;
-	Params.Parse(NodeName, Config);
-	return FFileSystemCacheStore::TryCreate(NodeName, Params, OutFlags, OutPath);
+	Params.Parse(Name, Config);
+	return FFileSystemCacheStore::TryCreate(Params, Owner, OutPath);
 }
 
 } // UE::DerivedData
