DDC: Added ICacheStoreStats to the memory cache and cleaned up graph node lifetime

#rnx

diff --git a/Engine/Source/Developer/DerivedDataCache/Private/DerivedDataBackends.cpp b/Engine/Source/Developer/DerivedDataCache/Private/DerivedDataBackends.cpp
--- a/Engine/Source/Developer/DerivedDataCache/Private/DerivedDataBackends.cpp
+++ b/Engine/Source/Developer/DerivedDataCache/Private/DerivedDataBackends.cpp
@@ -62,12 +62,12 @@ namespace UE::DerivedData
 {
 
 ILegacyCacheStore* CreateCacheStoreAsync(ILegacyCacheStore* InnerBackend, IMemoryCacheStore* MemoryCache);
-ILegacyCacheStore* CreateCacheStoreHierarchy(ICacheStoreOwner*& OutOwner, IMemoryCacheStore* MemoryCache);
+ILegacyCacheStore* CreateCacheStoreHierarchy(ICacheStoreOwner*& OutOwner, TFunctionRef<void (IMemoryCacheStore*&)> MemoryCacheCreator);
 ILegacyCacheStore* CreateCacheStoreThrottle(ILegacyCacheStore* InnerCache, uint32 LatencyMS, uint32 MaxBytesPerSecond);
 ILegacyCacheStore* CreateCacheStoreVerify(ILegacyCacheStore* InnerCache, bool bPutOnError);
 ILegacyCacheStore* CreateFileSystemCacheStore(const TCHAR* Name, const TCHAR* Config, ICacheStoreOwner& Owner, FString& OutPath);
 TTuple<ILegacyCacheStore*, ECacheStoreFlags> CreateHttpCacheStore(const TCHAR* NodeName, const TCHAR* Config);
-IMemoryCacheStore* CreateMemoryCacheStore(const TCHAR* Name, int64 MaxCacheSize, bool bCanBeDisabled);
+void CreateMemoryCacheStore(IMemoryCacheStore*& OutCache, const TCHAR* Name, ICacheStoreOwner* Owner, int64 MaxCacheSize, bool bCanBeDisabled);
 IPakFileCacheStore* CreatePakFileCacheStore(const TCHAR* Filename, bool bWriting, bool bCompressed);
 ILegacyCacheStore* CreateS3CacheStore(const TCHAR* RootManifestPath, const TCHAR* BaseUrl, const TCHAR* Region, const TCHAR* CanaryObjectKey, const TCHAR* CachePath);
 TTuple<ILegacyCacheStore*, ECacheStoreFlags> CreateZenCacheStore(const TCHAR* NodeName, const TCHAR* Config);
@@ -146,14 +146,14 @@ public:
 
 			if (!GraphName.IsEmpty() && GraphName != TEXT("Default"))
 			{
-				RootNode = CreateCacheStoreHierarchy(Hierarchy, GetMemoryCache());
-				CreatedNodes.AddUnique(RootNode);
+				RootNode = CreateCacheStoreHierarchy(Hierarchy, [this](IMemoryCacheStore*& OutCache) { GetMemoryCache(OutCache); });
 
 				if (!ParseNode(RootName, GEngineIni, *GraphName, ParsedNodes))
 				{
 					// Destroy any cache stores that have been created.
+					delete RootNode;
+					RootNode = nullptr;
 					ParsedNodes.Empty();
-					DestroyCreatedBackends();
 					MemoryCache = nullptr;
 					BootCache = nullptr;
 					WritePakCache = nullptr;
@@ -174,8 +174,7 @@ public:
 			{
 				// Try to use the default graph.
 				GraphName = FApp::IsEngineInstalled() ? TEXT("InstalledDerivedDataBackendGraph") : TEXT("DerivedDataBackendGraph");
-				RootNode = CreateCacheStoreHierarchy(Hierarchy, GetMemoryCache());
-				CreatedNodes.AddUnique(RootNode);
+				RootNode = CreateCacheStoreHierarchy(Hierarchy, [this](IMemoryCacheStore*& OutCache) { GetMemoryCache(OutCache); });
 
 				if (!ParseNode(RootName, GEngineIni, *GraphName, ParsedNodes))
 				{
@@ -203,10 +202,7 @@ public:
 		}
 
 		// Async must exist in the graph.
-		{
-			RootNode = CreateCacheStoreAsync(RootNode, GetMemoryCache());
-			CreatedNodes.AddUnique(RootNode);
-		}
+		RootNode = CreateCacheStoreAsync(RootNode, MemoryCache);
 
 		// Create a Verify node when using -DDC-Verify[=Type1[@Rate2][+Type2[@Rate2]...]].
 		if (bVerifyFound)
@@ -217,7 +213,6 @@ public:
 			{
 				IFileManager::Get().DeleteDirectory(*(FPaths::ProjectSavedDir() / TEXT("VerifyDDC/")), /*bRequireExists*/ false, /*bTree*/ true);
 				RootNode = CreateCacheStoreVerify(RootNode, /*bPutOnError*/ bVerifyFix);
-				CreatedNodes.AddUnique(RootNode);
 			}
 		}
 
@@ -225,7 +220,6 @@ public:
 		if (ILegacyCacheStore* ReplayNode = TryCreateCacheStoreReplay(RootNode))
 		{
 			RootNode = ReplayNode;
-			CreatedNodes.AddUnique(ReplayNode);
 		}
 
 		if (MaxKeyLength == 0)
@@ -660,7 +654,7 @@ public:
 		MaxCacheSize = FMath::Min(MaxCacheSize, MaxSupportedCacheSize);
 
 		UE_LOG(LogDerivedDataCache, Display, TEXT("%s: Max Cache Size: %d MB"), NodeName, MaxCacheSize);
-		BootCache = CreateMemoryCacheStore(TEXT("Boot"), MaxCacheSize * 1024 * 1024, /*bCanBeDisabled*/ true);
+		CreateMemoryCacheStore(BootCache, TEXT("Boot"), this, MaxCacheSize * 1024 * 1024, /*bCanBeDisabled*/ true);
 #endif
 
 		if (BootCache)
@@ -683,35 +677,35 @@ public:
 	{
 		FString Filename;
 		FParse::Value(Entry, TEXT("Filename="), Filename);
-		IMemoryCacheStore* Cache = CreateMemoryCacheStore(NodeName, /*MaxCacheSize*/ -1, /*bCanBeDisabled*/ false);
-		if (Cache && Filename.Len())
+		IMemoryCacheStore* Cache;
+		CreateMemoryCacheStore(Cache, NodeName, this, /*MaxCacheSize*/ -1, /*bCanBeDisabled*/ false);
+		check(Cache);
+		if (Filename.Len())
 		{
 			UE_LOG(LogDerivedDataCache, Display, TEXT("Memory nodes that load from a file are deprecated. Please remove the filename from the cache configuration."));
 		}
-		if (Cache)
-		{
-			Add(Cache, ECacheStoreFlags::Local | ECacheStoreFlags::Query | ECacheStoreFlags::Store);
-			return true;
-		}
-		return false;
+		return true;
 	}
 
-	IMemoryCacheStore* GetMemoryCache()
+	void GetMemoryCache(IMemoryCacheStore*& OutCache)
 	{
-		if (!MemoryCache)
+		if (MemoryCache)
 		{
-			MemoryCache = CreateMemoryCacheStore(TEXT("Memory"), 0, /*bCanBeDisabled*/ false);
-			CreatedNodes.Add(MemoryCache);
+			OutCache = MemoryCache;
+		}
+		else
+		{
+			// This is unconditionally added to the hierarchy and will be deleted by the hierarchy.
+			CreateMemoryCacheStore(OutCache, TEXT("Memory"), this, 0, /*bCanBeDisabled*/ false);
+			MemoryCache = OutCache;
 		}
-		return MemoryCache;
 	}
 
 	virtual ~FDerivedDataBackendGraph()
 	{
 		check(StaticGraph == this);
 		Replays.Empty();
-		RootCache = nullptr;
-		DestroyCreatedBackends();
+		delete RootCache;
 		StaticGraph = nullptr;
 	}
 
@@ -860,7 +854,6 @@ public:
 			ReadPak = CreatePakFileCacheStore(PakFilename, /*bWriting*/ false, /*bCompressed*/ false);
 
 			Hierarchy->Add(ReadPak, ECacheStoreFlags::Local | ECacheStoreFlags::Query | ECacheStoreFlags::StopStore);
-			CreatedNodes.AddUnique(ReadPak);
 			ReadPakCache.Add(ReadPak);
 		}
 		else
@@ -885,7 +878,6 @@ public:
 
 				Hierarchy->RemoveNotSafe(ReadPak);
 				ReadPakCache.RemoveAt(PakIndex);
-				CreatedNodes.Remove(ReadPak);
 				ReadPak->Close();
 				delete ReadPak;
 				return true;
@@ -909,8 +901,6 @@ private:
 	{
 		check(Hierarchy);
 
-		CreatedNodes.AddUnique(CacheStore);
-
 		// Parse any debug options for this node. E.g. -DDC-<Name>-MissRate
 		FBackendDebugOptions DebugOptions;
 		if (FBackendDebugOptions::ParseFromTokens(DebugOptions, *ActiveNodeName, FCommandLine::Get()))
@@ -940,7 +930,6 @@ private:
 			{
 				ILegacyCacheStore* ThrottleNode = CreateCacheStoreThrottle(CacheStore, LatencyMS, MaxBytesPerSecond);
 				ThrottleNodes.Add(CacheStore, ThrottleNode);
-				CreatedNodes.Add(ThrottleNode);
 				CacheStore = ThrottleNode;
 			}
 		}
@@ -964,11 +953,9 @@ private:
 	void RemoveNotSafe(ILegacyCacheStore* CacheStore) final
 	{
 		check(Hierarchy);
-		CreatedNodes.Remove(CacheStore);
 		if (ILegacyCacheStore* ThrottleNode = ThrottleNodes.FindRef(CacheStore))
 		{
 			Hierarchy->RemoveNotSafe(ThrottleNode);
-			CreatedNodes.Remove(ThrottleNode);
 			ThrottleNodes.Remove(CacheStore);
 			delete ThrottleNode;
 		}
@@ -990,16 +977,6 @@ private:
 		Hierarchy->DestroyStats(Stats);
 	}
 
-	/** Delete all created backends in the reversed order they were created. */
-	void DestroyCreatedBackends()
-	{
-		for (int32 BackendIndex = CreatedNodes.Num() - 1; BackendIndex >= 0; --BackendIndex)
-		{
-			delete CreatedNodes[BackendIndex];
-		}
-		CreatedNodes.Empty();
-	}
-
 	/** MountPak console command handler. */
 	void UnmountPakCommandHandler(const TArray<FString>& Args)
 	{
@@ -1115,8 +1092,6 @@ private:
 	/** Root of the graph */
 	ILegacyCacheStore*					RootCache;
 
-	/** References to all created backends */
-	TArray< ILegacyCacheStore* > CreatedNodes;
 	TMap<ILegacyCacheStore*, ILegacyCacheStore*> ThrottleNodes;
 
 	/** Instances of backend interfaces which exist in only one copy */
diff --git a/Engine/Source/Developer/DerivedDataCache/Private/DerivedDataCacheReplay.cpp b/Engine/Source/Developer/DerivedDataCache/Private/DerivedDataCacheReplay.cpp
--- a/Engine/Source/Developer/DerivedDataCache/Private/DerivedDataCacheReplay.cpp
+++ b/Engine/Source/Developer/DerivedDataCache/Private/DerivedDataCacheReplay.cpp
@@ -128,6 +128,7 @@ FCacheStoreReplay::FCacheStoreReplay(
 FCacheStoreReplay::~FCacheStoreReplay()
 {
 	FlushToArchive();
+	delete InnerCache;
 }
 
 template <typename RequestType>
diff --git a/Engine/Source/Developer/DerivedDataCache/Private/DerivedDataCacheStore.h b/Engine/Source/Developer/DerivedDataCache/Private/DerivedDataCacheStore.h
--- a/Engine/Source/Developer/DerivedDataCache/Private/DerivedDataCacheStore.h
+++ b/Engine/Source/Developer/DerivedDataCache/Private/DerivedDataCacheStore.h
@@ -174,7 +174,7 @@ enum class ECacheStoreStatusCode : uint8
 /**
  * Stats that represent how one cache store processed one request.
  */
-struct FCacheStoreRequestStats
+struct FCacheStoreRequestStats final
 {
 	/** Bucket that contains the data being accessed by this request. */
 	FCacheBucket Bucket;
@@ -216,7 +216,7 @@ struct FCacheStoreRequestStats
 };
 
 /** Updates StartTime, EndTime, MainThreadTime, OtherThreadTime based on its lifetime. */
-class FCacheStoreRequestTimer
+class FCacheStoreRequestTimer final
 {
 public:
 	explicit FCacheStoreRequestTimer(FCacheStoreRequestStats& Stats);
diff --git a/Engine/Source/Developer/DerivedDataCache/Private/DerivedDataCacheStoreAsync.cpp b/Engine/Source/Developer/DerivedDataCache/Private/DerivedDataCacheStoreAsync.cpp
--- a/Engine/Source/Developer/DerivedDataCache/Private/DerivedDataCacheStoreAsync.cpp
+++ b/Engine/Source/Developer/DerivedDataCache/Private/DerivedDataCacheStoreAsync.cpp
@@ -22,10 +22,11 @@ namespace UE::DerivedData
  *
  * Puts can be stored in a memory cache while they are in flight.
  */
-class FCacheStoreAsync : public ILegacyCacheStore
+class FCacheStoreAsync final : public ILegacyCacheStore
 {
 public:
 	FCacheStoreAsync(ILegacyCacheStore* InnerCache, IMemoryCacheStore* MemoryCache);
+	~FCacheStoreAsync() final;
 
 	virtual void Put(
 		TConstArrayView<FCachePutRequest> Requests,
@@ -123,16 +124,17 @@ FCacheStoreAsync::FCacheStoreAsync(ILegacyCacheStore* InInnerCache, IMemoryCache
 	check(InnerCache);
 }
 
+FCacheStoreAsync::~FCacheStoreAsync()
+{
+	delete InnerCache;
+}
+
 void FCacheStoreAsync::LegacyStats(FDerivedDataCacheStatsNode& OutNode)
 {
 	OutNode = {TEXT("Async"), TEXT(""), /*bIsLocal*/ true};
 	OutNode.UsageStats.Add(TEXT(""), UsageStats);
 
 	InnerCache->LegacyStats(OutNode.Children.Add_GetRef(MakeShared<FDerivedDataCacheStatsNode>()).Get());
-	if (MemoryCache)
-	{
-		MemoryCache->LegacyStats(OutNode.Children.Add_GetRef(MakeShared<FDerivedDataCacheStatsNode>()).Get());
-	}
 }
 
 template <typename RequestType, typename OnCompleteType, typename OnExecuteType>
diff --git a/Engine/Source/Developer/DerivedDataCache/Private/DerivedDataCacheStoreHierarchy.cpp b/Engine/Source/Developer/DerivedDataCache/Private/DerivedDataCacheStoreHierarchy.cpp
--- a/Engine/Source/Developer/DerivedDataCache/Private/DerivedDataCacheStoreHierarchy.cpp
+++ b/Engine/Source/Developer/DerivedDataCache/Private/DerivedDataCacheStoreHierarchy.cpp
@@ -155,8 +155,8 @@ void FCacheStoreStats::AddRequest(const FCacheStoreRequestStats& Stats)
 class FCacheStoreHierarchy final : public ILegacyCacheStore, public ICacheStoreOwner
 {
 public:
-	explicit FCacheStoreHierarchy(IMemoryCacheStore* MemoryCache);
-	~FCacheStoreHierarchy() final = default;
+	explicit FCacheStoreHierarchy(ICacheStoreOwner*& OutOwner, TFunctionRef<void (IMemoryCacheStore*&)> MemoryCacheCreator);
+	~FCacheStoreHierarchy() final;
 
 	void Add(ILegacyCacheStore* CacheStore, ECacheStoreFlags Flags) final;
 	void SetFlags(ILegacyCacheStore* CacheStore, ECacheStoreFlags Flags) final;
@@ -358,12 +358,19 @@ ENUM_CLASS_FLAGS(FCacheStoreHierarchy::ECacheStoreNodeFlags);
 
 ///////////////////////////////////////////////////////////////////////////////////////////////////////////////////////
 
-FCacheStoreHierarchy::FCacheStoreHierarchy(IMemoryCacheStore* InMemoryCache)
-	: MemoryCache(InMemoryCache)
+FCacheStoreHierarchy::FCacheStoreHierarchy(ICacheStoreOwner*& OutOwner, TFunctionRef<void (IMemoryCacheStore*&)> MemoryCacheCreator)
 {
-	if (MemoryCache)
+	// Assign OutOwner before MemoryCache in case the creator depends on the owner.
+	OutOwner = this;
+	MemoryCacheCreator(MemoryCache);
+}
+
+FCacheStoreHierarchy::~FCacheStoreHierarchy()
+{
+	// Delete nodes separately before Nodes is destroyed because destroying stats depends on it.
+	for (FCacheStoreNode& Node : ReverseIterate(Nodes))
 	{
-		Add(MemoryCache, ECacheStoreFlags::Local | ECacheStoreFlags::Query | ECacheStoreFlags::StopGetStore);
+		Node.AsyncCache.Reset();
 	}
 }
 
@@ -1511,11 +1518,9 @@ bool FCacheStoreHierarchy::LegacyDebugOptions(FBackendDebugOptions& Options)
 
 ///////////////////////////////////////////////////////////////////////////////////////////////////////////////////////
 
-ILegacyCacheStore* CreateCacheStoreHierarchy(ICacheStoreOwner*& OutOwner, IMemoryCacheStore* MemoryCache)
+ILegacyCacheStore* CreateCacheStoreHierarchy(ICacheStoreOwner*& OutOwner, TFunctionRef<void (IMemoryCacheStore*&)> MemoryCacheCreator)
 {
-	FCacheStoreHierarchy* Hierarchy = new FCacheStoreHierarchy(MemoryCache);
-	OutOwner = Hierarchy;
-	return Hierarchy;
+	return new FCacheStoreHierarchy(OutOwner, MemoryCacheCreator);
 }
 
 } // UE::DerivedData
diff --git a/Engine/Source/Developer/DerivedDataCache/Private/DerivedDataCacheStoreThrottle.cpp b/Engine/Source/Developer/DerivedDataCache/Private/DerivedDataCacheStoreThrottle.cpp
--- a/Engine/Source/Developer/DerivedDataCache/Private/DerivedDataCacheStoreThrottle.cpp
+++ b/Engine/Source/Developer/DerivedDataCache/Private/DerivedDataCacheStoreThrottle.cpp
@@ -28,6 +28,11 @@ public:
 		check(InnerCache);
 	}
 
+	~FCacheStoreThrottle() final
+	{
+		delete InnerCache;
+	}
+
 	void Put(
 		const TConstArrayView<FCachePutRequest> Requests,
 		IRequestOwner& Owner,
diff --git a/Engine/Source/Developer/DerivedDataCache/Private/DerivedDataCacheStoreVerify.cpp b/Engine/Source/Developer/DerivedDataCache/Private/DerivedDataCacheStoreVerify.cpp
--- a/Engine/Source/Developer/DerivedDataCache/Private/DerivedDataCacheStoreVerify.cpp
+++ b/Engine/Source/Developer/DerivedDataCache/Private/DerivedDataCacheStoreVerify.cpp
@@ -25,7 +25,7 @@ namespace UE::DerivedData
  *
  * This wraps a cache store and fails every get until a matching put occurs, then compares the derived data.
  */
-class FCacheStoreVerify : public ILegacyCacheStore
+class FCacheStoreVerify final : public ILegacyCacheStore
 {
 public:
 	FCacheStoreVerify(ILegacyCacheStore* InInnerCache, bool bInPutOnError)
@@ -68,6 +68,11 @@ public:
 			TEXT("Verify: Any record or value that differs will be overwritten."));
 	}
 
+	~FCacheStoreVerify()
+	{
+		delete InnerCache;
+	}
+
 	void Put(
 		TConstArrayView<FCachePutRequest> Requests,
 		IRequestOwner& Owner,
diff --git a/Engine/Source/Developer/DerivedDataCache/Private/FileSystemCacheStore.cpp b/Engine/Source/Developer/DerivedDataCache/Private/FileSystemCacheStore.cpp
--- a/Engine/Source/Developer/DerivedDataCache/Private/FileSystemCacheStore.cpp
+++ b/Engine/Source/Developer/DerivedDataCache/Private/FileSystemCacheStore.cpp
@@ -1703,8 +1703,8 @@ void FFileSystemCacheStore::GetChunks(
 				TRACE_COUNTER_ADD(FileSystemDDC_BytesWritten, RequestStats.PhysicalWriteSize);
 				if (!bExistsOnly)
 				{
-					RequestStats.LogicalReadSize += RawSize;
 					Buffer = ValueReader.Decompress(RawOffset, RawSize);
+					RequestStats.LogicalReadSize += Buffer.GetSize();
 				}
 				Status = bExistsOnly || Buffer.GetSize() == RawSize ? EStatus::Ok : EStatus::Error;
 			}
diff --git a/Engine/Source/Developer/DerivedDataCache/Private/MemoryCacheStore.cpp b/Engine/Source/Developer/DerivedDataCache/Private/MemoryCacheStore.cpp
--- a/Engine/Source/Developer/DerivedDataCache/Private/MemoryCacheStore.cpp
+++ b/Engine/Source/Developer/DerivedDataCache/Private/MemoryCacheStore.cpp
@@ -9,14 +9,10 @@
 #include "DerivedDataBackendInterface.h"
 #include "DerivedDataCachePrivate.h"
 #include "DerivedDataCacheRecord.h"
-#include "DerivedDataCacheUsageStats.h"
 #include "DerivedDataValue.h"
-#include "HAL/FileManager.h"
 #include "Misc/ScopeExit.h"
 #include "Misc/ScopeRWLock.h"
-#include "ProfilingDebugging/CookStats.h"
 #include "Serialization/CompactBinary.h"
-#include "Templates/UniquePtr.h"
 
 namespace UE::DerivedData
 {
@@ -27,7 +23,12 @@ namespace UE::DerivedData
 class FMemoryCacheStore final : public IMemoryCacheStore
 {
 public:
-	explicit FMemoryCacheStore(const TCHAR* InName, int64 InMaxCacheSize = -1, bool bCanBeDisabled = false);
+	explicit FMemoryCacheStore(
+		IMemoryCacheStore*& OutCache,
+		const TCHAR* InName,
+		ICacheStoreOwner* InOwner,
+		int64 InMaxCacheSize = -1,
+		bool bCanBeDisabled = false);
 	~FMemoryCacheStore() final;
 
 	// ICacheStore Interface
@@ -66,8 +67,6 @@ public:
 	void Disable() final;
 
 private:
-	FDerivedDataCacheUsageStats UsageStats;
-
 	struct FCacheRecordComponents
 	{
 		FCbObject Meta;
@@ -77,6 +76,9 @@ private:
 	/** Name of this cache (used for debugging) */
 	FString Name;
 
+	ICacheStoreOwner* StoreOwner;
+	ICacheStoreStats* StoreStats = nullptr;
+
 	/** Set of records in this cache. */
 	TMap<FCacheKey, FCacheRecordComponents> CacheRecords;
 	/** Set of values in this cache. */
@@ -117,20 +119,37 @@ protected:
 	FBackendDebugOptions DebugOptions;
 };
 
-FMemoryCacheStore::FMemoryCacheStore(const TCHAR* InName, int64 InMaxCacheSize, bool bInCanBeDisabled)
+FMemoryCacheStore::FMemoryCacheStore(
+	IMemoryCacheStore*& OutCache,
+	const TCHAR* InName,
+	ICacheStoreOwner* InOwner,
+	int64 InMaxCacheSize,
+	bool bInCanBeDisabled)
 	: Name(InName)
+	, StoreOwner(InOwner)
 	, MaxCacheSize(InMaxCacheSize < 0 ? 0 : uint64(InMaxCacheSize))
 	, bDisabled(false)
 	, CurrentCacheSize(0)
 	, bMaxSizeExceeded(false)
 	, bCanBeDisabled(bInCanBeDisabled)
 {
+	OutCache = this;
+	if (StoreOwner)
+	{
+		constexpr ECacheStoreFlags Flags = ECacheStoreFlags::Local | ECacheStoreFlags::Query | ECacheStoreFlags::Store;
+		StoreOwner->Add(this, Flags);
+		StoreStats = StoreOwner->CreateStats(this, Flags, TEXTVIEW("Memory"), Name);
+	}
 }
 
 FMemoryCacheStore::~FMemoryCacheStore()
 {
 	bShuttingDown = true;
 	Disable();
+	if (StoreOwner)
+	{
+		StoreOwner->DestroyStats(StoreStats);
+	}
 }
 
 void FMemoryCacheStore::Disable()
@@ -145,8 +164,7 @@ void FMemoryCacheStore::Disable()
 
 void FMemoryCacheStore::LegacyStats(FDerivedDataCacheStatsNode& OutNode)
 {
-	OutNode = {!bCanBeDisabled ? TEXT("Memory") : TEXT("Boot"), TEXT(""), /*bIsLocal*/ true};
-	OutNode.UsageStats.Add(TEXT(""), UsageStats);
+	checkNoEntry();
 }
 
 bool FMemoryCacheStore::LegacyDebugOptions(FBackendDebugOptions& InOptions)
@@ -177,6 +195,22 @@ void FMemoryCacheStore::Put(
 		const FCacheKey& Key = Record.GetKey();
 		const TConstArrayView<FValueWithId> Values = Record.GetValues();
 
+		FRequestStats RequestStats;
+		RequestStats.Bucket = Key.Bucket;
+		RequestStats.Type = ERequestType::Record;
+		RequestStats.Op = ERequestOp::Put;
+		ON_SCOPE_EXIT
+		{
+			RequestStats.Latency = RequestStats.EndTime - RequestStats.StartTime;
+			RequestStats.Status = Status;
+			if (StoreStats)
+			{
+				StoreStats->AddRequest(RequestStats);
+			}
+		};
+
+		FRequestTimer RequestTimer(RequestStats);
+
 		if (Algo::NoneOf(Values, &FValue::HasData))
 		{
 			continue;
@@ -189,7 +223,6 @@ void FMemoryCacheStore::Put(
 			continue;
 		}
 
-		COOK_STAT(auto Timer = UsageStats.TimePut());
 		const bool bReplaceExisting = !EnumHasAnyFlags(Request.Policy.GetRecordPolicy(), ECachePolicy::QueryLocal);
 
 		FWriteScopeLock ScopeLock(SynchronizationObject);
@@ -276,12 +309,18 @@ void FMemoryCacheStore::Put(
 		{
 			Components.Meta = Record.GetMeta();
 			Components.Values = Record.GetValues();
+			RequestStats.LogicalWriteSize += Components.Meta.GetSize();
+			for (const FValue& Value : Components.Values)
+			{
+				RequestStats.AddLogicalWrite(Value);
+			}
 		}
 		else if (!bReplaceExisting)
 		{
 			if (!Components.Meta)
 			{
 				Components.Meta = Record.GetMeta();
+				RequestStats.LogicalWriteSize += Components.Meta.GetSize();
 			}
 			for (FValueWithId& Value : Components.Values)
 			{
@@ -290,6 +329,7 @@ void FMemoryCacheStore::Put(
 					if (const FValueWithId& NewValue = Record.GetValue(Value.GetId()); NewValue && NewValue.GetRawHash() == Value.GetRawHash())
 					{
 						Value = NewValue;
+						RequestStats.AddLogicalWrite(Value);
 					}
 				}
 			}
@@ -299,8 +339,10 @@ void FMemoryCacheStore::Put(
 			FCacheRecordComponents ExistingComponents = MoveTemp(Components);
 			Components.Meta = Record.GetMeta();
 			Components.Values = Record.GetValues();
+			RequestStats.LogicalWriteSize += Components.Meta.GetSize();
 			for (FValueWithId& Value : Components.Values)
 			{
+				RequestStats.AddLogicalWrite(Value);
 				if (!Value.HasData())
 				{
 					if (const int32 ExistingValueIndex = Algo::BinarySearchBy(ExistingComponents.Values, Value.GetId(), &FValueWithId::GetId);
@@ -313,8 +355,6 @@ void FMemoryCacheStore::Put(
 				}
 			}
 		}
-
-		COOK_STAT(Timer.AddHit(RequiredSize));
 	}
 }
 
@@ -332,10 +372,25 @@ void FMemoryCacheStore::Get(
 	{
 		const FCacheKey& Key = Request.Key;
 		const FCacheRecordPolicy& Policy = Request.Policy;
-		COOK_STAT(auto Timer = EnumHasAnyFlags(Policy.GetRecordPolicy(), ECachePolicy::SkipData) ? UsageStats.TimeProbablyExists() : UsageStats.TimeGet());
 		FCacheRecordComponents Components;
 		EStatus Status = EStatus::Error;
 
+		FRequestStats RequestStats;
+		RequestStats.Bucket = Key.Bucket;
+		RequestStats.Type = ERequestType::Record;
+		RequestStats.Op = ERequestOp::Get;
+		ON_SCOPE_EXIT
+		{
+			RequestStats.Latency = RequestStats.EndTime - RequestStats.StartTime;
+			RequestStats.Status = Status;
+			if (StoreStats)
+			{
+				StoreStats->AddRequest(RequestStats);
+			}
+		};
+
+		FRequestTimer RequestTimer(RequestStats);
+
 		if (bCanBeDisabled && DebugOptions.ShouldSimulateGetMiss(Key))
 		{
 			UE_LOG(LogDerivedDataCache, Verbose, TEXT("%s: Simulated miss for get of %s from '%s'"),
@@ -356,9 +411,10 @@ void FMemoryCacheStore::Get(
 		{
 			const ECachePolicy RecordPolicy = Policy.GetRecordPolicy();
 
-			if (!EnumHasAnyFlags(RecordPolicy, ECachePolicy::SkipMeta))
+			if (!EnumHasAnyFlags(RecordPolicy, ECachePolicy::SkipMeta) && Components.Meta)
 			{
 				Builder.SetMeta(CopyTemp(Components.Meta));
+				RequestStats.LogicalReadSize += Components.Meta.GetSize();
 			}
 
 			for (const FValueWithId& Value : Components.Values)
@@ -378,16 +434,19 @@ void FMemoryCacheStore::Get(
 						break;
 					}
 				}
-				const bool bExistsOnly = EnumHasAnyFlags(ValuePolicy, ECachePolicy::SkipData);
-				Builder.AddValue(bExistsOnly ? Value.RemoveData() : Value);
+				if (EnumHasAnyFlags(ValuePolicy, ECachePolicy::SkipData))
+				{
+					Builder.AddValue(Value.RemoveData());
+				}
+				else
+				{
+					Builder.AddValue(Value);
+					RequestStats.AddLogicalRead(Value);
+				}
 			}
 		}
 
-		FCacheRecord Record = Builder.Build();
-		COOK_STAT(Timer.AddHitOrMiss(
-			Status == EStatus::Ok ? FCookStats::CallStats::EHitOrMiss::Hit : FCookStats::CallStats::EHitOrMiss::Miss,
-			Private::GetCacheRecordCompressedSize(Record)));
-		OnComplete({Request.Name, MoveTemp(Record), Request.UserData, Status});
+		OnComplete({Request.Name, Builder.Build(), Request.UserData, Status});
 	}
 }
 
@@ -424,6 +483,22 @@ void FMemoryCacheStore::PutValue(
 		const FCacheKey& Key = Request.Key;
 		const FValue& Value = Request.Value;
 
+		FRequestStats RequestStats;
+		RequestStats.Bucket = Key.Bucket;
+		RequestStats.Type = ERequestType::Value;
+		RequestStats.Op = ERequestOp::Put;
+		ON_SCOPE_EXIT
+		{
+			RequestStats.Latency = RequestStats.EndTime - RequestStats.StartTime;
+			RequestStats.Status = Status;
+			if (StoreStats)
+			{
+				StoreStats->AddRequest(RequestStats);
+			}
+		};
+
+		FRequestTimer RequestTimer(RequestStats);
+
 		if (!Value.HasData())
 		{
 			continue;
@@ -436,7 +511,6 @@ void FMemoryCacheStore::PutValue(
 			continue;
 		}
 
-		COOK_STAT(auto Timer = UsageStats.TimePut());
 		const int64 ValueSize = Value.GetData().GetCompressedSize();
 		const bool bReplaceExisting = !EnumHasAnyFlags(Request.Policy, ECachePolicy::QueryLocal);
 
@@ -472,7 +546,8 @@ void FMemoryCacheStore::PutValue(
 			CacheValues.Add(Key, Value);
 		}
 
-		COOK_STAT(Timer.AddHit(ValueSize));
+		RequestStats.AddLogicalWrite(Value);
+
 		Status = EStatus::Ok;
 	}
 }
@@ -491,10 +566,26 @@ void FMemoryCacheStore::GetValue(
 	{
 		const FCacheKey& Key = Request.Key;
 		const bool bExistsOnly = EnumHasAllFlags(Request.Policy, ECachePolicy::SkipData);
-		COOK_STAT(auto Timer = bExistsOnly ? UsageStats.TimeProbablyExists() : UsageStats.TimeGet());
 
 		FValue Value;
 		EStatus Status = EStatus::Error;
+
+		FRequestStats RequestStats;
+		RequestStats.Bucket = Key.Bucket;
+		RequestStats.Type = ERequestType::Value;
+		RequestStats.Op = ERequestOp::Get;
+		ON_SCOPE_EXIT
+		{
+			RequestStats.Latency = RequestStats.EndTime - RequestStats.StartTime;
+			RequestStats.Status = Status;
+			if (StoreStats)
+			{
+				StoreStats->AddRequest(RequestStats);
+			}
+		};
+
+		FRequestTimer RequestTimer(RequestStats);
+
 		if (bCanBeDisabled && DebugOptions.ShouldSimulateGetMiss(Key))
 		{
 			UE_LOG(LogDerivedDataCache, Verbose, TEXT("%s: Simulated miss for get of %s from '%s'"),
@@ -507,9 +598,10 @@ void FMemoryCacheStore::GetValue(
 		{
 			Status = EStatus::Ok;
 			Value = !bExistsOnly ? *CacheValue : CacheValue->RemoveData();
-			COOK_STAT(Timer.AddHit(Value.GetData().GetCompressedSize()));
 		}
 
+		RequestStats.AddLogicalRead(Value);
+
 		OnComplete({Request.Name, Request.Key, MoveTemp(Value), Request.UserData, Status});
 	}
 }
@@ -542,9 +634,23 @@ void FMemoryCacheStore::GetChunks(
 	FCompressedBufferReader Reader;
 	for (const FCacheGetChunkRequest& Request : Requests)
 	{
+		FRequestStats RequestStats;
+		RequestStats.Bucket = Request.Key.Bucket;
+		RequestStats.Type = Request.Id.IsNull() ? ERequestType::Value : ERequestType::Record;
+		RequestStats.Op = ERequestOp::GetChunk;
+		ON_SCOPE_EXIT
+		{
+			RequestStats.Latency = RequestStats.EndTime - RequestStats.StartTime;
+			if (StoreStats)
+			{
+				StoreStats->AddRequest(RequestStats);
+			}
+		};
+
+		FRequestTimer RequestTimer(RequestStats);
+
 		bool bProcessHit = false;
 		const bool bExistsOnly = EnumHasAnyFlags(Request.Policy, ECachePolicy::SkipData);
-		COOK_STAT(auto Timer = bExistsOnly ? UsageStats.TimeProbablyExists() : UsageStats.TimeGet());
 		if (bCanBeDisabled && DebugOptions.ShouldSimulateGetMiss(Request.Key))
 		{
 			UE_LOG(LogDerivedDataCache, Verbose, TEXT("%s: Simulated miss for get of %s from '%s'"),
@@ -595,26 +701,28 @@ void FMemoryCacheStore::GetChunks(
 		if (bProcessHit && Request.RawOffset <= Value.GetRawSize())
 		{
 			const uint64 RawSize = FMath::Min(Value.GetRawSize() - Request.RawOffset, Request.RawSize);
-			COOK_STAT(Timer.AddHit(RawSize));
 			FSharedBuffer Buffer;
 			if (Value.HasData() && !bExistsOnly)
 			{
 				Buffer = Reader.Decompress(Request.RawOffset, RawSize);
+				RequestStats.LogicalReadSize += Buffer.GetSize();
 			}
 			const EStatus Status = bExistsOnly || Buffer ? EStatus::Ok : EStatus::Error;
+			RequestStats.Status = Status;
 			OnComplete({Request.Name, Request.Key, Request.Id, Request.RawOffset,
 				RawSize, Value.GetRawHash(), MoveTemp(Buffer), Request.UserData, Status});
 		}
 		else
 		{
+			RequestStats.Status = EStatus::Error;
 			OnComplete(Request.MakeResponse(EStatus::Error));
 		}
 	}
 }
 
-IMemoryCacheStore* CreateMemoryCacheStore(const TCHAR* Name, int64 MaxCacheSize, bool bCanBeDisabled)
+void CreateMemoryCacheStore(IMemoryCacheStore*& OutCache, const TCHAR* Name, ICacheStoreOwner* Owner, int64 MaxCacheSize, bool bCanBeDisabled)
 {
-	return new FMemoryCacheStore(Name, MaxCacheSize, bCanBeDisabled);
+	new FMemoryCacheStore(OutCache, Name, Owner, MaxCacheSize, bCanBeDisabled);
 }
 
 } // UE::DerivedData
