DDC: Fixed the tests of the cache hierarchy

#rnx

diff --git a/Engine/Source/Developer/DerivedDataCache/Private/DerivedDataBackends.cpp b/Engine/Source/Developer/DerivedDataCache/Private/DerivedDataBackends.cpp
--- a/Engine/Source/Developer/DerivedDataCache/Private/DerivedDataBackends.cpp
+++ b/Engine/Source/Developer/DerivedDataCache/Private/DerivedDataBackends.cpp
@@ -61,7 +61,7 @@ int32 AddToAsyncTaskCounter(int32 Addend)
 namespace UE::DerivedData
 {
 
-ILegacyCacheStore* CreateCacheStoreAsync(ILegacyCacheStore* InnerBackend, IMemoryCacheStore* MemoryCache);
+ILegacyCacheStore* CreateCacheStoreAsync(ILegacyCacheStore* InnerBackend, IMemoryCacheStore* MemoryCache, bool bDeleteInnerCache);
 ILegacyCacheStore* CreateCacheStoreHierarchy(ICacheStoreOwner*& OutOwner, TFunctionRef<void (IMemoryCacheStore*&)> MemoryCacheCreator);
 ILegacyCacheStore* CreateCacheStoreThrottle(ILegacyCacheStore* InnerCache, uint32 LatencyMS, uint32 MaxBytesPerSecond);
 ILegacyCacheStore* CreateCacheStoreVerify(ILegacyCacheStore* InnerCache, bool bPutOnError);
@@ -202,7 +202,7 @@ public:
 		}
 
 		// Async must exist in the graph.
-		RootNode = CreateCacheStoreAsync(RootNode, MemoryCache);
+		RootNode = CreateCacheStoreAsync(RootNode, MemoryCache, /*bDeleteInnerCache*/ true);
 
 		// Create a Verify node when using -DDC-Verify[=Type1[@Rate2][+Type2[@Rate2]...]].
 		if (bVerifyFound)
diff --git a/Engine/Source/Developer/DerivedDataCache/Private/DerivedDataCacheStoreAsync.cpp b/Engine/Source/Developer/DerivedDataCache/Private/DerivedDataCacheStoreAsync.cpp
--- a/Engine/Source/Developer/DerivedDataCache/Private/DerivedDataCacheStoreAsync.cpp
+++ b/Engine/Source/Developer/DerivedDataCache/Private/DerivedDataCacheStoreAsync.cpp
@@ -25,7 +25,7 @@ namespace UE::DerivedData
 class FCacheStoreAsync final : public ILegacyCacheStore
 {
 public:
-	FCacheStoreAsync(ILegacyCacheStore* InnerCache, IMemoryCacheStore* MemoryCache);
+	FCacheStoreAsync(ILegacyCacheStore* InnerCache, IMemoryCacheStore* MemoryCache, bool bDeleteInnerCache);
 	~FCacheStoreAsync() final;
 
 	virtual void Put(
@@ -115,18 +115,23 @@ private:
 	ILegacyCacheStore* InnerCache;
 	IMemoryCacheStore* MemoryCache;
 	FDerivedDataCacheUsageStats UsageStats;
+	bool bDeleteInnerCache;
 };
 
-FCacheStoreAsync::FCacheStoreAsync(ILegacyCacheStore* InInnerCache, IMemoryCacheStore* InMemoryCache)
+FCacheStoreAsync::FCacheStoreAsync(ILegacyCacheStore* InInnerCache, IMemoryCacheStore* InMemoryCache, bool bInDeleteInnerCache)
 	: InnerCache(InInnerCache)
 	, MemoryCache(InMemoryCache)
+	, bDeleteInnerCache(bInDeleteInnerCache)
 {
 	check(InnerCache);
 }
 
 FCacheStoreAsync::~FCacheStoreAsync()
 {
-	delete InnerCache;
+	if (bDeleteInnerCache)
+	{
+		delete InnerCache;
+	}
 }
 
 void FCacheStoreAsync::LegacyStats(FDerivedDataCacheStatsNode& OutNode)
@@ -183,9 +188,9 @@ void FCacheStoreAsync::Execute(
 		});
 }
 
-ILegacyCacheStore* CreateCacheStoreAsync(ILegacyCacheStore* InnerCache, IMemoryCacheStore* MemoryCache)
+ILegacyCacheStore* CreateCacheStoreAsync(ILegacyCacheStore* InnerCache, IMemoryCacheStore* MemoryCache, bool bDeleteInnerCache)
 {
-	return new FCacheStoreAsync(InnerCache, MemoryCache);
+	return new FCacheStoreAsync(InnerCache, MemoryCache, bDeleteInnerCache);
 }
 
 } // UE::DerivedData
diff --git a/Engine/Source/Developer/DerivedDataCache/Private/DerivedDataCacheStoreHierarchy.cpp b/Engine/Source/Developer/DerivedDataCache/Private/DerivedDataCacheStoreHierarchy.cpp
--- a/Engine/Source/Developer/DerivedDataCache/Private/DerivedDataCacheStoreHierarchy.cpp
+++ b/Engine/Source/Developer/DerivedDataCache/Private/DerivedDataCacheStoreHierarchy.cpp
@@ -31,7 +31,7 @@
 namespace UE::DerivedData
 {
 
-ILegacyCacheStore* CreateCacheStoreAsync(ILegacyCacheStore* InnerCache, IMemoryCacheStore* MemoryCache);
+ILegacyCacheStore* CreateCacheStoreAsync(ILegacyCacheStore* InnerCache, IMemoryCacheStore* MemoryCache, bool bDeleteInnerCache);
 
 ///////////////////////////////////////////////////////////////////////////////////////////////////////////////////////
 
@@ -371,6 +371,7 @@ FCacheStoreHierarchy::~FCacheStoreHierarchy()
 	for (FCacheStoreNode& Node : ReverseIterate(Nodes))
 	{
 		Node.AsyncCache.Reset();
+		delete Node.Cache;
 	}
 }
 
@@ -379,7 +380,7 @@ void FCacheStoreHierarchy::Add(ILegacyCacheStore* CacheStore, ECacheStoreFlags F
 	FWriteScopeLock Lock(NodesLock);
 	checkf(!Algo::FindBy(Nodes, CacheStore, &FCacheStoreNode::Cache),
 		TEXT("Attempting to add a cache store that was previously registered to the hierarchy."));
-	TUniquePtr<ILegacyCacheStore> AsyncCacheStore(CreateCacheStoreAsync(CacheStore, MemoryCache));
+	TUniquePtr<ILegacyCacheStore> AsyncCacheStore(CreateCacheStoreAsync(CacheStore, MemoryCache, /*bDeleteInnerCache*/ false));
 	Nodes.Add({CacheStore, Flags, {}, MoveTemp(AsyncCacheStore)});
 	UpdateNodeFlags();
 }
diff --git a/Engine/Source/Developer/DerivedDataCache/Tests/DerivedDataCacheStoreHierarchyTest.cpp b/Engine/Source/Developer/DerivedDataCache/Tests/DerivedDataCacheStoreHierarchyTest.cpp
--- a/Engine/Source/Developer/DerivedDataCache/Tests/DerivedDataCacheStoreHierarchyTest.cpp
+++ b/Engine/Source/Developer/DerivedDataCache/Tests/DerivedDataCacheStoreHierarchyTest.cpp
@@ -13,7 +13,7 @@ namespace UE::DerivedData
 {
 
 class IMemoryCacheStore;
-ILegacyCacheStore* CreateCacheStoreHierarchy(ICacheStoreOwner*& OutOwner, IMemoryCacheStore* MemoryCache);
+ILegacyCacheStore* CreateCacheStoreHierarchy(ICacheStoreOwner*& OutOwner, TFunctionRef<void (IMemoryCacheStore*&)> MemoryCacheCreator);
 
 static FValue CreateHierarchyTestValue(uint32 Value)
 {
@@ -23,7 +23,7 @@ static FValue CreateHierarchyTestValue(uint32 Value)
 TEST_CASE("DerivedData::Cache::Hierarchy::PartialRecordPropagation", "[DerivedData]")
 {
 	ICacheStoreOwner* StoreOwner = nullptr;
-	TUniquePtr<ILegacyCacheStore> Hierarchy(CreateCacheStoreHierarchy(StoreOwner, nullptr));
+	TUniquePtr<ILegacyCacheStore> Hierarchy(CreateCacheStoreHierarchy(StoreOwner, [](IMemoryCacheStore*& OutCache) { OutCache = nullptr; }));
 
 	TUniquePtr<ITestCacheStore> Store0(CreateTestCacheStore(ECacheStoreFlags::Local | ECacheStoreFlags::Query | ECacheStoreFlags::Store, ETestCacheStoreFlags::Async));
 	StoreOwner->Add(Store0.Get(), Store0->GetFlags());
@@ -139,7 +139,7 @@ TEST_CASE("DerivedData::Cache::Hierarchy::PartialRecordPropagation", "[DerivedDa
 TEST_CASE("DerivedData::Cache::Hierarchy::PartialNonDeterministicRecordPropagation", "[DerivedData]")
 {
 	ICacheStoreOwner* StoreOwner = nullptr;
-	TUniquePtr<ILegacyCacheStore> Hierarchy(CreateCacheStoreHierarchy(StoreOwner, nullptr));
+	TUniquePtr<ILegacyCacheStore> Hierarchy(CreateCacheStoreHierarchy(StoreOwner, [](IMemoryCacheStore*& OutCache) { OutCache = nullptr; }));
 
 	TUniquePtr<ITestCacheStore> Store0(CreateTestCacheStore(ECacheStoreFlags::Local | ECacheStoreFlags::Query | ECacheStoreFlags::Store, ETestCacheStoreFlags::Async));
 	StoreOwner->Add(Store0.Get(), Store0->GetFlags());
@@ -288,7 +288,7 @@ TEST_CASE("DerivedData::Cache::Hierarchy::PartialNonDeterministicRecordPropagati
 TEST_CASE("DerivedData::Cache::Hierarchy::BlockingRequestOwner", "[DerivedData]")
 {
 	ICacheStoreOwner* StoreOwner = nullptr;
-	TUniquePtr<ILegacyCacheStore> Hierarchy(CreateCacheStoreHierarchy(StoreOwner, nullptr));
+	TUniquePtr<ILegacyCacheStore> Hierarchy(CreateCacheStoreHierarchy(StoreOwner, [](IMemoryCacheStore*& OutCache) { OutCache = nullptr; }));
 
 	TUniquePtr<ITestCacheStore> Store0(CreateTestCacheStore(ECacheStoreFlags::Local | ECacheStoreFlags::Query | ECacheStoreFlags::Store, ETestCacheStoreFlags::Async | ETestCacheStoreFlags::Wait));
 	StoreOwner->Add(Store0.Get(), Store0->GetFlags());
