DDC: Added ICacheStoreStats to the pak file cache and added physical size to memory cache stats

#rnx

diff --git a/Engine/Source/Developer/DerivedDataCache/Private/DerivedDataBackends.cpp b/Engine/Source/Developer/DerivedDataCache/Private/DerivedDataBackends.cpp
--- a/Engine/Source/Developer/DerivedDataCache/Private/DerivedDataBackends.cpp
+++ b/Engine/Source/Developer/DerivedDataCache/Private/DerivedDataBackends.cpp
@@ -68,7 +68,7 @@ ILegacyCacheStore* CreateCacheStoreVerify(ILegacyCacheStore* InnerCache, bool bP
 ILegacyCacheStore* CreateFileSystemCacheStore(const TCHAR* Name, const TCHAR* Config, ICacheStoreOwner& Owner, FString& OutPath);
 TTuple<ILegacyCacheStore*, ECacheStoreFlags> CreateHttpCacheStore(const TCHAR* NodeName, const TCHAR* Config);
 void CreateMemoryCacheStore(IMemoryCacheStore*& OutCache, const TCHAR* Name, ICacheStoreOwner* Owner, int64 MaxCacheSize, bool bCanBeDisabled);
-IPakFileCacheStore* CreatePakFileCacheStore(const TCHAR* Filename, bool bWriting, bool bCompressed);
+IPakFileCacheStore* CreatePakFileCacheStore(const TCHAR* Name, const TCHAR* Filename, bool bWriting, bool bCompressed, ICacheStoreOwner* Owner);
 ILegacyCacheStore* CreateS3CacheStore(const TCHAR* RootManifestPath, const TCHAR* BaseUrl, const TCHAR* Region, const TCHAR* CanaryObjectKey, const TCHAR* CachePath);
 TTuple<ILegacyCacheStore*, ECacheStoreFlags> CreateZenCacheStore(const TCHAR* NodeName, const TCHAR* Config);
 ILegacyCacheStore* TryCreateCacheStoreReplay(ILegacyCacheStore* InnerCache);
@@ -301,11 +301,11 @@ public:
 				}
 				else if (NodeType == TEXT("ReadPak"))
 				{
-					return ParsePak(*NodeName, *Entry, /*bWriting*/ false, ECacheStoreFlags::Local | ECacheStoreFlags::Query | ECacheStoreFlags::StopStore);
+					return ParsePak(*NodeName, *Entry, /*bWriting*/ false);
 				}
 				else if (NodeType == TEXT("WritePak"))
 				{
-					return ParsePak(*NodeName, *Entry, /*bWriting*/ true, ECacheStoreFlags::Local | ECacheStoreFlags::Store);
+					return ParsePak(*NodeName, *Entry, /*bWriting*/ true);
 				}
 				else if (NodeType == TEXT("S3"))
 				{
@@ -343,7 +343,7 @@ public:
 	 * @param bWriting true to create pak interface for writing
 	 * @return Pak file data backend interface instance or nullptr if unsuccessful
 	 */
-	bool ParsePak(const TCHAR* NodeName, const TCHAR* Entry, const bool bWriting, const ECacheStoreFlags Flags)
+	bool ParsePak(const TCHAR* NodeName, const TCHAR* Entry, const bool bWriting)
 	{
 		ILegacyCacheStore* PakNode = nullptr;
 		FString PakFilename;
@@ -367,9 +367,8 @@ public:
 			FGuid Temp = FGuid::NewGuid();
 			ReadPakFilename = PakFilename;
 			WritePakFilename = PakFilename + TEXT(".") + Temp.ToString();
-			WritePakCache = CreatePakFileCacheStore(*WritePakFilename, /*bWriting*/ true, bCompressed);
+			WritePakCache = CreatePakFileCacheStore(NodeName, *WritePakFilename, /*bWriting*/ true, bCompressed, this);
 			PakNode = WritePakCache;
-			Add(WritePakCache, Flags);
 			return true;
 		}
 		else
@@ -381,11 +380,10 @@ public:
 				return false;
 			}
 
-			IPakFileCacheStore* ReadPak = CreatePakFileCacheStore(*PakFilename, /*bWriting*/ false, bCompressed);
+			IPakFileCacheStore* ReadPak = CreatePakFileCacheStore(NodeName, *PakFilename, /*bWriting*/ false, bCompressed, this);
 			ReadPakFilename = PakFilename;
 			PakNode = ReadPak;
 			ReadPakCache.Add(ReadPak);
-			Add(ReadPak, Flags);
 			return true;
 		}
 	}
@@ -760,7 +758,7 @@ public:
 				for(const FString& MergePakName : MergePakList)
 				{
 					TUniquePtr<IPakFileCacheStore> ReadPak(
-						CreatePakFileCacheStore(*FPaths::Combine(*FPaths::GetPath(WritePakFilename), *MergePakName), /*bWriting*/ false, /*bCompressed*/ false));
+						CreatePakFileCacheStore(TEXT("Merge"), *FPaths::Combine(*FPaths::GetPath(WritePakFilename), *MergePakName), /*bWriting*/ false, /*bCompressed*/ false, /*Owner*/ nullptr));
 					WritePakCache->MergeCache(ReadPak.Get());
 				}
 			}
@@ -849,11 +847,9 @@ public:
 		// Assumptions: there's at least one read-only pak backend in the hierarchy
 		// and its parent is a hierarchical backend.
 		IPakFileCacheStore* ReadPak = nullptr;
-		if (Hierarchy && FPlatformFileManager::Get().GetPlatformFile().FileExists(PakFilename))
+		if (FPlatformFileManager::Get().GetPlatformFile().FileExists(PakFilename))
 		{
-			ReadPak = CreatePakFileCacheStore(PakFilename, /*bWriting*/ false, /*bCompressed*/ false);
-
-			Hierarchy->Add(ReadPak, ECacheStoreFlags::Local | ECacheStoreFlags::Query | ECacheStoreFlags::StopStore);
+			ReadPak = CreatePakFileCacheStore(TEXT("Mount"), PakFilename, /*bWriting*/ false, /*bCompressed*/ false, this);
 			ReadPakCache.Add(ReadPak);
 		}
 		else
@@ -876,7 +872,6 @@ public:
 				// Wait until all async requests are complete.
 				WaitForQuiescence(false);
 
-				Hierarchy->RemoveNotSafe(ReadPak);
 				ReadPakCache.RemoveAt(PakIndex);
 				ReadPak->Close();
 				delete ReadPak;
diff --git a/Engine/Source/Developer/DerivedDataCache/Private/MemoryCacheStore.cpp b/Engine/Source/Developer/DerivedDataCache/Private/MemoryCacheStore.cpp
--- a/Engine/Source/Developer/DerivedDataCache/Private/MemoryCacheStore.cpp
+++ b/Engine/Source/Developer/DerivedDataCache/Private/MemoryCacheStore.cpp
@@ -310,9 +310,11 @@ void FMemoryCacheStore::Put(
 			Components.Meta = Record.GetMeta();
 			Components.Values = Record.GetValues();
 			RequestStats.LogicalWriteSize += Components.Meta.GetSize();
+			RequestStats.PhysicalWriteSize += Components.Meta.GetSize();
 			for (const FValue& Value : Components.Values)
 			{
 				RequestStats.AddLogicalWrite(Value);
+				RequestStats.PhysicalWriteSize += Value.GetData().GetCompressedSize();
 			}
 		}
 		else if (!bReplaceExisting)
@@ -321,6 +323,7 @@ void FMemoryCacheStore::Put(
 			{
 				Components.Meta = Record.GetMeta();
 				RequestStats.LogicalWriteSize += Components.Meta.GetSize();
+				RequestStats.PhysicalWriteSize += Components.Meta.GetSize();
 			}
 			for (FValueWithId& Value : Components.Values)
 			{
@@ -330,6 +333,7 @@ void FMemoryCacheStore::Put(
 					{
 						Value = NewValue;
 						RequestStats.AddLogicalWrite(Value);
+						RequestStats.PhysicalWriteSize += Value.GetData().GetCompressedSize();
 					}
 				}
 			}
@@ -340,9 +344,11 @@ void FMemoryCacheStore::Put(
 			Components.Meta = Record.GetMeta();
 			Components.Values = Record.GetValues();
 			RequestStats.LogicalWriteSize += Components.Meta.GetSize();
+			RequestStats.PhysicalWriteSize += Components.Meta.GetSize();
 			for (FValueWithId& Value : Components.Values)
 			{
 				RequestStats.AddLogicalWrite(Value);
+				RequestStats.PhysicalWriteSize += Value.GetData().GetCompressedSize();
 				if (!Value.HasData())
 				{
 					if (const int32 ExistingValueIndex = Algo::BinarySearchBy(ExistingComponents.Values, Value.GetId(), &FValueWithId::GetId);
@@ -415,6 +421,7 @@ void FMemoryCacheStore::Get(
 			{
 				Builder.SetMeta(CopyTemp(Components.Meta));
 				RequestStats.LogicalReadSize += Components.Meta.GetSize();
+				RequestStats.PhysicalReadSize += Components.Meta.GetSize();
 			}
 
 			for (const FValueWithId& Value : Components.Values)
@@ -442,6 +449,7 @@ void FMemoryCacheStore::Get(
 				{
 					Builder.AddValue(Value);
 					RequestStats.AddLogicalRead(Value);
+					RequestStats.PhysicalReadSize += Value.GetData().GetCompressedSize();
 				}
 			}
 		}
@@ -547,6 +555,7 @@ void FMemoryCacheStore::PutValue(
 		}
 
 		RequestStats.AddLogicalWrite(Value);
+		RequestStats.PhysicalWriteSize += Value.GetData().GetCompressedSize();
 
 		Status = EStatus::Ok;
 	}
@@ -601,6 +610,7 @@ void FMemoryCacheStore::GetValue(
 		}
 
 		RequestStats.AddLogicalRead(Value);
+		RequestStats.PhysicalReadSize += Value.GetData().GetCompressedSize();
 
 		OnComplete({Request.Name, Request.Key, MoveTemp(Value), Request.UserData, Status});
 	}
diff --git a/Engine/Source/Developer/DerivedDataCache/Private/PakFileCacheStore.cpp b/Engine/Source/Developer/DerivedDataCache/Private/PakFileCacheStore.cpp
--- a/Engine/Source/Developer/DerivedDataCache/Private/PakFileCacheStore.cpp
+++ b/Engine/Source/Developer/DerivedDataCache/Private/PakFileCacheStore.cpp
@@ -22,6 +22,7 @@
 #include "Misc/FileHelper.h"
 #include "Misc/Paths.h"
 #include "Misc/PathViews.h"
+#include "Misc/ScopeExit.h"
 #include "Misc/ScopeRWLock.h"
 #include "ProfilingDebugging/CookStats.h"
 #include "Serialization/CompactBinary.h"
@@ -42,7 +43,7 @@ namespace UE::DerivedData
 class FPakFileCacheStore : public IPakFileCacheStore
 {
 public:
-	FPakFileCacheStore(const TCHAR* InFilename, bool bInWriting);
+	FPakFileCacheStore(const TCHAR* InName, const TCHAR* InCachePath, bool bInWriting, ICacheStoreOwner* InOwner = nullptr);
 	~FPakFileCacheStore();
 
 	void Close() final;
@@ -109,33 +110,36 @@ public:
 	bool LegacyDebugOptions(FBackendDebugOptions& Options) final;
 
 private:
-	[[nodiscard]] bool PutCacheRecord(FStringView Name, const FCacheRecord& Record, const FCacheRecordPolicy& Policy, uint64& OutWriteSize);
+	[[nodiscard]] bool PutCacheRecord(FStringView Name, const FCacheRecord& Record, const FCacheRecordPolicy& Policy, FRequestStats& Stats);
 
 	[[nodiscard]] FOptionalCacheRecord GetCacheRecordOnly(
 		FStringView Name,
 		const FCacheKey& Key,
-		const FCacheRecordPolicy& Policy);
+		const FCacheRecordPolicy& Policy,
+		FRequestStats& Stats);
 	[[nodiscard]] FOptionalCacheRecord GetCacheRecord(
 		FStringView Name,
 		const FCacheKey& Key,
 		const FCacheRecordPolicy& Policy,
-		EStatus& OutStatus);
+		EStatus& OutStatus,
+		FRequestStats& Stats);
 
-	[[nodiscard]] bool PutCacheValue(FStringView Name, const FCacheKey& Key, const FValue& Value, ECachePolicy Policy, uint64& OutWriteSize);
+	[[nodiscard]] bool PutCacheValue(FStringView Name, const FCacheKey& Key, const FValue& Value, ECachePolicy Policy, FRequestStats& Stats);
 
-	[[nodiscard]] bool GetCacheValueOnly(FStringView Name, const FCacheKey& Key, ECachePolicy Policy, FValue& OutValue);
-	[[nodiscard]] bool GetCacheValue(FStringView Name, const FCacheKey& Key, ECachePolicy Policy, FValue& OutValue);
+	[[nodiscard]] bool GetCacheValueOnly(FStringView Name, const FCacheKey& Key, ECachePolicy Policy, FValue& OutValue, FRequestStats& Stats);
+	[[nodiscard]] bool GetCacheValue(FStringView Name, const FCacheKey& Key, ECachePolicy Policy, FValue& OutValue, FRequestStats& Stats);
 
-	[[nodiscard]] bool PutCacheContent(FStringView Name, const FCompressedBuffer& Content, uint64& OutWriteSize);
+	[[nodiscard]] bool PutCacheContent(FStringView Name, const FCompressedBuffer& Content, FRequestStats& Stats);
 
-	[[nodiscard]] bool GetCacheContentExists(const FCacheKey& Key, const FIoHash& RawHash);
+	[[nodiscard]] bool GetCacheContentExists(const FCacheKey& Key, const FIoHash& RawHash, FRequestStats& Stats);
 	[[nodiscard]] bool GetCacheContent(
 		FStringView Name,
 		const FCacheKey& Key,
 		const FValueId& Id,
 		const FValue& Value,
 		ECachePolicy Policy,
-		FValue& OutValue);
+		FValue& OutValue,
+		FRequestStats& Stats);
 	void GetCacheContent(
 		const FStringView Name,
 		const FCacheKey& Key,
@@ -143,12 +147,13 @@ private:
 		const FValue& Value,
 		const ECachePolicy Policy,
 		FCompressedBufferReader& Reader,
-		TUniquePtr<FArchive>& OutArchive);
+		TUniquePtr<FArchive>& OutArchive,
+		FRequestStats& Stats);
 
-	[[nodiscard]] bool SaveFile(FStringView Path, FStringView DebugName, TFunctionRef<void (FArchive&)> WriteFunction);
-	[[nodiscard]] FSharedBuffer LoadFile(FStringView Path, FStringView DebugName);
-	[[nodiscard]] TUniquePtr<FArchive> OpenFile(FStringBuilderBase& Path, const FStringView DebugName);
-	[[nodiscard]] bool FileExists(FStringView Path);
+	[[nodiscard]] bool SaveFile(FStringView Path, FStringView DebugName, FRequestStats& Stats, TFunctionRef<void (FArchive&)> WriteFunction);
+	[[nodiscard]] FSharedBuffer LoadFile(FStringView Path, FStringView DebugName, FRequestStats& Stats);
+	[[nodiscard]] TUniquePtr<FArchive> OpenFile(FStringBuilderBase& Path, const FStringView DebugName, FRequestStats& Stats);
+	[[nodiscard]] bool FileExists(FStringView Path, FRequestStats& Stats);
 
 private:
 	FDerivedDataCacheUsageStats UsageStats;
@@ -185,6 +190,9 @@ private:
 	/** Maximum total size of compressed data stored within a value package, or a record package with one attachment. */
 	uint64 MaxValueSizeKB = 1024;
 
+	ICacheStoreOwner* StoreOwner = nullptr;
+	ICacheStoreStats* StoreStats = nullptr;
+
 	enum 
 	{
 		/** Magic number to use in header */
@@ -194,10 +202,11 @@ private:
 	friend class IPakFileCacheStore;
 };
 
-FPakFileCacheStore::FPakFileCacheStore(const TCHAR* const InCachePath, const bool bInWriting)
+FPakFileCacheStore::FPakFileCacheStore(const TCHAR* InName, const TCHAR* const InCachePath, const bool bInWriting, ICacheStoreOwner* const InOwner)
 	: bWriting(bInWriting)
 	, bClosed(false)
 	, CachePath(InCachePath)
+	, StoreOwner(InOwner)
 {
 	IPlatformFile& PlatformFile = FPlatformFileManager::Get().GetPlatformFile();
 	if (bWriting)
@@ -233,6 +242,13 @@ FPakFileCacheStore::FPakFileCacheStore(const TCHAR* const InCachePath, const boo
 				*CachePath, FileHandle->Size() / 1024 / 1024);
 		}
 	}
+
+	if (StoreOwner)
+	{
+		const ECacheStoreFlags Flags = ECacheStoreFlags::Local | (bWriting ? ECacheStoreFlags::Store : ECacheStoreFlags::Query | ECacheStoreFlags::StopStore);
+		StoreOwner->Add(this, Flags);
+		StoreStats = StoreOwner->CreateStats(this, Flags, TEXT("Pak File"), InName, CachePath);
+	}
 }
 
 FPakFileCacheStore::~FPakFileCacheStore()
@@ -254,6 +270,13 @@ void FPakFileCacheStore::Close()
 		CacheItems.Empty();
 		bClosed = true;
 	}
+
+	if (StoreOwner)
+	{
+		StoreOwner->DestroyStats(StoreStats);
+		StoreOwner->RemoveNotSafe(this);
+		StoreOwner = nullptr;
+	}
 }
 
 bool FPakFileCacheStore::CachedDataProbablyExists(const TCHAR* CacheKey)
@@ -568,10 +591,10 @@ void FPakFileCacheStore::MergeCache(IPakFileCacheStore* OtherPakInterface)
 bool IPakFileCacheStore::SortAndCopy(const FString &InputFilename, const FString &OutputFilename)
 {
 	// Open the input and output files
-	FPakFileCacheStore InputPak(*InputFilename, false);
+	FPakFileCacheStore InputPak(TEXT("Input"), *InputFilename, false);
 	if (InputPak.bClosed) return false;
 
-	FPakFileCacheStore OutputPak(*OutputFilename, true);
+	FPakFileCacheStore OutputPak(TEXT("Output"), *OutputFilename, true);
 	if (OutputPak.bClosed) return false;
 
 	// Sort the key names
@@ -606,8 +629,7 @@ bool IPakFileCacheStore::SortAndCopy(const FString &InputFilename, const FString
 
 void FPakFileCacheStore::LegacyStats(FDerivedDataCacheStatsNode& OutNode)
 {
-	OutNode = {TEXT("PakFile"), CachePath, /*bIsLocal*/ true};
-	OutNode.UsageStats.Add(TEXT(""), UsageStats);
+	checkNoEntry();
 }
 
 bool FPakFileCacheStore::LegacyDebugOptions(FBackendDebugOptions& InOptions)
@@ -624,22 +646,26 @@ void FPakFileCacheStore::Put(
 	for (const FCachePutRequest& Request : Requests)
 	{
 		bool bOk;
+		FRequestStats RequestStats;
+		RequestStats.Bucket = Request.Record.GetKey().Bucket;
+		RequestStats.Type = ERequestType::Record;
+		RequestStats.Op = ERequestOp::Put;
 		{
 			const FCacheRecord& Record = Request.Record;
 			TRACE_CPUPROFILER_EVENT_SCOPE(PakFileDDC_Put);
-			COOK_STAT(auto Timer = UsageStats.TimePut());
-			uint64 WriteSize = 0;
-			bOk = PutCacheRecord(Request.Name, Record, Request.Policy, WriteSize);
+			FRequestTimer RequestTimer(RequestStats);
+			bOk = PutCacheRecord(Request.Name, Record, Request.Policy, RequestStats);
 			if (bOk)
 			{
 				UE_LOG(LogDerivedDataCache, Verbose, TEXT("%s: Cache put complete for %s from '%s'"),
 					*CachePath, *WriteToString<96>(Record.GetKey()), *Request.Name);
-				if (WriteSize)
-				{
-					COOK_STAT(Timer.AddHit(WriteSize));
-				}
 			}
 		}
+		RequestStats.Status = bOk ? EStatus::Ok : EStatus::Error;
+		if (StoreStats)
+		{
+			StoreStats->AddRequest(RequestStats);
+		}
 		OnComplete(Request.MakeResponse(bOk ? EStatus::Ok : EStatus::Error));
 	}
 }
@@ -653,20 +679,29 @@ void FPakFileCacheStore::Get(
 	{
 		EStatus Status = EStatus::Error;
 		FOptionalCacheRecord Record;
+		FRequestStats RequestStats;
+		RequestStats.Bucket = Request.Key.Bucket;
+		RequestStats.Type = ERequestType::Record;
+		RequestStats.Op = ERequestOp::Get;
 		{
 			TRACE_CPUPROFILER_EVENT_SCOPE(PakFileDDC_Get);
-			COOK_STAT(auto Timer = UsageStats.TimeGet());
-			if ((Record = GetCacheRecord(Request.Name, Request.Key, Request.Policy, Status)))
+			FRequestTimer RequestTimer(RequestStats);
+			if ((Record = GetCacheRecord(Request.Name, Request.Key, Request.Policy, Status, RequestStats)))
 			{
 				UE_LOG(LogDerivedDataCache, Verbose, TEXT("%s: Cache hit for %s from '%s'"),
 					*CachePath, *WriteToString<96>(Request.Key), *Request.Name);
-				COOK_STAT(Timer.AddHit(Private::GetCacheRecordCompressedSize(Record.Get())));
 			}
 			else
 			{
 				Record = FCacheRecordBuilder(Request.Key).Build();
 			}
 		}
+		RequestStats.AddLogicalRead(Record.Get());
+		RequestStats.Status = Status;
+		if (StoreStats)
+		{
+			StoreStats->AddRequest(RequestStats);
+		}
 		OnComplete({Request.Name, MoveTemp(Record).Get(), Request.UserData, Status});
 	}
 }
@@ -679,21 +714,25 @@ void FPakFileCacheStore::PutValue(
 	for (const FCachePutValueRequest& Request : Requests)
 	{
 		bool bOk;
+		FRequestStats RequestStats;
+		RequestStats.Bucket = Request.Key.Bucket;
+		RequestStats.Type = ERequestType::Value;
+		RequestStats.Op = ERequestOp::Put;
 		{
 			TRACE_CPUPROFILER_EVENT_SCOPE(PakFileDDC_PutValue);
-			COOK_STAT(auto Timer = UsageStats.TimePut());
-			uint64 WriteSize = 0;
-			bOk = PutCacheValue(Request.Name, Request.Key, Request.Value, Request.Policy, WriteSize);
+			FRequestTimer RequestTimer(RequestStats);
+			bOk = PutCacheValue(Request.Name, Request.Key, Request.Value, Request.Policy, RequestStats);
 			if (bOk)
 			{
 				UE_LOG(LogDerivedDataCache, Verbose, TEXT("%s: Cache put complete for %s from '%s'"),
 					*CachePath, *WriteToString<96>(Request.Key), *Request.Name);
-				if (WriteSize)
-				{
-					COOK_STAT(Timer.AddHit(WriteSize));
-				}
 			}
 		}
+		RequestStats.Status = bOk ? EStatus::Ok : EStatus::Error;
+		if (StoreStats)
+		{
+			StoreStats->AddRequest(RequestStats);
+		}
 		OnComplete(Request.MakeResponse(bOk ? EStatus::Ok : EStatus::Error));
 	}
 }
@@ -707,17 +746,26 @@ void FPakFileCacheStore::GetValue(
 	{
 		bool bOk;
 		FValue Value;
+		FRequestStats RequestStats;
+		RequestStats.Bucket = Request.Key.Bucket;
+		RequestStats.Type = ERequestType::Value;
+		RequestStats.Op = ERequestOp::Get;
 		{
 			TRACE_CPUPROFILER_EVENT_SCOPE(PakFileDDC_GetValue);
-			COOK_STAT(auto Timer = UsageStats.TimeGet());
-			bOk = GetCacheValue(Request.Name, Request.Key, Request.Policy, Value);
+			FRequestTimer RequestTimer(RequestStats);
+			bOk = GetCacheValue(Request.Name, Request.Key, Request.Policy, Value, RequestStats);
 			if (bOk)
 			{
 				UE_LOG(LogDerivedDataCache, Verbose, TEXT("%s: Cache hit for %s from '%s'"),
 					*CachePath, *WriteToString<96>(Request.Key), *Request.Name);
-				COOK_STAT(Timer.AddHit(Value.GetData().GetCompressedSize()));
 			}
 		}
+		RequestStats.AddLogicalRead(Value);
+		RequestStats.Status = bOk ? EStatus::Ok : EStatus::Error;
+		if (StoreStats)
+		{
+			StoreStats->AddRequest(RequestStats);
+		}
 		OnComplete({Request.Name, Request.Key, Value, Request.UserData, bOk ? EStatus::Ok : EStatus::Error});
 	}
 }
@@ -742,10 +790,14 @@ void FPakFileCacheStore::GetChunks(
 		EStatus Status = EStatus::Error;
 		FSharedBuffer Buffer;
 		uint64 RawSize = 0;
+		FRequestStats RequestStats;
+		RequestStats.Bucket = Request.Key.Bucket;
+		RequestStats.Type = Request.Id.IsNull() ? ERequestType::Value : ERequestType::Record;
+		RequestStats.Op = ECacheStoreRequestOp::GetChunk;
 		{
 			TRACE_CPUPROFILER_EVENT_SCOPE(PakFileDDC_Get);
+			FRequestTimer RequestTimer(RequestStats);
 			const bool bExistsOnly = EnumHasAnyFlags(Request.Policy, ECachePolicy::SkipData);
-			COOK_STAT(auto Timer = bExistsOnly ? UsageStats.TimeProbablyExists() : UsageStats.TimeGet());
 			if (!(bHasValue && ValueKey == Request.Key && ValueId == Request.Id) || ValueReader.HasSource() < !bExistsOnly)
 			{
 				ValueReader.ResetSource();
@@ -761,7 +813,7 @@ void FPakFileCacheStore::GetChunks(
 						FCacheRecordPolicyBuilder PolicyBuilder(ECachePolicy::None);
 						PolicyBuilder.AddValuePolicy(Request.Id, Request.Policy);
 						Record.Reset();
-						Record = GetCacheRecordOnly(Request.Name, Request.Key, PolicyBuilder.Build());
+						Record = GetCacheRecordOnly(Request.Name, Request.Key, PolicyBuilder.Build(), RequestStats);
 					}
 					if (Record)
 					{
@@ -771,17 +823,17 @@ void FPakFileCacheStore::GetChunks(
 							Value = ValueWithId;
 							ValueId = Request.Id;
 							ValueKey = Request.Key;
-							GetCacheContent(Request.Name, Request.Key, ValueId, Value, Request.Policy, ValueReader, ValueAr);
+							GetCacheContent(Request.Name, Request.Key, ValueId, Value, Request.Policy, ValueReader, ValueAr, RequestStats);
 						}
 					}
 				}
 				else
 				{
 					ValueKey = Request.Key;
-					bHasValue = GetCacheValueOnly(Request.Name, Request.Key, Request.Policy, Value);
+					bHasValue = GetCacheValueOnly(Request.Name, Request.Key, Request.Policy, Value, RequestStats);
 					if (bHasValue)
 					{
-						GetCacheContent(Request.Name, Request.Key, Request.Id, Value, Request.Policy, ValueReader, ValueAr);
+						GetCacheContent(Request.Name, Request.Key, Request.Id, Value, Request.Policy, ValueReader, ValueAr, RequestStats);
 					}
 				}
 			}
@@ -791,14 +843,19 @@ void FPakFileCacheStore::GetChunks(
 				RawSize = FMath::Min(Value.GetRawSize() - RawOffset, Request.RawSize);
 				UE_LOG(LogDerivedDataCache, Verbose, TEXT("%s: Cache hit for %s from '%s'"),
 					*CachePath, *WriteToString<96>(Request.Key, '/', Request.Id), *Request.Name);
-				COOK_STAT(Timer.AddHit(!bExistsOnly ? RawSize : 0));
 				if (!bExistsOnly)
 				{
 					Buffer = ValueReader.Decompress(RawOffset, RawSize);
+					RequestStats.LogicalReadSize += Buffer.GetSize();
 				}
 				Status = bExistsOnly || Buffer.GetSize() == RawSize ? EStatus::Ok : EStatus::Error;
 			}
 		}
+		RequestStats.Status = Status;
+		if (StoreStats)
+		{
+			StoreStats->AddRequest(RequestStats);
+		}
 		OnComplete({Request.Name, Request.Key, Request.Id, Request.RawOffset,
 			RawSize, Value.GetRawHash(), MoveTemp(Buffer), Request.UserData, Status});
 	}
@@ -808,7 +865,7 @@ bool FPakFileCacheStore::PutCacheRecord(
 	const FStringView Name,
 	const FCacheRecord& Record,
 	const FCacheRecordPolicy& Policy,
-	uint64& OutWriteSize)
+	FRequestStats& Stats)
 {
 	if (!bWriting || bClosed)
 	{
@@ -838,10 +895,10 @@ bool FPakFileCacheStore::PutCacheRecord(
 
 	// Check if there is an existing record package.
 	bool bReplaceExisting = !EnumHasAnyFlags(RecordPolicy, ECachePolicy::QueryLocal);
-	bool bSaveRecord = bReplaceExisting;
+	bool bSavePackage = bReplaceExisting;
 	if (!bReplaceExisting)
 	{
-		bSaveRecord |= !FileExists(Path);
+		bSavePackage |= !FileExists(Path, Stats);
 	}
 
 	// Serialize the record to a package and remove attachments that will be stored externally.
@@ -853,19 +910,26 @@ bool FPakFileCacheStore::PutCacheRecord(
 	// Save the external content to storage.
 	for (FCompressedBuffer& Content : ExternalContent)
 	{
-		uint64 WriteSize = 0;
-		if (!PutCacheContent(Name, Content, WriteSize))
+		if (!PutCacheContent(Name, Content, Stats))
 		{
 			return false;
 		}
-		OutWriteSize += WriteSize;
 	}
 
 	// Save the record package to storage.
-	const auto WriteRecord = [&](FArchive& Ar) { Package.Save(Ar); OutWriteSize += uint64(Ar.TotalSize()); };
-	if (bSaveRecord && !SaveFile(Path, Name, WriteRecord))
+	const auto WritePackage = [&Package](FArchive& Ar) { Package.Save(Ar); };
+	if (bSavePackage)
 	{
-		return false;
+		if (!SaveFile(Path, Name, Stats, WritePackage))
+		{
+			return false;
+		}
+		if (const FCbObject& Meta = Record.GetMeta())
+		{
+			Stats.LogicalWriteSize += Meta.GetSize();
+		}
+		Stats.LogicalWriteSize += Algo::TransformAccumulate(Package.GetAttachments(),
+			[](const FCbAttachment& Attachment) { return Attachment.AsCompressedBinary().GetRawSize(); }, uint64(0));
 	}
 
 	return true;
@@ -874,7 +938,8 @@ bool FPakFileCacheStore::PutCacheRecord(
 FOptionalCacheRecord FPakFileCacheStore::GetCacheRecordOnly(
 	const FStringView Name,
 	const FCacheKey& Key,
-	const FCacheRecordPolicy& Policy)
+	const FCacheRecordPolicy& Policy,
+	FRequestStats& Stats)
 {
 	if (bClosed)
 	{
@@ -903,7 +968,7 @@ FOptionalCacheRecord FPakFileCacheStore::GetCacheRecordOnly(
 	FPathViews::Append(Path, TEXT("Buckets"), Key);
 
 	// Request the record from storage.
-	FSharedBuffer Buffer = LoadFile(Path, Name);
+	FSharedBuffer Buffer = LoadFile(Path, Name, Stats);
 	if (Buffer.IsNull())
 	{
 		UE_LOG(LogDerivedDataCache, Verbose, TEXT("%s: Cache miss with missing record for %s from '%.*s'"),
@@ -927,27 +992,28 @@ FOptionalCacheRecord FPakFileCacheStore::GetCacheRecordOnly(
 		{
 			UE_LOG(LogDerivedDataCache, Display, TEXT("%s: Cache miss with package load failure for %s from '%.*s'"),
 				*CachePath, *WriteToString<96>(Key), Name.Len(), Name.GetData());
-			return FOptionalCacheRecord();
+			return Record;
 		}
 		Record = FCacheRecord::Load(Package);
 		if (Record.IsNull())
 		{
 			UE_LOG(LogDerivedDataCache, Display, TEXT("%s: Cache miss with record load failure for %s from '%.*s'"),
 				*CachePath, *WriteToString<96>(Key), Name.Len(), Name.GetData());
-			return FOptionalCacheRecord();
+			return Record;
 		}
 	}
 
-	return Record.Get();
+	return Record;
 }
 
 FOptionalCacheRecord FPakFileCacheStore::GetCacheRecord(
 	const FStringView Name,
 	const FCacheKey& Key,
 	const FCacheRecordPolicy& Policy,
-	EStatus& OutStatus)
+	EStatus& OutStatus,
+	FRequestStats& Stats)
 {
-	FOptionalCacheRecord Record = GetCacheRecordOnly(Name, Key, Policy);
+	FOptionalCacheRecord Record = GetCacheRecordOnly(Name, Key, Policy, Stats);
 	if (Record.IsNull())
 	{
 		OutStatus = EStatus::Error;
@@ -969,7 +1035,7 @@ FOptionalCacheRecord FPakFileCacheStore::GetCacheRecord(
 		const FValueId& Id = Value.GetId();
 		const ECachePolicy ValuePolicy = Policy.GetValuePolicy(Id);
 		FValue Content;
-		if (GetCacheContent(Name, Key, Id, Value, ValuePolicy, Content))
+		if (GetCacheContent(Name, Key, Id, Value, ValuePolicy, Content, Stats))
 		{
 			RecordBuilder.AddValue(Id, MoveTemp(Content));
 		}
@@ -993,7 +1059,7 @@ bool FPakFileCacheStore::PutCacheValue(
 	const FCacheKey& Key,
 	const FValue& Value,
 	const ECachePolicy Policy,
-	uint64& OutWriteSize)
+	FRequestStats& Stats)
 {
 	if (!bWriting || bClosed)
 	{
@@ -1022,7 +1088,7 @@ bool FPakFileCacheStore::PutCacheValue(
 	const bool bReplaceExisting = !EnumHasAnyFlags(Policy, ECachePolicy::QueryLocal);
 	if (!bReplaceExisting)
 	{
-		bValueExists = FileExists(Path);
+		bValueExists = FileExists(Path, Stats);
 	}
 
 	// Save the value to a package and save the data to external content.
@@ -1038,7 +1104,7 @@ bool FPakFileCacheStore::PutCacheValue(
 		if (!Value.HasData())
 		{
 			// Verify that the content exists in storage.
-			if (!GetCacheContentExists(Key, Value.GetRawHash()))
+			if (!GetCacheContentExists(Key, Value.GetRawHash(), Stats))
 			{
 				UE_LOG(LogDerivedDataCache, Verbose, TEXT("%s: Failed due to missing data for put of %s from '%.*s'"),
 					*CachePath, *WriteToString<96>(Key), Name.Len(), Name.GetData());
@@ -1048,17 +1114,15 @@ bool FPakFileCacheStore::PutCacheValue(
 		else
 		{
 			// Save the external content to storage.
-			uint64 WriteSize = 0;
-			if (!PutCacheContent(Name, Value.GetData(), WriteSize))
+			if (!PutCacheContent(Name, Value.GetData(), Stats))
 			{
 				return false;
 			}
-			OutWriteSize += WriteSize;
 		}
 
 		// Save the value package to storage.
-		const auto WritePackage = [&](FArchive& Ar) { Package.Save(Ar); OutWriteSize += uint64(Ar.TotalSize()); };
-		if (!SaveFile(Path, Name, WritePackage))
+		const auto WritePackage = [&Package](FArchive& Ar) { Package.Save(Ar); };
+		if (!SaveFile(Path, Name, Stats, WritePackage))
 		{
 			return false;
 		}
@@ -1071,7 +1135,8 @@ bool FPakFileCacheStore::GetCacheValueOnly(
 	const FStringView Name,
 	const FCacheKey& Key,
 	const ECachePolicy Policy,
-	FValue& OutValue)
+	FValue& OutValue,
+	FRequestStats& Stats)
 {
 	if (bClosed)
 	{
@@ -1100,7 +1165,7 @@ bool FPakFileCacheStore::GetCacheValueOnly(
 	FPathViews::Append(Path, TEXT("Buckets"), Key);
 
 	// Request the value package from storage.
-	FSharedBuffer Buffer = LoadFile(Path, Name);
+	FSharedBuffer Buffer = LoadFile(Path, Name, Stats);
 	if (Buffer.IsNull())
 	{
 		UE_LOG(LogDerivedDataCache, Verbose, TEXT("%s: Cache miss with missing value for %s from '%.*s'"),
@@ -1142,31 +1207,33 @@ bool FPakFileCacheStore::GetCacheValue(
 	const FStringView Name,
 	const FCacheKey& Key,
 	const ECachePolicy Policy,
-	FValue& OutValue)
+	FValue& OutValue,
+	FRequestStats& Stats)
 {
-	return GetCacheValueOnly(Name, Key, Policy, OutValue) && GetCacheContent(Name, Key, {}, OutValue, Policy, OutValue);
+	return GetCacheValueOnly(Name, Key, Policy, OutValue, Stats) && GetCacheContent(Name, Key, {}, OutValue, Policy, OutValue, Stats);
 }
 
-bool FPakFileCacheStore::PutCacheContent(const FStringView Name, const FCompressedBuffer& Content, uint64& OutWriteSize)
+bool FPakFileCacheStore::PutCacheContent(const FStringView Name, const FCompressedBuffer& Content, FRequestStats& Stats)
 {
 	const FIoHash& RawHash = Content.GetRawHash();
 	TStringBuilder<256> Path;
 	FPathViews::Append(Path, TEXT("Content"), RawHash);
-	if (!FileExists(Path))
+	if (!FileExists(Path, Stats))
 	{
-		if (!SaveFile(Path, Name, [&Content, &OutWriteSize](FArchive& Ar) { Content.Save(Ar); OutWriteSize += uint64(Ar.TotalSize()); }))
+		if (!SaveFile(Path, Name, Stats, [&Content](FArchive& Ar) { Content.Save(Ar); }))
 		{
 			return false;
 		}
+		Stats.LogicalWriteSize += Content.GetRawSize();
 	}
 	return true;
 }
 
-bool FPakFileCacheStore::GetCacheContentExists(const FCacheKey& Key, const FIoHash& RawHash)
+bool FPakFileCacheStore::GetCacheContentExists(const FCacheKey& Key, const FIoHash& RawHash, FRequestStats& Stats)
 {
 	TStringBuilder<256> Path;
 	FPathViews::Append(Path, TEXT("Buckets"), Key);
-	return FileExists(Path);
+	return FileExists(Path, Stats);
 }
 
 bool FPakFileCacheStore::GetCacheContent(
@@ -1175,7 +1242,8 @@ bool FPakFileCacheStore::GetCacheContent(
 	const FValueId& Id,
 	const FValue& Value,
 	const ECachePolicy Policy,
-	FValue& OutValue)
+	FValue& OutValue,
+	FRequestStats& Stats)
 {
 	if (!EnumHasAnyFlags(Policy, ECachePolicy::Query))
 	{
@@ -1195,7 +1263,7 @@ bool FPakFileCacheStore::GetCacheContent(
 	FPathViews::Append(Path, TEXT("Content"), RawHash);
 	if (EnumHasAnyFlags(Policy, ECachePolicy::SkipData))
 	{
-		if (FileExists(Path))
+		if (FileExists(Path, Stats))
 		{
 			OutValue = Value;
 			return true;
@@ -1203,7 +1271,7 @@ bool FPakFileCacheStore::GetCacheContent(
 	}
 	else
 	{
-		if (FSharedBuffer CompressedData = LoadFile(Path, Name))
+		if (FSharedBuffer CompressedData = LoadFile(Path, Name, Stats))
 		{
 			if (FCompressedBuffer CompressedBuffer = FCompressedBuffer::FromCompressed(MoveTemp(CompressedData));
 				CompressedBuffer && CompressedBuffer.GetRawHash() == RawHash)
@@ -1233,8 +1301,29 @@ void FPakFileCacheStore::GetCacheContent(
 	const FValue& Value,
 	const ECachePolicy Policy,
 	FCompressedBufferReader& Reader,
-	TUniquePtr<FArchive>& OutArchive)
+	TUniquePtr<FArchive>& OutArchive,
+	FRequestStats& Stats)
 {
+	class FStatsArchive final : TUniquePtr<FArchive>, public FArchiveProxy
+	{
+	public:
+		FStatsArchive(FArchive& InArchive, FRequestStats& InStats)
+			: TUniquePtr<FArchive>(&InArchive)
+			, FArchiveProxy(InArchive)
+			, Stats(InStats)
+		{
+		}
+
+		void Serialize(void* V, int64 Length) final
+		{
+			Stats.PhysicalReadSize += uint64(Length);
+			FArchiveProxy::Serialize(V, Length);
+		}
+
+	private:
+		FRequestStats& Stats;
+	};
+
 	if (!EnumHasAnyFlags(Policy, ECachePolicy::Query))
 	{
 		return;
@@ -1256,16 +1345,17 @@ void FPakFileCacheStore::GetCacheContent(
 	FPathViews::Append(Path, TEXT("Content"), RawHash);
 	if (EnumHasAllFlags(Policy, ECachePolicy::SkipData))
 	{
-		if (FileExists(Path))
+		if (FileExists(Path, Stats))
 		{
 			return;
 		}
 	}
 	else
 	{
-		OutArchive = OpenFile(Path, Name);
+		OutArchive = OpenFile(Path, Name, Stats);
 		if (OutArchive)
 		{
+			OutArchive.Reset(new FStatsArchive(*OutArchive.Release(), Stats));
 			Reader.SetSource(*OutArchive);
 			if (Reader.GetRawHash() == RawHash)
 			{
@@ -1374,6 +1464,7 @@ private:
 bool FPakFileCacheStore::SaveFile(
 	const FStringView Path,
 	const FStringView DebugName,
+	FRequestStats& Stats,
 	TFunctionRef<void (FArchive&)> WriteFunction)
 {
 	FWriteScopeLock ScopeLock(SynchronizationObject);
@@ -1389,13 +1480,14 @@ bool FPakFileCacheStore::SaveFile(
 			UE_LOG(LogDerivedDataCache, VeryVerbose,
 				TEXT("%s: File %.*s from '%.*s' written with offset %" INT64_FMT ", size %" INT64_FMT", CRC 0x%08x."),
 				*CachePath, Path.Len(), Path.GetData(), DebugName.Len(), DebugName.GetData(), Item.Offset, Item.Size, Item.Crc);
+			Stats.PhysicalWriteSize += uint64(Item.Size);
 			return true;
 		}
 	}
 	return false;
 }
 
-FSharedBuffer FPakFileCacheStore::LoadFile(const FStringView Path, const FStringView DebugName)
+FSharedBuffer FPakFileCacheStore::LoadFile(const FStringView Path, const FStringView DebugName, FRequestStats& Stats)
 {
 	FWriteScopeLock ScopeLock(SynchronizationObject);
 	if (const FCacheValue* Item = CacheItems.FindByHash(GetTypeHash(Path), Path))
@@ -1429,6 +1521,7 @@ FSharedBuffer FPakFileCacheStore::LoadFile(const FStringView Path, const FString
 			}
 			else
 			{
+				Stats.PhysicalReadSize += uint64(Item->Size);
 				return MutableBuffer.MoveToShared();
 			}
 		}
@@ -1436,8 +1529,11 @@ FSharedBuffer FPakFileCacheStore::LoadFile(const FStringView Path, const FString
 	return FSharedBuffer();
 }
 
-TUniquePtr<FArchive> FPakFileCacheStore::OpenFile(FStringBuilderBase& Path, const FStringView DebugName)
+TUniquePtr<FArchive> FPakFileCacheStore::OpenFile(FStringBuilderBase& Path, const FStringView DebugName, FRequestStats& Stats)
 {
+	const FMonotonicTimePoint StartTime = FMonotonicTimePoint::Now();
+	ON_SCOPE_EXIT { Stats.AddLatency(FMonotonicTimePoint::Now() - StartTime); };
+
 	FReadScopeLock ScopeLock(SynchronizationObject);
 	const FStringView PathView(Path);
 	if (const FCacheValue* Item = CacheItems.FindByHash(GetTypeHash(PathView), PathView))
@@ -1452,7 +1548,7 @@ TUniquePtr<FArchive> FPakFileCacheStore::OpenFile(FStringBuilderBase& Path, cons
 	return nullptr;
 }
 
-bool FPakFileCacheStore::FileExists(const FStringView Path)
+bool FPakFileCacheStore::FileExists(const FStringView Path, FRequestStats& Stats)
 {
 	FReadScopeLock ScopeLock(SynchronizationObject);
 	const uint32 PathHash = GetTypeHash(Path);
@@ -1462,7 +1558,7 @@ bool FPakFileCacheStore::FileExists(const FStringView Path)
 class FCompressedPakFileCacheStore final : public FPakFileCacheStore
 {
 public:
-	FCompressedPakFileCacheStore(const TCHAR* InFilename, bool bInWriting);
+	using FPakFileCacheStore::FPakFileCacheStore;
 
 	void Put(
 		TConstArrayView<FCachePutRequest> Requests,
@@ -1481,11 +1577,6 @@ private:
 	static FValue Compress(const FValue& Value);
 };
 
-FCompressedPakFileCacheStore::FCompressedPakFileCacheStore(const TCHAR* InFilename, bool bInWriting)
-	: FPakFileCacheStore(InFilename, bInWriting)
-{
-}
-
 void FCompressedPakFileCacheStore::Put(
 	const TConstArrayView<FCachePutRequest> Requests,
 	IRequestOwner& Owner,
@@ -1549,9 +1640,9 @@ FValue FCompressedPakFileCacheStore::Compress(const FValue& Value)
 	return FValue(FCompressedBuffer::Compress(Data, RequiredCompressor, MinRequiredCompressionLevel, BlockSize));
 }
 
-IPakFileCacheStore* CreatePakFileCacheStore(const TCHAR* Filename, bool bWriting, bool bCompressed)
+IPakFileCacheStore* CreatePakFileCacheStore(const TCHAR* Name, const TCHAR* CachePath, bool bWriting, bool bCompressed, ICacheStoreOwner* Owner)
 {
-	return bCompressed ? new FCompressedPakFileCacheStore(Filename, bWriting) : new FPakFileCacheStore(Filename, bWriting);
+	return bCompressed ? new FCompressedPakFileCacheStore(Name, CachePath, bWriting, Owner) : new FPakFileCacheStore(Name, CachePath, bWriting, Owner);
 }
 
 } // UE::DerivedData
