DDC: Added ICacheStoreStats to the Zen cache

#rnx

diff --git a/Engine/Source/Developer/DerivedDataCache/Private/BatchView.h b/Engine/Source/Developer/DerivedDataCache/Private/BatchView.h
--- a/Engine/Source/Developer/DerivedDataCache/Private/BatchView.h
+++ b/Engine/Source/Developer/DerivedDataCache/Private/BatchView.h
@@ -23,6 +23,8 @@ template <typename T>
 class TBatchView
 {
 public:
+	TBatchView() = default;
+
 	/**
 	 * Constructor from list of elements and a function that is called on each element
 	 * to decide whether to add the element to the current batch.
diff --git a/Engine/Source/Developer/DerivedDataCache/Private/DerivedDataBackends.cpp b/Engine/Source/Developer/DerivedDataCache/Private/DerivedDataBackends.cpp
--- a/Engine/Source/Developer/DerivedDataCache/Private/DerivedDataBackends.cpp
+++ b/Engine/Source/Developer/DerivedDataCache/Private/DerivedDataBackends.cpp
@@ -70,7 +70,7 @@ TTuple<ILegacyCacheStore*, ECacheStoreFlags> CreateHttpCacheStore(const TCHAR* N
 void CreateMemoryCacheStore(IMemoryCacheStore*& OutCache, const TCHAR* Name, ICacheStoreOwner* Owner, int64 MaxCacheSize, bool bCanBeDisabled);
 IPakFileCacheStore* CreatePakFileCacheStore(const TCHAR* Name, const TCHAR* Filename, bool bWriting, bool bCompressed, ICacheStoreOwner* Owner);
 ILegacyCacheStore* CreateS3CacheStore(const TCHAR* RootManifestPath, const TCHAR* BaseUrl, const TCHAR* Region, const TCHAR* CanaryObjectKey, const TCHAR* CachePath);
-TTuple<ILegacyCacheStore*, ECacheStoreFlags> CreateZenCacheStore(const TCHAR* NodeName, const TCHAR* Config);
+ILegacyCacheStore* CreateZenCacheStore(const TCHAR* NodeName, const TCHAR* Config, ICacheStoreOwner* Owner);
 ILegacyCacheStore* TryCreateCacheStoreReplay(ILegacyCacheStore* InnerCache);
 
 /**
@@ -322,10 +322,8 @@ public:
 				}
 				else if (NodeType == TEXT("Zen"))
 				{
-					TTuple<ILegacyCacheStore*, ECacheStoreFlags> Node = CreateZenCacheStore(*NodeName, *Entry);
-					if (Node.Key)
+					if (CreateZenCacheStore(*NodeName, *Entry, this))
 					{
-						NodeOwner->Add(Node.Key, Node.Value);
 						return true;
 					}
 				}
diff --git a/Engine/Source/Developer/DerivedDataCache/Private/Tests/HttpCacheStoreTest.cpp b/Engine/Source/Developer/DerivedDataCache/Private/Tests/HttpCacheStoreTest.cpp
--- a/Engine/Source/Developer/DerivedDataCache/Private/Tests/HttpCacheStoreTest.cpp
+++ b/Engine/Source/Developer/DerivedDataCache/Private/Tests/HttpCacheStoreTest.cpp
@@ -44,7 +44,7 @@ ILegacyCacheStore* GetAnyHttpCacheStore(
 	FString& OutAccessToken,
 	FString& OutNamespace);
 
-TTuple<ILegacyCacheStore*, ECacheStoreFlags> CreateZenCacheStore(const TCHAR* NodeName, const TCHAR* Config);
+ILegacyCacheStore* CreateZenCacheStore(const TCHAR* NodeName, const TCHAR* Config, ICacheStoreOwner* Owner);
 
 class FHttpCacheStoreTestBase : public FAutomationTestBase
 {
@@ -668,11 +668,11 @@ bool FHttpCacheStoreTest::RunTest(const FString& Parameters)
 
 	FScopeZenService ScopeZenSiblingService(MoveTemp(ZenTestServiceSiblingSettings));
 	TUniquePtr<ILegacyCacheStore> ZenIntermediarySiblingBackend(CreateZenCacheStore(TEXT("TestSibling"),
-		*FString::Printf(TEXT("Host=%s, StructuredNamespace=%s"), ScopeZenSiblingService.GetInstance().GetURL(), *TestNamespace)).Key);
+		*FString::Printf(TEXT("Host=%s, StructuredNamespace=%s"), ScopeZenSiblingService.GetInstance().GetURL(), *TestNamespace), nullptr));
 
 	FScopeZenService ScopeZenService(MoveTemp(ZenTestServiceSettings));
 	TUniquePtr<ILegacyCacheStore> ZenIntermediaryBackend(CreateZenCacheStore(TEXT("Test"),
-		*FString::Printf(TEXT("Host=%s, StructuredNamespace=%s"), ScopeZenService.GetInstance().GetURL(), *TestNamespace)).Key);
+		*FString::Printf(TEXT("Host=%s, StructuredNamespace=%s"), ScopeZenService.GetInstance().GetURL(), *TestNamespace), nullptr));
 	auto WaitForZenPushToUpstream = [](ILegacyCacheStore* ZenBackend, TConstArrayView<FCacheRecord> Records)
 	{
 		// TODO: Expecting a legitimate means to wait for zen to finish pushing records to its upstream in the future
diff --git a/Engine/Source/Developer/DerivedDataCache/Private/ZenCacheStore.cpp b/Engine/Source/Developer/DerivedDataCache/Private/ZenCacheStore.cpp
--- a/Engine/Source/Developer/DerivedDataCache/Private/ZenCacheStore.cpp
+++ b/Engine/Source/Developer/DerivedDataCache/Private/ZenCacheStore.cpp
@@ -5,6 +5,7 @@
 
 #if UE_WITH_ZEN
 
+#include "Async/UniqueLock.h"
 #include "BatchView.h"
 #include "DerivedDataBackendInterface.h"
 #include "DerivedDataCachePrivate.h"
@@ -76,10 +77,19 @@ public:
 	 * @param ServiceUrl	Base url to the service including scheme.
 	 * @param Namespace		Namespace to use.
 	 */
-	FZenCacheStore(const TCHAR* ServiceUrl, const TCHAR* Namespace);
-
-
-	FZenCacheStore(UE::Zen::FServiceSettings&& InSettings, const TCHAR* Namespace);
+	FZenCacheStore(
+		const TCHAR* ServiceUrl,
+		const TCHAR* Namespace,
+		const TCHAR* Name,
+		const TCHAR* Config,
+		ICacheStoreOwner* Owner);
+
+	FZenCacheStore(
+		UE::Zen::FServiceSettings&& Settings,
+		const TCHAR* Namespace,
+		const TCHAR* Name,
+		const TCHAR* Config,
+		ICacheStoreOwner* Owner);
 
 	inline FString GetName() const { return ZenService.GetInstance().GetURL(); }
 
@@ -124,7 +134,7 @@ public:
 	const Zen::FZenServiceInstance& GetServiceInstance() const { return ZenService.GetInstance(); }
 
 private:
-	void Initialize(const TCHAR* Namespace);
+	void Initialize(const TCHAR* Namespace, const TCHAR* Name, const TCHAR* Config);
 
 	bool IsServiceReady();
 
@@ -142,6 +152,17 @@ private:
 	}
 
 private:
+	template <typename RequestType>
+	struct TRequestWithStats;
+
+	template <typename OutContainerType, typename InContainerType, typename BucketAccessorType, typename RequestTypeAccessorType>
+	static void StartRequests(
+		OutContainerType& Out,
+		const InContainerType& In,
+		BucketAccessorType BucketAccessor,
+		RequestTypeAccessorType TypeAccessor,
+		const ERequestOp Op);
+
 	class FPutOp;
 	class FGetOp;
 	class FPutValueOp;
@@ -153,7 +174,8 @@ private:
 
 	FString Namespace;
 	UE::Zen::FScopeZenService ZenService;
-	mutable FDerivedDataCacheUsageStats UsageStats;
+	ICacheStoreOwner* StoreOwner = nullptr;
+	ICacheStoreStats* StoreStats = nullptr;
 	THttpUniquePtr<IHttpConnectionPool> ConnectionPool;
 	FHttpRequestQueue RequestQueue;
 	bool bIsUsable = false;
@@ -165,6 +187,51 @@ private:
 	TAnsiStringBuilder<256> RpcUri;
 };
 
+template <typename RequestType>
+struct FZenCacheStore::TRequestWithStats
+{
+	RequestType Request;
+	mutable FRequestStats Stats;
+
+	explicit TRequestWithStats(const RequestType& InRequest)
+		: Request(InRequest)
+	{
+	}
+
+	void EndRequest(FZenCacheStore& Outer, const EStatus Status) const
+	{
+		{
+			TUniqueLock Lock(Stats.Mutex);
+			Stats.EndTime = FMonotonicTimePoint::Now();
+			Stats.Status = Status;
+		}
+		if (Outer.StoreStats)
+		{
+			Outer.StoreStats->AddRequest(Stats);
+		}
+	}
+};
+
+template <typename OutContainerType, typename InContainerType, typename BucketAccessorType, typename RequestTypeAccessorType>
+void FZenCacheStore::StartRequests(
+	OutContainerType& Out,
+	const InContainerType& In,
+	BucketAccessorType BucketAccessor,
+	RequestTypeAccessorType TypeAccessor,
+	const ERequestOp Op)
+{
+	const FMonotonicTimePoint Now = FMonotonicTimePoint::Now();
+	Out.Reserve(In.Num());
+	for (const auto& Request : In)
+	{
+		auto& RequestWithStats = Out[Out.Emplace(Request)];
+		RequestWithStats.Stats.Bucket = BucketAccessor(Request);
+		RequestWithStats.Stats.Type = TypeAccessor(Request);
+		RequestWithStats.Stats.Op = Op;
+		RequestWithStats.Stats.StartTime = Now;
+	}
+}
+
 class FZenCacheStore::FPutOp final : public FThreadSafeRefCountedObject
 {
 public:
@@ -174,22 +241,19 @@ public:
 		FOnCachePutComplete&& InOnComplete)
 		: CacheStore(InCacheStore)
 		, Owner(InOwner)
-		, Requests(InRequests)
-		, Batches(Requests, [this](const FCachePutRequest& NextRequest) {return BatchGroupingFilter(NextRequest);})
 		, OnComplete(MoveTemp(InOnComplete))
 	{
-		TRACE_COUNTER_ADD(ZenDDC_Put, (int64)Requests.Num());
-		COOK_STAT(Timers.Reserve(Requests.Num()));
-		COOK_STAT(for (const FCachePutRequest& Request : Requests)
-		{
-			Timers.Add(Request.Record.GetKey(), CacheStore.UsageStats.TimePut());
-		})
+		StartRequests(Requests, InRequests, [](const FCachePutRequest& Request) { return Request.Record.GetKey().Bucket; },
+			[](auto&) { return ERequestType::Record; }, ERequestOp::Put);
+		Batches = TBatchView<const TRequestWithStats<FCachePutRequest>>(Requests,
+			[this](const TRequestWithStats<FCachePutRequest>& NextRequest) { return BatchGroupingFilter(NextRequest.Request); });
+		TRACE_COUNTER_ADD(ZenDDC_Put, int64(Requests.Num()));
 	}
 
 	void IssueRequests()
 	{
 		FRequestBarrier Barrier(Owner);
-		for (TArrayView<const FCachePutRequest> Batch : Batches)
+		for (TArrayView<const TRequestWithStats<FCachePutRequest>> Batch : Batches)
 		{
 			FCbPackage BatchPackage;
 			FCbWriter BatchWriter;
@@ -200,13 +264,14 @@ public:
 
 				BatchWriter.BeginObject(ANSITEXTVIEW("Params"));
 				{
-					ECachePolicy BatchDefaultPolicy = Batch[0].Policy.GetRecordPolicy();
+					ECachePolicy BatchDefaultPolicy = Batch[0].Request.Policy.GetRecordPolicy();
 					BatchWriter << ANSITEXTVIEW("DefaultPolicy") << *WriteToString<128>(BatchDefaultPolicy);
 					BatchWriter.AddString(ANSITEXTVIEW("Namespace"), CacheStore.Namespace);
 
 					BatchWriter.BeginArray(ANSITEXTVIEW("Requests"));
-					for (const FCachePutRequest& Request : Batch)
+					for (const TRequestWithStats<FCachePutRequest>& RequestWithStats : Batch)
 					{
+						const FCachePutRequest& Request = RequestWithStats.Request;
 						const FCacheRecord& Record = Request.Record;
 
 						BatchWriter.BeginObject();
@@ -240,17 +305,19 @@ public:
 							++RequestIndex;
 							continue;
 						}
-						const FCachePutRequest& Request = Batch[RequestIndex++];
 
-						const FCacheKey& Key = Request.Record.GetKey();
+						const TRequestWithStats<FCachePutRequest>& RequestWithStats = Batch[RequestIndex++];
+						RequestWithStats.Stats.AddLatency(FMonotonicTimeSpan::FromSeconds(HttpResponse->GetStats().StartTransferTime));
+
+						const FCacheKey& Key = RequestWithStats.Request.Record.GetKey();
 						bool bPutSucceeded = ResponseField.AsBool();
 						if (CacheStore.DebugOptions.ShouldSimulatePutMiss(Key))
 						{
 							UE_LOG(LogDerivedDataCache, Verbose, TEXT("%s: Simulated miss for put of %s from '%s'"),
-								*CacheStore.GetName(), *WriteToString<96>(Key), *Request.Name);
+								*CacheStore.GetName(), *WriteToString<96>(Key), *RequestWithStats.Request.Name);
 							bPutSucceeded = false;
 						}
-						bPutSucceeded ? OnHit(Request) : OnMiss(Request);
+						bPutSucceeded ? OnHit(RequestWithStats) : OnMiss(RequestWithStats);
 					}
 					if (RequestIndex != Batch.Num())
 					{
@@ -266,9 +333,9 @@ public:
 						*CacheStore.GetName(), *WriteToString<256>(*HttpResponse));
 				}
 
-				for (const FCachePutRequest& Request : Batch.RightChop(RequestIndex))
+				for (const TRequestWithStats<FCachePutRequest>& RequestWithStats : Batch.RightChop(RequestIndex))
 				{
-					OnMiss(Request);
+					OnMiss(RequestWithStats);
 				}
 			};
 			CacheStore.EnqueueAsyncRpc(Owner, BatchPackage, MoveTemp(OnRpcComplete));
@@ -293,44 +360,41 @@ private:
 		return EBatchView::Continue;
 	}
 
-	void OnHit(const FCachePutRequest& Request)
+	void OnHit(const TRequestWithStats<FCachePutRequest>& RequestWithStats)
 	{
-		const FCacheKey& Key = Request.Record.GetKey();
 		UE_LOG(LogDerivedDataCache, Verbose, TEXT("%s: Cache Put complete for %s from '%s'"),
-			*CacheStore.GetName(), *WriteToString<96>(Key), *Request.Name);
-		int64 SentSize = 0;
-		for (const FValueWithId& Value : Request.Record.GetValues())
+			*CacheStore.GetName(), *WriteToString<96>(RequestWithStats.Request.Record.GetKey()), *RequestWithStats.Request.Name);
+
+		if (const FCbObject& Meta = RequestWithStats.Request.Record.GetMeta())
 		{
-			SentSize += Value.GetData().GetCompressed().GetSize();
+			RequestWithStats.Stats.PhysicalWriteSize += Meta.GetSize();
 		}
-		TRACE_COUNTER_ADD(ZenDDC_BytesSent, SentSize);
-		TRACE_COUNTER_INCREMENT(ZenDDC_PutHit);
-		COOK_STAT(Timers[Key].AddHit(SentSize));
+		for (const FValueWithId& Value : RequestWithStats.Request.Record.GetValues())
+		{
+			RequestWithStats.Stats.AddLogicalWrite(Value);
+			RequestWithStats.Stats.PhysicalWriteSize += Value.GetData().GetCompressedSize();
+		}
+		RequestWithStats.EndRequest(CacheStore, EStatus::Ok);
 
-		OnComplete(Request.MakeResponse(EStatus::Ok));
-	};
+		TRACE_COUNTER_ADD(ZenDDC_BytesSent, int64(RequestWithStats.Stats.PhysicalWriteSize));
+		TRACE_COUNTER_INCREMENT(ZenDDC_PutHit);
+		OnComplete(RequestWithStats.Request.MakeResponse(EStatus::Ok));
+	}
 
-	void OnMiss(const FCachePutRequest& Request)
+	void OnMiss(const TRequestWithStats<FCachePutRequest>& RequestWithStats)
 	{
-		const FCacheKey& Key = Request.Record.GetKey();
 		UE_LOG(LogDerivedDataCache, Verbose, TEXT("%s: Cache Put miss for '%s' from '%s'"),
-			*CacheStore.GetName(), *WriteToString<96>(Key), *Request.Name);
-		int64 SentSize = 0;
-		for (const FValueWithId& Value : Request.Record.GetValues())
-		{
-			SentSize += Value.GetData().GetCompressed().GetSize();
-		}
-		COOK_STAT(Timers[Key].AddMiss(SentSize));
-		OnComplete(Request.MakeResponse(EStatus::Error));
-	};
+			*CacheStore.GetName(), *WriteToString<96>(RequestWithStats.Request.Record.GetKey()), *RequestWithStats.Request.Name);
+		RequestWithStats.EndRequest(CacheStore, EStatus::Error);
+		OnComplete(RequestWithStats.Request.MakeResponse(EStatus::Error));
+	}
 
 	FZenCacheStore& CacheStore;
 	IRequestOwner& Owner;
-	const TArray<FCachePutRequest, TInlineAllocator<1>> Requests;
+	TArray<TRequestWithStats<FCachePutRequest>, TInlineAllocator<1>> Requests;
 	uint64 BatchSize = 0;
-	TBatchView<const FCachePutRequest> Batches;
+	TBatchView<const TRequestWithStats<FCachePutRequest>> Batches;
 	FOnCachePutComplete OnComplete;
-	COOK_STAT(TMap<FCacheKey, FCookStats::FScopedStatsCounter> Timers);
 };
 
 class FZenCacheStore::FGetOp final : public FThreadSafeRefCountedObject
@@ -342,16 +406,12 @@ public:
 		FOnCacheGetComplete&& InOnComplete)
 		: CacheStore(InCacheStore)
 		, Owner(InOwner)
-		, Requests(InRequests)
 		, OnComplete(MoveTemp(InOnComplete))
 	{
-		TRACE_COUNTER_ADD(ZenDDC_Get, (int64)Requests.Num());
+		StartRequests(Requests, InRequests, [](const FCacheGetRequest& Request) { return Request.Key.Bucket; },
+			[](auto&) { return ERequestType::Record; }, ERequestOp::Get);
+		TRACE_COUNTER_ADD(ZenDDC_Get, int64(Requests.Num()));
 		TRACE_COUNTER_ADD(ZenDDC_CacheRecordRequestCountInFlight, int64(Requests.Num()));
-		COOK_STAT(Timers.Reserve(Requests.Num()));
-		COOK_STAT(for (const FCacheGetRequest& Request : Requests)
-		{
-			Timers.Add(Request.Key, CacheStore.UsageStats.TimeGet());
-		})
 	}
 
 	virtual ~FGetOp()
@@ -364,127 +424,136 @@ public:
 		FRequestBarrier Barrier(Owner);
 		ForEachBatch(CacheStore.CacheRecordBatchSize, Requests.Num(),
 			[this](int32 BatchFirst, int32 BatchLast)
+		{
+			TConstArrayView<TRequestWithStats<FCacheGetRequest>> Batch(Requests.GetData() + BatchFirst, BatchLast - BatchFirst + 1);
+
+			FCbWriter BatchRequest;
+			BatchRequest.BeginObject();
 			{
-				TArrayView<const FCacheGetRequest> Batch(Requests.GetData() + BatchFirst, BatchLast - BatchFirst + 1);
+				BatchRequest << ANSITEXTVIEW("Method") << ANSITEXTVIEW("GetCacheRecords");
+				BatchRequest.AddInteger(ANSITEXTVIEW("Accept"), Zen::Http::kCbPkgMagic);
+				if (CacheStore.bIsLocalConnection)
+				{
+					BatchRequest.AddInteger(ANSITEXTVIEW("AcceptFlags"), static_cast<uint32_t>(Zen::Http::RpcAcceptOptions::kAllowLocalReferences));
+					BatchRequest.AddInteger(ANSITEXTVIEW("Pid"), FPlatformProcess::GetCurrentProcessId());
+				}
 
-				FCbWriter BatchRequest;
-				BatchRequest.BeginObject();
+				BatchRequest.BeginObject(ANSITEXTVIEW("Params"));
 				{
-					BatchRequest << ANSITEXTVIEW("Method") << ANSITEXTVIEW("GetCacheRecords");
-					BatchRequest.AddInteger(ANSITEXTVIEW("Accept"), Zen::Http::kCbPkgMagic);
-					if (CacheStore.bIsLocalConnection)
-					{
-						BatchRequest.AddInteger(ANSITEXTVIEW("AcceptFlags"), static_cast<uint32_t>(Zen::Http::RpcAcceptOptions::kAllowLocalReferences));
-						BatchRequest.AddInteger(ANSITEXTVIEW("Pid"), FPlatformProcess::GetCurrentProcessId());
-					}
+					ECachePolicy BatchDefaultPolicy = Batch[0].Request.Policy.GetRecordPolicy();
+					BatchRequest << ANSITEXTVIEW("DefaultPolicy") << *WriteToString<128>(BatchDefaultPolicy);
+					BatchRequest.AddString(ANSITEXTVIEW("Namespace"), CacheStore.Namespace);
 
-					BatchRequest.BeginObject(ANSITEXTVIEW("Params"));
+					BatchRequest.BeginArray(ANSITEXTVIEW("Requests"));
+					for (const TRequestWithStats<FCacheGetRequest>& RequestWithStats : Batch)
 					{
-						ECachePolicy BatchDefaultPolicy = Batch[0].Policy.GetRecordPolicy();
-						BatchRequest << ANSITEXTVIEW("DefaultPolicy") << *WriteToString<128>(BatchDefaultPolicy);
-						BatchRequest.AddString(ANSITEXTVIEW("Namespace"), CacheStore.Namespace);
-
-						BatchRequest.BeginArray(ANSITEXTVIEW("Requests"));
-						for (const FCacheGetRequest& Request : Batch)
+						BatchRequest.BeginObject();
 						{
-							BatchRequest.BeginObject();
+							BatchRequest << ANSITEXTVIEW("Key") << RequestWithStats.Request.Key;
+							if (!RequestWithStats.Request.Policy.IsDefault())
 							{
-								BatchRequest << ANSITEXTVIEW("Key") << Request.Key;
-								if (!Request.Policy.IsDefault())
-								{
-									BatchRequest << ANSITEXTVIEW("Policy") << Request.Policy;
-								}
+								BatchRequest << ANSITEXTVIEW("Policy") << RequestWithStats.Request.Policy;
 							}
-							BatchRequest.EndObject();
 						}
-						BatchRequest.EndArray();
+						BatchRequest.EndObject();
 					}
-					BatchRequest.EndObject();
+					BatchRequest.EndArray();
 				}
 				BatchRequest.EndObject();
+			}
+			BatchRequest.EndObject();
 
-				FGetOp* OriginalOp = this;
-				auto OnRpcComplete = [this, OpRef = TRefCountPtr<FGetOp>(OriginalOp), Batch](THttpUniquePtr<IHttpResponse>& HttpResponse, FCbPackage& Response)
+			FGetOp* OriginalOp = this;
+			auto OnRpcComplete = [this, OpRef = TRefCountPtr<FGetOp>(OriginalOp), Batch](THttpUniquePtr<IHttpResponse>& HttpResponse, FCbPackage& Response)
+			{
+				int32 RequestIndex = 0;
+				if (HttpResponse->GetErrorCode() == EHttpErrorCode::None && HttpResponse->GetStatusCode() >= 200 && HttpResponse->GetStatusCode() <= 299)
 				{
-					int32 RequestIndex = 0;
-					if (HttpResponse->GetErrorCode() == EHttpErrorCode::None && HttpResponse->GetStatusCode() >= 200 && HttpResponse->GetStatusCode() <= 299)
-					{
-						const FCbObject& ResponseObj = Response.GetObject();
+					const FCbObject& ResponseObj = Response.GetObject();
 						
-						for (FCbField RecordField : ResponseObj[ANSITEXTVIEW("Result")])
+					for (FCbField RecordField : ResponseObj[ANSITEXTVIEW("Result")])
+					{
+						if (RequestIndex >= Batch.Num())
 						{
-							if (RequestIndex >= Batch.Num())
-							{
-								++RequestIndex;
-								continue;
-							}
-							const FCacheGetRequest& Request = Batch[RequestIndex++];
+							++RequestIndex;
+							continue;
+						}
 
-							const FCacheKey& Key = Request.Key;
-							FOptionalCacheRecord Record;
+						const TRequestWithStats<FCacheGetRequest>& RequestWithStats = Batch[RequestIndex++];
+						RequestWithStats.Stats.AddLatency(FMonotonicTimeSpan::FromSeconds(HttpResponse->GetStats().StartTransferTime));
 
-							if (CacheStore.DebugOptions.ShouldSimulateGetMiss(Key))
-							{
-								UE_LOG(LogDerivedDataCache, Verbose, TEXT("%s: Simulated miss for put of '%s' from '%s'"),
-									*CacheStore.GetName(), *WriteToString<96>(Key), *Request.Name);
-							}
-							else if (!RecordField.IsNull())
-							{
-								Record = FCacheRecord::Load(Response, RecordField.AsObject());
-							}
-							Record ? OnHit(Request, MoveTemp(Record).Get()) : OnMiss(Request);
+						const FCacheKey& Key = RequestWithStats.Request.Key;
+						FOptionalCacheRecord Record;
+
+						if (CacheStore.DebugOptions.ShouldSimulateGetMiss(Key))
+						{
+							UE_LOG(LogDerivedDataCache, Verbose, TEXT("%s: Simulated miss for put of '%s' from '%s'"),
+								*CacheStore.GetName(), *WriteToString<96>(Key), *RequestWithStats.Request.Name);
 						}
-						if (RequestIndex != Batch.Num())
+						else if (!RecordField.IsNull())
 						{
-							UE_LOG(LogDerivedDataCache, Warning,
-								TEXT("%s: Invalid response received from GetCacheRecords RPC: %d results expected, received %d, from %s"),
-								*CacheStore.GetName(), Batch.Num(), RequestIndex, *WriteToString<256>(*HttpResponse));
+							Record = FCacheRecord::Load(Response, RecordField.AsObject());
 						}
+						Record ? OnHit(RequestWithStats, MoveTemp(Record).Get()) : OnMiss(RequestWithStats);
 					}
-					else if (HttpResponse->GetStatusCode() != 404)
+					if (RequestIndex != Batch.Num())
 					{
 						UE_LOG(LogDerivedDataCache, Warning,
-							TEXT("%s: Error response received from GetCacheRecords RPC: from %s"),
-							*CacheStore.GetName(), *WriteToString<256>(*HttpResponse));
+							TEXT("%s: Invalid response received from GetCacheRecords RPC: %d results expected, received %d, from %s"),
+							*CacheStore.GetName(), Batch.Num(), RequestIndex, *WriteToString<256>(*HttpResponse));
 					}
+				}
+				else if (HttpResponse->GetStatusCode() != 404)
+				{
+					UE_LOG(LogDerivedDataCache, Warning,
+						TEXT("%s: Error response received from GetCacheRecords RPC: from %s"),
+						*CacheStore.GetName(), *WriteToString<256>(*HttpResponse));
+				}
 
-					for (const FCacheGetRequest& Request : Batch.RightChop(RequestIndex))
-					{
-						OnMiss(Request);
-					}
-				};
+				for (const TRequestWithStats<FCacheGetRequest>& RequestWithStats : Batch.RightChop(RequestIndex))
+				{
+					OnMiss(RequestWithStats);
+				}
+			};
 
-				CacheStore.EnqueueAsyncRpc(Owner, BatchRequest.Save().AsObject(), MoveTemp(OnRpcComplete));
-			});
+			CacheStore.EnqueueAsyncRpc(Owner, BatchRequest.Save().AsObject(), MoveTemp(OnRpcComplete));
+		});
 	}
+
 private:
-	void OnHit(const FCacheGetRequest& Request, FCacheRecord&& Record)
+	void OnHit(const TRequestWithStats<FCacheGetRequest>& RequestWithStats, FCacheRecord&& Record)
 	{
 		UE_LOG(LogDerivedDataCache, Verbose, TEXT("%s: Cache hit for '%s' from '%s'"),
-			*CacheStore.GetName(), *WriteToString<96>(Request.Key), *Request.Name);
+			*CacheStore.GetName(), *WriteToString<96>(RequestWithStats.Request.Key), *RequestWithStats.Request.Name);
 
-		TRACE_COUNTER_INCREMENT(ZenDDC_GetHit);
-		int64 ReceivedSize = Private::GetCacheRecordCompressedSize(Record);
-		TRACE_COUNTER_ADD(ZenDDC_BytesReceived, ReceivedSize);
-		COOK_STAT(Timers[Request.Key].AddHit(ReceivedSize));
+		if (const FCbObject& Meta = Record.GetMeta())
+		{
+			RequestWithStats.Stats.PhysicalReadSize += Meta.GetSize();
+		}
+		for (const FValueWithId& Value : Record.GetValues())
+		{
+			RequestWithStats.Stats.AddLogicalRead(Value);
+			RequestWithStats.Stats.PhysicalReadSize += Value.GetData().GetCompressedSize();
+		}
+		RequestWithStats.EndRequest(CacheStore, EStatus::Ok);
 
-		OnComplete({Request.Name, MoveTemp(Record), Request.UserData, EStatus::Ok});
-	};
+		TRACE_COUNTER_INCREMENT(ZenDDC_GetHit);
+		TRACE_COUNTER_ADD(ZenDDC_BytesReceived, int64(RequestWithStats.Stats.PhysicalReadSize));
+		OnComplete({RequestWithStats.Request.Name, MoveTemp(Record), RequestWithStats.Request.UserData, EStatus::Ok});
+	}
 
-	void OnMiss(const FCacheGetRequest& Request)
+	void OnMiss(const TRequestWithStats<FCacheGetRequest>& RequestWithStats)
 	{
 		UE_LOG(LogDerivedDataCache, Verbose, TEXT("%s: Cache miss for '%s' from '%s'"),
-			*CacheStore.GetName(), *WriteToString<96>(Request.Key), *Request.Name);
-
-		COOK_STAT(Timers[Request.Key].AddMiss());
-		OnComplete(Request.MakeResponse(EStatus::Error));
-	};
+			*CacheStore.GetName(), *WriteToString<96>(RequestWithStats.Request.Key), *RequestWithStats.Request.Name);
+		RequestWithStats.EndRequest(CacheStore, EStatus::Error);
+		OnComplete(RequestWithStats.Request.MakeResponse(EStatus::Error));
+	}
 
 	FZenCacheStore& CacheStore;
 	IRequestOwner& Owner;
-	const TArray<FCacheGetRequest, TInlineAllocator<1>> Requests;
+	TArray<TRequestWithStats<FCacheGetRequest>, TInlineAllocator<1>> Requests;
 	FOnCacheGetComplete OnComplete;
-	COOK_STAT(TMap<FCacheKey, FCookStats::FScopedStatsCounter> Timers);
 };
 
 class FZenCacheStore::FPutValueOp final : public FThreadSafeRefCountedObject
@@ -496,22 +565,19 @@ public:
 		FOnCachePutValueComplete&& InOnComplete)
 		: CacheStore(InCacheStore)
 		, Owner(InOwner)
-		, Requests(InRequests)
-		, Batches(Requests, [this](const FCachePutValueRequest& NextRequest) {return BatchGroupingFilter(NextRequest);})
 		, OnComplete(MoveTemp(InOnComplete))
 	{
-		TRACE_COUNTER_ADD(ZenDDC_Put, (int64)Requests.Num());
-		COOK_STAT(Timers.Reserve(Requests.Num()));
-		COOK_STAT(for (const FCachePutValueRequest& Request : Requests)
-		{
-			Timers.Add(Request.Key, CacheStore.UsageStats.TimePut());
-		})
+		StartRequests(Requests, InRequests, [](const FCachePutValueRequest& Request) { return Request.Key.Bucket; },
+			[](auto&) { return ERequestType::Value; }, ERequestOp::Put);
+		TRACE_COUNTER_ADD(ZenDDC_Put, int64(Requests.Num()));
+		Batches = TBatchView<const TRequestWithStats<FCachePutValueRequest>>(Requests,
+			[this](const TRequestWithStats<FCachePutValueRequest>& NextRequest) { return BatchGroupingFilter(NextRequest.Request); });
 	}
 
 	void IssueRequests()
 	{
 		FRequestBarrier Barrier(Owner);
-		for (TArrayView<const FCachePutValueRequest> Batch : Batches)
+		for (TArrayView<const TRequestWithStats<FCachePutValueRequest>> Batch : Batches)
 		{
 			FCbPackage BatchPackage;
 			FCbWriter BatchWriter;
@@ -522,25 +588,25 @@ public:
 
 				BatchWriter.BeginObject(ANSITEXTVIEW("Params"));
 				{
-					ECachePolicy BatchDefaultPolicy = Batch[0].Policy;
+					ECachePolicy BatchDefaultPolicy = Batch[0].Request.Policy;
 					BatchWriter << ANSITEXTVIEW("DefaultPolicy") << *WriteToString<128>(BatchDefaultPolicy);
 					BatchWriter.AddString(ANSITEXTVIEW("Namespace"), CacheStore.Namespace);
 
 					BatchWriter.BeginArray("Requests");
-					for (const FCachePutValueRequest& Request : Batch)
+					for (const TRequestWithStats<FCachePutValueRequest>& RequestWithStats : Batch)
 					{
 						BatchWriter.BeginObject();
 						{
-							BatchWriter << ANSITEXTVIEW("Key") << Request.Key;
-							const FValue& Value = Request.Value;
+							BatchWriter << ANSITEXTVIEW("Key") << RequestWithStats.Request.Key;
+							const FValue& Value = RequestWithStats.Request.Value;
 							BatchWriter.AddBinaryAttachment("RawHash", Value.GetRawHash());
 							if (Value.HasData())
 							{
 								BatchPackage.AddAttachment(FCbAttachment(Value.GetData()));
 							}
-							if (Request.Policy != BatchDefaultPolicy)
+							if (RequestWithStats.Request.Policy != BatchDefaultPolicy)
 							{
-								BatchWriter << ANSITEXTVIEW("Policy") << WriteToString<128>(Request.Policy);
+								BatchWriter << ANSITEXTVIEW("Policy") << WriteToString<128>(RequestWithStats.Request.Policy);
 							}
 						}
 						BatchWriter.EndObject();
@@ -565,16 +631,18 @@ public:
 							++RequestIndex;
 							continue;
 						}
-						const FCachePutValueRequest& Request = Batch[RequestIndex++];
+
+						const TRequestWithStats<FCachePutValueRequest>& RequestWithStats = Batch[RequestIndex++];
+						RequestWithStats.Stats.AddLatency(FMonotonicTimeSpan::FromSeconds(HttpResponse->GetStats().StartTransferTime));
 
 						bool bPutSucceeded = ResponseField.AsBool();
-						if (CacheStore.DebugOptions.ShouldSimulatePutMiss(Request.Key))
+						if (CacheStore.DebugOptions.ShouldSimulatePutMiss(RequestWithStats.Request.Key))
 						{
 							UE_LOG(LogDerivedDataCache, Verbose, TEXT("%s: Simulated miss for PutValue of %s from '%s'"),
-								*CacheStore.GetName(), *WriteToString<96>(Request.Key), *Request.Name);
+								*CacheStore.GetName(), *WriteToString<96>(RequestWithStats.Request.Key), *RequestWithStats.Request.Name);
 							bPutSucceeded = false;
 						}
-						bPutSucceeded ? OnHit(Request) : OnMiss(Request);
+						bPutSucceeded ? OnHit(RequestWithStats) : OnMiss(RequestWithStats);
 					}
 					if (RequestIndex != Batch.Num())
 					{
@@ -590,9 +658,9 @@ public:
 						*CacheStore.GetName(), *WriteToString<256>(*HttpResponse));
 				}
 
-				for (const FCachePutValueRequest& Request : Batch.RightChop(RequestIndex))
+				for (const TRequestWithStats<FCachePutValueRequest>& RequestWithStats : Batch.RightChop(RequestIndex))
 				{
-					OnMiss(Request);
+					OnMiss(RequestWithStats);
 				}
 			};
 
@@ -613,33 +681,34 @@ private:
 		return EBatchView::Continue;
 	}
 
-	void OnHit(const FCachePutValueRequest& Request)
+	void OnHit(const TRequestWithStats<FCachePutValueRequest>& RequestWithStats)
 	{
 		UE_LOG(LogDerivedDataCache, Verbose, TEXT("%s: Cache PutValue complete for %s from '%s'"),
-			*CacheStore.GetName(), *WriteToString<96>(Request.Key), *Request.Name);
+			*CacheStore.GetName(), *WriteToString<96>(RequestWithStats.Request.Key), *RequestWithStats.Request.Name);
+
+		RequestWithStats.Stats.AddLogicalWrite(RequestWithStats.Request.Value);
+		RequestWithStats.Stats.PhysicalWriteSize += RequestWithStats.Request.Value.GetData().GetCompressedSize();
+		RequestWithStats.EndRequest(CacheStore, EStatus::Ok);
+
 		TRACE_COUNTER_INCREMENT(ZenDDC_PutHit);
-		int64 SentSize = Request.Value.GetData().GetCompressedSize();
-		TRACE_COUNTER_ADD(ZenDDC_BytesSent, SentSize);
-		COOK_STAT(Timers[Request.Key].AddHit(SentSize));
-		OnComplete(Request.MakeResponse(EStatus::Ok));
-	};
+		TRACE_COUNTER_ADD(ZenDDC_BytesSent, RequestWithStats.Stats.PhysicalWriteSize);
+		OnComplete(RequestWithStats.Request.MakeResponse(EStatus::Ok));
+	}
 
-	void OnMiss(const FCachePutValueRequest& Request)
+	void OnMiss(const TRequestWithStats<FCachePutValueRequest>& RequestWithStats)
 	{
 		UE_LOG(LogDerivedDataCache, Verbose, TEXT("%s: Cache PutValue miss for '%s' from '%s'"),
-			*CacheStore.GetName(), *WriteToString<96>(Request.Key), *Request.Name);
-
-		COOK_STAT(Timers[Request.Key].AddMiss());
-		OnComplete(Request.MakeResponse(EStatus::Error));
-	};
+			*CacheStore.GetName(), *WriteToString<96>(RequestWithStats.Request.Key), *RequestWithStats.Request.Name);
+		RequestWithStats.EndRequest(CacheStore, EStatus::Error);
+		OnComplete(RequestWithStats.Request.MakeResponse(EStatus::Error));
+	}
 
 	FZenCacheStore& CacheStore;
 	IRequestOwner& Owner;
-	const TArray<FCachePutValueRequest, TInlineAllocator<1>> Requests;
+	TArray<TRequestWithStats<FCachePutValueRequest>, TInlineAllocator<1>> Requests;
 	uint64 BatchSize = 0;
-	TBatchView<const FCachePutValueRequest> Batches;
+	TBatchView<const TRequestWithStats<FCachePutValueRequest>> Batches;
 	FOnCachePutValueComplete OnComplete;
-	COOK_STAT(TMap<FCacheKey, FCookStats::FScopedStatsCounter> Timers);
 };
 
 class FZenCacheStore::FGetValueOp final : public FThreadSafeRefCountedObject
@@ -651,16 +720,12 @@ public:
 		FOnCacheGetValueComplete&& InOnComplete)
 		: CacheStore(InCacheStore)
 		, Owner(InOwner)
-		, Requests(InRequests)
 		, OnComplete(MoveTemp(InOnComplete))
 	{
+		StartRequests(Requests, InRequests, [](const FCacheGetValueRequest& Request) { return Request.Key.Bucket; },
+			[](auto&) { return ERequestType::Value; }, ERequestOp::Get);
 		TRACE_COUNTER_ADD(ZenDDC_Get, (int64)Requests.Num());
 		TRACE_COUNTER_ADD(ZenDDC_CacheRecordRequestCountInFlight, int64(Requests.Num()));
-		COOK_STAT(Timers.Reserve(Requests.Num()));
-		COOK_STAT(for (const FCacheGetValueRequest& Request : Requests)
-		{
-			Timers.Add(Request.Key, CacheStore.UsageStats.TimeGet());
-		})
 	}
 
 	virtual ~FGetValueOp()
@@ -674,7 +739,7 @@ public:
 		ForEachBatch(CacheStore.CacheRecordBatchSize, Requests.Num(),
 			[this](int32 BatchFirst, int32 BatchLast)
 		{
-			TArrayView<const FCacheGetValueRequest> Batch(Requests.GetData() + BatchFirst, BatchLast - BatchFirst + 1);
+			TConstArrayView<TRequestWithStats<FCacheGetValueRequest>> Batch(Requests.GetData() + BatchFirst, BatchLast - BatchFirst + 1);
 
 			FCbWriter BatchRequest;
 			BatchRequest.BeginObject();
@@ -689,19 +754,19 @@ public:
 
 				BatchRequest.BeginObject(ANSITEXTVIEW("Params"));
 				{
-					ECachePolicy BatchDefaultPolicy = Batch[0].Policy;
+					ECachePolicy BatchDefaultPolicy = Batch[0].Request.Policy;
 					BatchRequest << ANSITEXTVIEW("DefaultPolicy") << *WriteToString<128>(BatchDefaultPolicy);
 					BatchRequest.AddString(ANSITEXTVIEW("Namespace"), CacheStore.Namespace);
 
 					BatchRequest.BeginArray("Requests");
-					for (const FCacheGetValueRequest& Request : Batch)
+					for (const TRequestWithStats<FCacheGetValueRequest>& RequestWithStats : Batch)
 					{
 						BatchRequest.BeginObject();
 						{
-							BatchRequest << ANSITEXTVIEW("Key") << Request.Key;
-							if (Request.Policy != BatchDefaultPolicy)
+							BatchRequest << ANSITEXTVIEW("Key") << RequestWithStats.Request.Key;
+							if (RequestWithStats.Request.Policy != BatchDefaultPolicy)
 							{
-								BatchRequest << ANSITEXTVIEW("Policy") << WriteToString<128>(Request.Policy);
+								BatchRequest << ANSITEXTVIEW("Policy") << WriteToString<128>(RequestWithStats.Request.Policy);
 							}
 						}
 						BatchRequest.EndObject();
@@ -727,14 +792,16 @@ public:
 							++RequestIndex;
 							continue;
 						}
-						const FCacheGetValueRequest& Request = Batch[RequestIndex++];
+
+						const TRequestWithStats<FCacheGetValueRequest>& RequestWithStats = Batch[RequestIndex++];
+						RequestWithStats.Stats.AddLatency(FMonotonicTimeSpan::FromSeconds(HttpResponse->GetStats().StartTransferTime));
 
 						FCbObjectView ResultObj = ResultField.AsObjectView();
 						TOptional<FValue> Value;
-						if (CacheStore.DebugOptions.ShouldSimulateGetMiss(Request.Key))
+						if (CacheStore.DebugOptions.ShouldSimulateGetMiss(RequestWithStats.Request.Key))
 						{
 							UE_LOG(LogDerivedDataCache, Verbose, TEXT("%s: Simulated miss for GetValue of '%s' from '%s'"),
-								*CacheStore.GetName(), *WriteToString<96>(Request.Key), *Request.Name);
+								*CacheStore.GetName(), *WriteToString<96>(RequestWithStats.Request.Key), *RequestWithStats.Request.Name);
 						}
 						else
 						{
@@ -754,7 +821,7 @@ public:
 								}
 							}
 						}
-						(bool)Value ? OnHit(Request, MoveTemp(*Value)) : OnMiss(Request);
+						(bool)Value ? OnHit(RequestWithStats, MoveTemp(*Value)) : OnMiss(RequestWithStats);
 					}
 					if (RequestIndex != Batch.Num())
 					{
@@ -770,42 +837,43 @@ public:
 						*CacheStore.GetName(), *WriteToString<256>(*HttpResponse));
 				}
 
-				for (const FCacheGetValueRequest& Request : Batch.RightChop(RequestIndex))
+				for (const TRequestWithStats<FCacheGetValueRequest>& RequestWithStats : Batch.RightChop(RequestIndex))
 				{
-					OnMiss(Request);
+					OnMiss(RequestWithStats);
 				}
 			};
 
 			CacheStore.EnqueueAsyncRpc(Owner, BatchRequest.Save().AsObject(), MoveTemp(OnRpcComplete));
 		});
 	}
+
 private:
-	void OnHit(const FCacheGetValueRequest& Request, FValue&& Value)
+	void OnHit(const TRequestWithStats<FCacheGetValueRequest>& RequestWithStats, FValue&& Value)
 	{
 		UE_LOG(LogDerivedDataCache, Verbose, TEXT("%s: Cache hit for '%s' from '%s'"),
-			*CacheStore.GetName(), *WriteToString<96>(Request.Key), *Request.Name);
-		TRACE_COUNTER_INCREMENT(ZenDDC_GetHit);
-		int64 ReceivedSize = Value.GetData().GetCompressedSize();
-		TRACE_COUNTER_ADD(ZenDDC_BytesReceived, ReceivedSize);
-		COOK_STAT(Timers[Request.Key].AddHit(ReceivedSize));
+			*CacheStore.GetName(), *WriteToString<96>(RequestWithStats.Request.Key), *RequestWithStats.Request.Name);
+
+		RequestWithStats.Stats.AddLogicalRead(Value);
+		RequestWithStats.Stats.PhysicalReadSize += Value.GetData().GetCompressedSize();
+		RequestWithStats.EndRequest(CacheStore, EStatus::Ok);
 
-		OnComplete({Request.Name, Request.Key, MoveTemp(Value), Request.UserData, EStatus::Ok});
+		TRACE_COUNTER_INCREMENT(ZenDDC_GetHit);
+		TRACE_COUNTER_ADD(ZenDDC_BytesReceived, int64(RequestWithStats.Stats.PhysicalReadSize));
+		OnComplete({RequestWithStats.Request.Name, RequestWithStats.Request.Key, MoveTemp(Value), RequestWithStats.Request.UserData, EStatus::Ok});
 	};
 
-	void OnMiss(const FCacheGetValueRequest& Request)
+	void OnMiss(const TRequestWithStats<FCacheGetValueRequest>& RequestWithStats)
 	{
 		UE_LOG(LogDerivedDataCache, Verbose, TEXT("%s: Cache miss for '%s' from '%s'"),
-			*CacheStore.GetName(), *WriteToString<96>(Request.Key), *Request.Name);
-
-		COOK_STAT(Timers[Request.Key].AddMiss());
-		OnComplete(Request.MakeResponse(EStatus::Error));
+			*CacheStore.GetName(), *WriteToString<96>(RequestWithStats.Request.Key), *RequestWithStats.Request.Name);
+		RequestWithStats.EndRequest(CacheStore, EStatus::Error);
+		OnComplete(RequestWithStats.Request.MakeResponse(EStatus::Error));
 	};
 
 	FZenCacheStore& CacheStore;
 	IRequestOwner& Owner;
-	const TArray<FCacheGetValueRequest, TInlineAllocator<1>> Requests;
+	TArray<TRequestWithStats<FCacheGetValueRequest>, TInlineAllocator<1>> Requests;
 	FOnCacheGetValueComplete OnComplete;
-	COOK_STAT(TMap<FCacheKey, FCookStats::FScopedStatsCounter> Timers);
 };
 
 class FZenCacheStore::FGetChunksOp final : public FThreadSafeRefCountedObject
@@ -817,17 +885,14 @@ public:
 		FOnCacheGetChunkComplete&& InOnComplete)
 		: CacheStore(InCacheStore)
 		, Owner(InOwner)
-		, Requests(InRequests)
 		, OnComplete(MoveTemp(InOnComplete))
 	{
-		TRACE_COUNTER_ADD(ZenDDC_ChunkRequestCountInFlight, int64(Requests.Num()));
+		StartRequests(Requests, InRequests, [](const FCacheGetChunkRequest& Request) { return Request.Key.Bucket; },
+			[](const FCacheGetChunkRequest& Request) { return Request.Id.IsNull() ? ERequestType::Value : ERequestType::Record; },
+			ERequestOp::GetChunk);
+		Algo::StableSortBy(Requests, &TRequestWithStats<FCacheGetChunkRequest>::Request, TChunkLess());
 		TRACE_COUNTER_ADD(ZenDDC_Get, int64(Requests.Num()));
-		Requests.StableSort(TChunkLess());
-		COOK_STAT(Timers.Reserve(Requests.Num()));
-		COOK_STAT(for (const FCacheGetChunkRequest& Request : Requests)
-		{
-			Timers.Add(Request.Key, CacheStore.UsageStats.TimeGet());
-		})
+		TRACE_COUNTER_ADD(ZenDDC_ChunkRequestCountInFlight, int64(Requests.Num()));
 	}
 
 	virtual ~FGetChunksOp()
@@ -841,7 +906,7 @@ public:
 		ForEachBatch(CacheStore.CacheChunksBatchSize, Requests.Num(),
 			[this](int32 BatchFirst, int32 BatchLast)
 		{
-			TArrayView<const FCacheGetChunkRequest> Batch(Requests.GetData() + BatchFirst, BatchLast - BatchFirst + 1);
+			TConstArrayView<TRequestWithStats<FCacheGetChunkRequest>> Batch(Requests.GetData() + BatchFirst, BatchLast - BatchFirst + 1);
 
 			FCbWriter BatchRequest;
 			BatchRequest.BeginObject();
@@ -856,13 +921,14 @@ public:
 
 				BatchRequest.BeginObject(ANSITEXTVIEW("Params"));
 				{
-					ECachePolicy DefaultPolicy = Batch[0].Policy;
+					ECachePolicy DefaultPolicy = Batch[0].Request.Policy;
 					BatchRequest << ANSITEXTVIEW("DefaultPolicy") << WriteToString<128>(DefaultPolicy);
 					BatchRequest.AddString(ANSITEXTVIEW("Namespace"), CacheStore.Namespace);
 
 					BatchRequest.BeginArray(ANSITEXTVIEW("ChunkRequests"));
-					for (const FCacheGetChunkRequest& Request : Batch)
+					for (const TRequestWithStats<FCacheGetChunkRequest>& RequestWithStats : Batch)
 					{
+						const FCacheGetChunkRequest& Request = RequestWithStats.Request;
 						BatchRequest.BeginObject();
 						{
 							BatchRequest << ANSITEXTVIEW("Key") << Request.Key;
@@ -911,13 +977,15 @@ public:
 							++RequestIndex;
 							continue;
 						}
-						const FCacheGetChunkRequest& Request = Batch[RequestIndex++];
+						const TRequestWithStats<FCacheGetChunkRequest>& RequestWithStats = Batch[RequestIndex++];
+						RequestWithStats.Stats.AddLatency(FMonotonicTimeSpan::FromSeconds(HttpResponse->GetStats().StartTransferTime));
 
 						FIoHash RawHash;
 						bool Succeeded = false;
 						uint64 RawSize = 0;
 						FCbObjectView ResultObject = ResultView.AsObjectView();
 						FSharedBuffer RequestedBytes;
+						const FCacheGetChunkRequest& Request = RequestWithStats.Request;
 						if (CacheStore.DebugOptions.ShouldSimulateGetMiss(Request.Key))
 						{
 							UE_LOG(LogDerivedDataCache, Verbose, TEXT("%s: Simulated miss for get of '%s' from '%s'"),
@@ -952,7 +1020,7 @@ public:
 								}
 							}
 						}
-						Succeeded ? OnHit(Request, MoveTemp(RawHash), RawSize, MoveTemp(RequestedBytes)) : OnMiss(Request);
+						Succeeded ? OnHit(RequestWithStats, MoveTemp(RawHash), RawSize, MoveTemp(RequestedBytes)) : OnMiss(RequestWithStats);
 					}
 				}
 				else if (HttpResponse->GetStatusCode() != 404)
@@ -962,41 +1030,46 @@ public:
 						*CacheStore.GetName(), *WriteToString<256>(*HttpResponse));
 				}
 
-				for (const FCacheGetChunkRequest& Request : Batch.RightChop(RequestIndex))
+				for (const TRequestWithStats<FCacheGetChunkRequest>& RequestWithStats : Batch.RightChop(RequestIndex))
 				{
-					OnMiss(Request);
+					OnMiss(RequestWithStats);
 				}
 			};
 			CacheStore.EnqueueAsyncRpc(Owner, BatchRequest.Save().AsObject(), MoveTemp(OnRpcComplete));
 		});
 	}
+
 private:
-	void OnHit(const FCacheGetChunkRequest& Request, FIoHash&& RawHash, uint64 RawSize, FSharedBuffer&& RequestedBytes)
+	void OnHit(const TRequestWithStats<FCacheGetChunkRequest>& RequestWithStats, FIoHash&& RawHash, uint64 RawSize, FSharedBuffer&& RequestedBytes)
 	{
+		const FCacheGetChunkRequest& Request = RequestWithStats.Request;
 		UE_LOG(LogDerivedDataCache, Verbose, TEXT("%s: CacheChunk hit for '%s' from '%s'"),
 			*CacheStore.GetName(), *WriteToString<96>(Request.Key, '/', Request.Id), *Request.Name);
+
+		// This is a rough estimate of physical read size until Zen communicates stats with each response.
+		RequestWithStats.Stats.LogicalReadSize += RequestedBytes.GetSize();
+		RequestWithStats.Stats.PhysicalReadSize += RequestedBytes.GetSize();
+		RequestWithStats.EndRequest(CacheStore, EStatus::Ok);
+
 		TRACE_COUNTER_INCREMENT(ZenDDC_GetHit);
-		int64 ReceivedSize = RequestedBytes.GetSize();
-		TRACE_COUNTER_ADD(ZenDDC_BytesReceived, ReceivedSize);
-		COOK_STAT(Timers[Request.Key].AddHit(RawSize));
+		TRACE_COUNTER_ADD(ZenDDC_BytesReceived, int64(RequestWithStats.Stats.PhysicalReadSize));
 		OnComplete({Request.Name, Request.Key, Request.Id, Request.RawOffset,
 			RawSize, MoveTemp(RawHash), MoveTemp(RequestedBytes), Request.UserData, EStatus::Ok});
 	};
 
-	void OnMiss(const FCacheGetChunkRequest& Request)
+	void OnMiss(const TRequestWithStats<FCacheGetChunkRequest>& RequestWithStats)
 	{
+		const FCacheGetChunkRequest& Request = RequestWithStats.Request;
 		UE_LOG(LogDerivedDataCache, Verbose, TEXT("%s: CacheChunk miss with missing value '%s' for '%s' from '%s'"),
 			*CacheStore.GetName(), *WriteToString<16>(Request.Id), *WriteToString<96>(Request.Key), *Request.Name);
-
-		COOK_STAT(Timers[Request.Key].AddMiss());
+		RequestWithStats.EndRequest(CacheStore, EStatus::Error);
 		OnComplete(Request.MakeResponse(EStatus::Error));
 	};
 
 	FZenCacheStore& CacheStore;
 	IRequestOwner& Owner;
-	TArray<FCacheGetChunkRequest, TInlineAllocator<1>> Requests;
+	TArray<TRequestWithStats<FCacheGetChunkRequest>, TInlineAllocator<1>> Requests;
 	FOnCacheGetChunkComplete OnComplete;
-	COOK_STAT(TMap<FCacheKey, FCookStats::FScopedStatsCounter> Timers);
 };
 
 class FZenCacheStore::FCbPackageReceiver final : public IHttpReceiver
@@ -1128,22 +1201,32 @@ private:
 
 FZenCacheStore::FZenCacheStore(
 	const TCHAR* InServiceUrl,
-	const TCHAR* InNamespace)
+	const TCHAR* InNamespace,
+	const TCHAR* InName,
+	const TCHAR* InConfig,
+	ICacheStoreOwner* InStoreOwner)
 	: ZenService(InServiceUrl)
+	, StoreOwner(InStoreOwner)
 {
-	Initialize(InNamespace);
+	Initialize(InNamespace, InName, InConfig);
 }
 
 FZenCacheStore::FZenCacheStore(
 	UE::Zen::FServiceSettings&& InSettings,
-	const TCHAR* InNamespace)
+	const TCHAR* InNamespace,
+	const TCHAR* InName,
+	const TCHAR* InConfig,
+	ICacheStoreOwner* InStoreOwner)
 	: ZenService(MoveTemp(InSettings))
+	, StoreOwner(InStoreOwner)
 {
-	Initialize(InNamespace);
+	Initialize(InNamespace, InName, InConfig);
 }
 
 void FZenCacheStore::Initialize(
-	const TCHAR* InNamespace)
+	const TCHAR* InNamespace,
+	const TCHAR* InName,
+	const TCHAR* InConfig)
 {
 	Namespace = InNamespace;
 	if (IsServiceReady())
@@ -1168,6 +1251,29 @@ void FZenCacheStore::Initialize(
 		bIsUsable = true;
 		bIsLocalConnection = ZenService.GetInstance().IsServiceRunningLocally();
 
+		if (StoreOwner)
+		{
+			bool bReadOnly = false;
+			FParse::Bool(InConfig, TEXT("ReadOnly="), bReadOnly);
+
+			// Default to locally launched service getting the Local cache store flag.  Can be overridden by explicit value in config.
+			bool bLocal = bIsLocalConnection;
+			FParse::Bool(InConfig, TEXT("Local="), bLocal);
+
+			// Default to non-locally launched service getting the Remote cache store flag.  Can be overridden by explicit value in config.
+			// In the future this could be extended to allow the Remote flag by default (even for locally launched instances) if they have upstreams configured.
+			bool bRemote = !bIsLocalConnection;
+			FParse::Bool(InConfig, TEXT("Remote="), bRemote);
+
+			ECacheStoreFlags Flags = ECacheStoreFlags::Query;
+			Flags |= bReadOnly ? ECacheStoreFlags::None : ECacheStoreFlags::Store;
+			Flags |= bLocal ? ECacheStoreFlags::Local : ECacheStoreFlags::None;
+			Flags |= bRemote ? ECacheStoreFlags::Remote : ECacheStoreFlags::None;
+
+			StoreOwner->Add(this, Flags);
+			StoreStats = StoreOwner->CreateStats(this, Flags, TEXT("Zen"), InName, ZenService.GetInstance().GetURL());
+		}
+
 		// Issue a request for stats as it will be fetched asynchronously and issuing now makes them available sooner for future callers.
 		Zen::FZenStats ZenStats;
 		ZenService.GetInstance().GetStats(ZenStats);
@@ -1218,59 +1324,7 @@ void FZenCacheStore::EnqueueAsyncRpc(IRequestOwner& Owner, const FCbPackage& Req
 
 void FZenCacheStore::LegacyStats(FDerivedDataCacheStatsNode& OutNode)
 {
-	EDerivedDataCacheStatus CacheStatus = EDerivedDataCacheStatus::None;
-	OutNode = { TEXT("Zen"), ZenService.GetInstance().GetURL(), /*bIsLocal*/ ZenService.GetInstance().IsServiceRunningLocally()};
-	OutNode.UsageStats.Add(TEXT(""), UsageStats);
-#if 0
-	// DE: 20230213 We might need to revisit this if we change so Zen handles upstream again
-	Zen::FZenStats ZenStats;
-
-	FDerivedDataCacheUsageStats LocalStats;
-	FDerivedDataCacheUsageStats RemoteStats;
-
-#if ENABLE_COOK_STATS
-	using EHitOrMiss = FCookStats::CallStats::EHitOrMiss;
-	using ECacheStatType = FCookStats::CallStats::EStatType;
-
-	ZenService.GetInstance().GetStats(ZenStats);
-
-	const int64 RemotePutSize = int64(ZenStats.UpstreamStats.TotalUploadedMB * 1024 * 1024);
-	const int64 RemoteGetSize = int64(ZenStats.UpstreamStats.TotalDownloadedMB * 1024 * 1024);
-	const int64 LocalGetSize = FMath::Max<int64>(0, UsageStats.GetStats.GetAccumulatedValueAnyThread(EHitOrMiss::Hit, ECacheStatType::Bytes) - RemoteGetSize);
-
-	LocalStats.PutStats = UsageStats.PutStats;
-	LocalStats.ExistsStats = UsageStats.ExistsStats;
-	LocalStats.PrefetchStats = UsageStats.PrefetchStats;
-
-	LocalStats.GetStats.Accumulate(EHitOrMiss::Hit, ECacheStatType::Counter, ZenStats.CacheStats.Hits - ZenStats.CacheStats.UpstreamHits, /*bIsInGameThread*/ false);
-	LocalStats.GetStats.Accumulate(EHitOrMiss::Miss, ECacheStatType::Counter, ZenStats.CacheStats.Misses + ZenStats.CacheStats.UpstreamHits, /*bIsInGameThread*/ false);
-	RemoteStats.GetStats.Accumulate(EHitOrMiss::Hit, ECacheStatType::Counter, ZenStats.CacheStats.UpstreamHits, /*bIsInGameThread*/ false);
-	RemoteStats.GetStats.Accumulate(EHitOrMiss::Miss, ECacheStatType::Counter, ZenStats.CacheStats.Misses, /*bIsInGameThread*/ false);
-
-	LocalStats.GetStats.Accumulate(EHitOrMiss::Hit, ECacheStatType::Bytes, LocalGetSize, /*bIsInGameThread*/ false);
-	RemoteStats.GetStats.Accumulate(EHitOrMiss::Hit, ECacheStatType::Bytes, RemoteGetSize, /*bIsInGameThread*/ false);
-	RemoteStats.PutStats.Accumulate(EHitOrMiss::Hit, ECacheStatType::Bytes, RemotePutSize, /*bIsInGameThread*/ false);
-#endif
-
-	if (ZenStats.UpstreamStats.EndPointStats.IsEmpty())
-	{
-		OutNode = {TEXT("Zen"), ZenService.GetInstance().GetURL(), /*bIsLocal*/ true};
-		OutNode.UsageStats.Add(TEXT(""), LocalStats);
-		return;
-	}
-
-	TSharedRef<FDerivedDataCacheStatsNode> LocalNode =
-		MakeShared<FDerivedDataCacheStatsNode>(TEXT("Zen"), ZenService.GetInstance().GetURL(), /*bIsLocal*/ true);
-	LocalNode->UsageStats.Add(TEXT(""), LocalStats);
-
-	TSharedRef<FDerivedDataCacheStatsNode> RemoteNode =
-		MakeShared<FDerivedDataCacheStatsNode>(ZenStats.UpstreamStats.EndPointStats[0].Name, ZenStats.UpstreamStats.EndPointStats[0].Url, /*bIsLocal*/ false);
-	RemoteNode->UsageStats.Add(TEXT(""), RemoteStats);
-
-	OutNode = {TEXT("Zen Group"), TEXT(""), /*bIsLocal*/ true};
-	OutNode.Children.Add(LocalNode);
-	OutNode.Children.Add(RemoteNode);
-#endif // 0
+	checkNoEntry();
 }
 
 bool FZenCacheStore::LegacyDebugOptions(FBackendDebugOptions& InOptions)
@@ -1329,7 +1383,7 @@ void FZenCacheStore::GetChunks(
 	GetChunksOp->IssueRequests();
 }
 
-TTuple<ILegacyCacheStore*, ECacheStoreFlags> CreateZenCacheStore(const TCHAR* NodeName, const TCHAR* Config)
+ILegacyCacheStore* CreateZenCacheStore(const TCHAR* NodeName, const TCHAR* Config, ICacheStoreOwner* Owner)
 {
 	FString ServiceUrl;
 	FParse::Value(Config, TEXT("Host="), ServiceUrl);
@@ -1356,7 +1410,7 @@ TTuple<ILegacyCacheStore*, ECacheStoreFlags> CreateZenCacheStore(const TCHAR* No
 	if (ServiceUrl == TEXT("None"))
 	{
 		UE_LOG(LogDerivedDataCache, Log, TEXT("Disabling %s data cache - host set to 'None'."), NodeName);
-		return MakeTuple<ILegacyCacheStore*, ECacheStoreFlags>(nullptr, ECacheStoreFlags::None);
+		return nullptr;
 	}
 
 	FString Namespace;
@@ -1377,7 +1431,7 @@ TTuple<ILegacyCacheStore*, ECacheStoreFlags> CreateZenCacheStore(const TCHAR* No
 		if (CachePathOverride == TEXT("None"))
 		{
 			UE_LOG(LogDerivedDataCache, Log, TEXT("Disabling %s data cache - path set to 'None'."), NodeName);
-			return MakeTuple<ILegacyCacheStore*, ECacheStoreFlags>(nullptr, ECacheStoreFlags::None);
+			return nullptr;
 		}
 	}
 
@@ -1394,7 +1448,7 @@ TTuple<ILegacyCacheStore*, ECacheStoreFlags> CreateZenCacheStore(const TCHAR* No
 		if (!DefaultServiceSettings.IsAutoLaunch())
 		{
 			UE_LOG(LogDerivedDataCache, Warning, TEXT("%s: Attempting to use a sandbox when there is no default autolaunch configured to interhit settings from.  Cache will be disabled."), NodeName);
-			return MakeTuple<ILegacyCacheStore*, ECacheStoreFlags>(nullptr, ECacheStoreFlags::None);
+			return nullptr;
 		}
 
 		// Make a unique local instance (not the default local instance) of ZenServer
@@ -1435,38 +1489,21 @@ TTuple<ILegacyCacheStore*, ECacheStoreFlags> CreateZenCacheStore(const TCHAR* No
 			}
 		}
 
-		Backend = MakeUnique<FZenCacheStore>(MoveTemp(ServiceSettings), *Namespace);
+		Backend = MakeUnique<FZenCacheStore>(MoveTemp(ServiceSettings), *Namespace, NodeName, Config, Owner);
 	}
 	else
 	{
-		Backend = MakeUnique<FZenCacheStore>(*ServiceUrl, *Namespace);
+		Backend = MakeUnique<FZenCacheStore>(*ServiceUrl, *Namespace, NodeName, Config, Owner);
 	}
 
 	if (!Backend->IsUsable())
 	{
 		UE_LOG(LogDerivedDataCache, Warning, TEXT("%s: Failed to contact the service (%s), will not use it."), NodeName, *Backend->GetName());
 		Backend.Reset();
-		return MakeTuple(nullptr, ECacheStoreFlags::None);
+		return nullptr;
 	}
 
-	bool bReadOnly = false;
-	FParse::Bool(Config, TEXT("ReadOnly="), bReadOnly);
-
-	// Default to locally launched service getting the Local cache store flag.  Can be overridden by explicit value in config.
-	const bool bServiceRunningLocally = Backend->GetServiceInstance().IsServiceRunningLocally();
-	bool bLocal = bServiceRunningLocally;
-	FParse::Bool(Config, TEXT("Local="), bLocal);
-
-	// Default to non-locally launched service getting the Remote cache store flag.  Can be overridden by explicit value in config.
-	// In the future this could be extended to allow the Remote flag by default (even for locally launched instances) if they have upstreams configured.
-	bool bRemote = !bServiceRunningLocally;
-	FParse::Bool(Config, TEXT("Remote="), bRemote);
-
-	ECacheStoreFlags Flags = ECacheStoreFlags::Query;
-	Flags |= bReadOnly ? ECacheStoreFlags::None : ECacheStoreFlags::Store;
-	Flags |= bLocal ? ECacheStoreFlags::Local : ECacheStoreFlags::None;
-	Flags |= bRemote ? ECacheStoreFlags::Remote : ECacheStoreFlags::None;
-	return MakeTuple(Backend.Release(), Flags);
+	return Backend.Release();
 }
 
 } // namespace UE::DerivedData
@@ -1476,10 +1513,10 @@ TTuple<ILegacyCacheStore*, ECacheStoreFlags> CreateZenCacheStore(const TCHAR* No
 namespace UE::DerivedData
 {
 
-TTuple<ILegacyCacheStore*, ECacheStoreFlags> CreateZenCacheStore(const TCHAR* NodeName, const TCHAR* Config)
+ILegacyCacheStore* CreateZenCacheStore(const TCHAR* NodeName, const TCHAR* Config, ICacheStoreOwner* Owner)
 {
 	UE_LOG(LogDerivedDataCache, Warning, TEXT("%s: Zen cache is not yet supported in the current build configuration."), NodeName);
-	return MakeTuple(nullptr, ECacheStoreFlags::None);
+	return nullptr;
 }
 
 } // UE::DerivedData
