DDC: Fixed leaking cache store stats in the Zen cache

#rnx

diff --git a/Engine/Source/Developer/DerivedDataCache/Private/ZenCacheStore.cpp b/Engine/Source/Developer/DerivedDataCache/Private/ZenCacheStore.cpp
--- a/Engine/Source/Developer/DerivedDataCache/Private/ZenCacheStore.cpp
+++ b/Engine/Source/Developer/DerivedDataCache/Private/ZenCacheStore.cpp
@@ -91,6 +91,8 @@ public:
 		const TCHAR* Config,
 		ICacheStoreOwner* Owner);
 
+	~FZenCacheStore() final;
+
 	inline FString GetName() const { return ZenService.GetInstance().GetURL(); }
 
 	/**
@@ -1223,6 +1225,14 @@ FZenCacheStore::FZenCacheStore(
 	Initialize(InNamespace, InName, InConfig);
 }
 
+FZenCacheStore::~FZenCacheStore()
+{
+	if (StoreStats)
+	{
+		StoreOwner->DestroyStats(StoreStats);
+	}
+}
+
 void FZenCacheStore::Initialize(
 	const TCHAR* InNamespace,
 	const TCHAR* InName,
