Make sure record policy is set up correctly in batch requests to Zen

diff --git a/Engine/Source/Developer/DerivedDataCache/Private/ZenCacheStore.cpp b/Engine/Source/Developer/DerivedDataCache/Private/ZenCacheStore.cpp
--- a/Engine/Source/Developer/DerivedDataCache/Private/ZenCacheStore.cpp
+++ b/Engine/Source/Developer/DerivedDataCache/Private/ZenCacheStore.cpp
@@ -280,7 +280,7 @@ public:
 						{
 							BatchWriter.SetName(ANSITEXTVIEW("Record"));
 							Record.Save(BatchPackage, BatchWriter);
-							if (!Request.Policy.IsDefault())
+							if ((!Request.Policy.IsUniform()) || (Request.Policy.GetRecordPolicy() != BatchDefaultPolicy))
 							{
 								BatchWriter << ANSITEXTVIEW("Policy") << Request.Policy;
 							}
@@ -452,7 +452,8 @@ public:
 						BatchRequest.BeginObject();
 						{
 							BatchRequest << ANSITEXTVIEW("Key") << RequestWithStats.Request.Key;
-							if (!RequestWithStats.Request.Policy.IsDefault())
+
+							if ((!RequestWithStats.Request.Policy.IsUniform()) || (RequestWithStats.Request.Policy.GetRecordPolicy() != BatchDefaultPolicy))
 							{
 								BatchRequest << ANSITEXTVIEW("Policy") << RequestWithStats.Request.Policy;
 							}
