Return error if we receive a partial record from zenserver

#prefligth 648a3d2f88385dd88cfdda4a

diff --git a/Engine/Source/Developer/DerivedDataCache/Private/ZenCacheStore.cpp b/Engine/Source/Developer/DerivedDataCache/Private/ZenCacheStore.cpp
--- a/Engine/Source/Developer/DerivedDataCache/Private/ZenCacheStore.cpp
+++ b/Engine/Source/Developer/DerivedDataCache/Private/ZenCacheStore.cpp
@@ -529,6 +529,7 @@ private:
 		UE_LOG(LogDerivedDataCache, Verbose, TEXT("%s: Cache hit for '%s' from '%s'"),
 			*CacheStore.GetName(), *WriteToString<96>(RequestWithStats.Request.Key), *RequestWithStats.Request.Name);
 
+		bool bInComplete = false;
 		if (const FCbObject& Meta = Record.GetMeta())
 		{
 			RequestWithStats.Stats.PhysicalReadSize += Meta.GetSize();
@@ -537,12 +538,26 @@ private:
 		{
 			RequestWithStats.Stats.AddLogicalRead(Value);
 			RequestWithStats.Stats.PhysicalReadSize += Value.GetData().GetCompressedSize();
+			ECachePolicy ValuePolicy = RequestWithStats.Request.Policy.GetValuePolicy(Value.GetId());
+			if (EnumHasAnyFlags(ValuePolicy, ECachePolicy::SkipData))
+			{
+				continue;
+			}
+			if (Value.HasData())
+			{
+				continue;
+			}
+			if (EnumHasAnyFlags(ValuePolicy, ECachePolicy::Query))
+			{
+				bInComplete = true;
 		}
-		RequestWithStats.EndRequest(CacheStore, EStatus::Ok);
+		}
+		EStatus Status = bInComplete ? EStatus::Error : EStatus::Ok;
+		RequestWithStats.EndRequest(CacheStore, Status);
 
 		TRACE_COUNTER_INCREMENT(ZenDDC_GetHit);
 		TRACE_COUNTER_ADD(ZenDDC_BytesReceived, int64(RequestWithStats.Stats.PhysicalReadSize));
-		OnComplete({RequestWithStats.Request.Name, MoveTemp(Record), RequestWithStats.Request.UserData, EStatus::Ok});
+		OnComplete({RequestWithStats.Request.Name, MoveTemp(Record), RequestWithStats.Request.UserData, Status});
 	}
 
 	void OnMiss(const TRequestWithStats<FCacheGetRequest>& RequestWithStats)
