DDC: Fixed the in-flight memory cache to be read-only

#rnx

diff --git a/Engine/Source/Developer/DerivedDataCache/Private/DerivedDataBackends.cpp b/Engine/Source/Developer/DerivedDataCache/Private/DerivedDataBackends.cpp
--- a/Engine/Source/Developer/DerivedDataCache/Private/DerivedDataBackends.cpp
+++ b/Engine/Source/Developer/DerivedDataCache/Private/DerivedDataBackends.cpp
@@ -67,7 +67,7 @@ ILegacyCacheStore* CreateCacheStoreThrottle(ILegacyCacheStore* InnerCache, uint3
 ILegacyCacheStore* CreateCacheStoreVerify(ILegacyCacheStore* InnerCache, bool bPutOnError);
 ILegacyCacheStore* CreateFileSystemCacheStore(const TCHAR* Name, const TCHAR* Config, ICacheStoreOwner& Owner, FString& OutPath);
 TTuple<ILegacyCacheStore*, ECacheStoreFlags> CreateHttpCacheStore(const TCHAR* NodeName, const TCHAR* Config);
-void CreateMemoryCacheStore(IMemoryCacheStore*& OutCache, const TCHAR* Name, ICacheStoreOwner* Owner, int64 MaxCacheSize, bool bCanBeDisabled);
+void CreateMemoryCacheStore(IMemoryCacheStore*& OutCache, const TCHAR* Name, const TCHAR* Config, ICacheStoreOwner* Owner);
 IPakFileCacheStore* CreatePakFileCacheStore(const TCHAR* Name, const TCHAR* Filename, bool bWriting, bool bCompressed, ICacheStoreOwner* Owner);
 ILegacyCacheStore* CreateS3CacheStore(const TCHAR* RootManifestPath, const TCHAR* BaseUrl, const TCHAR* Region, const TCHAR* CanaryObjectKey, const TCHAR* CachePath);
 ILegacyCacheStore* CreateZenCacheStore(const TCHAR* NodeName, const TCHAR* Config, ICacheStoreOwner* Owner);
@@ -638,28 +638,12 @@ public:
 
 		bBootFound = true;
 
+	#if WITH_EDITOR
 		// Only allow boot cache with the editor. We don't want other tools and utilities (eg. SCW) writing to the same file.
-#if WITH_EDITOR
-		FString Filename;
-		int64 MaxCacheSize = -1; // in MB
-		const int64 MaxSupportedCacheSize = 2048; // 2GB
-
-		FParse::Value(Entry, TEXT("MaxCacheSize="), MaxCacheSize);
-
-		// make sure MaxCacheSize does not exceed 2GB
-		MaxCacheSize = FMath::Min(MaxCacheSize, MaxSupportedCacheSize);
+		CreateMemoryCacheStore(BootCache, TEXT("Boot"), Entry, this);
+	#endif
 
-		UE_LOG(LogDerivedDataCache, Display, TEXT("%s: Max Cache Size: %d MB"), NodeName, MaxCacheSize);
-		CreateMemoryCacheStore(BootCache, TEXT("Boot"), this, MaxCacheSize * 1024 * 1024, /*bCanBeDisabled*/ true);
-#endif
-
-		if (BootCache)
-		{
-			Add(BootCache, ECacheStoreFlags::Local | ECacheStoreFlags::Query | ECacheStoreFlags::Store);
-			return true;
-		}
-
-		return false;
+		return !!BootCache;
 	}
 
 	/**
@@ -674,7 +658,7 @@ public:
 		FString Filename;
 		FParse::Value(Entry, TEXT("Filename="), Filename);
 		IMemoryCacheStore* Cache;
-		CreateMemoryCacheStore(Cache, NodeName, this, /*MaxCacheSize*/ -1, /*bCanBeDisabled*/ false);
+		CreateMemoryCacheStore(Cache, NodeName, TEXT(""), this);
 		check(Cache);
 		if (Filename.Len())
 		{
@@ -692,7 +676,7 @@ public:
 		else
 		{
 			// This is unconditionally added to the hierarchy and will be deleted by the hierarchy.
-			CreateMemoryCacheStore(OutCache, TEXT("Memory"), this, 0, /*bCanBeDisabled*/ false);
+			CreateMemoryCacheStore(OutCache, TEXT("Memory"), TEXT("-ReadOnly"), this);
 			MemoryCache = OutCache;
 		}
 	}
diff --git a/Engine/Source/Developer/DerivedDataCache/Private/MemoryCacheStore.cpp b/Engine/Source/Developer/DerivedDataCache/Private/MemoryCacheStore.cpp
--- a/Engine/Source/Developer/DerivedDataCache/Private/MemoryCacheStore.cpp
+++ b/Engine/Source/Developer/DerivedDataCache/Private/MemoryCacheStore.cpp
@@ -26,9 +26,8 @@ public:
 	explicit FMemoryCacheStore(
 		IMemoryCacheStore*& OutCache,
 		const TCHAR* InName,
-		ICacheStoreOwner* InOwner,
-		int64 InMaxCacheSize = -1,
-		bool bCanBeDisabled = false);
+		const TCHAR* InConfig,
+		ICacheStoreOwner* InOwner);
 	~FMemoryCacheStore() final;
 
 	// ICacheStore Interface
@@ -122,21 +121,33 @@ protected:
 FMemoryCacheStore::FMemoryCacheStore(
 	IMemoryCacheStore*& OutCache,
 	const TCHAR* InName,
-	ICacheStoreOwner* InOwner,
-	int64 InMaxCacheSize,
-	bool bInCanBeDisabled)
+	const TCHAR* InConfig,
+	ICacheStoreOwner* InOwner)
 	: Name(InName)
 	, StoreOwner(InOwner)
-	, MaxCacheSize(InMaxCacheSize < 0 ? 0 : uint64(InMaxCacheSize))
+	, MaxCacheSize(0)
 	, bDisabled(false)
 	, CurrentCacheSize(0)
 	, bMaxSizeExceeded(false)
-	, bCanBeDisabled(bInCanBeDisabled)
+	, bCanBeDisabled(false)
 {
 	OutCache = this;
+
+	// Make sure MaxCacheSize does not exceed 2 GB
+	int64 SignedMaxCacheSize = -1; // in MB
+	FParse::Value(InConfig, TEXT("MaxCacheSize="), SignedMaxCacheSize);
+	const int64 MaxSupportedCacheSize = 2048; // 2 GB
+	MaxCacheSize = FMath::Min(SignedMaxCacheSize, MaxSupportedCacheSize);
+	UE_LOG(LogDerivedDataCache, Display, TEXT("%s: Max Cache Size: %d MB"), *Name, SignedMaxCacheSize);
+	MaxCacheSize = SignedMaxCacheSize < 0 ? 0 : SignedMaxCacheSize * 1024 * 1024;
+
 	if (StoreOwner)
 	{
-		constexpr ECacheStoreFlags Flags = ECacheStoreFlags::Local | ECacheStoreFlags::Query | ECacheStoreFlags::Store;
+		ECacheStoreFlags Flags = ECacheStoreFlags::Local | ECacheStoreFlags::Query;
+		if (!FParse::Param(InConfig, TEXT("ReadOnly")))
+		{
+			Flags |= ECacheStoreFlags::Store;
+		}
 		StoreOwner->Add(this, Flags);
 		StoreStats = StoreOwner->CreateStats(this, Flags, TEXTVIEW("Memory"), Name);
 	}
@@ -730,9 +741,9 @@ void FMemoryCacheStore::GetChunks(
 	}
 }
 
-void CreateMemoryCacheStore(IMemoryCacheStore*& OutCache, const TCHAR* Name, ICacheStoreOwner* Owner, int64 MaxCacheSize, bool bCanBeDisabled)
+void CreateMemoryCacheStore(IMemoryCacheStore*& OutCache, const TCHAR* Name, const TCHAR* Config, ICacheStoreOwner* Owner)
 {
-	new FMemoryCacheStore(OutCache, Name, Owner, MaxCacheSize, bCanBeDisabled);
+	new FMemoryCacheStore(OutCache, Name, Config, Owner);
 }
 
 } // UE::DerivedData
