DDC: Prepared the Jupiter cache for ICacheStoreStats

#rnx

diff --git a/Engine/Source/Developer/DerivedDataCache/Private/DerivedDataBackends.cpp b/Engine/Source/Developer/DerivedDataCache/Private/DerivedDataBackends.cpp
--- a/Engine/Source/Developer/DerivedDataCache/Private/DerivedDataBackends.cpp
+++ b/Engine/Source/Developer/DerivedDataCache/Private/DerivedDataBackends.cpp
@@ -66,7 +66,7 @@ ILegacyCacheStore* CreateCacheStoreHierarchy(ICacheStoreOwner*& OutOwner, TFunct
 ILegacyCacheStore* CreateCacheStoreThrottle(ILegacyCacheStore* InnerCache, uint32 LatencyMS, uint32 MaxBytesPerSecond);
 ILegacyCacheStore* CreateCacheStoreVerify(ILegacyCacheStore* InnerCache, bool bPutOnError);
 ILegacyCacheStore* CreateFileSystemCacheStore(const TCHAR* Name, const TCHAR* Config, ICacheStoreOwner& Owner, FString& OutPath);
-TTuple<ILegacyCacheStore*, ECacheStoreFlags> CreateHttpCacheStore(const TCHAR* NodeName, const TCHAR* Config);
+ILegacyCacheStore* CreateHttpCacheStore(const TCHAR* NodeName, const TCHAR* Config, ICacheStoreOwner* Owner);
 void CreateMemoryCacheStore(IMemoryCacheStore*& OutCache, const TCHAR* Name, const TCHAR* Config, ICacheStoreOwner* Owner);
 IPakFileCacheStore* CreatePakFileCacheStore(const TCHAR* Name, const TCHAR* Filename, bool bWriting, bool bCompressed, ICacheStoreOwner* Owner);
 ILegacyCacheStore* CreateS3CacheStore(const TCHAR* RootManifestPath, const TCHAR* BaseUrl, const TCHAR* Region, const TCHAR* CanaryObjectKey, const TCHAR* CachePath);
@@ -313,10 +313,8 @@ public:
 				}
 				else if (NodeType == TEXT("Cloud") || NodeType == TEXT("Http"))
 				{
-					TTuple<ILegacyCacheStore*, ECacheStoreFlags> Node = CreateHttpCacheStore(*NodeName, *Entry);
-					if (Node.Key)
+					if (CreateHttpCacheStore(*NodeName, *Entry, this))
 					{
-						NodeOwner->Add(Node.Key, Node.Value);
 						return true;
 					}
 				}
diff --git a/Engine/Source/Developer/DerivedDataCache/Private/HttpCacheStore.cpp b/Engine/Source/Developer/DerivedDataCache/Private/HttpCacheStore.cpp
--- a/Engine/Source/Developer/DerivedDataCache/Private/HttpCacheStore.cpp
+++ b/Engine/Source/Developer/DerivedDataCache/Private/HttpCacheStore.cpp
@@ -7,6 +7,8 @@
 #if WITH_HTTP_DDC_BACKEND
 
 #include "Algo/Transform.h"
+#include "Async/Mutex.h"
+#include "Async/UniqueLock.h"
 #include "Compression/CompressedBuffer.h"
 #include "Containers/StringView.h"
 #include "Containers/Ticker.h"
@@ -216,6 +218,7 @@ FAnsiStringBuilderBase& operator<<(FAnsiStringBuilderBase& Builder, const FHttpA
 
 struct FHttpCacheStoreParams
 {
+	FString Name;
 	FString Host;
 	FString HostPinnedPublicKeys;
 	FString Namespace;
@@ -247,7 +250,7 @@ public:
 	/**
 	 * Creates the cache store client, checks health status and attempts to acquire an access token.
 	 */
-	explicit FHttpCacheStore(const FHttpCacheStoreParams& Params);
+	FHttpCacheStore(const FHttpCacheStoreParams& Params, ICacheStoreOwner* Owner);
 
 	~FHttpCacheStore();
 
@@ -313,6 +316,8 @@ private:
 
 	FAnsiStringBuilderBase EffectiveDomain;
 
+	ICacheStoreOwner* StoreOwner = nullptr;
+
 	FDerivedDataCacheUsageStats UsageStats;
 	FBackendDebugOptions DebugOptions;
 	THttpUniquePtr<IHttpConnectionPool> ConnectionPool;
@@ -325,8 +330,8 @@ private:
 	TUniquePtr<FHttpAccessToken> Access;
 	FTSTicker::FDelegateHandle RefreshAccessTokenHandle;
 	double RefreshAccessTokenTime = 0.0;
-	uint32 FailedLoginAttempts = 0;
 	uint32 LoginAttempts = 0;
+	uint32 FailedLoginAttempts = 0;
 	uint32 InteractiveLoginAttempts = 0;
 
 	bool bIsUsable = false;
@@ -336,8 +341,6 @@ private:
 
 	FHttpClientParams GetDefaultClientParams() const;
 
-	THttpUniquePtr<IHttpResponse> BeginIsServiceReady(IHttpClient& Client, TArray64<uint8>& Body);
-	bool EndIsServiceReady(THttpUniquePtr<IHttpResponse>& Response, TArray64<uint8>& Body);
 	bool AcquireAccessToken(IHttpClient* Client = nullptr);
 	void SetAccessTokenAndUnlock(FScopeLock &Lock, FStringView Token, double RefreshDelay = 0.0);
 	
@@ -351,24 +354,6 @@ private:
 
 	TUniquePtr<FHttpOperation> WaitForHttpOperation(EOperationCategory Category);
 
-	struct FGetCacheRecordOnlyResponse
-	{
-		FSharedString Name;
-		FCacheKey Key;
-		uint64 UserData = 0;
-		uint64 BytesReceived = 0;
-		FOptionalCacheRecord Record;
-		EStatus Status = EStatus::Error;
-	};
-	using FOnGetCacheRecordOnlyComplete = TUniqueFunction<void(FGetCacheRecordOnlyResponse&& Response)>;
-	void GetCacheRecordOnlyAsync(
-		IRequestOwner& Owner,
-		const FSharedString& Name,
-		const FCacheKey& Key,
-		const FCacheRecordPolicy& Policy,
-		uint64 UserData,
-		FOnGetCacheRecordOnlyComplete&& OnComplete);
-
 	void GetCacheRecordAsync(
 		IRequestOwner& Owner,
 		const FSharedString& Name,
@@ -402,14 +387,11 @@ private:
 		uint64 UserData,
 		FOnCacheGetValueComplete&& OnComplete);
 
-	void RefCachedDataProbablyExistsBatchAsync(
-		IRequestOwner& Owner,
-		TConstArrayView<FCacheGetValueRequest> ValueRefs,
-		FOnCacheGetValueComplete&& OnComplete);
-
 	class FHealthCheckOp;
 	class FPutPackageOp;
 	class FGetRecordOp;
+	class FGetValueOp;
+	class FExistsBatchOp;
 };
 
 //----------------------------------------------------------------------------------------------------------
@@ -714,16 +696,19 @@ private:
 class FHttpCacheStore::FPutPackageOp final : public FThreadSafeRefCountedObject
 {
 public:
-
-	struct FCachePutPackageResponse
+	struct FResponse
 	{
-		FSharedString Name;
-		FCacheKey Key;
-		uint64 UserData = 0;
 		uint64 BytesSent = 0;
 		EStatus Status = EStatus::Error;
 	};
-	using FOnCachePutPackageComplete = TUniqueFunction<void(FCachePutPackageResponse&& Response)>;
+	using FOnPackageComplete = TUniqueFunction<void (FResponse&& Response)>;
+
+	static TRefCountPtr<FPutPackageOp> New(FHttpCacheStore& CacheStore, IRequestOwner& Owner, const FSharedString& Name)
+	{
+		return new FPutPackageOp(CacheStore, Owner, Name);
+	}
+
+	void Put(const FCacheKey& Key, FCbPackage&& Package, FOnPackageComplete&& OnComplete);
 
 	/** Performs a multi-request operation for uploading a package of content. */
 	static void PutPackage(
@@ -732,164 +717,40 @@ public:
 		const FSharedString& Name,
 		FCacheKey Key,
 		FCbPackage&& Package,
-		FCacheRecordPolicy Policy,
-		uint64 UserData,
-		FOnCachePutPackageComplete&& OnComplete);
+		FOnPackageComplete&& OnComplete);
 
 private:
 	FHttpCacheStore& CacheStore;
 	IRequestOwner& Owner;
 	const FSharedString Name;
-	const FCacheKey Key;
-	const uint64 UserData;
-	std::atomic<uint64> BytesSent;
-	const FCbObject PackageObject;
-	const FIoHash PackageObjectHash;
-	const uint32 TotalBlobUploads;
-	std::atomic<uint32> SuccessfulBlobUploads;
-	std::atomic<uint32> PendingBlobUploads;
-	FOnCachePutPackageComplete OnComplete;
+
+	FCacheKey Key;
+	FCbObject Object;
+	FIoHash ObjectHash;
+	FOnPackageComplete OnPackageComplete;
+
+	std::atomic<uint64> BytesSent = 0;
+	std::atomic<uint32> SuccessfulBlobUploads = 0;
+	std::atomic<uint32> PendingBlobUploads = 0;
+	uint32 TotalBlobUploads = 0;
 
 	struct FCachePutRefResponse
 	{
-		FSharedString Name;
-		FCacheKey Key;
-		uint64 UserData = 0;
 		uint64 BytesSent = 0;
 		TConstArrayView<FIoHash> NeededBlobHashes;
 		EStatus Status = EStatus::Error;
 	};
 	using FOnCachePutRefComplete = TUniqueFunction<void(FCachePutRefResponse&& Response)>;
 
-	FPutPackageOp(
-		FHttpCacheStore& InCacheStore,
-		IRequestOwner& InOwner,
-		const FSharedString& InName,
-		const FCacheKey& InKey,
-		uint64 InUserData,
-		uint64 InBytesSent,
-		const FCbObject& InPackageObject,
-		const FIoHash& InPackageObjectHash,
-		uint32 InTotalBlobUploads,
-		FOnCachePutPackageComplete&& InOnComplete);
-
-	static void PutRefAsync(
-		FHttpCacheStore& CacheStore,
-		IRequestOwner& Owner,
-		const FSharedString& Name,
-		FCacheKey Key,
-		FCbObject Object,
-		FIoHash ObjectHash,
-		uint64 UserData,
-		bool bFinalize,
-		FOnCachePutRefComplete&& OnComplete);
-
-	static void OnPackagePutRefComplete(
-		FHttpCacheStore& CacheStore,
-		IRequestOwner& Owner,
-		const FSharedString& Name,
-		const FCacheKey& Key,
-		FCbPackage&& Package,
-		FCacheRecordPolicy Policy,
-		uint64 UserData,
-		FOnCachePutPackageComplete&& OnComplete,
-		FCachePutRefResponse&& Response);
-
-	void OnCompressedBlobUploadComplete(FHttpOperation& Operation);
-
-	void OnPutRefFinalizationComplete(FCachePutRefResponse&& Response);
-
-	FCachePutPackageResponse MakeResponse(uint64 InBytesSent, EStatus Status)
-	{
-		return FCachePutPackageResponse{ Name, Key, UserData, InBytesSent, Status };
-	};
-};
-
-//----------------------------------------------------------------------------------------------------------
-// FHttpCacheStore::FGetRecordOp
-//----------------------------------------------------------------------------------------------------------
-class FHttpCacheStore::FGetRecordOp final : public FThreadSafeRefCountedObject
-{
-public:
-	/** Performs a multi-request operation for downloading a record. */
-	static void GetRecord(
-		FHttpCacheStore& CacheStore,
-		IRequestOwner& Owner,
-		const FSharedString& Name,
-		const FCacheKey& Key,
-		const FCacheRecordPolicy& Policy,
-		uint64 UserData,
-		TUniqueFunction<void(FCacheGetResponse&& Response, uint64 BytesReceived)>&& OnComplete);
-
-	struct FGetCachedDataBatchResponse
-	{
-		FSharedString Name;
-		FCacheKey Key;
-		int32 ValueIndex;
-		uint64 BytesReceived = 0;
-		FCompressedBuffer DataBuffer;
-		EStatus Status = EStatus::Error;
-	};
-	using FOnGetCachedDataBatchComplete = TUniqueFunction<void(FGetCachedDataBatchResponse&& Response)>;
+	FPutPackageOp(FHttpCacheStore& CacheStore, IRequestOwner& Owner, const FSharedString& Name);
 
-	/** Utility method for fetching a batch of value data. */
-	template <typename ValueType, typename ValueIdGetterType>
-	static void GetDataBatch(
-		FHttpCacheStore& CacheStore,
-		IRequestOwner& Owner,
-		FSharedString Name,
-		const FCacheKey& Key,
-		TConstArrayView<ValueType> Values,
-		ValueIdGetterType ValueIdGetter,
-		FOnGetCachedDataBatchComplete&& OnComplete);
-private:
-	FHttpCacheStore& CacheStore;
-	IRequestOwner& Owner;
-	const FSharedString Name;
-	const FCacheKey Key;
-	const uint64 UserData;
-	std::atomic<uint64> BytesReceived;
-	TArray<FCompressedBuffer> FetchedBuffers;
-	const TArray<FValueWithId> RequiredGets;
-	const TArray<FValueWithId> RequiredHeads;
-	FCacheRecordBuilder RecordBuilder;
-	const uint32 TotalOperations;
-	std::atomic<uint32> SuccessfulOperations;
-	std::atomic<uint32> PendingOperations;
-	TUniqueFunction<void(FCacheGetResponse&& Response, uint64 BytesReceived)> OnComplete;
-
-	FGetRecordOp(
-		FHttpCacheStore& InCacheStore,
-		IRequestOwner& InOwner,
-		const FSharedString& InName,
-		const FCacheKey& InKey,
-		uint64 InUserData,
-		uint64 InBytesReceived,
-		TArray<FValueWithId>&& InRequiredGets,
-		TArray<FValueWithId>&& InRequiredHeads,
-		FCacheRecordBuilder&& InRecordBuilder,
-		TUniqueFunction<void(FCacheGetResponse&& Response, uint64 BytesReceived)>&& InOnComplete);
-
-	static void OnOnlyRecordComplete(
-		FHttpCacheStore& CacheStore,
-		IRequestOwner& Owner,
-		const FCacheRecordPolicy& Policy,
-		TUniqueFunction<void(FCacheGetResponse&& Response, uint64 BytesReceived)>&& OnComplete,
-		FGetCacheRecordOnlyResponse&& Response);
+	void BeginPutRef(bool bFinalize, FOnCachePutRefComplete&& OnComplete);
+	void EndPutRef(TUniquePtr<FHttpOperation> Operation, bool bFinalize, FOnCachePutRefComplete&& OnComplete);
 
-	struct FCachedDataProbablyExistsBatchResponse
-	{
-		FSharedString Name;
-		FCacheKey Key;
-		int32 ValueIndex;
-		EStatus Status = EStatus::Error;
-	};
-	using FOnCachedDataProbablyExistsBatchComplete = TUniqueFunction<void(FCachedDataProbablyExistsBatchResponse&& Response)>;
-	void DataProbablyExistsBatch(
-		TConstArrayView<FValueWithId> Values,
-		FOnCachedDataProbablyExistsBatchComplete&& OnComplete);
+	void BeginPutBlobs(FCbPackage&& Package, FCachePutRefResponse&& Response);
+	void EndPutBlob(FHttpOperation& Operation);
 
-	void FinishDataStep(bool bSuccess, uint64 InBytesReceived);
+	void EndPutRefFinalize(FCachePutRefResponse&& Response);
 };
 
 void FHttpCacheStore::FPutPackageOp::PutPackage(
@@ -898,57 +759,34 @@ void FHttpCacheStore::FPutPackageOp::PutPackage(
 	const FSharedString& Name,
 	FCacheKey Key,
 	FCbPackage&& Package,
-	FCacheRecordPolicy Policy,
-	uint64 UserData,
-	FOnCachePutPackageComplete&& OnComplete)
+	FOnPackageComplete&& OnComplete)
 {
 	// TODO: Jupiter currently always overwrites.  It doesn't have a "write if not present" feature (for records or attachments),
 	//		 but would require one to implement all policy correctly.
-
-	// Initial record upload
-	PutRefAsync(CacheStore, Owner, Name, Key, Package.GetObject(), Package.GetObjectHash(), UserData, false,
-		[&CacheStore, &Owner, Name = FSharedString(Name), Key, Package = MoveTemp(Package), Policy, UserData, OnComplete = MoveTemp(OnComplete)](FCachePutRefResponse&& Response) mutable
-		{
-			return OnPackagePutRefComplete(CacheStore, Owner, Name, Key, MoveTemp(Package), Policy, UserData, MoveTemp(OnComplete), MoveTemp(Response));
-		});
+	TRefCountPtr<FPutPackageOp> Self = New(CacheStore, Owner, Name);
+	Self->Put(Key, MoveTemp(Package), MoveTemp(OnComplete));
 }
 
-FHttpCacheStore::FPutPackageOp::FPutPackageOp(
-	FHttpCacheStore& InCacheStore,
-	IRequestOwner& InOwner,
-	const FSharedString& InName,
-	const FCacheKey& InKey,
-	uint64 InUserData,
-	uint64 InBytesSent,
-	const FCbObject& InPackageObject,
-	const FIoHash& InPackageObjectHash,
-	uint32 InTotalBlobUploads,
-	FOnCachePutPackageComplete&& InOnComplete)
+FHttpCacheStore::FPutPackageOp::FPutPackageOp(FHttpCacheStore& InCacheStore, IRequestOwner& InOwner, const FSharedString& InName)
 	: CacheStore(InCacheStore)
 	, Owner(InOwner)
 	, Name(InName)
-	, Key(InKey)
-	, UserData(InUserData)
-	, BytesSent(InBytesSent)
-	, PackageObject(InPackageObject)
-	, PackageObjectHash(InPackageObjectHash)
-	, TotalBlobUploads(InTotalBlobUploads)
-	, SuccessfulBlobUploads(0)
-	, PendingBlobUploads(InTotalBlobUploads)
-	, OnComplete(MoveTemp(InOnComplete))
 {
 }
 
-void FHttpCacheStore::FPutPackageOp::PutRefAsync(
-	FHttpCacheStore& CacheStore,
-	IRequestOwner& Owner,
-	const FSharedString& Name,
-	FCacheKey Key,
-	FCbObject Object,
-	FIoHash ObjectHash,
-	uint64 UserData,
-	bool bFinalize,
-	FOnCachePutRefComplete&& OnComplete)
+void FHttpCacheStore::FPutPackageOp::Put(const FCacheKey& InKey, FCbPackage&& Package, FOnPackageComplete&& OnComplete)
+{
+	Key = InKey;
+	Object = Package.GetObject();
+	ObjectHash = Package.GetObjectHash();
+	OnPackageComplete = MoveTemp(OnComplete);
+	BeginPutRef(/*bFinalize*/ false, [Self = TRefCountPtr(this), Package = MoveTemp(Package)](FCachePutRefResponse&& Response) mutable
+	{
+		return Self->BeginPutBlobs(MoveTemp(Package), MoveTemp(Response));
+	});
+}
+
+void FHttpCacheStore::FPutPackageOp::BeginPutRef(bool bFinalize, FOnCachePutRefComplete&& OnComplete)
 {
 	TAnsiStringBuilder<64> Bucket;
 	Algo::Transform(Key.Bucket.ToString(), AppendChars(Bucket), FCharAnsi::ToLower);
@@ -976,82 +814,83 @@ void FHttpCacheStore::FPutPackageOp::PutRefAsync(
 		LocalOperation.SetBody(Object.GetBuffer());
 	}
 	LocalOperation.AddAcceptType(EHttpMediaType::Json);
-	LocalOperation.SendAsync(Owner, [Operation = MoveTemp(Operation), &CacheStore, Name, Key, Object, UserData, bFinalize, OnComplete = MoveTemp(OnComplete)]
+	LocalOperation.SendAsync(Owner, [Self = TRefCountPtr(this), Operation = MoveTemp(Operation), bFinalize, OnComplete = MoveTemp(OnComplete)]() mutable
 	{
-		TRACE_CPUPROFILER_EVENT_SCOPE(HttpDDC_PutRefAsync_OnHttpRequestComplete);
+		Self->EndPutRef(MoveTemp(Operation), bFinalize, MoveTemp(OnComplete));
+	});
+}
 
-		const int32 StatusCode = Operation->GetStatusCode();
-		if (StatusCode >= 200 && StatusCode <= 204)
-		{
-			TArray<FIoHash> NeededBlobHashes;
+void FHttpCacheStore::FPutPackageOp::EndPutRef(
+	TUniquePtr<FHttpOperation> Operation,
+	bool bFinalize,
+	FOnCachePutRefComplete&& OnComplete)
+{
+	TRACE_CPUPROFILER_EVENT_SCOPE(HttpDDC_PutPackage_EndPutRef);
+
+	if (const int32 StatusCode = Operation->GetStatusCode(); StatusCode >= 200 && StatusCode <= 204)
+	{
+		TArray<FIoHash> NeededBlobHashes;
 
-			// Useful when debugging issues related to compressed/uncompressed blobs being returned from Jupiter
-			static const bool bHttpCacheAlwaysPut = FParse::Param(FCommandLine::Get(), TEXT("HttpCacheAlwaysPut"));
+		// Useful when debugging issues related to compressed/uncompressed blobs being returned from Jupiter
+		static const bool bHttpCacheAlwaysPut = FParse::Param(FCommandLine::Get(), TEXT("HttpCacheAlwaysPut"));
 
-			if (bHttpCacheAlwaysPut && !bFinalize)
+		if (bHttpCacheAlwaysPut && !bFinalize)
+		{
+			Object.IterateAttachments([&NeededBlobHashes](FCbFieldView AttachmentFieldView)
 			{
-				Object.IterateAttachments([&NeededBlobHashes](FCbFieldView AttachmentFieldView)
+				FIoHash AttachmentHash = AttachmentFieldView.AsHash();
+				if (!AttachmentHash.IsZero())
 				{
-					FIoHash AttachmentHash = AttachmentFieldView.AsHash();
-					if (!AttachmentHash.IsZero())
-					{
-						NeededBlobHashes.Add(AttachmentHash);
-					}
-				});
-			}
-			else if (TSharedPtr<FJsonObject> ResponseObject = Operation->GetBodyAsJson())
-			{
-				TArray<FString> NeedsArrayStrings;
-				ResponseObject->TryGetStringArrayField(TEXT("needs"), NeedsArrayStrings);
+					NeededBlobHashes.Add(AttachmentHash);
+				}
+			});
+		}
+		else if (TSharedPtr<FJsonObject> ResponseObject = Operation->GetBodyAsJson())
+		{
+			TArray<FString> NeedsArrayStrings;
+			ResponseObject->TryGetStringArrayField(TEXT("needs"), NeedsArrayStrings);
 
-				NeededBlobHashes.Reserve(NeedsArrayStrings.Num());
-				for (const FString& NeededString : NeedsArrayStrings)
+			NeededBlobHashes.Reserve(NeedsArrayStrings.Num());
+			for (const FString& NeededString : NeedsArrayStrings)
+			{
+				FIoHash BlobHash;
+				LexFromString(BlobHash, *NeededString);
+				if (!BlobHash.IsZero())
 				{
-					FIoHash BlobHash;
-					LexFromString(BlobHash, *NeededString);
-					if (!BlobHash.IsZero())
-					{
-						NeededBlobHashes.Add(BlobHash);
-					}
+					NeededBlobHashes.Add(BlobHash);
 				}
 			}
-
-			OnComplete({ Name, Key, UserData, Operation->GetBytesSent(), NeededBlobHashes, EStatus::Ok });
 		}
-		else
-		{
-			const EStatus Status = Operation->GetErrorCode() == EHttpErrorCode::Canceled ? EStatus::Canceled : EStatus::Error;
-			OnComplete({ Name, Key, UserData, Operation->GetBytesSent(), {}, Status });
-		}
-	});
+
+		OnComplete({Operation->GetBytesSent(), NeededBlobHashes, EStatus::Ok});
+	}
+	else
+	{
+		const EStatus Status = Operation->GetErrorCode() == EHttpErrorCode::Canceled ? EStatus::Canceled : EStatus::Error;
+		OnComplete({Operation->GetBytesSent(), {}, Status});
+	}
 }
 
-void FHttpCacheStore::FPutPackageOp::OnPackagePutRefComplete(
-	FHttpCacheStore& CacheStore,
-	IRequestOwner& Owner,
-	const FSharedString& Name,
-	const FCacheKey& Key,
-	FCbPackage&& Package,
-	FCacheRecordPolicy Policy,
-	uint64 UserData,
-	FOnCachePutPackageComplete&& OnComplete,
-	FCachePutRefResponse&& Response)
+void FHttpCacheStore::FPutPackageOp::BeginPutBlobs(FCbPackage&& Package, FCachePutRefResponse&& Response)
 {
 	if (Response.Status != EStatus::Ok)
 	{
 		if (Response.Status == EStatus::Error)
 		{
 			UE_LOG(LogDerivedDataCache, Log, TEXT("%s: Failed to put reference object for put of %s from '%s'"),
-				*CacheStore.Domain, *WriteToString<96>(Response.Key), *Response.Name);
+				*CacheStore.Domain, *WriteToString<96>(Key), *Name);
 		}
-		return OnComplete(FCachePutPackageResponse{ Name, Key, UserData, Response.BytesSent, Response.Status });
+		return OnPackageComplete({Response.BytesSent, Response.Status});
 	}
 
 	struct FCompressedBlobUpload
 	{
 		FIoHash Hash;
 		FSharedBuffer BlobBuffer;
-		FCompressedBlobUpload(const FIoHash& InHash, FSharedBuffer&& InBlobBuffer) : Hash(InHash), BlobBuffer(InBlobBuffer)
+
+		FCompressedBlobUpload(const FIoHash& InHash, FSharedBuffer&& InBlobBuffer)
+			: Hash(InHash)
+			, BlobBuffer(InBlobBuffer)
 		{
 		}
 	};
@@ -1100,29 +939,18 @@ void FHttpCacheStore::FPutPackageOp::OnPackagePutRefComplete(
 				bExpectedHashesSerialized = true;
 			}
 			UE_LOG(LogDerivedDataCache, Log, TEXT("%s: Server reported needed hash '%s' that is outside the set of expected hashes (%s) for put of %s from '%s'"),
-				*CacheStore.Domain, *WriteToString<96>(NeededBlobHash), ExpectedHashes.ToString(), *WriteToString<96>(Response.Key), *Response.Name);
+				*CacheStore.Domain, *WriteToString<96>(NeededBlobHash), ExpectedHashes.ToString(), *WriteToString<96>(Key), *Name);
 		}
 	}
 
 	if (CompressedBlobUploads.IsEmpty())
 	{
 		// No blobs need to be uploaded.  No finalization necessary.
-		return OnComplete(FCachePutPackageResponse{ Name, Key, UserData, Response.BytesSent, EStatus::Ok });
-	}
-
-	// Having this be a ref ensures we don't have the op reach 0 ref count as we queue up multiple operations which MAY execute synchronously
-	TRefCountPtr<FPutPackageOp> PutPackageOp = new FPutPackageOp(
-		CacheStore,
-		Owner,
-		Response.Name,
-		Response.Key,
-		Response.UserData,
-		Response.BytesSent,
-		Package.GetObject(),
-		Package.GetObjectHash(),
-		(uint32)CompressedBlobUploads.Num(),
-		MoveTemp(OnComplete)
-	);
+		return OnPackageComplete({Response.BytesSent, EStatus::Ok});
+	}
+
+	TotalBlobUploads = CompressedBlobUploads.Num();
+	PendingBlobUploads.store(TotalBlobUploads, std::memory_order_relaxed);
 
 	FRequestBarrier Barrier(Owner);
 	for (const FCompressedBlobUpload& CompressedBlobUpload : CompressedBlobUploads)
@@ -1133,14 +961,14 @@ void FHttpCacheStore::FPutPackageOp::OnPackagePutRefComplete(
 		LocalOperation.SetMethod(EHttpMethod::Put);
 		LocalOperation.SetContentType(EHttpMediaType::CompressedBinary);
 		LocalOperation.SetBody(FCompositeBuffer(CompressedBlobUpload.BlobBuffer));
-		LocalOperation.SendAsync(Owner, [Operation = MoveTemp(Operation), PutPackageOp]
+		LocalOperation.SendAsync(Owner, [Self = TRefCountPtr(this), Operation = MoveTemp(Operation)]
 		{
-			PutPackageOp->OnCompressedBlobUploadComplete(*Operation);
+			Self->EndPutBlob(*Operation);
 		});
 	}
 }
 
-void FHttpCacheStore::FPutPackageOp::OnCompressedBlobUploadComplete(FHttpOperation& Operation)
+void FHttpCacheStore::FPutPackageOp::EndPutBlob(FHttpOperation& Operation)
 {
 	BytesSent.fetch_add(Operation.GetBytesSent(), std::memory_order_relaxed);
 
@@ -1155,28 +983,26 @@ void FHttpCacheStore::FPutPackageOp::OnCompressedBlobUploadComplete(FHttpOperati
 		const uint32 LocalSuccessfulBlobUploads = SuccessfulBlobUploads.load(std::memory_order_relaxed);
 		if (Owner.IsCanceled())
 		{
-			OnComplete(MakeResponse(BytesSent.load(std::memory_order_relaxed), EStatus::Canceled));
+			OnPackageComplete({BytesSent.load(std::memory_order_relaxed), EStatus::Canceled});
 		}
 		else if (LocalSuccessfulBlobUploads == TotalBlobUploads)
 		{
-			// Perform finalization
-			PutRefAsync(CacheStore, Owner, Name, Key, PackageObject, PackageObjectHash, UserData, /*bFinalize*/ true,
-				[PutPackageOp = TRefCountPtr<FPutPackageOp>(this)](FCachePutRefResponse&& Response)
-				{
-					return PutPackageOp->OnPutRefFinalizationComplete(MoveTemp(Response));
-				});
+			BeginPutRef(/*bFinalize*/ true, [Self = TRefCountPtr(this)](FCachePutRefResponse&& Response)
+			{
+				return Self->EndPutRefFinalize(MoveTemp(Response));
+			});
 		}
 		else
 		{
 			const uint32 FailedBlobUploads = TotalBlobUploads - LocalSuccessfulBlobUploads;
 			UE_LOG(LogDerivedDataCache, Log, TEXT("%s: Failed to put %d/%d blobs for put of %s from '%s'"),
 				*CacheStore.Domain, FailedBlobUploads, TotalBlobUploads, *WriteToString<96>(Key), *Name);
-			OnComplete(MakeResponse(BytesSent.load(std::memory_order_relaxed), EStatus::Error));
+			OnPackageComplete({BytesSent.load(std::memory_order_relaxed), EStatus::Error});
 		}
 	}
 }
 
-void FHttpCacheStore::FPutPackageOp::OnPutRefFinalizationComplete(FCachePutRefResponse&& Response)
+void FHttpCacheStore::FPutPackageOp::EndPutRefFinalize(FCachePutRefResponse&& Response)
 {
 	BytesSent.fetch_add(Response.BytesSent, std::memory_order_relaxed);
 
@@ -1186,53 +1012,305 @@ void FHttpCacheStore::FPutPackageOp::OnPutRefFinalizationComplete(FCachePutRefRe
 			*CacheStore.Domain, *WriteToString<96>(Key), *Name);
 	}
 
-	return OnComplete(MakeResponse(BytesSent.load(std::memory_order_relaxed), Response.Status));
+	return OnPackageComplete({BytesSent.load(std::memory_order_relaxed), Response.Status});
 }
 
-void FHttpCacheStore::FGetRecordOp::GetRecord(
-	FHttpCacheStore& CacheStore,
-	IRequestOwner& Owner,
-	const FSharedString& Name,
-	const FCacheKey& Key,
-	const FCacheRecordPolicy& Policy,
-	uint64 UserData,
-	TUniqueFunction<void(FCacheGetResponse&& Response, uint64 BytesReceived)>&& OnComplete)
+//----------------------------------------------------------------------------------------------------------
+// FHttpCacheStore::FGetRecordOp
+//----------------------------------------------------------------------------------------------------------
+class FHttpCacheStore::FGetRecordOp final : public FThreadSafeRefCountedObject
+{
+public:
+	static TRefCountPtr<FGetRecordOp> New(FHttpCacheStore& CacheStore, IRequestOwner& Owner, const FSharedString& Name)
+	{
+		return new FGetRecordOp(CacheStore, Owner, Name);
+	}
+
+	struct FRecordResponse
+	{
+		uint64 BytesReceived = 0;
+		FCacheRecord Record;
+		EStatus Status = EStatus::Error;
+	};
+	using FOnRecordComplete = TUniqueFunction<void (FRecordResponse&& Response)>;
+
+	void GetRecordOnly(const FCacheKey& Key, const ECachePolicy RecordPolicy, FOnRecordComplete&& OnComplete);
+	void GetRecord(const FCacheKey& Key, const FCacheRecordPolicy& Policy, FOnRecordComplete&& OnComplete);
+
+	struct FValueResponse
+	{
+		FValueWithId Value;
+		EStatus Status = EStatus::Error;
+	};
+	using FOnValueComplete = TUniqueFunction<void (FValueResponse&& Response)>;
+
+	void GetValues(TConstArrayView<FValueWithId> Values, FOnValueComplete&& OnComplete);
+	void GetValuesExist(TConstArrayView<FValueWithId> Values, FOnValueComplete&& OnComplete);
+
+private:
+	FGetRecordOp(FHttpCacheStore& CacheStore, IRequestOwner& Owner, const FSharedString& Name);
+
+	void EndGetRef(TUniquePtr<FHttpOperation> Operation);
+
+	void BeginGetValues(const FCacheRecord& Record, const FCacheRecordPolicy& Policy, FOnRecordComplete&& OnComplete);
+	void EndGetValues(const FCacheRecordPolicy& Policy, EStatus Status);
+
+	FHttpCacheStore& CacheStore;
+	IRequestOwner& Owner;
+	FSharedString Name;
+	FCacheKey Key;
+	FCbPackage Package;
+	FOnRecordComplete OnRecordComplete;
+	int32 PendingValues = 0;
+	int32 FailedValues = 0;
+	mutable FMutex Mutex;
+};
+
+FHttpCacheStore::FGetRecordOp::FGetRecordOp(FHttpCacheStore& InCacheStore, IRequestOwner& InOwner, const FSharedString& InName)
+	: CacheStore(InCacheStore)
+	, Owner(InOwner)
+	, Name(InName)
+{
+}
+
+void FHttpCacheStore::FGetRecordOp::GetRecordOnly(const FCacheKey& InKey, const ECachePolicy RecordPolicy, FOnRecordComplete&& InOnComplete)
+{
+	if (!CacheStore.IsUsable())
+	{
+		UE_LOG(LogDerivedDataCache, VeryVerbose,
+			TEXT("%s: Skipped get of %s from '%s' because this cache store is not available"),
+			*CacheStore.Domain, *WriteToString<96>(Key), *Name);
+		return InOnComplete({0, FCacheRecordBuilder(Key).Build(), EStatus::Error});
+	}
+
+	// Skip the request if querying the cache is disabled.
+	if (!EnumHasAnyFlags(RecordPolicy, ECachePolicy::QueryRemote))
+	{
+		UE_LOG(LogDerivedDataCache, VeryVerbose, TEXT("%s: Skipped get of %s from '%s' due to cache policy"),
+			*CacheStore.Domain, *WriteToString<96>(Key), *Name);
+		return InOnComplete({0, FCacheRecordBuilder(Key).Build(), EStatus::Error});
+	}
+
+	if (CacheStore.DebugOptions.ShouldSimulateGetMiss(Key))
+	{
+		UE_LOG(LogDerivedDataCache, Verbose, TEXT("%s: Simulated miss for get of %s from '%s'"),
+			*CacheStore.Domain, *WriteToString<96>(Key), *Name);
+		return InOnComplete({0, FCacheRecordBuilder(Key).Build(), EStatus::Error});
+	}
+
+	Key = InKey;
+	OnRecordComplete = MoveTemp(InOnComplete);
+
+	TAnsiStringBuilder<64> Bucket;
+	Algo::Transform(Key.Bucket.ToString(), AppendChars(Bucket), FCharAnsi::ToLower);
+
+	TUniquePtr<FHttpOperation> Operation = CacheStore.WaitForHttpOperation(EOperationCategory::Get);
+	FHttpOperation& LocalOperation = *Operation;
+	LocalOperation.SetUri(WriteToAnsiString<256>(CacheStore.EffectiveDomain, ANSITEXTVIEW("/api/v1/refs/"), CacheStore.Namespace, '/', Bucket, '/', Key.Hash));
+	LocalOperation.SetMethod(EHttpMethod::Get);
+	LocalOperation.AddAcceptType(EHttpMediaType::CbObject);
+	LocalOperation.SetExpectedErrorCodes({404});
+	LocalOperation.SendAsync(Owner, [Self = TRefCountPtr(this), Operation = MoveTemp(Operation)]() mutable
+	{
+		Self->EndGetRef(MoveTemp(Operation));
+	});
+}
+
+void FHttpCacheStore::FGetRecordOp::EndGetRef(TUniquePtr<FHttpOperation> Operation)
 {
-	CacheStore.GetCacheRecordOnlyAsync(Owner, Name, Key, Policy, UserData, [&CacheStore, &Owner, Policy = FCacheRecordPolicy(Policy), OnComplete = MoveTemp(OnComplete)](FGetCacheRecordOnlyResponse&& Response) mutable
+	TRACE_CPUPROFILER_EVENT_SCOPE(HttpDDC_GetPackage_EndGetRef);
+
+	FOptionalCacheRecord Record;
+	EStatus Status = EStatus::Error;
+	ON_SCOPE_EXIT
+	{
+		const uint64 BytesReceived = Operation->GetBytesReceived();
+		Operation.Reset();
+		if (Record.IsNull())
+		{
+			Record = FCacheRecordBuilder(Key).Build();
+		}
+		FOnRecordComplete LocalOnComplete = MoveTemp(OnRecordComplete);
+		LocalOnComplete({BytesReceived, MoveTemp(Record).Get(), Status});
+	};
+
+	const int32 StatusCode = Operation->GetStatusCode();
+	if (StatusCode < 200 || StatusCode > 204)
 	{
-		OnOnlyRecordComplete(CacheStore, Owner, Policy, MoveTemp(OnComplete), MoveTemp(Response));
+		UE_LOG(LogDerivedDataCache, Verbose, TEXT("%s: Cache miss with missing package for %s from '%s'"),
+			*CacheStore.Domain, *WriteToString<96>(Key), *Name);
+		return;
+	}
+
+	FSharedBuffer Body = Operation->GetBody();
+
+	if (ValidateCompactBinary(Body, ECbValidateMode::Default) != ECbValidateError::None)
+	{
+		UE_LOG(LogDerivedDataCache, Log, TEXT("%s: Cache miss with invalid package for %s from '%s'"),
+			*CacheStore.Domain, *WriteToString<96>(Key), *Name);
+		return;
+	}
+
+	Package = FCbPackage(FCbObject(Body));
+	Record = FCacheRecord::Load(Package);
+
+	if (Record.IsNull())
+	{
+		UE_LOG(LogDerivedDataCache, Log, TEXT("%s: Cache miss with record load failure for %s from '%s'"),
+			*CacheStore.Domain, *WriteToString<96>(Key), *Name);
+		return;
+	}
+
+	Status = EStatus::Ok;
+}
+
+void FHttpCacheStore::FGetRecordOp::GetRecord(const FCacheKey& LocalKey, const FCacheRecordPolicy& Policy, FOnRecordComplete&& OnComplete)
+{
+	GetRecordOnly(LocalKey, Policy.GetRecordPolicy(), [Self = TRefCountPtr(this), Policy, OnComplete = MoveTemp(OnComplete)](FRecordResponse&& Response) mutable
+	{
+		if (Response.Status == EStatus::Ok)
+		{
+			Self->BeginGetValues(Response.Record, Policy, MoveTemp(OnComplete));
+		}
+		else
+		{
+			OnComplete(MoveTemp(Response));
+		}
 	});
 }
 
-template <typename ValueType, typename ValueIdGetterType>
-void FHttpCacheStore::FGetRecordOp::GetDataBatch(
-	FHttpCacheStore& CacheStore,
-	IRequestOwner& Owner,
-	FSharedString Name,
-	const FCacheKey& Key,
-	TConstArrayView<ValueType> Values,
-	ValueIdGetterType ValueIdGetter,
-	FOnGetCachedDataBatchComplete&& OnComplete)
+void FHttpCacheStore::FGetRecordOp::BeginGetValues(const FCacheRecord& Record, const FCacheRecordPolicy& Policy, FOnRecordComplete&& OnComplete)
+{
+	OnRecordComplete = MoveTemp(OnComplete);
+
+	TArray<FValueWithId> RequiredGets;
+	TArray<FValueWithId> RequiredHeads;
+
+	for (const FValueWithId& Value : Record.GetValues())
+	{
+		const ECachePolicy ValuePolicy = Policy.GetValuePolicy(Value.GetId());
+		if (EnumHasAnyFlags(ValuePolicy, ECachePolicy::QueryRemote))
+		{
+			(EnumHasAnyFlags(ValuePolicy, ECachePolicy::SkipData) ? RequiredHeads : RequiredGets).Emplace(Value);
+		}
+	}
+
+	PendingValues = RequiredGets.Num() + RequiredHeads.Num();
+
+	if (PendingValues == 0)
+	{
+		return EndGetValues(Policy, EStatus::Ok);
+	}
+
+	GetValues(RequiredGets, [Self = TRefCountPtr(this), Policy](FValueResponse&& Response)
+	{
+		TDynamicUniqueLock Lock(Self->Mutex);
+		const bool bComplete = --Self->PendingValues == 0;
+		if (Response.Value.HasData())
+		{
+			Self->Package.AddAttachment(FCbAttachment(Response.Value.GetData()));
+		}
+		else
+		{
+			++Self->FailedValues;
+		}
+		if (bComplete)
+		{
+			Lock.Unlock();
+			Self->EndGetValues(Policy, Response.Status);
+		}
+	});
+
+	GetValuesExist(RequiredHeads, [Self = TRefCountPtr(this), Policy](FValueResponse&& Response)
+	{
+		TDynamicUniqueLock Lock(Self->Mutex);
+		const bool bComplete = --Self->PendingValues == 0;
+		if (Response.Status != EStatus::Ok)
+		{
+			++Self->FailedValues;
+		}
+		if (bComplete)
+		{
+			Lock.Unlock();
+			Self->EndGetValues(Policy, Response.Status);
+		}
+	});
+}
+
+void FHttpCacheStore::FGetRecordOp::EndGetValues(const FCacheRecordPolicy& Policy, EStatus Status)
+{
+	FCacheRecordBuilder RecordBuilder(Key);
+	if (FOptionalCacheRecord Record = FCacheRecord::Load(Package))
+	{
+		if (!EnumHasAnyFlags(Policy.GetRecordPolicy(), ECachePolicy::SkipMeta))
+		{
+			RecordBuilder.SetMeta(CopyTemp(Record.Get().GetMeta()));
+		}
+		for (const FValueWithId& Value : Record.Get().GetValues())
+		{
+			const ECachePolicy ValuePolicy = Policy.GetValuePolicy(Value.GetId());
+			if (EnumHasAnyFlags(ValuePolicy, ECachePolicy::QueryRemote) && !EnumHasAnyFlags(ValuePolicy, ECachePolicy::SkipData))
+			{
+				if (Status == EStatus::Ok && !Value.HasData())
+				{
+					Status = EStatus::Error;
+				}
+				RecordBuilder.AddValue(Value);
+			}
+			else
+			{
+				RecordBuilder.AddValue(Value.RemoveData());
+			}
+		}
+	}
+
+	if (FailedValues)
+	{
+		Status = EStatus::Error;
+	}
+
+	FOnRecordComplete LocalOnComplete = MoveTemp(OnRecordComplete);
+	LocalOnComplete({0, RecordBuilder.Build(), Status});
+}
+
+void FHttpCacheStore::FGetRecordOp::GetValues(TConstArrayView<FValueWithId> Values, FOnValueComplete&& OnComplete)
 {
-	if (Values.IsEmpty())
+	int32 MissingDataCount = 0;
+	for (const FValueWithId& Value : Values)
+	{
+		if (Value.HasData())
+		{
+			OnComplete({Value, EStatus::Ok});
+			continue;
+		}
+		++MissingDataCount;
+	}
+
+	if (MissingDataCount == 0)
 	{
 		return;
 	}
 
+	// TODO: Jupiter does not currently provide a batched GET. Once it does, fetch every blob in one request.
+
 	FRequestBarrier Barrier(Owner);
-	TSharedRef<FOnGetCachedDataBatchComplete> SharedOnComplete = MakeShared<FOnGetCachedDataBatchComplete>(MoveTemp(OnComplete));
-	for (int32 ValueIndex = 0; ValueIndex < Values.Num(); ++ValueIndex)
+	TSharedRef<FOnValueComplete> SharedOnComplete = MakeShared<FOnValueComplete>(MoveTemp(OnComplete));
+	for (const FValueWithId& Value : Values)
 	{
-		const ValueType Value = Values[ValueIndex].RemoveData();
+		if (Value.HasData())
+		{
+			continue;
+		}
+
 		TUniquePtr<FHttpOperation> Operation = CacheStore.WaitForHttpOperation(EOperationCategory::Get);
 		FHttpOperation& LocalOperation = *Operation;
 		LocalOperation.SetUri(WriteToAnsiString<256>(CacheStore.EffectiveDomain, ANSITEXTVIEW("/api/v1/compressed-blobs/"), CacheStore.Namespace, '/', Value.GetRawHash()));
 		LocalOperation.SetMethod(EHttpMethod::Get);
 		LocalOperation.AddAcceptType(EHttpMediaType::Any);
 		LocalOperation.SetExpectedErrorCodes({404});
-		LocalOperation.SendAsync(Owner, [Operation = MoveTemp(Operation), &CacheStore, Name, Key, Value, ValueIndex, ValueIdGetter, SharedOnComplete]
+		LocalOperation.SendAsync(Owner, [Self = TRefCountPtr(this), Operation = MoveTemp(Operation), SharedOnComplete, Id = Value.GetId(), RawHash = Value.GetRawHash(), RawSize = Value.GetRawSize()]
 		{
-			TRACE_CPUPROFILER_EVENT_SCOPE(HttpDDC_GetDataBatch_OnHttpRequestComplete);
+			TRACE_CPUPROFILER_EVENT_SCOPE(HttpDDC_GetPackage_GetValues_OnResponse);
 
 			bool bHit = false;
 			FCompressedBuffer CompressedBuffer;
@@ -1256,302 +1334,452 @@ void FHttpCacheStore::FGetRecordOp::GetDataBatch(
 
 			if (bHit)
 			{
-				if (CompressedBuffer.GetRawHash() == Value.GetRawHash())
+				if (CompressedBuffer.GetRawHash() == RawHash && CompressedBuffer.GetRawSize() == RawSize)
 				{
-					SharedOnComplete.Get()({ Name, Key, ValueIndex, Operation->GetBytesReceived(), MoveTemp(CompressedBuffer), EStatus::Ok });
+					SharedOnComplete.Get()({FValueWithId(Id, MoveTemp(CompressedBuffer)), EStatus::Ok});
 				}
 				else
 				{
 					UE_LOG(LogDerivedDataCache, Display,
 						TEXT("%s: Cache miss with corrupted value %s with hash %s for %s from '%s'"),
-						*CacheStore.Domain, *ValueIdGetter(Value), *WriteToString<48>(Value.GetRawHash()),
-						*WriteToString<96>(Key), *Name);
-					SharedOnComplete.Get()({ Name, Key, ValueIndex, Operation->GetBytesReceived(), {}, EStatus::Error });
+						*Self->CacheStore.Domain, *WriteToString<32>(Id), *WriteToString<48>(RawHash),
+						*WriteToString<96>(Self->Key), *Self->Name);
+					SharedOnComplete.Get()({FValueWithId(Id, RawHash, RawSize), EStatus::Error});
 				}
 			}
 			else if (Operation->GetErrorCode() == EHttpErrorCode::Canceled)
 			{
-				SharedOnComplete.Get()({ Name, Key, ValueIndex, Operation->GetBytesReceived(), {}, EStatus::Canceled });
+				SharedOnComplete.Get()({FValueWithId(Id, RawHash, RawSize), EStatus::Canceled});
 			}
 			else
 			{
 				UE_LOG(LogDerivedDataCache, Verbose,
 					TEXT("%s: Cache miss with missing value %s with hash %s for %s from '%s'"),
-					*CacheStore.Domain, *ValueIdGetter(Value), *WriteToString<48>(Value.GetRawHash()), *WriteToString<96>(Key),
-					*Name);
-				SharedOnComplete.Get()({ Name, Key, ValueIndex, Operation->GetBytesReceived(), {}, EStatus::Error });
+					*Self->CacheStore.Domain, *WriteToString<32>(Id), *WriteToString<48>(RawHash),
+					*WriteToString<96>(Self->Key), *Self->Name);
+				SharedOnComplete.Get()({FValueWithId(Id, RawHash, RawSize), EStatus::Error});
 			}
 		});
 	}
 }
 
-FHttpCacheStore::FGetRecordOp::FGetRecordOp(
-	FHttpCacheStore& InCacheStore,
-	IRequestOwner& InOwner,
-	const FSharedString& InName,
-	const FCacheKey& InKey,
-	uint64 InUserData,
-	uint64 InBytesReceived,
-	TArray<FValueWithId>&& InRequiredGets,
-	TArray<FValueWithId>&& InRequiredHeads,
-	FCacheRecordBuilder&& InRecordBuilder,
-	TUniqueFunction<void(FCacheGetResponse&& Response, uint64 BytesReceived)>&& InOnComplete)
-	: CacheStore(InCacheStore)
-	, Owner(InOwner)
-	, Name(InName)
-	, Key(InKey)
-	, UserData(InUserData)
-	, BytesReceived(InBytesReceived)
-	, RequiredGets(MoveTemp(InRequiredGets))
-	, RequiredHeads(MoveTemp(InRequiredHeads))
-	, RecordBuilder(MoveTemp(InRecordBuilder))
-	, TotalOperations(RequiredGets.Num() + RequiredHeads.Num())
-	, SuccessfulOperations(0)
-	, PendingOperations(TotalOperations)
-	, OnComplete(MoveTemp(InOnComplete))
-{
-	FetchedBuffers.AddDefaulted(RequiredGets.Num());
-}
-
-void FHttpCacheStore::FGetRecordOp::OnOnlyRecordComplete(
-	FHttpCacheStore& CacheStore,
-	IRequestOwner& Owner,
-	const FCacheRecordPolicy& Policy,
-	TUniqueFunction<void(FCacheGetResponse&& Response, uint64 BytesReceived)>&& OnComplete,
-	FGetCacheRecordOnlyResponse&& Response)
+void FHttpCacheStore::FGetRecordOp::GetValuesExist(TConstArrayView<FValueWithId> Values, FOnValueComplete&& OnComplete)
 {
-	FCacheRecordBuilder RecordBuilder(Response.Key);
-	if (Response.Status != EStatus::Ok)
+	TArray<FValueWithId> QueryValues;
+	for (const FValueWithId& Value : Values)
 	{
-		return OnComplete({ Response.Name, RecordBuilder.Build(), Response.UserData, Response.Status }, Response.BytesReceived);
+		if (Value.HasData())
+		{
+			OnComplete({Value, EStatus::Ok});
+			continue;
+		}
+		QueryValues.Emplace(Value);
 	}
 
-	if (!EnumHasAnyFlags(Policy.GetRecordPolicy(), ECachePolicy::SkipMeta))
+	if (QueryValues.IsEmpty())
 	{
-		RecordBuilder.SetMeta(FCbObject(Response.Record.Get().GetMeta()));
+		return;
 	}
 
-	// TODO: There is not currently a batched GET endpoint for Jupiter.  Once there is, all payload data should be fetched in one call.
-	//		 In the meantime, we try to keep the code structured in a way that is friendly to future batching of GETs.
-
-	TArray<FValueWithId> RequiredGets;
-	TArray<FValueWithId> RequiredHeads;
+	TAnsiStringBuilder<256> Uri;
+	Uri << CacheStore.EffectiveDomain << ANSITEXTVIEW("/api/v1/compressed-blobs/") << CacheStore.Namespace << ANSITEXTVIEW("/exists?");
+	for (const FValueWithId& Value : QueryValues)
+	{
+		Uri << ANSITEXTVIEW("id=") << Value.GetRawHash() << '&';
+	}
+	Uri.RemoveSuffix(1);
 
-	for (FValueWithId Value : Response.Record.Get().GetValues())
+	FRequestBarrier Barrier(Owner);
+	TUniquePtr<FHttpOperation> Operation = CacheStore.WaitForHttpOperation(EOperationCategory::Get);
+	FHttpOperation& LocalOperation = *Operation;
+	LocalOperation.SetUri(Uri);
+	LocalOperation.SetMethod(EHttpMethod::Post);
+	LocalOperation.SetContentType(EHttpMediaType::FormUrlEncoded);
+	LocalOperation.AddAcceptType(EHttpMediaType::Json);
+	LocalOperation.SendAsync(Owner, [Self = TRefCountPtr(this), Operation = MoveTemp(Operation), Values = MoveTemp(QueryValues), OnComplete = MoveTemp(OnComplete)]() mutable
 	{
-		const ECachePolicy ValuePolicy = Policy.GetValuePolicy(Value.GetId());
-		if (IsValueDataReady(Value, ValuePolicy))
+		TRACE_CPUPROFILER_EVENT_SCOPE(HttpDDC_DataProbablyExistsBatch_OnHttpRequestComplete);
+
+		const TCHAR* DefaultMessage = TEXT("Cache exists miss for");
+		EStatus DefaultStatus = EStatus::Error;
+
+		if (Operation->GetErrorCode() == EHttpErrorCode::Canceled)
 		{
-			RecordBuilder.AddValue(MoveTemp(Value));
+			DefaultMessage = TEXT("Cache exists miss with canceled request for");
+			DefaultStatus = EStatus::Canceled;
 		}
-		else
+		else if (const int32 StatusCode = Operation->GetStatusCode(); StatusCode < 200 || StatusCode > 204)
 		{
-			if (EnumHasAnyFlags(ValuePolicy, ECachePolicy::SkipData))
-			{
-				RequiredHeads.Emplace(Value);
-			}
-			else
+			DefaultMessage = TEXT("Cache exists miss with failed response for");
+		}
+		else if (TSharedPtr<FJsonObject> ResponseObject = Operation->GetBodyAsJson(); !ResponseObject)
+		{
+			DefaultMessage = TEXT("Cache exists miss with invalid response for");
+		}
+		else if (TArray<FString> NeedsArrayStrings; ResponseObject->TryGetStringArrayField(TEXT("needs"), NeedsArrayStrings))
+		{
+			DefaultMessage = TEXT("Cache exists hit for");
+			DefaultStatus = EStatus::Ok;
+
+			for (const FString& NeedsString : NeedsArrayStrings)
 			{
-				RequiredGets.Emplace(Value);
+				const FIoHash NeedHash(NeedsString);
+				for (auto It = Values.CreateIterator(); It; ++It)
+				{
+					const FValueWithId& Value = *It;
+					if (Value.GetRawHash() == NeedHash)
+					{
+						UE_LOG(LogDerivedDataCache, Verbose,
+							TEXT("%s: Cache exists miss with missing value %s with hash %s for %s from '%s'"),
+							*Self->CacheStore.Domain, *WriteToString<32>(Value.GetId()),
+							*WriteToString<48>(Value.GetRawHash()), *WriteToString<96>(Self->Key), *Self->Name);
+						OnComplete({Value, EStatus::Error});
+						It.RemoveCurrentSwap();
+						break;
+					}
+				}
 			}
 		}
-	}
-
-	if (RequiredGets.IsEmpty() && RequiredHeads.IsEmpty())
-	{
-		return OnComplete({ Response.Name, RecordBuilder.Build(), Response.UserData, Response.Status }, Response.BytesReceived);
-	}
 
-	// Having this be a ref ensures we don't have the op reach 0 ref count in between the start of the exist batch operation and the get batch operation
-	TRefCountPtr<FGetRecordOp> GetRecordOp = new FGetRecordOp(
-		CacheStore,
-		Owner,
-		Response.Name,
-		Response.Key,
-		Response.UserData,
-		Response.BytesReceived,
-		MoveTemp(RequiredGets),
-		MoveTemp(RequiredHeads),
-		MoveTemp(RecordBuilder),
-		MoveTemp(OnComplete)
-	);
+		for (const FValueWithId& Value : Values)
+		{
+			UE_LOG(LogDerivedDataCache, Verbose,
+				TEXT("%s: %s value %s with hash %s for %s from '%s'"),
+				*Self->CacheStore.Domain, DefaultMessage, *WriteToString<32>(Value.GetId()),
+				*WriteToString<48>(Value.GetRawHash()), *WriteToString<96>(Self->Key), *Self->Name);
+			OnComplete({Value, DefaultStatus});
+		}
+	});
+}
 
-	auto IdGetter = [](const FValueWithId& Value)
+//----------------------------------------------------------------------------------------------------------
+// FHttpCacheStore::FGetValueOp
+//----------------------------------------------------------------------------------------------------------
+class FHttpCacheStore::FGetValueOp final : public FThreadSafeRefCountedObject
+{
+public:
+	struct FResponse
 	{
-		return FString(WriteToString<16>(Value.GetId()));
+		const FSharedString& Name;
+		const FCacheKey& Key;
+		FValue Value;
+		EStatus Status = EStatus::Error;
 	};
+	using FOnComplete = TUniqueFunction<void (FResponse&& Response)>;
 
+	static TRefCountPtr<FGetValueOp> New(FHttpCacheStore& CacheStore, IRequestOwner& Owner, const FSharedString& Name)
 	{
-		FRequestBarrier Barrier(Owner);
-		GetRecordOp->DataProbablyExistsBatch(GetRecordOp->RequiredHeads, [GetRecordOp](FCachedDataProbablyExistsBatchResponse&& Response)
-		{
-			GetRecordOp->FinishDataStep(Response.Status == EStatus::Ok, 0);
-		});
+		return new FGetValueOp(CacheStore, Owner, Name);
+	}
 
-		GetDataBatch<FValueWithId>(CacheStore, Owner, Response.Name, Response.Key, GetRecordOp->RequiredGets, IdGetter, [GetRecordOp](FGetCachedDataBatchResponse&& Response)
-		{
-			GetRecordOp->FetchedBuffers[Response.ValueIndex] = MoveTemp(Response.DataBuffer);
-			GetRecordOp->FinishDataStep(Response.Status == EStatus::Ok, Response.BytesReceived);
-		});
+	void Get(const FCacheKey& Key, ECachePolicy Policy, FOnComplete&& OnComplete);
+
+private:
+	FGetValueOp(FHttpCacheStore& CacheStore, IRequestOwner& Owner, const FSharedString& Name);
 
+	void EndGetRef(TUniquePtr<FHttpOperation> Operation);
+
+	FHttpCacheStore& CacheStore;
+	IRequestOwner& Owner;
+	FSharedString Name;
+	FCacheKey Key;
+	ECachePolicy Policy = ECachePolicy::None;
+	FOnComplete OnComplete;
+};
+
+FHttpCacheStore::FGetValueOp::FGetValueOp(FHttpCacheStore& InCacheStore, IRequestOwner& InOwner, const FSharedString& InName)
+	: CacheStore(InCacheStore)
+	, Owner(InOwner)
+	, Name(InName)
+{
+}
+
+void FHttpCacheStore::FGetValueOp::Get(const FCacheKey& InKey, ECachePolicy InPolicy, FOnComplete&& InOnComplete)
+{
+	Key = InKey;
+	Policy = InPolicy;
+	OnComplete = MoveTemp(InOnComplete);
+
+	const bool bSkipData = EnumHasAnyFlags(Policy, ECachePolicy::SkipData);
+
+	TAnsiStringBuilder<64> Bucket;
+	Algo::Transform(Key.Bucket.ToString(), AppendChars(Bucket), FCharAnsi::ToLower);
+
+	TUniquePtr<FHttpOperation> Operation = CacheStore.WaitForHttpOperation(EOperationCategory::Get);
+	FHttpOperation& LocalOperation = *Operation;
+	LocalOperation.SetUri(WriteToAnsiString<256>(CacheStore.EffectiveDomain, ANSITEXTVIEW("/api/v1/refs/"), CacheStore.Namespace, '/', Bucket, '/', Key.Hash));
+	LocalOperation.SetMethod(EHttpMethod::Get);
+	if (bSkipData)
+	{
+		LocalOperation.AddAcceptType(EHttpMediaType::CbObject);
+	}
+	else
+	{
+		LocalOperation.AddHeader(ANSITEXTVIEW("Accept"), ANSITEXTVIEW("application/x-jupiter-inline"));
 	}
+	LocalOperation.SetExpectedErrorCodes({404});
+	LocalOperation.SendAsync(Owner, [Self = TRefCountPtr(this), Operation = MoveTemp(Operation)]() mutable
+	{
+		Self->EndGetRef(MoveTemp(Operation));
+	});
 }
 
-void FHttpCacheStore::FGetRecordOp::DataProbablyExistsBatch(
-	TConstArrayView<FValueWithId> Values,
-	FOnCachedDataProbablyExistsBatchComplete&& InOnComplete)
+void FHttpCacheStore::FGetValueOp::EndGetRef(TUniquePtr<FHttpOperation> Operation)
 {
-	if (Values.IsEmpty())
+	TRACE_CPUPROFILER_EVENT_SCOPE(HttpDDC_GetValue_EndGetRef);
+
+	const bool bSkipData = EnumHasAnyFlags(Policy, ECachePolicy::SkipData);
+
+	const int32 StatusCode = Operation->GetStatusCode();
+	if (StatusCode < 200 || StatusCode > 204)
 	{
-		return;
+		UE_LOG(LogDerivedDataCache, Verbose, TEXT("%s: Cache miss with failed HTTP request for %s from '%s'"),
+			*CacheStore.Domain, *WriteToString<96>(Key), *Name);
+		return OnComplete({Name, Key, {}, EStatus::Error});
 	}
 
-	TAnsiStringBuilder<256> CompressedBlobsUri;
-	CompressedBlobsUri << CacheStore.EffectiveDomain << ANSITEXTVIEW("/api/v1/compressed-blobs/") << CacheStore.Namespace << ANSITEXTVIEW("/exists?");
-	bool bFirstItem = true;
-	for (const FValueWithId& Value : Values)
+	FSharedBuffer Body = Operation->GetBody();
+
+	if (bSkipData)
 	{
-		if (!bFirstItem)
+		if (ValidateCompactBinary(Body, ECbValidateMode::Default) != ECbValidateError::None)
 		{
-			CompressedBlobsUri << '&';
+			UE_LOG(LogDerivedDataCache, Display, TEXT("%s: Cache miss with invalid package for %s from '%s'"),
+				*CacheStore.Domain, *WriteToString<96>(Key), *Name);
+			return OnComplete({Name, Key, {}, EStatus::Error});
 		}
-		CompressedBlobsUri << ANSITEXTVIEW("id=") << Value.GetRawHash();
-		bFirstItem = false;
-	}
 
-	TUniquePtr<FHttpOperation> Operation = CacheStore.WaitForHttpOperation(EOperationCategory::Get);
-	FHttpOperation& LocalOperation = *Operation;
-	LocalOperation.SetUri(CompressedBlobsUri);
-	LocalOperation.SetMethod(EHttpMethod::Post);
-	LocalOperation.SetContentType(EHttpMediaType::FormUrlEncoded);
-	LocalOperation.AddAcceptType(EHttpMediaType::Json);
-	LocalOperation.SendAsync(Owner, [Operation = MoveTemp(Operation), this, Values = TArray<FValueWithId>(Values), InOnComplete = MoveTemp(InOnComplete)]
+		const FCbObjectView Object = FCbObject(Body);
+		const FIoHash RawHash = Object["RawHash"].AsHash();
+		const uint64 RawSize = Object["RawSize"].AsUInt64(MAX_uint64);
+		if (RawHash.IsZero() || RawSize == MAX_uint64)
+		{
+			UE_LOG(LogDerivedDataCache, Display, TEXT("%s: Cache miss with invalid value for %s from '%s'"),
+				*CacheStore.Domain, *WriteToString<96>(Key), *Name);
+			return OnComplete({Name, Key, {}, EStatus::Error});
+		}
+
+		OnComplete({Name, Key, FValue(RawHash, RawSize), EStatus::Ok});
+	}
+	else
 	{
-		TRACE_CPUPROFILER_EVENT_SCOPE(HttpDDC_DataProbablyExistsBatch_OnHttpRequestComplete);
+		FCompressedBuffer CompressedBuffer = FCompressedBuffer::FromCompressed(Body);
 
-		const int32 StatusCode = Operation->GetStatusCode();
-		if (StatusCode >= 200 && StatusCode <= 204)
+		if (!CompressedBuffer)
 		{
-			if (TSharedPtr<FJsonObject> ResponseObject = Operation->GetBodyAsJson())
+			if (FAnsiStringView ReceivedHashStr = Operation->GetHeader("X-Jupiter-InlinePayloadHash"); !ReceivedHashStr.IsEmpty())
 			{
-				TArray<FString> NeedsArrayStrings;
-				if (ResponseObject->TryGetStringArrayField(TEXT("needs"), NeedsArrayStrings))
+				FIoHash ReceivedHash(ReceivedHashStr);
+				FIoHash ComputedHash = FIoHash::HashBuffer(Body.GetView());
+				if (ReceivedHash == ComputedHash)
 				{
-					if (NeedsArrayStrings.IsEmpty())
-					{
-						for (int32 ValueIndex = 0; ValueIndex < Values.Num(); ++ValueIndex)
-						{
-							const FValueWithId& Value = Values[ValueIndex];
-							UE_LOG(LogDerivedDataCache, VeryVerbose,
-								TEXT("%s: Cache exists hit for value %s with hash %s for %s from '%s'"),
-								*CacheStore.Domain, *WriteToString<16>(Value.GetId()),
-								*WriteToString<48>(Value.GetRawHash()), *WriteToString<96>(Key), *Name);
-							InOnComplete({ Name, Key, ValueIndex, EStatus::Ok });
-						}
-						return;
-					}
+					CompressedBuffer = FCompressedBuffer::Compress(Body);
 				}
+			}
+		}
 
-				TBitArray<> ResultStatus(true, Values.Num());
-				for (const FString& NeedsString : NeedsArrayStrings)
-				{
-					const FIoHash NeedHash(NeedsString);
-					for (int32 ValueIndex = 0; ValueIndex < Values.Num(); ++ValueIndex)
-					{
-						const FValueWithId& Value = Values[ValueIndex];
-						if (ResultStatus[ValueIndex] && NeedHash == Value.GetRawHash())
-						{
-							ResultStatus[ValueIndex] = false;
-							break;
-						}
-					}
-				}
+		if (!CompressedBuffer)
+		{
+			UE_LOG(LogDerivedDataCache, Display, TEXT("%s: Cache miss with invalid package for %s from '%s'"),
+				*CacheStore.Domain, *WriteToString<96>(Key), *Name);
+			return OnComplete({Name, Key, {}, EStatus::Error});
+		}
 
-				for (int32 ValueIndex = 0; ValueIndex < Values.Num(); ++ValueIndex)
-				{
-					const FValueWithId& Value = Values[ValueIndex];
+		OnComplete({Name, Key, FValue(CompressedBuffer), EStatus::Ok});
+	}
+}
 
-					if (ResultStatus[ValueIndex])
-					{
-						UE_LOG(LogDerivedDataCache, VeryVerbose,
-							TEXT("%s: Cache exists hit for value %s with hash %s for %s from '%s'"),
-							*CacheStore.Domain, *WriteToString<32>(Value.GetId()),
-							*WriteToString<48>(Value.GetRawHash()), *WriteToString<96>(Key), *Name);
-						InOnComplete({ Name, Key, ValueIndex, EStatus::Ok });
-					}
-					else
-					{
-						UE_LOG(LogDerivedDataCache, Verbose,
-							TEXT("%s: Cache exists miss with missing value %s with hash %s for %s from '%s'"),
-							*CacheStore.Domain, *WriteToString<32>(Value.GetId()),
-							*WriteToString<48>(Value.GetRawHash()), *WriteToString<96>(Key), *Name);
-						InOnComplete({ Name, Key, ValueIndex, EStatus::Error });
-					}
-				}
-			}
-			else
-			{
-				for (int32 ValueIndex = 0; ValueIndex < Values.Num(); ++ValueIndex)
-				{
-					UE_LOG(LogDerivedDataCache, Log,
-						TEXT("%s: Cache exists miss with invalid response for value %s for %s from '%s'"),
-						*CacheStore.Domain, *WriteToString<32>(Values[ValueIndex].GetId()),
-						*WriteToString<96>(Key), *Name);
-					InOnComplete({ Name, Key, ValueIndex, EStatus::Error });
-				}
-			}
+//----------------------------------------------------------------------------------------------------------
+// FHttpCacheStore::FExistsBatchOp
+//----------------------------------------------------------------------------------------------------------
+class FHttpCacheStore::FExistsBatchOp final : public FThreadSafeRefCountedObject
+{
+public:
+	static TRefCountPtr<FExistsBatchOp> New(FHttpCacheStore& CacheStore, IRequestOwner& Owner)
+	{
+		return new FExistsBatchOp(CacheStore, Owner);
+	}
+
+	void Exists(TConstArrayView<FCacheGetValueRequest> Requests, FOnCacheGetValueComplete&& OnComplete);
+
+private:
+	FExistsBatchOp(FHttpCacheStore& CacheStore, IRequestOwner& Owner);
+
+	void EndExists(TUniquePtr<FHttpOperation> Operation);
+
+	FHttpCacheStore& CacheStore;
+	IRequestOwner& Owner;
+	TArray<FCacheGetValueRequest> Requests;
+	FOnCacheGetValueComplete OnComplete;
+};
+
+FHttpCacheStore::FExistsBatchOp::FExistsBatchOp(FHttpCacheStore& InCacheStore, IRequestOwner& InOwner)
+	: CacheStore(InCacheStore)
+	, Owner(InOwner)
+{
+}
+
+void FHttpCacheStore::FExistsBatchOp::Exists(TConstArrayView<FCacheGetValueRequest> InRequests, FOnCacheGetValueComplete&& InOnComplete)
+{
+	OnComplete = MoveTemp(InOnComplete);
+
+	Requests.Empty(InRequests.Num());
+	for (const FCacheGetValueRequest& Request : InRequests)
+	{
+		if (!CacheStore.IsUsable())
+		{
+			UE_LOG(LogDerivedDataCache, VeryVerbose,
+				TEXT("%s: Skipped exists check of %s from '%s' because this cache store is not available"),
+				*CacheStore.Domain, *WriteToString<96>(Request.Key), *Request.Name);
+			OnComplete(Request.MakeResponse(EStatus::Error));
+			continue;
+		}
+
+		if (!EnumHasAnyFlags(Request.Policy, ECachePolicy::QueryRemote))
+		{
+			UE_LOG(LogDerivedDataCache, VeryVerbose, TEXT("%s: Skipped exists check of %s from '%s' due to cache policy"),
+				*CacheStore.Domain, *WriteToString<96>(Request.Key), *Request.Name);
+			OnComplete(Request.MakeResponse(EStatus::Error));
+			continue;
 		}
-		else
+
+		if (CacheStore.DebugOptions.ShouldSimulateGetMiss(Request.Key))
 		{
-			for (int32 ValueIndex = 0; ValueIndex < Values.Num(); ++ValueIndex)
-			{
-				UE_LOG(LogDerivedDataCache, Verbose,
-					TEXT("%s: Cache exists miss with failed response for value %s for %s from '%s'"),
-					*CacheStore.Domain, *WriteToString<32>(Values[ValueIndex].GetId()),
-					*WriteToString<96>(Key), *Name);
-				InOnComplete({Name, Key, ValueIndex, EStatus::Error});
-			}
+			UE_LOG(LogDerivedDataCache, Verbose, TEXT("%s: Simulated miss for get of %s from '%s'"),
+				*CacheStore.Domain, *WriteToString<96>(Request.Key), *Request.Name);
+			OnComplete(Request.MakeResponse(EStatus::Error));
+			continue;
 		}
+
+		Requests.Emplace(Request);
+	}
+
+	if (Requests.IsEmpty())
+	{
+		return;
+	}
+
+	FCbWriter BodyWriter;
+	BodyWriter.BeginObject();
+	BodyWriter.BeginArray(ANSITEXTVIEW("ops"));
+	uint32 OpIndex = 0;
+	for (const FCacheGetValueRequest& Request : Requests)
+	{
+		BodyWriter.BeginObject();
+		BodyWriter.AddInteger(ANSITEXTVIEW("opId"), OpIndex);
+		BodyWriter.AddString(ANSITEXTVIEW("op"), ANSITEXTVIEW("GET"));
+		const FCacheKey& Key = Request.Key;
+		TAnsiStringBuilder<64> Bucket;
+		Algo::Transform(Key.Bucket.ToString(), AppendChars(Bucket), FCharAnsi::ToLower);
+		BodyWriter.AddString(ANSITEXTVIEW("bucket"), Bucket);
+		BodyWriter.AddString(ANSITEXTVIEW("key"), LexToString(Key.Hash));
+		BodyWriter.AddBool(ANSITEXTVIEW("resolveAttachments"), true);
+		BodyWriter.EndObject();
+		++OpIndex;
+	}
+	BodyWriter.EndArray();
+	BodyWriter.EndObject();
+	FCbFieldIterator Body = BodyWriter.Save();
+
+	TUniquePtr<FHttpOperation> Operation = CacheStore.WaitForHttpOperation(EOperationCategory::Get);
+	FHttpOperation& LocalOperation = *Operation;
+	LocalOperation.SetUri(WriteToAnsiString<256>(CacheStore.EffectiveDomain, ANSITEXTVIEW("/api/v1/refs/"), CacheStore.Namespace));
+	LocalOperation.SetMethod(EHttpMethod::Post);
+	LocalOperation.SetContentType(EHttpMediaType::CbObject);
+	LocalOperation.AddAcceptType(EHttpMediaType::CbObject);
+	LocalOperation.SetBody(FCompositeBuffer(Body.GetOuterBuffer()));
+	LocalOperation.SendAsync(Owner, [Self = TRefCountPtr(this), Operation = MoveTemp(Operation)]() mutable
+	{
+		Self->EndExists(MoveTemp(Operation));
 	});
 }
 
-void FHttpCacheStore::FGetRecordOp::FinishDataStep(bool bSuccess, uint64 InBytesReceived)
+void FHttpCacheStore::FExistsBatchOp::EndExists(TUniquePtr<FHttpOperation> Operation)
 {
-	BytesReceived.fetch_add(InBytesReceived, std::memory_order_relaxed);
-	if (bSuccess)
+	TRACE_CPUPROFILER_EVENT_SCOPE(HttpDDC_ExistsBatch_EndExists);
+
+	const int32 OverallStatusCode = Operation->GetStatusCode();
+	if (OverallStatusCode < 200 || OverallStatusCode > 204)
+	{
+		for (const FCacheGetValueRequest& Request : Requests)
+		{
+			UE_LOG(LogDerivedDataCache, Verbose, TEXT("%s: Cache miss with failed HTTP request for %s from '%s'"),
+				*CacheStore.Domain, *WriteToString<96>(Request.Key), *Request.Name);
+			OnComplete(Request.MakeResponse(EStatus::Error));
+		}
+		return;
+	}
+
+	FMemoryView ResponseView = Operation->GetBody();
+	if (ValidateCompactBinary(ResponseView, ECbValidateMode::Default) != ECbValidateError::None)
 	{
-		SuccessfulOperations.fetch_add(1, std::memory_order_relaxed);
+		for (const FCacheGetValueRequest& Request : Requests)
+		{
+			UE_LOG(LogDerivedDataCache, Log, TEXT("%s: Cache miss with corrupt response for %s from '%s'."),
+				*CacheStore.Domain, *WriteToString<96>(Request.Key), *Request.Name);
+			OnComplete(Request.MakeResponse(EStatus::Error));
+		}
+		return;
 	}
 
-	if (PendingOperations.fetch_sub(1, std::memory_order_acq_rel) == 1)
+	const FCbObjectView ResponseObject(ResponseView.GetData());
+	const FCbArrayView Results = ResponseObject[ANSITEXTVIEW("results")].AsArrayView();
+
+	if (Results.Num() != Requests.Num())
 	{
-		EStatus Status;
-		uint32 LocalSuccessfulOperations = SuccessfulOperations.load(std::memory_order_relaxed);
-		TOptional<FCacheRecord> Record;
-		if (LocalSuccessfulOperations == TotalOperations)
+		UE_LOG(LogDerivedDataCache, Log,
+			TEXT("%s: Cache exists returned unexpected quantity of results (expected %d, got %d)."),
+			*CacheStore.Domain, Requests.Num(), Results.Num());
+		for (const FCacheGetValueRequest& Request : Requests)
 		{
-			for (int32 Index = 0; Index < RequiredHeads.Num(); ++Index)
-			{
-				RecordBuilder.AddValue(RequiredHeads[Index].RemoveData());
-			}
+			UE_LOG(LogDerivedDataCache, Display, TEXT("%s: Cache miss with invalid response for %s from '%s'"),
+				*CacheStore.Domain, *WriteToString<96>(Request.Key), *Request.Name);
+			OnComplete(Request.MakeResponse(EStatus::Error));
+		}
+		return;
+	}
 
-			for (int32 Index = 0; Index < RequiredGets.Num(); ++Index)
-			{
-				RecordBuilder.AddValue(FValueWithId(RequiredGets[Index].GetId(), FetchedBuffers[Index]));
-			}
-			Record.Emplace(RecordBuilder.Build());
-			Status = EStatus::Ok;
+	for (FCbFieldView ResultField : Results)
+	{
+		const FCbObjectView ResultObject = ResultField.AsObjectView();
+		const uint32 OpId = ResultObject[ANSITEXTVIEW("opId")].AsUInt32();
+		const int32 StatusCode = ResultObject[ANSITEXTVIEW("statusCode")].AsInt32();
+		const FCbObjectView Value = ResultObject[ANSITEXTVIEW("response")].AsObjectView();
+
+		if (OpId >= (uint32)Requests.Num())
+		{
+			UE_LOG(LogDerivedDataCache, Display, TEXT("%s: Encountered invalid opId %d while querying %d values"),
+				*CacheStore.Domain, OpId, Requests.Num());
+			continue;
 		}
-		else
+
+		const FCacheGetValueRequest& Request = Requests[int32(OpId)];
+
+		if (StatusCode < 200 || StatusCode > 204)
 		{
-			Record.Emplace(FCacheRecordBuilder(Key).Build());
-			Status = EStatus::Error;
+			UE_LOG(LogDerivedDataCache, Verbose, TEXT("%s: Cache miss with unsuccessful response code %d for %s from '%s'"),
+				*CacheStore.Domain, StatusCode, *WriteToString<96>(Request.Key), *Request.Name);
+			OnComplete(Request.MakeResponse(EStatus::Error));
+			continue;
+		}
+
+		const FIoHash RawHash = Value[ANSITEXTVIEW("RawHash")].AsHash();
+		const uint64 RawSize = Value[ANSITEXTVIEW("RawSize")].AsUInt64(MAX_uint64);
+		if (RawHash.IsZero() || RawSize == MAX_uint64)
+		{
+			UE_LOG(LogDerivedDataCache, Display, TEXT("%s: Cache miss with invalid value for %s from '%s'"),
+				*CacheStore.Domain, *WriteToString<96>(Request.Key), *Request.Name);
+			OnComplete(Request.MakeResponse(EStatus::Error));
+			continue;
 		}
-		OnComplete({Name, MoveTemp(Record.GetValue()), UserData, Status}, BytesReceived.load(std::memory_order_relaxed));
+
+		OnComplete({Request.Name, Request.Key, FValue(RawHash, RawSize), Request.UserData, EStatus::Ok});
 	}
 }
 
-FHttpCacheStore::FHttpCacheStore(const FHttpCacheStoreParams& Params)
+///////////////////////////////////////////////////////////////////////////////////////////////////////////////////////
+
+FHttpCacheStore::FHttpCacheStore(const FHttpCacheStoreParams& Params, ICacheStoreOwner* Owner)
 	: Domain(Params.Host)
 	, Namespace(Params.Namespace)
 	, OAuthProvider(Params.OAuthProvider)
@@ -1561,6 +1789,7 @@ FHttpCacheStore::FHttpCacheStore(const FHttpCacheStoreParams& Params)
 	, OAuthProviderIdentifier(Params.OAuthProviderIdentifier)
 	, OAuthAccessToken(Params.OAuthAccessToken)
 	, HttpVersion(Params.HttpVersion)
+	, StoreOwner(Owner)
 	, bReadOnly(Params.bReadOnly)
 {
 	TRACE_CPUPROFILER_EVENT_SCOPE(HttpDDC_Construct);
@@ -1621,6 +1850,14 @@ FHttpCacheStore::FHttpCacheStore(const FHttpCacheStoreParams& Params)
 		NonBlockingPutRequestQueue = FHttpRequestQueue(*ConnectionPool, ClientParams);
 
 		bIsUsable = true;
+
+		if (StoreOwner)
+		{
+			const ECacheStoreFlags Flags = ECacheStoreFlags::Remote | ECacheStoreFlags::Query |
+				(Params.bReadOnly ? ECacheStoreFlags::None : ECacheStoreFlags::Store);
+			TStringBuilder<256> Path(InPlace, Domain, TEXTVIEW(" ("), Namespace, TEXTVIEW(")"));
+			StoreOwner->Add(this, Flags);
+		}
 	}
 
 	AnyInstance = this;
@@ -1919,86 +2156,6 @@ TUniquePtr<FHttpCacheStore::FHttpOperation> FHttpCacheStore::WaitForHttpOperatio
 	return MakeUnique<FHttpOperation>(MoveTemp(Request));
 }
 
-void FHttpCacheStore::GetCacheRecordOnlyAsync(
-	IRequestOwner& Owner,
-	const FSharedString& Name,
-	const FCacheKey& Key,
-	const FCacheRecordPolicy& Policy,
-	uint64 UserData,
-	FOnGetCacheRecordOnlyComplete&& OnComplete)
-{
-	auto MakeResponse = [Name = FSharedString(Name), Key, UserData](uint64 BytesReceived, EStatus Status)
-	{
-		return FGetCacheRecordOnlyResponse{ Name, Key, UserData, BytesReceived, {}, Status };
-	};
-
-	if (!IsUsable())
-	{
-		UE_LOG(LogDerivedDataCache, VeryVerbose,
-			TEXT("%s: Skipped get of %s from '%s' because this cache store is not available"),
-			*Domain, *WriteToString<96>(Key), *Name);
-		return OnComplete(MakeResponse(0, EStatus::Error));
-	}
-
-	// Skip the request if querying the cache is disabled.
-	if (!EnumHasAnyFlags(Policy.GetRecordPolicy(), ECachePolicy::QueryRemote))
-	{
-		UE_LOG(LogDerivedDataCache, VeryVerbose, TEXT("%s: Skipped get of %s from '%s' due to cache policy"),
-			*Domain, *WriteToString<96>(Key), *Name);
-		return OnComplete(MakeResponse(0, EStatus::Error));
-	}
-
-	if (DebugOptions.ShouldSimulateGetMiss(Key))
-	{
-		UE_LOG(LogDerivedDataCache, Verbose, TEXT("%s: Simulated miss for get of %s from '%s'"),
-			*Domain, *WriteToString<96>(Key), *Name);
-		return OnComplete(MakeResponse(0, EStatus::Error));
-	}
-
-	TAnsiStringBuilder<64> Bucket;
-	Algo::Transform(Key.Bucket.ToString(), AppendChars(Bucket), FCharAnsi::ToLower);
-
-	TUniquePtr<FHttpOperation> Operation = WaitForHttpOperation(EOperationCategory::Get);
-	FHttpOperation& LocalOperation = *Operation;
-	LocalOperation.SetUri(WriteToAnsiString<256>(EffectiveDomain, ANSITEXTVIEW("/api/v1/refs/"), Namespace, '/', Bucket, '/', Key.Hash));
-	LocalOperation.SetMethod(EHttpMethod::Get);
-	LocalOperation.AddAcceptType(EHttpMediaType::CbObject);
-	LocalOperation.SetExpectedErrorCodes({404});
-	LocalOperation.SendAsync(Owner, [Operation = MoveTemp(Operation), this, Name, Key, UserData, OnComplete = MoveTemp(OnComplete)]
-	{
-		TRACE_CPUPROFILER_EVENT_SCOPE(HttpDDC_GetCacheRecordOnlyAsync_OnHttpRequestComplete);
-
-		const int32 StatusCode = Operation->GetStatusCode();
-		if (StatusCode >= 200 && StatusCode <= 204)
-		{
-			FSharedBuffer Body = Operation->GetBody();
-
-			if (ValidateCompactBinary(Body, ECbValidateMode::Default) != ECbValidateError::None)
-			{
-				UE_LOG(LogDerivedDataCache, Log, TEXT("%s: Cache miss with invalid package for %s from '%s'"),
-					*Domain, *WriteToString<96>(Key), *Name);
-				OnComplete({ Name, Key, UserData, Operation->GetBytesReceived(), {}, EStatus::Error });
-			}
-			else if (FOptionalCacheRecord Record = FCacheRecord::Load(FCbPackage(FCbObject(Body))); Record.IsNull())
-			{
-				UE_LOG(LogDerivedDataCache, Log, TEXT("%s: Cache miss with record load failure for %s from '%s'"),
-					*Domain, *WriteToString<96>(Key), *Name);
-				OnComplete({ Name, Key, UserData, Operation->GetBytesReceived(), {}, EStatus::Error });
-			}
-			else
-			{
-				OnComplete({ Name, Key, UserData, Operation->GetBytesReceived(), MoveTemp(Record), EStatus::Ok });
-			}
-		}
-		else
-		{
-			UE_LOG(LogDerivedDataCache, Verbose, TEXT("%s: Cache miss with missing package for %s from '%s'"),
-				*Domain, *WriteToString<96>(Key), *Name);
-			OnComplete({ Name, Key, UserData, Operation->GetBytesReceived(), {}, EStatus::Error });
-		}
-	});
-}
-
 void FHttpCacheStore::PutCacheRecordAsync(
 	IRequestOwner& Owner,
 	const FSharedString& Name,
@@ -2042,7 +2199,7 @@ void FHttpCacheStore::PutCacheRecordAsync(
 
 	FCbPackage Package = Record.Save();
 
-	FPutPackageOp::PutPackage(*this, Owner, Name, Key, MoveTemp(Package), Policy, UserData, [MakeResponse = MoveTemp(MakeResponse), OnComplete = MoveTemp(OnComplete)](FPutPackageOp::FCachePutPackageResponse&& Response)
+	FPutPackageOp::PutPackage(*this, Owner, Name, Key, MoveTemp(Package), [MakeResponse = MoveTemp(MakeResponse), OnComplete = MoveTemp(OnComplete)](FPutPackageOp::FResponse&& Response)
 	{
 		OnComplete(MakeResponse(Response.Status), Response.BytesSent);
 	});
@@ -2097,7 +2254,7 @@ void FHttpCacheStore::PutCacheValueAsync(
 	FCbPackage Package(Writer.Save().AsObject());
 	Package.AddAttachment(FCbAttachment(Value.GetData()));
 
-	FPutPackageOp::PutPackage(*this, Owner, Name, Key, MoveTemp(Package), Policy, UserData, [MakeResponse = MoveTemp(MakeResponse), OnComplete = MoveTemp(OnComplete)](FPutPackageOp::FCachePutPackageResponse&& Response)
+	FPutPackageOp::PutPackage(*this, Owner, Name, Key, MoveTemp(Package), [MakeResponse = MoveTemp(MakeResponse), OnComplete = MoveTemp(OnComplete)](FPutPackageOp::FResponse&& Response)
 	{
 		OnComplete(MakeResponse(Response.Status), Response.BytesSent);
 	});
@@ -2137,88 +2294,10 @@ void FHttpCacheStore::GetCacheValueAsync(
 		return;
 	}
 
-	const bool bSkipData = EnumHasAnyFlags(Policy, ECachePolicy::SkipData);
-
-	TAnsiStringBuilder<64> Bucket;
-	Algo::Transform(Key.Bucket.ToString(), AppendChars(Bucket), FCharAnsi::ToLower);
-
-	TUniquePtr<FHttpOperation> Operation = WaitForHttpOperation(EOperationCategory::Get);
-	FHttpOperation& LocalOperation = *Operation;
-	LocalOperation.SetUri(WriteToAnsiString<256>(EffectiveDomain, ANSITEXTVIEW("/api/v1/refs/"), Namespace, '/', Bucket, '/', Key.Hash));
-	LocalOperation.SetMethod(EHttpMethod::Get);
-	if (bSkipData)
-	{
-		LocalOperation.AddAcceptType(EHttpMediaType::CbObject);
-	}
-	else
-	{
-		LocalOperation.AddHeader(ANSITEXTVIEW("Accept"), ANSITEXTVIEW("application/x-jupiter-inline"));
-	}
-	LocalOperation.SetExpectedErrorCodes({404});
-	LocalOperation.SendAsync(Owner, [Operation = MoveTemp(Operation), this, Name, Key, UserData, bSkipData, OnComplete = MoveTemp(OnComplete)]
+	TRefCountPtr<FGetValueOp> Op = FGetValueOp::New(*this, Owner, Name);
+	Op->Get(Key, Policy, [UserData, OnComplete = MoveTemp(OnComplete)](FGetValueOp::FResponse&& Response)
 	{
-		TRACE_CPUPROFILER_EVENT_SCOPE(HttpDDC_GetCacheValueAsync_OnHttpRequestComplete);
-
-		const int32 StatusCode = Operation->GetStatusCode();
-		if (StatusCode >= 200 && StatusCode <= 204)
-		{
-			FValue ResultValue;
-			FSharedBuffer ResponseBuffer = Operation->GetBody();
-
-			if (bSkipData)
-			{
-				if (ValidateCompactBinary(ResponseBuffer, ECbValidateMode::Default) != ECbValidateError::None)
-				{
-					UE_LOG(LogDerivedDataCache, Display, TEXT("%s: Cache miss with invalid package for %s from '%s'"),
-						*Domain, *WriteToString<96>(Key), *Name);
-					OnComplete({Name, Key, {}, UserData, EStatus::Error});
-					return;
-				}
-
-				const FCbObjectView Object = FCbObject(ResponseBuffer);
-				const FIoHash RawHash = Object["RawHash"].AsHash();
-				const uint64 RawSize = Object["RawSize"].AsUInt64(MAX_uint64);
-				if (RawHash.IsZero() || RawSize == MAX_uint64)
-				{
-					UE_LOG(LogDerivedDataCache, Display, TEXT("%s: Cache miss with invalid value for %s from '%s'"),
-						*Domain, *WriteToString<96>(Key), *Name);
-					OnComplete({Name, Key, {}, UserData, EStatus::Error});
-					return;
-				}
-				ResultValue = FValue(RawHash, RawSize);
-			}
-			else
-			{
-				FCompressedBuffer CompressedBuffer = FCompressedBuffer::FromCompressed(ResponseBuffer);
-				if (!CompressedBuffer)
-				{
-					if (FAnsiStringView ReceivedHashStr = Operation->GetHeader("X-Jupiter-InlinePayloadHash"); !ReceivedHashStr.IsEmpty())
-					{
-						FIoHash ReceivedHash(ReceivedHashStr);
-						FIoHash ComputedHash = FIoHash::HashBuffer(ResponseBuffer.GetView());
-						if (ReceivedHash == ComputedHash)
-						{
-							CompressedBuffer = FCompressedBuffer::Compress(ResponseBuffer);
-						}
-					}
-				}
-
-				if (!CompressedBuffer)
-				{
-					UE_LOG(LogDerivedDataCache, Display, TEXT("%s: Cache miss with invalid package for %s from '%s'"),
-						*Domain, *WriteToString<96>(Key), *Name);
-					OnComplete({Name, Key, {}, UserData, EStatus::Error});
-					return;
-				}
-				ResultValue = FValue(CompressedBuffer);
-			}
-			OnComplete({Name, Key, ResultValue, UserData, EStatus::Ok});
-			return;
-		}
-
-		UE_LOG(LogDerivedDataCache, Verbose, TEXT("%s: Cache miss with failed HTTP request for %s from '%s'"),
-			*Domain, *WriteToString<96>(Key), *Name);
-		OnComplete({Name, Key, {}, UserData, EStatus::Error});
+		OnComplete({Response.Name, Response.Key, MoveTemp(Response.Value), UserData, Response.Status});
 	});
 }
 
@@ -2230,147 +2309,10 @@ void FHttpCacheStore::GetCacheRecordAsync(
 	uint64 UserData,
 	TUniqueFunction<void(FCacheGetResponse&& Response, uint64 BytesReceived)>&& OnComplete)
 {
-	FGetRecordOp::GetRecord(*this, Owner, Name, Key, Policy, UserData, MoveTemp(OnComplete));
-}
-
-void FHttpCacheStore::RefCachedDataProbablyExistsBatchAsync(
-	IRequestOwner& Owner,
-	TConstArrayView<FCacheGetValueRequest> ValueRefs,
-	FOnCacheGetValueComplete&& OnComplete)
-{
-	if (ValueRefs.IsEmpty())
-	{
-		return;
-	}
-
-	if (!IsUsable())
-	{
-		for (const FCacheGetValueRequest& ValueRef : ValueRefs)
-		{
-			UE_LOG(LogDerivedDataCache, VeryVerbose,
-				TEXT("%s: Skipped exists check of %s from '%s' because this cache store is not available"),
-				*Domain, *WriteToString<96>(ValueRef.Key), *ValueRef.Name);
-			OnComplete(ValueRef.MakeResponse(EStatus::Error));
-		}
-		return;
-	}
-
-	FCbWriter RequestWriter;
-	RequestWriter.BeginObject();
-	RequestWriter.BeginArray(ANSITEXTVIEW("ops"));
-	uint32 OpIndex = 0;
-	for (const FCacheGetValueRequest& ValueRef : ValueRefs)
-	{
-		RequestWriter.BeginObject();
-		RequestWriter.AddInteger(ANSITEXTVIEW("opId"), OpIndex);
-		RequestWriter.AddString(ANSITEXTVIEW("op"), ANSITEXTVIEW("GET"));
-		const FCacheKey& Key = ValueRef.Key;
-		TAnsiStringBuilder<64> Bucket;
-		Algo::Transform(Key.Bucket.ToString(), AppendChars(Bucket), FCharAnsi::ToLower);
-		RequestWriter.AddString(ANSITEXTVIEW("bucket"), Bucket);
-		RequestWriter.AddString(ANSITEXTVIEW("key"), LexToString(Key.Hash));
-		RequestWriter.AddBool(ANSITEXTVIEW("resolveAttachments"), true);
-		RequestWriter.EndObject();
-		++OpIndex;
-	}
-	RequestWriter.EndArray();
-	RequestWriter.EndObject();
-	FCbFieldIterator RequestFields = RequestWriter.Save();
-
-	TUniquePtr<FHttpOperation> Operation = WaitForHttpOperation(EOperationCategory::Get);
-	FHttpOperation& LocalOperation = *Operation;
-	LocalOperation.SetUri(WriteToAnsiString<256>(EffectiveDomain, ANSITEXTVIEW("/api/v1/refs/"), Namespace));
-	LocalOperation.SetMethod(EHttpMethod::Post);
-	LocalOperation.SetContentType(EHttpMediaType::CbObject);
-	LocalOperation.AddAcceptType(EHttpMediaType::CbObject);
-	LocalOperation.SetBody(FCompositeBuffer(RequestFields.GetOuterBuffer()));
-	LocalOperation.SendAsync(Owner, [Operation = MoveTemp(Operation), this, ValueRefs = TArray<FCacheGetValueRequest>(ValueRefs), OnComplete = MoveTemp(OnComplete)]
+	TRefCountPtr<FGetRecordOp> Op = FGetRecordOp::New(*this, Owner, Name);
+	Op->GetRecord(Key, Policy, [Name, UserData, OnComplete = MoveTemp(OnComplete)](FGetRecordOp::FRecordResponse&& Response)
 	{
-		TRACE_CPUPROFILER_EVENT_SCOPE(HttpDDC_RefCachedDataProbablyExistsBatchAsync_OnHttpRequestComplete);
-
-		const int32 OverallStatusCode = Operation->GetStatusCode();
-		if (OverallStatusCode >= 200 && OverallStatusCode <= 204)
-		{
-			FMemoryView ResponseView = Operation->GetBody();
-			if (ValidateCompactBinary(ResponseView, ECbValidateMode::Default) != ECbValidateError::None)
-			{
-				for (const FCacheGetValueRequest& ValueRef : ValueRefs)
-				{
-					UE_LOG(LogDerivedDataCache, Log, TEXT("%s: Cache exists returned invalid results."), *Domain);
-					OnComplete(ValueRef.MakeResponse(EStatus::Error));
-				}
-				return;
-			}
-
-			const FCbObjectView ResponseObject(ResponseView.GetData());
-
-			FCbArrayView ResultsArrayView = ResponseObject[ANSITEXTVIEW("results")].AsArrayView();
-
-			if (ResultsArrayView.Num() != ValueRefs.Num())
-			{
-				for (const FCacheGetValueRequest& ValueRef : ValueRefs)
-				{
-					UE_LOG(LogDerivedDataCache, Log,
-						TEXT("%s: Cache exists returned unexpected quantity of results (expected %d, got %d)."),
-						*Domain, ValueRefs.Num(), ResultsArrayView.Num());
-						OnComplete(ValueRef.MakeResponse(EStatus::Error));
-				}
-				return;
-			}
-
-			for (FCbFieldView ResultFieldView : ResultsArrayView)
-			{
-				FCbObjectView ResultObjectView = ResultFieldView.AsObjectView();
-				uint32 OpId = ResultObjectView[ANSITEXTVIEW("opId")].AsUInt32();
-				FCbObjectView ResponseObjectView = ResultObjectView[ANSITEXTVIEW("response")].AsObjectView();
-				int32 StatusCode = ResultObjectView[ANSITEXTVIEW("statusCode")].AsInt32();
-
-				if (OpId >= (uint32)ValueRefs.Num())
-				{
-					UE_LOG(LogDerivedDataCache, Display, TEXT("%s: Encountered invalid opId %d while querying %d values"),
-						*Domain, OpId, ValueRefs.Num());
-					continue;
-				}
-
-				const FCacheGetValueRequest& ValueRef = ValueRefs[OpId];
-
-				if (StatusCode < 200 || StatusCode > 204)
-				{
-					UE_LOG(LogDerivedDataCache, Verbose, TEXT("%s: Cache miss with unsuccessful response code %d for %s from '%s'"),
-						*Domain, StatusCode, *WriteToString<96>(ValueRef.Key), *ValueRef.Name);
-					OnComplete(ValueRef.MakeResponse(EStatus::Error));
-					continue;
-				}
-
-				if (!EnumHasAnyFlags(ValueRef.Policy, ECachePolicy::QueryRemote))
-				{
-					UE_LOG(LogDerivedDataCache, VeryVerbose, TEXT("%s: Skipped exists check of %s from '%s' due to cache policy"),
-						*Domain, *WriteToString<96>(ValueRef.Key), *ValueRef.Name);
-					OnComplete(ValueRef.MakeResponse(EStatus::Error));
-					continue;
-				}
-
-				const FIoHash RawHash = ResponseObjectView[ANSITEXTVIEW("RawHash")].AsHash();
-				const uint64 RawSize = ResponseObjectView[ANSITEXTVIEW("RawSize")].AsUInt64(MAX_uint64);
-				if (RawHash.IsZero() || RawSize == MAX_uint64)
-				{
-					UE_LOG(LogDerivedDataCache, Display, TEXT("%s: Cache miss with invalid value for %s from '%s'"),
-						*Domain, *WriteToString<96>(ValueRef.Key), *ValueRef.Name);
-					OnComplete(ValueRef.MakeResponse(EStatus::Error));
-					continue;
-				}
-
-				OnComplete({ValueRef.Name, ValueRef.Key, FValue(RawHash, RawSize), ValueRef.UserData, EStatus::Ok});
-			}
-			return;
-		}
-
-		for (const FCacheGetValueRequest& ValueRef : ValueRefs)
-		{
-			UE_LOG(LogDerivedDataCache, Verbose, TEXT("%s: Cache miss with failed HTTP request for %s from '%s'"),
-				*Domain, *WriteToString<96>(ValueRef.Key), *ValueRef.Name);
-			OnComplete(ValueRef.MakeResponse(EStatus::Error));
-		}
+		OnComplete({Name, MoveTemp(Response.Record), UserData, Response.Status}, 0);
 	});
 }
 
@@ -2493,8 +2435,8 @@ void FHttpCacheStore::GetValue(
 	}
 	if (bBatchExistsCandidate)
 	{
-		RefCachedDataProbablyExistsBatchAsync(Owner, Requests,
-		[this, COOK_STAT(StartTime, bIsInGameThread, ) OnComplete = MoveTemp(OnComplete)](FCacheGetValueResponse&& Response)
+		TRefCountPtr<FExistsBatchOp> Op = FExistsBatchOp::New(*this, Owner);
+		Op->Exists(Requests, [this, COOK_STAT(StartTime, bIsInGameThread, ) OnComplete = MoveTemp(OnComplete)](FCacheGetValueResponse&& Response)
 		{
 			if (Response.Status != EStatus::Ok)
 			{
@@ -2517,7 +2459,6 @@ void FHttpCacheStore::GetValue(
 	{
 		FRequestBarrier Barrier(Owner);
 		TSharedRef<FOnCacheGetValueComplete> SharedOnComplete = MakeShared<FOnCacheGetValueComplete>(MoveTemp(OnComplete));
-		int64 HitBytes = 0;
 		for (const FCacheGetValueRequest& Request : Requests)
 		{
 			GetCacheValueAsync(Owner, Request.Name, Request.Key, Request.Policy, Request.UserData,
@@ -2557,7 +2498,6 @@ void FHttpCacheStore::GetValue(
 			});
 		}
 	}
-
 }
 
 void FHttpCacheStore::GetChunks(
@@ -2592,20 +2532,18 @@ void FHttpCacheStore::GetChunks(
 		{
 			ValueStatus = EStatus::Error;
 			ValueReader.ResetSource();
+			ValueBuffer.Reset();
 			ValueKey = {};
 			ValueId.Reset();
 			Value.Reset();
 			bHasValue = false;
 			if (Request.Id.IsValid())
 			{
+				FRequestOwner BlockingOwner(EPriority::Blocking);
+				TRefCountPtr<FGetRecordOp> Op = FGetRecordOp::New(*this, BlockingOwner, Request.Name);
 				if (!(Record && Record.Get().GetKey() == Request.Key))
 				{
-					FCacheRecordPolicyBuilder PolicyBuilder(ECachePolicy::None);
-					PolicyBuilder.AddValuePolicy(Request.Id, Request.Policy);
-					Record.Reset();
-
-					FRequestOwner BlockingOwner(EPriority::Blocking);
-					GetCacheRecordOnlyAsync(BlockingOwner, Request.Name, Request.Key, PolicyBuilder.Build(), 0, [&Record](FGetCacheRecordOnlyResponse&& Response)
+					Op->GetRecordOnly(Request.Key, Request.Policy, [&Record](FGetRecordOp::FRecordResponse&& Response)
 					{
 						Record = MoveTemp(Response.Record);
 					});
@@ -2621,38 +2559,22 @@ void FHttpCacheStore::GetChunks(
 
 					if (IsValueDataReady(Value, Request.Policy))
 					{
-						ValueReader.SetSource(Value.GetData());
+						ValueBuffer = Value.GetData();
+						ValueReader.SetSource(ValueBuffer);
 					}
-					else
+					else if (bHasValue)
 					{
-						auto IdGetter = [](const FValueWithId& Value)
-						{
-							return FString(WriteToString<16>(Value.GetId()));
-						};
-
-						FRequestOwner BlockingOwner(EPriority::Blocking);
-						bool bSucceeded = false;
-						FCompressedBuffer NewBuffer;
-						FGetRecordOp::GetDataBatch(*this, BlockingOwner, Request.Name, Request.Key, ::MakeArrayView({ ValueWithId }), IdGetter, [&bSucceeded, &NewBuffer](FGetRecordOp::FGetCachedDataBatchResponse&& Response)
+						Op->GetValues({ValueWithId}, [&Value](FGetRecordOp::FValueResponse&& Response)
 						{
-							if (Response.Status == EStatus::Ok)
-							{
-								bSucceeded = true;
-								NewBuffer = MoveTemp(Response.DataBuffer);
-							}
+							Value = MoveTemp(Response.Value);
 						});
 						BlockingOwner.Wait();
 
-						if (bSucceeded)
+						if (Value.HasData())
 						{
-							ValueBuffer = MoveTemp(NewBuffer);
+							ValueBuffer = Value.GetData();
 							ValueReader.SetSource(ValueBuffer);
 						}
-						else
-						{
-							ValueBuffer.Reset();
-							ValueReader.ResetSource();
-						}
 					}
 				}
 			}
@@ -2672,48 +2594,10 @@ void FHttpCacheStore::GetChunks(
 					bHasValue = bSucceeded;
 				}
 
-				if (bHasValue)
-				{
-					if (IsValueDataReady(Value, Request.Policy))
-					{
-						ValueReader.SetSource(Value.GetData());
-					}
-					else
-					{
-						auto IdGetter = [](const FValue& Value)
-						{
-							return FString(TEXT("Default"));
-						};
-
-						FRequestOwner BlockingOwner(EPriority::Blocking);
-						bool bSucceeded = false;
-						FCompressedBuffer NewBuffer;
-						FGetRecordOp::GetDataBatch(*this, BlockingOwner, Request.Name, Request.Key, ::MakeArrayView({ Value }), IdGetter, [&bSucceeded, &NewBuffer](FGetRecordOp::FGetCachedDataBatchResponse&& Response)
-						{
-							if (Response.Status == EStatus::Ok)
-							{
-								bSucceeded = true;
-								NewBuffer = MoveTemp(Response.DataBuffer);
-							}
-						});
-						BlockingOwner.Wait();
-
-						if (bSucceeded)
-						{
-							ValueBuffer = MoveTemp(NewBuffer);
-							ValueReader.SetSource(ValueBuffer);
-						}
-						else
-						{
-							ValueBuffer.Reset();
-							ValueReader.ResetSource();
-						}
-					}
-				}
-				else
+				if (bHasValue && IsValueDataReady(Value, Request.Policy))
 				{
-					ValueBuffer.Reset();
-					ValueReader.ResetSource();
+					ValueBuffer = Value.GetData();
+					ValueReader.SetSource(ValueBuffer);
 				}
 			}
 		}
@@ -2741,6 +2625,8 @@ void FHttpCacheStore::GetChunks(
 
 void FHttpCacheStoreParams::Parse(const TCHAR* NodeName, const TCHAR* Config)
 {
+	Name = NodeName;
+
 	FString ServerId;
 	if (FParse::Value(Config, TEXT("ServerID="), ServerId))
 	{
@@ -2888,7 +2774,7 @@ void FHttpCacheStoreParams::Parse(const TCHAR* NodeName, const TCHAR* Config)
 namespace UE::DerivedData
 {
 
-TTuple<ILegacyCacheStore*, ECacheStoreFlags> CreateHttpCacheStore(const TCHAR* NodeName, const TCHAR* Config)
+ILegacyCacheStore* CreateHttpCacheStore(const TCHAR* NodeName, const TCHAR* Config, ICacheStoreOwner* Owner)
 {
 #if !WITH_HTTP_DDC_BACKEND
 	UE_LOG(LogDerivedDataCache, Warning, TEXT("%s: HTTP cache is not yet supported in the current build configuration."), NodeName);
@@ -2965,16 +2851,15 @@ TTuple<ILegacyCacheStore*, ECacheStoreFlags> CreateHttpCacheStore(const TCHAR* N
 
 	if (bValidParams)
 	{
-		if (TUniquePtr<FHttpCacheStore> Store = MakeUnique<FHttpCacheStore>(Params); Store->IsUsable())
+		if (TUniquePtr<FHttpCacheStore> Store = MakeUnique<FHttpCacheStore>(Params, Owner); Store->IsUsable())
 		{
-			const ECacheStoreFlags StoreFlag = (Params.bReadOnly ? ECacheStoreFlags::None : ECacheStoreFlags::Store);
-			return MakeTuple(Store.Release(), ECacheStoreFlags::Remote | ECacheStoreFlags::Query | StoreFlag);
+			return Store.Release();
 		}
 		UE_LOG(LogDerivedDataCache, Warning, TEXT("%s: Failed to contact the service (%s), will not use it."), NodeName, *Params.Host);
 	}
 #endif
 
-	return MakeTuple(nullptr, ECacheStoreFlags::None);
+	return nullptr;
 }
 
 ILegacyCacheStore* GetAnyHttpCacheStore(
