DDC: Added ICacheStoreStats to S3

#rnx

diff --git a/Engine/Source/Developer/DerivedDataCache/Private/DerivedDataBackends.cpp b/Engine/Source/Developer/DerivedDataCache/Private/DerivedDataBackends.cpp
--- a/Engine/Source/Developer/DerivedDataCache/Private/DerivedDataBackends.cpp
+++ b/Engine/Source/Developer/DerivedDataCache/Private/DerivedDataBackends.cpp
@@ -66,11 +66,11 @@ ILegacyCacheStore* CreateCacheStoreHierarchy(ICacheStoreOwner*& OutOwner, TFunct
 ILegacyCacheStore* CreateCacheStoreThrottle(ILegacyCacheStore* InnerCache, uint32 LatencyMS, uint32 MaxBytesPerSecond);
 ILegacyCacheStore* CreateCacheStoreVerify(ILegacyCacheStore* InnerCache, bool bPutOnError);
 ILegacyCacheStore* CreateFileSystemCacheStore(const TCHAR* Name, const TCHAR* Config, ICacheStoreOwner& Owner, FString& OutPath);
-ILegacyCacheStore* CreateHttpCacheStore(const TCHAR* NodeName, const TCHAR* Config, ICacheStoreOwner* Owner);
+ILegacyCacheStore* CreateHttpCacheStore(const TCHAR* Name, const TCHAR* Config, ICacheStoreOwner* Owner);
 void CreateMemoryCacheStore(IMemoryCacheStore*& OutCache, const TCHAR* Name, const TCHAR* Config, ICacheStoreOwner* Owner);
 IPakFileCacheStore* CreatePakFileCacheStore(const TCHAR* Name, const TCHAR* Filename, bool bWriting, bool bCompressed, ICacheStoreOwner* Owner);
-ILegacyCacheStore* CreateS3CacheStore(const TCHAR* RootManifestPath, const TCHAR* BaseUrl, const TCHAR* Region, const TCHAR* CanaryObjectKey, const TCHAR* CachePath);
-ILegacyCacheStore* CreateZenCacheStore(const TCHAR* NodeName, const TCHAR* Config, ICacheStoreOwner* Owner);
+ILegacyCacheStore* CreateS3CacheStore(const TCHAR* Name, const TCHAR* Config, ICacheStoreOwner& Owner);
+ILegacyCacheStore* CreateZenCacheStore(const TCHAR* Name, const TCHAR* Config, ICacheStoreOwner* Owner);
 ILegacyCacheStore* TryCreateCacheStoreReplay(ILegacyCacheStore* InnerCache);
 
 /**
@@ -309,7 +309,10 @@ public:
 				}
 				else if (NodeType == TEXT("S3"))
 				{
-					return ParseS3Cache(*NodeName, *Entry);
+					if (CreateS3CacheStore(*NodeName, *Entry, *this))
+					{
+						return true;
+					}
 				}
 				else if (NodeType == TEXT("Cloud") || NodeType == TEXT("Http"))
 				{
@@ -544,79 +547,6 @@ public:
 		return false;
 	}
 
-	/**
-	 * Creates an S3 data cache interface.
-	 */
-	bool ParseS3Cache(const TCHAR* NodeName, const TCHAR* Entry)
-	{
-		FString ManifestPath;
-		if (!FParse::Value(Entry, TEXT("Manifest="), ManifestPath))
-		{
-			UE_LOG(LogDerivedDataCache, Error, TEXT("Node %s does not specify 'Manifest'."), NodeName);
-			return false;
-		}
-
-		FString BaseUrl;
-		if (!FParse::Value(Entry, TEXT("BaseUrl="), BaseUrl))
-		{
-			UE_LOG(LogDerivedDataCache, Error, TEXT("Node %s does not specify 'BaseUrl'."), NodeName);
-			return false;
-		}
-
-		FString CanaryObjectKey;
-		FParse::Value(Entry, TEXT("Canary="), CanaryObjectKey);
-
-		FString Region;
-		if (!FParse::Value(Entry, TEXT("Region="), Region))
-		{
-			UE_LOG(LogDerivedDataCache, Error, TEXT("Node %s does not specify 'Region'."), NodeName);
-			return false;
-		}
-
-		// Check the EnvPathOverride environment variable to allow persistent overriding of data cache path, eg for offsite workers.
-		FString EnvPathOverride;
-		FString CachePath = FPaths::ProjectSavedDir() / TEXT("S3DDC");
-		if (FParse::Value(Entry, TEXT("EnvPathOverride="), EnvPathOverride))
-		{
-			FString FilesystemCachePathEnv = FPlatformMisc::GetEnvironmentVariable(*EnvPathOverride);
-			if (FilesystemCachePathEnv.Len() > 0)
-			{
-				if (FilesystemCachePathEnv == TEXT("None"))
-				{
-					UE_LOG(LogDerivedDataCache, Log, TEXT("Node %s disabled due to %s=None"), NodeName, *EnvPathOverride);
-					return false;
-				}
-				else
-				{
-					CachePath = FilesystemCachePathEnv;
-					UE_LOG(LogDerivedDataCache, Log, TEXT("Found environment variable %s=%s"), *EnvPathOverride, *CachePath);
-				}
-			}
-
-			if (!EnvPathOverride.IsEmpty())
-			{
-				FString DDCPath;
-				if (FPlatformMisc::GetStoredValue(TEXT("Epic Games"), TEXT("GlobalDataCachePath"), *EnvPathOverride, DDCPath))
-				{
-					if ( DDCPath.Len() > 0 )
-					{
-						CachePath = DDCPath;
-						UE_LOG( LogDerivedDataCache, Log, TEXT("Found registry key GlobalDataCachePath %s=%s"), *EnvPathOverride, *CachePath );
-					}
-				}
-			}
-		}
-
-		if (ILegacyCacheStore* Backend = CreateS3CacheStore(*ManifestPath, *BaseUrl, *Region, *CanaryObjectKey, *CachePath))
-		{
-			Add(Backend, ECacheStoreFlags::Local | ECacheStoreFlags::Query | ECacheStoreFlags::StopStore);
-			return true;
-		}
-
-		UE_LOG(LogDerivedDataCache, Log, TEXT("S3 backend is not supported on the current platform."));
-		return false;
-	}
-
 	/**
 	 * Creates Boot data cache interface from ini settings.
 	 *
diff --git a/Engine/Source/Developer/DerivedDataCache/Private/S3CacheStore.cpp b/Engine/Source/Developer/DerivedDataCache/Private/S3CacheStore.cpp
--- a/Engine/Source/Developer/DerivedDataCache/Private/S3CacheStore.cpp
+++ b/Engine/Source/Developer/DerivedDataCache/Private/S3CacheStore.cpp
@@ -1,5 +1,6 @@
 // Copyright Epic Games, Inc. All Rights Reserved.
 
+#include "DerivedDataCacheStore.h"
 #include "DerivedDataLegacyCacheStore.h"
 
 #if WITH_S3_DDC_BACKEND
@@ -16,7 +17,6 @@
 #include "DerivedDataBackendInterface.h"
 #include "DerivedDataCacheRecord.h"
 #include "DerivedDataCachePrivate.h"
-#include "DerivedDataCacheUsageStats.h"
 #include "DerivedDataChunk.h"
 #include "DesktopPlatformModule.h"
 #include "Dom/JsonObject.h"
@@ -196,7 +196,9 @@ public:
 	 * @param  InCanaryObjectKey    Key for a canary object used to test whether this backend is usable
 	 * @param  InCachePath          Path to cache the DDC files
 	 */
-	FS3CacheStore(const TCHAR* InRootManifestPath, const TCHAR* InBaseUrl, const TCHAR* InRegion, const TCHAR* InCanaryObjectKey, const TCHAR* InCachePath);
+	FS3CacheStore(const TCHAR* InName, const TCHAR* InRootManifestPath, const TCHAR* InBaseUrl, const TCHAR* InRegion, const TCHAR* InCanaryObjectKey, const TCHAR* InCachePath, ICacheStoreOwner& Owner);
+
+	~FS3CacheStore() final;
 
 	inline const FString& GetName() const { return BaseUrl; }
 
@@ -249,32 +251,38 @@ private:
 	FString Region;
 	FString CanaryObjectKey;
 	FString CacheDir;
+
+	ICacheStoreOwner& StoreOwner;
+	ICacheStoreStats* StoreStats = nullptr;
+
 	TArray<FBundle> Bundles;
 	TUniquePtr<FRequestPool> RequestPool;
-	FDerivedDataCacheUsageStats UsageStats;
 	bool bEnabled;
 
 	[[nodiscard]] FOptionalCacheRecord GetCacheRecordOnly(
 		FStringView Name,
 		const FCacheKey& Key,
-		const FCacheRecordPolicy& Policy);
+		const FCacheRecordPolicy& Policy,
+		FRequestStats& Stats);
 	[[nodiscard]] FOptionalCacheRecord GetCacheRecord(
 		FStringView Name,
 		const FCacheKey& Key,
 		const FCacheRecordPolicy& Policy,
-		EStatus& OutStatus);
+		EStatus& OutStatus,
+		FRequestStats& Stats);
 
-	[[nodiscard]] bool GetCacheValueOnly(FStringView Name, const FCacheKey& Key, ECachePolicy Policy, FValue& OutValue);
-	[[nodiscard]] bool GetCacheValue(FStringView Name, const FCacheKey& Key, ECachePolicy Policy, FValue& OutValue);
+	[[nodiscard]] bool GetCacheValueOnly(FStringView Name, const FCacheKey& Key, ECachePolicy Policy, FValue& OutValue, FRequestStats& Stats);
+	[[nodiscard]] bool GetCacheValue(FStringView Name, const FCacheKey& Key, ECachePolicy Policy, FValue& OutValue, FRequestStats& Stats);
 
-	[[nodiscard]] bool GetCacheContentExists(const FCacheKey& Key, const FIoHash& RawHash) const;
+	[[nodiscard]] bool GetCacheContentExists(const FCacheKey& Key, const FIoHash& RawHash, FRequestStats& Stats) const;
 	[[nodiscard]] bool GetCacheContent(
 		FStringView Name,
 		const FCacheKey& Key,
 		const FValueId& Id,
 		const FValue& Value,
 		ECachePolicy Policy,
-		FValue& OutValue) const;
+		FValue& OutValue,
+		FRequestStats& Stats) const;
 	void GetCacheContent(
 		FStringView Name,
 		const FCacheKey& Key,
@@ -282,16 +290,17 @@ private:
 		const FValue& Value,
 		ECachePolicy Policy,
 		FCompressedBufferReader& Reader,
-		TUniquePtr<FArchive>& OutArchive) const;
+		TUniquePtr<FArchive>& OutArchive,
+		FRequestStats& Stats) const;
 
 	void BuildCachePackagePath(const FCacheKey& CacheKey, FStringBuilderBase& Path) const;
 	void BuildCacheContentPath(const FIoHash& RawHash, FStringBuilderBase& Path) const;
 
-	[[nodiscard]] bool LoadFileWithHash(FStringBuilderBase& Path, FStringView DebugName, TFunctionRef<void(FArchive&)> ReadFunction) const;
-	[[nodiscard]] bool LoadFile(FStringBuilderBase& Path, FStringView DebugName, TFunctionRef<void(FArchive&)> ReadFunction) const;
-	[[nodiscard]] TUniquePtr<FArchive> OpenFile(FStringBuilderBase& Path, FStringView DebugName) const;
+	[[nodiscard]] bool LoadFileWithHash(FStringBuilderBase& Path, FStringView DebugName, FRequestStats& Stats, TFunctionRef<void(FArchive&)> ReadFunction) const;
+	[[nodiscard]] bool LoadFile(FStringBuilderBase& Path, FStringView DebugName, FRequestStats& Stats, TFunctionRef<void(FArchive&)> ReadFunction) const;
+	[[nodiscard]] TUniquePtr<FArchive> OpenFile(FStringBuilderBase& Path, FStringView DebugName, FRequestStats& Stats) const;
 
-	[[nodiscard]] bool FileExists(FStringBuilderBase& Path) const;
+	[[nodiscard]] bool FileExists(FStringBuilderBase& Path, FRequestStats& Stats) const;
 
 	bool DownloadManifest(const FRootManifest& RootManifest, FFeedbackContext* Context);
 	void RemoveUnusedBundles();
@@ -894,12 +903,13 @@ struct FS3CacheStore::FRootManifest
 // FS3CacheStore
 //----------------------------------------------------------------------------------------------------------
 
-FS3CacheStore::FS3CacheStore(const TCHAR* InRootManifestPath, const TCHAR* InBaseUrl, const TCHAR* InRegion, const TCHAR* InCanaryObjectKey, const TCHAR* InCachePath)
+FS3CacheStore::FS3CacheStore(const TCHAR* InName, const TCHAR* InRootManifestPath, const TCHAR* InBaseUrl, const TCHAR* InRegion, const TCHAR* InCanaryObjectKey, const TCHAR* InCachePath, ICacheStoreOwner& InOwner)
 	: RootManifestPath(InRootManifestPath)
 	, BaseUrl(InBaseUrl)
 	, Region(InRegion)
 	, CanaryObjectKey(InCanaryObjectKey)
 	, CacheDir(InCachePath)
+	, StoreOwner(InOwner)
 	, bEnabled(false)
 {
 	FRootManifest RootManifest;
@@ -1008,6 +1018,10 @@ FS3CacheStore::FS3CacheStore(const TCHAR* InRootManifestPath, const TCHAR* InBas
 				// Mount all the bundles
 				ParallelFor(Bundles.Num(), [this](int32 Index) { ReadBundle(Bundles[Index]); });
 				bEnabled = true;
+
+				constexpr ECacheStoreFlags Flags = ECacheStoreFlags::Local | ECacheStoreFlags::Query | ECacheStoreFlags::StopStore;
+				StoreOwner.Add(this, Flags);
+				StoreStats = StoreOwner.CreateStats(this, Flags, TEXTVIEW("S3"), InName, BaseUrl);
 			}
 
 			Context->EndSlowTask();
@@ -1015,10 +1029,17 @@ FS3CacheStore::FS3CacheStore(const TCHAR* InRootManifestPath, const TCHAR* InBas
 	}
 }
 
+FS3CacheStore::~FS3CacheStore()
+{
+	if (StoreStats)
+	{
+		StoreOwner.DestroyStats(StoreStats);
+	}
+}
+
 void FS3CacheStore::LegacyStats(FDerivedDataCacheStatsNode& OutNode)
 {
-	OutNode = {TEXT("S3"), BaseUrl, /*bIsLocal*/ false};
-	OutNode.UsageStats.Add(TEXT(""), UsageStats);
+	checkNoEntry();
 }
 
 bool FS3CacheStore::DownloadManifest(const FRootManifest& RootManifest, FFeedbackContext* Context)
@@ -1182,7 +1203,8 @@ bool FS3CacheStore::LegacyDebugOptions(FBackendDebugOptions& InOptions)
 FOptionalCacheRecord FS3CacheStore::GetCacheRecordOnly(
 	const FStringView Name,
 	const FCacheKey& Key,
-	const FCacheRecordPolicy& Policy)
+	const FCacheRecordPolicy& Policy,
+	FRequestStats& Stats)
 {
 	if (!IsUsable())
 	{
@@ -1213,7 +1235,7 @@ FOptionalCacheRecord FS3CacheStore::GetCacheRecordOnly(
 	FOptionalCacheRecord Record;
 	{
 		FCbPackage Package;
-		if (!LoadFileWithHash(Path, Name, [&Package](FArchive& Ar) { Package.TryLoad(Ar); }))
+		if (!LoadFileWithHash(Path, Name, Stats, [&Package](FArchive& Ar) { Package.TryLoad(Ar); }))
 		{
 			UE_LOG(LogDerivedDataCache, Verbose, TEXT("%s: Cache miss with missing package for %s from '%.*s'"),
 				*GetName(), *WriteToString<96>(Key), Name.Len(), Name.GetData());
@@ -1243,9 +1265,10 @@ FOptionalCacheRecord FS3CacheStore::GetCacheRecord(
 	const FStringView Name,
 	const FCacheKey& Key,
 	const FCacheRecordPolicy& Policy,
-	EStatus& OutStatus)
+	EStatus& OutStatus,
+	FRequestStats& Stats)
 {
-	FOptionalCacheRecord Record = GetCacheRecordOnly(Name, Key, Policy);
+	FOptionalCacheRecord Record = GetCacheRecordOnly(Name, Key, Policy, Stats);
 	if (Record.IsNull())
 	{
 		OutStatus = EStatus::Error;
@@ -1267,7 +1290,7 @@ FOptionalCacheRecord FS3CacheStore::GetCacheRecord(
 		const FValueId& Id = Value.GetId();
 		const ECachePolicy ValuePolicy = Policy.GetValuePolicy(Id);
 		FValue Content;
-		if (GetCacheContent(Name, Key, Id, Value, ValuePolicy, Content))
+		if (GetCacheContent(Name, Key, Id, Value, ValuePolicy, Content, Stats))
 		{
 			RecordBuilder.AddValue(Id, MoveTemp(Content));
 		}
@@ -1290,7 +1313,8 @@ bool FS3CacheStore::GetCacheValueOnly(
 	const FStringView Name,
 	const FCacheKey& Key,
 	const ECachePolicy Policy,
-	FValue& OutValue)
+	FValue& OutValue,
+	FRequestStats& Stats)
 {
 	if (!IsUsable())
 	{
@@ -1319,7 +1343,7 @@ bool FS3CacheStore::GetCacheValueOnly(
 	BuildCachePackagePath(Key, Path);
 
 	FCbPackage Package;
-	if (!LoadFileWithHash(Path, Name, [&Package](FArchive& Ar) { Package.TryLoad(Ar); }))
+	if (!LoadFileWithHash(Path, Name, Stats, [&Package](FArchive& Ar) { Package.TryLoad(Ar); }))
 	{
 		UE_LOG(LogDerivedDataCache, Verbose, TEXT("%s: Cache miss with missing package for %s from '%.*s'"),
 			*GetName(), *WriteToString<96>(Key), Name.Len(), Name.GetData());
@@ -1367,16 +1391,17 @@ bool FS3CacheStore::GetCacheValue(
 	const FStringView Name,
 	const FCacheKey& Key,
 	const ECachePolicy Policy,
-	FValue& OutValue)
+	FValue& OutValue,
+	FRequestStats& Stats)
 {
-	return GetCacheValueOnly(Name, Key, Policy, OutValue) && GetCacheContent(Name, Key, {}, OutValue, Policy, OutValue);
+	return GetCacheValueOnly(Name, Key, Policy, OutValue, Stats) && GetCacheContent(Name, Key, {}, OutValue, Policy, OutValue, Stats);
 }
 
-bool FS3CacheStore::GetCacheContentExists(const FCacheKey& Key, const FIoHash& RawHash) const
+bool FS3CacheStore::GetCacheContentExists(const FCacheKey& Key, const FIoHash& RawHash, FRequestStats& Stats) const
 {
 	TStringBuilder<256> Path;
 	BuildCacheContentPath(RawHash, Path);
-	return FileExists(Path);
+	return FileExists(Path, Stats);
 }
 
 bool FS3CacheStore::GetCacheContent(
@@ -1385,7 +1410,8 @@ bool FS3CacheStore::GetCacheContent(
 	const FValueId& Id,
 	const FValue& Value,
 	const ECachePolicy Policy,
-	FValue& OutValue) const
+	FValue& OutValue,
+	FRequestStats& Stats) const
 {
 	if (!EnumHasAnyFlags(Policy, ECachePolicy::Query))
 	{
@@ -1405,7 +1431,7 @@ bool FS3CacheStore::GetCacheContent(
 	BuildCacheContentPath(RawHash, Path);
 	if (EnumHasAnyFlags(Policy, ECachePolicy::SkipData))
 	{
-		if (FileExists(Path))
+		if (FileExists(Path, Stats))
 		{
 			OutValue = Value;
 			return true;
@@ -1414,7 +1440,7 @@ bool FS3CacheStore::GetCacheContent(
 	else
 	{
 		FCompressedBuffer CompressedBuffer;
-		if (LoadFileWithHash(Path, Name, [&CompressedBuffer](FArchive& Ar) { CompressedBuffer = FCompressedBuffer::Load(Ar); }))
+		if (LoadFileWithHash(Path, Name, Stats, [&CompressedBuffer](FArchive& Ar) { CompressedBuffer = FCompressedBuffer::Load(Ar); }))
 		{
 			if (CompressedBuffer.GetRawHash() == RawHash)
 			{
@@ -1443,8 +1469,29 @@ void FS3CacheStore::GetCacheContent(
 	const FValue& Value,
 	const ECachePolicy Policy,
 	FCompressedBufferReader& Reader,
-	TUniquePtr<FArchive>& OutArchive) const
+	TUniquePtr<FArchive>& OutArchive,
+	FRequestStats& Stats) const
 {
+	class FStatsArchive final : TUniquePtr<FArchive>, public FArchiveProxy
+	{
+	public:
+		FStatsArchive(FArchive& InArchive, FRequestStats& InStats)
+			: TUniquePtr<FArchive>(&InArchive)
+			, FArchiveProxy(InArchive)
+			, Stats(InStats)
+		{
+		}
+
+		void Serialize(void* V, int64 Length) final
+		{
+			Stats.PhysicalReadSize += uint64(Length);
+			FArchiveProxy::Serialize(V, Length);
+		}
+
+	private:
+		FRequestStats& Stats;
+	};
+
 	if (!EnumHasAnyFlags(Policy, ECachePolicy::Query))
 	{
 		return;
@@ -1466,16 +1513,20 @@ void FS3CacheStore::GetCacheContent(
 	BuildCacheContentPath(RawHash, Path);
 	if (EnumHasAllFlags(Policy, ECachePolicy::SkipData))
 	{
-		if (FileExists(Path))
+		if (FileExists(Path, Stats))
 		{
 			return;
 		}
 	}
 	else
 	{
-		OutArchive = OpenFile(Path, Name);
+		OutArchive = OpenFile(Path, Name, Stats);
 		if (OutArchive)
 		{
+			OutArchive.Reset(new FStatsArchive(*OutArchive.Release(), Stats));
+			UE_LOG(LogDerivedDataCache, VeryVerbose,
+				TEXT("%s: Opened %s from '%.*s'"),
+				*GetName(), *Path, Name.Len(), Name.GetData());
 			Reader.SetSource(*OutArchive);
 			if (Reader.GetRawHash() == RawHash)
 			{
@@ -1510,9 +1561,10 @@ void FS3CacheStore::BuildCacheContentPath(const FIoHash& RawHash, FStringBuilder
 bool FS3CacheStore::LoadFileWithHash(
 	FStringBuilderBase& Path,
 	const FStringView DebugName,
+	FRequestStats& Stats,
 	const TFunctionRef<void (FArchive& Ar)> ReadFunction) const
 {
-	return LoadFile(Path, DebugName, [this, &Path, &DebugName, &ReadFunction](FArchive& Ar)
+	return LoadFile(Path, DebugName, Stats, [this, &Path, &DebugName, &ReadFunction](FArchive& Ar)
 	{
 		THashingArchiveProxy<FBlake3> HashAr(Ar);
 		ReadFunction(HashAr);
@@ -1534,6 +1586,7 @@ bool FS3CacheStore::LoadFileWithHash(
 bool FS3CacheStore::LoadFile(
 	FStringBuilderBase& Path,
 	const FStringView DebugName,
+	FRequestStats& Stats,
 	const TFunctionRef<void (FArchive& Ar)> ReadFunction) const
 {
 	check(IsUsable());
@@ -1542,7 +1595,7 @@ bool FS3CacheStore::LoadFile(
 	int64 ReadSize = 0;
 	bool bError = false;
 
-	if (TUniquePtr<FArchive> Ar = OpenFile(Path, DebugName))
+	if (TUniquePtr<FArchive> Ar = OpenFile(Path, DebugName, Stats))
 	{
 		int64 ReadStart = Ar->Tell();
 		ReadFunction(*Ar);
@@ -1564,11 +1617,19 @@ bool FS3CacheStore::LoadFile(
 		TEXT("%s: Loaded %s from '%.*s' (%" INT64_FMT " bytes, %.02f secs, %.2f MiB/s)"),
 		*GetName(), *Path, DebugName.Len(), DebugName.GetData(), ReadSize, ReadDuration, ReadSpeed);
 
+	if (ReadSize > 0)
+	{
+		Stats.PhysicalReadSize += uint64(ReadSize);
+	}
+
 	return !bError && ReadSize > 0;
 }
 
-TUniquePtr<FArchive> FS3CacheStore::OpenFile(FStringBuilderBase& Path, const FStringView DebugName) const
+TUniquePtr<FArchive> FS3CacheStore::OpenFile(FStringBuilderBase& Path, const FStringView DebugName, FRequestStats& Stats) const
 {
+	const FMonotonicTimePoint StartTime = FMonotonicTimePoint::Now();
+	ON_SCOPE_EXIT { Stats.AddLatency(FMonotonicTimePoint::Now() - StartTime); };
+
 	FSHAHash Hash;
 
 	auto AnsiString = StringCast<ANSICHAR>(*FString(Path.ToString()).ToUpper());
@@ -1591,8 +1652,11 @@ TUniquePtr<FArchive> FS3CacheStore::OpenFile(FStringBuilderBase& Path, const FSt
 	return nullptr;
 }
 
-bool FS3CacheStore::FileExists(FStringBuilderBase& Path) const
+bool FS3CacheStore::FileExists(FStringBuilderBase& Path, FRequestStats& Stats) const
 {
+	const FMonotonicTimePoint StartTime = FMonotonicTimePoint::Now();
+	ON_SCOPE_EXIT { Stats.AddLatency(FMonotonicTimePoint::Now() - StartTime); };
+
 	FSHAHash Hash;
 
 	auto AnsiString = StringCast<ANSICHAR>(*FString(Path.ToString()).ToUpper());
@@ -1627,22 +1691,28 @@ void FS3CacheStore::Get(
 	{
 		EStatus Status = EStatus::Error;
 		FOptionalCacheRecord Record;
+		FRequestStats RequestStats;
+		RequestStats.Bucket = Request.Key.Bucket;
+		RequestStats.Type = ERequestType::Record;
+		RequestStats.Op = ERequestOp::Get;
 		{
 			TRACE_CPUPROFILER_EVENT_SCOPE(S3DDC_Get);
 			TRACE_COUNTER_INCREMENT(S3DDC_Get);
-			COOK_STAT(auto Timer = UsageStats.TimeGet());
-			if ((Record = GetCacheRecord(Request.Name, Request.Key, Request.Policy, Status)))
+			FRequestTimer RequestTimer(RequestStats);
+			if ((Record = GetCacheRecord(Request.Name, Request.Key, Request.Policy, Status, RequestStats)))
 			{
 				UE_LOG(LogDerivedDataCache, Verbose, TEXT("%s: Cache hit for %s from '%s'"),
 					*GetName(), *WriteToString<96>(Request.Key), *Request.Name);
 				TRACE_COUNTER_INCREMENT(S3DDC_GetHit);
-				COOK_STAT(Timer.AddHit(Private::GetCacheRecordCompressedSize(Record.Get())));
 			}
 			else
 			{
 				Record = FCacheRecordBuilder(Request.Key).Build();
 			}
 		}
+		RequestStats.AddLogicalRead(Record.Get());
+		RequestStats.Status = Status;
+		StoreStats->AddRequest(RequestStats);
 		OnComplete({Request.Name, MoveTemp(Record).Get(), Request.UserData, Status});
 	}
 }
@@ -1664,19 +1734,25 @@ void FS3CacheStore::GetValue(
 	{
 		bool bOk;
 		FValue Value;
+		FRequestStats RequestStats;
+		RequestStats.Bucket = Request.Key.Bucket;
+		RequestStats.Type = ERequestType::Value;
+		RequestStats.Op = ERequestOp::Get;
 		{
 			TRACE_CPUPROFILER_EVENT_SCOPE(S3DDC_GetValue);
 			TRACE_COUNTER_INCREMENT(S3DDC_Get);
-			COOK_STAT(auto Timer = UsageStats.TimeGet());
-			bOk = GetCacheValue(Request.Name, Request.Key, Request.Policy, Value);
+			FRequestTimer RequestTimer(RequestStats);
+			bOk = GetCacheValue(Request.Name, Request.Key, Request.Policy, Value, RequestStats);
 			if (bOk)
 			{
 				UE_LOG(LogDerivedDataCache, Verbose, TEXT("%s: Cache hit for %s from '%s'"),
 					*GetName(), *WriteToString<96>(Request.Key), *Request.Name);
 				TRACE_COUNTER_INCREMENT(S3DDC_GetHit);
-				COOK_STAT(Timer.AddHit(Value.GetData().GetCompressedSize()));
 			}
 		}
+		RequestStats.AddLogicalRead(Value);
+		RequestStats.Status = bOk ? EStatus::Ok : EStatus::Error;
+		StoreStats->AddRequest(RequestStats);
 		OnComplete({Request.Name, Request.Key, Value, Request.UserData, bOk ? EStatus::Ok : EStatus::Error});
 	}
 }
@@ -1701,11 +1777,15 @@ void FS3CacheStore::GetChunks(
 		EStatus Status = EStatus::Error;
 		FSharedBuffer Buffer;
 		uint64 RawSize = 0;
+		FRequestStats RequestStats;
+		RequestStats.Bucket = Request.Key.Bucket;
+		RequestStats.Type = Request.Id.IsNull() ? ERequestType::Value : ERequestType::Record;
+		RequestStats.Op = ERequestOp::GetChunk;
 		{
 			TRACE_CPUPROFILER_EVENT_SCOPE(S3DDC_GetChunks);
 			TRACE_COUNTER_INCREMENT(S3DDC_Get);
 			const bool bExistsOnly = EnumHasAnyFlags(Request.Policy, ECachePolicy::SkipData);
-			COOK_STAT(auto Timer = bExistsOnly ? UsageStats.TimeProbablyExists() : UsageStats.TimeGet());
+			FRequestTimer RequestTimer(RequestStats);
 			if (!(bHasValue && ValueKey == Request.Key && ValueId == Request.Id) || ValueReader.HasSource() < !bExistsOnly)
 			{
 				ValueReader.ResetSource();
@@ -1721,7 +1801,7 @@ void FS3CacheStore::GetChunks(
 						FCacheRecordPolicyBuilder PolicyBuilder(ECachePolicy::None);
 						PolicyBuilder.AddValuePolicy(Request.Id, Request.Policy);
 						Record.Reset();
-						Record = GetCacheRecordOnly(Request.Name, Request.Key, PolicyBuilder.Build());
+						Record = GetCacheRecordOnly(Request.Name, Request.Key, PolicyBuilder.Build(), RequestStats);
 					}
 					if (Record)
 					{
@@ -1731,17 +1811,17 @@ void FS3CacheStore::GetChunks(
 							Value = ValueWithId;
 							ValueId = Request.Id;
 							ValueKey = Request.Key;
-							GetCacheContent(Request.Name, Request.Key, ValueId, Value, Request.Policy, ValueReader, ValueAr);
+							GetCacheContent(Request.Name, Request.Key, ValueId, Value, Request.Policy, ValueReader, ValueAr, RequestStats);
 						}
 					}
 				}
 				else
 				{
 					ValueKey = Request.Key;
-					bHasValue = GetCacheValueOnly(Request.Name, Request.Key, Request.Policy, Value);
+					bHasValue = GetCacheValueOnly(Request.Name, Request.Key, Request.Policy, Value, RequestStats);
 					if (bHasValue)
 					{
-						GetCacheContent(Request.Name, Request.Key, Request.Id, Value, Request.Policy, ValueReader, ValueAr);
+						GetCacheContent(Request.Name, Request.Key, Request.Id, Value, Request.Policy, ValueReader, ValueAr, RequestStats);
 					}
 				}
 			}
@@ -1749,39 +1829,92 @@ void FS3CacheStore::GetChunks(
 			{
 				const uint64 RawOffset = FMath::Min(Value.GetRawSize(), Request.RawOffset);
 				RawSize = FMath::Min(Value.GetRawSize() - RawOffset, Request.RawSize);
-				UE_LOG(LogDerivedDataCache, Verbose, TEXT("%s: Cache hit for %s from '%s'"),
-					*GetName(), *WriteToString<96>(Request.Key, '/', Request.Id), *Request.Name);
 				TRACE_COUNTER_INCREMENT(S3DDC_GetHit);
-				COOK_STAT(Timer.AddHit(!bExistsOnly ? RawSize : 0));
 				if (!bExistsOnly)
 				{
 					Buffer = ValueReader.Decompress(RawOffset, RawSize);
+					RequestStats.LogicalReadSize += Buffer.GetSize();
 				}
 				Status = bExistsOnly || Buffer.GetSize() == RawSize ? EStatus::Ok : EStatus::Error;
 			}
 		}
+		RequestStats.Status = Status;
+		StoreStats->AddRequest(RequestStats);
+		UE_CLOG(Status == EStatus::Ok, LogDerivedDataCache, Verbose, TEXT("%s: Cache hit for %s from '%s'"),
+			*GetName(), *WriteToString<96>(Request.Key, '/', Request.Id), *Request.Name);
 		OnComplete({Request.Name, Request.Key, Request.Id, Request.RawOffset,
 			RawSize, Value.GetRawHash(), MoveTemp(Buffer), Request.UserData, Status});
 	}
 }
 
-ILegacyCacheStore* CreateS3CacheStore(const TCHAR* RootManifestPath, const TCHAR* BaseUrl, const TCHAR* Region, const TCHAR* CanaryObjectKey, const TCHAR* CachePath)
-{
-	return new FS3CacheStore(RootManifestPath, BaseUrl, Region, CanaryObjectKey, CachePath);
-}
-
 } // UE::DerivedData
 
-#else
+#endif // WITH_S3_DDC_BACKEND
 
 namespace UE::DerivedData
 {
 
-ILegacyCacheStore* CreateS3CacheStore(const TCHAR* RootManifestPath, const TCHAR* BaseUrl, const TCHAR* Region, const TCHAR* CanaryObjectKey, const TCHAR* CachePath)
+ILegacyCacheStore* CreateS3CacheStore(const TCHAR* Name, const TCHAR* Config, ICacheStoreOwner& Owner)
 {
+#if WITH_S3_DDC_BACKEND
+	FString ManifestPath;
+	if (!FParse::Value(Config, TEXT("Manifest="), ManifestPath))
+	{
+		UE_LOG(LogDerivedDataCache, Error, TEXT("%s: Missing required parameter 'Manifest'"), Name);
+		return nullptr;
+	}
+
+	FString BaseUrl;
+	if (!FParse::Value(Config, TEXT("BaseUrl="), BaseUrl))
+	{
+		UE_LOG(LogDerivedDataCache, Error, TEXT("%s: Missing required parameter 'BaseUrl'"), Name);
+		return nullptr;
+	}
+
+	FString Region;
+	if (!FParse::Value(Config, TEXT("Region="), Region))
+	{
+		UE_LOG(LogDerivedDataCache, Error, TEXT("%s: Missing required parameter 'Region'"), Name);
+		return nullptr;
+	}
+
+	FString CanaryObjectKey;
+	FParse::Value(Config, TEXT("Canary="), CanaryObjectKey);
+
+	FString CachePath = FPaths::ProjectSavedDir() / TEXT("S3DDC");
+
+	FString Key;
+	if (FParse::Value(Config, TEXT("EnvPathOverride="), Key))
+	{
+		if (FString Value = FPlatformMisc::GetEnvironmentVariable(*Key); !Value.IsEmpty())
+		{
+			CachePath = MoveTemp(Value);
+			UE_LOG(LogDerivedDataCache, Log, TEXT("%s: Found environment variable %s=%s"), Name, *Key, *CachePath);
+		}
+		if (FString Value; FPlatformMisc::GetStoredValue(TEXT("Epic Games"), TEXT("GlobalDataCachePath"), *Key, Value) && !Value.IsEmpty())
+		{
+			CachePath = MoveTemp(Value);
+			UE_LOG(LogDerivedDataCache, Log, TEXT("%s: Found registry key GlobalDataCachePath %s=%s"), Name, *Key, *CachePath);
+		}
+	}
+
+	if (CachePath == TEXTVIEW("None"))
+	{
+		UE_LOG(LogDerivedDataCache, Log, TEXT("%s: Disabled because the path is configured to 'None'"), Name);
+		return nullptr;
+	}
+
+	TUniquePtr<FS3CacheStore> Store(new FS3CacheStore(Name, *ManifestPath, *BaseUrl, *Region, *CanaryObjectKey, *CachePath, Owner));
+	if (!Store->IsUsable())
+	{
+		return nullptr;
+	}
+
+	return Store.Release();
+#else
+	UE_LOG(LogDerivedDataCache, Log, TEXT("S3 backend is not supported on the current platform."));
 	return nullptr;
+#endif
 }
 
 } // UE::DerivedData
-
-#endif // WITH_S3_DDC_BACKEND
