DDC: Added ICacheStoreStats to the Jupiter cache

#rnx

diff --git a/Engine/Source/Developer/DerivedDataCache/Private/DerivedDataCacheStore.cpp b/Engine/Source/Developer/DerivedDataCache/Private/DerivedDataCacheStore.cpp
--- a/Engine/Source/Developer/DerivedDataCache/Private/DerivedDataCacheStore.cpp
+++ b/Engine/Source/Developer/DerivedDataCache/Private/DerivedDataCacheStore.cpp
@@ -17,14 +17,19 @@ FCacheStoreRequestTimer::FCacheStoreRequestTimer(FCacheStoreRequestStats& OutSta
 {
 }
 
-FCacheStoreRequestTimer::~FCacheStoreRequestTimer()
+void FCacheStoreRequestTimer::Stop()
 {
+	if (StartTime.IsInfinity())
+	{
+		return;
+	}
 	const FMonotonicTimePoint EndTime = FMonotonicTimePoint::Now();
 	const bool bIsInGameThread = IsInGameThread();
 	TUniqueLock Lock(Stats.Mutex);
 	(bIsInGameThread ? Stats.MainThreadTime : Stats.OtherThreadTime) = EndTime - StartTime;
 	Stats.StartTime = FMath::Min(Stats.StartTime, StartTime);
 	Stats.EndTime = FMath::Max(Stats.EndTime, EndTime);
+	StartTime = FMonotonicTimePoint::Infinity();
 }
 
 void FCacheStoreRequestStats::AddLatency(const FMonotonicTimeSpan InLatency)
diff --git a/Engine/Source/Developer/DerivedDataCache/Private/DerivedDataCacheStore.h b/Engine/Source/Developer/DerivedDataCache/Private/DerivedDataCacheStore.h
--- a/Engine/Source/Developer/DerivedDataCache/Private/DerivedDataCacheStore.h
+++ b/Engine/Source/Developer/DerivedDataCache/Private/DerivedDataCacheStore.h
@@ -220,7 +220,9 @@ class FCacheStoreRequestTimer final
 {
 public:
 	explicit FCacheStoreRequestTimer(FCacheStoreRequestStats& Stats);
-	~FCacheStoreRequestTimer();
+	inline ~FCacheStoreRequestTimer() { Stop(); }
+
+	void Stop();
 
 	FCacheStoreRequestTimer(const FCacheStoreRequestTimer&) = delete;
 	FCacheStoreRequestTimer& operator=(const FCacheStoreRequestTimer&) = delete;
@@ -237,11 +239,22 @@ class ICacheStoreStats
 {
 public:
 	virtual ~ICacheStoreStats() = default;
+
+	/** Returns the type name of the associated cache store. */
 	virtual FStringView GetType() const = 0;
+	/** Returns the name of the associated cache store in the hierarchy. */
 	virtual FStringView GetName() const = 0;
+	/** Returns the path to the associated cache store where applicable. */
 	virtual FStringView GetPath() const = 0;
+
+	/** Sets the flags detailing how the associated cache store operates. */
 	virtual void SetFlags(ECacheStoreFlags Flags) = 0;
+	/** Sets the status of the associated cache store. Persists until the next call. */
 	virtual void SetStatus(ECacheStoreStatusCode StatusCode, const FText& Status) = 0;
+	/** Sets a named attribute on the associated cache store. Persists until the next call. */
+	virtual void SetAttribute(FStringView Key, FStringView Value) = 0;
+
+	/** Adds stats for a single request that was processed by the associated cache store. */
 	virtual void AddRequest(const FCacheStoreRequestStats& Stats) = 0;
 };
 
diff --git a/Engine/Source/Developer/DerivedDataCache/Private/DerivedDataCacheStoreHierarchy.cpp b/Engine/Source/Developer/DerivedDataCache/Private/DerivedDataCacheStoreHierarchy.cpp
--- a/Engine/Source/Developer/DerivedDataCache/Private/DerivedDataCacheStoreHierarchy.cpp
+++ b/Engine/Source/Developer/DerivedDataCache/Private/DerivedDataCacheStoreHierarchy.cpp
@@ -48,6 +48,8 @@ public:
 	ECacheStoreStatusCode StatusCode = ECacheStoreStatusCode::None;
 	mutable FMutex Mutex;
 
+	TMap<FString, FString> Attributes;
+
 #if ENABLE_COOK_STATS
 	FCookStats::CallStats GetStats;
 	FCookStats::CallStats PutStats;
@@ -80,6 +82,7 @@ private:
 	FStringView GetPath() const final { return Path; }
 	void SetFlags(ECacheStoreFlags Flags) final;
 	void SetStatus(ECacheStoreStatusCode StatusCode, const FText& Status) final;
+	void SetAttribute(FStringView Key, FStringView Value) final;
 	void AddRequest(const FCacheStoreRequestStats& Stats) final;
 };
 
@@ -104,6 +107,12 @@ void FCacheStoreStats::SetStatus(ECacheStoreStatusCode InStatusCode, const FText
 	StatusCode = InStatusCode;
 }
 
+void FCacheStoreStats::SetAttribute(FStringView Key, FStringView Value)
+{
+	TUniqueLock Lock(Mutex);
+	Attributes.Emplace(Key, Value);
+}
+
 void FCacheStoreStats::AddRequest(const FCacheStoreRequestStats& Stats)
 {
 #if ENABLE_COOK_STATS
@@ -123,7 +132,7 @@ void FCacheStoreStats::AddRequest(const FCacheStoreRequestStats& Stats)
 		LatencyIndex %= LatencyHistory.Num();
 	}
 
-	if (Stats.PhysicalReadSize)
+	if (Stats.PhysicalReadSize && !Stats.LogicalWriteSize)
 	{
 		FRateStat& CurrentPhysicalReadSize = PhysicalReadSizeHistory[PhysicalReadSizeIndex++];
 		PhysicalReadSizeCount = FMath::Max(PhysicalReadSizeCount, PhysicalReadSizeIndex);
@@ -133,7 +142,7 @@ void FCacheStoreStats::AddRequest(const FCacheStoreRequestStats& Stats)
 		CurrentPhysicalReadSize.Amount = Stats.PhysicalReadSize;
 	}
 
-	if (Stats.PhysicalWriteSize)
+	if (Stats.PhysicalWriteSize && !Stats.LogicalReadSize)
 	{
 		FRateStat& CurrentPhysicalWriteSize = PhysicalWriteSizeHistory[PhysicalWriteSizeIndex++];
 		PhysicalWriteSizeCount = FMath::Max(PhysicalWriteSizeCount, PhysicalWriteSizeIndex);
@@ -372,6 +381,14 @@ FCacheStoreHierarchy::~FCacheStoreHierarchy()
 	{
 		Node.AsyncCache.Reset();
 		delete Node.Cache;
+
+		if (UNLIKELY(!Node.CacheStats.IsEmpty()))
+		{
+			FCacheStoreStats* Stats = Node.CacheStats[0];
+			UE_LOGFMT(LogDerivedDataCache, Fatal,
+				"Leaked stats for {Type} cache store '{Name}' with path '{Path}'.",
+				Stats->Type, Stats->Name, Stats->Path);
+		}
 	}
 }
 
@@ -1432,6 +1449,11 @@ static void ConvertToLegacyStats(FDerivedDataCacheStatsNode& OutNode, const FCac
 	FDerivedDataCacheUsageStats& UsageStats = OutNode.UsageStats.FindOrAdd({});
 	UsageStats.GetStats = Stats.GetStats;
 	UsageStats.PutStats = Stats.PutStats;
+
+	for (const TTuple<FString, FString>& Attribute : Stats.Attributes)
+	{
+		OutNode.CustomStats.Emplace(WriteToString<64>(Stats.Name, TEXT('.'), Attribute.Key), Attribute.Value);
+	}
 #endif
 
 	const auto MeasureActiveWallTime = [](TConstArrayView<FCacheStoreStats::FRateStat> History) -> FMonotonicTimeSpan
diff --git a/Engine/Source/Developer/DerivedDataCache/Private/HttpCacheStore.cpp b/Engine/Source/Developer/DerivedDataCache/Private/HttpCacheStore.cpp
--- a/Engine/Source/Developer/DerivedDataCache/Private/HttpCacheStore.cpp
+++ b/Engine/Source/Developer/DerivedDataCache/Private/HttpCacheStore.cpp
@@ -2,10 +2,10 @@
 
 #include "DerivedDataBackendInterface.h"
 #include "DerivedDataLegacyCacheStore.h"
-#include "Templates/Tuple.h"
 
 #if WITH_HTTP_DDC_BACKEND
 
+#include "Algo/AllOf.h"
 #include "Algo/Transform.h"
 #include "Async/Mutex.h"
 #include "Async/UniqueLock.h"
@@ -36,7 +36,6 @@
 #include "Misc/ScopeLock.h"
 #include "Misc/ScopeRWLock.h"
 #include "Misc/StringBuilder.h"
-#include "ProfilingDebugging/CookStats.h"
 #include "ProfilingDebugging/CountersTrace.h"
 #include "ProfilingDebugging/CpuProfilerTrace.h"
 #include "Serialization/CompactBinary.h"
@@ -317,6 +316,7 @@ private:
 	FAnsiStringBuilderBase EffectiveDomain;
 
 	ICacheStoreOwner* StoreOwner = nullptr;
+	ICacheStoreStats* StoreStats = nullptr;
 
 	FDerivedDataCacheUsageStats UsageStats;
 	FBackendDebugOptions DebugOptions;
@@ -354,36 +354,23 @@ private:
 
 	TUniquePtr<FHttpOperation> WaitForHttpOperation(EOperationCategory Category);
 
+	void PutCacheRecordAsync(IRequestOwner& Owner, const FCachePutRequest& Request, FOnCachePutComplete&& OnComplete);
+	void PutCacheValueAsync(IRequestOwner& Owner, const FCachePutValueRequest& Request, FOnCachePutValueComplete&& OnComplete);
+
 	void GetCacheRecordAsync(
 		IRequestOwner& Owner,
 		const FSharedString& Name,
 		const FCacheKey& Key,
 		const FCacheRecordPolicy& Policy,
 		uint64 UserData,
-		TUniqueFunction<void(FCacheGetResponse&& Response, uint64 BytesReceived)>&& OnComplete);
-
-	void PutCacheRecordAsync(
-		IRequestOwner& Owner,
-		const FSharedString& Name,
-		const FCacheRecord& Record,
-		const FCacheRecordPolicy& Policy,
-		uint64 UserData,
-		TUniqueFunction<void(FCachePutResponse&& Response, uint64 BytesSent)>&& OnComplete);
-
-	void PutCacheValueAsync(
-		IRequestOwner& Owner,
-		const FSharedString& Name,
-		const FCacheKey& Key,
-		const FValue& Value,
-		ECachePolicy Policy,
-		uint64 UserData,
-		TUniqueFunction<void(FCachePutValueResponse&& Response, uint64 BytesSent)>&& OnComplete);
+		FOnCacheGetComplete&& OnComplete);
 
 	void GetCacheValueAsync(
 		IRequestOwner& Owner,
 		FSharedString Name,
 		const FCacheKey& Key,
 		ECachePolicy Policy,
+		ERequestOp RequestOp,
 		uint64 UserData,
 		FOnCacheGetValueComplete&& OnComplete);
 
@@ -432,8 +419,7 @@ public:
 	FSharedBuffer GetBody() const { return ResponseBody; }
 	FString GetBodyAsString() const;
 	TSharedPtr<FJsonObject> GetBodyAsJson() const;
-	uint64 GetBytesSent() const { return Response->GetStats().SendSize; }
-	uint64 GetBytesReceived() const { return Response->GetStats().RecvSize; }
+	void GetStats(FRequestStats& OutStats) const;
 
 	friend FStringBuilderBase& operator<<(FStringBuilderBase& Builder, const FHttpOperation& Operation)
 	{
@@ -652,6 +638,18 @@ TSharedPtr<FJsonObject> FHttpCacheStore::FHttpOperation::GetBodyAsJson() const
 	return JsonObject;
 }
 
+void FHttpCacheStore::FHttpOperation::GetStats(FRequestStats& OutStats) const
+{
+	const FHttpResponseStats& Stats = Response->GetStats();
+	TUniqueLock Lock(OutStats.Mutex);
+	OutStats.PhysicalReadSize += Stats.RecvSize;
+	OutStats.PhysicalWriteSize += Stats.SendSize;
+	if (const EHttpMethod Method = Response->GetMethod(); Method == EHttpMethod::Get || Method == EHttpMethod::Head)
+	{
+		OutStats.AddLatency(FMonotonicTimeSpan::FromSeconds(Stats.StartTransferTime));
+	}
+}
+
 //----------------------------------------------------------------------------------------------------------
 // FHttpCacheStore::FHealthCheckOp
 //----------------------------------------------------------------------------------------------------------
@@ -698,7 +696,6 @@ class FHttpCacheStore::FPutPackageOp final : public FThreadSafeRefCountedObject
 public:
 	struct FResponse
 	{
-		uint64 BytesSent = 0;
 		EStatus Status = EStatus::Error;
 	};
 	using FOnPackageComplete = TUniqueFunction<void (FResponse&& Response)>;
@@ -710,14 +707,8 @@ public:
 
 	void Put(const FCacheKey& Key, FCbPackage&& Package, FOnPackageComplete&& OnComplete);
 
-	/** Performs a multi-request operation for uploading a package of content. */
-	static void PutPackage(
-		FHttpCacheStore& CacheStore,
-		IRequestOwner& Owner,
-		const FSharedString& Name,
-		FCacheKey Key,
-		FCbPackage&& Package,
-		FOnPackageComplete&& OnComplete);
+	const FRequestStats& ReadStats() const { return RequestStats; }
+	FRequestStats& EditStats() { return RequestStats; }
 
 private:
 	FHttpCacheStore& CacheStore;
@@ -729,14 +720,14 @@ private:
 	FIoHash ObjectHash;
 	FOnPackageComplete OnPackageComplete;
 
-	std::atomic<uint64> BytesSent = 0;
+	FRequestStats RequestStats;
+
 	std::atomic<uint32> SuccessfulBlobUploads = 0;
 	std::atomic<uint32> PendingBlobUploads = 0;
 	uint32 TotalBlobUploads = 0;
 
 	struct FCachePutRefResponse
 	{
-		uint64 BytesSent = 0;
 		TConstArrayView<FIoHash> NeededBlobHashes;
 		EStatus Status = EStatus::Error;
 	};
@@ -748,24 +739,12 @@ private:
 	void EndPutRef(TUniquePtr<FHttpOperation> Operation, bool bFinalize, FOnCachePutRefComplete&& OnComplete);
 
 	void BeginPutBlobs(FCbPackage&& Package, FCachePutRefResponse&& Response);
-	void EndPutBlob(FHttpOperation& Operation);
+	void EndPutBlob(FHttpOperation& Operation, uint64 LogicalSize);
 
 	void EndPutRefFinalize(FCachePutRefResponse&& Response);
-};
 
-void FHttpCacheStore::FPutPackageOp::PutPackage(
-	FHttpCacheStore& CacheStore,
-	IRequestOwner& Owner,
-	const FSharedString& Name,
-	FCacheKey Key,
-	FCbPackage&& Package,
-	FOnPackageComplete&& OnComplete)
-{
-	// TODO: Jupiter currently always overwrites.  It doesn't have a "write if not present" feature (for records or attachments),
-	//		 but would require one to implement all policy correctly.
-	TRefCountPtr<FPutPackageOp> Self = New(CacheStore, Owner, Name);
-	Self->Put(Key, MoveTemp(Package), MoveTemp(OnComplete));
-}
+	void EndPut(EStatus Status);
+};
 
 FHttpCacheStore::FPutPackageOp::FPutPackageOp(FHttpCacheStore& InCacheStore, IRequestOwner& InOwner, const FSharedString& InName)
 	: CacheStore(InCacheStore)
@@ -788,6 +767,8 @@ void FHttpCacheStore::FPutPackageOp::Put(const FCacheKey& InKey, FCbPackage&& Pa
 
 void FHttpCacheStore::FPutPackageOp::BeginPutRef(bool bFinalize, FOnCachePutRefComplete&& OnComplete)
 {
+	FRequestTimer RequestTimer(RequestStats);
+
 	TAnsiStringBuilder<64> Bucket;
 	Algo::Transform(Key.Bucket.ToString(), AppendChars(Bucket), FCharAnsi::ToLower);
 
@@ -814,8 +795,11 @@ void FHttpCacheStore::FPutPackageOp::BeginPutRef(bool bFinalize, FOnCachePutRefC
 		LocalOperation.SetBody(Object.GetBuffer());
 	}
 	LocalOperation.AddAcceptType(EHttpMediaType::Json);
+
+	RequestTimer.Stop();
 	LocalOperation.SendAsync(Owner, [Self = TRefCountPtr(this), Operation = MoveTemp(Operation), bFinalize, OnComplete = MoveTemp(OnComplete)]() mutable
 	{
+		Operation->GetStats(Self->RequestStats);
 		Self->EndPutRef(MoveTemp(Operation), bFinalize, MoveTemp(OnComplete));
 	});
 }
@@ -827,48 +811,49 @@ void FHttpCacheStore::FPutPackageOp::EndPutRef(
 {
 	TRACE_CPUPROFILER_EVENT_SCOPE(HttpDDC_PutPackage_EndPutRef);
 
-	if (const int32 StatusCode = Operation->GetStatusCode(); StatusCode >= 200 && StatusCode <= 204)
+	if (const int32 StatusCode = Operation->GetStatusCode(); StatusCode < 200 || StatusCode > 204)
 	{
-		TArray<FIoHash> NeededBlobHashes;
+		const EStatus Status = Operation->GetErrorCode() == EHttpErrorCode::Canceled ? EStatus::Canceled : EStatus::Error;
+		return OnComplete({{}, Status});
+	}
+
+	FRequestTimer RequestTimer(RequestStats);
 
-		// Useful when debugging issues related to compressed/uncompressed blobs being returned from Jupiter
-		static const bool bHttpCacheAlwaysPut = FParse::Param(FCommandLine::Get(), TEXT("HttpCacheAlwaysPut"));
+	TArray<FIoHash> NeededBlobHashes;
 
-		if (bHttpCacheAlwaysPut && !bFinalize)
+	// Useful when debugging issues related to compressed/uncompressed blobs being returned from Jupiter
+	static const bool bHttpCacheAlwaysPut = FParse::Param(FCommandLine::Get(), TEXT("HttpCacheAlwaysPut"));
+
+	if (bHttpCacheAlwaysPut && !bFinalize)
+	{
+		Object.IterateAttachments([&NeededBlobHashes](FCbFieldView AttachmentFieldView)
 		{
-			Object.IterateAttachments([&NeededBlobHashes](FCbFieldView AttachmentFieldView)
+			FIoHash AttachmentHash = AttachmentFieldView.AsHash();
+			if (!AttachmentHash.IsZero())
 			{
-				FIoHash AttachmentHash = AttachmentFieldView.AsHash();
-				if (!AttachmentHash.IsZero())
-				{
-					NeededBlobHashes.Add(AttachmentHash);
-				}
-			});
-		}
-		else if (TSharedPtr<FJsonObject> ResponseObject = Operation->GetBodyAsJson())
-		{
-			TArray<FString> NeedsArrayStrings;
-			ResponseObject->TryGetStringArrayField(TEXT("needs"), NeedsArrayStrings);
+				NeededBlobHashes.Add(AttachmentHash);
+			}
+		});
+	}
+	else if (TSharedPtr<FJsonObject> ResponseObject = Operation->GetBodyAsJson())
+	{
+		TArray<FString> NeedsArrayStrings;
+		ResponseObject->TryGetStringArrayField(TEXT("needs"), NeedsArrayStrings);
 
-			NeededBlobHashes.Reserve(NeedsArrayStrings.Num());
-			for (const FString& NeededString : NeedsArrayStrings)
+		NeededBlobHashes.Reserve(NeedsArrayStrings.Num());
+		for (const FString& NeededString : NeedsArrayStrings)
+		{
+			FIoHash BlobHash;
+			LexFromString(BlobHash, *NeededString);
+			if (!BlobHash.IsZero())
 			{
-				FIoHash BlobHash;
-				LexFromString(BlobHash, *NeededString);
-				if (!BlobHash.IsZero())
-				{
-					NeededBlobHashes.Add(BlobHash);
-				}
+				NeededBlobHashes.Add(BlobHash);
 			}
 		}
-
-		OnComplete({Operation->GetBytesSent(), NeededBlobHashes, EStatus::Ok});
-	}
-	else
-	{
-		const EStatus Status = Operation->GetErrorCode() == EHttpErrorCode::Canceled ? EStatus::Canceled : EStatus::Error;
-		OnComplete({Operation->GetBytesSent(), {}, Status});
 	}
+
+	RequestTimer.Stop();
+	OnComplete({NeededBlobHashes, EStatus::Ok});
 }
 
 void FHttpCacheStore::FPutPackageOp::BeginPutBlobs(FCbPackage&& Package, FCachePutRefResponse&& Response)
@@ -880,102 +865,89 @@ void FHttpCacheStore::FPutPackageOp::BeginPutBlobs(FCbPackage&& Package, FCacheP
 			UE_LOG(LogDerivedDataCache, Log, TEXT("%s: Failed to put reference object for put of %s from '%s'"),
 				*CacheStore.Domain, *WriteToString<96>(Key), *Name);
 		}
-		return OnPackageComplete({Response.BytesSent, Response.Status});
+		return EndPut(Response.Status);
 	}
 
-	struct FCompressedBlobUpload
-	{
-		FIoHash Hash;
-		FSharedBuffer BlobBuffer;
-
-		FCompressedBlobUpload(const FIoHash& InHash, FSharedBuffer&& InBlobBuffer)
-			: Hash(InHash)
-			, BlobBuffer(InBlobBuffer)
-		{
-		}
-	};
-
-	TArray<FCompressedBlobUpload> CompressedBlobUploads;
+	FRequestTimer RequestTimer(RequestStats);
 
 	// TODO: blob uploading and finalization should be replaced with a single batch compressed blob upload endpoint in the future.
 	TStringBuilder<128> ExpectedHashes;
 	bool bExpectedHashesSerialized = false;
 
 	// Needed blob upload (if any missing)
+	TArray<FCompressedBuffer> Blobs;
 	for (const FIoHash& NeededBlobHash : Response.NeededBlobHashes)
 	{
 		if (const FCbAttachment* Attachment = Package.FindAttachment(NeededBlobHash))
 		{
-			FSharedBuffer TempBuffer;
+			FCompressedBuffer Blob;
 			if (Attachment->IsCompressedBinary())
 			{
-				TempBuffer = Attachment->AsCompressedBinary().GetCompressed().ToShared();
+				Blob = Attachment->AsCompressedBinary();
 			}
 			else if (Attachment->IsBinary())
 			{
-				TempBuffer = FValue::Compress(Attachment->AsCompositeBinary()).GetData().GetCompressed().ToShared();
+				Blob = FValue::Compress(Attachment->AsCompositeBinary()).GetData();
 			}
 			else
 			{
-				TempBuffer = FValue::Compress(Attachment->AsObject().GetBuffer()).GetData().GetCompressed().ToShared();
+				Blob = FValue::Compress(Attachment->AsObject().GetBuffer()).GetData();
 			}
-
-			CompressedBlobUploads.Emplace(NeededBlobHash, MoveTemp(TempBuffer));
+			Blobs.Emplace(MoveTemp(Blob));
 		}
 		else
 		{
 			if (!bExpectedHashesSerialized)
 			{
-				bool bFirstHash = true;
 				for (const FCbAttachment& PackageAttachment : Package.GetAttachments())
 				{
-					if (!bFirstHash)
-					{
-						ExpectedHashes << TEXT(", ");
-					}
-					ExpectedHashes << PackageAttachment.GetHash();
-					bFirstHash = false;
+					ExpectedHashes << PackageAttachment.GetHash() << TEXTVIEW(", ");
+				}
+				if (ExpectedHashes.Len() >= 2)
+				{
+					ExpectedHashes.RemoveSuffix(2);
 				}
 				bExpectedHashesSerialized = true;
 			}
 			UE_LOG(LogDerivedDataCache, Log, TEXT("%s: Server reported needed hash '%s' that is outside the set of expected hashes (%s) for put of %s from '%s'"),
-				*CacheStore.Domain, *WriteToString<96>(NeededBlobHash), ExpectedHashes.ToString(), *WriteToString<96>(Key), *Name);
+				*CacheStore.Domain, *WriteToString<96>(NeededBlobHash), *ExpectedHashes, *WriteToString<96>(Key), *Name);
 		}
 	}
 
-	if (CompressedBlobUploads.IsEmpty())
+	if (Blobs.IsEmpty())
 	{
-		// No blobs need to be uploaded.  No finalization necessary.
-		return OnPackageComplete({Response.BytesSent, EStatus::Ok});
+		RequestTimer.Stop();
+		return EndPut(EStatus::Ok);
 	}
 
-	TotalBlobUploads = CompressedBlobUploads.Num();
+	TotalBlobUploads = Blobs.Num();
 	PendingBlobUploads.store(TotalBlobUploads, std::memory_order_relaxed);
 
 	FRequestBarrier Barrier(Owner);
-	for (const FCompressedBlobUpload& CompressedBlobUpload : CompressedBlobUploads)
+	for (const FCompressedBuffer& Blob : Blobs)
 	{
 		TUniquePtr<FHttpOperation> Operation = CacheStore.WaitForHttpOperation(EOperationCategory::Put);
 		FHttpOperation& LocalOperation = *Operation;
-		LocalOperation.SetUri(WriteToAnsiString<256>(CacheStore.EffectiveDomain, ANSITEXTVIEW("/api/v1/compressed-blobs/"), CacheStore.Namespace, '/', CompressedBlobUpload.Hash));
+		LocalOperation.SetUri(WriteToAnsiString<256>(CacheStore.EffectiveDomain, ANSITEXTVIEW("/api/v1/compressed-blobs/"), CacheStore.Namespace, '/', Blob.GetRawHash()));
 		LocalOperation.SetMethod(EHttpMethod::Put);
 		LocalOperation.SetContentType(EHttpMediaType::CompressedBinary);
-		LocalOperation.SetBody(FCompositeBuffer(CompressedBlobUpload.BlobBuffer));
-		LocalOperation.SendAsync(Owner, [Self = TRefCountPtr(this), Operation = MoveTemp(Operation)]
+		LocalOperation.SetBody(Blob.GetCompressed());
+		LocalOperation.SendAsync(Owner, [Self = TRefCountPtr(this), Operation = MoveTemp(Operation), LogicalSize = Blob.GetRawSize()]
 		{
-			Self->EndPutBlob(*Operation);
+			Operation->GetStats(Self->RequestStats);
+			Self->EndPutBlob(*Operation, LogicalSize);
 		});
 	}
 }
 
-void FHttpCacheStore::FPutPackageOp::EndPutBlob(FHttpOperation& Operation)
+void FHttpCacheStore::FPutPackageOp::EndPutBlob(FHttpOperation& Operation, uint64 LogicalSize)
 {
-	BytesSent.fetch_add(Operation.GetBytesSent(), std::memory_order_relaxed);
-
 	const int32 StatusCode = Operation.GetStatusCode();
 	if (StatusCode >= 200 && StatusCode <= 204)
 	{
 		SuccessfulBlobUploads.fetch_add(1, std::memory_order_relaxed);
+		TUniqueLock Lock(RequestStats.Mutex);
+		RequestStats.LogicalWriteSize += LogicalSize;
 	}
 
 	if (PendingBlobUploads.fetch_sub(1, std::memory_order_relaxed) == 1)
@@ -983,7 +955,7 @@ void FHttpCacheStore::FPutPackageOp::EndPutBlob(FHttpOperation& Operation)
 		const uint32 LocalSuccessfulBlobUploads = SuccessfulBlobUploads.load(std::memory_order_relaxed);
 		if (Owner.IsCanceled())
 		{
-			OnPackageComplete({BytesSent.load(std::memory_order_relaxed), EStatus::Canceled});
+			EndPut(EStatus::Canceled);
 		}
 		else if (LocalSuccessfulBlobUploads == TotalBlobUploads)
 		{
@@ -997,22 +969,31 @@ void FHttpCacheStore::FPutPackageOp::EndPutBlob(FHttpOperation& Operation)
 			const uint32 FailedBlobUploads = TotalBlobUploads - LocalSuccessfulBlobUploads;
 			UE_LOG(LogDerivedDataCache, Log, TEXT("%s: Failed to put %d/%d blobs for put of %s from '%s'"),
 				*CacheStore.Domain, FailedBlobUploads, TotalBlobUploads, *WriteToString<96>(Key), *Name);
-			OnPackageComplete({BytesSent.load(std::memory_order_relaxed), EStatus::Error});
+			EndPut(EStatus::Error);
 		}
 	}
 }
 
 void FHttpCacheStore::FPutPackageOp::EndPutRefFinalize(FCachePutRefResponse&& Response)
 {
-	BytesSent.fetch_add(Response.BytesSent, std::memory_order_relaxed);
-
 	if (Response.Status == EStatus::Error)
 	{
 		UE_LOG(LogDerivedDataCache, Log, TEXT("%s: Failed to finalize reference object for put of %s from '%s'"),
 			*CacheStore.Domain, *WriteToString<96>(Key), *Name);
 	}
 
-	return OnPackageComplete({BytesSent.load(std::memory_order_relaxed), Response.Status});
+	return EndPut(Response.Status);
+}
+
+void FHttpCacheStore::FPutPackageOp::EndPut(EStatus Status)
+{
+	RequestStats.EndTime = FMonotonicTimePoint::Now();
+	RequestStats.Status = Status;
+	OnPackageComplete({Status});
+	if (CacheStore.StoreStats)
+	{
+		CacheStore.StoreStats->AddRequest(RequestStats);
+	}
 }
 
 //----------------------------------------------------------------------------------------------------------
@@ -1028,7 +1009,6 @@ public:
 
 	struct FRecordResponse
 	{
-		uint64 BytesReceived = 0;
 		FCacheRecord Record;
 		EStatus Status = EStatus::Error;
 	};
@@ -1047,6 +1027,10 @@ public:
 	void GetValues(TConstArrayView<FValueWithId> Values, FOnValueComplete&& OnComplete);
 	void GetValuesExist(TConstArrayView<FValueWithId> Values, FOnValueComplete&& OnComplete);
 
+	const FRequestStats& ReadStats() const { return RequestStats; }
+	FRequestStats& EditStats() { return RequestStats; }
+	void RecordStats(EStatus Status);
+
 private:
 	FGetRecordOp(FHttpCacheStore& CacheStore, IRequestOwner& Owner, const FSharedString& Name);
 
@@ -1061,6 +1045,9 @@ private:
 	FCacheKey Key;
 	FCbPackage Package;
 	FOnRecordComplete OnRecordComplete;
+
+	FRequestStats RequestStats;
+
 	int32 PendingValues = 0;
 	int32 FailedValues = 0;
 	mutable FMutex Mutex;
@@ -1075,12 +1062,14 @@ FHttpCacheStore::FGetRecordOp::FGetRecordOp(FHttpCacheStore& InCacheStore, IRequ
 
 void FHttpCacheStore::FGetRecordOp::GetRecordOnly(const FCacheKey& InKey, const ECachePolicy RecordPolicy, FOnRecordComplete&& InOnComplete)
 {
+	FRequestTimer RequestTimer(RequestStats);
+
 	if (!CacheStore.IsUsable())
 	{
 		UE_LOG(LogDerivedDataCache, VeryVerbose,
 			TEXT("%s: Skipped get of %s from '%s' because this cache store is not available"),
 			*CacheStore.Domain, *WriteToString<96>(Key), *Name);
-		return InOnComplete({0, FCacheRecordBuilder(Key).Build(), EStatus::Error});
+		return InOnComplete({FCacheRecordBuilder(Key).Build(), EStatus::Error});
 	}
 
 	// Skip the request if querying the cache is disabled.
@@ -1088,14 +1077,14 @@ void FHttpCacheStore::FGetRecordOp::GetRecordOnly(const FCacheKey& InKey, const
 	{
 		UE_LOG(LogDerivedDataCache, VeryVerbose, TEXT("%s: Skipped get of %s from '%s' due to cache policy"),
 			*CacheStore.Domain, *WriteToString<96>(Key), *Name);
-		return InOnComplete({0, FCacheRecordBuilder(Key).Build(), EStatus::Error});
+		return InOnComplete({FCacheRecordBuilder(Key).Build(), EStatus::Error});
 	}
 
 	if (CacheStore.DebugOptions.ShouldSimulateGetMiss(Key))
 	{
 		UE_LOG(LogDerivedDataCache, Verbose, TEXT("%s: Simulated miss for get of %s from '%s'"),
 			*CacheStore.Domain, *WriteToString<96>(Key), *Name);
-		return InOnComplete({0, FCacheRecordBuilder(Key).Build(), EStatus::Error});
+		return InOnComplete({FCacheRecordBuilder(Key).Build(), EStatus::Error});
 	}
 
 	Key = InKey;
@@ -1110,8 +1099,11 @@ void FHttpCacheStore::FGetRecordOp::GetRecordOnly(const FCacheKey& InKey, const
 	LocalOperation.SetMethod(EHttpMethod::Get);
 	LocalOperation.AddAcceptType(EHttpMediaType::CbObject);
 	LocalOperation.SetExpectedErrorCodes({404});
+
+	RequestTimer.Stop();
 	LocalOperation.SendAsync(Owner, [Self = TRefCountPtr(this), Operation = MoveTemp(Operation)]() mutable
 	{
+		Operation->GetStats(Self->RequestStats);
 		Self->EndGetRef(MoveTemp(Operation));
 	});
 }
@@ -1120,18 +1112,20 @@ void FHttpCacheStore::FGetRecordOp::EndGetRef(TUniquePtr<FHttpOperation> Operati
 {
 	TRACE_CPUPROFILER_EVENT_SCOPE(HttpDDC_GetPackage_EndGetRef);
 
+	FRequestTimer RequestTimer(RequestStats);
+
 	FOptionalCacheRecord Record;
 	EStatus Status = EStatus::Error;
 	ON_SCOPE_EXIT
 	{
-		const uint64 BytesReceived = Operation->GetBytesReceived();
 		Operation.Reset();
 		if (Record.IsNull())
 		{
 			Record = FCacheRecordBuilder(Key).Build();
 		}
+		RequestTimer.Stop();
 		FOnRecordComplete LocalOnComplete = MoveTemp(OnRecordComplete);
-		LocalOnComplete({BytesReceived, MoveTemp(Record).Get(), Status});
+		LocalOnComplete({MoveTemp(Record).Get(), Status});
 	};
 
 	const int32 StatusCode = Operation->GetStatusCode();
@@ -1181,6 +1175,8 @@ void FHttpCacheStore::FGetRecordOp::GetRecord(const FCacheKey& LocalKey, const F
 
 void FHttpCacheStore::FGetRecordOp::BeginGetValues(const FCacheRecord& Record, const FCacheRecordPolicy& Policy, FOnRecordComplete&& OnComplete)
 {
+	FRequestTimer RequestTimer(RequestStats);
+
 	OnRecordComplete = MoveTemp(OnComplete);
 
 	TArray<FValueWithId> RequiredGets;
@@ -1197,6 +1193,8 @@ void FHttpCacheStore::FGetRecordOp::BeginGetValues(const FCacheRecord& Record, c
 
 	PendingValues = RequiredGets.Num() + RequiredHeads.Num();
 
+	RequestTimer.Stop();
+
 	if (PendingValues == 0)
 	{
 		return EndGetValues(Policy, EStatus::Ok);
@@ -1270,7 +1268,7 @@ void FHttpCacheStore::FGetRecordOp::EndGetValues(const FCacheRecordPolicy& Polic
 	}
 
 	FOnRecordComplete LocalOnComplete = MoveTemp(OnRecordComplete);
-	LocalOnComplete({0, RecordBuilder.Build(), Status});
+	LocalOnComplete({RecordBuilder.Build(), Status});
 }
 
 void FHttpCacheStore::FGetRecordOp::GetValues(TConstArrayView<FValueWithId> Values, FOnValueComplete&& OnComplete)
@@ -1293,6 +1291,8 @@ void FHttpCacheStore::FGetRecordOp::GetValues(TConstArrayView<FValueWithId> Valu
 
 	// TODO: Jupiter does not currently provide a batched GET. Once it does, fetch every blob in one request.
 
+	FRequestTimer RequestTimer(RequestStats);
+
 	FRequestBarrier Barrier(Owner);
 	TSharedRef<FOnValueComplete> SharedOnComplete = MakeShared<FOnValueComplete>(MoveTemp(OnComplete));
 	for (const FValueWithId& Value : Values)
@@ -1312,6 +1312,9 @@ void FHttpCacheStore::FGetRecordOp::GetValues(TConstArrayView<FValueWithId> Valu
 		{
 			TRACE_CPUPROFILER_EVENT_SCOPE(HttpDDC_GetPackage_GetValues_OnResponse);
 
+			FRequestTimer RequestTimer(Self->RequestStats);
+			Operation->GetStats(Self->RequestStats);
+
 			bool bHit = false;
 			FCompressedBuffer CompressedBuffer;
 			if (Operation->GetStatusCode() == 200)
@@ -1330,8 +1333,13 @@ void FHttpCacheStore::FGetRecordOp::GetValues(TConstArrayView<FValueWithId> Valu
 				default:
 					break;
 				}
+
+				TUniqueLock Lock(Self->RequestStats.Mutex);
+				Self->RequestStats.LogicalReadSize += CompressedBuffer.GetRawSize();
 			}
 
+			RequestTimer.Stop();
+
 			if (bHit)
 			{
 				if (CompressedBuffer.GetRawHash() == RawHash && CompressedBuffer.GetRawSize() == RawSize)
@@ -1381,6 +1389,8 @@ void FHttpCacheStore::FGetRecordOp::GetValuesExist(TConstArrayView<FValueWithId>
 		return;
 	}
 
+	FRequestTimer RequestTimer(RequestStats);
+
 	TAnsiStringBuilder<256> Uri;
 	Uri << CacheStore.EffectiveDomain << ANSITEXTVIEW("/api/v1/compressed-blobs/") << CacheStore.Namespace << ANSITEXTVIEW("/exists?");
 	for (const FValueWithId& Value : QueryValues)
@@ -1396,10 +1406,15 @@ void FHttpCacheStore::FGetRecordOp::GetValuesExist(TConstArrayView<FValueWithId>
 	LocalOperation.SetMethod(EHttpMethod::Post);
 	LocalOperation.SetContentType(EHttpMediaType::FormUrlEncoded);
 	LocalOperation.AddAcceptType(EHttpMediaType::Json);
+
+	RequestTimer.Stop();
 	LocalOperation.SendAsync(Owner, [Self = TRefCountPtr(this), Operation = MoveTemp(Operation), Values = MoveTemp(QueryValues), OnComplete = MoveTemp(OnComplete)]() mutable
 	{
 		TRACE_CPUPROFILER_EVENT_SCOPE(HttpDDC_DataProbablyExistsBatch_OnHttpRequestComplete);
 
+		FRequestTimer RequestTimer(Self->RequestStats);
+		Operation->GetStats(Self->RequestStats);
+
 		const TCHAR* DefaultMessage = TEXT("Cache exists miss for");
 		EStatus DefaultStatus = EStatus::Error;
 
@@ -1441,6 +1456,8 @@ void FHttpCacheStore::FGetRecordOp::GetValuesExist(TConstArrayView<FValueWithId>
 			}
 		}
 
+		RequestTimer.Stop();
+
 		for (const FValueWithId& Value : Values)
 		{
 			UE_LOG(LogDerivedDataCache, Verbose,
@@ -1452,6 +1469,16 @@ void FHttpCacheStore::FGetRecordOp::GetValuesExist(TConstArrayView<FValueWithId>
 	});
 }
 
+void FHttpCacheStore::FGetRecordOp::RecordStats(EStatus Status)
+{
+	RequestStats.EndTime = FMonotonicTimePoint::Now();
+	RequestStats.Status = Status;
+	if (CacheStore.StoreStats)
+	{
+		CacheStore.StoreStats->AddRequest(RequestStats);
+	}
+}
+
 //----------------------------------------------------------------------------------------------------------
 // FHttpCacheStore::FGetValueOp
 //----------------------------------------------------------------------------------------------------------
@@ -1474,10 +1501,14 @@ public:
 
 	void Get(const FCacheKey& Key, ECachePolicy Policy, FOnComplete&& OnComplete);
 
+	const FRequestStats& ReadStats() const { return RequestStats; }
+	FRequestStats& EditStats() { return RequestStats; }
+
 private:
 	FGetValueOp(FHttpCacheStore& CacheStore, IRequestOwner& Owner, const FSharedString& Name);
 
 	void EndGetRef(TUniquePtr<FHttpOperation> Operation);
+	void EndGet(FResponse&& Response);
 
 	FHttpCacheStore& CacheStore;
 	IRequestOwner& Owner;
@@ -1485,6 +1516,7 @@ private:
 	FCacheKey Key;
 	ECachePolicy Policy = ECachePolicy::None;
 	FOnComplete OnComplete;
+	FRequestStats RequestStats;
 };
 
 FHttpCacheStore::FGetValueOp::FGetValueOp(FHttpCacheStore& InCacheStore, IRequestOwner& InOwner, const FSharedString& InName)
@@ -1496,6 +1528,8 @@ FHttpCacheStore::FGetValueOp::FGetValueOp(FHttpCacheStore& InCacheStore, IReques
 
 void FHttpCacheStore::FGetValueOp::Get(const FCacheKey& InKey, ECachePolicy InPolicy, FOnComplete&& InOnComplete)
 {
+	FRequestTimer RequestTimer(RequestStats);
+
 	Key = InKey;
 	Policy = InPolicy;
 	OnComplete = MoveTemp(InOnComplete);
@@ -1518,8 +1552,11 @@ void FHttpCacheStore::FGetValueOp::Get(const FCacheKey& InKey, ECachePolicy InPo
 		LocalOperation.AddHeader(ANSITEXTVIEW("Accept"), ANSITEXTVIEW("application/x-jupiter-inline"));
 	}
 	LocalOperation.SetExpectedErrorCodes({404});
+
+	RequestTimer.Stop();
 	LocalOperation.SendAsync(Owner, [Self = TRefCountPtr(this), Operation = MoveTemp(Operation)]() mutable
 	{
+		Operation->GetStats(Self->RequestStats);
 		Self->EndGetRef(MoveTemp(Operation));
 	});
 }
@@ -1535,7 +1572,7 @@ void FHttpCacheStore::FGetValueOp::EndGetRef(TUniquePtr<FHttpOperation> Operatio
 	{
 		UE_LOG(LogDerivedDataCache, Verbose, TEXT("%s: Cache miss with failed HTTP request for %s from '%s'"),
 			*CacheStore.Domain, *WriteToString<96>(Key), *Name);
-		return OnComplete({Name, Key, {}, EStatus::Error});
+		return EndGet({Name, Key, {}, EStatus::Error});
 	}
 
 	FSharedBuffer Body = Operation->GetBody();
@@ -1546,7 +1583,7 @@ void FHttpCacheStore::FGetValueOp::EndGetRef(TUniquePtr<FHttpOperation> Operatio
 		{
 			UE_LOG(LogDerivedDataCache, Display, TEXT("%s: Cache miss with invalid package for %s from '%s'"),
 				*CacheStore.Domain, *WriteToString<96>(Key), *Name);
-			return OnComplete({Name, Key, {}, EStatus::Error});
+			return EndGet({Name, Key, {}, EStatus::Error});
 		}
 
 		const FCbObjectView Object = FCbObject(Body);
@@ -1556,10 +1593,10 @@ void FHttpCacheStore::FGetValueOp::EndGetRef(TUniquePtr<FHttpOperation> Operatio
 		{
 			UE_LOG(LogDerivedDataCache, Display, TEXT("%s: Cache miss with invalid value for %s from '%s'"),
 				*CacheStore.Domain, *WriteToString<96>(Key), *Name);
-			return OnComplete({Name, Key, {}, EStatus::Error});
+			return EndGet({Name, Key, {}, EStatus::Error});
 		}
 
-		OnComplete({Name, Key, FValue(RawHash, RawSize), EStatus::Ok});
+		EndGet({Name, Key, FValue(RawHash, RawSize), EStatus::Ok});
 	}
 	else
 	{
@@ -1567,6 +1604,7 @@ void FHttpCacheStore::FGetValueOp::EndGetRef(TUniquePtr<FHttpOperation> Operatio
 
 		if (!CompressedBuffer)
 		{
+			FRequestTimer RequestTimer(RequestStats);
 			if (FAnsiStringView ReceivedHashStr = Operation->GetHeader("X-Jupiter-InlinePayloadHash"); !ReceivedHashStr.IsEmpty())
 			{
 				FIoHash ReceivedHash(ReceivedHashStr);
@@ -1582,10 +1620,22 @@ void FHttpCacheStore::FGetValueOp::EndGetRef(TUniquePtr<FHttpOperation> Operatio
 		{
 			UE_LOG(LogDerivedDataCache, Display, TEXT("%s: Cache miss with invalid package for %s from '%s'"),
 				*CacheStore.Domain, *WriteToString<96>(Key), *Name);
-			return OnComplete({Name, Key, {}, EStatus::Error});
+			return EndGet({Name, Key, {}, EStatus::Error});
 		}
 
-		OnComplete({Name, Key, FValue(CompressedBuffer), EStatus::Ok});
+		EndGet({Name, Key, FValue(CompressedBuffer), EStatus::Ok});
+	}
+}
+
+void FHttpCacheStore::FGetValueOp::EndGet(FResponse&& Response)
+{
+	RequestStats.LogicalReadSize += Response.Value.GetRawSize();
+	RequestStats.EndTime = FMonotonicTimePoint::Now();
+	RequestStats.Status = Response.Status;
+	OnComplete(MoveTemp(Response));
+	if (CacheStore.StoreStats)
+	{
+		CacheStore.StoreStats->AddRequest(RequestStats);
 	}
 }
 
@@ -1607,10 +1657,13 @@ private:
 
 	void EndExists(TUniquePtr<FHttpOperation> Operation);
 
+	void EndRequest(const FCacheGetValueRequest& Request, const FValue& Value, EStatus Status);
+
 	FHttpCacheStore& CacheStore;
 	IRequestOwner& Owner;
 	TArray<FCacheGetValueRequest> Requests;
 	FOnCacheGetValueComplete OnComplete;
+	FRequestStats RequestStats;
 };
 
 FHttpCacheStore::FExistsBatchOp::FExistsBatchOp(FHttpCacheStore& InCacheStore, IRequestOwner& InOwner)
@@ -1621,6 +1674,8 @@ FHttpCacheStore::FExistsBatchOp::FExistsBatchOp(FHttpCacheStore& InCacheStore, I
 
 void FHttpCacheStore::FExistsBatchOp::Exists(TConstArrayView<FCacheGetValueRequest> InRequests, FOnCacheGetValueComplete&& InOnComplete)
 {
+	FRequestTimer RequestTimer(RequestStats);
+
 	OnComplete = MoveTemp(InOnComplete);
 
 	Requests.Empty(InRequests.Num());
@@ -1688,8 +1743,11 @@ void FHttpCacheStore::FExistsBatchOp::Exists(TConstArrayView<FCacheGetValueReque
 	LocalOperation.SetContentType(EHttpMediaType::CbObject);
 	LocalOperation.AddAcceptType(EHttpMediaType::CbObject);
 	LocalOperation.SetBody(FCompositeBuffer(Body.GetOuterBuffer()));
+
+	RequestTimer.Stop();
 	LocalOperation.SendAsync(Owner, [Self = TRefCountPtr(this), Operation = MoveTemp(Operation)]() mutable
 	{
+		Operation->GetStats(Self->RequestStats);
 		Self->EndExists(MoveTemp(Operation));
 	});
 }
@@ -1698,6 +1756,15 @@ void FHttpCacheStore::FExistsBatchOp::EndExists(TUniquePtr<FHttpOperation> Opera
 {
 	TRACE_CPUPROFILER_EVENT_SCOPE(HttpDDC_ExistsBatch_EndExists);
 
+	// Divide the stats evenly among the requests.
+	RequestStats.PhysicalReadSize /= Requests.Num();
+	RequestStats.PhysicalWriteSize /= Requests.Num();
+	RequestStats.MainThreadTime = FMonotonicTimeSpan::FromSeconds(RequestStats.MainThreadTime.ToSeconds() / Requests.Num());
+	RequestStats.OtherThreadTime = FMonotonicTimeSpan::FromSeconds(RequestStats.OtherThreadTime.ToSeconds() / Requests.Num());
+	RequestStats.EndTime = FMonotonicTimePoint::Now();
+	RequestStats.Type = ERequestType::Value;
+	RequestStats.Op = ERequestOp::Get;
+
 	const int32 OverallStatusCode = Operation->GetStatusCode();
 	if (OverallStatusCode < 200 || OverallStatusCode > 204)
 	{
@@ -1705,7 +1772,7 @@ void FHttpCacheStore::FExistsBatchOp::EndExists(TUniquePtr<FHttpOperation> Opera
 		{
 			UE_LOG(LogDerivedDataCache, Verbose, TEXT("%s: Cache miss with failed HTTP request for %s from '%s'"),
 				*CacheStore.Domain, *WriteToString<96>(Request.Key), *Request.Name);
-			OnComplete(Request.MakeResponse(EStatus::Error));
+			EndRequest(Request, {}, EStatus::Error);
 		}
 		return;
 	}
@@ -1717,7 +1784,7 @@ void FHttpCacheStore::FExistsBatchOp::EndExists(TUniquePtr<FHttpOperation> Opera
 		{
 			UE_LOG(LogDerivedDataCache, Log, TEXT("%s: Cache miss with corrupt response for %s from '%s'."),
 				*CacheStore.Domain, *WriteToString<96>(Request.Key), *Request.Name);
-			OnComplete(Request.MakeResponse(EStatus::Error));
+			EndRequest(Request, {}, EStatus::Error);
 		}
 		return;
 	}
@@ -1734,7 +1801,7 @@ void FHttpCacheStore::FExistsBatchOp::EndExists(TUniquePtr<FHttpOperation> Opera
 		{
 			UE_LOG(LogDerivedDataCache, Display, TEXT("%s: Cache miss with invalid response for %s from '%s'"),
 				*CacheStore.Domain, *WriteToString<96>(Request.Key), *Request.Name);
-			OnComplete(Request.MakeResponse(EStatus::Error));
+			EndRequest(Request, {}, EStatus::Error);
 		}
 		return;
 	}
@@ -1759,7 +1826,7 @@ void FHttpCacheStore::FExistsBatchOp::EndExists(TUniquePtr<FHttpOperation> Opera
 		{
 			UE_LOG(LogDerivedDataCache, Verbose, TEXT("%s: Cache miss with unsuccessful response code %d for %s from '%s'"),
 				*CacheStore.Domain, StatusCode, *WriteToString<96>(Request.Key), *Request.Name);
-			OnComplete(Request.MakeResponse(EStatus::Error));
+			EndRequest(Request, {}, EStatus::Error);
 			continue;
 		}
 
@@ -1769,11 +1836,22 @@ void FHttpCacheStore::FExistsBatchOp::EndExists(TUniquePtr<FHttpOperation> Opera
 		{
 			UE_LOG(LogDerivedDataCache, Display, TEXT("%s: Cache miss with invalid value for %s from '%s'"),
 				*CacheStore.Domain, *WriteToString<96>(Request.Key), *Request.Name);
-			OnComplete(Request.MakeResponse(EStatus::Error));
+			EndRequest(Request, {}, EStatus::Error);
 			continue;
 		}
 
-		OnComplete({Request.Name, Request.Key, FValue(RawHash, RawSize), Request.UserData, EStatus::Ok});
+		EndRequest(Request, FValue(RawHash, RawSize), EStatus::Ok);
+	}
+}
+
+void FHttpCacheStore::FExistsBatchOp::EndRequest(const FCacheGetValueRequest& Request, const FValue& Value, EStatus Status)
+{
+	RequestStats.EndTime = FMonotonicTimePoint::Now();
+	RequestStats.Status = Status;
+	OnComplete({Request.Name, Request.Key, Value, Request.UserData, Status});
+	if (CacheStore.StoreStats)
+	{
+		CacheStore.StoreStats->AddRequest(RequestStats);
 	}
 }
 
@@ -1857,6 +1935,14 @@ FHttpCacheStore::FHttpCacheStore(const FHttpCacheStoreParams& Params, ICacheStor
 				(Params.bReadOnly ? ECacheStoreFlags::None : ECacheStoreFlags::Store);
 			TStringBuilder<256> Path(InPlace, Domain, TEXTVIEW(" ("), Namespace, TEXTVIEW(")"));
 			StoreOwner->Add(this, Flags);
+			StoreStats = StoreOwner->CreateStats(this, Flags, TEXTVIEW("Unreal Cloud DDC"), Params.Name, Path);
+
+			StoreStats->SetAttribute(TEXTVIEW("Domain"), Domain);
+			StoreStats->SetAttribute(TEXTVIEW("EffectiveDomain"), WriteToString<128>(EffectiveDomain));
+			StoreStats->SetAttribute(TEXTVIEW("Namespace"), Namespace);
+			StoreStats->SetAttribute(TEXTVIEW("LoginAttempts"), WriteToString<16>(LoginAttempts));
+			StoreStats->SetAttribute(TEXTVIEW("InteractiveLoginAttempts"), WriteToString<16>(InteractiveLoginAttempts));
+			StoreStats->SetAttribute(TEXTVIEW("FailedLoginAttempts"), WriteToString<16>(FailedLoginAttempts));
 		}
 	}
 
@@ -1870,6 +1956,11 @@ FHttpCacheStore::~FHttpCacheStore()
 		FTSTicker::GetCoreTicker().RemoveTicker(RefreshAccessTokenHandle);
 	}
 
+	if (StoreStats)
+	{
+		StoreOwner->DestroyStats(StoreStats);
+	}
+
 	if (AnyInstance == this)
 	{
 		AnyInstance = nullptr;
@@ -1942,6 +2033,16 @@ bool FHttpCacheStore::AcquireAccessToken(IHttpClient* Client)
 		return true;
 	}
 
+	ON_SCOPE_EXIT
+	{
+		if (StoreStats)
+		{
+			StoreStats->SetAttribute(TEXTVIEW("LoginAttempts"), WriteToString<16>(LoginAttempts));
+			StoreStats->SetAttribute(TEXTVIEW("InteractiveLoginAttempts"), WriteToString<16>(InteractiveLoginAttempts));
+			StoreStats->SetAttribute(TEXTVIEW("FailedLoginAttempts"), WriteToString<16>(FailedLoginAttempts));
+		}
+	};
+
 	LoginAttempts++;
 
 	// Avoid spamming this if the service is down.
@@ -2156,107 +2257,119 @@ TUniquePtr<FHttpCacheStore::FHttpOperation> FHttpCacheStore::WaitForHttpOperatio
 	return MakeUnique<FHttpOperation>(MoveTemp(Request));
 }
 
-void FHttpCacheStore::PutCacheRecordAsync(
-	IRequestOwner& Owner,
-	const FSharedString& Name,
-	const FCacheRecord& Record,
-	const FCacheRecordPolicy& Policy,
-	uint64 UserData,
-	TUniqueFunction<void(FCachePutResponse&& Response, uint64 BytesSent)>&& OnComplete)
+void FHttpCacheStore::PutCacheRecordAsync(IRequestOwner& Owner, const FCachePutRequest& Request, FOnCachePutComplete&& OnComplete)
 {
-	const FCacheKey& Key = Record.GetKey();
-	auto MakeResponse = [Name = FSharedString(Name), Key = FCacheKey(Key), UserData](EStatus Status)
-	{
-		return FCachePutResponse{ Name, Key, UserData, Status };
-	};
+	const FCacheKey& Key = Request.Record.GetKey();
 
 	if (bReadOnly)
 	{
 		UE_LOG(LogDerivedDataCache, VeryVerbose,
 			TEXT("%s: Skipped put of %s from '%s' because this cache store is read-only"),
-			*Domain, *WriteToString<96>(Key), *Name);
-		return OnComplete(MakeResponse(EStatus::Error), 0);
+			*Domain, *WriteToString<96>(Key), *Request.Name);
+		return OnComplete(Request.MakeResponse(EStatus::Error));
 	}
 
 	// Skip the request if storing to the cache is disabled.
-	const ECachePolicy RecordPolicy = Policy.GetRecordPolicy();
+	const ECachePolicy RecordPolicy = Request.Policy.GetRecordPolicy();
 	if (!EnumHasAnyFlags(RecordPolicy, ECachePolicy::StoreRemote))
 	{
 		UE_LOG(LogDerivedDataCache, VeryVerbose, TEXT("%s: Skipped put of %s from '%s' due to cache policy"),
-			*Domain, *WriteToString<96>(Key), *Name);
-		return OnComplete(MakeResponse(EStatus::Error), 0);
+			*Domain, *WriteToString<96>(Key), *Request.Name);
+		return OnComplete(Request.MakeResponse(EStatus::Error));
 	}
 
 	if (DebugOptions.ShouldSimulatePutMiss(Key))
 	{
 		UE_LOG(LogDerivedDataCache, Verbose, TEXT("%s: Simulated miss for put of %s from '%s'"),
-			*Domain, *WriteToString<96>(Key), *Name);
-		return OnComplete(MakeResponse(EStatus::Error), 0);
+			*Domain, *WriteToString<96>(Key), *Request.Name);
+		return OnComplete(Request.MakeResponse(EStatus::Error));
 	}
 
 	// TODO: Jupiter currently always overwrites.  It doesn't have a "write if not present" feature (for records or attachments),
 	//		 but would require one to implement all policy correctly.
 
-	FCbPackage Package = Record.Save();
+	TRefCountPtr<FPutPackageOp> Op = FPutPackageOp::New(*this, Owner, Request.Name);
 
-	FPutPackageOp::PutPackage(*this, Owner, Name, Key, MoveTemp(Package), [MakeResponse = MoveTemp(MakeResponse), OnComplete = MoveTemp(OnComplete)](FPutPackageOp::FResponse&& Response)
+	FCbPackage Package;
 	{
-		OnComplete(MakeResponse(Response.Status), Response.BytesSent);
+		FRequestStats& RequestStats = Op->EditStats();
+		RequestStats.Bucket = Key.Bucket;
+		RequestStats.Type = ERequestType::Record;
+		RequestStats.Op = ERequestOp::Put;
+
+		FRequestTimer RequestTimer(RequestStats);
+		Package = Request.Record.Save();
+	}
+
+	Op->Put(Key, MoveTemp(Package), [Op, Request, OnComplete = MoveTemp(OnComplete)](FPutPackageOp::FResponse&& Response)
+	{
+		if (Response.Status == EStatus::Ok)
+		{
+			if (const FCbObject& Meta = Request.Record.GetMeta())
+			{
+				Op->EditStats().LogicalWriteSize += Meta.GetSize();
+			}
+		}
+		TRACE_COUNTER_ADD(HttpDDC_BytesReceived, Op->ReadStats().PhysicalReadSize);
+		TRACE_COUNTER_ADD(HttpDDC_BytesSent, Op->ReadStats().PhysicalWriteSize);
+		OnComplete(Request.MakeResponse(Response.Status));
 	});
 }
 
-void FHttpCacheStore::PutCacheValueAsync(
-	IRequestOwner& Owner,
-	const FSharedString& Name,
-	const FCacheKey& Key,
-	const FValue& Value,
-	const ECachePolicy Policy,
-	uint64 UserData,
-	TUniqueFunction<void(FCachePutValueResponse&& Response, uint64 BytesSent)>&& OnComplete)
+void FHttpCacheStore::PutCacheValueAsync(IRequestOwner& Owner, const FCachePutValueRequest& Request, FOnCachePutValueComplete&& OnComplete)
 {
-	auto MakeResponse = [Name = FSharedString(Name), Key = FCacheKey(Key), UserData](EStatus Status)
-	{
-		return FCachePutValueResponse{ Name, Key, UserData, Status };
-	};
-
 	if (bReadOnly)
 	{
 		UE_LOG(LogDerivedDataCache, VeryVerbose,
 			TEXT("%s: Skipped put of %s from '%s' because this cache store is read-only"),
-			*Domain, *WriteToString<96>(Key), *Name);
-		return OnComplete(MakeResponse(EStatus::Error), 0);
+			*Domain, *WriteToString<96>(Request.Key), *Request.Name);
+		return OnComplete(Request.MakeResponse(EStatus::Error));
 	}
 
 	// Skip the request if storing to the cache is disabled.
-	if (!EnumHasAnyFlags(Policy, ECachePolicy::StoreRemote))
+	if (!EnumHasAnyFlags(Request.Policy, ECachePolicy::StoreRemote))
 	{
 		UE_LOG(LogDerivedDataCache, VeryVerbose, TEXT("%s: Skipped put of %s from '%s' due to cache policy"),
-			*Domain, *WriteToString<96>(Key), *Name);
-		return OnComplete(MakeResponse(EStatus::Error), 0);
+			*Domain, *WriteToString<96>(Request.Key), *Request.Name);
+		return OnComplete(Request.MakeResponse(EStatus::Error));
 	}
 
-	if (DebugOptions.ShouldSimulatePutMiss(Key))
+	if (DebugOptions.ShouldSimulatePutMiss(Request.Key))
 	{
 		UE_LOG(LogDerivedDataCache, Verbose, TEXT("%s: Simulated miss for put of %s from '%s'"),
-			*Domain, *WriteToString<96>(Key), *Name);
-		return OnComplete(MakeResponse(EStatus::Error), 0);
+			*Domain, *WriteToString<96>(Request.Key), *Request.Name);
+		return OnComplete(Request.MakeResponse(EStatus::Error));
 	}
 
 	// TODO: Jupiter currently always overwrites.  It doesn't have a "write if not present" feature (for records or attachments),
 	//		 but would require one to implement all policy correctly.
 
-	FCbWriter Writer;
-	Writer.BeginObject();
-	Writer.AddBinaryAttachment("RawHash", Value.GetRawHash());
-	Writer.AddInteger("RawSize", Value.GetRawSize());
-	Writer.EndObject();
+	TRefCountPtr<FPutPackageOp> Op = FPutPackageOp::New(*this, Owner, Request.Name);
 
-	FCbPackage Package(Writer.Save().AsObject());
-	Package.AddAttachment(FCbAttachment(Value.GetData()));
+	FCbPackage Package;
+	{
+		FRequestStats& RequestStats = Op->EditStats();
+		RequestStats.Bucket = Request.Key.Bucket;
+		RequestStats.Type = ERequestType::Value;
+		RequestStats.Op = ERequestOp::Put;
 
-	FPutPackageOp::PutPackage(*this, Owner, Name, Key, MoveTemp(Package), [MakeResponse = MoveTemp(MakeResponse), OnComplete = MoveTemp(OnComplete)](FPutPackageOp::FResponse&& Response)
+		FRequestTimer RequestTimer(RequestStats);
+
+		FCbWriter Writer;
+		Writer.BeginObject();
+		Writer.AddBinaryAttachment("RawHash", Request.Value.GetRawHash());
+		Writer.AddInteger("RawSize", Request.Value.GetRawSize());
+		Writer.EndObject();
+
+		Package.SetObject(Writer.Save().AsObject());
+		Package.AddAttachment(FCbAttachment(Request.Value.GetData()));
+	}
+
+	Op->Put(Request.Key, MoveTemp(Package), [Op, Request, OnComplete = MoveTemp(OnComplete)](FPutPackageOp::FResponse&& Response)
 	{
-		OnComplete(MakeResponse(Response.Status), Response.BytesSent);
+		TRACE_COUNTER_ADD(HttpDDC_BytesReceived, Op->ReadStats().PhysicalReadSize);
+		TRACE_COUNTER_ADD(HttpDDC_BytesSent, Op->ReadStats().PhysicalWriteSize);
+		OnComplete(Request.MakeResponse(Response.Status));
 	});
 }
 
@@ -2265,6 +2378,7 @@ void FHttpCacheStore::GetCacheValueAsync(
 	FSharedString Name,
 	const FCacheKey& Key,
 	ECachePolicy Policy,
+	ERequestOp RequestOp,
 	uint64 UserData,
 	FOnCacheGetValueComplete&& OnComplete)
 {
@@ -2295,8 +2409,16 @@ void FHttpCacheStore::GetCacheValueAsync(
 	}
 
 	TRefCountPtr<FGetValueOp> Op = FGetValueOp::New(*this, Owner, Name);
-	Op->Get(Key, Policy, [UserData, OnComplete = MoveTemp(OnComplete)](FGetValueOp::FResponse&& Response)
+
+	FRequestStats& RequestStats = Op->EditStats();
+	RequestStats.Bucket = Key.Bucket;
+	RequestStats.Type = ERequestType::Value;
+	RequestStats.Op = RequestOp;
+
+	Op->Get(Key, Policy, [Op, UserData, OnComplete = MoveTemp(OnComplete)](FGetValueOp::FResponse&& Response)
 	{
+		TRACE_COUNTER_ADD(HttpDDC_BytesReceived, Op->ReadStats().PhysicalReadSize);
+		TRACE_COUNTER_ADD(HttpDDC_BytesSent, Op->ReadStats().PhysicalWriteSize);
 		OnComplete({Response.Name, Response.Key, MoveTemp(Response.Value), UserData, Response.Status});
 	});
 }
@@ -2307,43 +2429,34 @@ void FHttpCacheStore::GetCacheRecordAsync(
 	const FCacheKey& Key,
 	const FCacheRecordPolicy& Policy,
 	uint64 UserData,
-	TUniqueFunction<void(FCacheGetResponse&& Response, uint64 BytesReceived)>&& OnComplete)
+	FOnCacheGetComplete&& OnComplete)
 {
 	TRefCountPtr<FGetRecordOp> Op = FGetRecordOp::New(*this, Owner, Name);
-	Op->GetRecord(Key, Policy, [Name, UserData, OnComplete = MoveTemp(OnComplete)](FGetRecordOp::FRecordResponse&& Response)
+
+	FRequestStats& RequestStats = Op->EditStats();
+	RequestStats.Bucket = Key.Bucket;
+	RequestStats.Type = ERequestType::Record;
+	RequestStats.Op = ERequestOp::Get;
+
+	Op->GetRecord(Key, Policy, [Op, Name, UserData, OnComplete = MoveTemp(OnComplete)](FGetRecordOp::FRecordResponse&& Response)
 	{
-		OnComplete({Name, MoveTemp(Response.Record), UserData, Response.Status}, 0);
+		if (Response.Status == EStatus::Ok)
+		{
+			if (const FCbObject& Meta = Response.Record.GetMeta())
+			{
+				Op->EditStats().LogicalReadSize += Meta.GetSize();
+			}
+		}
+		Op->RecordStats(Response.Status);
+		TRACE_COUNTER_ADD(HttpDDC_BytesReceived, Op->ReadStats().PhysicalReadSize);
+		TRACE_COUNTER_ADD(HttpDDC_BytesSent, Op->ReadStats().PhysicalWriteSize);
+		OnComplete({Name, MoveTemp(Response.Record), UserData, Response.Status});
 	});
 }
 
 void FHttpCacheStore::LegacyStats(FDerivedDataCacheStatsNode& OutNode)
 {
-	OutNode = {TEXT("Unreal Cloud DDC"), FString::Printf(TEXT("%s (%s)"), *Domain, *Namespace), /*bIsLocal*/ false};
-	OutNode.UsageStats.Add(TEXT(""), UsageStats);
-
-#if ENABLE_COOK_STATS
-	const int64 GetHits = UsageStats.GetStats.GetAccumulatedValueAnyThread(FCookStats::CallStats::EHitOrMiss::Hit, FCookStats::CallStats::EStatType::Counter);
-	const int64 GetMisses = UsageStats.GetStats.GetAccumulatedValueAnyThread(FCookStats::CallStats::EHitOrMiss::Miss, FCookStats::CallStats::EStatType::Counter);
-	const int64 PutHits = UsageStats.PutStats.GetAccumulatedValueAnyThread(FCookStats::CallStats::EHitOrMiss::Hit, FCookStats::CallStats::EStatType::Counter);
-	const int64 PutMisses = UsageStats.PutStats.GetAccumulatedValueAnyThread(FCookStats::CallStats::EHitOrMiss::Miss, FCookStats::CallStats::EStatType::Counter);
-	const int64 TotalGets = GetHits + GetMisses;
-	const int64 TotalPuts = PutHits + PutMisses;
-
-	const FString BaseName(TEXTVIEW("CloudDDC."));
-
-	OutNode.CustomStats = FCookStatsManager::CreateKeyValueArray(	FString(WriteToString<64>(BaseName, TEXTVIEW("Domain"))), Domain,
-																	FString(WriteToString<64>(BaseName, TEXTVIEW("EffectiveDomain"))), *EffectiveDomain,
-																	FString(WriteToString<64>(BaseName, TEXTVIEW("Namespace"))), Namespace,
-																	FString(WriteToString<64>(BaseName, TEXTVIEW("LoginAttempts"))), LoginAttempts,
-																	FString(WriteToString<64>(BaseName, TEXTVIEW("InteractiveLoginAttempts"))), InteractiveLoginAttempts,
-																	FString(WriteToString<64>(BaseName, TEXTVIEW("FailedLoginAttempts"))), FailedLoginAttempts,
-																	FString(WriteToString<64>(BaseName, TEXTVIEW("GetHits"))), GetHits,
-																	FString(WriteToString<64>(BaseName, TEXTVIEW("GetMisses"))), GetMisses,
-																	FString(WriteToString<64>(BaseName, TEXTVIEW("TotalGets"))), TotalGets,
-																	FString(WriteToString<64>(BaseName, TEXTVIEW("PutHits"))), PutHits,
-																	FString(WriteToString<64>(BaseName, TEXTVIEW("PutMisses"))), PutMisses,
-																	FString(WriteToString<64>(BaseName, TEXTVIEW("TotalPuts"))), TotalPuts);
-#endif
+	checkNoEntry();
 }
 
 void FHttpCacheStore::Put(
@@ -2352,17 +2465,17 @@ void FHttpCacheStore::Put(
 	FOnCachePutComplete&& OnComplete)
 {
 	TRACE_CPUPROFILER_EVENT_SCOPE(HttpDDC_Put);
+	TRACE_COUNTER_ADD(HttpDDC_Put, Requests.Num());
+
 	FRequestBarrier Barrier(Owner);
 	TSharedRef<FOnCachePutComplete> SharedOnComplete = MakeShared<FOnCachePutComplete>(MoveTemp(OnComplete));
 	for (const FCachePutRequest& Request : Requests)
 	{
-		PutCacheRecordAsync(Owner, Request.Name, Request.Record, Request.Policy, Request.UserData,
-			[COOK_STAT(Timer = UsageStats.TimePut(), ) SharedOnComplete](FCachePutResponse&& Response, uint64 BytesSent) mutable
+		PutCacheRecordAsync(Owner, Request, [SharedOnComplete](FCachePutResponse&& Response)
 		{
-			TRACE_COUNTER_ADD(HttpDDC_BytesSent, BytesSent);
 			if (Response.Status == EStatus::Ok)
 			{
-				COOK_STAT(if (BytesSent) { Timer.AddHit(BytesSent); });
+				TRACE_COUNTER_INCREMENT(HttpDDC_PutHit);
 			}
 			SharedOnComplete.Get()(MoveTemp(Response));
 		});
@@ -2375,17 +2488,18 @@ void FHttpCacheStore::Get(
 	FOnCacheGetComplete&& OnComplete)
 {
 	TRACE_CPUPROFILER_EVENT_SCOPE(HttpDDC_Get);
+	TRACE_COUNTER_ADD(HttpDDC_Get, Requests.Num());
+
 	FRequestBarrier Barrier(Owner);
 	TSharedRef<FOnCacheGetComplete> SharedOnComplete = MakeShared<FOnCacheGetComplete>(MoveTemp(OnComplete));
 	for (const FCacheGetRequest& Request : Requests)
 	{
 		GetCacheRecordAsync(Owner, Request.Name, Request.Key, Request.Policy, Request.UserData,
-			[COOK_STAT(Timer = UsageStats.TimeGet(), ) SharedOnComplete](FCacheGetResponse&& Response, uint64 BytesReceived) mutable
+			[SharedOnComplete](FCacheGetResponse&& Response)
 		{
-			TRACE_COUNTER_ADD(HttpDDC_BytesReceived, BytesReceived);
 			if (Response.Status == EStatus::Ok)
 			{
-				COOK_STAT(Timer.AddHit(BytesReceived););
+				TRACE_COUNTER_INCREMENT(HttpDDC_GetHit);
 			}
 			SharedOnComplete.Get()(MoveTemp(Response));
 		});
@@ -2398,17 +2512,17 @@ void FHttpCacheStore::PutValue(
 	FOnCachePutValueComplete&& OnComplete)
 {
 	TRACE_CPUPROFILER_EVENT_SCOPE(HttpDDC_PutValue);
+	TRACE_COUNTER_ADD(HttpDDC_Put, Requests.Num());
+
 	FRequestBarrier Barrier(Owner);
 	TSharedRef<FOnCachePutValueComplete> SharedOnComplete = MakeShared<FOnCachePutValueComplete>(MoveTemp(OnComplete));
 	for (const FCachePutValueRequest& Request : Requests)
 	{
-		PutCacheValueAsync(Owner, Request.Name, Request.Key, Request.Value, Request.Policy, Request.UserData,
-			[COOK_STAT(Timer = UsageStats.TimePut(),) SharedOnComplete](FCachePutValueResponse&& Response, uint64 BytesSent) mutable
+		PutCacheValueAsync(Owner, Request, [SharedOnComplete](FCachePutValueResponse&& Response)
 		{
-			TRACE_COUNTER_ADD(HttpDDC_BytesSent, BytesSent);
 			if (Response.Status == EStatus::Ok)
 			{
-				COOK_STAT(if (BytesSent) { Timer.AddHit(BytesSent); });
+				TRACE_COUNTER_INCREMENT(HttpDDC_PutHit);
 			}
 			SharedOnComplete.Get()(MoveTemp(Response));
 		});
@@ -2421,38 +2535,21 @@ void FHttpCacheStore::GetValue(
 	FOnCacheGetValueComplete&& OnComplete)
 {
 	TRACE_CPUPROFILER_EVENT_SCOPE(HttpDDC_GetValue);
-	COOK_STAT(double StartTime = FPlatformTime::Seconds());
-	COOK_STAT(bool bIsInGameThread = IsInGameThread());
+	TRACE_COUNTER_ADD(HttpDDC_Get, Requests.Num());
 
-	bool bBatchExistsCandidate = true;
-	for (const FCacheGetValueRequest& Request : Requests)
-	{
-		if (!EnumHasAnyFlags(Request.Policy, ECachePolicy::SkipData))
-		{
-			bBatchExistsCandidate = false;
-			break;
-		}
-	}
-	if (bBatchExistsCandidate)
+	const auto HasSkipData = [](const FCacheGetValueRequest& Request) { return EnumHasAnyFlags(Request.Policy, ECachePolicy::SkipData); };
+	if (Algo::AllOf(Requests, HasSkipData))
 	{
 		TRefCountPtr<FExistsBatchOp> Op = FExistsBatchOp::New(*this, Owner);
-		Op->Exists(Requests, [this, COOK_STAT(StartTime, bIsInGameThread, ) OnComplete = MoveTemp(OnComplete)](FCacheGetValueResponse&& Response)
+		Op->Exists(Requests, [this, OnComplete = MoveTemp(OnComplete)](FCacheGetValueResponse&& Response)
 		{
-			if (Response.Status != EStatus::Ok)
-			{
-				COOK_STAT(UsageStats.GetStats.Accumulate(FCookStats::CallStats::EHitOrMiss::Miss, FCookStats::CallStats::EStatType::Counter, 1l, bIsInGameThread));
-				OnComplete(MoveTemp(Response));
-			}
-			else
+			if (Response.Status == EStatus::Ok)
 			{
+				TRACE_COUNTER_INCREMENT(HttpDDC_GetHit);
 				UE_LOG(LogDerivedDataCache, Verbose, TEXT("%s: Cache hit for %s from '%s'"),
 					*Domain, *WriteToString<96>(Response.Key), *Response.Name);
-				COOK_STAT(UsageStats.GetStats.Accumulate(FCookStats::CallStats::EHitOrMiss::Hit, FCookStats::CallStats::EStatType::Counter, 1l, bIsInGameThread));
-				OnComplete(MoveTemp(Response));
 			}
-
-			COOK_STAT(const int64 CyclesUsed = int64((FPlatformTime::Seconds() - StartTime) / FPlatformTime::GetSecondsPerCycle()));
-			COOK_STAT(UsageStats.GetStats.Accumulate(FCookStats::CallStats::EHitOrMiss::Hit, FCookStats::CallStats::EStatType::Cycles, CyclesUsed, bIsInGameThread));
+			OnComplete(MoveTemp(Response));
 		});
 	}
 	else
@@ -2461,40 +2558,26 @@ void FHttpCacheStore::GetValue(
 		TSharedRef<FOnCacheGetValueComplete> SharedOnComplete = MakeShared<FOnCacheGetValueComplete>(MoveTemp(OnComplete));
 		for (const FCacheGetValueRequest& Request : Requests)
 		{
-			GetCacheValueAsync(Owner, Request.Name, Request.Key, Request.Policy, Request.UserData,
-			[this, COOK_STAT(StartTime, bIsInGameThread,) Policy = Request.Policy, SharedOnComplete](FCacheGetValueResponse&& Response)
+			GetCacheValueAsync(Owner, Request.Name, Request.Key, Request.Policy, ERequestOp::Get, Request.UserData,
+			[this, Policy = Request.Policy, SharedOnComplete](FCacheGetValueResponse&& Response)
 			{
-				const FOnCacheGetValueComplete& OnComplete = SharedOnComplete.Get();
-				check(OnComplete);
-				if (Response.Status != EStatus::Ok)
+				if (Response.Status == EStatus::Ok && !IsValueDataReady(Response.Value, Policy) && !EnumHasAnyFlags(Policy, ECachePolicy::SkipData))
 				{
-					COOK_STAT(UsageStats.GetStats.Accumulate(FCookStats::CallStats::EHitOrMiss::Miss, FCookStats::CallStats::EStatType::Counter, 1l, bIsInGameThread));
-					OnComplete(MoveTemp(Response));
+					Response.Status = EStatus::Error;
+					// With inline fetching, expect we will always have a value we can use.
+					// Even SkipData/Exists can rely on the blob existing if the ref is reported to exist.
+					UE_LOG(LogDerivedDataCache, Log, TEXT("%s: Cache miss due to inlining failure for %s from '%s'"),
+						*Domain, *WriteToString<96>(Response.Key), *Response.Name);
 				}
-				else
+
+				if (Response.Status == EStatus::Ok)
 				{
-					if (!IsValueDataReady(Response.Value, Policy) && !EnumHasAnyFlags(Policy, ECachePolicy::SkipData))
-					{
-						// With inline fetching, expect we will always have a value we can use.  Even SkipData/Exists can rely on the blob existing if the ref is reported to exist.
-						UE_LOG(LogDerivedDataCache, Log, TEXT("%s: Cache miss due to inlining failure for %s from '%s'"),
-									*Domain, *WriteToString<96>(Response.Key), *Response.Name);
-						COOK_STAT(UsageStats.GetStats.Accumulate(FCookStats::CallStats::EHitOrMiss::Miss, FCookStats::CallStats::EStatType::Counter, 1l, bIsInGameThread));
-						OnComplete(MoveTemp(Response));
-					}
-					else
-					{
-						UE_LOG(LogDerivedDataCache, Verbose, TEXT("%s: Cache hit for %s from '%s'"),
-									*Domain, *WriteToString<96>(Response.Key), *Response.Name);
-						uint64 ValueSize = Response.Value.GetData().GetCompressedSize();
-						TRACE_COUNTER_ADD(HttpDDC_BytesReceived, ValueSize);
-						COOK_STAT(UsageStats.GetStats.Accumulate(FCookStats::CallStats::EHitOrMiss::Hit, FCookStats::CallStats::EStatType::Counter, 1l, bIsInGameThread));
-						OnComplete({ Response.Name, Response.Key, Response.Value, Response.UserData, EStatus::Ok });
-						
-						COOK_STAT(UsageStats.GetStats.Accumulate(FCookStats::CallStats::EHitOrMiss::Hit, FCookStats::CallStats::EStatType::Bytes, ValueSize, bIsInGameThread));
-					}
+					TRACE_COUNTER_INCREMENT(HttpDDC_GetHit);
+					UE_LOG(LogDerivedDataCache, Verbose, TEXT("%s: Cache hit for %s from '%s'"),
+						*Domain, *WriteToString<96>(Response.Key), *Response.Name);
 				}
-				COOK_STAT(const int64 CyclesUsed = int64((FPlatformTime::Seconds() - StartTime) / FPlatformTime::GetSecondsPerCycle()));
-				COOK_STAT(UsageStats.GetStats.Accumulate(FCookStats::CallStats::EHitOrMiss::Hit, FCookStats::CallStats::EStatType::Cycles, CyclesUsed, bIsInGameThread));
+
+				SharedOnComplete.Get()(MoveTemp(Response));
 			});
 		}
 	}
@@ -2506,6 +2589,7 @@ void FHttpCacheStore::GetChunks(
 	FOnCacheGetChunkComplete&& OnComplete)
 {
 	TRACE_CPUPROFILER_EVENT_SCOPE(HttpDDC_GetChunks);
+	TRACE_COUNTER_ADD(HttpDDC_GetHit, Requests.Num());
 	// TODO: This is inefficient because Jupiter doesn't allow us to get only part of a compressed blob, so we have to
 	//		 get the whole thing and then decompress only the portion we need.  Furthermore, because there is no propagation
 	//		 between cache stores during chunk requests, the fetched result won't end up in the local store.
@@ -2527,7 +2611,6 @@ void FHttpCacheStore::GetChunks(
 	for (const FCacheGetChunkRequest& Request : SortedRequests)
 	{
 		const bool bExistsOnly = EnumHasAnyFlags(Request.Policy, ECachePolicy::SkipData);
-		COOK_STAT(auto Timer = bExistsOnly ? UsageStats.TimeProbablyExists() : UsageStats.TimeGet());
 		if (!(bHasValue && ValueKey == Request.Key && ValueId == Request.Id) || ValueReader.HasSource() < !bExistsOnly)
 		{
 			ValueStatus = EStatus::Error;
@@ -2540,12 +2623,16 @@ void FHttpCacheStore::GetChunks(
 			if (Request.Id.IsValid())
 			{
 				FRequestOwner BlockingOwner(EPriority::Blocking);
+				bool bOpUsed = false;
+				EStatus OpStatus = EStatus::Ok;
 				TRefCountPtr<FGetRecordOp> Op = FGetRecordOp::New(*this, BlockingOwner, Request.Name);
 				if (!(Record && Record.Get().GetKey() == Request.Key))
 				{
-					Op->GetRecordOnly(Request.Key, Request.Policy, [&Record](FGetRecordOp::FRecordResponse&& Response)
+					bOpUsed = true;
+					Op->GetRecordOnly(Request.Key, Request.Policy, [&Record, &OpStatus](FGetRecordOp::FRecordResponse&& Response)
 					{
 						Record = MoveTemp(Response.Record);
+						OpStatus = Response.Status;
 					});
 					BlockingOwner.Wait();
 				}
@@ -2564,9 +2651,11 @@ void FHttpCacheStore::GetChunks(
 					}
 					else if (bHasValue)
 					{
-						Op->GetValues({ValueWithId}, [&Value](FGetRecordOp::FValueResponse&& Response)
+						bOpUsed = true;
+						Op->GetValues({ValueWithId}, [&Value, &OpStatus](FGetRecordOp::FValueResponse&& Response)
 						{
 							Value = MoveTemp(Response.Value);
+							OpStatus = Response.Status;
 						});
 						BlockingOwner.Wait();
 
@@ -2577,22 +2666,24 @@ void FHttpCacheStore::GetChunks(
 						}
 					}
 				}
+				if (bOpUsed)
+				{
+					Op->RecordStats(OpStatus);
+					TRACE_COUNTER_ADD(HttpDDC_BytesReceived, Op->ReadStats().PhysicalReadSize);
+					TRACE_COUNTER_ADD(HttpDDC_BytesSent, Op->ReadStats().PhysicalWriteSize);
+				}
 			}
 			else
 			{
 				ValueKey = Request.Key;
 
+				FRequestOwner BlockingOwner(EPriority::Blocking);
+				GetCacheValueAsync(BlockingOwner, Request.Name, Request.Key, Request.Policy, ERequestOp::GetChunk, 0, [&bHasValue, &Value](FCacheGetValueResponse&& Response)
 				{
-					FRequestOwner BlockingOwner(EPriority::Blocking);
-					bool bSucceeded = false;
-					GetCacheValueAsync(BlockingOwner, Request.Name, Request.Key, Request.Policy, 0, [&bSucceeded, &Value](FCacheGetValueResponse&& Response)
-					{
-						Value = MoveTemp(Response.Value);
-						bSucceeded = Response.Status == EStatus::Ok;
-					});
-					BlockingOwner.Wait();
-					bHasValue = bSucceeded;
-				}
+					bHasValue = Response.Status == EStatus::Ok;
+					Value = MoveTemp(Response.Value);
+				});
+				BlockingOwner.Wait();
 
 				if (bHasValue && IsValueDataReady(Value, Request.Policy))
 				{
@@ -2607,13 +2698,16 @@ void FHttpCacheStore::GetChunks(
 			const uint64 RawSize = FMath::Min(Value.GetRawSize() - RawOffset, Request.RawSize);
 			UE_LOG(LogDerivedDataCache, Verbose, TEXT("%s: Cache hit for %s from '%s'"),
 				*Domain, *WriteToString<96>(Request.Key, '/', Request.Id), *Request.Name);
-			COOK_STAT(Timer.AddHit(!bExistsOnly ? RawSize : 0));
 			FSharedBuffer Buffer;
 			if (!bExistsOnly)
 			{
 				Buffer = ValueReader.Decompress(RawOffset, RawSize);
 			}
 			const EStatus ChunkStatus = bExistsOnly || Buffer.GetSize() == RawSize ? EStatus::Ok : EStatus::Error;
+			if (ChunkStatus == EStatus::Ok)
+			{
+				TRACE_COUNTER_INCREMENT(HttpDDC_GetHit);
+			}
 			OnComplete({Request.Name, Request.Key, Request.Id, Request.RawOffset,
 				RawSize, Value.GetRawHash(), MoveTemp(Buffer), Request.UserData, ChunkStatus});
 			continue;
