DDC: Disable stats for cache stores that should not be visible

#rnx

diff --git a/Engine/Source/Developer/DerivedDataCache/Private/DerivedDataBackends.cpp b/Engine/Source/Developer/DerivedDataCache/Private/DerivedDataBackends.cpp
--- a/Engine/Source/Developer/DerivedDataCache/Private/DerivedDataBackends.cpp
+++ b/Engine/Source/Developer/DerivedDataCache/Private/DerivedDataBackends.cpp
@@ -604,7 +604,7 @@ public:
 		else
 		{
 			// This is unconditionally added to the hierarchy and will be deleted by the hierarchy.
-			CreateMemoryCacheStore(OutCache, TEXT("Memory"), TEXT("-ReadOnly"), this);
+			CreateMemoryCacheStore(OutCache, TEXT("Memory"), TEXT("-ReadOnly -NoStats"), this);
 			MemoryCache = OutCache;
 		}
 	}
diff --git a/Engine/Source/Developer/DerivedDataCache/Private/DerivedDataCacheStoreHierarchy.cpp b/Engine/Source/Developer/DerivedDataCache/Private/DerivedDataCacheStoreHierarchy.cpp
--- a/Engine/Source/Developer/DerivedDataCache/Private/DerivedDataCacheStoreHierarchy.cpp
+++ b/Engine/Source/Developer/DerivedDataCache/Private/DerivedDataCacheStoreHierarchy.cpp
@@ -1520,16 +1520,9 @@ void FCacheStoreHierarchy::LegacyStats(FDerivedDataCacheStatsNode& OutNode)
 	OutNode.Children.Reserve(Nodes.Num());
 	for (const FCacheStoreNode& Node : Nodes)
 	{
-		if (Node.CacheStats.IsEmpty())
+		for (const FCacheStoreStats* Stats : Node.CacheStats)
 		{
-			Node.Cache->LegacyStats(OutNode.Children.Add_GetRef(MakeShared<FDerivedDataCacheStatsNode>()).Get());
-		}
-		else
-		{
-			for (const FCacheStoreStats* Stats : Node.CacheStats)
-			{
-				ConvertToLegacyStats(OutNode.Children.Add_GetRef(MakeShared<FDerivedDataCacheStatsNode>()).Get(), *Stats);
-			}
+			ConvertToLegacyStats(OutNode.Children.Add_GetRef(MakeShared<FDerivedDataCacheStatsNode>()).Get(), *Stats);
 		}
 	}
 }
diff --git a/Engine/Source/Developer/DerivedDataCache/Private/FileSystemCacheStore.cpp b/Engine/Source/Developer/DerivedDataCache/Private/FileSystemCacheStore.cpp
--- a/Engine/Source/Developer/DerivedDataCache/Private/FileSystemCacheStore.cpp
+++ b/Engine/Source/Developer/DerivedDataCache/Private/FileSystemCacheStore.cpp
@@ -1219,7 +1219,10 @@ FFileSystemCacheStore::FFileSystemCacheStore(
 	Flags |= SpeedClass == EBackendSpeedClass::Local ? ECacheStoreFlags::Local : ECacheStoreFlags::Remote;
 
 	StoreOwner.Add(this, Flags);
-	StoreStats = StoreOwner.CreateStats(this, Flags, TEXTVIEW("File System"), Params.CacheName, CachePath);
+	if (!Params.bDeleteOnly)
+	{
+		StoreStats = StoreOwner.CreateStats(this, Flags, TEXTVIEW("File System"), Params.CacheName, CachePath);
+	}
 
 	UpdateStatus();
 
@@ -1230,7 +1233,10 @@ FFileSystemCacheStore::FFileSystemCacheStore(
 
 FFileSystemCacheStore::~FFileSystemCacheStore()
 {
-	StoreOwner.DestroyStats(StoreStats);
+	if (StoreStats)
+	{
+		StoreOwner.DestroyStats(StoreStats);
+	}
 
 	TUniqueLock Lock(ActiveStoresMutex);
 	ActiveStores.Remove(this);
@@ -2807,13 +2813,16 @@ bool FFileSystemCacheStore::IsDeactivatedForPerformance()
 
 void FFileSystemCacheStore::UpdateStatus()
 {
-	if (bDeactivedForPerformance.load(std::memory_order_relaxed))
-	{
-		StoreStats->SetStatus(ECacheStoreStatusCode::Warning, NSLOCTEXT("DerivedDataCache", "DeactivatedForPerformance", "Deactivated for performance"));
-	}
-	else
+	if (StoreStats)
 	{
-		StoreStats->SetStatus(ECacheStoreStatusCode::None, {});
+		if (bDeactivedForPerformance.load(std::memory_order_relaxed))
+		{
+			StoreStats->SetStatus(ECacheStoreStatusCode::Warning, NSLOCTEXT("DerivedDataCache", "DeactivatedForPerformance", "Deactivated for performance"));
+		}
+		else
+		{
+			StoreStats->SetStatus(ECacheStoreStatusCode::None, {});
+		}
 	}
 }
 
diff --git a/Engine/Source/Developer/DerivedDataCache/Private/MemoryCacheStore.cpp b/Engine/Source/Developer/DerivedDataCache/Private/MemoryCacheStore.cpp
--- a/Engine/Source/Developer/DerivedDataCache/Private/MemoryCacheStore.cpp
+++ b/Engine/Source/Developer/DerivedDataCache/Private/MemoryCacheStore.cpp
@@ -149,7 +149,10 @@ FMemoryCacheStore::FMemoryCacheStore(
 			Flags |= ECacheStoreFlags::Store;
 		}
 		StoreOwner->Add(this, Flags);
-		StoreStats = StoreOwner->CreateStats(this, Flags, TEXTVIEW("Memory"), Name);
+		if (!FParse::Param(InConfig, TEXT("NoStats")))
+		{
+			StoreStats = StoreOwner->CreateStats(this, Flags, TEXTVIEW("Memory"), Name);
+		}
 	}
 }
 
@@ -157,7 +160,7 @@ FMemoryCacheStore::~FMemoryCacheStore()
 {
 	bShuttingDown = true;
 	Disable();
-	if (StoreOwner)
+	if (StoreStats)
 	{
 		StoreOwner->DestroyStats(StoreStats);
 	}
diff --git a/Engine/Source/Editor/DerivedDataEditor/Private/SDerivedDataDialogs.cpp b/Engine/Source/Editor/DerivedDataEditor/Private/SDerivedDataDialogs.cpp
--- a/Engine/Source/Editor/DerivedDataEditor/Private/SDerivedDataDialogs.cpp
+++ b/Engine/Source/Editor/DerivedDataEditor/Private/SDerivedDataDialogs.cpp
@@ -646,11 +646,6 @@ TSharedRef<SWidget> SDerivedDataCacheStatisticsDialog::GetGridPanel()
 
 	for (TSharedRef<const FDerivedDataCacheStatsNode> Node : LeafUsageStats)
 	{
-		if (Node->GetCacheType().Equals(TEXT("Memory")))
-		{
-			continue;
-		}
-
 		FDerivedDataCacheUsageStats Stats;
 
 		for (const auto& KVP : Node->UsageStats)
