DDC: Added request name to FCacheStoreRequestStats

#rnx

diff --git a/Engine/Source/Developer/DerivedDataCache/Private/DerivedDataCacheStore.h b/Engine/Source/Developer/DerivedDataCache/Private/DerivedDataCacheStore.h
--- a/Engine/Source/Developer/DerivedDataCache/Private/DerivedDataCacheStore.h
+++ b/Engine/Source/Developer/DerivedDataCache/Private/DerivedDataCacheStore.h
@@ -176,6 +176,8 @@ enum class ECacheStoreStatusCode : uint8
  */
 struct FCacheStoreRequestStats final
 {
+	/** A name to identify this request for logging and profiling. An object path is typically sufficient. */
+	FSharedString Name;
 	/** Bucket that contains the data being accessed by this request. */
 	FCacheBucket Bucket;
 	/** Size of the raw data and metadata that was read from the cache store. */
diff --git a/Engine/Source/Developer/DerivedDataCache/Private/FileSystemCacheStore.cpp b/Engine/Source/Developer/DerivedDataCache/Private/FileSystemCacheStore.cpp
--- a/Engine/Source/Developer/DerivedDataCache/Private/FileSystemCacheStore.cpp
+++ b/Engine/Source/Developer/DerivedDataCache/Private/FileSystemCacheStore.cpp
@@ -1503,6 +1503,7 @@ void FFileSystemCacheStore::Put(
 	{
 		bool bOk;
 		FRequestStats RequestStats;
+		RequestStats.Name = Request.Name;
 		RequestStats.Bucket = Request.Record.GetKey().Bucket;
 		RequestStats.Type = ERequestType::Record;
 		RequestStats.Op = ERequestOp::Put;
@@ -1538,6 +1539,7 @@ void FFileSystemCacheStore::Get(
 		EStatus Status = EStatus::Error;
 		FOptionalCacheRecord Record;
 		FRequestStats RequestStats;
+		RequestStats.Name = Request.Name;
 		RequestStats.Bucket = Request.Key.Bucket;
 		RequestStats.Type = ERequestType::Record;
 		RequestStats.Op = ERequestOp::Get;
@@ -1574,6 +1576,7 @@ void FFileSystemCacheStore::PutValue(
 	{
 		bool bOk;
 		FRequestStats RequestStats;
+		RequestStats.Name = Request.Name;
 		RequestStats.Bucket = Request.Key.Bucket;
 		RequestStats.Type = ERequestType::Value;
 		RequestStats.Op = ERequestOp::Put;
@@ -1608,6 +1611,7 @@ void FFileSystemCacheStore::GetValue(
 		bool bOk;
 		FValue Value;
 		FRequestStats RequestStats;
+		RequestStats.Name = Request.Name;
 		RequestStats.Bucket = Request.Key.Bucket;
 		RequestStats.Type = ERequestType::Value;
 		RequestStats.Op = ERequestOp::Get;
@@ -1653,6 +1657,7 @@ void FFileSystemCacheStore::GetChunks(
 		FSharedBuffer Buffer;
 		uint64 RawSize = 0;
 		FRequestStats RequestStats;
+		RequestStats.Name = Request.Name;
 		RequestStats.Bucket = Request.Key.Bucket;
 		RequestStats.Type = Request.Id.IsNull() ? ERequestType::Value : ERequestType::Record;
 		RequestStats.Op = ERequestOp::GetChunk;
diff --git a/Engine/Source/Developer/DerivedDataCache/Private/HttpCacheStore.cpp b/Engine/Source/Developer/DerivedDataCache/Private/HttpCacheStore.cpp
--- a/Engine/Source/Developer/DerivedDataCache/Private/HttpCacheStore.cpp
+++ b/Engine/Source/Developer/DerivedDataCache/Private/HttpCacheStore.cpp
@@ -751,6 +751,7 @@ FHttpCacheStore::FPutPackageOp::FPutPackageOp(FHttpCacheStore& InCacheStore, IRe
 	, Owner(InOwner)
 	, Name(InName)
 {
+	RequestStats.Name = Name;
 }
 
 void FHttpCacheStore::FPutPackageOp::Put(const FCacheKey& InKey, FCbPackage&& Package, FOnPackageComplete&& OnComplete)
@@ -1058,6 +1059,7 @@ FHttpCacheStore::FGetRecordOp::FGetRecordOp(FHttpCacheStore& InCacheStore, IRequ
 	, Owner(InOwner)
 	, Name(InName)
 {
+	RequestStats.Name = Name;
 }
 
 void FHttpCacheStore::FGetRecordOp::GetRecordOnly(const FCacheKey& InKey, const ECachePolicy RecordPolicy, FOnRecordComplete&& InOnComplete)
@@ -1524,6 +1526,7 @@ FHttpCacheStore::FGetValueOp::FGetValueOp(FHttpCacheStore& InCacheStore, IReques
 	, Owner(InOwner)
 	, Name(InName)
 {
+	RequestStats.Name = Name;
 }
 
 void FHttpCacheStore::FGetValueOp::Get(const FCacheKey& InKey, ECachePolicy InPolicy, FOnComplete&& InOnComplete)
diff --git a/Engine/Source/Developer/DerivedDataCache/Private/MemoryCacheStore.cpp b/Engine/Source/Developer/DerivedDataCache/Private/MemoryCacheStore.cpp
--- a/Engine/Source/Developer/DerivedDataCache/Private/MemoryCacheStore.cpp
+++ b/Engine/Source/Developer/DerivedDataCache/Private/MemoryCacheStore.cpp
@@ -210,6 +210,7 @@ void FMemoryCacheStore::Put(
 		const TConstArrayView<FValueWithId> Values = Record.GetValues();
 
 		FRequestStats RequestStats;
+		RequestStats.Name = Request.Name;
 		RequestStats.Bucket = Key.Bucket;
 		RequestStats.Type = ERequestType::Record;
 		RequestStats.Op = ERequestOp::Put;
@@ -396,6 +397,7 @@ void FMemoryCacheStore::Get(
 		EStatus Status = EStatus::Error;
 
 		FRequestStats RequestStats;
+		RequestStats.Name = Request.Name;
 		RequestStats.Bucket = Key.Bucket;
 		RequestStats.Type = ERequestType::Record;
 		RequestStats.Op = ERequestOp::Get;
@@ -506,6 +508,7 @@ void FMemoryCacheStore::PutValue(
 		const FValue& Value = Request.Value;
 
 		FRequestStats RequestStats;
+		RequestStats.Name = Request.Name;
 		RequestStats.Bucket = Key.Bucket;
 		RequestStats.Type = ERequestType::Value;
 		RequestStats.Op = ERequestOp::Put;
@@ -594,6 +597,7 @@ void FMemoryCacheStore::GetValue(
 		EStatus Status = EStatus::Error;
 
 		FRequestStats RequestStats;
+		RequestStats.Name = Request.Name;
 		RequestStats.Bucket = Key.Bucket;
 		RequestStats.Type = ERequestType::Value;
 		RequestStats.Op = ERequestOp::Get;
@@ -659,6 +663,7 @@ void FMemoryCacheStore::GetChunks(
 	for (const FCacheGetChunkRequest& Request : Requests)
 	{
 		FRequestStats RequestStats;
+		RequestStats.Name = Request.Name;
 		RequestStats.Bucket = Request.Key.Bucket;
 		RequestStats.Type = Request.Id.IsNull() ? ERequestType::Value : ERequestType::Record;
 		RequestStats.Op = ERequestOp::GetChunk;
diff --git a/Engine/Source/Developer/DerivedDataCache/Private/PakFileCacheStore.cpp b/Engine/Source/Developer/DerivedDataCache/Private/PakFileCacheStore.cpp
--- a/Engine/Source/Developer/DerivedDataCache/Private/PakFileCacheStore.cpp
+++ b/Engine/Source/Developer/DerivedDataCache/Private/PakFileCacheStore.cpp
@@ -647,6 +647,7 @@ void FPakFileCacheStore::Put(
 	{
 		bool bOk;
 		FRequestStats RequestStats;
+		RequestStats.Name = Request.Name;
 		RequestStats.Bucket = Request.Record.GetKey().Bucket;
 		RequestStats.Type = ERequestType::Record;
 		RequestStats.Op = ERequestOp::Put;
@@ -680,6 +681,7 @@ void FPakFileCacheStore::Get(
 		EStatus Status = EStatus::Error;
 		FOptionalCacheRecord Record;
 		FRequestStats RequestStats;
+		RequestStats.Name = Request.Name;
 		RequestStats.Bucket = Request.Key.Bucket;
 		RequestStats.Type = ERequestType::Record;
 		RequestStats.Op = ERequestOp::Get;
@@ -715,6 +717,7 @@ void FPakFileCacheStore::PutValue(
 	{
 		bool bOk;
 		FRequestStats RequestStats;
+		RequestStats.Name = Request.Name;
 		RequestStats.Bucket = Request.Key.Bucket;
 		RequestStats.Type = ERequestType::Value;
 		RequestStats.Op = ERequestOp::Put;
@@ -747,6 +750,7 @@ void FPakFileCacheStore::GetValue(
 		bool bOk;
 		FValue Value;
 		FRequestStats RequestStats;
+		RequestStats.Name = Request.Name;
 		RequestStats.Bucket = Request.Key.Bucket;
 		RequestStats.Type = ERequestType::Value;
 		RequestStats.Op = ERequestOp::Get;
@@ -791,6 +795,7 @@ void FPakFileCacheStore::GetChunks(
 		FSharedBuffer Buffer;
 		uint64 RawSize = 0;
 		FRequestStats RequestStats;
+		RequestStats.Name = Request.Name;
 		RequestStats.Bucket = Request.Key.Bucket;
 		RequestStats.Type = Request.Id.IsNull() ? ERequestType::Value : ERequestType::Record;
 		RequestStats.Op = ECacheStoreRequestOp::GetChunk;
diff --git a/Engine/Source/Developer/DerivedDataCache/Private/S3CacheStore.cpp b/Engine/Source/Developer/DerivedDataCache/Private/S3CacheStore.cpp
--- a/Engine/Source/Developer/DerivedDataCache/Private/S3CacheStore.cpp
+++ b/Engine/Source/Developer/DerivedDataCache/Private/S3CacheStore.cpp
@@ -1692,6 +1692,7 @@ void FS3CacheStore::Get(
 		EStatus Status = EStatus::Error;
 		FOptionalCacheRecord Record;
 		FRequestStats RequestStats;
+		RequestStats.Name = Request.Name;
 		RequestStats.Bucket = Request.Key.Bucket;
 		RequestStats.Type = ERequestType::Record;
 		RequestStats.Op = ERequestOp::Get;
@@ -1735,6 +1736,7 @@ void FS3CacheStore::GetValue(
 		bool bOk;
 		FValue Value;
 		FRequestStats RequestStats;
+		RequestStats.Name = Request.Name;
 		RequestStats.Bucket = Request.Key.Bucket;
 		RequestStats.Type = ERequestType::Value;
 		RequestStats.Op = ERequestOp::Get;
@@ -1778,6 +1780,7 @@ void FS3CacheStore::GetChunks(
 		FSharedBuffer Buffer;
 		uint64 RawSize = 0;
 		FRequestStats RequestStats;
+		RequestStats.Name = Request.Name;
 		RequestStats.Bucket = Request.Key.Bucket;
 		RequestStats.Type = Request.Id.IsNull() ? ERequestType::Value : ERequestType::Record;
 		RequestStats.Op = ERequestOp::GetChunk;
diff --git a/Engine/Source/Developer/DerivedDataCache/Private/ZenCacheStore.cpp b/Engine/Source/Developer/DerivedDataCache/Private/ZenCacheStore.cpp
--- a/Engine/Source/Developer/DerivedDataCache/Private/ZenCacheStore.cpp
+++ b/Engine/Source/Developer/DerivedDataCache/Private/ZenCacheStore.cpp
@@ -227,6 +227,7 @@ void FZenCacheStore::StartRequests(
 	for (const auto& Request : In)
 	{
 		auto& RequestWithStats = Out[Out.Emplace(Request)];
+		RequestWithStats.Stats.Name = Request.Name;
 		RequestWithStats.Stats.Bucket = BucketAccessor(Request);
 		RequestWithStats.Stats.Type = TypeAccessor(Request);
 		RequestWithStats.Stats.Op = Op;
