DDC: Gather bucket stats from the new ICacheStats

#rnx

diff --git a/Engine/Source/Developer/DerivedDataCache/Private/DerivedDataBackendInterface.h b/Engine/Source/Developer/DerivedDataCache/Private/DerivedDataBackendInterface.h
--- a/Engine/Source/Developer/DerivedDataCache/Private/DerivedDataBackendInterface.h
+++ b/Engine/Source/Developer/DerivedDataCache/Private/DerivedDataBackendInterface.h
@@ -133,6 +133,8 @@ public:
 	 *  Gather the usage of the DDC hierarchically.
 	 */
 	virtual TSharedRef<FDerivedDataCacheStatsNode> GatherUsageStats() const = 0;
+
+	virtual void GatherResourceStats(TArray<FDerivedDataCacheResourceStat>& DDCResourceStats) const = 0;
 };
 
 /** Lexical conversions from and to enums */
diff --git a/Engine/Source/Developer/DerivedDataCache/Private/DerivedDataBackends.cpp b/Engine/Source/Developer/DerivedDataCache/Private/DerivedDataBackends.cpp
--- a/Engine/Source/Developer/DerivedDataCache/Private/DerivedDataBackends.cpp
+++ b/Engine/Source/Developer/DerivedDataCache/Private/DerivedDataBackends.cpp
@@ -801,6 +801,14 @@ public:
 		return Stats;
 	}
 
+	virtual void GatherResourceStats(TArray<FDerivedDataCacheResourceStat>& DDCResourceStats) const override
+	{
+		if (Hierarchy)
+		{
+			Hierarchy->LegacyResourceStats(DDCResourceStats);
+		}
+	}
+
 private:
 	void Add(ILegacyCacheStore* CacheStore, ECacheStoreFlags Flags) final
 	{
@@ -882,6 +890,12 @@ private:
 		Hierarchy->DestroyStats(Stats);
 	}
 
+	void LegacyResourceStats(TArray<FDerivedDataCacheResourceStat>& OutStats) const final
+	{
+		check(Hierarchy);
+		Hierarchy->LegacyResourceStats(OutStats);
+	}
+
 	/** MountPak console command handler. */
 	void UnmountPakCommandHandler(const TArray<FString>& Args)
 	{
diff --git a/Engine/Source/Developer/DerivedDataCache/Private/DerivedDataCache.cpp b/Engine/Source/Developer/DerivedDataCache/Private/DerivedDataCache.cpp
--- a/Engine/Source/Developer/DerivedDataCache/Private/DerivedDataCache.cpp
+++ b/Engine/Source/Developer/DerivedDataCache/Private/DerivedDataCache.cpp
@@ -213,7 +213,6 @@ namespace UE::DerivedData::CookStats
 }
 #endif
 
-void GatherDerivedDataCacheResourceStats(TArray<FDerivedDataCacheResourceStat>& DDCResourceStats);
 void GatherDerivedDataCacheSummaryStats(FDerivedDataCacheSummaryStats& DDCSummaryStats);
 
 /** Whether we want to verify the DDC (pass in -VerifyDDC on the command line)*/
@@ -934,7 +933,7 @@ public:
 
 	virtual void GatherResourceStats(TArray<FDerivedDataCacheResourceStat>& DDCResourceStats) const override
 	{
-		GatherDerivedDataCacheResourceStats(DDCResourceStats);
+		Backend->GatherResourceStats(DDCResourceStats);
 	}
 
 	virtual void GatherSummaryStats(FDerivedDataCacheSummaryStats& DDCSummaryStats) const override
@@ -949,7 +948,7 @@ public:
 		// Gather the latest resource stats
 		TArray<FDerivedDataCacheResourceStat> ResourceStats;
 
-		GatherDerivedDataCacheResourceStats(ResourceStats);
+		GatherResourceStats(ResourceStats);
 
 		FDerivedDataCacheResourceStat ResourceStatsTotal(TEXT("Total"));
 
diff --git a/Engine/Source/Developer/DerivedDataCache/Private/DerivedDataCacheKey.cpp b/Engine/Source/Developer/DerivedDataCache/Private/DerivedDataCacheKey.cpp
--- a/Engine/Source/Developer/DerivedDataCache/Private/DerivedDataCacheKey.cpp
+++ b/Engine/Source/Developer/DerivedDataCache/Private/DerivedDataCacheKey.cpp
@@ -2,6 +2,7 @@
 
 #include "DerivedDataCacheKey.h"
 
+#include "Containers/Map.h"
 #include "Containers/Set.h"
 #include "Containers/StringConv.h"
 #include "DerivedDataCachePrivate.h"
@@ -94,9 +95,13 @@ public:
 	template <typename CharType>
 	inline FCacheBucket FindOrAdd(TStringView<CharType> Name);
 
+	inline void GetDisplayName(FCacheBucket Bucket, FStringBuilderBase& OutDisplayName);
+	inline void SetDisplayName(FCacheBucket Bucket, FStringView DisplayName);
+
 private:
 	FRWLock Lock;
 	TSet<FCacheBucketOwner, FCacheBucketOwnerKeyFuncs> Buckets;
+	TMap<FCacheBucket, FString> DisplayNames;
 };
 
 template <typename CharType>
@@ -121,6 +126,35 @@ inline FCacheBucket FCacheBuckets::FindOrAdd(const TStringView<CharType> Name)
 	return Buckets.FindOrAddByHash(Hash, MoveTemp(LocalBucket));
 }
 
+void FCacheBuckets::GetDisplayName(FCacheBucket Bucket, FStringBuilderBase& OutDisplayName)
+{
+	if (FReadScopeLock ReadLock(Lock); const FString* DisplayName = DisplayNames.Find(Bucket))
+	{
+		OutDisplayName << *DisplayName;
+		return;
+	}
+	FAnsiStringView BucketName = Bucket.ToString();
+	if (BucketName.StartsWith(ANSITEXTVIEW("Legacy")))
+	{
+		BucketName.RightChopInline(TEXTVIEW("Legacy").Len());
+	}
+	OutDisplayName << BucketName;
+}
+
+void FCacheBuckets::SetDisplayName(FCacheBucket Bucket, FStringView DisplayName)
+{
+	if (FReadScopeLock ReadLock(Lock); const FString* ExistingDisplayName = DisplayNames.Find(Bucket))
+	{
+		if (*ExistingDisplayName == DisplayName)
+		{
+			return;
+		}
+	}
+
+	FWriteScopeLock WriteLock(Lock);
+	DisplayNames.Emplace(Bucket, DisplayName);
+}
+
 static FCacheBuckets& GetCacheBuckets()
 {
 	static FCacheBuckets Buckets;
@@ -142,6 +176,29 @@ FCacheBucket::FCacheBucket(FWideStringView InName)
 {
 }
 
+FCacheBucket::FCacheBucket(FUtf8StringView Name, FStringView DisplayName)
+	: FCacheBucket(Name)
+{
+	if (!DisplayName.IsEmpty())
+	{
+		Private::GetCacheBuckets().SetDisplayName(*this, DisplayName);
+	}
+}
+
+FCacheBucket::FCacheBucket(FWideStringView Name, FStringView DisplayName)
+	: FCacheBucket(Name)
+{
+	if (!DisplayName.IsEmpty())
+	{
+		Private::GetCacheBuckets().SetDisplayName(*this, DisplayName);
+	}
+}
+
+void FCacheBucket::ToDisplayName(FStringBuilderBase& OutDisplayName) const
+{
+	Private::GetCacheBuckets().GetDisplayName(*this, OutDisplayName);
+}
+
 FCbWriter& operator<<(FCbWriter& Writer, const FCacheBucket Bucket)
 {
 	Writer.AddString(Bucket.ToString());
diff --git a/Engine/Source/Developer/DerivedDataCache/Private/DerivedDataCacheStats.cpp b/Engine/Source/Developer/DerivedDataCache/Private/DerivedDataCacheStats.cpp
new file mode 100644
--- /dev/null
+++ b/Engine/Source/Developer/DerivedDataCache/Private/DerivedDataCacheStats.cpp
@@ -0,0 +1,315 @@
+// Copyright Epic Games, Inc. All Rights Reserved.
+
+#include "DerivedDataCacheStats.h"
+
+#include "Algo/BinarySearch.h"
+
+namespace UE::DerivedData
+{
+
+void FRequestCounter::AddRequest(const ERequestType Type, const ERequestOp Op, const EStatus Status)
+{
+	switch (Type)
+	{
+	case ERequestType::Record:
+		switch (Op)
+		{
+		case ERequestOp::Put:
+			++(Status == EStatus::Ok ? PutRecord.Hits : PutRecord.Misses);
+			break;
+		case ERequestOp::Get:
+			++(Status == EStatus::Ok ? GetRecord.Hits : GetRecord.Misses);
+			break;
+		case ERequestOp::GetChunk:
+			++(Status == EStatus::Ok ? GetRecordChunk.Hits : GetRecordChunk.Misses);
+			break;
+		}
+		break;
+	case ERequestType::Value:
+		switch (Op)
+		{
+		case ERequestOp::Put:
+			++(Status == EStatus::Ok ? PutValue.Hits : PutValue.Misses);
+			break;
+		case ERequestOp::Get:
+			++(Status == EStatus::Ok ? GetValue.Hits : GetValue.Misses);
+			break;
+		case ERequestOp::GetChunk:
+			++(Status == EStatus::Ok ? GetValueChunk.Hits : GetValueChunk.Misses);
+			break;
+		}
+		break;
+	}
+}
+
+///////////////////////////////////////////////////////////////////////////////////////////////////////////////////////
+
+void FTimeAveragedStat::Add(FMonotonicTimePoint StartTime, FMonotonicTimePoint EndTime, double Value)
+{
+	checkf(StartTime <= EndTime, TEXT("StartTime %.3f is later than EndTime %.3f."), StartTime.ToSeconds(), EndTime.ToSeconds());
+
+	TUniqueLock Lock(Mutex);
+
+	if (EndTime <= LastTime - Period)
+	{
+		return;
+	}
+	else
+	{
+		const int32 Index = Algo::LowerBoundBy(EndValues, EndTime, &FValue::Time);
+		EndValues.Insert({EndTime, Value}, Index);
+	}
+
+	if (StartTime < LastTime)
+	{
+		AccumulatedValue += Value;
+		++AccumulatedValueCount;
+	}
+	else
+	{
+		const int32 Index = Algo::LowerBoundBy(StartValues, StartTime, &FValue::Time);
+		StartValues.Insert({StartTime, Value}, Index);
+	}
+
+	const int32 Count = ActiveRanges.Num();
+	const int32 StartIndex = Algo::LowerBoundBy(ActiveRanges, StartTime, &FRange::EndTime);
+	if (StartIndex == Count)
+	{
+		ActiveRanges.Add({StartTime, EndTime});
+	}
+	else
+	{
+		FRange& StartRange = ActiveRanges[StartIndex];
+		if (EndTime < StartRange.StartTime)
+		{
+			ActiveRanges.Insert({StartTime, EndTime}, StartIndex);
+		}
+		else
+		{
+			if (StartTime < StartRange.StartTime)
+			{
+				StartRange.StartTime = StartTime;
+			}
+
+			const int32 EndIndex = Algo::LowerBoundBy(ActiveRanges, EndTime, &FRange::StartTime);
+			StartRange.EndTime = FMath::Max(EndTime, ActiveRanges[EndIndex - 1].EndTime);
+
+			ActiveRanges.RemoveAt(StartIndex + 1, EndIndex - StartIndex - 1, /*bAllowShrinking*/ false);
+		}
+	}
+
+	// Avoid unbounded growth by updating when behind by more than one period.
+	if (LastTime + Period < EndTime)
+	{
+		Update(EndTime);
+	}
+}
+
+double FTimeAveragedStat::GetRate(FMonotonicTimePoint Time)
+{
+	TUniqueLock Lock(Mutex);
+	if (LastTime < Time)
+	{
+		Update(Time);
+	}
+	return AverageRate;
+}
+
+double FTimeAveragedStat::GetValue(FMonotonicTimePoint Time)
+{
+	TUniqueLock Lock(Mutex);
+	if (LastTime < Time)
+	{
+		Update(Time);
+	}
+	return AverageValue;
+}
+
+void FTimeAveragedStat::Update(FMonotonicTimePoint Time)
+{
+	LastTime = Time;
+
+	const FMonotonicTimePoint StartTime = Time - Period;
+	const FMonotonicTimePoint EndTime = Time;
+
+	// NOTE: Values are not being prorated when entering or exiting their time range.
+	//       While this is technically inaccurate, the error will likely be small due
+	//       to the use case being requests of one second averaged over one minute.
+
+	// Add values entering the current period.
+	int32 StartCount = 0;
+	for (const FValue& Value : StartValues)
+	{
+		if (EndTime < Value.Time)
+		{
+			break;
+		}
+		AccumulatedValue += Value.Value;
+		++AccumulatedValueCount;
+		++StartCount;
+	}
+	StartValues.RemoveAt(0, StartCount, /*bAllowShrinking*/ false);
+
+	// Remove values exiting the current period.
+	int32 EndCount = 0;
+	for (const FValue& Value : EndValues)
+	{
+		if (StartTime <= Value.Time)
+		{
+			break;
+		}
+		AccumulatedValue -= Value.Value;
+		--AccumulatedValueCount;
+		++EndCount;
+	}
+	EndValues.RemoveAt(0, EndCount, /*bAllowShrinking*/ false);
+
+	// Remove ranges before the current period.
+	int32 RangeCount = 0;
+	for (const FRange& Range : ActiveRanges)
+	{
+		if (StartTime < Range.EndTime)
+		{
+			break;
+		}
+		++RangeCount;
+	}
+	ActiveRanges.RemoveAt(0, RangeCount, /*bAllowShrinking*/ false);
+
+	// Accumulate active time in the current period.
+	FMonotonicTimeSpan ActiveTime;
+	for (const FRange& Range : ActiveRanges)
+	{
+		const FMonotonicTimePoint RangeStart = FMath::Max(Range.StartTime, StartTime);
+		const FMonotonicTimePoint RangeEnd = FMath::Min(Range.EndTime, EndTime);
+		ActiveTime += RangeEnd - RangeStart;
+	}
+
+	if (AccumulatedValueCount && !ActiveTime.IsZero())
+	{
+		AverageRate = AccumulatedValue / ActiveTime.ToSeconds();
+		AverageValue = AccumulatedValue / AccumulatedValueCount;
+	}
+}
+
+///////////////////////////////////////////////////////////////////////////////////////////////////////////////////////
+
+FCacheStoreStats::FCacheStoreStats(FCacheStats& InCacheStats, ECacheStoreFlags InFlags, FStringView InType, FStringView InName, FStringView InPath)
+	: Type(InType)
+	, Name(InName)
+	, Path(InPath)
+	, Flags(InFlags)
+	, CacheStats(InCacheStats)
+{
+	const FMonotonicTimeSpan Period = FMonotonicTimeSpan::FromSeconds(60.0);
+	AveragePhysicalReadSize.SetPeriod(Period);
+	AveragePhysicalWriteSize.SetPeriod(Period);
+}
+
+void FCacheStoreStats::SetFlags(ECacheStoreFlags InFlags)
+{
+	TUniqueLock Lock(Mutex);
+	Flags = InFlags;
+}
+
+void FCacheStoreStats::SetStatus(ECacheStoreStatusCode InStatusCode, const FText& InStatus)
+{
+	TUniqueLock Lock(Mutex);
+	Status = InStatus;
+	StatusCode = InStatusCode;
+}
+
+void FCacheStoreStats::SetAttribute(FStringView Key, FStringView Value)
+{
+	TUniqueLock Lock(Mutex);
+	Attributes.Emplace(Key, Value);
+}
+
+void FCacheStoreStats::AddRequest(const FCacheStoreRequestStats& Stats)
+{
+	checkf(Stats.Bucket.IsValid(), TEXT("Stats for request '%s' did not set a bucket."), *Stats.Name);
+
+#if ENABLE_COOK_STATS
+	using FCallStats = FCookStats::CallStats;
+	using EHitOrMiss = FCallStats::EHitOrMiss;
+	using EStatType = FCallStats::EStatType;
+
+	const EHitOrMiss HitOrMiss = Stats.Status == EStatus::Ok ? EHitOrMiss::Hit : EHitOrMiss::Miss;
+	const bool bIsGet = Stats.Op != ECacheStoreRequestOp::Put;
+	const bool bIsInGameThread = IsInGameThread();
+#endif
+
+	{
+		// Accumulate only physical size and time from each cache store.
+		// Request count and logical size are tracked by the cache hierarchy.
+
+		FCacheBucketStats& BucketStats = CacheStats.GetBucket(Stats.Bucket);
+		TUniqueLock Lock(BucketStats.Mutex);
+
+		BucketStats.PhysicalReadSize += Stats.PhysicalReadSize;
+		BucketStats.PhysicalWriteSize += Stats.PhysicalWriteSize;
+		BucketStats.MainThreadTime += Stats.MainThreadTime;
+		BucketStats.OtherThreadTime += Stats.OtherThreadTime;
+
+	#if ENABLE_COOK_STATS
+		FCallStats& CallStats = bIsGet ? BucketStats.GetStats : BucketStats.PutStats;
+		CallStats.Accumulate(HitOrMiss, EStatType::Cycles, int64(Stats.MainThreadTime.ToSeconds() / FPlatformTime::GetSecondsPerCycle64()), /*bIsInGameThread*/ true);
+		CallStats.Accumulate(HitOrMiss, EStatType::Cycles, int64(Stats.OtherThreadTime.ToSeconds() / FPlatformTime::GetSecondsPerCycle64()), /*bIsInGameThread*/ false);
+		CallStats.Accumulate(HitOrMiss, EStatType::Bytes, bIsGet ? Stats.PhysicalReadSize : Stats.PhysicalWriteSize, bIsInGameThread);
+	#endif
+	}
+
+	TUniqueLock Lock(Mutex);
+
+	LogicalReadSize += Stats.LogicalReadSize;
+	LogicalWriteSize += Stats.LogicalWriteSize;
+	PhysicalReadSize += Stats.PhysicalReadSize;
+	PhysicalWriteSize += Stats.PhysicalWriteSize;
+
+	MainThreadTime += Stats.MainThreadTime;
+	OtherThreadTime += Stats.OtherThreadTime;
+
+	RequestCount.AddRequest(Stats.Type, Stats.Op, Stats.Status);
+
+	if (!Stats.Latency.IsInfinity())
+	{
+		AverageLatency.Add(Stats.StartTime, Stats.EndTime, Stats.Latency.ToSeconds());
+	}
+
+	if (Stats.PhysicalReadSize && !Stats.LogicalWriteSize)
+	{
+		// Skip physical reads for put operations because they pull down the read rate in time periods with few gets.
+		AveragePhysicalReadSize.Add(Stats.StartTime, Stats.EndTime, double(Stats.PhysicalReadSize));
+	}
+
+	if (Stats.PhysicalWriteSize && !Stats.LogicalReadSize)
+	{
+		// Skip physical writes for get operations because they pull down the write rate in time periods with few puts.
+		AveragePhysicalWriteSize.Add(Stats.StartTime, Stats.EndTime, double(Stats.PhysicalWriteSize));
+	}
+
+#if ENABLE_COOK_STATS
+	FCookStats::CallStats& CallStats = bIsGet ? GetStats : PutStats;
+	CallStats.Accumulate(HitOrMiss, EStatType::Counter, 1, bIsInGameThread);
+	CallStats.Accumulate(HitOrMiss, EStatType::Cycles, int64(Stats.MainThreadTime.ToSeconds() / FPlatformTime::GetSecondsPerCycle64()), /*bIsInGameThread*/ true);
+	CallStats.Accumulate(HitOrMiss, EStatType::Cycles, int64(Stats.OtherThreadTime.ToSeconds() / FPlatformTime::GetSecondsPerCycle64()), /*bIsInGameThread*/ false);
+	CallStats.Accumulate(HitOrMiss, EStatType::Bytes, bIsGet ? Stats.PhysicalReadSize : Stats.PhysicalWriteSize, bIsInGameThread);
+#endif
+}
+
+///////////////////////////////////////////////////////////////////////////////////////////////////////////////////////
+
+FCacheBucketStats& FCacheStats::GetBucket(FCacheBucket Bucket)
+{
+	TUniqueLock Lock(Mutex);
+	const int32 Index = Algo::LowerBoundBy(BucketStats, Bucket, [](TUniquePtr<FCacheBucketStats>& Stats) { return Stats->Bucket; });
+	if (!BucketStats.IsValidIndex(Index) || BucketStats[Index]->Bucket != Bucket)
+	{
+		FCacheBucketStats* Stats = new FCacheBucketStats;
+		Stats->Bucket = Bucket;
+		BucketStats.EmplaceAt(Index, Stats);
+	}
+	return *BucketStats[Index];
+}
+
+} // UE::DerivedData
diff --git a/Engine/Source/Developer/DerivedDataCache/Private/DerivedDataCacheStats.h b/Engine/Source/Developer/DerivedDataCache/Private/DerivedDataCacheStats.h
new file mode 100644
--- /dev/null
+++ b/Engine/Source/Developer/DerivedDataCache/Private/DerivedDataCacheStats.h
@@ -0,0 +1,195 @@
+// Copyright Epic Games, Inc. All Rights Reserved.
+
+#pragma once
+
+#include "Async/Mutex.h"
+#include "Async/UniqueLock.h"
+#include "Containers/Array.h"
+#include "DerivedDataCacheKey.h"
+#include "DerivedDataCacheStore.h"
+#include "Internationalization/Text.h"
+#include "Misc/MonotonicTime.h"
+#include "ProfilingDebugging/CookStats.h"
+
+namespace UE::DerivedData { class FCacheStats; }
+
+namespace UE::DerivedData
+{
+
+class FRequestCounter
+{
+	struct FCounter
+	{
+		uint64 Hits = 0;
+		uint64 Misses = 0;
+	};
+
+	using ERequestOp = ECacheStoreRequestOp;
+	using ERequestType = ECacheStoreRequestType;
+
+public:
+	void AddRequest(ERequestType Type, ERequestOp Op, EStatus Status);
+
+	constexpr uint64 Puts() const { return PutHits() + PutMisses(); }
+	constexpr uint64 PutHits() const { return PutRecord.Hits + PutValue.Hits; }
+	constexpr uint64 PutMisses() const { return PutRecord.Misses + PutValue.Misses; }
+
+	constexpr uint64 Gets() const { return GetHits() + GetMisses(); }
+	constexpr uint64 GetHits() const { return GetRecord.Hits + GetValue.Hits + GetRecordChunk.Hits + GetValueChunk.Hits; }
+	constexpr uint64 GetMisses() const { return GetRecord.Misses + GetValue.Misses + GetRecordChunk.Misses + GetValueChunk.Misses; }
+
+	constexpr uint64 Total() const { return Puts() + Gets(); }
+	constexpr uint64 TotalHits() const { return PutHits() + GetHits(); }
+	constexpr uint64 TotalMisses() const { return PutMisses() + GetMisses(); }
+
+private:
+	FCounter PutRecord;
+	FCounter PutValue;
+	FCounter GetRecord;
+	FCounter GetValue;
+	FCounter GetRecordChunk;
+	FCounter GetValueChunk;
+};
+
+class FTimeAveragedStat
+{
+public:
+	void SetPeriod(FMonotonicTimeSpan NewPeriod) { Period = NewPeriod; }
+
+	void Add(FMonotonicTimePoint StartTime, FMonotonicTimePoint EndTime, double Value);
+
+	/** Calculate and return the average rate at the time. Time must be monotonically increasing. */
+	double GetRate(FMonotonicTimePoint Time);
+
+	/** Calculate and return the average value at the time. Time must be monotonically increasing. */
+	double GetValue(FMonotonicTimePoint Time);
+
+private:
+	void Update(FMonotonicTimePoint Time);
+
+	struct FRange
+	{
+		FMonotonicTimePoint StartTime;
+		FMonotonicTimePoint EndTime;
+	};
+
+	struct FValue
+	{
+		FMonotonicTimePoint Time;
+		double Value = 0.0;
+	};
+
+	/** Period of time that this stat is averaged over. */
+	FMonotonicTimeSpan Period;
+	/** Time at which the most recent update occurred. */
+	FMonotonicTimePoint LastTime;
+	/** Ranges of time that have values. */
+	TArray<FRange> ActiveRanges;
+	/** Values that start later than the most recent update time. */
+	TArray<FValue> StartValues;
+	/** Values that end later than the most recent update time. */
+	TArray<FValue> EndValues;
+	/** Average rate from the last update that had values, or 0. */
+	double AverageRate = 0.0;
+	/** Average value from the last update that had values, or 0. */
+	double AverageValue = 0.0;
+	/** Sum of values in the period up to the most recent update time. */
+	double AccumulatedValue = 0.0;
+	/** Number of values within the period at the most recent update time. */
+	int32 AccumulatedValueCount = 0;
+	/** Mutex that can be used to synchronize access to the stat from multiple threads. */
+	FMutex Mutex;
+};
+
+///////////////////////////////////////////////////////////////////////////////////////////////////////////////////////
+
+class FCacheBucketStats final
+{
+public:
+	/** Name of the cache bucket. */
+	FCacheBucket Bucket;
+	/** Size of the raw data and metadata that was read from the cache bucket. */
+	uint64 LogicalReadSize = 0;
+	/** Size of the raw data and metadata that was written to the cache bucket. */
+	uint64 LogicalWriteSize = 0;
+	/** Size of the compressed data and anything else that was read from the cache bucket. */
+	uint64 PhysicalReadSize = 0;
+	/** Size of the compressed data and anything else that was written to the cache bucket. */
+	uint64 PhysicalWriteSize = 0;
+	/** Approximate CPU time on the main thread that was needed to satisfy requests in the cache bucket. */
+	FMonotonicTimeSpan MainThreadTime;
+	/** Approximate CPU time off the main thread that was needed to satisfy requests in the cache bucket. */
+	FMonotonicTimeSpan OtherThreadTime;
+	/** Count of requests that accessed the cache bucket. */
+	FRequestCounter RequestCount;
+	/** Mutex that can be used to synchronize access to the cache buckets stats from multiple threads. */
+	mutable FMutex Mutex;
+
+#if ENABLE_COOK_STATS
+	FCookStats::CallStats GetStats;
+	FCookStats::CallStats PutStats;
+#endif
+};
+
+class FCacheStoreStats final : public ICacheStoreStats
+{
+public:
+	FCacheStoreStats(FCacheStats& CacheStats, ECacheStoreFlags Flags, FStringView Type, FStringView Name, FStringView Path);
+
+	FString Type;
+	FString Name;
+	FString Path;
+	FText Status;
+	ECacheStoreFlags Flags = ECacheStoreFlags::None;
+	ECacheStoreStatusCode StatusCode = ECacheStoreStatusCode::None;
+	mutable FMutex Mutex;
+
+	TMap<FString, FString> Attributes;
+
+	FCacheStats& CacheStats;
+
+	/** Size of the raw data and metadata that was read from the cache store. */
+	uint64 LogicalReadSize = 0;
+	/** Size of the raw data and metadata that was written to the cache store. */
+	uint64 LogicalWriteSize = 0;
+	/** Size of the compressed data and anything else that was read from the cache store. */
+	uint64 PhysicalReadSize = 0;
+	/** Size of the compressed data and anything else that was written to the cache store. */
+	uint64 PhysicalWriteSize = 0;
+	/** Approximate CPU time on the main thread that was needed. */
+	FMonotonicTimeSpan MainThreadTime;
+	/** Approximate CPU time off the main thread that was needed. */
+	FMonotonicTimeSpan OtherThreadTime;
+	/** Count of requests that accessed the cache store. */
+	FRequestCounter RequestCount;
+
+	FTimeAveragedStat AverageLatency;
+	FTimeAveragedStat AveragePhysicalReadSize;
+	FTimeAveragedStat AveragePhysicalWriteSize;
+
+#if ENABLE_COOK_STATS
+	FCookStats::CallStats GetStats;
+	FCookStats::CallStats PutStats;
+#endif
+
+private:
+	FStringView GetType() const final { return Type; }
+	FStringView GetName() const final { return Name; }
+	FStringView GetPath() const final { return Path; }
+	void SetFlags(ECacheStoreFlags Flags) final;
+	void SetStatus(ECacheStoreStatusCode StatusCode, const FText& Status) final;
+	void SetAttribute(FStringView Key, FStringView Value) final;
+	void AddRequest(const FCacheStoreRequestStats& Stats) final;
+};
+
+class FCacheStats final
+{
+public:
+	TArray<TUniquePtr<FCacheBucketStats>> BucketStats;
+	TArray<TUniquePtr<FCacheStoreStats>> StoreStats;
+	mutable FMutex Mutex;
+
+	FCacheBucketStats& GetBucket(FCacheBucket Bucket);
+};
+
+} // UE::DerivedData
diff --git a/Engine/Source/Developer/DerivedDataCache/Private/DerivedDataCacheStore.h b/Engine/Source/Developer/DerivedDataCache/Private/DerivedDataCacheStore.h
--- a/Engine/Source/Developer/DerivedDataCache/Private/DerivedDataCacheStore.h
+++ b/Engine/Source/Developer/DerivedDataCache/Private/DerivedDataCacheStore.h
@@ -3,10 +3,12 @@
 #pragma once
 
 #include "Async/Mutex.h"
+#include "Containers/ContainersFwd.h"
 #include "DerivedDataCache.h"
 #include "Misc/MonotonicTime.h"
 
 class FText;
+struct FDerivedDataCacheResourceStat;
 
 namespace UE::DerivedData { class FCacheRecord; }
 namespace UE::DerivedData { class FCacheStoreRequestTimer; }
@@ -162,6 +164,8 @@ public:
 
 	virtual ICacheStoreStats* CreateStats(ILegacyCacheStore* CacheStore, ECacheStoreFlags Flags, FStringView Type, FStringView Name, FStringView Path = {}) = 0;
 	virtual void DestroyStats(ICacheStoreStats* Stats) = 0;
+
+	virtual void LegacyResourceStats(TArray<FDerivedDataCacheResourceStat>& OutStats) const = 0;
 };
 
 enum class ECacheStoreStatusCode : uint8
diff --git a/Engine/Source/Developer/DerivedDataCache/Private/DerivedDataCacheStoreHierarchy.cpp b/Engine/Source/Developer/DerivedDataCache/Private/DerivedDataCacheStoreHierarchy.cpp
--- a/Engine/Source/Developer/DerivedDataCache/Private/DerivedDataCacheStoreHierarchy.cpp
+++ b/Engine/Source/Developer/DerivedDataCache/Private/DerivedDataCacheStoreHierarchy.cpp
@@ -11,6 +11,7 @@
 #include "Containers/Array.h"
 #include "Containers/StaticArray.h"
 #include "DerivedDataCache.h"
+#include "DerivedDataCacheStats.h"
 #include "DerivedDataCacheStore.h"
 #include "DerivedDataCacheUsageStats.h"
 #include "DerivedDataLegacyCacheStore.h"
@@ -28,6 +29,8 @@
 #include "Templates/UniquePtr.h"
 #include <atomic>
 
+void GatherDerivedDataCacheResourceStats(TArray<FDerivedDataCacheResourceStat>& DDCResourceStats);
+
 namespace UE::DerivedData
 {
 
@@ -35,132 +38,6 @@ ILegacyCacheStore* CreateCacheStoreAsync(ILegacyCacheStore* InnerCache, IMemoryC
 
 ///////////////////////////////////////////////////////////////////////////////////////////////////////////////////////
 
-class FCacheStoreStats final : public ICacheStoreStats
-{
-public:
-	FCacheStoreStats(ECacheStoreFlags Flags, FStringView Type, FStringView Name, FStringView Path);
-
-	FString Type;
-	FString Name;
-	FString Path;
-	FText Status;
-	ECacheStoreFlags Flags = ECacheStoreFlags::None;
-	ECacheStoreStatusCode StatusCode = ECacheStoreStatusCode::None;
-	mutable FMutex Mutex;
-
-	TMap<FString, FString> Attributes;
-
-#if ENABLE_COOK_STATS
-	FCookStats::CallStats GetStats;
-	FCookStats::CallStats PutStats;
-#endif
-
-	struct FRateStat
-	{
-		FMonotonicTimePoint StartTime;
-		FMonotonicTimePoint EndTime;
-		uint64 Amount = 0;
-	};
-
-	constexpr static int32 SpeedStatsMax = 16;
-
-	TStaticArray<FMonotonicTimeSpan, SpeedStatsMax> LatencyHistory;
-	int32 LatencyIndex = 0;
-	int32 LatencyCount = 0;
-
-	TStaticArray<FRateStat, SpeedStatsMax> PhysicalReadSizeHistory;
-	int32 PhysicalReadSizeIndex = 0;
-	int32 PhysicalReadSizeCount = 0;
-
-	TStaticArray<FRateStat, SpeedStatsMax> PhysicalWriteSizeHistory;
-	int32 PhysicalWriteSizeIndex = 0;
-	int32 PhysicalWriteSizeCount = 0;
-
-private:
-	FStringView GetType() const final { return Type; }
-	FStringView GetName() const final { return Name; }
-	FStringView GetPath() const final { return Path; }
-	void SetFlags(ECacheStoreFlags Flags) final;
-	void SetStatus(ECacheStoreStatusCode StatusCode, const FText& Status) final;
-	void SetAttribute(FStringView Key, FStringView Value) final;
-	void AddRequest(const FCacheStoreRequestStats& Stats) final;
-};
-
-FCacheStoreStats::FCacheStoreStats(ECacheStoreFlags InFlags, FStringView InType, FStringView InName, FStringView InPath)
-	: Type(InType)
-	, Name(InName)
-	, Path(InPath)
-	, Flags(InFlags)
-{
-}
-
-void FCacheStoreStats::SetFlags(ECacheStoreFlags InFlags)
-{
-	TUniqueLock Lock(Mutex);
-	Flags = InFlags;
-}
-
-void FCacheStoreStats::SetStatus(ECacheStoreStatusCode InStatusCode, const FText& InStatus)
-{
-	TUniqueLock Lock(Mutex);
-	Status = InStatus;
-	StatusCode = InStatusCode;
-}
-
-void FCacheStoreStats::SetAttribute(FStringView Key, FStringView Value)
-{
-	TUniqueLock Lock(Mutex);
-	Attributes.Emplace(Key, Value);
-}
-
-void FCacheStoreStats::AddRequest(const FCacheStoreRequestStats& Stats)
-{
-#if ENABLE_COOK_STATS
-	using FCallStats = FCookStats::CallStats;
-	using EHitOrMiss = FCallStats::EHitOrMiss;
-	using EStatType = FCallStats::EStatType;
-	const bool bIsInGameThread = IsInGameThread();
-	const bool bIsGet = Stats.Op != ECacheStoreRequestOp::Put;
-#endif
-
-	TUniqueLock Lock(Mutex);
-
-	if (!Stats.Latency.IsInfinity())
-	{
-		LatencyHistory[LatencyIndex++] = Stats.Latency;
-		LatencyCount = FMath::Max(LatencyCount, LatencyIndex);
-		LatencyIndex %= LatencyHistory.Num();
-	}
-
-	if (Stats.PhysicalReadSize && !Stats.LogicalWriteSize)
-	{
-		FRateStat& CurrentPhysicalReadSize = PhysicalReadSizeHistory[PhysicalReadSizeIndex++];
-		PhysicalReadSizeCount = FMath::Max(PhysicalReadSizeCount, PhysicalReadSizeIndex);
-		PhysicalReadSizeIndex %= PhysicalReadSizeHistory.Num();
-		CurrentPhysicalReadSize.StartTime = Stats.StartTime;
-		CurrentPhysicalReadSize.EndTime = Stats.EndTime;
-		CurrentPhysicalReadSize.Amount = Stats.PhysicalReadSize;
-	}
-
-	if (Stats.PhysicalWriteSize && !Stats.LogicalReadSize)
-	{
-		FRateStat& CurrentPhysicalWriteSize = PhysicalWriteSizeHistory[PhysicalWriteSizeIndex++];
-		PhysicalWriteSizeCount = FMath::Max(PhysicalWriteSizeCount, PhysicalWriteSizeIndex);
-		PhysicalWriteSizeIndex %= PhysicalWriteSizeHistory.Num();
-		CurrentPhysicalWriteSize.StartTime = Stats.StartTime;
-		CurrentPhysicalWriteSize.EndTime = Stats.EndTime;
-		CurrentPhysicalWriteSize.Amount = Stats.PhysicalWriteSize;
-	}
-
-#if ENABLE_COOK_STATS
-	FCookStats::CallStats& CallStats = bIsGet ? GetStats : PutStats;
-	CallStats.Accumulate(Stats.Status == EStatus::Ok ? EHitOrMiss::Hit : EHitOrMiss::Miss, EStatType::Counter, 1, bIsInGameThread);
-	CallStats.Accumulate(Stats.Status == EStatus::Ok ? EHitOrMiss::Hit : EHitOrMiss::Miss, EStatType::Cycles, int64(Stats.MainThreadTime.ToSeconds() / FPlatformTime::GetSecondsPerCycle64()), /*bIsInGameThread*/ true);
-	CallStats.Accumulate(Stats.Status == EStatus::Ok ? EHitOrMiss::Hit : EHitOrMiss::Miss, EStatType::Cycles, int64(Stats.OtherThreadTime.ToSeconds() / FPlatformTime::GetSecondsPerCycle64()), /*bIsInGameThread*/ false);
-	CallStats.Accumulate(Stats.Status == EStatus::Ok ? EHitOrMiss::Hit : EHitOrMiss::Miss, EStatType::Bytes, bIsGet ? Stats.PhysicalReadSize : Stats.PhysicalWriteSize, bIsInGameThread);
-#endif
-}
-
 class FCacheStoreHierarchy final : public ILegacyCacheStore, public ICacheStoreOwner
 {
 public:
@@ -174,6 +51,8 @@ public:
 	ICacheStoreStats* CreateStats(ILegacyCacheStore* CacheStore, ECacheStoreFlags Flags, FStringView Type, FStringView Name, FStringView Path) final;
 	void DestroyStats(ICacheStoreStats* Stats) final;
 
+	void LegacyResourceStats(TArray<FDerivedDataCacheResourceStat>& OutStats) const final;
+
 	void Put(
 		TConstArrayView<FCachePutRequest> Requests,
 		IRequestOwner& Owner,
@@ -211,6 +90,8 @@ private:
 	// Caller must hold a write lock on CacheStoresLock.
 	void UpdateNodeFlags();
 
+	void RecordStats(FCacheBucket Bucket, ERequestType Type, ERequestOp Op, EStatus Status, uint64 LogicalReadSize, uint64 LogicalWriteSize);
+
 	class FCounterEvent;
 	class FDynamicRequestBarrier;
 
@@ -236,6 +117,9 @@ private:
 	static bool CanStoreIfOk(ECachePolicy Policy, ECacheStoreNodeFlags Flags);
 	static bool CanQueryIfError(ECachePolicy Policy, ECacheStoreNodeFlags Flags);
 
+	static uint64 MeasureLogicalRecordSize(const FCacheRecord& Record);
+	static uint64 MeasureLogicalValueSize(const FValue& Value);
+
 	struct FCacheStoreNode
 	{
 		ILegacyCacheStore* Cache{};
@@ -249,6 +133,7 @@ private:
 	ECacheStoreNodeFlags CombinedNodeFlags{};
 	TArray<FCacheStoreNode, TInlineAllocator<8>> Nodes;
 	IMemoryCacheStore* MemoryCache;
+	FCacheStats CacheStats;
 };
 
 ///////////////////////////////////////////////////////////////////////////////////////////////////////////////////////
@@ -340,6 +225,9 @@ struct FCacheStoreHierarchy::TBatchParams
 	static FGetResponse FilterResponseByRequest(const FGetResponse& Response, const FGetRequest& Request);
 	static FPutRequest MakePutRequest(const FGetResponse& Response, const FGetRequest& Request);
 	static FGetRequest MakeGetRequest(const FPutRequest& Request, int32 RequestIndex);
+	static uint64 MeasureLogicalSize(const FGetResponse& Response);
+	static uint64 MeasureLogicalSize(const FPutRequest& Request);
+	static ERequestType GetBatchType();
 };
 
 ///////////////////////////////////////////////////////////////////////////////////////////////////////////////////////
@@ -462,7 +350,12 @@ ICacheStoreStats* FCacheStoreHierarchy::CreateStats(ILegacyCacheStore* CacheStor
 	{
 		if (Node.Cache == CacheStore)
 		{
-			Node.CacheStats.Emplace(new FCacheStoreStats(Flags, Type, Name, Path));
+			TUniquePtr<FCacheStoreStats> StoreStats(new FCacheStoreStats(CacheStats, Flags, Type, Name, Path));
+			Node.CacheStats.Emplace(StoreStats.Get());
+
+			TUniqueLock StatsLock(CacheStats.Mutex);
+			CacheStats.StoreStats.Emplace(MoveTemp(StoreStats));
+
 			return Node.CacheStats.Last();
 		}
 	}
@@ -480,7 +373,11 @@ void FCacheStoreHierarchy::DestroyStats(ICacheStoreStats* Stats)
 	{
 		if (Node.CacheStats.RemoveSingle((FCacheStoreStats*)Stats))
 		{
-			delete Stats;
+			TUniqueLock StatsLock(CacheStats.Mutex);
+			CacheStats.StoreStats.RemoveAll([Stats](TUniquePtr<FCacheStoreStats>& Match)
+			{
+				return Match.Get() == (FCacheStoreStats*)Stats;
+			});
 			return;
 		}
 	}
@@ -489,6 +386,104 @@ void FCacheStoreHierarchy::DestroyStats(ICacheStoreStats* Stats)
 		Stats->GetType(), Stats->GetName(), Stats->GetPath());
 }
 
+void FCacheStoreHierarchy::LegacyResourceStats(TArray<FDerivedDataCacheResourceStat>& OutStats) const
+{
+#if ENABLE_COOK_STATS
+	using FCallStats = FCookStats::CallStats;
+	using EHitOrMiss = FCallStats::EHitOrMiss;
+	using EStatType = FCallStats::EStatType;
+
+	TArray<const FCacheBucketStats*, TInlineAllocator<64>> Buckets;
+	{
+		TUniqueLock StatsLock(CacheStats.Mutex);
+		for (const TUniquePtr<FCacheBucketStats>& Bucket : CacheStats.BucketStats)
+		{
+			Buckets.Add(Bucket.Get());
+		}
+	}
+
+	TMap<FString, int32> BucketNameToIndex;
+
+	for (const FCacheBucketStats* BucketStats : Buckets)
+	{
+		TUniqueLock Lock(BucketStats->Mutex);
+
+		TStringBuilder<64> DisplayName;
+		BucketStats->Bucket.ToDisplayName(DisplayName);
+		FString BucketName(DisplayName);
+
+		int32 Index;
+		if (const int32* IndexPtr = BucketNameToIndex.Find(BucketName))
+		{
+			Index = *IndexPtr;
+		}
+		else
+		{
+			Index = OutStats.Emplace(BucketName);
+			BucketNameToIndex.Add(BucketName, Index);
+		}
+
+		const int64 BuildCycles =
+			BucketStats->GetStats.GetAccumulatedValueAnyThread(EHitOrMiss::Miss, EStatType::Cycles) +
+			BucketStats->PutStats.GetAccumulatedValueAnyThread(EHitOrMiss::Hit, EStatType::Cycles);
+		OutStats[Index] += FDerivedDataCacheResourceStat(BucketName, /*bIsGameThreadTime*/ false,
+			double(BucketStats->GetStats.GetAccumulatedValueAnyThread(EHitOrMiss::Hit, EStatType::Cycles)) * FPlatformTime::GetSecondsPerCycle(),
+			double(BucketStats->GetStats.GetAccumulatedValueAnyThread(EHitOrMiss::Hit, EStatType::Bytes)) / 1024.0 / 1024.0,
+			BucketStats->GetStats.GetAccumulatedValueAnyThread(EHitOrMiss::Hit, EStatType::Counter),
+			double(BuildCycles) * FPlatformTime::GetSecondsPerCycle(),
+			double(BucketStats->PutStats.GetAccumulatedValueAnyThread(EHitOrMiss::Miss, EStatType::Bytes)) / 1024.0 / 1024.0,
+			BucketStats->PutStats.GetAccumulatedValueAnyThread(EHitOrMiss::Hit, EStatType::Counter));
+
+		const int64 MainThreadCycles =
+			BucketStats->GetStats.GetAccumulatedValue(EHitOrMiss::Hit, EStatType::Cycles, /*bIsInGameThread*/ true) +
+			BucketStats->GetStats.GetAccumulatedValue(EHitOrMiss::Miss, EStatType::Cycles, /*bIsInGameThread*/ true) +
+			BucketStats->PutStats.GetAccumulatedValue(EHitOrMiss::Hit, EStatType::Cycles, /*bIsInGameThread*/ true) +
+			BucketStats->PutStats.GetAccumulatedValue(EHitOrMiss::Miss, EStatType::Cycles, /*bIsInGameThread*/ true);
+		OutStats[Index].GameThreadTimeSec += double(MainThreadCycles) * FPlatformTime::GetSecondsPerCycle();
+	}
+
+	// Add in build times from the legacy stats because those are not captured by ICacheStats.
+	TArray<FDerivedDataCacheResourceStat> ResourceStats;
+	GatherDerivedDataCacheResourceStats(ResourceStats);
+
+	TMultiMap<FString, int32> LegacyNameToIndex;
+	for (auto It = ResourceStats.CreateConstIterator(); It; ++It)
+	{
+		if (int32 ParenIndex; It->AssetType.FindChar(TEXT('('), ParenIndex))
+		{
+			LegacyNameToIndex.Add(It->AssetType.Left(ParenIndex - 1), It.GetIndex());
+		}
+		else
+		{
+			LegacyNameToIndex.Add(It->AssetType, It.GetIndex());
+		}
+	}
+
+	for (FDerivedDataCacheResourceStat& Stat : OutStats)
+	{
+		double ResourceBuildTimeSec = 0.0;
+		double ResourceGameThreadTimeSec = 0.0;
+
+		TArray<int32, TInlineAllocator<4>> ResourceStatsForType;
+		LegacyNameToIndex.MultiFind(Stat.AssetType, ResourceStatsForType);
+		for (int32 Index : ResourceStatsForType)
+		{
+			const FDerivedDataCacheResourceStat& ResourceStat = ResourceStats[Index];
+			ResourceBuildTimeSec += ResourceStat.BuildTimeSec;
+			ResourceGameThreadTimeSec += ResourceStat.GameThreadTimeSec;
+		}
+
+		Stat.BuildTimeSec += ResourceBuildTimeSec;
+		if (Stat.GameThreadTimeSec < ResourceGameThreadTimeSec)
+		{
+			// Add the game thread time that was tracked in excess of the cache.
+			// This is an approximation of the build time.
+			Stat.GameThreadTimeSec += ResourceGameThreadTimeSec - Stat.GameThreadTimeSec;
+		}
+	}
+#endif
+}
+
 FCacheRecordPolicy FCacheStoreHierarchy::AddPolicy(const FCacheRecordPolicy& BasePolicy, ECachePolicy Policy)
 {
 	return BasePolicy.Transform([Policy](ECachePolicy P) { return P | Policy; });
@@ -545,17 +540,19 @@ class FCacheStoreHierarchy::TPutBatch final : public FBatchBase, public Params
 	using Params::Put;
 	using Params::Get;
 	using Params::MakeGetRequest;
+	using Params::MeasureLogicalSize;
+	using Params::GetBatchType;
 
 public:
 	static void Begin(
-		const FCacheStoreHierarchy& Hierarchy,
+		FCacheStoreHierarchy& Hierarchy,
 		TConstArrayView<FPutRequest> Requests,
 		IRequestOwner& Owner,
 		FOnPutComplete&& OnComplete);
 
 private:
 	TPutBatch(
-		const FCacheStoreHierarchy& InHierarchy,
+		FCacheStoreHierarchy& InHierarchy,
 		const TConstArrayView<FPutRequest> InRequests,
 		IRequestOwner& InOwner,
 		FOnPutComplete&& InOnComplete)
@@ -577,13 +574,15 @@ private:
 	bool DispatchPutRequests();
 	void CompletePutRequest(FPutResponse&& Response);
 
+	void FinishRequest(FPutResponse&& Response, const FPutRequest& Request);
+
 	struct FRequestState
 	{
 		bool bOk = false;
 		bool bStop = false;
 	};
 
-	const FCacheStoreHierarchy& Hierarchy;
+	FCacheStoreHierarchy& Hierarchy;
 	TArray<FPutRequest, TInlineAllocator<1>> Requests;
 	IRequestOwner& BatchOwner;
 	FOnPutComplete OnComplete;
@@ -597,7 +596,7 @@ private:
 
 template <typename Params>
 void FCacheStoreHierarchy::TPutBatch<Params>::Begin(
-	const FCacheStoreHierarchy& InHierarchy,
+	FCacheStoreHierarchy& InHierarchy,
 	const TConstArrayView<FPutRequest> InRequests,
 	IRequestOwner& InOwner,
 	FOnPutComplete&& InOnComplete)
@@ -632,7 +631,7 @@ void FCacheStoreHierarchy::TPutBatch<Params>::DispatchRequests()
 		const FRequestState& State = States[RequestIndex];
 		if (!State.bOk && !State.bStop)
 		{
-			OnComplete(Request.MakeResponse(BatchOwner.IsCanceled() ? EStatus::Canceled : EStatus::Error));
+			FinishRequest(Request.MakeResponse(BatchOwner.IsCanceled() ? EStatus::Canceled : EStatus::Error), Request);
 		}
 		++RequestIndex;
 	}
@@ -692,7 +691,8 @@ void FCacheStoreHierarchy::TPutBatch<Params>::CompleteGetRequest(FGetResponse&&
 		State.bStop = true;
 		if (!State.bOk)
 		{
-			OnComplete(Requests[RequestIndex].MakeResponse(Response.Status));
+			const FPutRequest& Request = Requests[RequestIndex];
+			FinishRequest(Request.MakeResponse(Response.Status), Request);
 		}
 	}
 	if (RemainingRequestCount.Signal())
@@ -757,8 +757,9 @@ void FCacheStoreHierarchy::TPutBatch<Params>::CompletePutRequest(FPutResponse&&
 		FRequestState& State = States[RequestIndex];
 		check(!State.bOk && !State.bStop);
 		State.bOk = true;
-		Response.UserData = Requests[RequestIndex].UserData;
-		OnComplete(MoveTemp(Response));
+		const FPutRequest& Request = Requests[RequestIndex];
+		Response.UserData = Request.UserData;
+		FinishRequest(MoveTemp(Response), Request);
 	}
 	if (RemainingRequestCount.Signal())
 	{
@@ -767,6 +768,13 @@ void FCacheStoreHierarchy::TPutBatch<Params>::CompletePutRequest(FPutResponse&&
 	}
 }
 
+template <typename Params>
+void FCacheStoreHierarchy::TPutBatch<Params>::FinishRequest(FPutResponse&& Response, const FPutRequest& Request)
+{
+	Hierarchy.RecordStats(Response.Key.Bucket, GetBatchType(), ERequestOp::Put, Response.Status, 0, MeasureLogicalSize(Request));
+	OnComplete(MoveTemp(Response));
+}
+
 ///////////////////////////////////////////////////////////////////////////////////////////////////////////////////////
 
 template <typename Params>
@@ -785,17 +793,19 @@ class FCacheStoreHierarchy::TGetBatch final : public FBatchBase, public Params
 	using Params::ModifyPolicyForResponse;
 	using Params::FilterResponseByRequest;
 	using Params::MakePutRequest;
+	using Params::MeasureLogicalSize;
+	using Params::GetBatchType;
 
 public:
 	static void Begin(
-		const FCacheStoreHierarchy& Hierarchy,
+		FCacheStoreHierarchy& Hierarchy,
 		TConstArrayView<FGetRequest> Requests,
 		IRequestOwner& Owner,
 		FOnGetComplete&& OnComplete);
 
 private:
 	TGetBatch(
-		const FCacheStoreHierarchy& InHierarchy,
+		FCacheStoreHierarchy& InHierarchy,
 		const TConstArrayView<FGetRequest> InRequests,
 		IRequestOwner& InOwner,
 		FOnGetComplete&& InOnComplete)
@@ -815,6 +825,8 @@ private:
 	void DispatchRequests();
 	void CompleteRequest(FGetResponse&& Response);
 
+	void FinishRequest(FGetResponse&& Response, const FGetRequest& Request);
+
 	struct FState
 	{
 		FGetRequest Request;
@@ -822,7 +834,7 @@ private:
 		int32 NodeIndex = 0;
 	};
 
-	const FCacheStoreHierarchy& Hierarchy;
+	FCacheStoreHierarchy& Hierarchy;
 	FOnGetComplete OnComplete;
 	TArray<FState, TInlineAllocator<8>> States;
 
@@ -833,7 +845,7 @@ private:
 
 template <typename Params>
 void FCacheStoreHierarchy::TGetBatch<Params>::Begin(
-	const FCacheStoreHierarchy& InHierarchy,
+	FCacheStoreHierarchy& InHierarchy,
 	const TConstArrayView<FGetRequest> InRequests,
 	IRequestOwner& InOwner,
 	FOnGetComplete&& InOnComplete)
@@ -949,7 +961,7 @@ void FCacheStoreHierarchy::TGetBatch<Params>::DispatchRequests()
 	{
 		if (State.Response.Status != EStatus::Ok)
 		{
-			OnComplete(FilterResponseByRequest(State.Response, State.Request));
+			FinishRequest(FilterResponseByRequest(State.Response, State.Request), State.Request);
 		}
 	}
 }
@@ -1001,7 +1013,7 @@ void FCacheStoreHierarchy::TGetBatch<Params>::CompleteRequest(FGetResponse&& Res
 	if (bFirstOk)
 	{
 		// Values may be fetched to populate earlier nodes. Remove values if requested.
-		OnComplete(FilterResponseByRequest(State.Response, State.Request));
+		FinishRequest(FilterResponseByRequest(State.Response, State.Request), State.Request);
 	}
 
 	if (Response.Status == EStatus::Ok)
@@ -1022,6 +1034,13 @@ void FCacheStoreHierarchy::TGetBatch<Params>::CompleteRequest(FGetResponse&& Res
 	++State.NodeIndex;
 }
 
+template <typename Params>
+void FCacheStoreHierarchy::TGetBatch<Params>::FinishRequest(FGetResponse&& Response, const FGetRequest& Request)
+{
+	Hierarchy.RecordStats(Request.Key.Bucket, GetBatchType(), ERequestOp::Get, Response.Status, MeasureLogicalSize(Response), 0);
+	OnComplete(MoveTemp(Response));
+}
+
 ///////////////////////////////////////////////////////////////////////////////////////////////////////////////////////
 
 template <>
@@ -1192,6 +1211,24 @@ FCacheGetRequest FCacheStoreHierarchy::FCacheRecordBatchParams::MakeGetRequest(
 	return {Request.Name, Request.Record.GetKey(), AddPolicy(Request.Policy, ECachePolicy::SkipData), uint64(RequestIndex)};
 }
 
+template <>
+uint64 FCacheStoreHierarchy::FCacheRecordBatchParams::MeasureLogicalSize(const FGetResponse& Response)
+{
+	return MeasureLogicalRecordSize(Response.Record);
+}
+
+template <>
+uint64 FCacheStoreHierarchy::FCacheRecordBatchParams::MeasureLogicalSize(const FPutRequest& Request)
+{
+	return MeasureLogicalRecordSize(Request.Record);
+}
+
+template <>
+ECacheStoreRequestType FCacheStoreHierarchy::FCacheRecordBatchParams::GetBatchType()
+{
+	return ERequestType::Record;
+}
+
 void FCacheStoreHierarchy::Put(
 	const TConstArrayView<FCachePutRequest> Requests,
 	IRequestOwner& Owner,
@@ -1268,6 +1305,24 @@ FCacheGetValueRequest FCacheStoreHierarchy::FCacheValueBatchParams::MakeGetReque
 	return {Request.Name, Request.Key, AddPolicy(Request.Policy, ECachePolicy::SkipData), uint64(RequestIndex)};
 }
 
+template <>
+uint64 FCacheStoreHierarchy::FCacheValueBatchParams::MeasureLogicalSize(const FGetResponse& Response)
+{
+	return MeasureLogicalValueSize(Response.Value);
+}
+
+template <>
+uint64 FCacheStoreHierarchy::FCacheValueBatchParams::MeasureLogicalSize(const FPutRequest& Request)
+{
+	return MeasureLogicalValueSize(Request.Value);
+}
+
+template <>
+ECacheStoreRequestType FCacheStoreHierarchy::FCacheValueBatchParams::GetBatchType()
+{
+	return ERequestType::Value;
+}
+
 void FCacheStoreHierarchy::PutValue(
 	const TConstArrayView<FCachePutValueRequest> Requests,
 	IRequestOwner& Owner,
@@ -1290,14 +1345,14 @@ class FCacheStoreHierarchy::FGetChunksBatch final : public FBatchBase
 {
 public:
 	static void Begin(
-		const FCacheStoreHierarchy& Hierarchy,
+		FCacheStoreHierarchy& Hierarchy,
 		TConstArrayView<FCacheGetChunkRequest> Requests,
 		IRequestOwner& Owner,
 		FOnCacheGetChunkComplete&& OnComplete);
 
 private:
 	FGetChunksBatch(
-		const FCacheStoreHierarchy& InHierarchy,
+		FCacheStoreHierarchy& InHierarchy,
 		const TConstArrayView<FCacheGetChunkRequest> InRequests,
 		IRequestOwner& InOwner,
 		FOnCacheGetChunkComplete&& InOnComplete)
@@ -1315,13 +1370,15 @@ private:
 	void DispatchRequests();
 	void CompleteRequest(FCacheGetChunkResponse&& Response);
 
+	void FinishRequest(FCacheGetChunkResponse&& Response);
+
 	struct FState
 	{
 		FCacheGetChunkRequest Request;
 		EStatus Status = EStatus::Error;
 	};
 
-	const FCacheStoreHierarchy& Hierarchy;
+	FCacheStoreHierarchy& Hierarchy;
 	FOnCacheGetChunkComplete OnComplete;
 	TArray<FState, TInlineAllocator<8>> States;
 
@@ -1331,7 +1388,7 @@ private:
 };
 
 void FCacheStoreHierarchy::FGetChunksBatch::Begin(
-	const FCacheStoreHierarchy& InHierarchy,
+	FCacheStoreHierarchy& InHierarchy,
 	const TConstArrayView<FCacheGetChunkRequest> InRequests,
 	IRequestOwner& InOwner,
 	FOnCacheGetChunkComplete&& InOnComplete)
@@ -1390,7 +1447,7 @@ void FCacheStoreHierarchy::FGetChunksBatch::DispatchRequests()
 	{
 		if (State.Status != EStatus::Ok)
 		{
-			OnComplete(State.Request.MakeResponse(Owner.IsCanceled() ? EStatus::Canceled : EStatus::Error));
+			FinishRequest(State.Request.MakeResponse(Owner.IsCanceled() ? EStatus::Canceled : EStatus::Error));
 		}
 	}
 }
@@ -1402,7 +1459,7 @@ void FCacheStoreHierarchy::FGetChunksBatch::CompleteRequest(FCacheGetChunkRespon
 	{
 		check(State.Status == EStatus::Error);
 		Response.UserData = State.Request.UserData;
-		OnComplete(MoveTemp(Response));
+		FinishRequest(MoveTemp(Response));
 	}
 	State.Status = Response.Status;
 
@@ -1413,6 +1470,13 @@ void FCacheStoreHierarchy::FGetChunksBatch::CompleteRequest(FCacheGetChunkRespon
 	}
 }
 
+void FCacheStoreHierarchy::FGetChunksBatch::FinishRequest(FCacheGetChunkResponse&& Response)
+{
+	const ERequestType RequestType = Response.Id.IsNull() ? ERequestType::Value : ERequestType::Record;
+	Hierarchy.RecordStats(Response.Key.Bucket, RequestType, ERequestOp::GetChunk, Response.Status, Response.RawData.GetSize(), 0);
+	OnComplete(MoveTemp(Response));
+}
+
 void FCacheStoreHierarchy::GetChunks(
 	TConstArrayView<FCacheGetChunkRequest> Requests,
 	IRequestOwner& Owner,
@@ -1423,7 +1487,7 @@ void FCacheStoreHierarchy::GetChunks(
 
 ///////////////////////////////////////////////////////////////////////////////////////////////////////////////////////
 
-static void ConvertToLegacyStats(FDerivedDataCacheStatsNode& OutNode, const FCacheStoreStats& Stats)
+static void ConvertToLegacyStats(FDerivedDataCacheStatsNode& OutNode, FCacheStoreStats& Stats)
 {
 	EDerivedDataCacheStatus StatusCode;
 	switch (Stats.StatusCode)
@@ -1456,62 +1520,10 @@ static void ConvertToLegacyStats(FDerivedDataCacheStatsNode& OutNode, const FCac
 	}
 #endif
 
-	const auto MeasureActiveWallTime = [](TConstArrayView<FCacheStoreStats::FRateStat> History) -> FMonotonicTimeSpan
-	{
-		TSortedMap<FMonotonicTimePoint, FMonotonicTimePoint, TInlineAllocator<FCacheStoreStats::SpeedStatsMax>> WallTimes;
-		for (const FCacheStoreStats::FRateStat& RateStat : History)
-		{
-			WallTimes.Add(RateStat.StartTime, RateStat.EndTime);
-		}
-
-		FMonotonicTimeSpan Duration;
-		FMonotonicTimePoint ActiveStartTime = FMonotonicTimePoint::Infinity();
-		FMonotonicTimePoint ActiveEndTime;
-		for (const TTuple<FMonotonicTimePoint, FMonotonicTimePoint>& WallTime : WallTimes)
-		{
-			if (ActiveStartTime.IsInfinity())
-			{
-				ActiveStartTime = WallTime.Key;
-				ActiveEndTime = WallTime.Value;
-			}
-			if (ActiveEndTime < WallTime.Key)
-			{
-				Duration += ActiveEndTime - ActiveStartTime;
-				ActiveStartTime = WallTime.Key;
-			}
-			if (ActiveEndTime < WallTime.Value)
-			{
-				ActiveEndTime = WallTime.Value;
-			}
-		}
-		Duration += ActiveEndTime - ActiveStartTime;
-		return Duration;
-	};
-
-	FDerivedDataCacheSpeedStats& SpeedStats = OutNode.SpeedStats;
-
-	if (TConstArrayView<FMonotonicTimeSpan> LatencyHistory(&Stats.LatencyHistory[0], Stats.LatencyCount); !LatencyHistory.IsEmpty())
-	{
-		SpeedStats.LatencyMS = Algo::Accumulate(LatencyHistory, FMonotonicTimeSpan::Zero()).ToMilliseconds() / LatencyHistory.Num();
-	}
-
-	if (TConstArrayView<FCacheStoreStats::FRateStat> PhysicalReadSizeHistory(&Stats.PhysicalReadSizeHistory[0], Stats.PhysicalReadSizeCount); !PhysicalReadSizeHistory.IsEmpty())
-	{
-		const uint64 PhysicalReadSize = Algo::TransformAccumulate(PhysicalReadSizeHistory, &FCacheStoreStats::FRateStat::Amount, uint64(0));
-		if (const FMonotonicTimeSpan WallTime = MeasureActiveWallTime(PhysicalReadSizeHistory); !WallTime.IsZero())
-		{
-			SpeedStats.ReadSpeedMBs = double(PhysicalReadSize) / WallTime.ToSeconds() / 1024.0 / 1024.0;
-		}
-	}
-
-	if (TConstArrayView<FCacheStoreStats::FRateStat> PhysicalWriteSizeHistory(&Stats.PhysicalWriteSizeHistory[0], Stats.PhysicalWriteSizeCount); !PhysicalWriteSizeHistory.IsEmpty())
-	{
-		const uint64 PhysicalWriteSize = Algo::TransformAccumulate(PhysicalWriteSizeHistory, &FCacheStoreStats::FRateStat::Amount, uint64(0));
-		if (const FMonotonicTimeSpan WallTime = MeasureActiveWallTime(PhysicalWriteSizeHistory); !WallTime.IsZero())
-		{
-			SpeedStats.WriteSpeedMBs = double(PhysicalWriteSize) / WallTime.ToSeconds() / 1024.0 / 1024.0;
-		}
-	}
+	FMonotonicTimePoint Now = FMonotonicTimePoint::Now();
+	OutNode.SpeedStats.LatencyMS = Stats.AverageLatency.GetValue(Now);
+	OutNode.SpeedStats.ReadSpeedMBs = Stats.AveragePhysicalReadSize.GetRate(Now) / 1024.0 / 1024.0;
+	OutNode.SpeedStats.WriteSpeedMBs = Stats.AveragePhysicalWriteSize.GetRate(Now) / 1024.0 / 1024.0;
 }
 
 void FCacheStoreHierarchy::LegacyStats(FDerivedDataCacheStatsNode& OutNode)
@@ -1520,7 +1532,7 @@ void FCacheStoreHierarchy::LegacyStats(FDerivedDataCacheStatsNode& OutNode)
 	OutNode.Children.Reserve(Nodes.Num());
 	for (const FCacheStoreNode& Node : Nodes)
 	{
-		for (const FCacheStoreStats* Stats : Node.CacheStats)
+		for (FCacheStoreStats* Stats : Node.CacheStats)
 		{
 			ConvertToLegacyStats(OutNode.Children.Add_GetRef(MakeShared<FDerivedDataCacheStatsNode>()).Get(), *Stats);
 		}
@@ -1532,6 +1544,48 @@ bool FCacheStoreHierarchy::LegacyDebugOptions(FBackendDebugOptions& Options)
 	return false;
 }
 
+void FCacheStoreHierarchy::RecordStats(FCacheBucket Bucket, ERequestType Type, ERequestOp Op, EStatus Status, uint64 LogicalReadSize, uint64 LogicalWriteSize)
+{
+	// Accumulate only request count and logical size from the hierarchy.
+	// Physical size and time are tracked by the individual cache stores.
+
+	FCacheBucketStats& BucketStats = CacheStats.GetBucket(Bucket);
+	TUniqueLock Lock(BucketStats.Mutex);
+	BucketStats.LogicalReadSize += LogicalReadSize;
+	BucketStats.LogicalWriteSize += LogicalWriteSize;
+	BucketStats.RequestCount.AddRequest(Type, Op, Status);
+
+#if ENABLE_COOK_STATS
+	using FCallStats = FCookStats::CallStats;
+	using EHitOrMiss = FCallStats::EHitOrMiss;
+	using EStatType = FCallStats::EStatType;
+
+	const bool bIsInGameThread = IsInGameThread();
+	const EHitOrMiss HitOrMiss = Status == EStatus::Ok ? EHitOrMiss::Hit : EHitOrMiss::Miss;
+	FCallStats& CallStats = (Op == ERequestOp::Put) ? BucketStats.PutStats : BucketStats.GetStats;
+	CallStats.Accumulate(HitOrMiss, EStatType::Counter, 1, bIsInGameThread);
+#endif
+}
+
+uint64 FCacheStoreHierarchy::MeasureLogicalRecordSize(const FCacheRecord& Record)
+{
+	uint64 LogicalSize = 0;
+	if (const FCbObject& Meta = Record.GetMeta())
+	{
+		LogicalSize += Meta.GetSize();
+	}
+	for (const FValueWithId& Value : Record.GetValues())
+	{
+		LogicalSize += MeasureLogicalValueSize(Value);
+	}
+	return LogicalSize;
+}
+
+uint64 FCacheStoreHierarchy::MeasureLogicalValueSize(const FValue& Value)
+{
+	return Value.HasData() ? Value.GetRawSize() : 0;
+}
+
 ///////////////////////////////////////////////////////////////////////////////////////////////////////////////////////
 
 ILegacyCacheStore* CreateCacheStoreHierarchy(ICacheStoreOwner*& OutOwner, TFunctionRef<void (IMemoryCacheStore*&)> MemoryCacheCreator)
diff --git a/Engine/Source/Developer/DerivedDataCache/Private/HttpCacheStore.cpp b/Engine/Source/Developer/DerivedDataCache/Private/HttpCacheStore.cpp
--- a/Engine/Source/Developer/DerivedDataCache/Private/HttpCacheStore.cpp
+++ b/Engine/Source/Developer/DerivedDataCache/Private/HttpCacheStore.cpp
@@ -1091,6 +1091,7 @@ void FHttpCacheStore::FGetRecordOp::GetRecordOnly(const FCacheKey& InKey, const
 
 	Key = InKey;
 	OnRecordComplete = MoveTemp(InOnComplete);
+	RequestStats.Bucket = Key.Bucket;
 
 	TAnsiStringBuilder<64> Bucket;
 	Algo::Transform(Key.Bucket.ToString(), AppendChars(Bucket), FCharAnsi::ToLower);
@@ -1787,6 +1788,7 @@ void FHttpCacheStore::FExistsBatchOp::EndExists(TUniquePtr<FHttpOperation> Opera
 		{
 			UE_LOG(LogDerivedDataCache, Log, TEXT("%s: Cache miss with corrupt response for %s from '%s'."),
 				*CacheStore.Domain, *WriteToString<96>(Request.Key), *Request.Name);
+			RequestStats.Bucket = Request.Key.Bucket;
 			EndRequest(Request, {}, EStatus::Error);
 		}
 		return;
@@ -1804,6 +1806,7 @@ void FHttpCacheStore::FExistsBatchOp::EndExists(TUniquePtr<FHttpOperation> Opera
 		{
 			UE_LOG(LogDerivedDataCache, Display, TEXT("%s: Cache miss with invalid response for %s from '%s'"),
 				*CacheStore.Domain, *WriteToString<96>(Request.Key), *Request.Name);
+			RequestStats.Bucket = Request.Key.Bucket;
 			EndRequest(Request, {}, EStatus::Error);
 		}
 		return;
@@ -1824,6 +1827,7 @@ void FHttpCacheStore::FExistsBatchOp::EndExists(TUniquePtr<FHttpOperation> Opera
 		}
 
 		const FCacheGetValueRequest& Request = Requests[int32(OpId)];
+		RequestStats.Bucket = Request.Key.Bucket;
 
 		if (StatusCode < 200 || StatusCode > 204)
 		{
@@ -2671,6 +2675,9 @@ void FHttpCacheStore::GetChunks(
 				}
 				if (bOpUsed)
 				{
+					FRequestStats& RequestStats = Op->EditStats();
+					RequestStats.Type = ERequestType::Record;
+					RequestStats.Op = ERequestOp::GetChunk;
 					Op->RecordStats(OpStatus);
 					TRACE_COUNTER_ADD(HttpDDC_BytesReceived, Op->ReadStats().PhysicalReadSize);
 					TRACE_COUNTER_ADD(HttpDDC_BytesSent, Op->ReadStats().PhysicalWriteSize);
diff --git a/Engine/Source/Developer/DerivedDataCache/Public/DerivedDataCacheKey.h b/Engine/Source/Developer/DerivedDataCache/Public/DerivedDataCacheKey.h
--- a/Engine/Source/Developer/DerivedDataCache/Public/DerivedDataCacheKey.h
+++ b/Engine/Source/Developer/DerivedDataCache/Public/DerivedDataCacheKey.h
@@ -46,6 +46,13 @@ public:
 	UE_API explicit FCacheBucket(FUtf8StringView Name);
 	UE_API explicit FCacheBucket(FWideStringView Name);
 
+	/** Create a cache bucket from a name and a display name. */
+	UE_API FCacheBucket(FUtf8StringView Name, FStringView DisplayName);
+	UE_API FCacheBucket(FWideStringView Name, FStringView DisplayName);
+
+	/** Append the display name if available or the bucket name if not. */
+	UE_API void ToDisplayName(FStringBuilderBase& OutDisplayName) const;
+
 	/** Whether this is null. */
 	inline bool IsNull() const { return !Name; }
 	/** Whether this is not null. */
diff --git a/Engine/Source/Editor/UnrealEd/Private/Cooker/CookProfiling.cpp b/Engine/Source/Editor/UnrealEd/Private/Cooker/CookProfiling.cpp
--- a/Engine/Source/Editor/UnrealEd/Private/Cooker/CookProfiling.cpp
+++ b/Engine/Source/Editor/UnrealEd/Private/Cooker/CookProfiling.cpp
@@ -44,6 +44,7 @@
 #include "IAnalyticsProviderET.h"
 #include "StudioAnalytics.h"
 #include "DerivedDataCacheInterface.h"
+#include "DerivedDataCacheUsageStats.h"
 #include "Virtualization/VirtualizationSystem.h"
 #include "Experimental/ZenServerInterface.h"
 #endif
@@ -905,33 +906,6 @@ void LogCookStats(ECookMode::Type CookMode)
 		return;
 	}
 
-	/** Used for custom logging of DDC Resource usage stats. */
-	struct FDDCResourceUsageStat
-	{
-	public:
-		FDDCResourceUsageStat(FString InAssetType, double InTotalTimeSec, bool bIsGameThreadTime, double InSizeMB, int64 InAssetsBuilt) : AssetType(MoveTemp(InAssetType)), TotalTimeSec(InTotalTimeSec), GameThreadTimeSec(bIsGameThreadTime ? InTotalTimeSec : 0.0), SizeMB(InSizeMB), AssetsBuilt(InAssetsBuilt) {}
-		void Accumulate(const FDDCResourceUsageStat& OtherStat)
-		{
-			TotalTimeSec += OtherStat.TotalTimeSec;
-			GameThreadTimeSec += OtherStat.GameThreadTimeSec;
-			SizeMB += OtherStat.SizeMB;
-			AssetsBuilt += OtherStat.AssetsBuilt;
-		}
-		FString AssetType;
-		double TotalTimeSec;
-		double GameThreadTimeSec;
-		double SizeMB;
-		int64 AssetsBuilt;
-	};
-
-	/** Used for custom TSet comparison of DDC Resource usage stats. */
-	struct FDDCResourceUsageStatKeyFuncs : BaseKeyFuncs<FDDCResourceUsageStat, FString, false>
-	{
-		static const FString& GetSetKey(const FDDCResourceUsageStat& Element) { return Element.AssetType; }
-		static bool Matches(const FString& A, const FString& B) { return A == B; }
-		static uint32 GetKeyHash(const FString& Key) { return GetTypeHash(Key); }
-	};
-
 	/** Used to store profile data for custom logging. */
 	struct FCookProfileData
 	{
@@ -943,7 +917,7 @@ void LogCookStats(ECookMode::Type CookMode)
 	};
 
 	// instead of printing the usage stats generically, we capture them so we can log a subset of them in an easy-to-read way.
-	TSet<FDDCResourceUsageStat, FDDCResourceUsageStatKeyFuncs> DDCResourceUsageStats;
+	TArray<FDerivedDataCacheResourceStat> DDCResourceUsageStats;
 	TArray<FCookStatsManager::StringKeyValue> DDCSummaryStats;
 	TArray<FCookProfileData> CookProfileData;
 	TArray<FString> StatCategories;
@@ -963,54 +937,7 @@ void LogCookStats(ECookMode::Type CookMode)
 		}
 		else if (StatName.EndsWith(TEXT(".Usage"), ESearchCase::IgnoreCase))
 		{
-			// Anything that ends in .Usage is assumed to be an instance of FCookStats.FDDCResourceUsageStats. We'll log that using custom formatting.
-			FString AssetType = StatName;
-			AssetType.RemoveFromEnd(TEXT(".Usage"), ESearchCase::IgnoreCase);
-			// See if the asset has a subtype (found via the "Node" parameter")
-			const FCookStatsManager::StringKeyValue* AssetSubType = StatAttributes.FindByPredicate([](const FCookStatsManager::StringKeyValue& Item) { return Item.Key == TEXT("Node"); });
-			if (AssetSubType && AssetSubType->Value.Len() > 0)
-			{
-				AssetType += FString::Printf(TEXT(" (%s)"), *AssetSubType->Value);
-			}
-			// Pull the Time and Size attributes and AddOrAccumulate them into the set of stats. Ugly string/container manipulation code courtesy of UE/C++.
-			const FCookStatsManager::StringKeyValue* AssetTimeSecAttr = StatAttributes.FindByPredicate([](const FCookStatsManager::StringKeyValue& Item) { return Item.Key == TEXT("TimeSec"); });
-			double AssetTimeSec = 0.0;
-			if (AssetTimeSecAttr)
-			{
-				LexFromString(AssetTimeSec, *AssetTimeSecAttr->Value);
-			}
-			const FCookStatsManager::StringKeyValue* AssetSizeMBAttr = StatAttributes.FindByPredicate([](const FCookStatsManager::StringKeyValue& Item) { return Item.Key == TEXT("MB"); });
-			double AssetSizeMB = 0.0;
-			if (AssetSizeMBAttr)
-			{
-				LexFromString(AssetSizeMB, *AssetSizeMBAttr->Value);
-			}
-			const FCookStatsManager::StringKeyValue* ThreadNameAttr = StatAttributes.FindByPredicate([](const FCookStatsManager::StringKeyValue& Item) { return Item.Key == TEXT("ThreadName"); });
-			bool bIsGameThreadTime = ThreadNameAttr != nullptr && ThreadNameAttr->Value == TEXT("GameThread");
-
-			const FCookStatsManager::StringKeyValue* HitOrMissAttr = StatAttributes.FindByPredicate([](const FCookStatsManager::StringKeyValue& Item) { return Item.Key == TEXT("HitOrMiss"); });
-			bool bWasMiss = HitOrMissAttr != nullptr && HitOrMissAttr->Value == TEXT("Miss");
-			int64 AssetsBuilt = 0;
-			if (bWasMiss)
-			{
-				const FCookStatsManager::StringKeyValue* CountAttr = StatAttributes.FindByPredicate([](const FCookStatsManager::StringKeyValue& Item) { return Item.Key == TEXT("Count"); });
-				if (CountAttr)
-				{
-					LexFromString(AssetsBuilt, *CountAttr->Value);
-				}
-			}
-
-
-			FDDCResourceUsageStat Stat(AssetType, AssetTimeSec, bIsGameThreadTime, AssetSizeMB, AssetsBuilt);
-			FDDCResourceUsageStat* ExistingStat = DDCResourceUsageStats.Find(Stat.AssetType);
-			if (ExistingStat)
-			{
-				ExistingStat->Accumulate(Stat);
-			}
-			else
-			{
-				DDCResourceUsageStats.Add(Stat);
-			}
+			// These are gathered through GatherResourceStats.
 			bStatUsedCustomFormatting = true;
 		}
 		else if (StatName == TEXT("DDC.Summary"))
@@ -1039,6 +966,8 @@ void LogCookStats(ECookMode::Type CookMode)
 		}
 	};
 
+	GetDerivedDataCacheRef().GatherResourceStats(DDCResourceUsageStats);
+
 	FCookStatsManager::LogCookStats(LogStatsFunc);
 
 	UE_LOG(LogCook, Display, TEXT("Misc Cook Stats"));
@@ -1119,28 +1048,20 @@ void LogCookStats(ECookMode::Type CookMode)
 
 	DumpDerivedDataBuildRemoteExecutorStats();
 
-	if (DDCResourceUsageStats.Num() > 0)
+	if (!DDCResourceUsageStats.IsEmpty())
 	{
-		// sort the list
-		TArray<FDDCResourceUsageStat> SortedDDCResourceUsageStats;
-		SortedDDCResourceUsageStats.Empty(DDCResourceUsageStats.Num());
-		for (const FDDCResourceUsageStat& Stat : DDCResourceUsageStats)
-		{
-			SortedDDCResourceUsageStats.Emplace(Stat);
-		}
-		SortedDDCResourceUsageStats.Sort([](const FDDCResourceUsageStat& LHS, const FDDCResourceUsageStat& RHS)
-			{
-				return LHS.TotalTimeSec > RHS.TotalTimeSec;
-			});
+		Algo::SortBy(DDCResourceUsageStats, [](const FDerivedDataCacheResourceStat& Stat) { return Stat.BuildTimeSec + Stat.LoadTimeSec; }, TGreater());
 
 		UE_LOG(LogCook, Display, TEXT(""));
 		UE_LOG(LogCook, Display, TEXT("DDC Resource Stats"));
 		UE_LOG(LogCook, Display, TEXT("======================================================================================================="));
 		UE_LOG(LogCook, Display, TEXT("Asset Type                          Total Time (Sec)  GameThread Time (Sec)  Assets Built  MB Processed"));
 		UE_LOG(LogCook, Display, TEXT("----------------------------------  ----------------  ---------------------  ------------  ------------"));
-		for (const FDDCResourceUsageStat& Stat : SortedDDCResourceUsageStats)
+		for (const FDerivedDataCacheResourceStat& Stat : DDCResourceUsageStats)
 		{
-			UE_LOG(LogCook, Display, TEXT("%-34s  %16.2f  %21.2f  %12d  %12.2f"), *Stat.AssetType, Stat.TotalTimeSec, Stat.GameThreadTimeSec, Stat.AssetsBuilt, Stat.SizeMB);
+			UE_LOG(LogCook, Display, TEXT("%-34s  %16.2f  %21.2f  %12d  %12.2f"),
+				*Stat.AssetType, Stat.LoadTimeSec + Stat.BuildTimeSec, Stat.GameThreadTimeSec,
+				Stat.BuildCount, Stat.LoadSizeMB + Stat.BuildSizeMB);
 		}
 	}
 
