DDC: Fixed missing cache bucket for stats

#rnx

diff --git a/Engine/Source/Developer/DerivedDataCache/Private/HttpCacheStore.cpp b/Engine/Source/Developer/DerivedDataCache/Private/HttpCacheStore.cpp
--- a/Engine/Source/Developer/DerivedDataCache/Private/HttpCacheStore.cpp
+++ b/Engine/Source/Developer/DerivedDataCache/Private/HttpCacheStore.cpp
@@ -1776,6 +1776,7 @@ void FHttpCacheStore::FExistsBatchOp::EndExists(TUniquePtr<FHttpOperation> Opera
 		{
 			UE_LOG(LogDerivedDataCache, Verbose, TEXT("%s: Cache miss with failed HTTP request for %s from '%s'"),
 				*CacheStore.Domain, *WriteToString<96>(Request.Key), *Request.Name);
+			RequestStats.Bucket = Request.Key.Bucket;
 			EndRequest(Request, {}, EStatus::Error);
 		}
 		return;
