DDC: Fixed benign race caught by TSan in FRequestOwnerShared

#rnx

diff --git a/Engine/Source/Developer/DerivedDataCache/Private/DerivedDataRequestOwner.cpp b/Engine/Source/Developer/DerivedDataCache/Private/DerivedDataRequestOwner.cpp
--- a/Engine/Source/Developer/DerivedDataCache/Private/DerivedDataRequestOwner.cpp
+++ b/Engine/Source/Developer/DerivedDataCache/Private/DerivedDataRequestOwner.cpp
@@ -29,8 +29,8 @@ public:
 	void BeginBarrier(ERequestBarrierFlags Flags) final;
 	void EndBarrier(ERequestBarrierFlags Flags) final;
 
-	inline EPriority GetPriority() const final { return Priority; }
-	inline bool IsCanceled() const final { return bIsCanceled; }
+	inline EPriority GetPriority() const final { return Priority.load(std::memory_order_relaxed); }
+	inline bool IsCanceled() const final { return bIsCanceled.load(std::memory_order_relaxed); }
 
 	inline void SetPriority(EPriority Priority) final;
 	inline void Cancel() final;
@@ -46,11 +46,11 @@ private:
 	FEventCount BarrierEvent;
 	uint32 BarrierCount{0};
 	uint32 PriorityBarrierCount{0};
-	EPriority Priority{EPriority::Normal};
+	std::atomic<EPriority> Priority{EPriority::Normal};
+	std::atomic<bool> bIsCanceled{false};
+	std::atomic<bool> bKeepAlive{false};
 	uint8 bPriorityChangedInBarrier : 1;
 	uint8 bBeginExecuted : 1;
-	bool bIsCanceled{false};
-	bool bKeepAlive{false};
 };
 
 FRequestOwnerShared::FRequestOwnerShared(EPriority NewPriority)
@@ -77,7 +77,7 @@ void FRequestOwnerShared::Begin(IRequest* Request)
 		{
 			return;
 		}
-		NewPriority = Priority;
+		NewPriority = GetPriority();
 	}
 	// Loop until priority is stable. Another thread may be changing the priority concurrently.
 	for (EPriority CheckPriority; ; NewPriority = CheckPriority)
@@ -86,7 +86,7 @@ void FRequestOwnerShared::Begin(IRequest* Request)
 		
 		{
 			FReadScopeLock ScopeLock(Lock);
-			CheckPriority = Priority;
+			CheckPriority = GetPriority();
 		}
 
 		if (CheckPriority == NewPriority)
@@ -143,13 +143,13 @@ void FRequestOwnerShared::EndBarrier(ERequestBarrierFlags Flags)
 void FRequestOwnerShared::SetPriority(EPriority NewPriority)
 {
 	TArray<TRefCountPtr<IRequest>, TInlineAllocator<16>> LocalRequests;
-	if (FWriteScopeLock WriteLock(Lock); Priority == NewPriority)
+	if (FWriteScopeLock WriteLock(Lock); GetPriority() == NewPriority)
 	{
 		return;
 	}
 	else
 	{
-		Priority = NewPriority;
+		Priority.store(NewPriority, std::memory_order_relaxed);
 		LocalRequests = Requests;
 		bPriorityChangedInBarrier = (PriorityBarrierCount > 0);
 	}
@@ -169,7 +169,7 @@ void FRequestOwnerShared::Cancel()
 
 		{
 			FWriteScopeLock WriteLock(Lock);
-			bIsCanceled = true;
+			bIsCanceled.store(true, std::memory_order_relaxed);
 
 			if (!Requests.IsEmpty())
 			{
@@ -239,18 +239,12 @@ bool FRequestOwnerShared::Poll() const
 
 void FRequestOwnerShared::KeepAlive()
 {
-	FWriteScopeLock WriteLock(Lock);
 	bKeepAlive = true;
 }
 
 void FRequestOwnerShared::Destroy()
 {
-	bool bLocalKeepAlive;
-	{
-		FWriteScopeLock ScopeLock(Lock);
-		bLocalKeepAlive = bKeepAlive;
-	}
-	if (!bLocalKeepAlive)
+	if (!bKeepAlive)
 	{
 		Cancel();
 	}
