Change HTTP cache store to allow GetChunks to execute asynchronous network operations.

#rnx

diff --git a/Engine/Source/Developer/DerivedDataCache/Private/HttpCacheStore.cpp b/Engine/Source/Developer/DerivedDataCache/Private/HttpCacheStore.cpp
--- a/Engine/Source/Developer/DerivedDataCache/Private/HttpCacheStore.cpp
+++ b/Engine/Source/Developer/DerivedDataCache/Private/HttpCacheStore.cpp
@@ -373,6 +373,19 @@ private:
 		uint64 UserData,
 		FOnCacheGetValueComplete&& OnComplete);
 
+	void FinishChunkRequest(
+		const FCacheGetChunkRequest& Request,
+		EStatus Status,
+		const FValue& Value,
+		FCompressedBufferReader& ValueReader,
+		const TSharedRef<FOnCacheGetChunkComplete>& SharedOnComplete);
+
+	void GetChunkGroupAsync(
+		IRequestOwner& Owner,
+		const FCacheGetChunkRequest* StartRequest,
+		const FCacheGetChunkRequest* EndRequest,
+		TSharedRef<FOnCacheGetChunkComplete>& SharedOnComplete);
+
 	class FHealthCheckOp;
 	class FPutPackageOp;
 	class FGetRecordOp;
@@ -1031,6 +1044,11 @@ public:
 	FRequestStats& EditStats() { return RequestStats; }
 	void RecordStats(EStatus Status);
 
+	int32 GetFailedValues() const { return FailedValues; }
+	void PrepareForPendingValues(int32 InPendingValues) { PendingValues = InPendingValues; }
+	bool FinishPendingValueFetch(const FValueWithId& Value, bool bAppendToPackage);
+	bool FinishPendingValueExists(EStatus Status);
+
 private:
 	FGetRecordOp(FHttpCacheStore& CacheStore, IRequestOwner& Owner, const FSharedString& Name);
 
@@ -1175,6 +1193,35 @@ void FHttpCacheStore::FGetRecordOp::GetRecord(const FCacheKey& LocalKey, const F
 	});
 }
 
+bool FHttpCacheStore::FGetRecordOp::FinishPendingValueFetch(const FValueWithId& Value, bool bAppendToPackage)
+{
+	TDynamicUniqueLock Lock(Mutex);
+	const bool bComplete = --PendingValues == 0;
+	if (Value.HasData())
+	{
+		if (bAppendToPackage)
+		{
+			Package.AddAttachment(FCbAttachment(Value.GetData()));
+		}
+	}
+	else
+	{
+		++FailedValues;
+	}
+	return bComplete;
+}
+
+bool FHttpCacheStore::FGetRecordOp::FinishPendingValueExists(EStatus Status)
+{
+	TDynamicUniqueLock Lock(Mutex);
+	const bool bComplete = --PendingValues == 0;
+	if (Status != EStatus::Ok)
+	{
+		++FailedValues;
+	}
+	return bComplete;
+}
+
 void FHttpCacheStore::FGetRecordOp::BeginGetValues(const FCacheRecord& Record, const FCacheRecordPolicy& Policy, FOnRecordComplete&& OnComplete)
 {
 	FRequestTimer RequestTimer(RequestStats);
@@ -1193,7 +1240,7 @@ void FHttpCacheStore::FGetRecordOp::BeginGetValues(const FCacheRecord& Record, c
 		}
 	}
 
-	PendingValues = RequiredGets.Num() + RequiredHeads.Num();
+	PrepareForPendingValues(RequiredGets.Num() + RequiredHeads.Num());
 
 	RequestTimer.Stop();
 
@@ -1204,34 +1251,16 @@ void FHttpCacheStore::FGetRecordOp::BeginGetValues(const FCacheRecord& Record, c
 
 	GetValues(RequiredGets, [Self = TRefCountPtr(this), Policy](FValueResponse&& Response)
 	{
-		TDynamicUniqueLock Lock(Self->Mutex);
-		const bool bComplete = --Self->PendingValues == 0;
-		if (Response.Value.HasData())
-		{
-			Self->Package.AddAttachment(FCbAttachment(Response.Value.GetData()));
-		}
-		else
+		if (Self->FinishPendingValueFetch(Response.Value, true))
 		{
-			++Self->FailedValues;
-		}
-		if (bComplete)
-		{
-			Lock.Unlock();
 			Self->EndGetValues(Policy, Response.Status);
 		}
 	});
 
 	GetValuesExist(RequiredHeads, [Self = TRefCountPtr(this), Policy](FValueResponse&& Response)
 	{
-		TDynamicUniqueLock Lock(Self->Mutex);
-		const bool bComplete = --Self->PendingValues == 0;
-		if (Response.Status != EStatus::Ok)
-		{
-			++Self->FailedValues;
-		}
-		if (bComplete)
+		if (Self->FinishPendingValueExists(Response.Status))
 		{
-			Lock.Unlock();
 			Self->EndGetValues(Policy, Response.Status);
 		}
 	});
@@ -2461,6 +2490,237 @@ void FHttpCacheStore::GetCacheRecordAsync(
 	});
 }
 
+void FHttpCacheStore::FinishChunkRequest(
+	const FCacheGetChunkRequest& Request,
+	EStatus Status,
+	const FValue& Value,
+	FCompressedBufferReader& ValueReader,
+	const TSharedRef<FOnCacheGetChunkComplete>& SharedOnComplete)
+{
+	if (Status == EStatus::Ok)
+	{
+		const uint64 RawOffset = FMath::Min(Value.GetRawSize(), Request.RawOffset);
+		const uint64 RawSize = FMath::Min(Value.GetRawSize() - RawOffset, Request.RawSize);
+		UE_LOG(LogDerivedDataCache, Verbose, TEXT("%s: Cache hit for %s from '%s'"),
+			*Domain, *WriteToString<96>(Request.Key, '/', Request.Id), *Request.Name);
+		FSharedBuffer Buffer;
+		const bool bExistsOnly = EnumHasAnyFlags(Request.Policy, ECachePolicy::SkipData);
+		if (!bExistsOnly)
+		{
+			Buffer = ValueReader.Decompress(RawOffset, RawSize);
+		}
+		const EStatus ChunkStatus = bExistsOnly || Buffer.GetSize() == RawSize ? EStatus::Ok : EStatus::Error;
+		if (ChunkStatus == EStatus::Ok)
+		{
+			TRACE_COUNTER_INCREMENT(HttpDDC_GetHit);
+		}
+		SharedOnComplete.Get()({ Request.Name, Request.Key, Request.Id, Request.RawOffset,
+			RawSize, Value.GetRawHash(), MoveTemp(Buffer), Request.UserData, ChunkStatus });
+	}
+	else
+	{
+		UE_LOG(LogDerivedDataCache, Verbose, TEXT("%s: Cache miss for %s from '%s'"),
+			*Domain, *WriteToString<96>(Request.Key, '/', Request.Id), *Request.Name);
+
+		SharedOnComplete.Get()(Request.MakeResponse(Status));
+	}
+}
+
+static void AppendGetAndHeadOpsForChunkRequestGroupItem(
+	const FCacheGetChunkRequest& Request,
+	const FValueWithId& ValueWithId,
+	TArray<FValueWithId>& RequiredGets,
+	TArray<TArray<FCacheGetChunkRequest>>& RequiredGetRequests,
+	TArray<FValueWithId>& RequiredHeads,
+	TArray<TArray<FCacheGetChunkRequest>>& RequiredHeadRequests)
+{
+	const bool bAlreadyRequiredGet = !RequiredGets.IsEmpty() && RequiredGets.Last() == ValueWithId;
+	const bool bAlreadyRequiredHead = !RequiredHeads.IsEmpty() && RequiredHeads.Last() == ValueWithId;
+	if (EnumHasAnyFlags(Request.Policy, ECachePolicy::SkipData))
+	{
+		if (!bAlreadyRequiredHead && !bAlreadyRequiredGet)
+		{
+			RequiredHeads.Emplace(ValueWithId);
+			RequiredHeadRequests.AddDefaulted();
+		}
+		if (bAlreadyRequiredGet)
+		{
+			RequiredGetRequests.Last().Add(Request);
+		}
+		else
+		{
+			RequiredHeadRequests.Last().Add(Request);
+		}
+	}
+	else
+	{
+		if (!bAlreadyRequiredGet)
+		{
+			RequiredGets.Emplace(ValueWithId);
+			if (bAlreadyRequiredHead)
+			{
+				//Steal existing head contents first
+				RequiredGetRequests.Emplace(MoveTemp(RequiredHeadRequests.Last()));
+				RequiredHeads.SetNum(RequiredHeads.Num() - 1, false);
+			}
+			else
+			{
+				RequiredGetRequests.AddDefaulted();
+			}
+		}
+
+		RequiredGetRequests.Last().Add(Request);
+	}
+}
+
+void FHttpCacheStore::GetChunkGroupAsync(
+	IRequestOwner& Owner,
+	const FCacheGetChunkRequest* StartRequest,
+	const FCacheGetChunkRequest* EndRequest,
+	TSharedRef<FOnCacheGetChunkComplete>& SharedOnComplete)
+{
+	if ((StartRequest == nullptr) || (StartRequest >= EndRequest))
+	{
+		return;
+	}
+
+	ECachePolicy GroupPolicy(ECachePolicy::None);
+	TArray<FCacheGetChunkRequest> RequestGroup;
+	RequestGroup.Reserve(static_cast<int>(EndRequest - StartRequest));
+	for (const FCacheGetChunkRequest* Request = StartRequest; Request != EndRequest; ++Request)
+	{
+		RequestGroup.Add(*Request);
+		GroupPolicy |= Request->Policy;
+	}
+
+	if (StartRequest->Id.IsValid())
+	{
+		// Get Record and contained Values within the request group
+		TRefCountPtr<FGetRecordOp> Op = FGetRecordOp::New(*this, Owner, StartRequest->Name);
+
+		Op->GetRecordOnly(StartRequest->Key, GroupPolicy, [this, Op = TRefCountPtr(Op), RequestGroup = MoveTemp(RequestGroup), SharedOnComplete](FGetRecordOp::FRecordResponse&& Response) mutable
+		{
+			auto RecordStats = [](FGetRecordOp& Op, FCacheBucket Bucket, EStatus Status)
+			{
+				FRequestStats& RequestStats = Op.EditStats();
+				RequestStats.Type = ERequestType::Record;
+				RequestStats.Bucket = Bucket;
+				RequestStats.Op = ERequestOp::GetChunk;
+				Op.RecordStats(Status);
+				TRACE_COUNTER_ADD(HttpDDC_BytesReceived, Op.ReadStats().PhysicalReadSize);
+				TRACE_COUNTER_ADD(HttpDDC_BytesSent, Op.ReadStats().PhysicalWriteSize);
+			};
+
+			if (Response.Status == EStatus::Ok)
+			{
+				// Get Values on the record
+				FRequestTimer RequestTimer(Op->EditStats());
+
+				TArray<FValueWithId> RequiredGets;
+				TArray<TArray<FCacheGetChunkRequest>> RequiredGetRequests;
+				TArray<FValueWithId> RequiredHeads;
+				TArray<TArray<FCacheGetChunkRequest>> RequiredHeadRequests;
+				FCompressedBufferReader NullReader;
+				for (const FCacheGetChunkRequest& Request : RequestGroup)
+				{
+					const FValueWithId& ValueWithId = Response.Record.GetValue(Request.Id);
+					bool bHasValue = ValueWithId.IsValid();
+					FValue Value = ValueWithId;
+
+					if (!bHasValue || IsValueDataReady(Value, Request.Policy))
+					{
+						FinishChunkRequest(Request, Response.Status, Value, NullReader, SharedOnComplete);
+					}
+					else
+					{
+						AppendGetAndHeadOpsForChunkRequestGroupItem(Request, ValueWithId, RequiredGets,RequiredGetRequests, RequiredHeads, RequiredHeadRequests);
+					}
+				}
+
+				int32 PendingValues = RequiredGets.Num() + RequiredHeads.Num();
+				Op->PrepareForPendingValues(PendingValues);
+
+				RequestTimer.Stop();
+
+				if (PendingValues == 0)
+				{
+					RecordStats(*Op, RequestGroup[0].Key.Bucket, Response.Status);
+					return;
+				}
+
+				Op->GetValues(RequiredGets, [this, RecordStats, Op = TRefCountPtr(Op), ChunkRequestsForValues = MoveTemp(RequiredGetRequests), SharedOnComplete](FGetRecordOp::FValueResponse&& Response)
+				{
+					int FoundRequestsIndex = Algo::BinarySearchBy(ChunkRequestsForValues, Response.Value.GetId(), [](const TArray<FCacheGetChunkRequest>& ChunkRequests)
+					{
+						check(!ChunkRequests.IsEmpty());
+						return ChunkRequests[0].Id;
+					});
+
+					check(FoundRequestsIndex != INDEX_NONE);
+					const TArray<FCacheGetChunkRequest>& ChunkRequests = ChunkRequestsForValues[FoundRequestsIndex];
+					FCompressedBufferReader ValueReader(Response.Value.GetData());
+
+					if (Op->FinishPendingValueFetch(Response.Value, false))
+					{
+						RecordStats(*Op, ChunkRequests[0].Key.Bucket, Op->GetFailedValues() > 0 ? EStatus::Error : EStatus::Ok);
+					}
+
+					for (const FCacheGetChunkRequest& ChunkRequest : ChunkRequests)
+					{
+						FinishChunkRequest(ChunkRequest, Response.Status, Response.Value, ValueReader, SharedOnComplete);
+					}
+				});
+
+				Op->GetValuesExist(RequiredHeads, [this, RecordStats, Op = TRefCountPtr(Op), ChunkRequestsForValues = MoveTemp(RequiredHeadRequests), SharedOnComplete](FGetRecordOp::FValueResponse&& Response)
+				{
+					int FoundRequestsIndex = Algo::BinarySearchBy(ChunkRequestsForValues, Response.Value.GetId(), [](const TArray<FCacheGetChunkRequest>& ChunkRequests)
+					{
+						check(!ChunkRequests.IsEmpty());
+						return ChunkRequests[0].Id;
+					});
+
+					check(FoundRequestsIndex != INDEX_NONE);
+					const TArray<FCacheGetChunkRequest>& ChunkRequests = ChunkRequestsForValues[FoundRequestsIndex];
+
+					if (Op->FinishPendingValueExists(Response.Status))
+					{
+						RecordStats(*Op, ChunkRequests[0].Key.Bucket, Op->GetFailedValues() > 0 ? EStatus::Error : EStatus::Ok);
+					}
+
+					FCompressedBufferReader NullReader;
+					for (const FCacheGetChunkRequest& ChunkRequest : ChunkRequests)
+					{
+						FinishChunkRequest(ChunkRequest, Response.Status, Response.Value, NullReader, SharedOnComplete);
+					}
+				});
+			}
+			else
+			{
+				FCompressedBufferReader NullReader;
+				FValue DummyValue;
+				for (const FCacheGetChunkRequest& Request : RequestGroup)
+				{
+					FinishChunkRequest(Request, Response.Status, DummyValue, NullReader, SharedOnComplete);
+				}
+
+				RecordStats(*Op, RequestGroup[0].Key.Bucket, Response.Status);
+			}
+		});
+	}
+	else
+	{
+		// Get Value for the request group
+		GetCacheValueAsync(Owner, StartRequest->Name, StartRequest->Key, GroupPolicy, ERequestOp::GetChunk, 0, [this, RequestGroup = MoveTemp(RequestGroup), SharedOnComplete](FCacheGetValueResponse&& Response)
+		{
+			FCompressedBufferReader ValueReader(Response.Value.GetData());
+			for (const FCacheGetChunkRequest& Request : RequestGroup)
+			{
+				FinishChunkRequest(Request, Response.Status, Response.Value, ValueReader, SharedOnComplete);
+			}
+		});
+	}
+}
+
 void FHttpCacheStore::LegacyStats(FDerivedDataCacheStatsNode& OutNode)
 {
 	checkNoEntry();
@@ -2596,7 +2856,13 @@ void FHttpCacheStore::GetChunks(
 	FOnCacheGetChunkComplete&& OnComplete)
 {
 	TRACE_CPUPROFILER_EVENT_SCOPE(HttpDDC_GetChunks);
-	TRACE_COUNTER_ADD(HttpDDC_GetHit, Requests.Num());
+	TRACE_COUNTER_ADD(HttpDDC_Get, Requests.Num());
+
+	if (Requests.IsEmpty())
+	{
+		return;
+	}
+
 	// TODO: This is inefficient because Jupiter doesn't allow us to get only part of a compressed blob, so we have to
 	//		 get the whole thing and then decompress only the portion we need.  Furthermore, because there is no propagation
 	//		 between cache stores during chunk requests, the fetched result won't end up in the local store.
@@ -2607,125 +2873,21 @@ void FHttpCacheStore::GetChunks(
 	TArray<FCacheGetChunkRequest, TInlineAllocator<16>> SortedRequests(Requests);
 	SortedRequests.StableSort(TChunkLess());
 
-	bool bHasValue = false;
-	FValue Value;
-	FValueId ValueId;
-	FCacheKey ValueKey;
-	FCompressedBuffer ValueBuffer;
-	FCompressedBufferReader ValueReader;
-	EStatus ValueStatus = EStatus::Error;
-	FOptionalCacheRecord Record;
+	const FCacheGetChunkRequest* PendingGroupStartRequest = &SortedRequests[0];
+
+	FRequestBarrier Barrier(Owner);
+	TSharedRef<FOnCacheGetChunkComplete> SharedOnComplete = MakeShared<FOnCacheGetChunkComplete>(MoveTemp(OnComplete));
 	for (const FCacheGetChunkRequest& Request : SortedRequests)
 	{
 		const bool bExistsOnly = EnumHasAnyFlags(Request.Policy, ECachePolicy::SkipData);
-		if (!(bHasValue && ValueKey == Request.Key && ValueId == Request.Id) || ValueReader.HasSource() < !bExistsOnly)
-		{
-			ValueStatus = EStatus::Error;
-			ValueReader.ResetSource();
-			ValueBuffer.Reset();
-			ValueKey = {};
-			ValueId.Reset();
-			Value.Reset();
-			bHasValue = false;
-			if (Request.Id.IsValid())
-			{
-				FRequestOwner BlockingOwner(EPriority::Blocking);
-				bool bOpUsed = false;
-				EStatus OpStatus = EStatus::Ok;
-				TRefCountPtr<FGetRecordOp> Op = FGetRecordOp::New(*this, BlockingOwner, Request.Name);
-				if (!(Record && Record.Get().GetKey() == Request.Key))
-				{
-					bOpUsed = true;
-					Op->GetRecordOnly(Request.Key, Request.Policy, [&Record, &OpStatus](FGetRecordOp::FRecordResponse&& Response)
-					{
-						Record = MoveTemp(Response.Record);
-						OpStatus = Response.Status;
-					});
-					BlockingOwner.Wait();
-				}
-				if (Record)
-				{
-					const FValueWithId& ValueWithId = Record.Get().GetValue(Request.Id);
-					bHasValue = ValueWithId.IsValid();
-					Value = ValueWithId;
-					ValueId = Request.Id;
-					ValueKey = Request.Key;
-
-					if (IsValueDataReady(Value, Request.Policy))
-					{
-						ValueBuffer = Value.GetData();
-						ValueReader.SetSource(ValueBuffer);
-					}
-					else if (bHasValue)
-					{
-						bOpUsed = true;
-						Op->GetValues({ValueWithId}, [&Value, &OpStatus](FGetRecordOp::FValueResponse&& Response)
-						{
-							Value = MoveTemp(Response.Value);
-							OpStatus = Response.Status;
-						});
-						BlockingOwner.Wait();
-
-						if (Value.HasData())
-						{
-							ValueBuffer = Value.GetData();
-							ValueReader.SetSource(ValueBuffer);
-						}
-					}
-				}
-				if (bOpUsed)
-				{
-					FRequestStats& RequestStats = Op->EditStats();
-					RequestStats.Type = ERequestType::Record;
-					RequestStats.Bucket = Request.Key.Bucket;
-					RequestStats.Op = ERequestOp::GetChunk;
-					Op->RecordStats(OpStatus);
-					TRACE_COUNTER_ADD(HttpDDC_BytesReceived, Op->ReadStats().PhysicalReadSize);
-					TRACE_COUNTER_ADD(HttpDDC_BytesSent, Op->ReadStats().PhysicalWriteSize);
-				}
-			}
-			else
-			{
-				ValueKey = Request.Key;
-
-				FRequestOwner BlockingOwner(EPriority::Blocking);
-				GetCacheValueAsync(BlockingOwner, Request.Name, Request.Key, Request.Policy, ERequestOp::GetChunk, 0, [&bHasValue, &Value](FCacheGetValueResponse&& Response)
-				{
-					bHasValue = Response.Status == EStatus::Ok;
-					Value = MoveTemp(Response.Value);
-				});
-				BlockingOwner.Wait();
-
-				if (bHasValue && IsValueDataReady(Value, Request.Policy))
-				{
-					ValueBuffer = Value.GetData();
-					ValueReader.SetSource(ValueBuffer);
-				}
-			}
-		}
-		if (bHasValue)
+		const bool bMatchesExistingGroup = PendingGroupStartRequest != nullptr && PendingGroupStartRequest->Key == Request.Key && PendingGroupStartRequest->Id.IsValid() == Request.Id.IsValid();
+		if (!bMatchesExistingGroup)
 		{
-			const uint64 RawOffset = FMath::Min(Value.GetRawSize(), Request.RawOffset);
-			const uint64 RawSize = FMath::Min(Value.GetRawSize() - RawOffset, Request.RawSize);
-			UE_LOG(LogDerivedDataCache, Verbose, TEXT("%s: Cache hit for %s from '%s'"),
-				*Domain, *WriteToString<96>(Request.Key, '/', Request.Id), *Request.Name);
-			FSharedBuffer Buffer;
-			if (!bExistsOnly)
-			{
-				Buffer = ValueReader.Decompress(RawOffset, RawSize);
-			}
-			const EStatus ChunkStatus = bExistsOnly || Buffer.GetSize() == RawSize ? EStatus::Ok : EStatus::Error;
-			if (ChunkStatus == EStatus::Ok)
-			{
-				TRACE_COUNTER_INCREMENT(HttpDDC_GetHit);
-			}
-			OnComplete({Request.Name, Request.Key, Request.Id, Request.RawOffset,
-				RawSize, Value.GetRawHash(), MoveTemp(Buffer), Request.UserData, ChunkStatus});
-			continue;
+			GetChunkGroupAsync(Owner, PendingGroupStartRequest, &Request, SharedOnComplete);
+			PendingGroupStartRequest = &Request;
 		}
-
-		OnComplete(Request.MakeResponse(EStatus::Error));
 	}
+	GetChunkGroupAsync(Owner, PendingGroupStartRequest, SortedRequests.GetData() + SortedRequests.Num(), SharedOnComplete);
 }
 
 void FHttpCacheStoreParams::Parse(const TCHAR* NodeName, const TCHAR* Config)
