DDC: Fixed -DDC-Verify being ignored if there is no verify node in the graph

#rnx

diff --git a/Engine/Source/Developer/DerivedDataCache/Private/DerivedDataBackends.cpp b/Engine/Source/Developer/DerivedDataCache/Private/DerivedDataBackends.cpp
--- a/Engine/Source/Developer/DerivedDataCache/Private/DerivedDataBackends.cpp
+++ b/Engine/Source/Developer/DerivedDataCache/Private/DerivedDataBackends.cpp
@@ -204,16 +204,14 @@ public:
 		// Async must exist in the graph.
 		RootNode = CreateCacheStoreAsync(RootNode, MemoryCache, /*bDeleteInnerCache*/ true);
 
-		// Create a Verify node when using -DDC-Verify[=Type1[@Rate2][+Type2[@Rate2]...]].
-		if (bVerifyFound)
+		// Create a Verify node when using -DDC-Verify[=Type1[@Rate2][+Type2[@Rate2]...]] or the graph has a verify node.
+		if (FString VerifyArg;
+			FParse::Value(FCommandLine::Get(), TEXT("-DDC-Verify="), VerifyArg) ||
+			FParse::Param(FCommandLine::Get(), TEXT("DDC-Verify")) ||
+			bVerifyFound)
 		{
-			FString VerifyArg;
-			if (FParse::Value(FCommandLine::Get(), TEXT("-DDC-Verify="), VerifyArg) ||
-				FParse::Param(FCommandLine::Get(), TEXT("DDC-Verify")))
-			{
-				IFileManager::Get().DeleteDirectory(*(FPaths::ProjectSavedDir() / TEXT("VerifyDDC/")), /*bRequireExists*/ false, /*bTree*/ true);
-				RootNode = CreateCacheStoreVerify(RootNode, /*bPutOnError*/ bVerifyFix);
-			}
+			IFileManager::Get().DeleteDirectory(*(FPaths::ProjectSavedDir() / TEXT("VerifyDDC/")), /*bRequireExists*/ false, /*bTree*/ true);
+			RootNode = CreateCacheStoreVerify(RootNode, /*bPutOnError*/ bVerifyFix);
 		}
 
 		// Create a Replay node when requested on the command line.
