DDC: Cleaned up the cache replay uses of UE_LOG that were recently fixed

#rnx

diff --git a/Engine/Source/Developer/DerivedDataCache/Private/DerivedDataCacheReplay.cpp b/Engine/Source/Developer/DerivedDataCache/Private/DerivedDataCacheReplay.cpp
--- a/Engine/Source/Developer/DerivedDataCache/Private/DerivedDataCacheReplay.cpp
+++ b/Engine/Source/Developer/DerivedDataCache/Private/DerivedDataCacheReplay.cpp
@@ -167,9 +167,9 @@ void FCacheStoreReplay::SerializeRequests(
 
 	if (UE_LOG_ACTIVE(LogDerivedDataCache, Verbose))
 	{
-		TUtf8StringBuilder<1024> Batch;
+		TStringBuilder<1024> Batch;
 		CompactBinaryToCompactJson(Writer.Save().AsObject(), Batch);
-		UE_LOG(LogDerivedDataCache, Verbose, TEXT("Replay: %s"), StringCast<TCHAR>(*Batch).Get());
+		UE_LOG(LogDerivedDataCache, Verbose, TEXT("Replay: %s"), *Batch);
 	}
 
 	WriteToArchive(Writer);
@@ -405,9 +405,9 @@ bool FCacheReplayReader::FState::DispatchRequests(
 	{
 		if (UE_LOG_ACTIVE(LogDerivedDataCache, Verbose))
 		{
-			TUtf8StringBuilder<1024> Batch;
+			TStringBuilder<1024> Batch;
 			CompactBinaryToCompactJson(Object, Batch);
-			UE_LOG(LogDerivedDataCache, Verbose, TEXT("Replay: %s"), StringCast<TCHAR>(*Batch).Get());
+			UE_LOG(LogDerivedDataCache, Verbose, TEXT("Replay: %s"), *Batch);
 		}
 
 		TRACE_CPUPROFILER_EVENT_SCOPE(ReplayDDC_Dispatch);
