DDC: Fixed creation of stats when using a throttle node

diff --git a/Engine/Source/Developer/DerivedDataCache/Private/DerivedDataBackends.cpp b/Engine/Source/Developer/DerivedDataCache/Private/DerivedDataBackends.cpp
--- a/Engine/Source/Developer/DerivedDataCache/Private/DerivedDataBackends.cpp
+++ b/Engine/Source/Developer/DerivedDataCache/Private/DerivedDataBackends.cpp
@@ -879,6 +879,10 @@ private:
 	ICacheStoreStats* CreateStats(ILegacyCacheStore* CacheStore, ECacheStoreFlags Flags, FStringView Type, FStringView Name, FStringView Path) final
 	{
 		check(Hierarchy);
+		if (ILegacyCacheStore* ThrottleNode = ThrottleNodes.FindRef(CacheStore))
+		{
+			CacheStore = ThrottleNode;
+		}
 		return Hierarchy->CreateStats(CacheStore, Flags, Type, Name, Path);
 	}
 
