Fix mislabeled bVerbose variable name in Jupiter Http code.

diff --git a/Engine/Source/Developer/DerivedDataCache/Private/HttpCacheStore.cpp b/Engine/Source/Developer/DerivedDataCache/Private/HttpCacheStore.cpp
--- a/Engine/Source/Developer/DerivedDataCache/Private/HttpCacheStore.cpp
+++ b/Engine/Source/Developer/DerivedDataCache/Private/HttpCacheStore.cpp
@@ -512,10 +512,10 @@ private:
 		if (UE_LOG_ACTIVE(LogDerivedDataCache, Display))
 		{
 			const int32 StatusCode = LocalResponse.GetStatusCode();
-			const bool bVerbose = (StatusCode >= 200 && StatusCode < 300) || Operation->ExpectedErrorCodes.Contains(StatusCode);
+			const bool bUnexpectedError = !((StatusCode >= 200 && StatusCode < 300) || Operation->ExpectedErrorCodes.Contains(StatusCode));
 
 			TStringBuilder<80> StatsText;
-			if (!bVerbose || UE_LOG_ACTIVE(LogDerivedDataCache, Verbose))
+			if (bUnexpectedError || UE_LOG_ACTIVE(LogDerivedDataCache, Verbose))
 			{
 				const FHttpResponseStats& Stats = LocalResponse.GetStats();
 				if (Stats.SendSize)
@@ -529,11 +529,7 @@ private:
 				StatsText.Appendf(TEXT("%.3f seconds %.3f|%.3f|%.3f|%.3f"), Stats.TotalTime, Stats.NameResolveTime, Stats.ConnectTime, Stats.TlsConnectTime, Stats.StartTransferTime);
 			}
 
-			if (bVerbose)
-			{
-				UE_LOG(LogDerivedDataCache, Verbose, TEXT("HTTP: %s (%s)"), *WriteToString<256>(LocalResponse), *StatsText);
-			}
-			else
+			if (bUnexpectedError)
 			{
 				FString Body = Operation->GetBodyAsString();
 				Body.ReplaceCharInline(TEXT('\r'), TEXT(' '));
@@ -541,6 +537,10 @@ private:
 				UE_LOG(LogDerivedDataCache, Display,
 					TEXT("HTTP: %s (%s) %s"), *WriteToString<256>(LocalResponse), *StatsText, *Body);
 			}
+			else
+			{
+				UE_LOG(LogDerivedDataCache, Verbose, TEXT("HTTP: %s (%s)"), *WriteToString<256>(LocalResponse), *StatsText);
+			}
 		}
 	}
 
