DDC: Disabled propagation of gets from the cache of in-flight puts and added support for write-only cache stores

- Do not invoke the put completion callback from the hierarchy until a put to a readable node succeeds, which is now tracked by bFinished.
- Use StopGetStore on the in-flight put cache to stop gets from that cache from propagating as duplicate puts.
- Disable the in-flight put cache for async wrappers used by propagation of gets.

diff --git a/Engine/Source/Developer/DerivedDataCache/Private/DerivedDataBackends.cpp b/Engine/Source/Developer/DerivedDataCache/Private/DerivedDataBackends.cpp
--- a/Engine/Source/Developer/DerivedDataCache/Private/DerivedDataBackends.cpp
+++ b/Engine/Source/Developer/DerivedDataCache/Private/DerivedDataBackends.cpp
@@ -602,7 +602,7 @@ public:
 		else
 		{
 			// This is unconditionally added to the hierarchy and will be deleted by the hierarchy.
-			CreateMemoryCacheStore(OutCache, TEXT("Memory"), TEXT("-ReadOnly -NoStats"), this);
+			CreateMemoryCacheStore(OutCache, TEXT("Memory"), TEXT("-ReadOnly -StopGetStore -NoStats"), this);
 			MemoryCache = OutCache;
 		}
 	}
diff --git a/Engine/Source/Developer/DerivedDataCache/Private/DerivedDataCacheStoreAsync.cpp b/Engine/Source/Developer/DerivedDataCache/Private/DerivedDataCacheStoreAsync.cpp
--- a/Engine/Source/Developer/DerivedDataCache/Private/DerivedDataCacheStoreAsync.cpp
+++ b/Engine/Source/Developer/DerivedDataCache/Private/DerivedDataCacheStoreAsync.cpp
@@ -39,7 +39,7 @@ public:
 			Execute(COOK_STAT(&FDerivedDataCacheUsageStats::TimePut,) Requests, Owner,
 				[this, OnComplete = MoveTemp(OnComplete)](FCachePutResponse&& Response)
 				{
-					MemoryCache->Delete(Response.Key);
+					MemoryCache->Delete(Response.Key, Response.Name);
 					OnComplete(MoveTemp(Response));
 				}, &ICacheStore::Put);
 		}
@@ -68,7 +68,7 @@ public:
 			Execute(COOK_STAT(&FDerivedDataCacheUsageStats::TimePut,) Requests, Owner,
 				[this, OnComplete = MoveTemp(OnComplete)](FCachePutValueResponse&& Response)
 				{
-					MemoryCache->DeleteValue(Response.Key);
+					MemoryCache->DeleteValue(Response.Key, Response.Name);
 					OnComplete(MoveTemp(Response));
 				}, &ICacheStore::PutValue);
 		}
diff --git a/Engine/Source/Developer/DerivedDataCache/Private/DerivedDataCacheStoreHierarchy.cpp b/Engine/Source/Developer/DerivedDataCache/Private/DerivedDataCacheStoreHierarchy.cpp
--- a/Engine/Source/Developer/DerivedDataCache/Private/DerivedDataCacheStoreHierarchy.cpp
+++ b/Engine/Source/Developer/DerivedDataCache/Private/DerivedDataCacheStoreHierarchy.cpp
@@ -120,6 +120,9 @@ private:
 	static uint64 MeasureLogicalRecordSize(const FCacheRecord& Record);
 	static uint64 MeasureLogicalValueSize(const FValue& Value);
 
+	static const FCacheKey& GetKey(const FCachePutRequest& Request) { return Request.Record.GetKey(); }
+	static const FCacheKey& GetKey(const FCachePutValueRequest& Request) { return Request.Key; }
+
 	struct FCacheStoreNode
 	{
 		ILegacyCacheStore* Cache{};
@@ -285,7 +288,7 @@ void FCacheStoreHierarchy::Add(ILegacyCacheStore* CacheStore, ECacheStoreFlags F
 	FWriteScopeLock Lock(NodesLock);
 	checkf(!Algo::FindBy(Nodes, CacheStore, &FCacheStoreNode::Cache),
 		TEXT("Attempting to add a cache store that was previously registered to the hierarchy."));
-	TUniquePtr<ILegacyCacheStore> AsyncCacheStore(CreateCacheStoreAsync(CacheStore, MemoryCache, /*bDeleteInnerCache*/ false));
+	TUniquePtr<ILegacyCacheStore> AsyncCacheStore(CreateCacheStoreAsync(CacheStore, /*MemoryCache*/ nullptr, /*bDeleteInnerCache*/ false));
 	Nodes.Add({CacheStore, Flags, {}, MoveTemp(AsyncCacheStore)});
 	UpdateNodeFlags();
 }
@@ -580,6 +583,7 @@ private:
 	{
 		bool bOk = false;
 		bool bStop = false;
+		bool bFinished = false;
 	};
 
 	FCacheStoreHierarchy& Hierarchy;
@@ -629,9 +633,10 @@ void FCacheStoreHierarchy::TPutBatch<Params>::DispatchRequests()
 	for (const FPutRequest& Request : Requests)
 	{
 		const FRequestState& State = States[RequestIndex];
-		if (!State.bOk && !State.bStop)
+		if (!State.bFinished)
 		{
-			FinishRequest(Request.MakeResponse(BatchOwner.IsCanceled() ? EStatus::Canceled : EStatus::Error), Request);
+			const EStatus Status = BatchOwner.IsCanceled() ? EStatus::Canceled : (State.bOk ? EStatus::Ok : EStatus::Error);
+			FinishRequest(Request.MakeResponse(Status), Request);
 		}
 		++RequestIndex;
 	}
@@ -689,8 +694,9 @@ void FCacheStoreHierarchy::TPutBatch<Params>::CompleteGetRequest(FGetResponse&&
 		FRequestState& State = States[RequestIndex];
 		check(!State.bStop);
 		State.bStop = true;
-		if (!State.bOk)
+		if (!State.bFinished)
 		{
+			State.bFinished = true;
 			const FPutRequest& Request = Requests[RequestIndex];
 			FinishRequest(Request.MakeResponse(Response.Status), Request);
 		}
@@ -723,7 +729,7 @@ bool FCacheStoreHierarchy::TPutBatch<Params>::DispatchPutRequests()
 		const FRequestState& State = States[RequestIndex];
 		if (!State.bStop && CanStore(GetCombinedPolicy(Request.Policy), Node.CacheFlags))
 		{
-			(State.bOk ? AsyncNodeRequests : NodeRequests).Add_GetRef(Request).UserData = uint64(RequestIndex);
+			(State.bFinished ? AsyncNodeRequests : NodeRequests).Add_GetRef(Request).UserData = uint64(RequestIndex);
 		}
 		++RequestIndex;
 	}
@@ -755,11 +761,20 @@ void FCacheStoreHierarchy::TPutBatch<Params>::CompletePutRequest(FPutResponse&&
 	{
 		const int32 RequestIndex = int32(Response.UserData);
 		FRequestState& State = States[RequestIndex];
-		check(!State.bOk && !State.bStop);
 		State.bOk = true;
-		const FPutRequest& Request = Requests[RequestIndex];
-		Response.UserData = Request.UserData;
-		FinishRequest(MoveTemp(Response), Request);
+		check(!State.bFinished);
+		bool bCanQuery;
+		{
+			FReadScopeLock Lock(Hierarchy.NodesLock);
+			bCanQuery = EnumHasAnyFlags(Hierarchy.Nodes[NodePutIndex].CacheFlags, ECacheStoreFlags::Query);
+		}
+		if (bCanQuery)
+		{
+			State.bFinished = true;
+			const FPutRequest& Request = Requests[RequestIndex];
+			Response.UserData = Request.UserData;
+			FinishRequest(MoveTemp(Response), Request);
+		}
 	}
 	if (RemainingRequestCount.Signal())
 	{
diff --git a/Engine/Source/Developer/DerivedDataCache/Private/MemoryCacheStore.cpp b/Engine/Source/Developer/DerivedDataCache/Private/MemoryCacheStore.cpp
--- a/Engine/Source/Developer/DerivedDataCache/Private/MemoryCacheStore.cpp
+++ b/Engine/Source/Developer/DerivedDataCache/Private/MemoryCacheStore.cpp
@@ -60,8 +60,8 @@ public:
 
 	// IMemoryCacheStore Interface
 
-	void Delete(const FCacheKey& Key) final;
-	void DeleteValue(const FCacheKey& Key) final;
+	void Delete(const FCacheKey& Key, const FSharedString& Name) final;
+	void DeleteValue(const FCacheKey& Key, const FSharedString& Name) final;
 
 	void Disable() final;
 
@@ -148,6 +148,10 @@ FMemoryCacheStore::FMemoryCacheStore(
 		{
 			Flags |= ECacheStoreFlags::Store;
 		}
+		if (FParse::Param(InConfig, TEXT("StopGetStore")))
+		{
+			Flags |= ECacheStoreFlags::StopGetStore;
+		}
 		StoreOwner->Add(this, Flags);
 		if (!FParse::Param(InConfig, TEXT("NoStats")))
 		{
@@ -474,7 +478,7 @@ void FMemoryCacheStore::Get(
 	}
 }
 
-void FMemoryCacheStore::Delete(const FCacheKey& Key)
+void FMemoryCacheStore::Delete(const FCacheKey& Key, const FSharedString& DebugName)
 {
 	FWriteScopeLock ScopeLock(SynchronizationObject);
 	FCacheRecordComponents Components;
@@ -634,7 +638,7 @@ void FMemoryCacheStore::GetValue(
 	}
 }
 
-void FMemoryCacheStore::DeleteValue(const FCacheKey& Key)
+void FMemoryCacheStore::DeleteValue(const FCacheKey& Key, const FSharedString& DebugName)
 {
 	FWriteScopeLock ScopeLock(SynchronizationObject);
 	FValue Value;
diff --git a/Engine/Source/Developer/DerivedDataCache/Private/MemoryCacheStore.h b/Engine/Source/Developer/DerivedDataCache/Private/MemoryCacheStore.h
--- a/Engine/Source/Developer/DerivedDataCache/Private/MemoryCacheStore.h
+++ b/Engine/Source/Developer/DerivedDataCache/Private/MemoryCacheStore.h
@@ -10,8 +10,8 @@ namespace UE::DerivedData
 class IMemoryCacheStore : public ILegacyCacheStore
 {
 public:
-	virtual void Delete(const FCacheKey& Key) = 0;
-	virtual void DeleteValue(const FCacheKey& Key) = 0;
+	virtual void Delete(const FCacheKey& Key, const FSharedString& Name) = 0;
+	virtual void DeleteValue(const FCacheKey& Key, const FSharedString& Name) = 0;
 
 	virtual void Disable() = 0;
 };
diff --git a/Engine/Source/Runtime/Engine/Private/TextureDerivedData.cpp b/Engine/Source/Runtime/Engine/Private/TextureDerivedData.cpp
--- a/Engine/Source/Runtime/Engine/Private/TextureDerivedData.cpp
+++ b/Engine/Source/Runtime/Engine/Private/TextureDerivedData.cpp
@@ -1242,7 +1242,7 @@ int64 PutDerivedDataInCache(FTexturePlatformData* DerivedData, const FString& De
 		FCbObject MetadataObject = BuildMetadata.ToCompactBinaryWithDefaults();
 		UE::DerivedData::FValue Value = UE::DerivedData::FValue::Compress(MetadataObject.GetBuffer());
 
-		const UE::DerivedData::FSharedString Name = TextureName;
+		const UE::DerivedData::FSharedString Name(WriteToString<256>(TextureName, TEXTVIEW(" [Meta]")));
 		UE::DerivedData::FRequestOwner AsyncOwner(UE::DerivedData::EPriority::Normal);
 		const UE::DerivedData::ECachePolicy Policy = bReplaceExistingDDC ? UE::DerivedData::ECachePolicy::Store : UE::DerivedData::ECachePolicy::Default;
 		UE::DerivedData::GetCache().PutValue({ {Name, GetTextureDerivedMetadataKeyFromSuffix(DerivedDataKeySuffix), MoveTemp(Value), Policy} }, AsyncOwner);
