Fix handling of cancelation in Zen DDC client to avoid having it treated as equivalent to an unexpected error.

diff --git a/Engine/Source/Developer/DerivedDataCache/Private/ZenCacheStore.cpp b/Engine/Source/Developer/DerivedDataCache/Private/ZenCacheStore.cpp
--- a/Engine/Source/Developer/DerivedDataCache/Private/ZenCacheStore.cpp
+++ b/Engine/Source/Developer/DerivedDataCache/Private/ZenCacheStore.cpp
@@ -329,7 +329,7 @@ public:
 							*CacheStore.GetName(), Batch.Num(), RequestIndex, *WriteToString<256>(*HttpResponse));
 					}
 				}
-				else if (HttpResponse->GetStatusCode() != 404)
+				else if ((HttpResponse->GetErrorCode() != EHttpErrorCode::Canceled) && (HttpResponse->GetStatusCode() != 404))
 				{
 					UE_LOG(LogDerivedDataCache, Warning,
 						TEXT("%s: Error response received from PutCacheRecords RPC: from %s"),
@@ -338,7 +338,14 @@ public:
 
 				for (const TRequestWithStats<FCachePutRequest>& RequestWithStats : Batch.RightChop(RequestIndex))
 				{
-					OnMiss(RequestWithStats);
+					if (HttpResponse->GetErrorCode() == EHttpErrorCode::Canceled)
+					{
+						OnCanceled(RequestWithStats);
+					}
+					else
+					{
+						OnMiss(RequestWithStats);
+					}
 				}
 			};
 			CacheStore.EnqueueAsyncRpc(Owner, BatchPackage, MoveTemp(OnRpcComplete));
@@ -365,7 +372,7 @@ private:
 
 	void OnHit(const TRequestWithStats<FCachePutRequest>& RequestWithStats)
 	{
-		UE_LOG(LogDerivedDataCache, Verbose, TEXT("%s: Cache Put complete for %s from '%s'"),
+		UE_LOG(LogDerivedDataCache, Verbose, TEXT("%s: Cache put complete for %s from '%s'"),
 			*CacheStore.GetName(), *WriteToString<96>(RequestWithStats.Request.Record.GetKey()), *RequestWithStats.Request.Name);
 
 		if (const FCbObject& Meta = RequestWithStats.Request.Record.GetMeta())
@@ -386,12 +393,20 @@ private:
 
 	void OnMiss(const TRequestWithStats<FCachePutRequest>& RequestWithStats)
 	{
-		UE_LOG(LogDerivedDataCache, Verbose, TEXT("%s: Cache Put miss for '%s' from '%s'"),
+		UE_LOG(LogDerivedDataCache, Verbose, TEXT("%s: Cache put failed for '%s' from '%s'"),
 			*CacheStore.GetName(), *WriteToString<96>(RequestWithStats.Request.Record.GetKey()), *RequestWithStats.Request.Name);
 		RequestWithStats.EndRequest(CacheStore, EStatus::Error);
 		OnComplete(RequestWithStats.Request.MakeResponse(EStatus::Error));
 	}
 
+	void OnCanceled(const TRequestWithStats<FCachePutRequest>& RequestWithStats)
+	{
+		UE_LOG(LogDerivedDataCache, Verbose, TEXT("%s: Cache put failed with canceled request for '%s' from '%s'"),
+			*CacheStore.GetName(), *WriteToString<96>(RequestWithStats.Request.Record.GetKey()), *RequestWithStats.Request.Name);
+		RequestWithStats.EndRequest(CacheStore, EStatus::Canceled);
+		OnComplete(RequestWithStats.Request.MakeResponse(EStatus::Canceled));
+	}
+
 	FZenCacheStore& CacheStore;
 	IRequestOwner& Owner;
 	TArray<TRequestWithStats<FCachePutRequest>, TInlineAllocator<1>> Requests;
@@ -507,7 +522,7 @@ public:
 							*CacheStore.GetName(), Batch.Num(), RequestIndex, *WriteToString<256>(*HttpResponse));
 					}
 				}
-				else if (HttpResponse->GetStatusCode() != 404)
+				else if ((HttpResponse->GetErrorCode() != EHttpErrorCode::Canceled) && (HttpResponse->GetStatusCode() != 404))
 				{
 					UE_LOG(LogDerivedDataCache, Warning,
 						TEXT("%s: Error response received from GetCacheRecords RPC: from %s"),
@@ -516,7 +531,14 @@ public:
 
 				for (const TRequestWithStats<FCacheGetRequest>& RequestWithStats : Batch.RightChop(RequestIndex))
 				{
-					OnMiss(RequestWithStats);
+					if (HttpResponse->GetErrorCode() == EHttpErrorCode::Canceled)
+					{
+						OnCanceled(RequestWithStats);
+					}
+					else
+					{
+						OnMiss(RequestWithStats);
+					}
 				}
 			};
 
@@ -569,6 +591,14 @@ private:
 		OnComplete(RequestWithStats.Request.MakeResponse(EStatus::Error));
 	}
 
+	void OnCanceled(const TRequestWithStats<FCacheGetRequest>& RequestWithStats)
+	{
+		UE_LOG(LogDerivedDataCache, Verbose, TEXT("%s: Cache miss with canceled request for '%s' from '%s'"),
+			*CacheStore.GetName(), *WriteToString<96>(RequestWithStats.Request.Key), *RequestWithStats.Request.Name);
+		RequestWithStats.EndRequest(CacheStore, EStatus::Canceled);
+		OnComplete(RequestWithStats.Request.MakeResponse(EStatus::Canceled));
+	}
+
 	FZenCacheStore& CacheStore;
 	IRequestOwner& Owner;
 	TArray<TRequestWithStats<FCacheGetRequest>, TInlineAllocator<1>> Requests;
@@ -670,7 +700,7 @@ public:
 							*CacheStore.GetName(), Batch.Num(), RequestIndex, *WriteToString<256>(*HttpResponse));
 					}
 				}
-				else if (HttpResponse->GetStatusCode() != 404)
+				else if ((HttpResponse->GetErrorCode() != EHttpErrorCode::Canceled) && (HttpResponse->GetStatusCode() != 404))
 				{
 					UE_LOG(LogDerivedDataCache, Warning,
 						TEXT("%s: Error response received from PutCacheValues RPC: from %s"),
@@ -679,7 +709,14 @@ public:
 
 				for (const TRequestWithStats<FCachePutValueRequest>& RequestWithStats : Batch.RightChop(RequestIndex))
 				{
-					OnMiss(RequestWithStats);
+					if (HttpResponse->GetErrorCode() == EHttpErrorCode::Canceled)
+					{
+						OnCanceled(RequestWithStats);
+					}
+					else
+					{
+						OnMiss(RequestWithStats);
+					}
 				}
 			};
 
@@ -702,7 +739,7 @@ private:
 
 	void OnHit(const TRequestWithStats<FCachePutValueRequest>& RequestWithStats)
 	{
-		UE_LOG(LogDerivedDataCache, Verbose, TEXT("%s: Cache PutValue complete for %s from '%s'"),
+		UE_LOG(LogDerivedDataCache, Verbose, TEXT("%s: Cache put complete for %s from '%s'"),
 			*CacheStore.GetName(), *WriteToString<96>(RequestWithStats.Request.Key), *RequestWithStats.Request.Name);
 
 		RequestWithStats.Stats.AddLogicalWrite(RequestWithStats.Request.Value);
@@ -716,12 +753,20 @@ private:
 
 	void OnMiss(const TRequestWithStats<FCachePutValueRequest>& RequestWithStats)
 	{
-		UE_LOG(LogDerivedDataCache, Verbose, TEXT("%s: Cache PutValue miss for '%s' from '%s'"),
+		UE_LOG(LogDerivedDataCache, Verbose, TEXT("%s: Cache put failed for '%s' from '%s'"),
 			*CacheStore.GetName(), *WriteToString<96>(RequestWithStats.Request.Key), *RequestWithStats.Request.Name);
 		RequestWithStats.EndRequest(CacheStore, EStatus::Error);
 		OnComplete(RequestWithStats.Request.MakeResponse(EStatus::Error));
 	}
 
+	void OnCanceled(const TRequestWithStats<FCachePutValueRequest>& RequestWithStats)
+	{
+		UE_LOG(LogDerivedDataCache, Verbose, TEXT("%s: Cache put failed with canceled request for '%s' from '%s'"),
+			*CacheStore.GetName(), *WriteToString<96>(RequestWithStats.Request.Key), *RequestWithStats.Request.Name);
+		RequestWithStats.EndRequest(CacheStore, EStatus::Canceled);
+		OnComplete(RequestWithStats.Request.MakeResponse(EStatus::Canceled));
+	}
+
 	FZenCacheStore& CacheStore;
 	IRequestOwner& Owner;
 	TArray<TRequestWithStats<FCachePutValueRequest>, TInlineAllocator<1>> Requests;
@@ -849,7 +894,7 @@ public:
 							*CacheStore.GetName(), Batch.Num(), RequestIndex, *WriteToString<256>(*HttpResponse));
 					}
 				}
-				else if (HttpResponse->GetStatusCode() != 404)
+				else if ((HttpResponse->GetErrorCode() != EHttpErrorCode::Canceled) && (HttpResponse->GetStatusCode() != 404))
 				{
 					UE_LOG(LogDerivedDataCache, Warning,
 						TEXT("%s: Error response received from GetCacheValues RPC: from %s"),
@@ -858,7 +903,14 @@ public:
 
 				for (const TRequestWithStats<FCacheGetValueRequest>& RequestWithStats : Batch.RightChop(RequestIndex))
 				{
-					OnMiss(RequestWithStats);
+					if (HttpResponse->GetErrorCode() == EHttpErrorCode::Canceled)
+					{
+						OnCanceled(RequestWithStats);
+					}
+					else
+					{
+						OnMiss(RequestWithStats);
+					}
 				}
 			};
 
@@ -889,6 +941,14 @@ private:
 		OnComplete(RequestWithStats.Request.MakeResponse(EStatus::Error));
 	};
 
+	void OnCanceled(const TRequestWithStats<FCacheGetValueRequest>& RequestWithStats)
+	{
+		UE_LOG(LogDerivedDataCache, Verbose, TEXT("%s: Cache miss with canceled request for '%s' from '%s'"),
+			*CacheStore.GetName(), *WriteToString<96>(RequestWithStats.Request.Key), *RequestWithStats.Request.Name);
+		RequestWithStats.EndRequest(CacheStore, EStatus::Canceled);
+		OnComplete(RequestWithStats.Request.MakeResponse(EStatus::Canceled));
+	};
+
 	FZenCacheStore& CacheStore;
 	IRequestOwner& Owner;
 	TArray<TRequestWithStats<FCacheGetValueRequest>, TInlineAllocator<1>> Requests;
@@ -1042,7 +1102,7 @@ public:
 						Succeeded ? OnHit(RequestWithStats, MoveTemp(RawHash), RawSize, MoveTemp(RequestedBytes)) : OnMiss(RequestWithStats);
 					}
 				}
-				else if (HttpResponse->GetStatusCode() != 404)
+				else if ((HttpResponse->GetErrorCode() != EHttpErrorCode::Canceled) && (HttpResponse->GetStatusCode() != 404))
 				{
 					UE_LOG(LogDerivedDataCache, Warning,
 						TEXT("%s: Error response received from GetChunks RPC: from %s"),
@@ -1051,7 +1111,14 @@ public:
 
 				for (const TRequestWithStats<FCacheGetChunkRequest>& RequestWithStats : Batch.RightChop(RequestIndex))
 				{
-					OnMiss(RequestWithStats);
+					if (HttpResponse->GetErrorCode() == EHttpErrorCode::Canceled)
+					{
+						OnCanceled(RequestWithStats);
+					}
+					else
+					{
+						OnMiss(RequestWithStats);
+					}
 				}
 			};
 			CacheStore.EnqueueAsyncRpc(Owner, BatchRequest.Save().AsObject(), MoveTemp(OnRpcComplete));
@@ -1062,7 +1129,7 @@ private:
 	void OnHit(const TRequestWithStats<FCacheGetChunkRequest>& RequestWithStats, FIoHash&& RawHash, uint64 RawSize, FSharedBuffer&& RequestedBytes)
 	{
 		const FCacheGetChunkRequest& Request = RequestWithStats.Request;
-		UE_LOG(LogDerivedDataCache, Verbose, TEXT("%s: CacheChunk hit for '%s' from '%s'"),
+		UE_LOG(LogDerivedDataCache, Verbose, TEXT("%s: Cache hit for '%s' from '%s'"),
 			*CacheStore.GetName(), *WriteToString<96>(Request.Key, '/', Request.Id), *Request.Name);
 
 		// This is a rough estimate of physical read size until Zen communicates stats with each response.
@@ -1079,12 +1146,21 @@ private:
 	void OnMiss(const TRequestWithStats<FCacheGetChunkRequest>& RequestWithStats)
 	{
 		const FCacheGetChunkRequest& Request = RequestWithStats.Request;
-		UE_LOG(LogDerivedDataCache, Verbose, TEXT("%s: CacheChunk miss with missing value '%s' for '%s' from '%s'"),
+		UE_LOG(LogDerivedDataCache, Verbose, TEXT("%s: Cache miss with missing value '%s' for '%s' from '%s'"),
 			*CacheStore.GetName(), *WriteToString<16>(Request.Id), *WriteToString<96>(Request.Key), *Request.Name);
 		RequestWithStats.EndRequest(CacheStore, EStatus::Error);
 		OnComplete(Request.MakeResponse(EStatus::Error));
 	};
 
+	void OnCanceled(const TRequestWithStats<FCacheGetChunkRequest>& RequestWithStats)
+	{
+		const FCacheGetChunkRequest& Request = RequestWithStats.Request;
+		UE_LOG(LogDerivedDataCache, Verbose, TEXT("%s: Cache miss with canceled request for '%s' from '%s'"),
+			*CacheStore.GetName(), *WriteToString<96>(Request.Key, '/', Request.Id), *Request.Name);
+		RequestWithStats.EndRequest(CacheStore, EStatus::Canceled);
+		OnComplete(Request.MakeResponse(EStatus::Canceled));
+	};
+
 	FZenCacheStore& CacheStore;
 	IRequestOwner& Owner;
 	TArray<TRequestWithStats<FCacheGetChunkRequest>, TInlineAllocator<1>> Requests;
