Add CPU time and latency stats to ZenCacheStore
Fix LogicalReadSize for ZenCacheStore GetChunks
Add missing fields for GatherAnalytics in ZenServerInterface

diff --git a/Engine/Source/Developer/DerivedDataCache/Private/ZenCacheStore.cpp b/Engine/Source/Developer/DerivedDataCache/Private/ZenCacheStore.cpp
--- a/Engine/Source/Developer/DerivedDataCache/Private/ZenCacheStore.cpp
+++ b/Engine/Source/Developer/DerivedDataCache/Private/ZenCacheStore.cpp
@@ -253,8 +253,24 @@ public:
 		TRACE_COUNTER_ADD(ZenDDC_Put, int64(Requests.Num()));
 	}
 
+	virtual ~FPutOp()
+	{
+		FMonotonicTimeSpan AverageMainThreadTime = FMonotonicTimeSpan::FromSeconds(Requests[0].Stats.MainThreadTime.ToSeconds() / Requests.Num());
+		FMonotonicTimeSpan AverageOtherThreadTime = FMonotonicTimeSpan::FromSeconds(Requests[0].Stats.OtherThreadTime.ToSeconds() / Requests.Num());
+		FMonotonicTimeSpan AverageLatency = FMonotonicTimeSpan::FromSeconds(Requests[0].Stats.Latency.ToSeconds() / Requests.Num());
+
+		for (TRequestWithStats<FCachePutRequest>& Request : Requests)
+		{
+			Request.Stats.MainThreadTime = AverageMainThreadTime;
+			Request.Stats.OtherThreadTime = AverageOtherThreadTime;
+			Request.Stats.Latency = AverageLatency;
+		}
+	}
+
 	void IssueRequests()
 	{
+		FRequestTimer RequestTimer(Requests[0].Stats);
+
 		FRequestBarrier Barrier(Owner);
 		for (TArrayView<const TRequestWithStats<FCachePutRequest>> Batch : Batches)
 		{
@@ -297,6 +313,9 @@ public:
 
 			auto OnRpcComplete = [this, OpRef = TRefCountPtr<FPutOp>(this), Batch](THttpUniquePtr<IHttpResponse>& HttpResponse, FCbPackage& Response)
 			{
+				FRequestTimer RequestTimer(Requests[0].Stats);
+				Requests[0].Stats.Latency = FMonotonicTimeSpan::FromSeconds(HttpResponse->GetStats().StartTransferTime);
+
 				int32 RequestIndex = 0;
 				if (HttpResponse->GetErrorCode() == EHttpErrorCode::None && HttpResponse->GetStatusCode() >= 200 && HttpResponse->GetStatusCode() <= 299)
 				{
@@ -310,7 +329,6 @@ public:
 						}
 
 						const TRequestWithStats<FCachePutRequest>& RequestWithStats = Batch[RequestIndex++];
-						RequestWithStats.Stats.AddLatency(FMonotonicTimeSpan::FromSeconds(HttpResponse->GetStats().StartTransferTime));
 
 						const FCacheKey& Key = RequestWithStats.Request.Record.GetKey();
 						bool bPutSucceeded = ResponseField.AsBool();
@@ -435,10 +453,22 @@ public:
 	virtual ~FGetOp()
 	{
 		TRACE_COUNTER_SUBTRACT(ZenDDC_CacheRecordRequestCountInFlight, int64(Requests.Num()));
+		FMonotonicTimeSpan AverageMainThreadTime = FMonotonicTimeSpan::FromSeconds(Requests[0].Stats.MainThreadTime.ToSeconds() / Requests.Num());
+		FMonotonicTimeSpan AverageOtherThreadTime = FMonotonicTimeSpan::FromSeconds(Requests[0].Stats.OtherThreadTime.ToSeconds() / Requests.Num());
+		FMonotonicTimeSpan AverageLatency = FMonotonicTimeSpan::FromSeconds(Requests[0].Stats.Latency.ToSeconds() / Requests.Num());
+
+		for (TRequestWithStats<FCacheGetRequest>& Request : Requests)
+		{
+			Request.Stats.MainThreadTime = AverageMainThreadTime;
+			Request.Stats.OtherThreadTime = AverageOtherThreadTime;
+			Request.Stats.Latency = AverageLatency;
+		}
 	}
 
 	void IssueRequests()
 	{
+		FRequestTimer RequestTimer(Requests[0].Stats);
+
 		FRequestBarrier Barrier(Owner);
 		ForEachBatch(CacheStore.CacheRecordBatchSize, Requests.Num(),
 			[this](int32 BatchFirst, int32 BatchLast)
@@ -485,6 +515,9 @@ public:
 			FGetOp* OriginalOp = this;
 			auto OnRpcComplete = [this, OpRef = TRefCountPtr<FGetOp>(OriginalOp), Batch](THttpUniquePtr<IHttpResponse>& HttpResponse, FCbPackage& Response)
 			{
+				FRequestTimer RequestTimer(Requests[0].Stats);
+				Requests[0].Stats.Latency = FMonotonicTimeSpan::FromSeconds(HttpResponse->GetStats().StartTransferTime);
+
 				int32 RequestIndex = 0;
 				if (HttpResponse->GetErrorCode() == EHttpErrorCode::None && HttpResponse->GetStatusCode() >= 200 && HttpResponse->GetStatusCode() <= 299)
 				{
@@ -499,7 +532,6 @@ public:
 						}
 
 						const TRequestWithStats<FCacheGetRequest>& RequestWithStats = Batch[RequestIndex++];
-						RequestWithStats.Stats.AddLatency(FMonotonicTimeSpan::FromSeconds(HttpResponse->GetStats().StartTransferTime));
 
 						const FCacheKey& Key = RequestWithStats.Request.Key;
 						FOptionalCacheRecord Record;
@@ -528,7 +560,7 @@ public:
 						TEXT("%s: Error response received from GetCacheRecords RPC: from %s"),
 						*CacheStore.GetName(), *WriteToString<256>(*HttpResponse));
 				}
-
+					
 				for (const TRequestWithStats<FCacheGetRequest>& RequestWithStats : Batch.RightChop(RequestIndex))
 				{
 					if (HttpResponse->GetErrorCode() == EHttpErrorCode::Canceled)
@@ -623,8 +655,24 @@ public:
 			[this](const TRequestWithStats<FCachePutValueRequest>& NextRequest) { return BatchGroupingFilter(NextRequest.Request); });
 	}
 
+	virtual ~FPutValueOp()
+	{
+		FMonotonicTimeSpan AverageMainThreadTime = FMonotonicTimeSpan::FromSeconds(Requests[0].Stats.MainThreadTime.ToSeconds() / Requests.Num());
+		FMonotonicTimeSpan AverageOtherThreadTime = FMonotonicTimeSpan::FromSeconds(Requests[0].Stats.OtherThreadTime.ToSeconds() / Requests.Num());
+		FMonotonicTimeSpan AverageLatency = FMonotonicTimeSpan::FromSeconds(Requests[0].Stats.Latency.ToSeconds() / Requests.Num());
+
+		for (TRequestWithStats<FCachePutValueRequest>& Request : Requests)
+		{
+			Request.Stats.MainThreadTime = AverageMainThreadTime;
+			Request.Stats.OtherThreadTime = AverageOtherThreadTime;
+			Request.Stats.Latency = AverageLatency;
+		}
+	}
+
 	void IssueRequests()
 	{
+		FRequestTimer RequestTimer(Requests[0].Stats);
+
 		FRequestBarrier Barrier(Owner);
 		for (TArrayView<const TRequestWithStats<FCachePutValueRequest>> Batch : Batches)
 		{
@@ -669,6 +717,9 @@ public:
 
 			auto OnRpcComplete = [this, OpRef = TRefCountPtr<FPutValueOp>(this), Batch](THttpUniquePtr<IHttpResponse>& HttpResponse, FCbPackage& Response)
 			{
+				FRequestTimer RequestTimer(Requests[0].Stats);
+				Requests[0].Stats.Latency = FMonotonicTimeSpan::FromSeconds(HttpResponse->GetStats().StartTransferTime);
+
 				int32 RequestIndex = 0;
 				if (HttpResponse->GetErrorCode() == EHttpErrorCode::None && HttpResponse->GetStatusCode() >= 200 && HttpResponse->GetStatusCode() <= 299)
 				{
@@ -682,7 +733,6 @@ public:
 						}
 
 						const TRequestWithStats<FCachePutValueRequest>& RequestWithStats = Batch[RequestIndex++];
-						RequestWithStats.Stats.AddLatency(FMonotonicTimeSpan::FromSeconds(HttpResponse->GetStats().StartTransferTime));
 
 						bool bPutSucceeded = ResponseField.AsBool();
 						if (CacheStore.DebugOptions.ShouldSimulatePutMiss(RequestWithStats.Request.Key))
@@ -795,10 +845,22 @@ public:
 	virtual ~FGetValueOp()
 	{
 		TRACE_COUNTER_SUBTRACT(ZenDDC_CacheRecordRequestCountInFlight, int64(Requests.Num()));
+		FMonotonicTimeSpan AverageMainThreadTime = FMonotonicTimeSpan::FromSeconds(Requests[0].Stats.MainThreadTime.ToSeconds() / Requests.Num());
+		FMonotonicTimeSpan AverageOtherThreadTime = FMonotonicTimeSpan::FromSeconds(Requests[0].Stats.OtherThreadTime.ToSeconds() / Requests.Num());
+		FMonotonicTimeSpan AverageLatency = FMonotonicTimeSpan::FromSeconds(Requests[0].Stats.Latency.ToSeconds() / Requests.Num());
+
+		for (TRequestWithStats<FCacheGetValueRequest>& Request : Requests)
+		{
+			Request.Stats.MainThreadTime = AverageMainThreadTime;
+			Request.Stats.OtherThreadTime = AverageOtherThreadTime;
+			Request.Stats.Latency = AverageLatency;
+		}
 	}
 
 	void IssueRequests()
 	{
+		FRequestTimer RequestTimer(Requests[0].Stats);
+
 		FRequestBarrier Barrier(Owner);
 		ForEachBatch(CacheStore.CacheRecordBatchSize, Requests.Num(),
 			[this](int32 BatchFirst, int32 BatchLast)
@@ -844,6 +906,9 @@ public:
 			FGetValueOp* OriginalOp = this;
 			auto OnRpcComplete = [this, OpRef = TRefCountPtr<FGetValueOp>(OriginalOp), Batch](THttpUniquePtr<IHttpResponse>& HttpResponse, FCbPackage& Response)
 			{
+				FRequestTimer RequestTimer(Requests[0].Stats);
+				Requests[0].Stats.Latency = FMonotonicTimeSpan::FromSeconds(HttpResponse->GetStats().StartTransferTime);
+
 				int32 RequestIndex = 0;
 				if (HttpResponse->GetErrorCode() == EHttpErrorCode::None && HttpResponse->GetStatusCode() >= 200 && HttpResponse->GetStatusCode() <= 299)
 				{
@@ -858,7 +923,6 @@ public:
 						}
 
 						const TRequestWithStats<FCacheGetValueRequest>& RequestWithStats = Batch[RequestIndex++];
-						RequestWithStats.Stats.AddLatency(FMonotonicTimeSpan::FromSeconds(HttpResponse->GetStats().StartTransferTime));
 
 						FCbObjectView ResultObj = ResultField.AsObjectView();
 						TOptional<FValue> Value;
@@ -977,10 +1041,22 @@ public:
 	virtual ~FGetChunksOp()
 	{
 		TRACE_COUNTER_SUBTRACT(ZenDDC_ChunkRequestCountInFlight, int64(Requests.Num()));
+		FMonotonicTimeSpan AverageMainThreadTime = FMonotonicTimeSpan::FromSeconds(Requests[0].Stats.MainThreadTime.ToSeconds() / Requests.Num());
+		FMonotonicTimeSpan AverageOtherThreadTime = FMonotonicTimeSpan::FromSeconds(Requests[0].Stats.OtherThreadTime.ToSeconds() / Requests.Num());
+		FMonotonicTimeSpan AverageLatency = FMonotonicTimeSpan::FromSeconds(Requests[0].Stats.Latency.ToSeconds() / Requests.Num());
+
+		for (TRequestWithStats<FCacheGetChunkRequest>& Request : Requests)
+		{
+			Request.Stats.MainThreadTime = AverageMainThreadTime;
+			Request.Stats.OtherThreadTime = AverageOtherThreadTime;
+			Request.Stats.Latency = AverageLatency;
+		}
 	}
 
 	void IssueRequests()
 	{
+		FRequestTimer RequestTimer(Requests[0].Stats);
+
 		FRequestBarrier Barrier(Owner);
 		ForEachBatch(CacheStore.CacheChunksBatchSize, Requests.Num(),
 			[this](int32 BatchFirst, int32 BatchLast)
@@ -1044,6 +1120,9 @@ public:
 			FGetChunksOp* OriginalOp = this;
 			auto OnRpcComplete = [this, OpRef = TRefCountPtr<FGetChunksOp>(OriginalOp), Batch](THttpUniquePtr<IHttpResponse>& HttpResponse, FCbPackage& Response)
 			{
+				FRequestTimer RequestTimer(Requests[0].Stats);
+				Requests[0].Stats.Latency = FMonotonicTimeSpan::FromSeconds(HttpResponse->GetStats().StartTransferTime);
+
 				int32 RequestIndex = 0;
 				if (HttpResponse->GetErrorCode() == EHttpErrorCode::None && HttpResponse->GetStatusCode() >= 200 && HttpResponse->GetStatusCode() <= 299)
 				{
@@ -1057,7 +1136,6 @@ public:
 							continue;
 						}
 						const TRequestWithStats<FCacheGetChunkRequest>& RequestWithStats = Batch[RequestIndex++];
-						RequestWithStats.Stats.AddLatency(FMonotonicTimeSpan::FromSeconds(HttpResponse->GetStats().StartTransferTime));
 
 						FIoHash RawHash;
 						bool Succeeded = false;
@@ -1133,7 +1211,7 @@ private:
 			*CacheStore.GetName(), *WriteToString<96>(Request.Key, '/', Request.Id), *Request.Name);
 
 		// This is a rough estimate of physical read size until Zen communicates stats with each response.
-		RequestWithStats.Stats.LogicalReadSize += RequestedBytes.GetSize();
+		RequestWithStats.Stats.LogicalReadSize += RawSize;
 		RequestWithStats.Stats.PhysicalReadSize += RequestedBytes.GetSize();
 		RequestWithStats.EndRequest(CacheStore, EStatus::Ok);
 
diff --git a/Engine/Source/Developer/Zen/Private/ZenServerInterface.cpp b/Engine/Source/Developer/Zen/Private/ZenServerInterface.cpp
--- a/Engine/Source/Developer/Zen/Private/ZenServerInterface.cpp
+++ b/Engine/Source/Developer/Zen/Private/ZenServerInterface.cpp
@@ -2355,86 +2355,262 @@ FZenServiceInstance::RequestGC(const bool* OverrideCollectSmallObjects, const ui
 bool 
 FZenServiceInstance::GatherAnalytics(TArray<FAnalyticsEventAttribute>& Attributes)
 {
-	FZenCacheStats ZenStats;
+	FZenCacheStats ZenCacheStats;
+	FZenProjectStats ZenProjectStats;
 
-	if (GetCacheStats(ZenStats) == false)
+	if (GetCacheStats(ZenCacheStats) == false)
+		return false;
+
+	if (GetProjectStats(ZenProjectStats) == false)
 		return false;
 
 	const FString BaseName = TEXT("Zen_");
 
 	{
 		FString AttrName = BaseName + TEXT("Enabled");
-		Attributes.Emplace(MoveTemp(AttrName), ZenStats.bIsValid);
+		Attributes.Emplace(MoveTemp(AttrName), ZenCacheStats.bIsValid && ZenProjectStats.bIsValid);
 	}
 
+	///////////// Cache
 	{
-		FString AttrName = BaseName + TEXT("Cache_HitRatio");
-		Attributes.Emplace(MoveTemp(AttrName), ZenStats.General.HitRatio);
+		FString AttrName = BaseName + TEXT("Cache_Size_Disk");
+		Attributes.Emplace(MoveTemp(AttrName), ZenCacheStats.General.Size.Disk);
+	}
+
+	{
+		FString AttrName = BaseName + TEXT("Cache_Size_Memory");
+		Attributes.Emplace(MoveTemp(AttrName), ZenCacheStats.General.Size.Memory);
 	}
 
 	{
 		FString AttrName = BaseName + TEXT("Cache_Hits");
-		Attributes.Emplace(MoveTemp(AttrName), ZenStats.General.Hits);
+		Attributes.Emplace(MoveTemp(AttrName), ZenCacheStats.General.Hits);
 	}
 
 	{
 		FString AttrName = BaseName + TEXT("Cache_Misses");
-		Attributes.Emplace(MoveTemp(AttrName), ZenStats.General.Misses);
+		Attributes.Emplace(MoveTemp(AttrName), ZenCacheStats.General.Misses);
 	}
 
 	{
-		FString AttrName = BaseName + TEXT("Cache_Size_Disk");
-		Attributes.Emplace(MoveTemp(AttrName), ZenStats.General.Size.Disk);
+		FString AttrName = BaseName + TEXT("Cache_Writes");
+		Attributes.Emplace(MoveTemp(AttrName), ZenCacheStats.General.Writes);
 	}
 
 	{
-		FString AttrName = BaseName + TEXT("Cache_Size_Memory");
-		Attributes.Emplace(MoveTemp(AttrName), ZenStats.General.Size.Memory);
+		FString AttrName = BaseName + TEXT("Cache_HitRatio");
+		Attributes.Emplace(MoveTemp(AttrName), ZenCacheStats.General.HitRatio);
 	}
 
 	{
 		FString AttrName = BaseName + TEXT("Cache_UpstreamHits");
-		Attributes.Emplace(MoveTemp(AttrName), ZenStats.General.UpstreamHits);
+		Attributes.Emplace(MoveTemp(AttrName), ZenCacheStats.General.UpstreamHits);
 	}
 
 	{
 		FString AttrName = BaseName + TEXT("Cache_UpstreamRatio");
-		Attributes.Emplace(MoveTemp(AttrName), ZenStats.General.UpstreamRatio);
+		Attributes.Emplace(MoveTemp(AttrName), ZenCacheStats.General.UpstreamRatio);
+	}
+
+
+	{
+		FString AttrName = BaseName + TEXT("Cache_Cas_Hits");
+		Attributes.Emplace(MoveTemp(AttrName), ZenCacheStats.General.CidHits);
+	}
+
+	{
+		FString AttrName = BaseName + TEXT("Cache_Cas_Misses");
+		Attributes.Emplace(MoveTemp(AttrName), ZenCacheStats.General.CidMisses);
+	}
+
+	{
+		FString AttrName = BaseName + TEXT("Cache_Cas_Writes");
+		Attributes.Emplace(MoveTemp(AttrName), ZenCacheStats.General.CidWrites);
+	}
+
+	{
+		FString AttrName = BaseName + TEXT("Cache_Requests");
+		Attributes.Emplace(MoveTemp(AttrName), ZenCacheStats.General.RequestCount);
+	}
+
+	{
+		FString AttrName = BaseName + TEXT("Cache_BadRequests");
+		Attributes.Emplace(MoveTemp(AttrName), ZenCacheStats.General.BadRequestCount);
+	}
+
+
+	{
+		FString AttrName = BaseName + TEXT("Cache_Requests_Count");
+		Attributes.Emplace(MoveTemp(AttrName), ZenCacheStats.Request.Count);
+	}
+
+	{
+		FString AttrName = BaseName + TEXT("Cache_Requests_RateMean");
+		Attributes.Emplace(MoveTemp(AttrName), ZenCacheStats.Request.RateMean);
+	}
+
+	{
+		FString AttrName = BaseName + TEXT("Cache_Requests_TAverage");
+		Attributes.Emplace(MoveTemp(AttrName), ZenCacheStats.Request.TAverage);
+	}
+
+	{
+		FString AttrName = BaseName + TEXT("Cache_Requests_TMin");
+		Attributes.Emplace(MoveTemp(AttrName), ZenCacheStats.Request.TMin);
+	}
+
+	{
+		FString AttrName = BaseName + TEXT("Cache_Requests_TMax");
+		Attributes.Emplace(MoveTemp(AttrName), ZenCacheStats.Request.TMax);
 	}
 
 	{
 		FString AttrName = BaseName + TEXT("Cache_TotalUploadedMB");
-		Attributes.Emplace(MoveTemp(AttrName), ZenStats.Upstream.TotalUploadedMB);
+		Attributes.Emplace(MoveTemp(AttrName), ZenCacheStats.Upstream.TotalUploadedMB);
 	}
 
 	{
 		FString AttrName = BaseName + TEXT("Upstream_TotalDownloadedMB");
-		Attributes.Emplace(MoveTemp(AttrName), ZenStats.Upstream.TotalDownloadedMB);
+		Attributes.Emplace(MoveTemp(AttrName), ZenCacheStats.Upstream.TotalDownloadedMB);
 	}
 
 	{
 		FString AttrName = BaseName + TEXT("Upstream_TotalUploadedMB");
-		Attributes.Emplace(MoveTemp(AttrName), ZenStats.Upstream.TotalUploadedMB);
+		Attributes.Emplace(MoveTemp(AttrName), ZenCacheStats.Upstream.TotalUploadedMB);
+	}
+
+	{
+		FString AttrName = BaseName + TEXT("Upstream_Requests_Count");
+		Attributes.Emplace(MoveTemp(AttrName), ZenCacheStats.UpstreamRequest.Count);
+	}
+	{
+		FString AttrName = BaseName + TEXT("Upstream_Requests_RateMean");
+		Attributes.Emplace(MoveTemp(AttrName), ZenCacheStats.UpstreamRequest.RateMean);
+	}
+	{
+		FString AttrName = BaseName + TEXT("Upstream_Requests_TAverage");
+		Attributes.Emplace(MoveTemp(AttrName), ZenCacheStats.UpstreamRequest.TAverage);
+	}
+	{
+		FString AttrName = BaseName + TEXT("Upstream_Requests_TMin");
+		Attributes.Emplace(MoveTemp(AttrName), ZenCacheStats.UpstreamRequest.TMin);
+	}
+	{
+		FString AttrName = BaseName + TEXT("Upstream_Requests_TMax");
+		Attributes.Emplace(MoveTemp(AttrName), ZenCacheStats.UpstreamRequest.TMax);
 	}
 
 	{
 		FString AttrName = BaseName + TEXT("Cas_Size_Large");
-		Attributes.Emplace(MoveTemp(AttrName), ZenStats.CID.Size.Large);
+		Attributes.Emplace(MoveTemp(AttrName), ZenCacheStats.CID.Size.Large);
 	}
 
 	{
 		FString AttrName = BaseName + TEXT("Cas_Size_Small");
-		Attributes.Emplace(MoveTemp(AttrName), ZenStats.CID.Size.Small);
+		Attributes.Emplace(MoveTemp(AttrName), ZenCacheStats.CID.Size.Small);
 	}
 
 	{
 		FString AttrName = BaseName + TEXT("Cas_Size_Tiny");
-		Attributes.Emplace(MoveTemp(AttrName), ZenStats.CID.Size.Tiny);
+		Attributes.Emplace(MoveTemp(AttrName), ZenCacheStats.CID.Size.Tiny);
 	}
 
 	{
 		FString AttrName = BaseName + TEXT("Cas_Size_Total");
-		Attributes.Emplace(MoveTemp(AttrName), ZenStats.CID.Size.Total);
+		Attributes.Emplace(MoveTemp(AttrName), ZenCacheStats.CID.Size.Total);
+	}
+
+	///////////// Project
+	{
+		FString AttrName = BaseName + TEXT("Project_Size_Disk");
+		Attributes.Emplace(MoveTemp(AttrName), ZenProjectStats.General.Size.Disk);
+	}
+
+	{
+		FString AttrName = BaseName + TEXT("Project_Size_Memory");
+		Attributes.Emplace(MoveTemp(AttrName), ZenProjectStats.General.Size.Memory);
+	}
+
+	{
+		FString AttrName = BaseName + TEXT("Project_WriteCount");
+		Attributes.Emplace(MoveTemp(AttrName), ZenProjectStats.General.Project.WriteCount);
+	}
+
+	{
+		FString AttrName = BaseName + TEXT("Project_ReadCount");
+		Attributes.Emplace(MoveTemp(AttrName), ZenProjectStats.General.Project.ReadCount);
+	}
+
+	{
+		FString AttrName = BaseName + TEXT("Project_DeleteCount");
+		Attributes.Emplace(MoveTemp(AttrName), ZenProjectStats.General.Project.DeleteCount);
+	}
+
+	{
+		FString AttrName = BaseName + TEXT("Project_Oplog_WriteCount");
+		Attributes.Emplace(MoveTemp(AttrName), ZenProjectStats.General.Oplog.WriteCount);
+	}
+
+	{
+		FString AttrName = BaseName + TEXT("Project_Oplog_ReadCount");
+		Attributes.Emplace(MoveTemp(AttrName), ZenProjectStats.General.Oplog.ReadCount);
+	}
+
+	{
+		FString AttrName = BaseName + TEXT("Project_Oplog_DeleteCount");
+		Attributes.Emplace(MoveTemp(AttrName), ZenProjectStats.General.Oplog.DeleteCount);
+	}
+
+	{
+		FString AttrName = BaseName + TEXT("Project_Op_Hits");
+		Attributes.Emplace(MoveTemp(AttrName), ZenProjectStats.General.Op.HitCount);
+	}
+
+	{
+		FString AttrName = BaseName + TEXT("Project_Op_Misses");
+		Attributes.Emplace(MoveTemp(AttrName), ZenProjectStats.General.Op.MissCount);
+	}
+
+	{
+		FString AttrName = BaseName + TEXT("Project_Op_Writes");
+		Attributes.Emplace(MoveTemp(AttrName), ZenProjectStats.General.Op.WriteCount);
+	}
+
+	{
+		FString AttrName = BaseName + TEXT("Project_Chunk_Hits");
+		Attributes.Emplace(MoveTemp(AttrName), ZenProjectStats.General.Chunk.HitCount);
+	}
+
+	{
+		FString AttrName = BaseName + TEXT("Project_Chunk_Misses");
+		Attributes.Emplace(MoveTemp(AttrName), ZenProjectStats.General.Chunk.MissCount);
+	}
+
+	{
+		FString AttrName = BaseName + TEXT("Project_Chunk_Writes");
+		Attributes.Emplace(MoveTemp(AttrName), ZenProjectStats.General.Chunk.WriteCount);
+	}
+
+	{
+		FString AttrName = BaseName + TEXT("Project_Requests");
+		Attributes.Emplace(MoveTemp(AttrName), ZenProjectStats.General.RequestCount);
+	}
+
+	{
+		FString AttrName = BaseName + TEXT("Project_BadRequests");
+		Attributes.Emplace(MoveTemp(AttrName), ZenProjectStats.General.BadRequestCount);
+	}
+
+	{
+		FString AttrName = BaseName + TEXT("Project_Op_HitRatio");
+		double Total = static_cast<double>(ZenProjectStats.General.Op.HitCount + ZenProjectStats.General.Op.MissCount);
+		Attributes.Emplace(MoveTemp(AttrName), Total > 0 ? static_cast<double>(ZenProjectStats.General.Op.HitCount) / Total : 0.0);
+	}
+
+	{
+		FString AttrName = BaseName + TEXT("Project_Chunk_HitRatio");
+		double Total = static_cast<double>(ZenProjectStats.General.Chunk.HitCount + ZenProjectStats.General.Chunk.MissCount);
+		Attributes.Emplace(MoveTemp(AttrName), Total > 0 ? static_cast<double>(ZenProjectStats.General.Chunk.HitCount) / Total : 0.0);
 	}
 
 	return true;
