DDC: Eliminated most blocking during put requests in HttpCacheStore

diff --git a/Engine/Source/Developer/DerivedDataCache/Private/HttpCacheStore.cpp b/Engine/Source/Developer/DerivedDataCache/Private/HttpCacheStore.cpp
--- a/Engine/Source/Developer/DerivedDataCache/Private/HttpCacheStore.cpp
+++ b/Engine/Source/Developer/DerivedDataCache/Private/HttpCacheStore.cpp
@@ -7,6 +7,7 @@
 
 #include "Algo/AllOf.h"
 #include "Algo/Transform.h"
+#include "Async/ManualResetEvent.h"
 #include "Async/Mutex.h"
 #include "Async/UniqueLock.h"
 #include "Compression/CompressedBuffer.h"
@@ -180,6 +181,122 @@ static bool TryResolveCanonicalHost(const FAnsiStringView Uri, FAnsiStringBuilde
 	return false;
 }
 
+class FHttpCacheStoreRequestQueue
+{
+public:
+	using FOnRequest = TUniqueFunction<void (THttpUniquePtr<IHttpRequest>&& Request)>;
+
+	void Initialize(IHttpConnectionPool& ConnectionPool, const FHttpClientParams& ClientParams)
+	{
+		FHttpClientParams QueueParams = ClientParams;
+		QueueParams.OnDestroyRequest = [this, OnDestroyRequest = MoveTemp(QueueParams.OnDestroyRequest)]
+		{
+			if (OnDestroyRequest)
+			{
+				OnDestroyRequest();
+			}
+			if (!Queue.IsEmpty())
+			{
+				if (THttpUniquePtr<IHttpRequest> Request = Client->TryCreateRequest({}))
+				{
+					if (FQueueRequest* Waiter = Queue.Pop())
+					{
+						Waiter->Complete(MoveTemp(Request));
+					}
+				}
+			}
+		};
+		Client = ConnectionPool.CreateClient(QueueParams);
+	}
+
+	void CreateRequestAsync(IRequestOwner& Owner, const FHttpRequestParams& Params, FOnRequest&& OnRequest)
+	{
+		if (Params.bIgnoreMaxRequests)
+		{
+			THttpUniquePtr<IHttpRequest> Request = Client->TryCreateRequest(Params);
+			checkf(Request, TEXT("IHttpClient::TryCreateRequest returned null in spite of bIgnoreMaxRequests."));
+			OnRequest(MoveTemp(Request));
+			return;
+		}
+
+		while (THttpUniquePtr<IHttpRequest> Request = Client->TryCreateRequest(Params))
+		{
+			if (FQueueRequest* Waiter = Queue.Pop())
+			{
+				Waiter->Complete(MoveTemp(Request));
+			}
+			else
+			{
+				OnRequest(MoveTemp(Request));
+				return;
+			}
+		}
+
+		Queue.Push(new FQueueRequest(Owner, MoveTemp(OnRequest)));
+
+		while (THttpUniquePtr<IHttpRequest> Request = Client->TryCreateRequest(Params))
+		{
+			if (FQueueRequest* Waiter = Queue.Pop())
+			{
+				Waiter->Complete(MoveTemp(Request));
+			}
+			else
+			{
+				return;
+			}
+		}
+	}
+
+private:
+	class FQueueRequest : FRequestBase
+	{
+	public:
+		FQueueRequest(IRequestOwner& InOwner, FOnRequest&& InOnRequest)
+			: Owner(InOwner)
+			, OnRequest(MoveTemp(InOnRequest))
+		{
+			Owner.Begin(this);
+		}
+
+		void Complete(THttpUniquePtr<IHttpRequest>&& Request)
+		{
+			if (bComplete.exchange(true))
+			{
+				return OnComplete.Wait();
+			}
+			Owner.End(this, [this](THttpUniquePtr<IHttpRequest>&& Request)
+			{
+				OnRequest(MoveTemp(Request));
+				OnComplete.Notify();
+			}, MoveTemp(Request));
+		}
+
+	private:
+		void SetPriority(EPriority Priority) final
+		{
+		}
+
+		void Cancel() final
+		{
+			Complete({});
+		}
+
+		void Wait() final
+		{
+			TRACE_CPUPROFILER_EVENT_SCOPE(HttpDDC_WaitOperation);
+			OnComplete.Wait();
+		}
+
+		IRequestOwner& Owner;
+		FOnRequest OnRequest;
+		FManualResetEvent OnComplete;
+		std::atomic<bool> bComplete = false;
+	};
+
+	THttpUniquePtr<IHttpClient> Client;
+	TLockFreePointerListFIFO<FQueueRequest, 0> Queue;
+};
+
 /**
  * Encapsulation for access token shared by all requests.
  */
@@ -323,7 +440,7 @@ private:
 	FHttpRequestQueue GetRequestQueues[2];
 	FHttpRequestQueue PutRequestQueues[2];
 	FHttpRequestQueue NonBlockingGetRequestQueue;
-	FHttpRequestQueue NonBlockingPutRequestQueue;
+	FHttpCacheStoreRequestQueue NonBlockingPutRequestQueue;
 
 	FCriticalSection AccessCs;
 	TUniquePtr<FHttpAccessToken> Access;
@@ -352,6 +469,8 @@ private:
 	class FHttpOperation;
 
 	TUniquePtr<FHttpOperation> WaitForHttpOperation(EOperationCategory Category);
+	void WaitForHttpOperationAsync(IRequestOwner& Owner, EOperationCategory Category, TUniqueFunction<void (TUniquePtr<FHttpOperation>&&)>&& OnOperation);
+	void WaitForHttpRequestAsync(IRequestOwner& Owner, EOperationCategory Category, TUniqueFunction<void (THttpUniquePtr<IHttpRequest>&&)>&& OnRequest);
 
 	void PutCacheRecordAsync(IRequestOwner& Owner, const FCachePutRequest& Request, FOnCachePutComplete&& OnComplete);
 	void PutCacheValueAsync(IRequestOwner& Owner, const FCachePutValueRequest& Request, FOnCachePutValueComplete&& OnComplete);
@@ -747,7 +866,9 @@ private:
 
 	FPutPackageOp(FHttpCacheStore& CacheStore, IRequestOwner& Owner, const FSharedString& Name);
 
-	void BeginPutRef(bool bFinalize, FOnCachePutRefComplete&& OnComplete);
+	void BeginOperation(bool bFinalize, FOnCachePutRefComplete&& OnComplete);
+
+	void BeginPutRef(TUniquePtr<FHttpOperation> Operation, bool bFinalize, FOnCachePutRefComplete&& OnComplete);
 	void EndPutRef(TUniquePtr<FHttpOperation> Operation, bool bFinalize, FOnCachePutRefComplete&& OnComplete);
 
 	void BeginPutBlobs(FCbPackage&& Package, FCachePutRefResponse&& Response);
@@ -772,13 +893,21 @@ void FHttpCacheStore::FPutPackageOp::Put(const FCacheKey& InKey, FCbPackage&& Pa
 	Object = Package.GetObject();
 	ObjectHash = Package.GetObjectHash();
 	OnPackageComplete = MoveTemp(OnComplete);
-	BeginPutRef(/*bFinalize*/ false, [Self = TRefCountPtr(this), Package = MoveTemp(Package)](FCachePutRefResponse&& Response) mutable
+	BeginOperation(/*bFinalize*/ false, [Self = TRefCountPtr(this), Package = MoveTemp(Package)](FCachePutRefResponse&& Response) mutable
 	{
 		return Self->BeginPutBlobs(MoveTemp(Package), MoveTemp(Response));
 	});
 }
 
-void FHttpCacheStore::FPutPackageOp::BeginPutRef(bool bFinalize, FOnCachePutRefComplete&& OnComplete)
+void FHttpCacheStore::FPutPackageOp::BeginOperation(bool bFinalize, FOnCachePutRefComplete&& OnComplete)
+{
+	CacheStore.WaitForHttpOperationAsync(Owner, EOperationCategory::Put, [Self = TRefCountPtr(this), bFinalize, OnComplete = MoveTemp(OnComplete)](TUniquePtr<FHttpOperation>&& Operation) mutable
+	{
+		Self->BeginPutRef(MoveTemp(Operation), bFinalize, MoveTemp(OnComplete));
+	});
+}
+
+void FHttpCacheStore::FPutPackageOp::BeginPutRef(TUniquePtr<FHttpOperation> Operation, bool bFinalize, FOnCachePutRefComplete&& OnComplete)
 {
 	FRequestTimer RequestTimer(RequestStats);
 
@@ -792,7 +921,6 @@ void FHttpCacheStore::FPutPackageOp::BeginPutRef(bool bFinalize, FOnCachePutRefC
 		RefsUri << ANSITEXTVIEW("/finalize/") << ObjectHash;
 	}
 
-	TUniquePtr<FHttpOperation> Operation = CacheStore.WaitForHttpOperation(EOperationCategory::Put);
 	FHttpOperation& LocalOperation = *Operation;
 	LocalOperation.SetUri(RefsUri);
 	if (bFinalize)
@@ -939,16 +1067,18 @@ void FHttpCacheStore::FPutPackageOp::BeginPutBlobs(FCbPackage&& Package, FCacheP
 	FRequestBarrier Barrier(Owner);
 	for (const FCompressedBuffer& Blob : Blobs)
 	{
-		TUniquePtr<FHttpOperation> Operation = CacheStore.WaitForHttpOperation(EOperationCategory::Put);
-		FHttpOperation& LocalOperation = *Operation;
-		LocalOperation.SetUri(WriteToAnsiString<256>(CacheStore.EffectiveDomain, ANSITEXTVIEW("/api/v1/compressed-blobs/"), CacheStore.Namespace, '/', Blob.GetRawHash()));
-		LocalOperation.SetMethod(EHttpMethod::Put);
-		LocalOperation.SetContentType(EHttpMediaType::CompressedBinary);
-		LocalOperation.SetBody(Blob.GetCompressed());
-		LocalOperation.SendAsync(Owner, [Self = TRefCountPtr(this), Operation = MoveTemp(Operation), LogicalSize = Blob.GetRawSize()]
+		CacheStore.WaitForHttpOperationAsync(Owner, EOperationCategory::Put, [this, Blob](TUniquePtr<FHttpOperation>&& Operation)
 		{
-			Operation->GetStats(Self->RequestStats);
-			Self->EndPutBlob(*Operation, LogicalSize);
+			FHttpOperation& LocalOperation = *Operation;
+			LocalOperation.SetUri(WriteToAnsiString<256>(CacheStore.EffectiveDomain, ANSITEXTVIEW("/api/v1/compressed-blobs/"), CacheStore.Namespace, '/', Blob.GetRawHash()));
+			LocalOperation.SetMethod(EHttpMethod::Put);
+			LocalOperation.SetContentType(EHttpMediaType::CompressedBinary);
+			LocalOperation.SetBody(Blob.GetCompressed());
+			LocalOperation.SendAsync(Owner, [Self = TRefCountPtr(this), Operation = MoveTemp(Operation), LogicalSize = Blob.GetRawSize()]
+			{
+				Operation->GetStats(Self->RequestStats);
+				Self->EndPutBlob(*Operation, LogicalSize);
+			});
 		});
 	}
 }
@@ -972,7 +1102,7 @@ void FHttpCacheStore::FPutPackageOp::EndPutBlob(FHttpOperation& Operation, uint6
 		}
 		else if (LocalSuccessfulBlobUploads == TotalBlobUploads)
 		{
-			BeginPutRef(/*bFinalize*/ true, [Self = TRefCountPtr(this)](FCachePutRefResponse&& Response)
+			BeginOperation(/*bFinalize*/ true, [Self = TRefCountPtr(this)](FCachePutRefResponse&& Response)
 			{
 				return Self->EndPutRefFinalize(MoveTemp(Response));
 			});
@@ -1962,7 +2092,7 @@ FHttpCacheStore::FHttpCacheStore(const FHttpCacheStoreParams& Params, ICacheStor
 
 		ClientParams.MaxRequests = UE_HTTPDDC_NONBLOCKING_PUT_REQUEST_POOL_SIZE;
 		ClientParams.MinRequests = UE_HTTPDDC_NONBLOCKING_PUT_REQUEST_POOL_SIZE;
-		NonBlockingPutRequestQueue = FHttpRequestQueue(*ConnectionPool, ClientParams);
+		NonBlockingPutRequestQueue.Initialize(*ConnectionPool, ClientParams);
 
 		bIsUsable = true;
 
@@ -2270,7 +2400,12 @@ TUniquePtr<FHttpCacheStore::FHttpOperation> FHttpCacheStore::WaitForHttpOperatio
 		}
 		else
 		{
-			Request = NonBlockingPutRequestQueue.CreateRequest(Params);
+			FRequestOwner BlockingOwner(EPriority::Blocking);
+			NonBlockingPutRequestQueue.CreateRequestAsync(BlockingOwner, Params, [&Request](THttpUniquePtr<IHttpRequest>&& AsyncRequest)
+			{
+				Request = MoveTemp(AsyncRequest);
+			});
+			BlockingOwner.Wait();
 		}
 	}
 	else
@@ -2294,6 +2429,52 @@ TUniquePtr<FHttpCacheStore::FHttpOperation> FHttpCacheStore::WaitForHttpOperatio
 	return MakeUnique<FHttpOperation>(MoveTemp(Request));
 }
 
+void FHttpCacheStore::WaitForHttpOperationAsync(IRequestOwner& Owner, EOperationCategory Category, TUniqueFunction<void (TUniquePtr<FHttpOperation>&&)>&& OnOperation)
+{
+	WaitForHttpRequestAsync(Owner, Category, [this, OnOperation = MoveTemp(OnOperation)](THttpUniquePtr<IHttpRequest>&& Request)
+	{
+		if (Access && RefreshAccessTokenTime > 0.0 && RefreshAccessTokenTime < FPlatformTime::Seconds())
+		{
+			AcquireAccessToken();
+		}
+
+		if (Access)
+		{
+			Request->AddHeader(ANSITEXTVIEW("Authorization"), WriteToAnsiString<1024>(*Access));
+		}
+
+		OnOperation(MakeUnique<FHttpOperation>(MoveTemp(Request)));
+	});
+}
+
+void FHttpCacheStore::WaitForHttpRequestAsync(IRequestOwner& Owner, EOperationCategory Category, TUniqueFunction<void (THttpUniquePtr<IHttpRequest>&&)>&& OnRequest)
+{
+	FHttpRequestParams Params;
+	if (FPlatformProcess::SupportsMultithreading() && bHttpEnableAsync)
+	{
+		if (Category == EOperationCategory::Get)
+		{
+			return OnRequest(NonBlockingGetRequestQueue.CreateRequest(Params));
+		}
+		else
+		{
+			return NonBlockingPutRequestQueue.CreateRequestAsync(Owner, Params, MoveTemp(OnRequest));
+		}
+	}
+	else
+	{
+		const bool bIsInGameThread = IsInGameThread();
+		if (Category == EOperationCategory::Get)
+		{
+			return OnRequest(GetRequestQueues[bIsInGameThread].CreateRequest(Params));
+		}
+		else
+		{
+			return OnRequest(PutRequestQueues[bIsInGameThread].CreateRequest(Params));
+		}
+	}
+}
+
 void FHttpCacheStore::PutCacheRecordAsync(IRequestOwner& Owner, const FCachePutRequest& Request, FOnCachePutComplete&& OnComplete)
 {
 	const FCacheKey& Key = Request.Record.GetKey();
