Fix bug in Zen DDC client.  Don't assume that if the data hash of a response is available as an attachment that it should be included in the reponse.  If the request is skipping data, then even if the attachment is present (eg: from another request in the batch), it should not be included when processing the request that had SkipData.

diff --git a/Engine/Source/Developer/DerivedDataCache/Private/ZenCacheStore.cpp b/Engine/Source/Developer/DerivedDataCache/Private/ZenCacheStore.cpp
--- a/Engine/Source/Developer/DerivedDataCache/Private/ZenCacheStore.cpp
+++ b/Engine/Source/Developer/DerivedDataCache/Private/ZenCacheStore.cpp
@@ -935,7 +935,8 @@ public:
 						{
 							FCbFieldView RawHashField = ResultObj["RawHash"];
 							FIoHash RawHash = RawHashField.AsHash();
-							if (const FCbAttachment* Attachment = Response.FindAttachment(RawHash))
+							const FCbAttachment* Attachment = EnumHasAnyFlags(RequestWithStats.Request.Policy, ECachePolicy::SkipData) ? nullptr : Response.FindAttachment(RawHash);
+							if (Attachment)
 							{
 								Value.Emplace(Attachment->AsCompressedBinary());
 							}
@@ -1154,7 +1155,9 @@ public:
 							RawHash = HashView.AsHash();
 							if (!HashView.HasError())
 							{
-								if (const FCbAttachment* Attachment = Response.FindAttachment(HashView.AsHash()))
+								const FCbAttachment* Attachment = EnumHasAnyFlags(RequestWithStats.Request.Policy, ECachePolicy::SkipData) ? nullptr : Response.FindAttachment(HashView.AsHash());
+
+								if (Attachment)
 								{
 									FCompressedBuffer CompressedBuffer = Attachment->AsCompressedBinary();
 									if (CompressedBuffer)
