Downgrade messages from ZenCacheStore from Warning to Display as they can be expected to come up due to network conditions or temporary outages and do not break functionality.

diff --git a/Engine/Source/Developer/DerivedDataCache/Private/ZenCacheStore.cpp b/Engine/Source/Developer/DerivedDataCache/Private/ZenCacheStore.cpp
--- a/Engine/Source/Developer/DerivedDataCache/Private/ZenCacheStore.cpp
+++ b/Engine/Source/Developer/DerivedDataCache/Private/ZenCacheStore.cpp
@@ -342,14 +342,14 @@ public:
 					}
 					if (RequestIndex != Batch.Num())
 					{
-						UE_LOG(LogDerivedDataCache, Warning,
+						UE_LOG(LogDerivedDataCache, Display,
 							TEXT("%s: Invalid response received from PutCacheRecords RPC: %d results expected, received %d, from %s"),
 							*CacheStore.GetName(), Batch.Num(), RequestIndex, *WriteToString<256>(*HttpResponse));
 					}
 				}
 				else if ((HttpResponse->GetErrorCode() != EHttpErrorCode::Canceled) && (HttpResponse->GetStatusCode() != 404))
 				{
-					UE_LOG(LogDerivedDataCache, Warning,
+					UE_LOG(LogDerivedDataCache, Display,
 						TEXT("%s: Error response received from PutCacheRecords RPC: from %s"),
 						*CacheStore.GetName(), *WriteToString<256>(*HttpResponse));
 				}
@@ -549,14 +549,14 @@ public:
 					}
 					if (RequestIndex != Batch.Num())
 					{
-						UE_LOG(LogDerivedDataCache, Warning,
+						UE_LOG(LogDerivedDataCache, Display,
 							TEXT("%s: Invalid response received from GetCacheRecords RPC: %d results expected, received %d, from %s"),
 							*CacheStore.GetName(), Batch.Num(), RequestIndex, *WriteToString<256>(*HttpResponse));
 					}
 				}
 				else if ((HttpResponse->GetErrorCode() != EHttpErrorCode::Canceled) && (HttpResponse->GetStatusCode() != 404))
 				{
-					UE_LOG(LogDerivedDataCache, Warning,
+					UE_LOG(LogDerivedDataCache, Display,
 						TEXT("%s: Error response received from GetCacheRecords RPC: from %s"),
 						*CacheStore.GetName(), *WriteToString<256>(*HttpResponse));
 				}
@@ -745,14 +745,14 @@ public:
 					}
 					if (RequestIndex != Batch.Num())
 					{
-						UE_LOG(LogDerivedDataCache, Warning,
+						UE_LOG(LogDerivedDataCache, Display,
 							TEXT("%s: Invalid response received from PutCacheValues RPC: %d results expected, received %d, from %s"),
 							*CacheStore.GetName(), Batch.Num(), RequestIndex, *WriteToString<256>(*HttpResponse));
 					}
 				}
 				else if ((HttpResponse->GetErrorCode() != EHttpErrorCode::Canceled) && (HttpResponse->GetStatusCode() != 404))
 				{
-					UE_LOG(LogDerivedDataCache, Warning,
+					UE_LOG(LogDerivedDataCache, Display,
 						TEXT("%s: Error response received from PutCacheValues RPC: from %s"),
 						*CacheStore.GetName(), *WriteToString<256>(*HttpResponse));
 				}
@@ -954,14 +954,14 @@ public:
 					}
 					if (RequestIndex != Batch.Num())
 					{
-						UE_LOG(LogDerivedDataCache, Warning,
+						UE_LOG(LogDerivedDataCache, Display,
 							TEXT("%s: Invalid response received from GetCacheValues RPC: %d results expected, received %d from %s"),
 							*CacheStore.GetName(), Batch.Num(), RequestIndex, *WriteToString<256>(*HttpResponse));
 					}
 				}
 				else if ((HttpResponse->GetErrorCode() != EHttpErrorCode::Canceled) && (HttpResponse->GetStatusCode() != 404))
 				{
-					UE_LOG(LogDerivedDataCache, Warning,
+					UE_LOG(LogDerivedDataCache, Display,
 						TEXT("%s: Error response received from GetCacheValues RPC: from %s"),
 						*CacheStore.GetName(), *WriteToString<256>(*HttpResponse));
 				}
@@ -1185,7 +1185,7 @@ public:
 				}
 				else if ((HttpResponse->GetErrorCode() != EHttpErrorCode::Canceled) && (HttpResponse->GetStatusCode() != 404))
 				{
-					UE_LOG(LogDerivedDataCache, Warning,
+					UE_LOG(LogDerivedDataCache, Display,
 						TEXT("%s: Error response received from GetChunks RPC: from %s"),
 						*CacheStore.GetName(), *WriteToString<256>(*HttpResponse));
 				}
