DDC: Added cancellation support to FHttpCacheStoreRequestQueue

#rnx

diff --git a/Engine/Source/Developer/DerivedDataCache/Private/HttpCacheStore.cpp b/Engine/Source/Developer/DerivedDataCache/Private/HttpCacheStore.cpp
--- a/Engine/Source/Developer/DerivedDataCache/Private/HttpCacheStore.cpp
+++ b/Engine/Source/Developer/DerivedDataCache/Private/HttpCacheStore.cpp
@@ -262,7 +262,8 @@ private:
 		{
 			if (bComplete.exchange(true))
 			{
-				return OnComplete.Wait();
+				OnComplete.Wait();
+				return;
 			}
 			Owner.End(this, [this](THttpUniquePtr<IHttpRequest>&& Request)
 			{
@@ -469,7 +470,11 @@ private:
 	class FHttpOperation;
 
 	TUniquePtr<FHttpOperation> WaitForHttpOperation(EOperationCategory Category);
+
+	/** Invokes the callback when an operation is available, or with null if canceled. */
 	void WaitForHttpOperationAsync(IRequestOwner& Owner, EOperationCategory Category, TUniqueFunction<void (TUniquePtr<FHttpOperation>&&)>&& OnOperation);
+
+	/** Invokes the callback when a request is available, or with null if canceled. */
 	void WaitForHttpRequestAsync(IRequestOwner& Owner, EOperationCategory Category, TUniqueFunction<void (THttpUniquePtr<IHttpRequest>&&)>&& OnRequest);
 
 	void PutCacheRecordAsync(IRequestOwner& Owner, const FCachePutRequest& Request, FOnCachePutComplete&& OnComplete);
@@ -872,7 +877,7 @@ private:
 	void EndPutRef(TUniquePtr<FHttpOperation> Operation, bool bFinalize, FOnCachePutRefComplete&& OnComplete);
 
 	void BeginPutBlobs(FCbPackage&& Package, FCachePutRefResponse&& Response);
-	void EndPutBlob(FHttpOperation& Operation, uint64 LogicalSize);
+	void EndPutBlob(FHttpOperation* Operation, uint64 LogicalSize);
 
 	void EndPutRefFinalize(FCachePutRefResponse&& Response);
 
@@ -895,7 +900,7 @@ void FHttpCacheStore::FPutPackageOp::Put(const FCacheKey& InKey, FCbPackage&& Pa
 	OnPackageComplete = MoveTemp(OnComplete);
 	BeginOperation(/*bFinalize*/ false, [Self = TRefCountPtr(this), Package = MoveTemp(Package)](FCachePutRefResponse&& Response) mutable
 	{
-		return Self->BeginPutBlobs(MoveTemp(Package), MoveTemp(Response));
+		Self->BeginPutBlobs(MoveTemp(Package), MoveTemp(Response));
 	});
 }
 
@@ -909,6 +914,12 @@ void FHttpCacheStore::FPutPackageOp::BeginOperation(bool bFinalize, FOnCachePutR
 
 void FHttpCacheStore::FPutPackageOp::BeginPutRef(TUniquePtr<FHttpOperation> Operation, bool bFinalize, FOnCachePutRefComplete&& OnComplete)
 {
+	if (UNLIKELY(!Operation))
+	{
+		OnComplete({{}, EStatus::Canceled});
+		return;
+	}
+
 	FRequestTimer RequestTimer(RequestStats);
 
 	TAnsiStringBuilder<64> Bucket;
@@ -955,7 +966,8 @@ void FHttpCacheStore::FPutPackageOp::EndPutRef(
 	if (const int32 StatusCode = Operation->GetStatusCode(); StatusCode < 200 || StatusCode > 204)
 	{
 		const EStatus Status = Operation->GetErrorCode() == EHttpErrorCode::Canceled ? EStatus::Canceled : EStatus::Error;
-		return OnComplete({{}, Status});
+		OnComplete({{}, Status});
+		return;
 	}
 
 	FRequestTimer RequestTimer(RequestStats);
@@ -1006,7 +1018,8 @@ void FHttpCacheStore::FPutPackageOp::BeginPutBlobs(FCbPackage&& Package, FCacheP
 			UE_LOG(LogDerivedDataCache, Log, TEXT("%s: Failed to put reference object for put of %s from '%s'"),
 				*CacheStore.Domain, *WriteToString<96>(Key), *Name);
 		}
-		return EndPut(Response.Status);
+		EndPut(Response.Status);
+		return;
 	}
 
 	FRequestTimer RequestTimer(RequestStats);
@@ -1058,7 +1071,8 @@ void FHttpCacheStore::FPutPackageOp::BeginPutBlobs(FCbPackage&& Package, FCacheP
 	if (Blobs.IsEmpty())
 	{
 		RequestTimer.Stop();
-		return EndPut(EStatus::Ok);
+		EndPut(EStatus::Ok);
+		return;
 	}
 
 	TotalBlobUploads = Blobs.Num();
@@ -1067,30 +1081,38 @@ void FHttpCacheStore::FPutPackageOp::BeginPutBlobs(FCbPackage&& Package, FCacheP
 	FRequestBarrier Barrier(Owner);
 	for (const FCompressedBuffer& Blob : Blobs)
 	{
-		CacheStore.WaitForHttpOperationAsync(Owner, EOperationCategory::Put, [this, Blob](TUniquePtr<FHttpOperation>&& Operation)
+		CacheStore.WaitForHttpOperationAsync(Owner, EOperationCategory::Put, [Self = TRefCountPtr(this), Blob](TUniquePtr<FHttpOperation>&& Operation)
 		{
+			if (UNLIKELY(!Operation))
+			{
+				Self->EndPutBlob(nullptr, 0);
+				return;
+			}
 			FHttpOperation& LocalOperation = *Operation;
-			LocalOperation.SetUri(WriteToAnsiString<256>(CacheStore.EffectiveDomain, ANSITEXTVIEW("/api/v1/compressed-blobs/"), CacheStore.Namespace, '/', Blob.GetRawHash()));
+			LocalOperation.SetUri(WriteToAnsiString<256>(Self->CacheStore.EffectiveDomain, ANSITEXTVIEW("/api/v1/compressed-blobs/"), Self->CacheStore.Namespace, '/', Blob.GetRawHash()));
 			LocalOperation.SetMethod(EHttpMethod::Put);
 			LocalOperation.SetContentType(EHttpMediaType::CompressedBinary);
 			LocalOperation.SetBody(Blob.GetCompressed());
-			LocalOperation.SendAsync(Owner, [Self = TRefCountPtr(this), Operation = MoveTemp(Operation), LogicalSize = Blob.GetRawSize()]
+			LocalOperation.SendAsync(Self->Owner, [Self, Operation = MoveTemp(Operation), LogicalSize = Blob.GetRawSize()]
 			{
 				Operation->GetStats(Self->RequestStats);
-				Self->EndPutBlob(*Operation, LogicalSize);
+				Self->EndPutBlob(Operation.Get(), LogicalSize);
 			});
 		});
 	}
 }
 
-void FHttpCacheStore::FPutPackageOp::EndPutBlob(FHttpOperation& Operation, uint64 LogicalSize)
+void FHttpCacheStore::FPutPackageOp::EndPutBlob(FHttpOperation* Operation, uint64 LogicalSize)
 {
-	const int32 StatusCode = Operation.GetStatusCode();
-	if (StatusCode >= 200 && StatusCode <= 204)
+	if (Operation)
 	{
-		SuccessfulBlobUploads.fetch_add(1, std::memory_order_relaxed);
-		TUniqueLock Lock(RequestStats.Mutex);
-		RequestStats.LogicalWriteSize += LogicalSize;
+		const int32 StatusCode = Operation->GetStatusCode();
+		if (StatusCode >= 200 && StatusCode <= 204)
+		{
+			SuccessfulBlobUploads.fetch_add(1, std::memory_order_relaxed);
+			TUniqueLock Lock(RequestStats.Mutex);
+			RequestStats.LogicalWriteSize += LogicalSize;
+		}
 	}
 
 	if (PendingBlobUploads.fetch_sub(1, std::memory_order_relaxed) == 1)
@@ -1104,7 +1126,7 @@ void FHttpCacheStore::FPutPackageOp::EndPutBlob(FHttpOperation& Operation, uint6
 		{
 			BeginOperation(/*bFinalize*/ true, [Self = TRefCountPtr(this)](FCachePutRefResponse&& Response)
 			{
-				return Self->EndPutRefFinalize(MoveTemp(Response));
+				Self->EndPutRefFinalize(MoveTemp(Response));
 			});
 		}
 		else
@@ -1125,7 +1147,7 @@ void FHttpCacheStore::FPutPackageOp::EndPutRefFinalize(FCachePutRefResponse&& Re
 			*CacheStore.Domain, *WriteToString<96>(Key), *Name);
 	}
 
-	return EndPut(Response.Status);
+	EndPut(Response.Status);
 }
 
 void FHttpCacheStore::FPutPackageOp::EndPut(EStatus Status)
@@ -1377,7 +1399,8 @@ void FHttpCacheStore::FGetRecordOp::BeginGetValues(const FCacheRecord& Record, c
 
 	if (PendingValues == 0)
 	{
-		return EndGetValues(Policy, EStatus::Ok);
+		EndGetValues(Policy, EStatus::Ok);
+		return;
 	}
 
 	GetValues(RequiredGets, [Self = TRefCountPtr(this), Policy](FValueResponse&& Response)
@@ -2433,6 +2456,12 @@ void FHttpCacheStore::WaitForHttpOperationAsync(IRequestOwner& Owner, EOperation
 {
 	WaitForHttpRequestAsync(Owner, Category, [this, OnOperation = MoveTemp(OnOperation)](THttpUniquePtr<IHttpRequest>&& Request)
 	{
+		if (UNLIKELY(!Request))
+		{
+			OnOperation({});
+			return;
+		}
+
 		if (Access && RefreshAccessTokenTime > 0.0 && RefreshAccessTokenTime < FPlatformTime::Seconds())
 		{
 			AcquireAccessToken();
@@ -2454,11 +2483,11 @@ void FHttpCacheStore::WaitForHttpRequestAsync(IRequestOwner& Owner, EOperationCa
 	{
 		if (Category == EOperationCategory::Get)
 		{
-			return OnRequest(NonBlockingGetRequestQueue.CreateRequest(Params));
+			OnRequest(NonBlockingGetRequestQueue.CreateRequest(Params));
 		}
 		else
 		{
-			return NonBlockingPutRequestQueue.CreateRequestAsync(Owner, Params, MoveTemp(OnRequest));
+			NonBlockingPutRequestQueue.CreateRequestAsync(Owner, Params, MoveTemp(OnRequest));
 		}
 	}
 	else
@@ -2466,11 +2495,11 @@ void FHttpCacheStore::WaitForHttpRequestAsync(IRequestOwner& Owner, EOperationCa
 		const bool bIsInGameThread = IsInGameThread();
 		if (Category == EOperationCategory::Get)
 		{
-			return OnRequest(GetRequestQueues[bIsInGameThread].CreateRequest(Params));
+			OnRequest(GetRequestQueues[bIsInGameThread].CreateRequest(Params));
 		}
 		else
 		{
-			return OnRequest(PutRequestQueues[bIsInGameThread].CreateRequest(Params));
+			OnRequest(PutRequestQueues[bIsInGameThread].CreateRequest(Params));
 		}
 	}
 }
