DDC: Removed the synchronous request pools from HttpCacheStore

These are not used by default and were only kept to support A/B testing in the past.

#rnx

diff --git a/Engine/Source/Developer/DerivedDataCache/Private/HttpCacheStore.cpp b/Engine/Source/Developer/DerivedDataCache/Private/HttpCacheStore.cpp
--- a/Engine/Source/Developer/DerivedDataCache/Private/HttpCacheStore.cpp
+++ b/Engine/Source/Developer/DerivedDataCache/Private/HttpCacheStore.cpp
@@ -63,23 +63,14 @@
 #include "Ssl.h"
 #endif
 
-#define UE_HTTPDDC_GET_REQUEST_POOL_SIZE 48
-#define UE_HTTPDDC_PUT_REQUEST_POOL_SIZE 16
-#define UE_HTTPDDC_NONBLOCKING_GET_REQUEST_POOL_SIZE 128
-#define UE_HTTPDDC_NONBLOCKING_PUT_REQUEST_POOL_SIZE 24
+#define UE_HTTPDDC_GET_REQUEST_POOL_SIZE 128
+#define UE_HTTPDDC_PUT_REQUEST_POOL_SIZE 24
 #define UE_HTTPDDC_MAX_FAILED_LOGIN_ATTEMPTS 16
 #define UE_HTTPDDC_MAX_ATTEMPTS 4
 
 namespace UE::DerivedData
 {
 
-static bool bHttpEnableAsync = true;
-static FAutoConsoleVariableRef CVarHttpEnableAsync(
-	TEXT("DDC.Http.EnableAsync"),
-	bHttpEnableAsync,
-	TEXT("If true, asynchronous operations are permitted, otherwise all operations are forced to be synchronous."),
-	ECVF_Default);
-
 TRACE_DECLARE_INT_COUNTER(HttpDDC_Get, TEXT("HttpDDC Get"));
 TRACE_DECLARE_INT_COUNTER(HttpDDC_GetHit, TEXT("HttpDDC Get Hit"));
 TRACE_DECLARE_INT_COUNTER(HttpDDC_Put, TEXT("HttpDDC Put"));
@@ -438,10 +429,8 @@ private:
 	FDerivedDataCacheUsageStats UsageStats;
 	FBackendDebugOptions DebugOptions;
 	THttpUniquePtr<IHttpConnectionPool> ConnectionPool;
-	FHttpRequestQueue GetRequestQueues[2];
-	FHttpRequestQueue PutRequestQueues[2];
-	FHttpRequestQueue NonBlockingGetRequestQueue;
-	FHttpCacheStoreRequestQueue NonBlockingPutRequestQueue;
+	FHttpCacheStoreRequestQueue GetRequestQueue;
+	FHttpCacheStoreRequestQueue PutRequestQueue;
 
 	FCriticalSection AccessCs;
 	TUniquePtr<FHttpAccessToken> Access;
@@ -2101,21 +2090,11 @@ FHttpCacheStore::FHttpCacheStore(const FHttpCacheStoreParams& Params, ICacheStor
 	{
 		ClientParams.MaxRequests = UE_HTTPDDC_GET_REQUEST_POOL_SIZE;
 		ClientParams.MinRequests = UE_HTTPDDC_GET_REQUEST_POOL_SIZE;
-		GetRequestQueues[0] = FHttpRequestQueue(*ConnectionPool, ClientParams);
-		GetRequestQueues[1] = FHttpRequestQueue(*ConnectionPool, ClientParams);
+		GetRequestQueue.Initialize(*ConnectionPool, ClientParams);
 
 		ClientParams.MaxRequests = UE_HTTPDDC_PUT_REQUEST_POOL_SIZE;
 		ClientParams.MinRequests = UE_HTTPDDC_PUT_REQUEST_POOL_SIZE;
-		PutRequestQueues[0] = FHttpRequestQueue(*ConnectionPool, ClientParams);
-		PutRequestQueues[1] = FHttpRequestQueue(*ConnectionPool, ClientParams);
-
-		ClientParams.MaxRequests = UE_HTTPDDC_NONBLOCKING_GET_REQUEST_POOL_SIZE;
-		ClientParams.MinRequests = UE_HTTPDDC_NONBLOCKING_GET_REQUEST_POOL_SIZE;
-		NonBlockingGetRequestQueue = FHttpRequestQueue(*ConnectionPool, ClientParams);
-
-		ClientParams.MaxRequests = UE_HTTPDDC_NONBLOCKING_PUT_REQUEST_POOL_SIZE;
-		ClientParams.MinRequests = UE_HTTPDDC_NONBLOCKING_PUT_REQUEST_POOL_SIZE;
-		NonBlockingPutRequestQueue.Initialize(*ConnectionPool, ClientParams);
+		PutRequestQueue.Initialize(*ConnectionPool, ClientParams);
 
 		bIsUsable = true;
 
@@ -2414,34 +2393,15 @@ TUniquePtr<FHttpCacheStore::FHttpOperation> FHttpCacheStore::WaitForHttpOperatio
 
 	THttpUniquePtr<IHttpRequest> Request;
 
-	FHttpRequestParams Params;
-	if (FPlatformProcess::SupportsMultithreading() && bHttpEnableAsync)
 	{
-		if (Category == EOperationCategory::Get)
-		{
-			Request = NonBlockingGetRequestQueue.CreateRequest(Params);
-		}
-		else
-		{
-			FRequestOwner BlockingOwner(EPriority::Blocking);
-			NonBlockingPutRequestQueue.CreateRequestAsync(BlockingOwner, Params, [&Request](THttpUniquePtr<IHttpRequest>&& AsyncRequest)
-			{
-				Request = MoveTemp(AsyncRequest);
-			});
-			BlockingOwner.Wait();
-		}
-	}
-	else
-	{
-		const bool bIsInGameThread = IsInGameThread();
-		if (Category == EOperationCategory::Get)
-		{
-			Request = GetRequestQueues[bIsInGameThread].CreateRequest(Params);
-		}
-		else
+		FHttpRequestParams Params;
+		FRequestOwner BlockingOwner(EPriority::Blocking);
+		FHttpCacheStoreRequestQueue& RequestQueue = (Category == EOperationCategory::Get) ? GetRequestQueue : PutRequestQueue;
+		RequestQueue.CreateRequestAsync(BlockingOwner, Params, [&Request](THttpUniquePtr<IHttpRequest>&& AsyncRequest)
 		{
-			Request = PutRequestQueues[bIsInGameThread].CreateRequest(Params);
-		}
+			Request = MoveTemp(AsyncRequest);
+		});
+		BlockingOwner.Wait();
 	}
 
 	if (Access)
@@ -2479,29 +2439,8 @@ void FHttpCacheStore::WaitForHttpOperationAsync(IRequestOwner& Owner, EOperation
 void FHttpCacheStore::WaitForHttpRequestAsync(IRequestOwner& Owner, EOperationCategory Category, TUniqueFunction<void (THttpUniquePtr<IHttpRequest>&&)>&& OnRequest)
 {
 	FHttpRequestParams Params;
-	if (FPlatformProcess::SupportsMultithreading() && bHttpEnableAsync)
-	{
-		if (Category == EOperationCategory::Get)
-		{
-			OnRequest(NonBlockingGetRequestQueue.CreateRequest(Params));
-		}
-		else
-		{
-			NonBlockingPutRequestQueue.CreateRequestAsync(Owner, Params, MoveTemp(OnRequest));
-		}
-	}
-	else
-	{
-		const bool bIsInGameThread = IsInGameThread();
-		if (Category == EOperationCategory::Get)
-		{
-			OnRequest(GetRequestQueues[bIsInGameThread].CreateRequest(Params));
-		}
-		else
-		{
-			OnRequest(PutRequestQueues[bIsInGameThread].CreateRequest(Params));
-		}
-	}
+	FHttpCacheStoreRequestQueue& RequestQueue = (Category == EOperationCategory::Get) ? GetRequestQueue : PutRequestQueue;
+	RequestQueue.CreateRequestAsync(Owner, Params, MoveTemp(OnRequest));
 }
 
 void FHttpCacheStore::PutCacheRecordAsync(IRequestOwner& Owner, const FCachePutRequest& Request, FOnCachePutComplete&& OnComplete)
