Remove erroneous text from filesystem cache store speed test logs.  The logic for whether the file store gets used or not is more subtle than implied by the log text.

diff --git a/Engine/Source/Developer/DerivedDataCache/Private/FileSystemCacheStore.cpp b/Engine/Source/Developer/DerivedDataCache/Private/FileSystemCacheStore.cpp
--- a/Engine/Source/Developer/DerivedDataCache/Private/FileSystemCacheStore.cpp
+++ b/Engine/Source/Developer/DerivedDataCache/Private/FileSystemCacheStore.cpp
@@ -1387,7 +1387,7 @@ bool FFileSystemCacheStore::RunSpeedTest(
 				TCHAR ErrorBuffer[1024];
 				FPlatformMisc::GetSystemErrorMessage(ErrorBuffer, 1024, ErrorCode);
 				UE_LOG(LogDerivedDataCache, Warning,
-					TEXT("%s: Failed to read from %s, this cache store will be disabled. ReadError: %u (%s)"),
+					TEXT("%s: Failed to read from %s. ReadError: %u (%s)"),
 					*CachePath, *FilePath, ErrorCode, ErrorBuffer);
 				bReadTestPassed = false;
 				break;
@@ -1436,7 +1436,7 @@ bool FFileSystemCacheStore::RunSpeedTest(
 				TCHAR ErrorBuffer[1024];
 				FPlatformMisc::GetSystemErrorMessage(ErrorBuffer, 1024, ErrorCode);
 				UE_LOG(LogDerivedDataCache, Warning,
-					TEXT("%s: Failed to write to %s, this cache store will be disabled. ReadError: %u (%s)"),
+					TEXT("%s: Failed to write to %s. WriteError: %u (%s)"),
 					*CachePath, *FilePath, ErrorCode, ErrorBuffer);
 				bWriteTestPassed = false;
 				break;
