DDC: Fixed a crash when a boot cache is present

diff --git a/Engine/Source/Developer/DerivedDataCache/Private/DerivedDataBackends.cpp b/Engine/Source/Developer/DerivedDataCache/Private/DerivedDataBackends.cpp
--- a/Engine/Source/Developer/DerivedDataCache/Private/DerivedDataBackends.cpp
+++ b/Engine/Source/Developer/DerivedDataCache/Private/DerivedDataBackends.cpp
@@ -566,7 +566,7 @@ public:
 
 	#if WITH_EDITOR
 		// Only allow boot cache with the editor. We don't want other tools and utilities (eg. SCW) writing to the same file.
-		CreateMemoryCacheStore(BootCache, TEXT("Boot"), Entry, this);
+		CreateMemoryCacheStore(BootCache, TEXT("Boot"), *WriteToString<128>(TEXT("-Boot "), Entry), this);
 	#endif
 
 		return !!BootCache;
diff --git a/Engine/Source/Developer/DerivedDataCache/Private/MemoryCacheStore.cpp b/Engine/Source/Developer/DerivedDataCache/Private/MemoryCacheStore.cpp
--- a/Engine/Source/Developer/DerivedDataCache/Private/MemoryCacheStore.cpp
+++ b/Engine/Source/Developer/DerivedDataCache/Private/MemoryCacheStore.cpp
@@ -129,7 +129,7 @@ FMemoryCacheStore::FMemoryCacheStore(
 	, bDisabled(false)
 	, CurrentCacheSize(0)
 	, bMaxSizeExceeded(false)
-	, bCanBeDisabled(false)
+	, bCanBeDisabled(FParse::Param(InConfig, TEXT("Boot")))
 {
 	OutCache = this;
 
