Allow HttpCacheStore to retry on requests that experience an Unknown error code.  This is to help mitigate issues where the client experiences "connection reset by peer".

diff --git a/Engine/Source/Developer/DerivedDataCache/Private/HttpCacheStore.cpp b/Engine/Source/Developer/DerivedDataCache/Private/HttpCacheStore.cpp
--- a/Engine/Source/Developer/DerivedDataCache/Private/HttpCacheStore.cpp
+++ b/Engine/Source/Developer/DerivedDataCache/Private/HttpCacheStore.cpp
@@ -606,7 +606,8 @@ private:
 			return false;
 		}
 
-		if (LocalResponse.GetErrorCode() == EHttpErrorCode::TimedOut)
+		EHttpErrorCode ErrorCode = LocalResponse.GetErrorCode();
+		if ((ErrorCode == EHttpErrorCode::TimedOut) || (ErrorCode == EHttpErrorCode::Unknown))
 		{
 			return true;
 		}
