Avoid printing response body when receiving an unexpected error in the HttpCacheStore.  The response body could be of arbitrary type and could be truncated or malformed, so it's not safe to attempt to reinterpret it as a UTF8 string.  If GetBodyAsString is expanded later to cover all media types and the possibility of truncated or malformed data, we can consider re-enabling the printing of the body at the time of an unexpected error.

diff --git a/Engine/Source/Developer/DerivedDataCache/Private/HttpCacheStore.cpp b/Engine/Source/Developer/DerivedDataCache/Private/HttpCacheStore.cpp
--- a/Engine/Source/Developer/DerivedDataCache/Private/HttpCacheStore.cpp
+++ b/Engine/Source/Developer/DerivedDataCache/Private/HttpCacheStore.cpp
@@ -645,11 +645,8 @@ private:
 
 			if (bUnexpectedError)
 			{
-				FString Body = Operation->GetBodyAsString();
-				Body.ReplaceCharInline(TEXT('\r'), TEXT(' '));
-				Body.ReplaceCharInline(TEXT('\n'), TEXT(' '));
 				UE_LOG(LogDerivedDataCache, Display,
-					TEXT("HTTP: %s (%s) %s"), *WriteToString<256>(LocalResponse), *StatsText, *Body);
+					TEXT("HTTP: %s (%s)"), *WriteToString<256>(LocalResponse), *StatsText);
 			}
 			else
 			{
