DDC: Eliminated most blocking during get requests in HttpCacheStore

diff --git a/Engine/Source/Developer/DerivedDataCache/Private/HttpCacheStore.cpp b/Engine/Source/Developer/DerivedDataCache/Private/HttpCacheStore.cpp
--- a/Engine/Source/Developer/DerivedDataCache/Private/HttpCacheStore.cpp
+++ b/Engine/Source/Developer/DerivedDataCache/Private/HttpCacheStore.cpp
@@ -458,8 +458,6 @@ private:
 
 	class FHttpOperation;
 
-	TUniquePtr<FHttpOperation> WaitForHttpOperation(EOperationCategory Category);
-
 	/** Invokes the callback when an operation is available, or with null if canceled. */
 	void WaitForHttpOperationAsync(IRequestOwner& Owner, EOperationCategory Category, TUniqueFunction<void (TUniquePtr<FHttpOperation>&&)>&& OnOperation);
 
@@ -1075,6 +1073,7 @@ void FHttpCacheStore::FPutPackageOp::BeginPutBlobs(FCbPackage&& Package, FCacheP
 				Self->EndPutBlob(nullptr, 0);
 				return;
 			}
+
 			FHttpOperation& LocalOperation = *Operation;
 			LocalOperation.SetUri(WriteToAnsiString<256>(Self->CacheStore.EffectiveDomain, ANSITEXTVIEW("/api/v1/compressed-blobs/"), Self->CacheStore.Namespace, '/', Blob.GetRawHash()));
 			LocalOperation.SetMethod(EHttpMethod::Put);
@@ -1196,6 +1195,12 @@ private:
 	void BeginGetValues(const FCacheRecord& Record, const FCacheRecordPolicy& Policy, FOnRecordComplete&& OnComplete);
 	void EndGetValues(const FCacheRecordPolicy& Policy, EStatus Status);
 
+	void BeginGetValue(TUniquePtr<FHttpOperation>&& Operation, const FValueWithId& Value, const TSharedRef<FOnValueComplete>& OnComplete);
+	void EndGetValue(FHttpOperation& Operation, const FValueWithId& Value, const FOnValueComplete& OnComplete);
+
+	void BeginGetValuesExist(TUniquePtr<FHttpOperation>&& Operation, TArray<FValueWithId>&& Values, FOnValueComplete&& OnComplete);
+	void EndGetValuesExist(FHttpOperation* Operation, TArray<FValueWithId>&& Values, FOnValueComplete&& OnComplete);
+
 	FHttpCacheStore& CacheStore;
 	IRequestOwner& Owner;
 	FSharedString Name;
@@ -1250,21 +1255,29 @@ void FHttpCacheStore::FGetRecordOp::GetRecordOnly(const FCacheKey& InKey, const
 	OnRecordComplete = MoveTemp(InOnComplete);
 	RequestStats.Bucket = Key.Bucket;
 
-	TAnsiStringBuilder<64> Bucket;
-	Algo::Transform(Key.Bucket.ToString(), AppendChars(Bucket), FCharAnsi::ToLower);
-
-	TUniquePtr<FHttpOperation> Operation = CacheStore.WaitForHttpOperation(EOperationCategory::Get);
-	FHttpOperation& LocalOperation = *Operation;
-	LocalOperation.SetUri(WriteToAnsiString<256>(CacheStore.EffectiveDomain, ANSITEXTVIEW("/api/v1/refs/"), CacheStore.Namespace, '/', Bucket, '/', Key.Hash));
-	LocalOperation.SetMethod(EHttpMethod::Get);
-	LocalOperation.AddAcceptType(EHttpMediaType::CbObject);
-	LocalOperation.SetExpectedErrorCodes({404});
-
 	RequestTimer.Stop();
-	LocalOperation.SendAsync(Owner, [Self = TRefCountPtr(this), Operation = MoveTemp(Operation)]() mutable
+	CacheStore.WaitForHttpOperationAsync(Owner, EOperationCategory::Get, [Self = TRefCountPtr(this)](TUniquePtr<FHttpOperation>&& Operation)
 	{
-		Operation->GetStats(Self->RequestStats);
-		Self->EndGetRef(MoveTemp(Operation));
+		if (UNLIKELY(!Operation))
+		{
+			Self->EndGetRef(MoveTemp(Operation));
+			return;
+		}
+
+		TAnsiStringBuilder<64> Bucket;
+		Algo::Transform(Self->Key.Bucket.ToString(), AppendChars(Bucket), FCharAnsi::ToLower);
+
+		FHttpOperation& LocalOperation = *Operation;
+		LocalOperation.SetUri(WriteToAnsiString<256>(Self->CacheStore.EffectiveDomain, ANSITEXTVIEW("/api/v1/refs/"), Self->CacheStore.Namespace, '/', Bucket, '/', Self->Key.Hash));
+		LocalOperation.SetMethod(EHttpMethod::Get);
+		LocalOperation.AddAcceptType(EHttpMediaType::CbObject);
+		LocalOperation.SetExpectedErrorCodes({404});
+
+		LocalOperation.SendAsync(Self->Owner, [Self, Operation = MoveTemp(Operation)]() mutable
+		{
+			Operation->GetStats(Self->RequestStats);
+			Self->EndGetRef(MoveTemp(Operation));
+		});
 	});
 }
 
@@ -1275,7 +1288,7 @@ void FHttpCacheStore::FGetRecordOp::EndGetRef(TUniquePtr<FHttpOperation> Operati
 	FRequestTimer RequestTimer(RequestStats);
 
 	FOptionalCacheRecord Record;
-	EStatus Status = EStatus::Error;
+	EStatus Status = Operation ? EStatus::Error : EStatus::Canceled;
 	ON_SCOPE_EXIT
 	{
 		Operation.Reset();
@@ -1288,6 +1301,11 @@ void FHttpCacheStore::FGetRecordOp::EndGetRef(TUniquePtr<FHttpOperation> Operati
 		LocalOnComplete({MoveTemp(Record).Get(), Status});
 	};
 
+	if (UNLIKELY(!Operation))
+	{
+		return;
+	}
+
 	const int32 StatusCode = Operation->GetStatusCode();
 	if (StatusCode < 200 || StatusCode > 204)
 	{
@@ -1471,75 +1489,97 @@ void FHttpCacheStore::FGetRecordOp::GetValues(TConstArrayView<FValueWithId> Valu
 	{
 		if (Value.HasData())
 		{
+			(*SharedOnComplete)({Value, EStatus::Ok});
 			continue;
 		}
 
-		TUniquePtr<FHttpOperation> Operation = CacheStore.WaitForHttpOperation(EOperationCategory::Get);
-		FHttpOperation& LocalOperation = *Operation;
-		LocalOperation.SetUri(WriteToAnsiString<256>(CacheStore.EffectiveDomain, ANSITEXTVIEW("/api/v1/compressed-blobs/"), CacheStore.Namespace, '/', Value.GetRawHash()));
-		LocalOperation.SetMethod(EHttpMethod::Get);
-		LocalOperation.AddAcceptType(EHttpMediaType::Any);
-		LocalOperation.SetExpectedErrorCodes({404});
-		LocalOperation.SendAsync(Owner, [Self = TRefCountPtr(this), Operation = MoveTemp(Operation), SharedOnComplete, Id = Value.GetId(), RawHash = Value.GetRawHash(), RawSize = Value.GetRawSize()]
+		CacheStore.WaitForHttpOperationAsync(Owner, EOperationCategory::Get, [Self = TRefCountPtr(this), SharedOnComplete, Value](TUniquePtr<FHttpOperation>&& Operation)
 		{
-			TRACE_CPUPROFILER_EVENT_SCOPE(HttpDDC_GetPackage_GetValues_OnResponse);
+			Self->BeginGetValue(MoveTemp(Operation), Value, SharedOnComplete);
+		});
+	}
+}
 
-			FRequestTimer RequestTimer(Self->RequestStats);
-			Operation->GetStats(Self->RequestStats);
+void FHttpCacheStore::FGetRecordOp::BeginGetValue(
+	TUniquePtr<FHttpOperation>&& Operation,
+	const FValueWithId& Value,
+	const TSharedRef<FOnValueComplete>& OnComplete)
+{
+	if (UNLIKELY(!Operation))
+	{
+		(*OnComplete)({Value, EStatus::Canceled});
+		return;
+	}
 
-			bool bHit = false;
-			FCompressedBuffer CompressedBuffer;
-			if (Operation->GetStatusCode() == 200)
-			{
-				switch (Operation->GetContentType())
-				{
-				case EHttpMediaType::Any:
-				case EHttpMediaType::CompressedBinary:
-					CompressedBuffer = FCompressedBuffer::FromCompressed(Operation->GetBody());
-					bHit = true;
-					break;
-				case EHttpMediaType::Binary:
-					CompressedBuffer = FValue::Compress(Operation->GetBody()).GetData();
-					bHit = true;
-					break;
-				default:
-					break;
-				}
+	FHttpOperation& LocalOperation = *Operation;
+	LocalOperation.SetUri(WriteToAnsiString<256>(CacheStore.EffectiveDomain, ANSITEXTVIEW("/api/v1/compressed-blobs/"), CacheStore.Namespace, '/', Value.GetRawHash()));
+	LocalOperation.SetMethod(EHttpMethod::Get);
+	LocalOperation.AddAcceptType(EHttpMediaType::Any);
+	LocalOperation.SetExpectedErrorCodes({404});
+	LocalOperation.SendAsync(Owner, [Self = TRefCountPtr(this), Operation = MoveTemp(Operation), OnComplete, Value]
+	{
+		Self->EndGetValue(*Operation, Value, *OnComplete);
+	});
+}
 
-				TUniqueLock Lock(Self->RequestStats.Mutex);
-				Self->RequestStats.LogicalReadSize += CompressedBuffer.GetRawSize();
-			}
+void FHttpCacheStore::FGetRecordOp::EndGetValue(FHttpOperation& Operation, const FValueWithId& Value, const FOnValueComplete& OnComplete)
+{
+	TRACE_CPUPROFILER_EVENT_SCOPE(HttpDDC_GetPackage_GetValues_OnResponse);
 
-			RequestTimer.Stop();
+	FRequestTimer RequestTimer(RequestStats);
+	Operation.GetStats(RequestStats);
 
-			if (bHit)
-			{
-				if (CompressedBuffer.GetRawHash() == RawHash && CompressedBuffer.GetRawSize() == RawSize)
-				{
-					SharedOnComplete.Get()({FValueWithId(Id, MoveTemp(CompressedBuffer)), EStatus::Ok});
-				}
-				else
-				{
-					UE_LOG(LogDerivedDataCache, Display,
-						TEXT("%s: Cache miss with corrupted value %s with hash %s for %s from '%s'"),
-						*Self->CacheStore.Domain, *WriteToString<32>(Id), *WriteToString<48>(RawHash),
-						*WriteToString<96>(Self->Key), *Self->Name);
-					SharedOnComplete.Get()({FValueWithId(Id, RawHash, RawSize), EStatus::Error});
-				}
-			}
-			else if (Operation->GetErrorCode() == EHttpErrorCode::Canceled)
-			{
-				SharedOnComplete.Get()({FValueWithId(Id, RawHash, RawSize), EStatus::Canceled});
-			}
-			else
-			{
-				UE_LOG(LogDerivedDataCache, Verbose,
-					TEXT("%s: Cache miss with missing value %s with hash %s for %s from '%s'"),
-					*Self->CacheStore.Domain, *WriteToString<32>(Id), *WriteToString<48>(RawHash),
-					*WriteToString<96>(Self->Key), *Self->Name);
-				SharedOnComplete.Get()({FValueWithId(Id, RawHash, RawSize), EStatus::Error});
-			}
-		});
+	bool bHit = false;
+	FCompressedBuffer CompressedBuffer;
+	if (Operation.GetStatusCode() == 200)
+	{
+		switch (Operation.GetContentType())
+		{
+		case EHttpMediaType::Any:
+		case EHttpMediaType::CompressedBinary:
+			CompressedBuffer = FCompressedBuffer::FromCompressed(Operation.GetBody());
+			bHit = true;
+			break;
+		case EHttpMediaType::Binary:
+			CompressedBuffer = FValue::Compress(Operation.GetBody()).GetData();
+			bHit = true;
+			break;
+		default:
+			break;
+		}
+
+		TUniqueLock Lock(RequestStats.Mutex);
+		RequestStats.LogicalReadSize += CompressedBuffer.GetRawSize();
+	}
+
+	RequestTimer.Stop();
+
+	if (bHit)
+	{
+		if (CompressedBuffer.GetRawHash() == Value.GetRawHash() && CompressedBuffer.GetRawSize() == Value.GetRawSize())
+		{
+			OnComplete({FValueWithId(Value.GetId(), MoveTemp(CompressedBuffer)), EStatus::Ok});
+		}
+		else
+		{
+			UE_LOG(LogDerivedDataCache, Display,
+				TEXT("%s: Cache miss with corrupted value %s with hash %s for %s from '%s'"),
+				*CacheStore.Domain, *WriteToString<32>(Value.GetId()), *WriteToString<48>(Value.GetRawHash()),
+				*WriteToString<96>(Key), *Name);
+			OnComplete({Value, EStatus::Error});
+		}
+	}
+	else if (Operation.GetErrorCode() == EHttpErrorCode::Canceled)
+	{
+		OnComplete({Value, EStatus::Canceled});
+	}
+	else
+	{
+		UE_LOG(LogDerivedDataCache, Verbose,
+			TEXT("%s: Cache miss with missing value %s with hash %s for %s from '%s'"),
+			*CacheStore.Domain, *WriteToString<32>(Value.GetId()), *WriteToString<48>(Value.GetRawHash()),
+			*WriteToString<96>(Key), *Name);
+		OnComplete({Value, EStatus::Error});
 	}
 }
 
@@ -1561,18 +1601,34 @@ void FHttpCacheStore::FGetRecordOp::GetValuesExist(TConstArrayView<FValueWithId>
 		return;
 	}
 
+	FRequestTimer RequestTimer(RequestStats);
+	RequestTimer.Stop();
+
+	FRequestBarrier Barrier(Owner);
+	CacheStore.WaitForHttpOperationAsync(Owner, EOperationCategory::Get, [Self = TRefCountPtr(this), Values = MoveTemp(QueryValues), OnComplete = MoveTemp(OnComplete)](TUniquePtr<FHttpOperation>&& Operation) mutable
+	{
+		Self->BeginGetValuesExist(MoveTemp(Operation), MoveTemp(Values), MoveTemp(OnComplete));
+	});
+}
+
+void FHttpCacheStore::FGetRecordOp::BeginGetValuesExist(TUniquePtr<FHttpOperation>&& Operation, TArray<FValueWithId>&& Values, FOnValueComplete&& OnComplete)
+{
+	if (UNLIKELY(!Operation))
+	{
+		EndGetValuesExist(nullptr, MoveTemp(Values), MoveTemp(OnComplete));
+		return;
+	}
+
 	FRequestTimer RequestTimer(RequestStats);
 
 	TAnsiStringBuilder<256> Uri;
 	Uri << CacheStore.EffectiveDomain << ANSITEXTVIEW("/api/v1/compressed-blobs/") << CacheStore.Namespace << ANSITEXTVIEW("/exists?");
-	for (const FValueWithId& Value : QueryValues)
+	for (const FValueWithId& Value : Values)
 	{
 		Uri << ANSITEXTVIEW("id=") << Value.GetRawHash() << '&';
 	}
 	Uri.RemoveSuffix(1);
 
-	FRequestBarrier Barrier(Owner);
-	TUniquePtr<FHttpOperation> Operation = CacheStore.WaitForHttpOperation(EOperationCategory::Get);
 	FHttpOperation& LocalOperation = *Operation;
 	LocalOperation.SetUri(Uri);
 	LocalOperation.SetMethod(EHttpMethod::Post);
@@ -1580,65 +1636,74 @@ void FHttpCacheStore::FGetRecordOp::GetValuesExist(TConstArrayView<FValueWithId>
 	LocalOperation.AddAcceptType(EHttpMediaType::Json);
 
 	RequestTimer.Stop();
-	LocalOperation.SendAsync(Owner, [Self = TRefCountPtr(this), Operation = MoveTemp(Operation), Values = MoveTemp(QueryValues), OnComplete = MoveTemp(OnComplete)]() mutable
+	LocalOperation.SendAsync(Owner, [Self = TRefCountPtr(this), Operation = MoveTemp(Operation), Values = MoveTemp(Values), OnComplete = MoveTemp(OnComplete)]() mutable
 	{
-		TRACE_CPUPROFILER_EVENT_SCOPE(HttpDDC_DataProbablyExistsBatch_OnHttpRequestComplete);
+		Self->EndGetValuesExist(Operation.Get(), MoveTemp(Values), MoveTemp(OnComplete));
+	});
+}
 
-		FRequestTimer RequestTimer(Self->RequestStats);
-		Operation->GetStats(Self->RequestStats);
+void FHttpCacheStore::FGetRecordOp::EndGetValuesExist(FHttpOperation* Operation, TArray<FValueWithId>&& Values, FOnValueComplete&& OnComplete)
+{
+	TRACE_CPUPROFILER_EVENT_SCOPE(HttpDDC_DataProbablyExistsBatch_OnHttpRequestComplete);
 
-		const TCHAR* DefaultMessage = TEXT("Cache exists miss for");
-		EStatus DefaultStatus = EStatus::Error;
+	FRequestTimer RequestTimer(RequestStats);
 
-		if (Operation->GetErrorCode() == EHttpErrorCode::Canceled)
-		{
-			DefaultMessage = TEXT("Cache exists miss with canceled request for");
-			DefaultStatus = EStatus::Canceled;
-		}
-		else if (const int32 StatusCode = Operation->GetStatusCode(); StatusCode < 200 || StatusCode > 204)
-		{
-			DefaultMessage = TEXT("Cache exists miss with failed response for");
-		}
-		else if (TSharedPtr<FJsonObject> ResponseObject = Operation->GetBodyAsJson(); !ResponseObject)
-		{
-			DefaultMessage = TEXT("Cache exists miss with invalid response for");
-		}
-		else if (TArray<FString> NeedsArrayStrings; ResponseObject->TryGetStringArrayField(TEXT("needs"), NeedsArrayStrings))
-		{
-			DefaultMessage = TEXT("Cache exists hit for");
-			DefaultStatus = EStatus::Ok;
+	if (Operation)
+	{
+		Operation->GetStats(RequestStats);
+	}
+
+	const TCHAR* DefaultMessage = TEXT("Cache exists miss for");
+	EStatus DefaultStatus = EStatus::Error;
+
+	if (!Operation || Operation->GetErrorCode() == EHttpErrorCode::Canceled)
+	{
+		DefaultMessage = TEXT("Cache exists miss with canceled request for");
+		DefaultStatus = EStatus::Canceled;
+	}
+	else if (const int32 StatusCode = Operation->GetStatusCode(); StatusCode < 200 || StatusCode > 204)
+	{
+		DefaultMessage = TEXT("Cache exists miss with failed response for");
+	}
+	else if (TSharedPtr<FJsonObject> ResponseObject = Operation->GetBodyAsJson(); !ResponseObject)
+	{
+		DefaultMessage = TEXT("Cache exists miss with invalid response for");
+	}
+	else if (TArray<FString> NeedsArrayStrings; ResponseObject->TryGetStringArrayField(TEXT("needs"), NeedsArrayStrings))
+	{
+		DefaultMessage = TEXT("Cache exists hit for");
+		DefaultStatus = EStatus::Ok;
 
-			for (const FString& NeedsString : NeedsArrayStrings)
+		for (const FString& NeedsString : NeedsArrayStrings)
+		{
+			const FIoHash NeedHash(NeedsString);
+			for (auto It = Values.CreateIterator(); It; ++It)
 			{
-				const FIoHash NeedHash(NeedsString);
-				for (auto It = Values.CreateIterator(); It; ++It)
+				const FValueWithId& Value = *It;
+				if (Value.GetRawHash() == NeedHash)
 				{
-					const FValueWithId& Value = *It;
-					if (Value.GetRawHash() == NeedHash)
-					{
-						UE_LOG(LogDerivedDataCache, Verbose,
-							TEXT("%s: Cache exists miss with missing value %s with hash %s for %s from '%s'"),
-							*Self->CacheStore.Domain, *WriteToString<32>(Value.GetId()),
-							*WriteToString<48>(Value.GetRawHash()), *WriteToString<96>(Self->Key), *Self->Name);
-						OnComplete({Value, EStatus::Error});
-						It.RemoveCurrentSwap();
-						break;
-					}
+					UE_LOG(LogDerivedDataCache, Verbose,
+						TEXT("%s: Cache exists miss with missing value %s with hash %s for %s from '%s'"),
+						*CacheStore.Domain, *WriteToString<32>(Value.GetId()),
+						*WriteToString<48>(Value.GetRawHash()), *WriteToString<96>(Key), *Name);
+					OnComplete({Value, EStatus::Error});
+					It.RemoveCurrentSwap();
+					break;
 				}
 			}
 		}
+	}
 
-		RequestTimer.Stop();
+	RequestTimer.Stop();
 
-		for (const FValueWithId& Value : Values)
-		{
-			UE_LOG(LogDerivedDataCache, Verbose,
-				TEXT("%s: %s value %s with hash %s for %s from '%s'"),
-				*Self->CacheStore.Domain, DefaultMessage, *WriteToString<32>(Value.GetId()),
-				*WriteToString<48>(Value.GetRawHash()), *WriteToString<96>(Self->Key), *Self->Name);
-			OnComplete({Value, DefaultStatus});
-		}
-	});
+	for (const FValueWithId& Value : Values)
+	{
+		UE_LOG(LogDerivedDataCache, Verbose,
+			TEXT("%s: %s value %s with hash %s for %s from '%s'"),
+			*CacheStore.Domain, DefaultMessage, *WriteToString<32>(Value.GetId()),
+			*WriteToString<48>(Value.GetRawHash()), *WriteToString<96>(Key), *Name);
+		OnComplete({Value, DefaultStatus});
+	}
 }
 
 void FHttpCacheStore::FGetRecordOp::RecordStats(EStatus Status)
@@ -1679,7 +1744,8 @@ public:
 private:
 	FGetValueOp(FHttpCacheStore& CacheStore, IRequestOwner& Owner, const FSharedString& Name);
 
-	void EndGetRef(TUniquePtr<FHttpOperation> Operation);
+	void BeginGetRef(TUniquePtr<FHttpOperation>&& Operation);
+	void EndGetRef(FHttpOperation& Operation);
 	void EndGet(FResponse&& Response);
 
 	FHttpCacheStore& CacheStore;
@@ -1707,12 +1773,30 @@ void FHttpCacheStore::FGetValueOp::Get(const FCacheKey& InKey, ECachePolicy InPo
 	Policy = InPolicy;
 	OnComplete = MoveTemp(InOnComplete);
 
+	RequestTimer.Stop();
+	CacheStore.WaitForHttpOperationAsync(Owner, EOperationCategory::Get, [Self = TRefCountPtr(this)](TUniquePtr<FHttpOperation>&& Operation)
+	{
+		Self->BeginGetRef(MoveTemp(Operation));
+	});
+}
+
+void FHttpCacheStore::FGetValueOp::BeginGetRef(TUniquePtr<FHttpOperation>&& Operation)
+{
+	if (UNLIKELY(!Operation))
+	{
+		UE_LOG(LogDerivedDataCache, Verbose, TEXT("%s: Cache miss with failed with canceled request for %s from '%s'"),
+			*CacheStore.Domain, *WriteToString<96>(Key), *Name);
+		EndGet({Name, Key, {}, EStatus::Canceled});
+		return;
+	}
+
+	FRequestTimer RequestTimer(RequestStats);
+
 	const bool bSkipData = EnumHasAnyFlags(Policy, ECachePolicy::SkipData);
 
 	TAnsiStringBuilder<64> Bucket;
 	Algo::Transform(Key.Bucket.ToString(), AppendChars(Bucket), FCharAnsi::ToLower);
 
-	TUniquePtr<FHttpOperation> Operation = CacheStore.WaitForHttpOperation(EOperationCategory::Get);
 	FHttpOperation& LocalOperation = *Operation;
 	LocalOperation.SetUri(WriteToAnsiString<256>(CacheStore.EffectiveDomain, ANSITEXTVIEW("/api/v1/refs/"), CacheStore.Namespace, '/', Bucket, '/', Key.Hash));
 	LocalOperation.SetMethod(EHttpMethod::Get);
@@ -1729,18 +1813,19 @@ void FHttpCacheStore::FGetValueOp::Get(const FCacheKey& InKey, ECachePolicy InPo
 	RequestTimer.Stop();
 	LocalOperation.SendAsync(Owner, [Self = TRefCountPtr(this), Operation = MoveTemp(Operation)]() mutable
 	{
-		Operation->GetStats(Self->RequestStats);
-		Self->EndGetRef(MoveTemp(Operation));
+		Self->EndGetRef(*Operation);
 	});
 }
 
-void FHttpCacheStore::FGetValueOp::EndGetRef(TUniquePtr<FHttpOperation> Operation)
+void FHttpCacheStore::FGetValueOp::EndGetRef(FHttpOperation& Operation)
 {
 	TRACE_CPUPROFILER_EVENT_SCOPE(HttpDDC_GetValue_EndGetRef);
 
+	Operation.GetStats(RequestStats);
+
 	const bool bSkipData = EnumHasAnyFlags(Policy, ECachePolicy::SkipData);
 
-	const int32 StatusCode = Operation->GetStatusCode();
+	const int32 StatusCode = Operation.GetStatusCode();
 	if (StatusCode < 200 || StatusCode > 204)
 	{
 		UE_LOG(LogDerivedDataCache, Verbose, TEXT("%s: Cache miss with failed HTTP request for %s from '%s'"),
@@ -1748,7 +1833,7 @@ void FHttpCacheStore::FGetValueOp::EndGetRef(TUniquePtr<FHttpOperation> Operatio
 		return EndGet({Name, Key, {}, EStatus::Error});
 	}
 
-	FSharedBuffer Body = Operation->GetBody();
+	FSharedBuffer Body = Operation.GetBody();
 
 	if (bSkipData)
 	{
@@ -1778,7 +1863,7 @@ void FHttpCacheStore::FGetValueOp::EndGetRef(TUniquePtr<FHttpOperation> Operatio
 		if (!CompressedBuffer)
 		{
 			FRequestTimer RequestTimer(RequestStats);
-			if (FAnsiStringView ReceivedHashStr = Operation->GetHeader("X-Jupiter-InlinePayloadHash"); !ReceivedHashStr.IsEmpty())
+			if (FAnsiStringView ReceivedHashStr = Operation.GetHeader("X-Jupiter-InlinePayloadHash"); !ReceivedHashStr.IsEmpty())
 			{
 				FIoHash ReceivedHash(ReceivedHashStr);
 				FIoHash ComputedHash = FIoHash::HashBuffer(Body.GetView());
@@ -1828,7 +1913,8 @@ public:
 private:
 	FExistsBatchOp(FHttpCacheStore& CacheStore, IRequestOwner& Owner);
 
-	void EndExists(TUniquePtr<FHttpOperation> Operation);
+	void BeginExists(TUniquePtr<FHttpOperation>&& Operation, FCbFieldIterator&& Body);
+	void EndExists(FHttpOperation& Operation);
 
 	void EndRequest(const FCacheGetValueRequest& Request, const FValue& Value, EStatus Status);
 
@@ -1909,7 +1995,29 @@ void FHttpCacheStore::FExistsBatchOp::Exists(TConstArrayView<FCacheGetValueReque
 	BodyWriter.EndObject();
 	FCbFieldIterator Body = BodyWriter.Save();
 
-	TUniquePtr<FHttpOperation> Operation = CacheStore.WaitForHttpOperation(EOperationCategory::Get);
+	RequestTimer.Stop();
+	CacheStore.WaitForHttpOperationAsync(Owner, EOperationCategory::Get, [Self = TRefCountPtr(this), Body = MoveTemp(Body)](TUniquePtr<FHttpOperation>&& Operation) mutable
+	{
+		Self->BeginExists(MoveTemp(Operation), MoveTemp(Body));
+	});
+}
+
+void FHttpCacheStore::FExistsBatchOp::BeginExists(TUniquePtr<FHttpOperation>&& Operation, FCbFieldIterator&& Body)
+{
+	if (UNLIKELY(!Operation))
+	{
+		for (const FCacheGetValueRequest& Request : Requests)
+		{
+			UE_LOG(LogDerivedDataCache, Verbose, TEXT("%s: Cache miss with canceled request for %s from '%s'"),
+				*CacheStore.Domain, *WriteToString<96>(Request.Key), *Request.Name);
+			RequestStats.Bucket = Request.Key.Bucket;
+			EndRequest(Request, {}, EStatus::Canceled);
+		}
+		return;
+	}
+
+	FRequestTimer RequestTimer(RequestStats);
+
 	FHttpOperation& LocalOperation = *Operation;
 	LocalOperation.SetUri(WriteToAnsiString<256>(CacheStore.EffectiveDomain, ANSITEXTVIEW("/api/v1/refs/"), CacheStore.Namespace));
 	LocalOperation.SetMethod(EHttpMethod::Post);
@@ -1920,15 +2028,18 @@ void FHttpCacheStore::FExistsBatchOp::Exists(TConstArrayView<FCacheGetValueReque
 	RequestTimer.Stop();
 	LocalOperation.SendAsync(Owner, [Self = TRefCountPtr(this), Operation = MoveTemp(Operation)]() mutable
 	{
-		Operation->GetStats(Self->RequestStats);
-		Self->EndExists(MoveTemp(Operation));
+		Self->EndExists(*Operation);
 	});
 }
 
-void FHttpCacheStore::FExistsBatchOp::EndExists(TUniquePtr<FHttpOperation> Operation)
+void FHttpCacheStore::FExistsBatchOp::EndExists(FHttpOperation& Operation)
 {
 	TRACE_CPUPROFILER_EVENT_SCOPE(HttpDDC_ExistsBatch_EndExists);
 
+	FRequestTimer RequestTimer(RequestStats);
+
+	Operation.GetStats(RequestStats);
+
 	// Divide the stats evenly among the requests.
 	RequestStats.PhysicalReadSize /= Requests.Num();
 	RequestStats.PhysicalWriteSize /= Requests.Num();
@@ -1938,9 +2049,10 @@ void FHttpCacheStore::FExistsBatchOp::EndExists(TUniquePtr<FHttpOperation> Opera
 	RequestStats.Type = ERequestType::Value;
 	RequestStats.Op = ERequestOp::Get;
 
-	const int32 OverallStatusCode = Operation->GetStatusCode();
+	const int32 OverallStatusCode = Operation.GetStatusCode();
 	if (OverallStatusCode < 200 || OverallStatusCode > 204)
 	{
+		RequestTimer.Stop();
 		for (const FCacheGetValueRequest& Request : Requests)
 		{
 			UE_LOG(LogDerivedDataCache, Verbose, TEXT("%s: Cache miss with failed HTTP request for %s from '%s'"),
@@ -1951,9 +2063,10 @@ void FHttpCacheStore::FExistsBatchOp::EndExists(TUniquePtr<FHttpOperation> Opera
 		return;
 	}
 
-	FMemoryView ResponseView = Operation->GetBody();
+	FMemoryView ResponseView = Operation.GetBody();
 	if (ValidateCompactBinary(ResponseView, ECbValidateMode::Default) != ECbValidateError::None)
 	{
+		RequestTimer.Stop();
 		for (const FCacheGetValueRequest& Request : Requests)
 		{
 			UE_LOG(LogDerivedDataCache, Log, TEXT("%s: Cache miss with corrupt response for %s from '%s'."),
@@ -1964,6 +2077,8 @@ void FHttpCacheStore::FExistsBatchOp::EndExists(TUniquePtr<FHttpOperation> Opera
 		return;
 	}
 
+	RequestTimer.Stop();
+
 	const FCbObjectView ResponseObject(ResponseView.GetData());
 	const FCbArrayView Results = ResponseObject[ANSITEXTVIEW("results")].AsArrayView();
 
@@ -2382,34 +2497,6 @@ void FHttpCacheStore::SetAccessTokenAndUnlock(FScopeLock& Lock, FStringView Toke
 	}
 }
 
-TUniquePtr<FHttpCacheStore::FHttpOperation> FHttpCacheStore::WaitForHttpOperation(EOperationCategory Category)
-{
-	if (Access && RefreshAccessTokenTime > 0.0 && RefreshAccessTokenTime < FPlatformTime::Seconds())
-	{
-		AcquireAccessToken();
-	}
-
-	THttpUniquePtr<IHttpRequest> Request;
-
-	{
-		FHttpRequestParams Params;
-		FRequestOwner BlockingOwner(EPriority::Blocking);
-		FHttpCacheStoreRequestQueue& RequestQueue = (Category == EOperationCategory::Get) ? GetRequestQueue : PutRequestQueue;
-		RequestQueue.CreateRequestAsync(BlockingOwner, Params, [&Request](THttpUniquePtr<IHttpRequest>&& AsyncRequest)
-		{
-			Request = MoveTemp(AsyncRequest);
-		});
-		BlockingOwner.Wait();
-	}
-
-	if (Access)
-	{
-		Request->AddHeader(ANSITEXTVIEW("Authorization"), WriteToAnsiString<1024>(*Access));
-	}
-
-	return MakeUnique<FHttpOperation>(MoveTemp(Request));
-}
-
 void FHttpCacheStore::WaitForHttpOperationAsync(IRequestOwner& Owner, EOperationCategory Category, TUniqueFunction<void (TUniquePtr<FHttpOperation>&&)>&& OnOperation)
 {
 	WaitForHttpRequestAsync(Owner, Category, [this, OnOperation = MoveTemp(OnOperation)](THttpUniquePtr<IHttpRequest>&& Request)
