DDC: Fixed accumulation and reporting of cache store latency

diff --git a/Engine/Source/Developer/DerivedDataCache/Private/DerivedDataCacheStats.cpp b/Engine/Source/Developer/DerivedDataCache/Private/DerivedDataCacheStats.cpp
--- a/Engine/Source/Developer/DerivedDataCache/Private/DerivedDataCacheStats.cpp
+++ b/Engine/Source/Developer/DerivedDataCache/Private/DerivedDataCacheStats.cpp
@@ -185,9 +185,9 @@ void FTimeAveragedStat::Update(FMonotonicTimePoint Time)
 		ActiveTime += RangeEnd - RangeStart;
 	}
 
-	if (AccumulatedValueCount && !ActiveTime.IsZero())
+	if (AccumulatedValueCount)
 	{
-		AverageRate = AccumulatedValue / ActiveTime.ToSeconds();
+		AverageRate = !ActiveTime.IsZero() ? AccumulatedValue / ActiveTime.ToSeconds() : 0.0;
 		AverageValue = AccumulatedValue / AccumulatedValueCount;
 	}
 }
@@ -202,6 +202,7 @@ FCacheStoreStats::FCacheStoreStats(FCacheStats& InCacheStats, ECacheStoreFlags I
 	, CacheStats(InCacheStats)
 {
 	const FMonotonicTimeSpan Period = FMonotonicTimeSpan::FromSeconds(60.0);
+	AverageLatency.SetPeriod(Period);
 	AveragePhysicalReadSize.SetPeriod(Period);
 	AveragePhysicalWriteSize.SetPeriod(Period);
 }
diff --git a/Engine/Source/Developer/DerivedDataCache/Private/DerivedDataCacheStoreHierarchy.cpp b/Engine/Source/Developer/DerivedDataCache/Private/DerivedDataCacheStoreHierarchy.cpp
--- a/Engine/Source/Developer/DerivedDataCache/Private/DerivedDataCacheStoreHierarchy.cpp
+++ b/Engine/Source/Developer/DerivedDataCache/Private/DerivedDataCacheStoreHierarchy.cpp
@@ -1536,7 +1536,7 @@ static void ConvertToLegacyStats(FDerivedDataCacheStatsNode& OutNode, FCacheStor
 #endif
 
 	FMonotonicTimePoint Now = FMonotonicTimePoint::Now();
-	OutNode.SpeedStats.LatencyMS = Stats.AverageLatency.GetValue(Now);
+	OutNode.SpeedStats.LatencyMS = Stats.AverageLatency.GetValue(Now) * 1000.0;
 	OutNode.SpeedStats.ReadSpeedMBs = Stats.AveragePhysicalReadSize.GetRate(Now) / 1024.0 / 1024.0;
 	OutNode.SpeedStats.WriteSpeedMBs = Stats.AveragePhysicalWriteSize.GetRate(Now) / 1024.0 / 1024.0;
 }
