Add mechanism to redirect filesystem cache stores to other cache store entries in the graph.

diff --git a/Engine/Source/Developer/DerivedDataCache/Private/DerivedDataBackends.cpp b/Engine/Source/Developer/DerivedDataCache/Private/DerivedDataBackends.cpp
--- a/Engine/Source/Developer/DerivedDataCache/Private/DerivedDataBackends.cpp
+++ b/Engine/Source/Developer/DerivedDataCache/Private/DerivedDataBackends.cpp
@@ -65,7 +65,7 @@ ILegacyCacheStore* CreateCacheStoreAsync(ILegacyCacheStore* InnerBackend, IMemor
 ILegacyCacheStore* CreateCacheStoreHierarchy(ICacheStoreOwner*& OutOwner, TFunctionRef<void (IMemoryCacheStore*&)> MemoryCacheCreator);
 ILegacyCacheStore* CreateCacheStoreThrottle(ILegacyCacheStore* InnerCache, uint32 LatencyMS, uint32 MaxBytesPerSecond);
 ILegacyCacheStore* CreateCacheStoreVerify(ILegacyCacheStore* InnerCache, bool bPutOnError);
-ILegacyCacheStore* CreateFileSystemCacheStore(const TCHAR* Name, const TCHAR* Config, ICacheStoreOwner& Owner, FString& OutPath);
+ILegacyCacheStore* CreateFileSystemCacheStore(const TCHAR* Name, const TCHAR* Config, ICacheStoreOwner& Owner, ICacheStoreGraph* Graph, FString& OutPath, bool& OutRedirected);
 ILegacyCacheStore* CreateHttpCacheStore(const TCHAR* Name, const TCHAR* Config, ICacheStoreOwner* Owner);
 void CreateMemoryCacheStore(IMemoryCacheStore*& OutCache, const TCHAR* Name, const TCHAR* Config, ICacheStoreOwner* Owner);
 IPakFileCacheStore* CreatePakFileCacheStore(const TCHAR* Name, const TCHAR* Filename, bool bWriting, bool bCompressed, ICacheStoreOwner* Owner);
@@ -77,10 +77,10 @@ ILegacyCacheStore* TryCreateCacheStoreReplay(ILegacyCacheStore* InnerCache);
  * This class is used to create a singleton that represents the derived data cache hierarchy and all of the wrappers necessary
  * ideally this would be data driven and the backends would be plugins...
  */
-class FDerivedDataBackendGraph final : public FDerivedDataBackend, ICacheStoreOwner
+class FDerivedDataBackendGraph final : public FDerivedDataBackend, ICacheStoreOwner, ICacheStoreGraph
 {
 public:
-	using FParsedNode = TPair<ILegacyCacheStore*, ECacheStoreFlags>;
+	using FParsedNode = ILegacyCacheStore*;
 	using FParsedNodeMap = TMap<FString, FParsedNode>;
 
 	/**
@@ -117,6 +117,7 @@ public:
 
 		RootCache = nullptr;
 		FParsedNodeMap ParsedNodes;
+		TGuardValue ParsedNodeMapGuard(ActiveParsedNodeMap, &ParsedNodes);
 
 		ILegacyCacheStore* RootNode = nullptr;
 
@@ -251,7 +252,7 @@ public:
 	{
 		if (const FParsedNode* ParsedNode = InParsedNodes.Find(NodeName))
 		{
-			UE_LOG(LogDerivedDataCache, Warning, TEXT("Node %s was referenced more than once in the graph. Nodes may not be shared."), *NodeName);
+			UE_LOG(LogDerivedDataCache, Display, TEXT("Node %s was referenced more than once in the graph. Nodes may not be shared."), *NodeName);
 			return false;
 		}
 
@@ -271,15 +272,15 @@ public:
 			{
 				if (NodeType == TEXT("FileSystem"))
 				{
-					return ParseDataCache(*NodeName, *Entry);
+					return ParseFileSystemCache(*NodeName, *Entry, InParsedNodes);
 				}
 				else if (NodeType == TEXT("Boot"))
 				{
-					return ParseBootCache(*NodeName, *Entry);
+					return ParseBootCache(*NodeName, *Entry, InParsedNodes);
 				}
 				else if (NodeType == TEXT("Memory"))
 				{
-					return ParseMemoryCache(*NodeName, *Entry);
+					return ParseMemoryCache(*NodeName, *Entry, InParsedNodes);
 				}
 				else if (NodeType == TEXT("Hierarchical"))
 				{
@@ -299,30 +300,33 @@ public:
 				}
 				else if (NodeType == TEXT("ReadPak"))
 				{
-					return ParsePak(*NodeName, *Entry, /*bWriting*/ false);
+					return ParsePak(*NodeName, *Entry, /*bWriting*/ false, InParsedNodes);
 				}
 				else if (NodeType == TEXT("WritePak"))
 				{
-					return ParsePak(*NodeName, *Entry, /*bWriting*/ true);
+					return ParsePak(*NodeName, *Entry, /*bWriting*/ true, InParsedNodes);
 				}
 				else if (NodeType == TEXT("S3"))
 				{
-					if (CreateS3CacheStore(*NodeName, *Entry, *this))
+					if (ILegacyCacheStore* CacheStore = CreateS3CacheStore(*NodeName, *Entry, *this))
 					{
+						InParsedNodes.Add(NodeName, CacheStore);
 						return true;
 					}
 				}
 				else if (NodeType == TEXT("Cloud") || NodeType == TEXT("Http"))
 				{
-					if (CreateHttpCacheStore(*NodeName, *Entry, this))
+					if (ILegacyCacheStore* CacheStore = CreateHttpCacheStore(*NodeName, *Entry, this))
 					{
+						InParsedNodes.Add(NodeName, CacheStore);
 						return true;
 					}
 				}
 				else if (NodeType == TEXT("Zen"))
 				{
-					if (CreateZenCacheStore(*NodeName, *Entry, this))
+					if (ILegacyCacheStore* CacheStore = CreateZenCacheStore(*NodeName, *Entry, this))
 					{
+						InParsedNodes.Add(NodeName,CacheStore);
 						return true;
 					}
 				}
@@ -340,7 +344,7 @@ public:
 	 * @param bWriting true to create pak interface for writing
 	 * @return Pak file data backend interface instance or nullptr if unsuccessful
 	 */
-	bool ParsePak(const TCHAR* NodeName, const TCHAR* Entry, const bool bWriting)
+	bool ParsePak(const TCHAR* NodeName, const TCHAR* Entry, const bool bWriting, FParsedNodeMap& InParsedNodes)
 	{
 		ILegacyCacheStore* PakNode = nullptr;
 		FString PakFilename;
@@ -366,6 +370,7 @@ public:
 			WritePakFilename = PakFilename + TEXT(".") + Temp.ToString();
 			WritePakCache = CreatePakFileCacheStore(NodeName, *WritePakFilename, /*bWriting*/ true, bCompressed, this);
 			PakNode = WritePakCache;
+			InParsedNodes.Add(NodeName, WritePakCache);
 			return true;
 		}
 		else
@@ -381,6 +386,7 @@ public:
 			ReadPakFilename = PakFilename;
 			PakNode = ReadPak;
 			ReadPakCache.Add(ReadPak);
+			InParsedNodes.Add(NodeName, ReadPak);
 			return true;
 		}
 	}
@@ -533,15 +539,21 @@ public:
 	 * @param Entry Node definition.
 	 * @return Filesystem data cache backend interface instance or nullptr if unsuccessful
 	 */
-	bool ParseDataCache(const TCHAR* NodeName, const TCHAR* Config)
+	bool ParseFileSystemCache(const TCHAR* NodeName, const TCHAR* Config, FParsedNodeMap& InParsedNodes)
 	{
+		bool bRedirected = false;
 		FString Path;
-		if (ILegacyCacheStore* Store = CreateFileSystemCacheStore(NodeName, Config, *this, Path))
+		if (ILegacyCacheStore* Store = CreateFileSystemCacheStore(NodeName, Config, *this, this, Path, bRedirected))
 		{
-			bUsingSharedDDC |= NodeName == TEXTVIEW("Shared");
-			Directories.AddUnique(Path);
+			if (!bRedirected)
+			{
+				bUsingSharedDDC |= NodeName == TEXTVIEW("Shared");
+				Directories.AddUnique(Path);
+			}
+			InParsedNodes.Add(NodeName, Store);
 			return true;
 		}
+
 		return false;
 	}
 
@@ -553,7 +565,7 @@ public:
 	 * @param OutFilename filename specified for the cache
 	 * @return Boot data cache backend interface instance or nullptr if unsuccessful
 	 */
-	bool ParseBootCache(const TCHAR* NodeName, const TCHAR* Entry)
+	bool ParseBootCache(const TCHAR* NodeName, const TCHAR* Entry, FParsedNodeMap& InParsedNodes)
 	{
 		UE_LOG(LogDerivedDataCache, Display, TEXT("Boot nodes are deprecated. Please remove the Boot node from the cache graph."));
 		if (bBootFound)
@@ -569,6 +581,11 @@ public:
 		CreateMemoryCacheStore(BootCache, TEXT("Boot"), *WriteToString<128>(TEXT("-Boot "), Entry), this);
 	#endif
 
+		if (!!BootCache)
+		{
+			InParsedNodes.Add(NodeName, BootCache);
+		}
+
 		return !!BootCache;
 	}
 
@@ -579,13 +596,14 @@ public:
 	 * @param Entry Node definition.
 	 * @return Memory data cache backend interface instance or nullptr if unsuccessful
 	 */
-	bool ParseMemoryCache(const TCHAR* NodeName, const TCHAR* Entry)
+	bool ParseMemoryCache(const TCHAR* NodeName, const TCHAR* Entry, FParsedNodeMap& InParsedNodes)
 	{
 		FString Filename;
 		FParse::Value(Entry, TEXT("Filename="), Filename);
 		IMemoryCacheStore* Cache;
 		CreateMemoryCacheStore(Cache, NodeName, TEXT(""), this);
 		check(Cache);
+		InParsedNodes.Add(NodeName, Cache);
 		if (Filename.Len())
 		{
 			UE_LOG(LogDerivedDataCache, Display, TEXT("Memory nodes that load from a file are deprecated. Please remove the filename from the cache configuration."));
@@ -898,6 +916,32 @@ private:
 		Hierarchy->LegacyResourceStats(OutStats);
 	}
 
+	ILegacyCacheStore* Create(const TCHAR* Name) final
+	{
+		if (!ActiveParsedNodeMap)
+		{
+			UE_LOG(LogDerivedDataCache, Warning, TEXT("Requesting creation of node %s during an unsupported time for node creation."), Name);
+			return nullptr;
+		}
+
+		if (const FParsedNode* ParsedNode = ActiveParsedNodeMap->Find(Name))
+		{
+			return nullptr;
+		}
+
+		if (!ParseNode(Name, GEngineIni, *GraphName, *ActiveParsedNodeMap))
+		{
+			return nullptr;
+		}
+
+		const FParsedNode* ParsedNode = ActiveParsedNodeMap->Find(Name);
+		if (ParsedNode)
+		{
+			return *ParsedNode;
+		}
+		return nullptr;
+	}
+
 	/** MountPak console command handler. */
 	void UnmountPakCommandHandler(const TArray<FString>& Args)
 	{
@@ -1046,6 +1090,7 @@ private:
 	FAutoConsoleCommand LoadReplayCommand;
 	TArray<FCacheReplayReader> Replays;
 
+	FParsedNodeMap* ActiveParsedNodeMap = nullptr;
 	FString ActiveNodeName;
 	FString ActiveNodeConfig;
 
diff --git a/Engine/Source/Developer/DerivedDataCache/Private/DerivedDataCacheStore.h b/Engine/Source/Developer/DerivedDataCache/Private/DerivedDataCacheStore.h
--- a/Engine/Source/Developer/DerivedDataCache/Private/DerivedDataCacheStore.h
+++ b/Engine/Source/Developer/DerivedDataCache/Private/DerivedDataCacheStore.h
@@ -168,6 +168,12 @@ public:
 	virtual void LegacyResourceStats(TArray<FDerivedDataCacheResourceStat>& OutStats) const = 0;
 };
 
+class ICacheStoreGraph
+{
+public:
+	virtual ILegacyCacheStore* Create(const TCHAR* Name) = 0;
+};
+
 enum class ECacheStoreStatusCode : uint8
 {
 	None = 0,
diff --git a/Engine/Source/Developer/DerivedDataCache/Private/FileSystemCacheStore.cpp b/Engine/Source/Developer/DerivedDataCache/Private/FileSystemCacheStore.cpp
--- a/Engine/Source/Developer/DerivedDataCache/Private/FileSystemCacheStore.cpp
+++ b/Engine/Source/Developer/DerivedDataCache/Private/FileSystemCacheStore.cpp
@@ -781,6 +781,12 @@ struct FFileSystemCacheStoreParams
 	/** Path to the optional access log that records every accessed file. */
 	FString AccessLogPath;
 
+	/** Optional name of a file containing redirection information so that if the file is present and valid, another cache store can be used in place of this FileSystemCacheStore. */
+	FString RedirectionFileName;
+
+	/** Optional name of a key within the RedirectionFileName containing redirection information so that if the file+key is present and valid, another cache store can be used in place of this FileSystemCacheStore. */
+	FString RedirectionKeyName;
+
 	/** Maximum total size of compressed data stored within a record package with multiple attachments. */
 	uint64 MaxRecordSizeKB = 256;
 	/** Maximum total size of compressed data stored within a value package, or a record package with one attachment. */
@@ -830,7 +836,9 @@ public:
 	static ILegacyCacheStore* TryCreate(
 		const FFileSystemCacheStoreParams& Params,
 		ICacheStoreOwner& Owner,
-		FString& OutPath);
+		ICacheStoreGraph* Graph,
+		FString& OutPath,
+		bool& bOutRedirected);
 
 private:
 	FFileSystemCacheStore(
@@ -849,6 +857,11 @@ private:
 		std::atomic<uint32>* NumLatencyTestsCompleted,
 		std::atomic<bool>* AbandonRequest);
 
+	static ILegacyCacheStore* TryRedirection(
+		const FFileSystemCacheStoreParams& Params,
+		ICacheStoreOwner& Owner,
+		ICacheStoreGraph* Graph);
+	
 	// ICacheStore Interface
 
 	void Put(
@@ -988,7 +1001,9 @@ private:
 ILegacyCacheStore* FFileSystemCacheStore::TryCreate(
 	const FFileSystemCacheStoreParams& Params,
 	ICacheStoreOwner& Owner,
-	FString& OutPath)
+	ICacheStoreGraph* Graph,
+	FString& OutPath,
+	bool& bOutRedirected)
 {
 	// If we find a platform that has more stringent limits, this needs to be rethought.
 	checkf(GMaxCacheRootLen + GMaxCacheKeyLen <= FPlatformMisc::GetMaxPathLength(),
@@ -1025,6 +1040,12 @@ ILegacyCacheStore* FFileSystemCacheStore::TryCreate(
 		bool bShared = Params.CacheName == TEXTVIEW("Shared");
 		if (!bShared || IFileManager::Get().DirectoryExists(*Params.CachePath))
 		{
+			if (ILegacyCacheStore* RedirectedStore = TryRedirection(Params, Owner, Graph))
+			{
+				bOutRedirected = true;
+				return RedirectedStore;
+			}
+
 			FDerivedDataCacheSpeedStats LocalSpeedStats;
 			LocalSpeedStats.ReadSpeedMBs = 999;
 			LocalSpeedStats.WriteSpeedMBs = 999;
@@ -1483,6 +1504,47 @@ bool FFileSystemCacheStore::RunSpeedTest(
 	return bWriteTestPassed || bReadTestPassed;
 }
 
+ILegacyCacheStore* FFileSystemCacheStore::TryRedirection(
+	const FFileSystemCacheStoreParams& Params,
+	ICacheStoreOwner& Owner,
+	ICacheStoreGraph* Graph)
+{
+	if (Graph && !Params.RedirectionFileName.IsEmpty())
+	{
+		FString RedirectionFileName = FPaths::Combine(Params.CachePath, Params.RedirectionFileName);
+		FConfigFile RedirectionFile;
+		if (FConfigCacheIni::LoadLocalIniFile(RedirectionFile, *RedirectionFileName, false))
+		{
+			FString RedirectionKeyName = Params.RedirectionKeyName.IsEmpty() ? TEXT("Default") : Params.RedirectionKeyName;
+			FString RedirectionData;
+			if (RedirectionFile.GetString(TEXT("Redirect"), *RedirectionKeyName, RedirectionData))
+			{
+				RedirectionData.TrimStartInline();
+				RedirectionData.RemoveFromStart(TEXT("("));
+				RedirectionData.RemoveFromEnd(TEXT(")"));
+
+				FString EnvName;
+				FString EnvValue;
+				if (FParse::Value(*RedirectionData, TEXT("SetEnvName="), EnvName) && FParse::Value(*RedirectionData, TEXT("SetEnvValue="), EnvValue))
+				{
+					FPlatformMisc::SetEnvironmentVar(*EnvName, *EnvValue);
+				}
+
+				FString TargetName;
+				if (FParse::Value(*RedirectionData, TEXT("Target="), TargetName))
+				{
+					if (ILegacyCacheStore* RedirectedStore = Graph->Create(*TargetName))
+					{
+						return RedirectedStore;
+					}
+				}
+			}
+		}
+	}
+
+	return nullptr;
+}
+
 void FFileSystemCacheStore::LegacyStats(FDerivedDataCacheStatsNode& OutNode)
 {
 	checkNoEntry();
@@ -2985,17 +3047,22 @@ void FFileSystemCacheStoreParams::Parse(const TCHAR* Name, const TCHAR* Config)
 	FParse::Value(Config, TEXT("FoldersToClean="), MaxDirectoryScanCount);
 
 	GConfig->GetDouble(TEXT("DDCCleanup"), TEXT("TimeToWaitAfterInit"), MaintenanceDelayInSeconds, GEngineIni);
+
+	FParse::Value(Config, TEXT("RedirectionFileName="), RedirectionFileName);
+	FParse::Value(Config, TEXT("RedirectionKeyName="), RedirectionKeyName);
 }
 
 ILegacyCacheStore* CreateFileSystemCacheStore(
 	const TCHAR* Name,
 	const TCHAR* Config,
 	ICacheStoreOwner& Owner,
-	FString& OutPath)
+	ICacheStoreGraph* Graph,
+	FString& OutPath,
+	bool& bOutRedirected)
 {
 	FFileSystemCacheStoreParams Params;
 	Params.Parse(Name, Config);
-	return FFileSystemCacheStore::TryCreate(Params, Owner, OutPath);
+	return FFileSystemCacheStore::TryCreate(Params, Owner, Graph, OutPath, bOutRedirected);
 }
 
 } // UE::DerivedData
