DDC: Temporarily restored blocking during get requests in HttpCacheStore

#rnx

diff --git a/Engine/Source/Developer/DerivedDataCache/Private/HttpCacheStore.cpp b/Engine/Source/Developer/DerivedDataCache/Private/HttpCacheStore.cpp
--- a/Engine/Source/Developer/DerivedDataCache/Private/HttpCacheStore.cpp
+++ b/Engine/Source/Developer/DerivedDataCache/Private/HttpCacheStore.cpp
@@ -458,6 +458,8 @@ private:
 
 	class FHttpOperation;
 
+	TUniquePtr<FHttpOperation> WaitForHttpOperation(EOperationCategory Category);
+
 	/** Invokes the callback when an operation is available, or with null if canceled. */
 	void WaitForHttpOperationAsync(IRequestOwner& Owner, EOperationCategory Category, TUniqueFunction<void (TUniquePtr<FHttpOperation>&&)>&& OnOperation);
 
@@ -1255,8 +1257,9 @@ void FHttpCacheStore::FGetRecordOp::GetRecordOnly(const FCacheKey& InKey, const
 	OnRecordComplete = MoveTemp(InOnComplete);
 	RequestStats.Bucket = Key.Bucket;
 
+	TUniquePtr<FHttpOperation> Operation = CacheStore.WaitForHttpOperation(EOperationCategory::Get);
+	TRefCountPtr Self(this);
 	RequestTimer.Stop();
-	CacheStore.WaitForHttpOperationAsync(Owner, EOperationCategory::Get, [Self = TRefCountPtr(this)](TUniquePtr<FHttpOperation>&& Operation)
 	{
 		if (UNLIKELY(!Operation))
 		{
@@ -1278,7 +1281,7 @@ void FHttpCacheStore::FGetRecordOp::GetRecordOnly(const FCacheKey& InKey, const
 			Operation->GetStats(Self->RequestStats);
 			Self->EndGetRef(MoveTemp(Operation));
 		});
-	});
+	}
 }
 
 void FHttpCacheStore::FGetRecordOp::EndGetRef(TUniquePtr<FHttpOperation> Operation)
@@ -1482,6 +1485,7 @@ void FHttpCacheStore::FGetRecordOp::GetValues(TConstArrayView<FValueWithId> Valu
 	// TODO: Jupiter does not currently provide a batched GET. Once it does, fetch every blob in one request.
 
 	FRequestTimer RequestTimer(RequestStats);
+	RequestTimer.Stop();
 
 	FRequestBarrier Barrier(Owner);
 	TSharedRef<FOnValueComplete> SharedOnComplete = MakeShared<FOnValueComplete>(MoveTemp(OnComplete));
@@ -1493,10 +1497,11 @@ void FHttpCacheStore::FGetRecordOp::GetValues(TConstArrayView<FValueWithId> Valu
 			continue;
 		}
 
-		CacheStore.WaitForHttpOperationAsync(Owner, EOperationCategory::Get, [Self = TRefCountPtr(this), SharedOnComplete, Value](TUniquePtr<FHttpOperation>&& Operation)
+		TUniquePtr<FHttpOperation> Operation = CacheStore.WaitForHttpOperation(EOperationCategory::Get);
+		TRefCountPtr Self(this);
 		{
 			Self->BeginGetValue(MoveTemp(Operation), Value, SharedOnComplete);
-		});
+		}
 	}
 }
 
@@ -1602,13 +1607,14 @@ void FHttpCacheStore::FGetRecordOp::GetValuesExist(TConstArrayView<FValueWithId>
 	}
 
 	FRequestTimer RequestTimer(RequestStats);
-	RequestTimer.Stop();
 
 	FRequestBarrier Barrier(Owner);
-	CacheStore.WaitForHttpOperationAsync(Owner, EOperationCategory::Get, [Self = TRefCountPtr(this), Values = MoveTemp(QueryValues), OnComplete = MoveTemp(OnComplete)](TUniquePtr<FHttpOperation>&& Operation) mutable
+	TUniquePtr<FHttpOperation> Operation = CacheStore.WaitForHttpOperation(EOperationCategory::Get);
+	TRefCountPtr Self(this);
+	RequestTimer.Stop();
 	{
-		Self->BeginGetValuesExist(MoveTemp(Operation), MoveTemp(Values), MoveTemp(OnComplete));
-	});
+		Self->BeginGetValuesExist(MoveTemp(Operation), MoveTemp(QueryValues), MoveTemp(OnComplete));
+	}
 }
 
 void FHttpCacheStore::FGetRecordOp::BeginGetValuesExist(TUniquePtr<FHttpOperation>&& Operation, TArray<FValueWithId>&& Values, FOnValueComplete&& OnComplete)
@@ -1773,11 +1779,12 @@ void FHttpCacheStore::FGetValueOp::Get(const FCacheKey& InKey, ECachePolicy InPo
 	Policy = InPolicy;
 	OnComplete = MoveTemp(InOnComplete);
 
+	TUniquePtr<FHttpOperation> Operation = CacheStore.WaitForHttpOperation(EOperationCategory::Get);
+	TRefCountPtr Self(this);
 	RequestTimer.Stop();
-	CacheStore.WaitForHttpOperationAsync(Owner, EOperationCategory::Get, [Self = TRefCountPtr(this)](TUniquePtr<FHttpOperation>&& Operation)
 	{
 		Self->BeginGetRef(MoveTemp(Operation));
-	});
+	}
 }
 
 void FHttpCacheStore::FGetValueOp::BeginGetRef(TUniquePtr<FHttpOperation>&& Operation)
@@ -1995,11 +2002,12 @@ void FHttpCacheStore::FExistsBatchOp::Exists(TConstArrayView<FCacheGetValueReque
 	BodyWriter.EndObject();
 	FCbFieldIterator Body = BodyWriter.Save();
 
+	TUniquePtr<FHttpOperation> Operation = CacheStore.WaitForHttpOperation(EOperationCategory::Get);
+	TRefCountPtr Self(this);
 	RequestTimer.Stop();
-	CacheStore.WaitForHttpOperationAsync(Owner, EOperationCategory::Get, [Self = TRefCountPtr(this), Body = MoveTemp(Body)](TUniquePtr<FHttpOperation>&& Operation) mutable
 	{
 		Self->BeginExists(MoveTemp(Operation), MoveTemp(Body));
-	});
+	}
 }
 
 void FHttpCacheStore::FExistsBatchOp::BeginExists(TUniquePtr<FHttpOperation>&& Operation, FCbFieldIterator&& Body)
@@ -2497,6 +2505,34 @@ void FHttpCacheStore::SetAccessTokenAndUnlock(FScopeLock& Lock, FStringView Toke
 	}
 }
 
+TUniquePtr<FHttpCacheStore::FHttpOperation> FHttpCacheStore::WaitForHttpOperation(EOperationCategory Category)
+{
+	if (Access && RefreshAccessTokenTime > 0.0 && RefreshAccessTokenTime < FPlatformTime::Seconds())
+	{
+		AcquireAccessToken();
+	}
+
+	THttpUniquePtr<IHttpRequest> Request;
+
+	{
+		FHttpRequestParams Params;
+		FRequestOwner BlockingOwner(EPriority::Blocking);
+		FHttpCacheStoreRequestQueue& RequestQueue = (Category == EOperationCategory::Get) ? GetRequestQueue : PutRequestQueue;
+		RequestQueue.CreateRequestAsync(BlockingOwner, Params, [&Request](THttpUniquePtr<IHttpRequest>&& AsyncRequest)
+		{
+			Request = MoveTemp(AsyncRequest);
+		});
+		BlockingOwner.Wait();
+	}
+
+	if (Access)
+	{
+		Request->AddHeader(ANSITEXTVIEW("Authorization"), WriteToAnsiString<1024>(*Access));
+	}
+
+	return MakeUnique<FHttpOperation>(MoveTemp(Request));
+}
+
 void FHttpCacheStore::WaitForHttpOperationAsync(IRequestOwner& Owner, EOperationCategory Category, TUniqueFunction<void (TUniquePtr<FHttpOperation>&&)>&& OnOperation)
 {
 	WaitForHttpRequestAsync(Owner, Category, [this, OnOperation = MoveTemp(OnOperation)](THttpUniquePtr<IHttpRequest>&& Request)
