Fix memory leak caused by reference counted objects capturing references to themselves in function captures held on members of itself.  This manifested as memory never being freed when performing some operations with HttpCacheStore and ZenCacheStore.

diff --git a/Engine/Source/Developer/DerivedDataCache/Private/HttpCacheStore.cpp b/Engine/Source/Developer/DerivedDataCache/Private/HttpCacheStore.cpp
--- a/Engine/Source/Developer/DerivedDataCache/Private/HttpCacheStore.cpp
+++ b/Engine/Source/Developer/DerivedDataCache/Private/HttpCacheStore.cpp
@@ -1142,7 +1142,9 @@ void FHttpCacheStore::FPutPackageOp::EndPut(EStatus Status)
 {
 	RequestStats.EndTime = FMonotonicTimePoint::Now();
 	RequestStats.Status = Status;
-	OnPackageComplete({Status});
+	// Ensuring that the OnPackageComplete method is destroyed by the time we exit this method by moving it to a local scope variable
+	FOnPackageComplete LocalOnComplete = MoveTemp(OnPackageComplete);
+	LocalOnComplete({Status});
 	if (CacheStore.StoreStats)
 	{
 		CacheStore.StoreStats->AddRequest(RequestStats);
@@ -1300,6 +1302,7 @@ void FHttpCacheStore::FGetRecordOp::EndGetRef(TUniquePtr<FHttpOperation> Operati
 			Record = FCacheRecordBuilder(Key).Build();
 		}
 		RequestTimer.Stop();
+		// Ensuring that the OnRecordComplete method is destroyed by the time we exit this method by moving it to a local scope variable
 		FOnRecordComplete LocalOnComplete = MoveTemp(OnRecordComplete);
 		LocalOnComplete({MoveTemp(Record).Get(), Status});
 	};
@@ -1460,6 +1463,7 @@ void FHttpCacheStore::FGetRecordOp::EndGetValues(const FCacheRecordPolicy& Polic
 		Status = EStatus::Error;
 	}
 
+	// Ensuring that the OnRecordComplete method is destroyed by the time we exit this method by moving it to a local scope variable
 	FOnRecordComplete LocalOnComplete = MoveTemp(OnRecordComplete);
 	LocalOnComplete({RecordBuilder.Build(), Status});
 }
@@ -1897,7 +1901,9 @@ void FHttpCacheStore::FGetValueOp::EndGet(FResponse&& Response)
 	RequestStats.LogicalReadSize += Response.Value.GetRawSize();
 	RequestStats.EndTime = FMonotonicTimePoint::Now();
 	RequestStats.Status = Response.Status;
-	OnComplete(MoveTemp(Response));
+	// Ensuring that the OnComplete method is destroyed by the time we exit this method by moving it to a local scope variable
+	FOnComplete LocalOnComplete = MoveTemp(OnComplete);
+	LocalOnComplete(MoveTemp(Response));
 	if (CacheStore.StoreStats)
 	{
 		CacheStore.StoreStats->AddRequest(RequestStats);
@@ -2043,6 +2049,11 @@ void FHttpCacheStore::FExistsBatchOp::BeginExists(TUniquePtr<FHttpOperation>&& O
 void FHttpCacheStore::FExistsBatchOp::EndExists(FHttpOperation& Operation)
 {
 	TRACE_CPUPROFILER_EVENT_SCOPE(HttpDDC_ExistsBatch_EndExists);
+	ON_SCOPE_EXIT
+	{
+		// OnComplete may be called multiple times in the span of EndExists, but by the time this method finishes, it will never be used and can be destroyed
+		OnComplete.Reset();
+	};
 
 	FRequestTimer RequestTimer(RequestStats);
 
diff --git a/Engine/Source/Developer/DerivedDataCache/Private/ZenCacheStore.cpp b/Engine/Source/Developer/DerivedDataCache/Private/ZenCacheStore.cpp
--- a/Engine/Source/Developer/DerivedDataCache/Private/ZenCacheStore.cpp
+++ b/Engine/Source/Developer/DerivedDataCache/Private/ZenCacheStore.cpp
@@ -1357,7 +1357,9 @@ private:
 				// Launch a task for the completion function since it can execute arbitrary code.
 				Self->Owner->LaunchTask(TEXT("ZenHttpComplete"), [Self = TRefCountPtr(Self)]
 				{
-					Self->OnRpcComplete(Self->Response, Self->Package);
+					// Ensuring that the OnRpcComplete method is destroyed by the time we exit this method by moving it to a local scope variable
+					FOnRpcComplete LocalOnComplete = MoveTemp(Self->OnRpcComplete);
+					LocalOnComplete(Self->Response, Self->Package);
 				});
 			}
 		});
