[ddc] expose option for more granular prioritization of ddc query/response tasks

diff --git a/Engine/Source/Developer/DerivedDataCache/Private/DerivedDataRequestOwner.cpp b/Engine/Source/Developer/DerivedDataCache/Private/DerivedDataRequestOwner.cpp
--- a/Engine/Source/Developer/DerivedDataCache/Private/DerivedDataRequestOwner.cpp
+++ b/Engine/Source/Developer/DerivedDataCache/Private/DerivedDataRequestOwner.cpp
@@ -281,6 +281,20 @@ void IRequestOwner::LaunchTask(const TCHAR* DebugName, TUniqueFunction<void ()>&
 
 	Tasks::FTaskEvent TaskEvent(TEXT("LaunchTaskRequest"));
 	FTaskRequest* Request = new FTaskRequest;
+	ETaskPriority TaskPriority;
+	switch (GetPriority())
+	{
+	case EPriority::Highest:
+	case EPriority::Blocking:
+		TaskPriority = ETaskPriority::High;
+		break;
+	case EPriority::High:
+		TaskPriority = ETaskPriority::BackgroundHigh;
+		break;
+	default:
+		TaskPriority = ETaskPriority::BackgroundNormal;
+		break;
+	}
 	Request->Task = Launch(
 		DebugName,
 		[this, Request, TaskBody = MoveTemp(TaskBody)]
@@ -288,7 +302,7 @@ void IRequestOwner::LaunchTask(const TCHAR* DebugName, TUniqueFunction<void ()>&
 			End(Request, TaskBody);
 		},
 		TaskEvent,
-		GetPriority() <= EPriority::Normal ? ETaskPriority::BackgroundNormal : ETaskPriority::BackgroundHigh);
+		TaskPriority);
 	Begin(Request);
 	TaskEvent.Trigger();
 }
