Ensure that FCacheStoreRequestTimer increments timings and doesn't overwrite them.

diff --git a/Engine/Source/Developer/DerivedDataCache/Private/DerivedDataCacheStore.cpp b/Engine/Source/Developer/DerivedDataCache/Private/DerivedDataCacheStore.cpp
--- a/Engine/Source/Developer/DerivedDataCache/Private/DerivedDataCacheStore.cpp
+++ b/Engine/Source/Developer/DerivedDataCache/Private/DerivedDataCacheStore.cpp
@@ -26,7 +26,7 @@ void FCacheStoreRequestTimer::Stop()
 	const FMonotonicTimePoint EndTime = FMonotonicTimePoint::Now();
 	const bool bIsInGameThread = IsInGameThread();
 	TUniqueLock Lock(Stats.Mutex);
-	(bIsInGameThread ? Stats.MainThreadTime : Stats.OtherThreadTime) = EndTime - StartTime;
+	(bIsInGameThread ? Stats.MainThreadTime : Stats.OtherThreadTime) += EndTime - StartTime;
 	Stats.StartTime = FMath::Min(Stats.StartTime, StartTime);
 	Stats.EndTime = FMath::Max(Stats.EndTime, EndTime);
 	StartTime = FMonotonicTimePoint::Infinity();
