Fix a mis-named variable (name change only).

diff --git a/Engine/Source/Developer/DerivedDataCache/Private/FileSystemCacheStore.cpp b/Engine/Source/Developer/DerivedDataCache/Private/FileSystemCacheStore.cpp
--- a/Engine/Source/Developer/DerivedDataCache/Private/FileSystemCacheStore.cpp
+++ b/Engine/Source/Developer/DerivedDataCache/Private/FileSystemCacheStore.cpp
@@ -993,7 +993,7 @@ private:
 	FRWLock PerformanceReEvaluationTaskLock;
 	Tasks::TTask<std::atomic<EPerformanceReEvaluationResult>> PerformanceReEvaluationTask;
 	std::atomic<int64> LastPerformanceEvaluationTicks;
-	std::atomic<bool> bDeactivedForPerformance = false;
+	std::atomic<bool> bDeactivatedForPerformance = false;
 	bool bDeactivationDeferredClean = false;
 	float DeactivateAtMs;
 };
@@ -1101,7 +1101,7 @@ FFileSystemCacheStore::FFileSystemCacheStore(
 	bReadOnly = bReadOnly || !bWriteTestPassed;
 
 	const bool bLocalDeactivatedForPerformance = (Params.DeactivateAtMs > 0.f) && (SpeedStats.LatencyMS >= Params.DeactivateAtMs);
-	bDeactivedForPerformance.store(bLocalDeactivatedForPerformance, std::memory_order_relaxed);
+	bDeactivatedForPerformance.store(bLocalDeactivatedForPerformance, std::memory_order_relaxed);
 
 	// classify and report on these times
 	if (SpeedStats.LatencyMS < 1)
@@ -2764,7 +2764,7 @@ bool FFileSystemCacheStore::FileExists(FStringBuilderBase& Path, FRequestStats&
 
 bool FFileSystemCacheStore::IsDeactivatedForPerformance()
 {
-	if ((DeactivateAtMs <= 0.f) || !bDeactivedForPerformance.load(std::memory_order_relaxed))
+	if ((DeactivateAtMs <= 0.f) || !bDeactivatedForPerformance.load(std::memory_order_relaxed))
 	{
 		return false;
 	}
@@ -2808,7 +2808,7 @@ bool FFileSystemCacheStore::IsDeactivatedForPerformance()
 							*CachePath);
 					}
 
-					bDeactivedForPerformance.store(bLocalDeactivatedForPerformance, std::memory_order_relaxed);
+					bDeactivatedForPerformance.store(bLocalDeactivatedForPerformance, std::memory_order_relaxed);
 					UpdateStatus();
 					return bLocalDeactivatedForPerformance;
 				}
@@ -2882,7 +2882,7 @@ void FFileSystemCacheStore::UpdateStatus()
 {
 	if (StoreStats)
 	{
-		if (bDeactivedForPerformance.load(std::memory_order_relaxed))
+		if (bDeactivatedForPerformance.load(std::memory_order_relaxed))
 		{
 			StoreStats->SetStatus(ECacheStoreStatusCode::Warning, NSLOCTEXT("DerivedDataCache", "DeactivatedForPerformance", "Deactivated for performance"));
 		}
