Add performance disabling (via latency) for ZenCacheStore.

diff --git a/Engine/Source/Developer/DerivedDataCache/Private/DerivedDataCacheStats.cpp b/Engine/Source/Developer/DerivedDataCache/Private/DerivedDataCacheStats.cpp
--- a/Engine/Source/Developer/DerivedDataCache/Private/DerivedDataCacheStats.cpp
+++ b/Engine/Source/Developer/DerivedDataCache/Private/DerivedDataCacheStats.cpp
@@ -1,6 +1,7 @@
 // Copyright Epic Games, Inc. All Rights Reserved.
 
 #include "DerivedDataCacheStats.h"
+#include "DerivedDataLegacyCacheStore.h"
 
 #include "Algo/BinarySearch.h"
 
@@ -274,6 +275,8 @@ void FCacheStoreStats::AddRequest(const FCacheStoreRequestStats& Stats)
 
 	if (!Stats.Latency.IsInfinity())
 	{
+		UE_LOG(LogDerivedDataCache, VeryVerbose, TEXT("%s: Adding request latency of %.2fms from %s on %s with status %s"),
+			*Name, Stats.Latency.ToMilliseconds(), LexToString(Stats.Op), LexToString(Stats.Type), *WriteToString<32>(Stats.Status));
 		AverageLatency.Add(Stats.StartTime, Stats.EndTime, Stats.Latency.ToSeconds());
 	}
 
@@ -298,6 +301,22 @@ void FCacheStoreStats::AddRequest(const FCacheStoreRequestStats& Stats)
 #endif
 }
 
+void FCacheStoreStats::AddLatency(FMonotonicTimePoint StartTime, FMonotonicTimePoint EndTime, FMonotonicTimeSpan Latency)
+{
+	if (!Latency.IsInfinity())
+	{
+		UE_LOG(LogDerivedDataCache, VeryVerbose, TEXT("%s: Adding non-request latency of %.2fms"),
+			*Name, Latency.ToMilliseconds());
+		AverageLatency.Add(StartTime, EndTime, Latency.ToSeconds());
+	}
+}
+
+double FCacheStoreStats::GetAverageLatency()
+{
+	TUniqueLock Lock(Mutex);
+	return AverageLatency.GetValue(FMonotonicTimePoint::Now());
+}
+
 ///////////////////////////////////////////////////////////////////////////////////////////////////////////////////////
 
 FCacheBucketStats& FCacheStats::GetBucket(FCacheBucket Bucket)
diff --git a/Engine/Source/Developer/DerivedDataCache/Private/DerivedDataCacheStats.h b/Engine/Source/Developer/DerivedDataCache/Private/DerivedDataCacheStats.h
--- a/Engine/Source/Developer/DerivedDataCache/Private/DerivedDataCacheStats.h
+++ b/Engine/Source/Developer/DerivedDataCache/Private/DerivedDataCacheStats.h
@@ -180,6 +180,8 @@ private:
 	void SetStatus(ECacheStoreStatusCode StatusCode, const FText& Status) final;
 	void SetAttribute(FStringView Key, FStringView Value) final;
 	void AddRequest(const FCacheStoreRequestStats& Stats) final;
+	void AddLatency(FMonotonicTimePoint StartTime, FMonotonicTimePoint EndTime, FMonotonicTimeSpan Latency) final;
+	double GetAverageLatency() final;
 };
 
 class FCacheStats final
diff --git a/Engine/Source/Developer/DerivedDataCache/Private/DerivedDataCacheStore.h b/Engine/Source/Developer/DerivedDataCache/Private/DerivedDataCacheStore.h
--- a/Engine/Source/Developer/DerivedDataCache/Private/DerivedDataCacheStore.h
+++ b/Engine/Source/Developer/DerivedDataCache/Private/DerivedDataCacheStore.h
@@ -120,6 +120,22 @@ enum class ECacheStoreRequestType : uint8
 	Value,
 };
 
+inline const TCHAR* LexToString(ECacheStoreRequestType CacheStoreRequestType)
+{
+	switch (CacheStoreRequestType)
+	{
+	case ECacheStoreRequestType::None:
+		return TEXT("None");
+	case ECacheStoreRequestType::Record:
+		return TEXT("Record");
+	case ECacheStoreRequestType::Value:
+		return TEXT("Value");
+	}
+
+	checkNoEntry();
+	return TEXT("Unknown value! (Update LexToString!)");
+}
+
 /**
  * The operation performed by the request.
  */
@@ -131,6 +147,24 @@ enum class ECacheStoreRequestOp : uint8
 	GetChunk,
 };
 
+inline const TCHAR* LexToString(ECacheStoreRequestOp CacheStoreRequestOp)
+{
+	switch (CacheStoreRequestOp)
+	{
+	case ECacheStoreRequestOp::None:
+		return TEXT("None");
+	case ECacheStoreRequestOp::Put:
+		return TEXT("Put");
+	case ECacheStoreRequestOp::Get:
+		return TEXT("Get");
+	case ECacheStoreRequestOp::GetChunk:
+		return TEXT("GetChunk");
+	}
+
+	checkNoEntry();
+	return TEXT("Unknown value! (Update LexToString!)");
+}
+
 enum class ECacheStoreFlags : uint32
 {
 	None            = 0,
@@ -268,6 +302,12 @@ public:
 
 	/** Adds stats for a single request that was processed by the associated cache store. */
 	virtual void AddRequest(const FCacheStoreRequestStats& Stats) = 0;
+
+	/** Adds stats for latency for a measurement that was done for the associated cache store. */
+	virtual void AddLatency(FMonotonicTimePoint StartTime, FMonotonicTimePoint EndTime, FMonotonicTimeSpan Latency) = 0;
+
+	/** Gets the average latency value for the current time in seconds. */
+	virtual double GetAverageLatency() = 0;
 };
 
 template <typename RequestRangeType, typename OnCompleteType>
diff --git a/Engine/Source/Developer/DerivedDataCache/Private/ZenCacheStore.cpp b/Engine/Source/Developer/DerivedDataCache/Private/ZenCacheStore.cpp
--- a/Engine/Source/Developer/DerivedDataCache/Private/ZenCacheStore.cpp
+++ b/Engine/Source/Developer/DerivedDataCache/Private/ZenCacheStore.cpp
@@ -5,6 +5,7 @@
 
 #if UE_WITH_ZEN
 
+#include "Async/ManualResetEvent.h"
 #include "Async/UniqueLock.h"
 #include "BatchView.h"
 #include "DerivedDataBackendInterface.h"
@@ -16,6 +17,7 @@
 #include "DerivedDataRequestOwner.h"
 #include "Experimental/ZenStatistics.h"
 #include "HAL/FileManager.h"
+#include "HAL/Thread.h"
 #include "Http/HttpClient.h"
 #include "Math/UnrealMathUtility.h"
 #include "Misc/App.h"
@@ -146,6 +148,9 @@ private:
 	void EnqueueAsyncRpc(IRequestOwner& Owner, FCbObject RequestObject, FOnRpcComplete&& OnComplete);
 	void EnqueueAsyncRpc(IRequestOwner& Owner, const FCbPackage& RequestPackage, FOnRpcComplete&& OnComplete);
 
+	void ConditionalEvaluatePerformance();
+	void UpdateStatus();
+
 	template <typename T, typename... ArgTypes>
 	static TRefCountPtr<T> MakeAsyncOp(ArgTypes&&... Args)
 	{
@@ -171,6 +176,17 @@ private:
 	class FGetValueOp;
 	class FGetChunksOp;
 
+	enum class EHealthStatus
+	{
+		Unknown,
+		Ok,
+		Error,
+	};
+
+	using FOnHealthStatusComplete = TUniqueFunction<void(THttpUniquePtr<IHttpResponse>& HttpResponse, EHealthStatus HealthStatus)>;
+	class FHealthStatusReceiver;
+	class FAsyncHealthStatusReceiver;
+
 	class FCbPackageReceiver;
 	class FAsyncCbPackageReceiver;
 
@@ -182,11 +198,19 @@ private:
 	FHttpRequestQueue RequestQueue;
 	bool bIsUsable = false;
 	bool bIsLocalConnection = false;
+	bool bTryEvaluatePerformance = false;
+	std::atomic<bool> bDeactivatedForPerformance = false;
 	int32 BatchPutMaxBytes = 1024*1024;
 	int32 CacheRecordBatchSize = 8;
 	int32 CacheChunksBatchSize = 8;
 	FBackendDebugOptions DebugOptions;
 	TAnsiStringBuilder<256> RpcUri;
+	std::atomic<int64> LastPerformanceEvaluationTicks;
+	TOptional<FThread> PerformanceEvaluationThread;
+	FManualResetEvent PerformanceEvaluationThreadShutdownEvent;
+	float DeactivateAtMs = -1.0f;
+	ECacheStoreFlags OperationalFlags;
+	FRequestOwner PerformanceEvaluationRequestOwner;
 };
 
 template <typename RequestType>
@@ -265,6 +289,7 @@ public:
 			Request.Stats.OtherThreadTime = AverageOtherThreadTime;
 			Request.Stats.Latency = AverageLatency;
 		}
+		CacheStore.ConditionalEvaluatePerformance();
 	}
 
 	void IssueRequests()
@@ -314,7 +339,8 @@ public:
 			auto OnRpcComplete = [this, OpRef = TRefCountPtr<FPutOp>(this), Batch](THttpUniquePtr<IHttpResponse>& HttpResponse, FCbPackage& Response)
 			{
 				FRequestTimer RequestTimer(Requests[0].Stats);
-				Requests[0].Stats.Latency = FMonotonicTimeSpan::FromSeconds(HttpResponse->GetStats().StartTransferTime);
+				// Latency can't be measured for Put operations because it is intertwined with upload time.
+				Requests[0].Stats.Latency = FMonotonicTimeSpan::Infinity();
 
 				int32 RequestIndex = 0;
 				if (HttpResponse->GetErrorCode() == EHttpErrorCode::None && HttpResponse->GetStatusCode() >= 200 && HttpResponse->GetStatusCode() <= 299)
@@ -463,6 +489,7 @@ public:
 			Request.Stats.OtherThreadTime = AverageOtherThreadTime;
 			Request.Stats.Latency = AverageLatency;
 		}
+		CacheStore.ConditionalEvaluatePerformance();
 	}
 
 	void IssueRequests()
@@ -516,7 +543,8 @@ public:
 			auto OnRpcComplete = [this, OpRef = TRefCountPtr<FGetOp>(OriginalOp), Batch](THttpUniquePtr<IHttpResponse>& HttpResponse, FCbPackage& Response)
 			{
 				FRequestTimer RequestTimer(Requests[0].Stats);
-				Requests[0].Stats.Latency = FMonotonicTimeSpan::FromSeconds(HttpResponse->GetStats().StartTransferTime);
+				const FHttpResponseStats& ResponseStats = HttpResponse->GetStats();
+				Requests[0].Stats.Latency = FMonotonicTimeSpan::FromSeconds(ResponseStats.StartTransferTime - ResponseStats.ConnectTime);
 
 				int32 RequestIndex = 0;
 				if (HttpResponse->GetErrorCode() == EHttpErrorCode::None && HttpResponse->GetStatusCode() >= 200 && HttpResponse->GetStatusCode() <= 299)
@@ -667,6 +695,7 @@ public:
 			Request.Stats.OtherThreadTime = AverageOtherThreadTime;
 			Request.Stats.Latency = AverageLatency;
 		}
+		CacheStore.ConditionalEvaluatePerformance();
 	}
 
 	void IssueRequests()
@@ -718,7 +747,8 @@ public:
 			auto OnRpcComplete = [this, OpRef = TRefCountPtr<FPutValueOp>(this), Batch](THttpUniquePtr<IHttpResponse>& HttpResponse, FCbPackage& Response)
 			{
 				FRequestTimer RequestTimer(Requests[0].Stats);
-				Requests[0].Stats.Latency = FMonotonicTimeSpan::FromSeconds(HttpResponse->GetStats().StartTransferTime);
+				// Latency can't be measured for Put operations because it is intertwined with upload time.
+				Requests[0].Stats.Latency = FMonotonicTimeSpan::Infinity();
 
 				int32 RequestIndex = 0;
 				if (HttpResponse->GetErrorCode() == EHttpErrorCode::None && HttpResponse->GetStatusCode() >= 200 && HttpResponse->GetStatusCode() <= 299)
@@ -855,6 +885,7 @@ public:
 			Request.Stats.OtherThreadTime = AverageOtherThreadTime;
 			Request.Stats.Latency = AverageLatency;
 		}
+		CacheStore.ConditionalEvaluatePerformance();
 	}
 
 	void IssueRequests()
@@ -907,7 +938,8 @@ public:
 			auto OnRpcComplete = [this, OpRef = TRefCountPtr<FGetValueOp>(OriginalOp), Batch](THttpUniquePtr<IHttpResponse>& HttpResponse, FCbPackage& Response)
 			{
 				FRequestTimer RequestTimer(Requests[0].Stats);
-				Requests[0].Stats.Latency = FMonotonicTimeSpan::FromSeconds(HttpResponse->GetStats().StartTransferTime);
+				const FHttpResponseStats& ResponseStats = HttpResponse->GetStats();
+				Requests[0].Stats.Latency = FMonotonicTimeSpan::FromSeconds(ResponseStats.StartTransferTime - ResponseStats.ConnectTime);
 
 				int32 RequestIndex = 0;
 				if (HttpResponse->GetErrorCode() == EHttpErrorCode::None && HttpResponse->GetStatusCode() >= 200 && HttpResponse->GetStatusCode() <= 299)
@@ -1052,6 +1084,7 @@ public:
 			Request.Stats.OtherThreadTime = AverageOtherThreadTime;
 			Request.Stats.Latency = AverageLatency;
 		}
+		CacheStore.ConditionalEvaluatePerformance();
 	}
 
 	void IssueRequests()
@@ -1122,7 +1155,8 @@ public:
 			auto OnRpcComplete = [this, OpRef = TRefCountPtr<FGetChunksOp>(OriginalOp), Batch](THttpUniquePtr<IHttpResponse>& HttpResponse, FCbPackage& Response)
 			{
 				FRequestTimer RequestTimer(Requests[0].Stats);
-				Requests[0].Stats.Latency = FMonotonicTimeSpan::FromSeconds(HttpResponse->GetStats().StartTransferTime);
+				const FHttpResponseStats& ResponseStats = HttpResponse->GetStats();
+				Requests[0].Stats.Latency = FMonotonicTimeSpan::FromSeconds(ResponseStats.StartTransferTime - ResponseStats.ConnectTime);
 
 				int32 RequestIndex = 0;
 				if (HttpResponse->GetErrorCode() == EHttpErrorCode::None && HttpResponse->GetStatusCode() >= 200 && HttpResponse->GetStatusCode() <= 299)
@@ -1248,6 +1282,135 @@ private:
 	FOnCacheGetChunkComplete OnComplete;
 };
 
+class FZenCacheStore::FHealthStatusReceiver final : public IHttpReceiver
+{
+public:
+	FHealthStatusReceiver(const FHealthStatusReceiver&) = delete;
+	FHealthStatusReceiver& operator=(const FHealthStatusReceiver&) = delete;
+
+	explicit FHealthStatusReceiver(EHealthStatus& OutHealthStatus, IHttpReceiver* InNext = nullptr)
+		: HealthStatus(OutHealthStatus)
+		, Next(InNext)
+	{
+		HealthStatus = EHealthStatus::Unknown;
+	}
+
+private:
+	IHttpReceiver* OnCreate(IHttpResponse& Response) final
+	{
+		return &BodyReceiver;
+	}
+
+	IHttpReceiver* OnComplete(IHttpResponse& Response) final
+	{
+		FUtf8StringView ResponseStringView(reinterpret_cast<const UTF8CHAR*>(BodyArray.GetData()), IntCastChecked<int32>(BodyArray.Num()));
+		if (ResponseStringView == UTF8TEXTVIEW("OK!"))
+		{
+			HealthStatus = EHealthStatus::Ok;
+		}
+		else
+		{
+			HealthStatus = EHealthStatus::Error;
+		}
+		return Next;
+	}
+
+private:
+	EHealthStatus& HealthStatus;
+	IHttpReceiver* Next;
+	TArray64<uint8> BodyArray;
+	FHttpByteArrayReceiver BodyReceiver{ BodyArray, this };
+};
+
+class FZenCacheStore::FAsyncHealthStatusReceiver final : public FRequestBase, public IHttpReceiver
+{
+public:
+	FAsyncHealthStatusReceiver(const FAsyncHealthStatusReceiver&) = delete;
+	FAsyncHealthStatusReceiver& operator=(const FAsyncHealthStatusReceiver&) = delete;
+
+	FAsyncHealthStatusReceiver(
+		THttpUniquePtr<IHttpRequest>&& InRequest,
+		IRequestOwner* InOwner,
+		Zen::FZenServiceInstance& InZenServiceInstance,
+		FOnHealthStatusComplete&& InOnHealthStatusComplete)
+		: Request(MoveTemp(InRequest))
+		, Owner(InOwner)
+		, ZenServiceInstance(InZenServiceInstance)
+		, BaseReceiver(HealthStatus, this)
+		, OnHealthStatusComplete(MoveTemp(InOnHealthStatusComplete))
+	{
+		Request->SendAsync(this, Response);
+	}
+
+private:
+	// IRequest Interface
+
+	void SetPriority(EPriority Priority) final {}
+	void Cancel() final { Monitor->Cancel(); }
+	void Wait() final { Monitor->Wait(); }
+
+	// IHttpReceiver Interface
+
+	IHttpReceiver* OnCreate(IHttpResponse& LocalResponse) final
+	{
+		Monitor = LocalResponse.GetMonitor();
+		Owner->Begin(this);
+		return &BaseReceiver;
+	}
+
+	bool ShouldRecoverAndRetry(IHttpResponse& LocalResponse)
+	{
+		if (!ZenServiceInstance.IsServiceRunningLocally())
+		{
+			return false;
+		}
+
+		if ((LocalResponse.GetErrorCode() == EHttpErrorCode::Connect) ||
+			(LocalResponse.GetErrorCode() == EHttpErrorCode::TlsConnect) ||
+			(LocalResponse.GetErrorCode() == EHttpErrorCode::TimedOut))
+		{
+			return true;
+		}
+
+		return false;
+	}
+
+	IHttpReceiver* OnComplete(IHttpResponse& LocalResponse) final
+	{
+		Owner->End(this, [Self = this, &LocalResponse]
+		{
+			if (Self->ShouldRecoverAndRetry(LocalResponse) && Self->ZenServiceInstance.TryRecovery())
+			{
+				new FAsyncHealthStatusReceiver(MoveTemp(Self->Request), Self->Owner, Self->ZenServiceInstance, MoveTemp(Self->OnHealthStatusComplete));
+				return;
+			}
+
+			Self->Request.Reset();
+			if (Self->OnHealthStatusComplete)
+			{
+				// Launch a task for the completion function since it can execute arbitrary code.
+				Self->Owner->LaunchTask(TEXT("ZenHealthStatusComplete"), [Self = TRefCountPtr(Self)]
+				{
+					// Ensuring that the OnRpcComplete method is destroyed by the time we exit this method by moving it to a local scope variable
+					FOnHealthStatusComplete LocalOnComplete = MoveTemp(Self->OnHealthStatusComplete);
+					LocalOnComplete(Self->Response, Self->HealthStatus);
+				});
+			}
+		});
+		return nullptr;
+	}
+
+private:
+	THttpUniquePtr<IHttpRequest> Request;
+	THttpUniquePtr<IHttpResponse> Response;
+	TRefCountPtr<IHttpResponseMonitor> Monitor;
+	IRequestOwner* Owner;
+	Zen::FZenServiceInstance& ZenServiceInstance;
+	EHealthStatus HealthStatus;
+	FHealthStatusReceiver BaseReceiver;
+	FOnHealthStatusComplete OnHealthStatusComplete;
+};
+
 class FZenCacheStore::FCbPackageReceiver final : public IHttpReceiver
 {
 public:
@@ -1385,6 +1548,7 @@ FZenCacheStore::FZenCacheStore(
 	ICacheStoreOwner* InStoreOwner)
 	: ZenService(InServiceUrl)
 	, StoreOwner(InStoreOwner)
+	, PerformanceEvaluationRequestOwner(EPriority::Low)
 {
 	Initialize(InNamespace, InName, InConfig);
 }
@@ -1397,12 +1561,22 @@ FZenCacheStore::FZenCacheStore(
 	ICacheStoreOwner* InStoreOwner)
 	: ZenService(MoveTemp(InSettings))
 	, StoreOwner(InStoreOwner)
+	, PerformanceEvaluationRequestOwner(EPriority::Low)
 {
 	Initialize(InNamespace, InName, InConfig);
 }
 
 FZenCacheStore::~FZenCacheStore()
 {
+	PerformanceEvaluationRequestOwner.Cancel();
+	if (PerformanceEvaluationThread.IsSet())
+	{
+		PerformanceEvaluationThreadShutdownEvent.Notify();
+		PerformanceEvaluationThread->Join();
+		PerformanceEvaluationThread.Reset();
+		PerformanceEvaluationThreadShutdownEvent.Reset();
+	}
+
 	if (StoreStats)
 	{
 		StoreOwner->DestroyStats(StoreStats);
@@ -1414,6 +1588,7 @@ void FZenCacheStore::Initialize(
 	const TCHAR* InName,
 	const TCHAR* InConfig)
 {
+	LastPerformanceEvaluationTicks.store(FDateTime::UtcNow().GetTicks(), std::memory_order_relaxed);
 	Namespace = InNamespace;
 	if (IsServiceReady())
 	{
@@ -1451,13 +1626,18 @@ void FZenCacheStore::Initialize(
 			bool bRemote = !bIsLocalConnection;
 			FParse::Bool(InConfig, TEXT("Remote="), bRemote);
 
+			FParse::Value(InConfig, TEXT("DeactivateAt="), DeactivateAtMs);
+
 			ECacheStoreFlags Flags = ECacheStoreFlags::Query;
 			Flags |= bReadOnly ? ECacheStoreFlags::None : ECacheStoreFlags::Store;
 			Flags |= bLocal ? ECacheStoreFlags::Local : ECacheStoreFlags::None;
 			Flags |= bRemote ? ECacheStoreFlags::Remote : ECacheStoreFlags::None;
 
+			OperationalFlags = Flags;
+
 			StoreOwner->Add(this, Flags);
 			StoreStats = StoreOwner->CreateStats(this, Flags, TEXT("Zen"), InName, ZenService.GetInstance().GetURL());
+			bTryEvaluatePerformance = !GIsBuildMachine && (StoreStats != nullptr) && (DeactivateAtMs > 0.0f);
 		}
 
 		// Issue a request for stats as it will be fetched asynchronously and issuing now makes them available sooner for future callers.
@@ -1508,6 +1688,115 @@ void FZenCacheStore::EnqueueAsyncRpc(IRequestOwner& Owner, const FCbPackage& Req
 	new FAsyncCbPackageReceiver(MoveTemp(Request), &Owner, ZenService.GetInstance(), MoveTemp(OnComplete));
 }
 
+void FZenCacheStore::ConditionalEvaluatePerformance()
+{
+	if (!bTryEvaluatePerformance)
+	{
+		return;
+	}
+
+	// Look for an opportunity to measure and evaluate if performance is acceptable.
+	int64 LocalLastPerfEvaluationTicks = LastPerformanceEvaluationTicks.load(std::memory_order_relaxed);
+	FTimespan TimespanSinceLastPerfEval = FDateTime::UtcNow() - FDateTime(LocalLastPerfEvaluationTicks);
+
+	if (TimespanSinceLastPerfEval < FTimespan::FromSeconds(30))
+	{
+		return;
+	}
+
+	if (!LastPerformanceEvaluationTicks.compare_exchange_strong(LocalLastPerfEvaluationTicks, FDateTime::UtcNow().GetTicks()))
+	{
+		return;
+	}
+
+	// We won the race and get to do the performance check
+
+	if (PerformanceEvaluationThread.IsSet() && PerformanceEvaluationThreadShutdownEvent.IsNotified())
+	{
+		// Join and cleanup old thread before we consider whether we need to start a new one
+		PerformanceEvaluationThread->Join();
+		PerformanceEvaluationThread.Reset();
+		PerformanceEvaluationThreadShutdownEvent.Reset();
+	}
+
+	if (StoreStats->GetAverageLatency() * 1000 > DeactivateAtMs)
+	{
+		if (!bDeactivatedForPerformance.load(std::memory_order_relaxed))
+		{
+			StoreOwner->SetFlags(this, ECacheStoreFlags::None);
+			UE_LOG(LogDerivedDataCache, Warning,
+				TEXT("%s: Performance does not meet minimum criteria. "
+					"It will be deactivated until performance measurements improve. "
+					"If this is consistent, consider disabling this cache store through "
+					"environment variables or other configuration."),
+				*GetName());
+			bDeactivatedForPerformance.store(true, std::memory_order_relaxed);
+			UpdateStatus();
+		}
+
+		if (!PerformanceEvaluationThread.IsSet())
+		{
+			PerformanceEvaluationThread.Emplace(TEXT("ZenCacheStore Performance Evaluation"), [this]
+			{
+				while (!PerformanceEvaluationThreadShutdownEvent.WaitFor(FMonotonicTimeSpan::FromSeconds(30.0)))
+				{
+					IRequestOwner& Owner(PerformanceEvaluationRequestOwner);
+					FRequestBarrier Barrier(Owner);
+					THttpUniquePtr<IHttpRequest> Request = RequestQueue.CreateRequest({});
+					TAnsiStringBuilder<256> StatusUri;
+					StatusUri << ZenService.GetInstance().GetURL() << ANSITEXTVIEW("/health/status");
+					Request->SetUri(StatusUri);
+					Request->SetMethod(EHttpMethod::Get);
+					Request->AddAcceptType(EHttpMediaType::Text);
+					new FAsyncHealthStatusReceiver(MoveTemp(Request), &Owner, ZenService.GetInstance(), [this, StartTime = FMonotonicTimePoint::Now()](THttpUniquePtr<IHttpResponse>& HttpResponse, EHealthStatus HealthStatus)
+						{
+							if (HealthStatus != EHealthStatus::Ok)
+							{
+								return;
+							}
+
+							double LatencySec = (HttpResponse->GetStats().StartTransferTime - HttpResponse->GetStats().ConnectTime);
+							StoreStats->AddLatency(StartTime, FMonotonicTimePoint::Now(), FMonotonicTimeSpan::FromSeconds(LatencySec));
+							if (StoreStats->GetAverageLatency() * 1000 <= DeactivateAtMs)
+							{
+								if (PerformanceEvaluationThread.IsSet())
+								{
+									PerformanceEvaluationThreadShutdownEvent.Notify();
+								}
+
+								if (bDeactivatedForPerformance.load(std::memory_order_relaxed))
+								{
+									StoreOwner->SetFlags(this, OperationalFlags);
+									UE_LOG(LogDerivedDataCache, Display,
+										TEXT("%s: Performance has improved and meets minimum performance criteria. "
+											"It will be reactivated now."),
+										*GetName());
+									bDeactivatedForPerformance.store(false, std::memory_order_relaxed);
+									UpdateStatus();
+								}
+							}
+						});
+				}
+			});
+		}
+	}
+}
+
+void FZenCacheStore::UpdateStatus()
+{
+	if (StoreStats)
+	{
+		if (bDeactivatedForPerformance.load(std::memory_order_relaxed))
+		{
+			StoreStats->SetStatus(ECacheStoreStatusCode::Warning, NSLOCTEXT("DerivedDataCache", "DeactivatedForPerformance", "Deactivated for performance"));
+		}
+		else
+		{
+			StoreStats->SetStatus(ECacheStoreStatusCode::None, {});
+		}
+	}
+}
+
 void FZenCacheStore::LegacyStats(FDerivedDataCacheStatsNode& OutNode)
 {
 	checkNoEntry();
