Change certain messages relating to PUTting data into jupiter from Log verbosity to Display verbosity.

diff --git a/Engine/Source/Developer/DerivedDataCache/Private/HttpCacheStore.cpp b/Engine/Source/Developer/DerivedDataCache/Private/HttpCacheStore.cpp
--- a/Engine/Source/Developer/DerivedDataCache/Private/HttpCacheStore.cpp
+++ b/Engine/Source/Developer/DerivedDataCache/Private/HttpCacheStore.cpp
@@ -1002,7 +1002,7 @@ void FHttpCacheStore::FPutPackageOp::BeginPutBlobs(FCbPackage&& Package, FCacheP
 	{
 		if (Response.Status == EStatus::Error)
 		{
-			UE_LOG(LogDerivedDataCache, Log, TEXT("%s: Failed to put reference object for put of %s from '%s'"),
+			UE_LOG(LogDerivedDataCache, Display, TEXT("%s: Failed to put reference object for put of %s from '%s'"),
 				*CacheStore.Domain, *WriteToString<96>(Key), *Name);
 		}
 		EndPut(Response.Status);
@@ -1050,7 +1050,7 @@ void FHttpCacheStore::FPutPackageOp::BeginPutBlobs(FCbPackage&& Package, FCacheP
 				}
 				bExpectedHashesSerialized = true;
 			}
-			UE_LOG(LogDerivedDataCache, Log, TEXT("%s: Server reported needed hash '%s' that is outside the set of expected hashes (%s) for put of %s from '%s'"),
+			UE_LOG(LogDerivedDataCache, Display, TEXT("%s: Server reported needed hash '%s' that is outside the set of expected hashes (%s) for put of %s from '%s'"),
 				*CacheStore.Domain, *WriteToString<96>(NeededBlobHash), *ExpectedHashes, *WriteToString<96>(Key), *Name);
 		}
 	}
