Reduce frequency of "partial put" operations where we put a ref but not the blobs it refers to because of either a 502 "Bad Gateway" response, or because we experience an "operation too slow" condition.

diff --git a/Engine/Source/Developer/DerivedDataCache/Private/HttpCacheStore.cpp b/Engine/Source/Developer/DerivedDataCache/Private/HttpCacheStore.cpp
--- a/Engine/Source/Developer/DerivedDataCache/Private/HttpCacheStore.cpp
+++ b/Engine/Source/Developer/DerivedDataCache/Private/HttpCacheStore.cpp
@@ -612,8 +612,11 @@ private:
 			return true;
 		}
 
-		// Too many requests, make a new attempt.
-		if (LocalResponse.GetStatusCode() == 429)
+		// Make a new attempt if the response status code is any of:
+		// 429 - Too many requests
+		// 502 - Bad gateway
+		int32 StatusCode = LocalResponse.GetStatusCode();
+		if ((StatusCode == 429) || (StatusCode == 502))
 		{
 			return true;
 		}
@@ -2307,8 +2310,8 @@ FHttpClientParams FHttpCacheStore::GetDefaultClientParams() const
 	FHttpClientParams ClientParams;
 	ClientParams.DnsCacheTimeout = 15;
 	ClientParams.ConnectTimeout = 3 * 1000;
-	ClientParams.LowSpeedLimit = 1024;
-	ClientParams.LowSpeedTime = 10;
+	ClientParams.LowSpeedLimit = 10;
+	ClientParams.LowSpeedTime = 30;
 	ClientParams.TlsLevel = EHttpTlsLevel::All;
 	ClientParams.bFollowRedirects = true;
 	ClientParams.bFollow302Post = true;
