Restoring the previous rate limit on the jupiter client, and specifically disabling the rate limit for operations in the PUT queue.

diff --git a/Engine/Source/Developer/DerivedDataCache/Private/HttpCacheStore.cpp b/Engine/Source/Developer/DerivedDataCache/Private/HttpCacheStore.cpp
--- a/Engine/Source/Developer/DerivedDataCache/Private/HttpCacheStore.cpp
+++ b/Engine/Source/Developer/DerivedDataCache/Private/HttpCacheStore.cpp
@@ -2227,6 +2227,9 @@ FHttpCacheStore::FHttpCacheStore(const FHttpCacheStoreParams& Params, ICacheStor
 		ClientParams.MinRequests = UE_HTTPDDC_GET_REQUEST_POOL_SIZE;
 		GetRequestQueue.Initialize(*ConnectionPool, ClientParams);
 
+		// Disabling rate limits during PUT operations as the cause too many spurious failures to put blobs or finalize refs
+		ClientParams.LowSpeedLimit = 0;
+		ClientParams.LowSpeedTime = 0;
 		ClientParams.MaxRequests = UE_HTTPDDC_PUT_REQUEST_POOL_SIZE;
 		ClientParams.MinRequests = UE_HTTPDDC_PUT_REQUEST_POOL_SIZE;
 		PutRequestQueue.Initialize(*ConnectionPool, ClientParams);
@@ -2310,8 +2313,8 @@ FHttpClientParams FHttpCacheStore::GetDefaultClientParams() const
 	FHttpClientParams ClientParams;
 	ClientParams.DnsCacheTimeout = 15;
 	ClientParams.ConnectTimeout = 3 * 1000;
-	ClientParams.LowSpeedLimit = 10;
-	ClientParams.LowSpeedTime = 30;
+	ClientParams.LowSpeedLimit = 1024;
+	ClientParams.LowSpeedTime = 10;
 	ClientParams.TlsLevel = EHttpTlsLevel::All;
 	ClientParams.bFollowRedirects = true;
 	ClientParams.bFollow302Post = true;
