DDC Http Store - Added ability to configure AuthScheme (defaulting to JWTBearer as before) and the ability to directly specify a auth token. This can be used to run with the ServiceAccount scheme in Cloud DDC which we intend to use for getting onboarded quicker by licensees (it is not a recommended approch to use for wider deployments across teams).

diff --git a/Engine/Source/Developer/DerivedDataCache/Private/HttpCacheStore.cpp b/Engine/Source/Developer/DerivedDataCache/Private/HttpCacheStore.cpp
--- a/Engine/Source/Developer/DerivedDataCache/Private/HttpCacheStore.cpp
+++ b/Engine/Source/Developer/DerivedDataCache/Private/HttpCacheStore.cpp
@@ -295,7 +295,7 @@ private:
 class FHttpAccessToken
 {
 public:
-	void SetToken(FStringView Token);
+	void SetToken(FStringView Scheme, FStringView Token);
 	inline uint32 GetSerial() const { return Serial.load(std::memory_order_relaxed); }
 	friend FAnsiStringBuilderBase& operator<<(FAnsiStringBuilderBase& Builder, const FHttpAccessToken& Token);
 
@@ -305,13 +305,20 @@ private:
 	std::atomic<uint32> Serial;
 };
 
-void FHttpAccessToken::SetToken(const FStringView Token)
+void FHttpAccessToken::SetToken(const FStringView Scheme, const FStringView Token)
 {
 	FWriteScopeLock WriteLock(Lock);
-	const FAnsiStringView Prefix = ANSITEXTVIEW("Bearer ");
+	const int32 SchemeLen = FPlatformString::ConvertedLength<ANSICHAR>(Scheme.GetData(), Scheme.Len());
 	const int32 TokenLen = FPlatformString::ConvertedLength<ANSICHAR>(Token.GetData(), Token.Len());
-	Header.Empty(Prefix.Len() + TokenLen);
-	Header.Append(Prefix.GetData(), Prefix.Len());
+
+	Header.Empty(SchemeLen + 1 + TokenLen);
+	
+	const int32 SchemeIndex = Header.AddUninitialized(SchemeLen);
+	FPlatformString::Convert(Header.GetData() + SchemeIndex, SchemeLen, Scheme.GetData(), Scheme.Len());
+	
+	const FAnsiStringView Seperator = ANSITEXTVIEW(" ");
+	Header.Append(Seperator.GetData(), Seperator.Len());
+
 	const int32 TokenIndex = Header.AddUninitialized(TokenLen);
 	FPlatformString::Convert(Header.GetData() + TokenIndex, TokenLen, Token.GetData(), Token.Len());
 	Serial.fetch_add(1, std::memory_order_relaxed);
@@ -337,6 +344,8 @@ struct FHttpCacheStoreParams
 	FString OAuthProviderIdentifier;
 	FString OAuthAccessToken;
 	FString OAuthPinnedPublicKeys;
+	FString AuthScheme;
+
 	bool bResolveHostCanonicalName = true;
 	bool bReadOnly = false;
 
@@ -420,6 +429,7 @@ private:
 	FString OAuthProviderIdentifier;
 	FString OAuthAccessToken;
 	FString HttpVersion;
+	FString AuthScheme;
 
 	FAnsiStringBuilderBase EffectiveDomain;
 
@@ -2181,6 +2191,7 @@ FHttpCacheStore::FHttpCacheStore(const FHttpCacheStoreParams& Params, ICacheStor
 	, OAuthProviderIdentifier(Params.OAuthProviderIdentifier)
 	, OAuthAccessToken(Params.OAuthAccessToken)
 	, HttpVersion(Params.HttpVersion)
+	, AuthScheme(Params.AuthScheme)
 	, StoreOwner(Owner)
 	, bReadOnly(Params.bReadOnly)
 {
@@ -2484,7 +2495,7 @@ void FHttpCacheStore::SetAccessTokenAndUnlock(FScopeLock& Lock, FStringView Toke
 	{
 		Access = MakeUnique<FHttpAccessToken>();
 	}
-	Access->SetToken(Token);
+	Access->SetToken(AuthScheme, Token);
 
 	constexpr double RefreshGracePeriod = 20.0f;
 	if (RefreshDelay > RefreshGracePeriod)
@@ -3304,6 +3315,7 @@ void FHttpCacheStoreParams::Parse(const TCHAR* NodeName, const TCHAR* Config)
 
 	FParse::Value(Config, TEXT("OAuthProviderIdentifier="), OAuthProviderIdentifier);
 
+	FParse::Value(Config, TEXT("OAuthAccess="), OAuthAccessToken);
 	if (FParse::Value(Config, TEXT("OAuthAccessTokenEnvOverride="), OverrideName))
 	{
 		FString AccessToken = FPlatformMisc::GetEnvironmentVariable(*OverrideName);
@@ -3315,6 +3327,12 @@ void FHttpCacheStoreParams::Parse(const TCHAR* NodeName, const TCHAR* Config)
 		}
 	}
 
+	FParse::Value(Config, TEXT("AuthScheme="), AuthScheme);
+	if (AuthScheme.IsEmpty())
+	{
+		AuthScheme = "Bearer";
+	}
+
 	FParse::Value(Config, TEXT("OAuthPinnedPublicKeys="), OAuthPinnedPublicKeys);
 
 	// Cache Params
