Restore the printing of the response body when an unexpected error is received, but now with consideration for the content type of the response.

#okfogithub public

diff --git a/Engine/Source/Developer/DerivedDataCache/Private/HttpCacheStore.cpp b/Engine/Source/Developer/DerivedDataCache/Private/HttpCacheStore.cpp
--- a/Engine/Source/Developer/DerivedDataCache/Private/HttpCacheStore.cpp
+++ b/Engine/Source/Developer/DerivedDataCache/Private/HttpCacheStore.cpp
@@ -658,8 +658,11 @@ private:
 
 			if (bUnexpectedError)
 			{
+				FString Body = Operation->GetBodyAsString();
+				Body.ReplaceCharInline(TEXT('\r'), TEXT(' '));
+				Body.ReplaceCharInline(TEXT('\n'), TEXT(' '));
 				UE_LOG(LogDerivedDataCache, Display,
-					TEXT("HTTP: %s (%s)"), *WriteToString<256>(LocalResponse), *StatsText);
+					TEXT("HTTP: %s (%s) %s"), *WriteToString<256>(LocalResponse), *StatsText, *Body);
 			}
 			else
 			{
@@ -752,9 +755,15 @@ void FHttpCacheStore::FHttpOperation::SendAsync(IRequestOwner& Owner, TUniqueFun
 FString FHttpCacheStore::FHttpOperation::GetBodyAsString() const
 {
 	static_assert(sizeof(uint8) == sizeof(UTF8CHAR));
-	const int32 Len = IntCastChecked<int32>(ResponseBody.GetSize());
-	if (GetContentType() == EHttpMediaType::CbObject)
-	{
+	uint64 ResponseBodySize = ResponseBody.GetSize();
+	EHttpMediaType ContentType = GetContentType();
+	switch (ContentType)
+	{
+	case EHttpMediaType::Text:
+	case EHttpMediaType::Json:
+	case EHttpMediaType::Yaml:
+		return FString::ConstructFromPtrSize((const UTF8CHAR*)ResponseBody.GetData(), int32(FMath::Clamp<uint64>(ResponseBodySize, 0, MAX_int32)));
+	case EHttpMediaType::CbObject:
 		if (ValidateCompactBinary(ResponseBody, ECbValidateMode::Default) == ECbValidateError::None)
 		{
 			TUtf8StringBuilder<1024> JsonStringBuilder;
@@ -762,8 +771,19 @@ FString FHttpCacheStore::FHttpOperation::GetBodyAsString() const
 			CompactBinaryToCompactJson(ResponseObject, JsonStringBuilder);
 			return JsonStringBuilder.ToString();
 		}
+		return FString::Printf(TEXT("Invalid compact binary object of size %" UINT64_FMT), ResponseBodySize);
+	case EHttpMediaType::CompressedBinary:
+		{
+			FCompressedBuffer Buffer = FCompressedBuffer::FromCompressed(ResponseBody);
+			if (!Buffer.IsNull())
+			{
+				return FString::Printf(TEXT("CompressedBuffer rawhash:%s, rawsize:%" UINT64_FMT ", compressedsize:%" UINT64_FMT), *WriteToString<32>(Buffer.GetRawHash()), Buffer.GetRawSize(), Buffer.GetCompressedSize());
+			}
+			return FString::Printf(TEXT("Invalid compressed buffer of size %" UINT64_FMT), ResponseBodySize);
+		}
+	default:
+		return FString::Printf(TEXT("Content type '%s' of size %" UINT64_FMT), *WriteToString<32>(LexToString(ContentType)), ResponseBodySize);
 	}
-	return FString::ConstructFromPtrSize((const UTF8CHAR*)ResponseBody.GetData(), Len);
 }
 
 TSharedPtr<FJsonObject> FHttpCacheStore::FHttpOperation::GetBodyAsJson() const
