Downgrade warning severity messages to display severity when we are unable to reach a remote zenserver (as opposed to a local one).

diff --git a/Engine/Source/Developer/DerivedDataCache/Private/ZenCacheStore.cpp b/Engine/Source/Developer/DerivedDataCache/Private/ZenCacheStore.cpp
--- a/Engine/Source/Developer/DerivedDataCache/Private/ZenCacheStore.cpp
+++ b/Engine/Source/Developer/DerivedDataCache/Private/ZenCacheStore.cpp
@@ -103,6 +103,12 @@ public:
 	 */
 	inline bool IsUsable() const { return bIsUsable; }
 
+	/**
+	 * Checks if cache service is on the local machine.
+	 * @return true if it is local
+	 */
+	inline bool IsLocalConnection() const { return bIsLocalConnection; }
+
 	// ICacheStore
 
 	void Put(
@@ -1973,7 +1979,14 @@ ILegacyCacheStore* CreateZenCacheStore(const TCHAR* NodeName, const TCHAR* Confi
 
 	if (!Backend->IsUsable())
 	{
-		UE_LOG(LogDerivedDataCache, Warning, TEXT("%s: Failed to contact the service (%s), will not use it."), NodeName, *Backend->GetName());
+		if (Backend->IsLocalConnection())
+		{
+			UE_LOG(LogDerivedDataCache, Warning, TEXT("%s: Failed to contact the service (%s), will not use it."), NodeName, *Backend->GetName());
+		}
+		else
+		{
+			UE_LOG(LogDerivedDataCache, Display, TEXT("%s: Failed to contact the service (%s), will not use it."), NodeName, *Backend->GetName());
+		}
 		Backend.Reset();
 		return nullptr;
 	}
diff --git a/Engine/Source/Developer/Zen/Private/ZenServerInterface.cpp b/Engine/Source/Developer/Zen/Private/ZenServerInterface.cpp
--- a/Engine/Source/Developer/Zen/Private/ZenServerInterface.cpp
+++ b/Engine/Source/Developer/Zen/Private/ZenServerInterface.cpp
@@ -1816,7 +1816,14 @@ FZenServiceInstance::IsServiceReady()
 		}
 		else
 		{
-			UE_LOG(LogZenServiceInstance, Warning, TEXT("Unable to reach ZenServer HTTP service at %s. Status: %d . Response: %s"), ZenDomain.ToString(), Request.GetResponseCode(), *Request.GetResponseAsString());
+			if (IsServiceRunningLocally())
+			{
+				UE_LOG(LogZenServiceInstance, Warning, TEXT("Unable to reach ZenServer HTTP service at %s. Status: %d . Response: %s"), ZenDomain.ToString(), Request.GetResponseCode(), *Request.GetResponseAsString());
+			}
+			else
+			{
+				UE_LOG(LogZenServiceInstance, Display, TEXT("Unable to reach ZenServer HTTP service at %s. Status: %d . Response: %s"), ZenDomain.ToString(), Request.GetResponseCode(), *Request.GetResponseAsString());
+			}
 		}
 	}
 	return false;
