Change the strategy for HttpCacheStore Put operations to have larger and distinct queues for each of the put operation phases.  Instead of one queue of 24 maximum in-flight requests, we can now have 64 PutRef requests AND 64 PutBlobs  requests, and 64 PutFinalize requests operating simultaneously.  The purpose is to allow per-shader caching to:

1. Put a higher volume of individual cache items in a lesser amount of wall time
2. Reduce the gap in time between when a ref is put and when it is finalized during times of heavy Put workload

Along with this change there is a System.DerivedDataCache.HttpDerivedDataBackend.CacheStoreStressPut automated test (StressFilter) that puts 1000 records containing 12 byte values in each.  Before this change, the test took 21 seconds to complete.  It now takes 9 seconds to complete.  There are opportunities to improve this further through batching.

diff --git a/Engine/Source/Developer/DerivedDataCache/Private/HttpCacheStore.cpp b/Engine/Source/Developer/DerivedDataCache/Private/HttpCacheStore.cpp
--- a/Engine/Source/Developer/DerivedDataCache/Private/HttpCacheStore.cpp
+++ b/Engine/Source/Developer/DerivedDataCache/Private/HttpCacheStore.cpp
@@ -64,7 +64,9 @@
 #endif
 
 #define UE_HTTPDDC_GET_REQUEST_POOL_SIZE 128
-#define UE_HTTPDDC_PUT_REQUEST_POOL_SIZE 24
+#define UE_HTTPDDC_PUTREF_REQUEST_POOL_SIZE 64
+#define UE_HTTPDDC_PUTBLOBS_REQUEST_POOL_SIZE 64
+#define UE_HTTPDDC_PUTFINALIZE_REQUEST_POOL_SIZE 64
 #define UE_HTTPDDC_MAX_FAILED_LOGIN_ATTEMPTS 16
 #define UE_HTTPDDC_MAX_ATTEMPTS 4
 
@@ -440,7 +442,9 @@ private:
 	FBackendDebugOptions DebugOptions;
 	THttpUniquePtr<IHttpConnectionPool> ConnectionPool;
 	FHttpCacheStoreRequestQueue GetRequestQueue;
-	FHttpCacheStoreRequestQueue PutRequestQueue;
+	FHttpCacheStoreRequestQueue PutRefRequestQueue;
+	FHttpCacheStoreRequestQueue PutBlobsRequestQueue;
+	FHttpCacheStoreRequestQueue PutFinalizeRequestQueue;
 
 	FCriticalSection AccessCs;
 	TUniquePtr<FHttpAccessToken> Access;
@@ -463,11 +467,14 @@ private:
 	enum class EOperationCategory
 	{
 		Get,
-		Put,
+		PutRef,
+		PutBlobs,
+		PutFinalize
 	};
 
 	class FHttpOperation;
 
+	FHttpCacheStoreRequestQueue& PickRequestQueue(EOperationCategory Category);
 	TUniquePtr<FHttpOperation> WaitForHttpOperation(EOperationCategory Category);
 
 	/** Invokes the callback when an operation is available, or with null if canceled. */
@@ -926,7 +933,7 @@ void FHttpCacheStore::FPutPackageOp::Put(const FCacheKey& InKey, FCbPackage&& Pa
 
 void FHttpCacheStore::FPutPackageOp::BeginOperation(bool bFinalize, FOnCachePutRefComplete&& OnComplete)
 {
-	CacheStore.WaitForHttpOperationAsync(Owner, EOperationCategory::Put, [Self = TRefCountPtr(this), bFinalize, OnComplete = MoveTemp(OnComplete)](TUniquePtr<FHttpOperation>&& Operation) mutable
+	CacheStore.WaitForHttpOperationAsync(Owner, bFinalize ? EOperationCategory::PutFinalize : EOperationCategory::PutRef, [Self = TRefCountPtr(this), bFinalize, OnComplete = MoveTemp(OnComplete)](TUniquePtr<FHttpOperation>&& Operation) mutable
 	{
 		Self->BeginPutRef(MoveTemp(Operation), bFinalize, MoveTemp(OnComplete));
 	});
@@ -1101,7 +1108,7 @@ void FHttpCacheStore::FPutPackageOp::BeginPutBlobs(FCbPackage&& Package, FCacheP
 	FRequestBarrier Barrier(Owner);
 	for (const FCompressedBuffer& Blob : Blobs)
 	{
-		CacheStore.WaitForHttpOperationAsync(Owner, EOperationCategory::Put, [Self = TRefCountPtr(this), Blob](TUniquePtr<FHttpOperation>&& Operation)
+		CacheStore.WaitForHttpOperationAsync(Owner, EOperationCategory::PutBlobs, [Self = TRefCountPtr(this), Blob](TUniquePtr<FHttpOperation>&& Operation)
 		{
 			if (UNLIKELY(!Operation))
 			{
@@ -2261,9 +2268,15 @@ FHttpCacheStore::FHttpCacheStore(const FHttpCacheStoreParams& Params, ICacheStor
 		// Disabling rate limits during PUT operations as the cause too many spurious failures to put blobs or finalize refs
 		ClientParams.LowSpeedLimit = 0;
 		ClientParams.LowSpeedTime = 0;
-		ClientParams.MaxRequests = UE_HTTPDDC_PUT_REQUEST_POOL_SIZE;
-		ClientParams.MinRequests = UE_HTTPDDC_PUT_REQUEST_POOL_SIZE;
-		PutRequestQueue.Initialize(*ConnectionPool, ClientParams);
+		ClientParams.MaxRequests = UE_HTTPDDC_PUTREF_REQUEST_POOL_SIZE;
+		ClientParams.MinRequests = UE_HTTPDDC_PUTREF_REQUEST_POOL_SIZE;
+		PutRefRequestQueue.Initialize(*ConnectionPool, ClientParams);
+		ClientParams.MaxRequests = UE_HTTPDDC_PUTBLOBS_REQUEST_POOL_SIZE;
+		ClientParams.MinRequests = UE_HTTPDDC_PUTBLOBS_REQUEST_POOL_SIZE;
+		PutBlobsRequestQueue.Initialize(*ConnectionPool, ClientParams);
+		ClientParams.MaxRequests = UE_HTTPDDC_PUTFINALIZE_REQUEST_POOL_SIZE;
+		ClientParams.MinRequests = UE_HTTPDDC_PUTFINALIZE_REQUEST_POOL_SIZE;
+		PutFinalizeRequestQueue.Initialize(*ConnectionPool, ClientParams);
 
 		bIsUsable = true;
 
@@ -2553,6 +2566,24 @@ void FHttpCacheStore::SetAccessTokenAndUnlock(FScopeLock& Lock, FStringView Toke
 	}
 }
 
+FHttpCacheStoreRequestQueue& FHttpCacheStore::PickRequestQueue(EOperationCategory Category)
+{
+	switch (Category)
+	{
+	case EOperationCategory::Get:
+		return GetRequestQueue;
+	case EOperationCategory::PutRef:
+		return PutRefRequestQueue;
+	case EOperationCategory::PutBlobs:
+		return PutBlobsRequestQueue;
+	case EOperationCategory::PutFinalize:
+		return PutFinalizeRequestQueue;
+	default:
+		checkNoEntry();
+		return GetRequestQueue;
+	}
+}
+
 TUniquePtr<FHttpCacheStore::FHttpOperation> FHttpCacheStore::WaitForHttpOperation(EOperationCategory Category)
 {
 	if (Access && RefreshAccessTokenTime > 0.0 && RefreshAccessTokenTime < FPlatformTime::Seconds())
@@ -2565,7 +2596,7 @@ TUniquePtr<FHttpCacheStore::FHttpOperation> FHttpCacheStore::WaitForHttpOperatio
 	{
 		FHttpRequestParams Params;
 		FRequestOwner BlockingOwner(EPriority::Blocking);
-		FHttpCacheStoreRequestQueue& RequestQueue = (Category == EOperationCategory::Get) ? GetRequestQueue : PutRequestQueue;
+		FHttpCacheStoreRequestQueue& RequestQueue = PickRequestQueue(Category);
 		RequestQueue.CreateRequestAsync(BlockingOwner, Params, [&Request](THttpUniquePtr<IHttpRequest>&& AsyncRequest)
 		{
 			Request = MoveTemp(AsyncRequest);
@@ -2608,7 +2639,7 @@ void FHttpCacheStore::WaitForHttpOperationAsync(IRequestOwner& Owner, EOperation
 void FHttpCacheStore::WaitForHttpRequestAsync(IRequestOwner& Owner, EOperationCategory Category, TUniqueFunction<void (THttpUniquePtr<IHttpRequest>&&)>&& OnRequest)
 {
 	FHttpRequestParams Params;
-	FHttpCacheStoreRequestQueue& RequestQueue = (Category == EOperationCategory::Get) ? GetRequestQueue : PutRequestQueue;
+	FHttpCacheStoreRequestQueue& RequestQueue = PickRequestQueue(Category);
 	RequestQueue.CreateRequestAsync(Owner, Params, MoveTemp(OnRequest));
 }
 
diff --git a/Engine/Source/Developer/DerivedDataCache/Private/Tests/HttpCacheStoreTest.cpp b/Engine/Source/Developer/DerivedDataCache/Private/Tests/HttpCacheStoreTest.cpp
--- a/Engine/Source/Developer/DerivedDataCache/Private/Tests/HttpCacheStoreTest.cpp
+++ b/Engine/Source/Developer/DerivedDataCache/Private/Tests/HttpCacheStoreTest.cpp
@@ -14,6 +14,7 @@
 #include "Misc/Paths.h"
 #include "Misc/SecureHash.h"
 #include "Serialization/CompactBinaryWriter.h"
+#include "ProfilingDebugging/ScopedTimers.h"
 
 // Test is targeted at HttpDerivedDataBackend but with some backend test interface it could be generalized
 // to function against all backends.
@@ -487,7 +488,7 @@ protected:
 	static inline FString TestNamespace;
 };
 
-TArray<FCacheRecord> CreateTestCacheRecords(ICacheStore* InTestBackend, uint32 InNumKeys, uint32 InNumValues, FCbObject MetaContents = FCbObject(), const char* BucketName = nullptr)
+TArray<FCacheRecord> CreateTestCacheRecords(ICacheStore* InTestBackend, uint32 InNumKeys, uint32 InNumValues, FCbObject MetaContents = FCbObject(), const char* BucketName = nullptr, uint8 Salt = 0)
 {
 	using namespace UE::DerivedData;
 	FCacheBucket TestCacheBucket(BucketName ? BucketName : "AutoTestDummy");
@@ -511,7 +512,7 @@ TArray<FCacheRecord> CreateTestCacheRecords(ICacheStore* InTestBackend, uint32 I
 			ValueContents.AddUninitialized(NumBytes);
 			for (int32 ContentIndex = 0; ContentIndex < NumBytes; ++ContentIndex)
 			{
-				ValueContents[ContentIndex] = (uint8)(KeyIndex + ContentIndex);
+				ValueContents[ContentIndex] = (uint8)(KeyIndex + ContentIndex + Salt);
 			}
 			Values.Emplace(MakeSharedBufferFromArray(MoveTemp(ValueContents)));
 			HashBuilder.Update(Values.Last().GetView());
@@ -558,7 +559,7 @@ TArray<FCacheRecord> CreateTestCacheRecords(ICacheStore* InTestBackend, uint32 I
 	return CacheRecords;
 }
 
-TArray<FValue> CreateTestCacheValues(ICacheStore* InTestBackend, uint32 InNumValues, const char* BucketName = nullptr)
+TArray<FValue> CreateTestCacheValues(ICacheStore* InTestBackend, uint32 InNumValues, const char* BucketName = nullptr, uint8 Salt = 1)
 {
 	using namespace UE::DerivedData;
 	FCacheBucket TestCacheBucket(BucketName ? BucketName : "AutoTestDummy");
@@ -575,7 +576,7 @@ TArray<FValue> CreateTestCacheValues(ICacheStore* InTestBackend, uint32 InNumVal
 		ValueContents.AddUninitialized(NumBytes);
 		for (int32 ContentIndex = 0; ContentIndex < NumBytes; ++ContentIndex)
 		{
-			ValueContents[ContentIndex] = (uint8)(ValueIndex + ContentIndex + 52); // offset of 52 to keep the contents distinct from record test data
+			ValueContents[ContentIndex] = (uint8)(ValueIndex + ContentIndex + Salt);
 		}
 		ValueBuffers.Emplace(MakeSharedBufferFromArray(MoveTemp(ValueContents)));
 	}
@@ -772,6 +773,87 @@ bool FHttpCacheStoreTest::RunTest(const FString& Parameters)
 	return true;
 }
 
+static void PutTestRecords(ILegacyCacheStore* TestBackend, uint32 NumRecords, bool bSynchronous = false, TArray<FCacheKey>* OutKeys = nullptr)
+{
+	FCacheBucket TestCacheBucket("AutoTestDummy");
+
+	const int64 RunUniqueId = FDateTime::Now().GetTicks();
+
+	TArray<FCachePutRequest> PutRequests;
+	PutRequests.Reserve(NumRecords);
+
+	TArray<FSharedBuffer> ValueBuffers;
+	for (uint32 ValueIndex = 0; ValueIndex < NumRecords; ++ValueIndex)
+	{
+		TArray<uint8> ValueContents;
+		const int32 NumBytes = 12;
+		ValueContents.AddUninitialized(NumBytes);
+		*reinterpret_cast<int64*>(ValueContents.GetData()) = RunUniqueId;
+		*(reinterpret_cast<uint32*>(ValueContents.GetData()) + 2) = ValueIndex;
+		ValueBuffers.Emplace(MakeSharedBufferFromArray(MoveTemp(ValueContents)));
+	}
+
+	if (OutKeys)
+	{
+		OutKeys->Reserve(OutKeys->Num() + NumRecords);
+	}
+	uint64 KeyIndex = 0;
+	for (const FSharedBuffer& ValueBuffer : ValueBuffers)
+	{
+		FIoHash ValueHash(FIoHash::HashBuffer(ValueBuffer));
+
+		FCacheKey Key;
+		Key.Bucket = TestCacheBucket;
+		Key.Hash = ValueHash;
+
+		if (OutKeys)
+		{
+			OutKeys->Add(Key);
+		}
+
+		FCacheRecordBuilder RecordBuilder(Key);
+		RecordBuilder.AddValue(FValueId::FromHash(ValueHash), ValueBuffer);
+
+		PutRequests.Add({ {TEXT("AutoTestStressPutRecord")}, RecordBuilder.Build(), ECachePolicy::Default, KeyIndex });
+	}
+
+	FRequestOwner RequestOwner(bSynchronous ? EPriority::Blocking : EPriority::Normal);
+	{
+		FRequestBarrier RequestBarrier(RequestOwner);
+		if (!bSynchronous)
+		{
+			RequestOwner.KeepAlive();
+		}
+		TestBackend->Put(PutRequests, RequestOwner, [](FCachePutResponse&& Response)
+		{
+			check(Response.Status == EStatus::Ok);
+		});
+	}
+
+	if (bSynchronous)
+	{
+		RequestOwner.Wait();
+	}
+}
+
+
+// Stress test Put operations
+IMPLEMENT_HTTPDERIVEDDATA_AUTOMATION_TEST(FHttpCacheStoreStressPutTest, ".CacheStoreStressPut", EAutomationTestFlags::EditorContext | EAutomationTestFlags::StressFilter)
+bool FHttpCacheStoreStressPutTest::RunTest(const FString& Parameters)
+{
+	using namespace UE::DerivedData;
+	ILegacyCacheStore* TestBackend = GetTestBackend();
+
+	constexpr uint32 NumRecords = 1000;
+	{
+		FAutoScopedDurationTimer AutoTimer;
+		PutTestRecords(TestBackend, NumRecords, true);
+		UE_LOG(LogDerivedDataCache, Display, TEXT("Putting %u records (containing values of size 12 bytes): %.2f s"), NumRecords, AutoTimer.GetTime());
+	}
+
+	return true;
+}
+
 } // UE::DerivedData
 
 #endif // #if WITH_DEV_AUTOMATION_TESTS && WITH_HTTP_DDC_BACKEND
