Change the way ZenCacheStore configuration and presentation is implemented to align with the pattern used in HttpCacheStore so that we can use a ServerID parameter to refer to a configuration in the StorageServers section of the config while still retaining the ability to override values as needed.  No config changes at this point, those will happen later.

diff --git a/Engine/Source/Developer/DerivedDataCache/Private/ZenCacheStore.cpp b/Engine/Source/Developer/DerivedDataCache/Private/ZenCacheStore.cpp
--- a/Engine/Source/Developer/DerivedDataCache/Private/ZenCacheStore.cpp
+++ b/Engine/Source/Developer/DerivedDataCache/Private/ZenCacheStore.cpp
@@ -66,6 +66,24 @@ void ForEachBatch(const int32 BatchSize, const int32 TotalCount, T&& Fn)
 	}
 }
 
+struct FZenCacheStoreParams
+{
+	FString Name;
+	FString Host;
+	FString Namespace;
+	FString Sandbox;
+	int32 MaxBatchPutKB = 1024;
+	int32 RecordBatchSize = 8;
+	int32 ChunksBatchSize = 8;
+	float DeactivateAtMs = -1.0f;
+	TOptional<bool> bLocal;
+	TOptional<bool> bRemote;
+	bool bFlush = false;
+	bool bReadOnly = false;
+
+	void Parse(const TCHAR* NodeName, const TCHAR* Config);
+};
+
 /**
  * Backend for a HTTP based caching service (Zen)
  */
@@ -80,17 +98,12 @@ public:
 	 * @param Namespace		Namespace to use.
 	 */
 	FZenCacheStore(
-		const TCHAR* ServiceUrl,
-		const TCHAR* Namespace,
-		const TCHAR* Name,
-		const TCHAR* Config,
+		const FZenCacheStoreParams& InParams,
 		ICacheStoreOwner* Owner);
 
 	FZenCacheStore(
 		UE::Zen::FServiceSettings&& Settings,
-		const TCHAR* Namespace,
-		const TCHAR* Name,
-		const TCHAR* Config,
+		const FZenCacheStoreParams& InParams,
 		ICacheStoreOwner* Owner);
 
 	~FZenCacheStore() final;
@@ -144,7 +157,7 @@ public:
 	const Zen::FZenServiceInstance& GetServiceInstance() const { return ZenService.GetInstance(); }
 
 private:
-	void Initialize(const TCHAR* Namespace, const TCHAR* Name, const TCHAR* Config);
+	void Initialize(const FZenCacheStoreParams& Params);
 
 	bool IsServiceReady();
 
@@ -206,7 +219,7 @@ private:
 	bool bIsLocalConnection = false;
 	bool bTryEvaluatePerformance = false;
 	std::atomic<bool> bDeactivatedForPerformance = false;
-	int32 BatchPutMaxBytes = 1024*1024;
+	int32 MaxBatchPutKB = 1024;
 	int32 CacheRecordBatchSize = 8;
 	int32 CacheChunksBatchSize = 8;
 	FBackendDebugOptions DebugOptions;
@@ -412,7 +425,7 @@ private:
 			RecordSize += Value.GetData().GetCompressedSize();
 		}
 		BatchSize += RecordSize;
-		if (BatchSize > CacheStore.BatchPutMaxBytes)
+		if (BatchSize > CacheStore.MaxBatchPutKB*1024)
 		{
 			BatchSize = RecordSize;
 			return EBatchView::NewBatch;
@@ -815,7 +828,7 @@ private:
 	{
 		uint64 ValueSize = sizeof(FCacheKey) + NextRequest.Value.GetData().GetCompressedSize();
 		BatchSize += ValueSize;
-		if (BatchSize > CacheStore.BatchPutMaxBytes)
+		if (BatchSize > CacheStore.MaxBatchPutKB*1024)
 		{
 			BatchSize = ValueSize;
 			return EBatchView::NewBatch;
@@ -1547,29 +1560,24 @@ private:
 };
 
 FZenCacheStore::FZenCacheStore(
-	const TCHAR* InServiceUrl,
-	const TCHAR* InNamespace,
-	const TCHAR* InName,
-	const TCHAR* InConfig,
+	const FZenCacheStoreParams& InParams,
 	ICacheStoreOwner* InStoreOwner)
-	: ZenService(InServiceUrl)
+	: ZenService(*InParams.Host)
 	, StoreOwner(InStoreOwner)
 	, PerformanceEvaluationRequestOwner(EPriority::Low)
 {
-	Initialize(InNamespace, InName, InConfig);
+	Initialize(InParams);
 }
 
 FZenCacheStore::FZenCacheStore(
 	UE::Zen::FServiceSettings&& InSettings,
-	const TCHAR* InNamespace,
-	const TCHAR* InName,
-	const TCHAR* InConfig,
+	const FZenCacheStoreParams& InParams,
 	ICacheStoreOwner* InStoreOwner)
 	: ZenService(MoveTemp(InSettings))
 	, StoreOwner(InStoreOwner)
 	, PerformanceEvaluationRequestOwner(EPriority::Low)
 {
-	Initialize(InNamespace, InName, InConfig);
+	Initialize(InParams);
 }
 
 FZenCacheStore::~FZenCacheStore()
@@ -1589,13 +1597,10 @@ FZenCacheStore::~FZenCacheStore()
 	}
 }
 
-void FZenCacheStore::Initialize(
-	const TCHAR* InNamespace,
-	const TCHAR* InName,
-	const TCHAR* InConfig)
+void FZenCacheStore::Initialize(const FZenCacheStoreParams& Params)
 {
 	LastPerformanceEvaluationTicks.store(FDateTime::UtcNow().GetTicks(), std::memory_order_relaxed);
-	Namespace = InNamespace;
+	Namespace = Params.Namespace;
 	if (IsServiceReady())
 	{
 		RpcUri << ZenService.GetInstance().GetURL() << ANSITEXTVIEW("/z$/$rpc");
@@ -1620,30 +1625,28 @@ void FZenCacheStore::Initialize(
 
 		if (StoreOwner)
 		{
-			bool bReadOnly = false;
-			FParse::Bool(InConfig, TEXT("ReadOnly="), bReadOnly);
-
 			// Default to locally launched service getting the Local cache store flag.  Can be overridden by explicit value in config.
-			bool bLocal = bIsLocalConnection;
-			FParse::Bool(InConfig, TEXT("Local="), bLocal);
+			bool bLocal = Params.bLocal.Get(bIsLocalConnection);
 
 			// Default to non-locally launched service getting the Remote cache store flag.  Can be overridden by explicit value in config.
 			// In the future this could be extended to allow the Remote flag by default (even for locally launched instances) if they have upstreams configured.
-			bool bRemote = !bIsLocalConnection;
-			FParse::Bool(InConfig, TEXT("Remote="), bRemote);
+			bool bRemote = Params.bRemote.Get(!bIsLocalConnection);
 
-			FParse::Value(InConfig, TEXT("DeactivateAt="), DeactivateAtMs);
+			DeactivateAtMs = Params.DeactivateAtMs;
 
 			ECacheStoreFlags Flags = ECacheStoreFlags::Query;
-			Flags |= bReadOnly ? ECacheStoreFlags::None : ECacheStoreFlags::Store;
+			Flags |= Params.bReadOnly ? ECacheStoreFlags::None : ECacheStoreFlags::Store;
 			Flags |= bLocal ? ECacheStoreFlags::Local : ECacheStoreFlags::None;
 			Flags |= bRemote ? ECacheStoreFlags::Remote : ECacheStoreFlags::None;
 
 			OperationalFlags = Flags;
 
 			StoreOwner->Add(this, Flags);
-			StoreStats = StoreOwner->CreateStats(this, Flags, TEXT("Zen"), InName, ZenService.GetInstance().GetURL());
+			TStringBuilder<256> Path(InPlace, ZenService.GetInstance().GetURL(), TEXTVIEW(" ("), Namespace, TEXTVIEW(")"));
+			StoreStats = StoreOwner->CreateStats(this, Flags, TEXT("Zen"), *Params.Name, Path);
 			bTryEvaluatePerformance = !GIsBuildMachine && (StoreStats != nullptr) && (DeactivateAtMs > 0.0f);
+
+			StoreStats->SetAttribute(TEXTVIEW("Namespace"), Namespace);
 		}
 
 		// Issue a request for stats as it will be fetched asynchronously and issuing now makes them available sooner for future callers.
@@ -1651,9 +1654,9 @@ void FZenCacheStore::Initialize(
 		ZenService.GetInstance().GetCacheStats(ZenStats);
 	}
 
-	GConfig->GetInt(TEXT("Zen"), TEXT("BatchPutMaxBytes"), BatchPutMaxBytes, GEngineIni);
-	GConfig->GetInt(TEXT("Zen"), TEXT("CacheRecordBatchSize"), CacheRecordBatchSize, GEngineIni);
-	GConfig->GetInt(TEXT("Zen"), TEXT("CacheChunksBatchSize"), CacheChunksBatchSize, GEngineIni);
+	MaxBatchPutKB = Params.MaxBatchPutKB;
+	CacheRecordBatchSize = Params.RecordBatchSize;
+	CacheChunksBatchSize = Params.ChunksBatchSize;
 }
 
 bool FZenCacheStore::IsServiceReady()
@@ -1864,63 +1867,106 @@ void FZenCacheStore::GetChunks(
 	GetChunksOp->IssueRequests();
 }
 
-ILegacyCacheStore* CreateZenCacheStore(const TCHAR* NodeName, const TCHAR* Config, ICacheStoreOwner* Owner)
+void FZenCacheStoreParams::Parse(const TCHAR* NodeName, const TCHAR* Config)
 {
-	FString ServiceUrl;
-	FParse::Value(Config, TEXT("Host="), ServiceUrl);
+	Name = NodeName;
+
+	if (FString ServerId; FParse::Value(Config, TEXT("ServerID="), ServerId))
+	{
+		FString ServerEntry;
+		const TCHAR* ServerSection = TEXT("StorageServers");
+		if (GConfig->GetString(ServerSection, *ServerId, ServerEntry, GEngineIni))
+		{
+			Parse(NodeName, *ServerEntry);
+		}
+		else
+		{
+			UE_LOG(LogDerivedDataCache, Warning, TEXT("%s: Using ServerID=%s which was not found in [%s]"), NodeName, *ServerId, ServerSection);
+		}
+	}
+
+	FParse::Value(Config, TEXT("Host="), Host);
 
 	FString OverrideName;
 	if (FParse::Value(Config, TEXT("EnvHostOverride="), OverrideName))
 	{
-		FString ServiceUrlEnv = FPlatformMisc::GetEnvironmentVariable(*OverrideName);
-		if (!ServiceUrlEnv.IsEmpty())
+		FString HostEnv = FPlatformMisc::GetEnvironmentVariable(*OverrideName);
+		if (!HostEnv.IsEmpty())
 		{
-			ServiceUrl = ServiceUrlEnv;
-			UE_LOG(LogDerivedDataCache, Log, TEXT("%s: Found environment override for Host %s=%s"), NodeName, *OverrideName, *ServiceUrl);
+			Host = HostEnv;
+			UE_LOG(LogDerivedDataCache, Log, TEXT("%s: Found environment override for Host %s=%s"), NodeName, *OverrideName, *Host);
 		}
 	}
 
 	if (FParse::Value(Config, TEXT("CommandLineHostOverride="), OverrideName))
 	{
-		if (FParse::Value(FCommandLine::Get(), *(OverrideName + TEXT("=")), ServiceUrl))
+		if (FParse::Value(FCommandLine::Get(), *(OverrideName + TEXT("=")), Host))
 		{
-			UE_LOG(LogDerivedDataCache, Log, TEXT("%s: Found command line override for Host %s=%s"), NodeName, *OverrideName, *ServiceUrl);
+			UE_LOG(LogDerivedDataCache, Log, TEXT("%s: Found command line override for Host %s=%s"), NodeName, *OverrideName, *Host);
 		}
 	}
 
-	if (ServiceUrl == TEXT("None"))
+	FParse::Value(Config, TEXT("Namespace="), Namespace);
+	FParse::Value(Config, TEXT("StructuredNamespace="), Namespace);
+
+	FParse::Bool(Config, TEXT("ReadOnly="), bReadOnly);
+
+	// Sandbox and flush configuration for use in Cold/Warm type use cases
+	FParse::Value(Config, TEXT("Sandbox="), Sandbox);
+	FParse::Bool(Config, TEXT("Flush="), bFlush);
+
+	// Performance deactivation
+	FParse::Value(Config, TEXT("DeactivateAt="), DeactivateAtMs);
+
+	// Explicit local and remote configuration
+	if (bool bExplicitLocal = false; FParse::Bool(Config, TEXT("Local="), bExplicitLocal))
 	{
-		UE_LOG(LogDerivedDataCache, Log, TEXT("Disabling %s data cache - host set to 'None'."), NodeName);
+		bLocal = bExplicitLocal;
+	}
+
+	if (bool bExplicitRemote = false; FParse::Bool(Config, TEXT("Remote="), bExplicitRemote))
+	{
+		bRemote = bExplicitRemote;
+	}
+
+	// Request batch fracturing configuration
+	FParse::Value(Config, TEXT("MaxBatchPutKB="), MaxBatchPutKB);
+	FParse::Value(Config, TEXT("RecordBatchSize="), RecordBatchSize);
+	FParse::Value(Config, TEXT("ChunksBatchSize="), ChunksBatchSize);
+}
+
+ILegacyCacheStore* CreateZenCacheStore(const TCHAR* NodeName, const TCHAR* Config, ICacheStoreOwner* Owner)
+{
+	FZenCacheStoreParams Params;
+	Params.Parse(NodeName, Config);
+
+	if (Params.Host == TEXTVIEW("None"))
+	{
+		UE_LOG(LogDerivedDataCache, Log, TEXT("%s: Disabled because Host is set to 'None'"), NodeName);
 		return nullptr;
 	}
 
-	FString Namespace;
-	if (!FParse::Value(Config, TEXT("StructuredNamespace="), Namespace) && !FParse::Value(Config, TEXT("Namespace="), Namespace))
+	if (Params.Namespace.IsEmpty())
 	{
-		Namespace = FApp::GetProjectName();
-		UE_LOG(LogDerivedDataCache, Warning, TEXT("%s: Missing required parameter 'Namespace', falling back to '%s'"), NodeName, *Namespace);
+		Params.Namespace = FApp::GetProjectName();
+		UE_LOG(LogDerivedDataCache, Warning, TEXT("%s: Missing required parameter 'Namespace', falling back to '%s'"), NodeName, *Params.Namespace);
 	}
 
-	FString Sandbox;
-	FParse::Value(Config, TEXT("Sandbox="), Sandbox);
-	bool bHasSandbox = !Sandbox.IsEmpty();
+	bool bHasSandbox = !Params.Sandbox.IsEmpty();
 	bool bUseLocalDataCachePathOverrides = !bHasSandbox;
 
 	FString CachePathOverride;
-	if (bUseLocalDataCachePathOverrides && UE::Zen::Private::IsLocalAutoLaunched(ServiceUrl) && UE::Zen::Private::GetLocalDataCachePathOverride(CachePathOverride))
+	if (bUseLocalDataCachePathOverrides && UE::Zen::Private::IsLocalAutoLaunched(Params.Host) && UE::Zen::Private::GetLocalDataCachePathOverride(CachePathOverride))
 	{
 		if (CachePathOverride == TEXT("None"))
 		{
-			UE_LOG(LogDerivedDataCache, Log, TEXT("Disabling %s data cache - path set to 'None'."), NodeName);
+			UE_LOG(LogDerivedDataCache, Log, TEXT("%s: Disabled because path is set to 'None'"), NodeName);
 			return nullptr;
 		}
 	}
 
 	TUniquePtr<FZenCacheStore> Backend;
 
-	bool bFlush = false;
-	FParse::Bool(Config, TEXT("Flush="), bFlush);
-
 	if (bHasSandbox)
 	{
 		Zen::FServiceSettings DefaultServiceSettings;
@@ -1944,7 +1990,7 @@ ILegacyCacheStore* CreateZenCacheStore(const TCHAR* NodeName, const TCHAR* Confi
 
 		FPaths::NormalizeDirectoryName(AutoLaunchSettings.DataPath);
 		AutoLaunchSettings.DataPath += TEXT("_");
-		AutoLaunchSettings.DataPath += Sandbox;
+		AutoLaunchSettings.DataPath += Params.Sandbox;
 		AutoLaunchSettings.bIsDefaultSharedRunContext = false;
 
 		// The unique local instances will always limit process lifetime for now to avoid accumulating many of them
@@ -1952,7 +1998,7 @@ ILegacyCacheStore* CreateZenCacheStore(const TCHAR* NodeName, const TCHAR* Confi
 
 		// Flush the cache if requested.
 		uint32 MultiprocessId = UE::GetMultiprocessId();
-		if (bFlush && (MultiprocessId == 0))
+		if (Params.bFlush && (MultiprocessId == 0))
 		{
 			bool bStopped = true;
 			if (UE::Zen::IsLocalServiceRunning(*AutoLaunchSettings.DataPath))
@@ -1970,11 +2016,11 @@ ILegacyCacheStore* CreateZenCacheStore(const TCHAR* NodeName, const TCHAR* Confi
 			}
 		}
 
-		Backend = MakeUnique<FZenCacheStore>(MoveTemp(ServiceSettings), *Namespace, NodeName, Config, Owner);
+		Backend = MakeUnique<FZenCacheStore>(MoveTemp(ServiceSettings), Params, Owner);
 	}
 	else
 	{
-		Backend = MakeUnique<FZenCacheStore>(*ServiceUrl, *Namespace, NodeName, Config, Owner);
+		Backend = MakeUnique<FZenCacheStore>(Params, Owner);
 	}
 
 	if (!Backend->IsUsable())
