Don't attempt immediate retry when a 502 error happens during communication with the UE Cloud DDC server.

diff --git a/Engine/Source/Developer/DerivedDataCache/Private/HttpCacheStore.cpp b/Engine/Source/Developer/DerivedDataCache/Private/HttpCacheStore.cpp
--- a/Engine/Source/Developer/DerivedDataCache/Private/HttpCacheStore.cpp
+++ b/Engine/Source/Developer/DerivedDataCache/Private/HttpCacheStore.cpp
@@ -631,9 +631,8 @@ private:
 
 		// Make a new attempt if the response status code is any of:
 		// 429 - Too many requests
-		// 502 - Bad gateway
 		int32 StatusCode = LocalResponse.GetStatusCode();
-		if ((StatusCode == 429) || (StatusCode == 502))
+		if (StatusCode == 429)
 		{
 			return true;
 		}
