Add EQueuedWorkFlags::DoNotRunInsideBusyWait to tasks started by DDC so that as long running tasks, they don't get picked up by arbitrary busy waits.

diff --git a/Engine/Source/Developer/DerivedDataCache/Private/DerivedDataCache.cpp b/Engine/Source/Developer/DerivedDataCache/Private/DerivedDataCache.cpp
--- a/Engine/Source/Developer/DerivedDataCache/Private/DerivedDataCache.cpp
+++ b/Engine/Source/Developer/DerivedDataCache/Private/DerivedDataCache.cpp
@@ -600,7 +600,7 @@ public:
 		UE_LOG(LogDerivedDataCache, VeryVerbose, TEXT("GetSynchronous %s from '%s'"), *CacheKey, *DataDeriver->GetDebugContextString());
 		FAsyncTask<FBuildAsyncWorker> PendingTask(Backend, DataDeriver, *CacheKey, DataDeriver->GetDebugContextString(), true);
 		AddToAsyncCompletionCounter(1);
-		PendingTask.StartSynchronousTask();
+		PendingTask.StartSynchronousTask(EQueuedWorkPriority::Normal, EQueuedWorkFlags::DoNotRunInsideBusyWait);
 		OutData = TArray<uint8>(MoveTemp(PendingTask.GetTask().Data));
 		if (bDataWasBuilt)
 		{
@@ -623,11 +623,11 @@ public:
 		AddToAsyncCompletionCounter(1);
 		if (!bSync)
 		{
-			AsyncTask->StartBackgroundTask(DataDeriver->GetCustomThreadPool());
+			AsyncTask->StartBackgroundTask(DataDeriver->GetCustomThreadPool(), EQueuedWorkPriority::Normal, EQueuedWorkFlags::DoNotRunInsideBusyWait);
 		}
 		else
 		{
-			AsyncTask->StartSynchronousTask();
+			AsyncTask->StartSynchronousTask(EQueuedWorkPriority::Normal, EQueuedWorkFlags::DoNotRunInsideBusyWait);
 		}
 		// Must return a valid handle
 		check(Handle != 0);
@@ -710,7 +710,7 @@ public:
 		UE_LOG(LogDerivedDataCache, VeryVerbose, TEXT("GetSynchronous %s from '%.*s'"), CacheKey, DebugContext.Len(), DebugContext.GetData());
 		FAsyncTask<FBuildAsyncWorker> PendingTask(Backend, nullptr, CacheKey, DebugContext, true);
 		AddToAsyncCompletionCounter(1);
-		PendingTask.StartSynchronousTask();
+		PendingTask.StartSynchronousTask(EQueuedWorkPriority::Normal, EQueuedWorkFlags::DoNotRunInsideBusyWait);
 		OutData = DataType(MoveTemp(PendingTask.GetTask().Data));
 		return PendingTask.GetTask().bSuccess;
 	}
@@ -736,7 +736,7 @@ public:
 		PendingTasks.Add(Handle, AsyncTask);
 		AddToAsyncCompletionCounter(1);
 		// This request is I/O only, doesn't do any processing, send it to the I/O only thread-pool to avoid wasting worker threads on long I/O waits.
-		AsyncTask->StartBackgroundTask(GCacheThreadPool);
+		AsyncTask->StartBackgroundTask(GCacheThreadPool, EQueuedWorkPriority::Normal, EQueuedWorkFlags::DoNotRunInsideBusyWait);
 		return Handle;
 	}
 
diff --git a/Engine/Source/Developer/DerivedDataCache/Private/DerivedDataThreadPoolTask.cpp b/Engine/Source/Developer/DerivedDataCache/Private/DerivedDataThreadPoolTask.cpp
--- a/Engine/Source/Developer/DerivedDataCache/Private/DerivedDataThreadPoolTask.cpp
+++ b/Engine/Source/Developer/DerivedDataCache/Private/DerivedDataThreadPoolTask.cpp
@@ -35,6 +35,11 @@ private:
 
 	void DoThreadedWork() final { Execute(); }
 	void Abandon() final { Execute(); }
+	EQueuedWorkFlags GetQueuedWorkFlags() const final
+	{
+		return EQueuedWorkFlags::DoNotRunInsideBusyWait;
+	}
+	
 	int64 GetRequiredMemory() const final
 	{
 		return int64(MemoryEstimate);
