Fix a typo in a log message.

diff --git a/Engine/Source/Developer/DerivedDataCache/Private/ZenCacheStore.cpp b/Engine/Source/Developer/DerivedDataCache/Private/ZenCacheStore.cpp
--- a/Engine/Source/Developer/DerivedDataCache/Private/ZenCacheStore.cpp
+++ b/Engine/Source/Developer/DerivedDataCache/Private/ZenCacheStore.cpp
@@ -1974,7 +1974,7 @@ ILegacyCacheStore* CreateZenCacheStore(const TCHAR* NodeName, const TCHAR* Confi
 
 		if (!DefaultServiceSettings.IsAutoLaunch())
 		{
-			UE_LOG(LogDerivedDataCache, Warning, TEXT("%s: Attempting to use a sandbox when there is no default autolaunch configured to interhit settings from.  Cache will be disabled."), NodeName);
+			UE_LOG(LogDerivedDataCache, Warning, TEXT("%s: Attempting to use a sandbox when there is no default autolaunch configured to inherit settings from.  Cache will be disabled."), NodeName);
 			return nullptr;
 		}
 
