DDC: Eliminated most blocking during get requests in HttpCacheStore

This restores the changes from 29016116 that were temporarily reverted.

#rnx

diff --git a/Engine/Source/Developer/DerivedDataCache/Private/HttpCacheStore.cpp b/Engine/Source/Developer/DerivedDataCache/Private/HttpCacheStore.cpp
--- a/Engine/Source/Developer/DerivedDataCache/Private/HttpCacheStore.cpp
+++ b/Engine/Source/Developer/DerivedDataCache/Private/HttpCacheStore.cpp
@@ -475,7 +475,6 @@ private:
 	class FHttpOperation;
 
 	FHttpCacheStoreRequestQueue& PickRequestQueue(EOperationCategory Category);
-	TUniquePtr<FHttpOperation> WaitForHttpOperation(EOperationCategory Category);
 
 	/** Invokes the callback when an operation is available, or with null if canceled. */
 	void WaitForHttpOperationAsync(IRequestOwner& Owner, EOperationCategory Category, TUniqueFunction<void (TUniquePtr<FHttpOperation>&&)>&& OnOperation);
@@ -1298,9 +1297,8 @@ void FHttpCacheStore::FGetRecordOp::GetRecordOnly(const FCacheKey& InKey, const
 	OnRecordComplete = MoveTemp(InOnComplete);
 	RequestStats.Bucket = Key.Bucket;
 
-	TUniquePtr<FHttpOperation> Operation = CacheStore.WaitForHttpOperation(EOperationCategory::Get);
-	TRefCountPtr Self(this);
 	RequestTimer.Stop();
+	CacheStore.WaitForHttpOperationAsync(Owner, EOperationCategory::Get, [Self = TRefCountPtr(this)](TUniquePtr<FHttpOperation>&& Operation)
 	{
 		if (UNLIKELY(!Operation))
 		{
@@ -1322,7 +1320,7 @@ void FHttpCacheStore::FGetRecordOp::GetRecordOnly(const FCacheKey& InKey, const
 			Operation->GetStats(Self->RequestStats);
 			Self->EndGetRef(MoveTemp(Operation));
 		});
-	}
+	});
 }
 
 void FHttpCacheStore::FGetRecordOp::EndGetRef(TUniquePtr<FHttpOperation> Operation)
@@ -1540,11 +1538,10 @@ void FHttpCacheStore::FGetRecordOp::GetValues(TConstArrayView<FValueWithId> Valu
 			continue;
 		}
 
-		TUniquePtr<FHttpOperation> Operation = CacheStore.WaitForHttpOperation(EOperationCategory::Get);
-		TRefCountPtr Self(this);
+		CacheStore.WaitForHttpOperationAsync(Owner, EOperationCategory::Get, [Self = TRefCountPtr(this), SharedOnComplete, Value](TUniquePtr<FHttpOperation>&& Operation)
 		{
 			Self->BeginGetValue(MoveTemp(Operation), Value, SharedOnComplete);
-		}
+		});
 	}
 }
 
@@ -1650,14 +1647,13 @@ void FHttpCacheStore::FGetRecordOp::GetValuesExist(TConstArrayView<FValueWithId>
 	}
 
 	FRequestTimer RequestTimer(RequestStats);
+	RequestTimer.Stop();
 
 	FRequestBarrier Barrier(Owner);
-	TUniquePtr<FHttpOperation> Operation = CacheStore.WaitForHttpOperation(EOperationCategory::Get);
-	TRefCountPtr Self(this);
-	RequestTimer.Stop();
+	CacheStore.WaitForHttpOperationAsync(Owner, EOperationCategory::Get, [Self = TRefCountPtr(this), Values = MoveTemp(QueryValues), OnComplete = MoveTemp(OnComplete)](TUniquePtr<FHttpOperation>&& Operation) mutable
 	{
-		Self->BeginGetValuesExist(MoveTemp(Operation), MoveTemp(QueryValues), MoveTemp(OnComplete));
-	}
+		Self->BeginGetValuesExist(MoveTemp(Operation), MoveTemp(Values), MoveTemp(OnComplete));
+	});
 }
 
 void FHttpCacheStore::FGetRecordOp::BeginGetValuesExist(TUniquePtr<FHttpOperation>&& Operation, TArray<FValueWithId>&& Values, FOnValueComplete&& OnComplete)
@@ -1822,12 +1818,11 @@ void FHttpCacheStore::FGetValueOp::Get(const FCacheKey& InKey, ECachePolicy InPo
 	Policy = InPolicy;
 	OnComplete = MoveTemp(InOnComplete);
 
-	TUniquePtr<FHttpOperation> Operation = CacheStore.WaitForHttpOperation(EOperationCategory::Get);
-	TRefCountPtr Self(this);
 	RequestTimer.Stop();
+	CacheStore.WaitForHttpOperationAsync(Owner, EOperationCategory::Get, [Self = TRefCountPtr(this)](TUniquePtr<FHttpOperation>&& Operation)
 	{
 		Self->BeginGetRef(MoveTemp(Operation));
-	}
+	});
 }
 
 void FHttpCacheStore::FGetValueOp::BeginGetRef(TUniquePtr<FHttpOperation>&& Operation)
@@ -2047,12 +2042,11 @@ void FHttpCacheStore::FExistsBatchOp::Exists(TConstArrayView<FCacheGetValueReque
 	BodyWriter.EndObject();
 	FCbFieldIterator Body = BodyWriter.Save();
 
-	TUniquePtr<FHttpOperation> Operation = CacheStore.WaitForHttpOperation(EOperationCategory::Get);
-	TRefCountPtr Self(this);
 	RequestTimer.Stop();
+	CacheStore.WaitForHttpOperationAsync(Owner, EOperationCategory::Get, [Self = TRefCountPtr(this), Body = MoveTemp(Body)](TUniquePtr<FHttpOperation>&& Operation) mutable
 	{
 		Self->BeginExists(MoveTemp(Operation), MoveTemp(Body));
-	}
+	});
 }
 
 void FHttpCacheStore::FExistsBatchOp::BeginExists(TUniquePtr<FHttpOperation>&& Operation, FCbFieldIterator&& Body)
@@ -2583,34 +2577,6 @@ FHttpCacheStoreRequestQueue& FHttpCacheStore::PickRequestQueue(EOperationCategor
 	}
 }
 
-TUniquePtr<FHttpCacheStore::FHttpOperation> FHttpCacheStore::WaitForHttpOperation(EOperationCategory Category)
-{
-	if (Access && RefreshAccessTokenTime > 0.0 && RefreshAccessTokenTime < FPlatformTime::Seconds())
-	{
-		AcquireAccessToken();
-	}
-
-	THttpUniquePtr<IHttpRequest> Request;
-
-	{
-		FHttpRequestParams Params;
-		FRequestOwner BlockingOwner(EPriority::Blocking);
-		FHttpCacheStoreRequestQueue& RequestQueue = PickRequestQueue(Category);
-		RequestQueue.CreateRequestAsync(BlockingOwner, Params, [&Request](THttpUniquePtr<IHttpRequest>&& AsyncRequest)
-		{
-			Request = MoveTemp(AsyncRequest);
-		});
-		BlockingOwner.Wait();
-	}
-
-	if (Access)
-	{
-		Request->AddHeader(ANSITEXTVIEW("Authorization"), WriteToAnsiString<1024>(*Access));
-	}
-
-	return MakeUnique<FHttpOperation>(MoveTemp(Request));
-}
-
 void FHttpCacheStore::WaitForHttpOperationAsync(IRequestOwner& Owner, EOperationCategory Category, TUniqueFunction<void (TUniquePtr<FHttpOperation>&&)>&& OnOperation)
 {
 	WaitForHttpRequestAsync(Owner, Category, [this, OnOperation = MoveTemp(OnOperation)](THttpUniquePtr<IHttpRequest>&& Request)
