If we find an active parsed note in the ddc graph when creating it, return it instead of returning nullptr.

The main cooker process sets up the environment variable UE-ZenSharedDataCacheHost as part of it's initialization so the ZenShared instance is created by the worker processes when they initialize before the redirection for the file share is detected.

In the main cooker process that env-variable is not set so it skips creating the ZenShared instance until it finds the redirection.

Renamed ICacheStoreGraph::Create -> ICacheStoreGraph::FindOrCreate to better reflect the functionality.

diff --git a/Engine/Source/Developer/DerivedDataCache/Private/DerivedDataBackends.cpp b/Engine/Source/Developer/DerivedDataCache/Private/DerivedDataBackends.cpp
--- a/Engine/Source/Developer/DerivedDataCache/Private/DerivedDataBackends.cpp
+++ b/Engine/Source/Developer/DerivedDataCache/Private/DerivedDataBackends.cpp
@@ -916,7 +916,7 @@ private:
 		Hierarchy->LegacyResourceStats(OutStats);
 	}
 
-	ILegacyCacheStore* Create(const TCHAR* Name) final
+	ILegacyCacheStore* FindOrCreate(const TCHAR* Name) final
 	{
 		if (!ActiveParsedNodeMap)
 		{
@@ -926,7 +926,7 @@ private:
 
 		if (const FParsedNode* ParsedNode = ActiveParsedNodeMap->Find(Name))
 		{
-			return nullptr;
+			return *ParsedNode;
 		}
 
 		if (!ParseNode(Name, GEngineIni, *GraphName, *ActiveParsedNodeMap))
diff --git a/Engine/Source/Developer/DerivedDataCache/Private/DerivedDataCacheStore.h b/Engine/Source/Developer/DerivedDataCache/Private/DerivedDataCacheStore.h
--- a/Engine/Source/Developer/DerivedDataCache/Private/DerivedDataCacheStore.h
+++ b/Engine/Source/Developer/DerivedDataCache/Private/DerivedDataCacheStore.h
@@ -205,7 +205,7 @@ public:
 class ICacheStoreGraph
 {
 public:
-	virtual ILegacyCacheStore* Create(const TCHAR* Name) = 0;
+	virtual ILegacyCacheStore* FindOrCreate(const TCHAR* Name) = 0;
 };
 
 enum class ECacheStoreStatusCode : uint8
diff --git a/Engine/Source/Developer/DerivedDataCache/Private/FileSystemCacheStore.cpp b/Engine/Source/Developer/DerivedDataCache/Private/FileSystemCacheStore.cpp
--- a/Engine/Source/Developer/DerivedDataCache/Private/FileSystemCacheStore.cpp
+++ b/Engine/Source/Developer/DerivedDataCache/Private/FileSystemCacheStore.cpp
@@ -1536,10 +1536,19 @@ ILegacyCacheStore* FFileSystemCacheStore::TryRedirection(
 				FString TargetName;
 				if (FParse::Value(*RedirectionData, TEXT("Target="), TargetName))
 				{
-					if (ILegacyCacheStore* RedirectedStore = Graph->Create(*TargetName))
+					UE_LOG(LogDerivedDataCache, Log, TEXT("%s: Found redirection to '%s'"),
+						*Params.CachePath, *TargetName);
+					if (ILegacyCacheStore* RedirectedStore = Graph->FindOrCreate(*TargetName))
 					{
+						UE_LOG(LogDerivedDataCache, Log, TEXT("%s: Successfully redirected to '%s'"),
+							*Params.CachePath, *TargetName);
 						return RedirectedStore;
 					}
+					else
+					{
+						UE_LOG(LogDerivedDataCache, Warning, TEXT("%s: Failed to redirect to '%s'"),
+							*Params.CachePath, *TargetName);
+					}
 				}
 			}
 		}
