Use the low speed limit with zen cache store client, but with really tolerant threshold (only require 30 bytes over 60 seconds).

diff --git a/Engine/Source/Developer/DerivedDataCache/Private/ZenCacheStore.cpp b/Engine/Source/Developer/DerivedDataCache/Private/ZenCacheStore.cpp
--- a/Engine/Source/Developer/DerivedDataCache/Private/ZenCacheStore.cpp
+++ b/Engine/Source/Developer/DerivedDataCache/Private/ZenCacheStore.cpp
@@ -1618,6 +1618,8 @@ void FZenCacheStore::Initialize(const FZenCacheStoreParams& Params)
 		ClientParams.Version = EHttpVersion::V2;
 		ClientParams.MaxRequests = RequestPoolSize + RequestPoolOverflowSize;
 		ClientParams.MinRequests = RequestPoolSize;
+		ClientParams.LowSpeedLimit = 30;
+		ClientParams.LowSpeedTime = 60;
 		RequestQueue = FHttpRequestQueue(*ConnectionPool, ClientParams);
 
 		bIsUsable = true;
