Avoid logging Display verbosity messages when a DDC request to HttpCacheStore is canceled.  Also rename ExpectedErrorCodes to ExpectedStatusCodes to avoid conflating the distinct meanings of ErrorCodes and StatusCodes.

diff --git a/Engine/Source/Developer/DerivedDataCache/Private/HttpCacheStore.cpp b/Engine/Source/Developer/DerivedDataCache/Private/HttpCacheStore.cpp
--- a/Engine/Source/Developer/DerivedDataCache/Private/HttpCacheStore.cpp
+++ b/Engine/Source/Developer/DerivedDataCache/Private/HttpCacheStore.cpp
@@ -545,7 +545,7 @@ public:
 	void SetBody(const FCompositeBuffer& Body) { Request->SetBody(Body); }
 	void SetContentType(EHttpMediaType Type) { Request->SetContentType(Type); }
 	void AddAcceptType(EHttpMediaType Type) { Request->AddAcceptType(Type); }
-	void SetExpectedErrorCodes(TConstArrayView<int32> Codes) { ExpectedErrorCodes = Codes; }
+	void SetExpectedStatusCodes(TConstArrayView<int32> Codes) { ExpectedStatusCodes = Codes; }
 
 	// Send Request
 
@@ -576,7 +576,7 @@ private:
 	FSharedBuffer ResponseBody;
 	THttpUniquePtr<IHttpRequest> Request;
 	THttpUniquePtr<IHttpResponse> Response;
-	TArray<int32, TInlineAllocator<4>> ExpectedErrorCodes;
+	TArray<int32, TInlineAllocator<4>> ExpectedStatusCodes;
 	uint32 AttemptCount = 0;
 };
 
@@ -644,8 +644,21 @@ private:
 	{
 		if (UE_LOG_ACTIVE(LogDerivedDataCache, Display))
 		{
+			EHttpErrorCode ErrorCode = LocalResponse.GetErrorCode();
 			const int32 StatusCode = LocalResponse.GetStatusCode();
-			const bool bUnexpectedError = !((StatusCode >= 200 && StatusCode < 300) || Operation->ExpectedErrorCodes.Contains(StatusCode));
+			bool bUnexpectedError = false;
+			if (ErrorCode == EHttpErrorCode::None)
+			{
+				bUnexpectedError = !((StatusCode >= 200 && StatusCode < 300) || Operation->ExpectedStatusCodes.Contains(StatusCode));
+			}
+			else if (ErrorCode == EHttpErrorCode::Canceled)
+			{
+				// No logging, this is expected to happen.
+			}
+			else
+			{
+				bUnexpectedError = true;
+			}
 
 			TStringBuilder<80> StatsText;
 			if (bUnexpectedError || UE_LOG_ACTIVE(LogDerivedDataCache, Verbose))
@@ -1315,7 +1328,7 @@ void FHttpCacheStore::FGetRecordOp::GetRecordOnly(const FCacheKey& InKey, const
 		LocalOperation.SetUri(WriteToAnsiString<256>(Self->CacheStore.EffectiveDomain, ANSITEXTVIEW("/api/v1/refs/"), Self->CacheStore.Namespace, '/', Bucket, '/', Self->Key.Hash));
 		LocalOperation.SetMethod(EHttpMethod::Get);
 		LocalOperation.AddAcceptType(EHttpMediaType::CbObject);
-		LocalOperation.SetExpectedErrorCodes({404});
+		LocalOperation.SetExpectedStatusCodes({404});
 
 		LocalOperation.SendAsync(Self->Owner, [Self, Operation = MoveTemp(Operation)]() mutable
 		{
@@ -1563,7 +1576,7 @@ void FHttpCacheStore::FGetRecordOp::BeginGetValue(
 	LocalOperation.SetUri(WriteToAnsiString<256>(CacheStore.EffectiveDomain, ANSITEXTVIEW("/api/v1/compressed-blobs/"), CacheStore.Namespace, '/', Value.GetRawHash()));
 	LocalOperation.SetMethod(EHttpMethod::Get);
 	LocalOperation.AddAcceptType(EHttpMediaType::Any);
-	LocalOperation.SetExpectedErrorCodes({404});
+	LocalOperation.SetExpectedStatusCodes({404});
 	LocalOperation.SendAsync(Owner, [Self = TRefCountPtr(this), Operation = MoveTemp(Operation), OnComplete, Value]
 	{
 		Self->EndGetValue(*Operation, Value, *OnComplete);
@@ -1858,7 +1871,7 @@ void FHttpCacheStore::FGetValueOp::BeginGetRef(TUniquePtr<FHttpOperation>&& Oper
 	{
 		LocalOperation.AddHeader(ANSITEXTVIEW("Accept"), ANSITEXTVIEW("application/x-jupiter-inline"));
 	}
-	LocalOperation.SetExpectedErrorCodes({404});
+	LocalOperation.SetExpectedStatusCodes({404});
 
 	RequestTimer.Stop();
 	LocalOperation.SendAsync(Owner, [Self = TRefCountPtr(this), Operation = MoveTemp(Operation)]() mutable
