DDC: Added a fatal error if the cache is not readable or writable

diff --git a/Engine/Source/Developer/DerivedDataCache/Private/DerivedDataBackends.cpp b/Engine/Source/Developer/DerivedDataCache/Private/DerivedDataBackends.cpp
--- a/Engine/Source/Developer/DerivedDataCache/Private/DerivedDataBackends.cpp
+++ b/Engine/Source/Developer/DerivedDataCache/Private/DerivedDataBackends.cpp
@@ -149,7 +149,7 @@ public:
 			{
 				RootNode = CreateCacheStoreHierarchy(Hierarchy, [this](IMemoryCacheStore*& OutCache) { GetMemoryCache(OutCache); });
 
-				if (!ParseNode(RootName, GEngineIni, *GraphName, ParsedNodes))
+				if (!ParseNode(RootName, GEngineIni, *GraphName, ParsedNodes) || !Hierarchy->HasAllFlags(ECacheStoreFlags::Query | ECacheStoreFlags::Store))
 				{
 					// Destroy any cache stores that have been created.
 					delete RootNode;
@@ -167,7 +167,7 @@ public:
 					bVerifyFound = false;
 					bVerifyFix = false;
 					UE_LOG(LogDerivedDataCache, Warning,
-						TEXT("Unable to create cache graph '%s'. Reverting to the default graph."), *GraphName);
+						TEXT("Unable to create or use cache graph '%s'. Reverting to the default graph."), *GraphName);
 				}
 			}
 
@@ -177,7 +177,7 @@ public:
 				GraphName = FApp::IsEngineInstalled() ? TEXT("InstalledDerivedDataBackendGraph") : TEXT("DerivedDataBackendGraph");
 				RootNode = CreateCacheStoreHierarchy(Hierarchy, [this](IMemoryCacheStore*& OutCache) { GetMemoryCache(OutCache); });
 
-				if (!ParseNode(RootName, GEngineIni, *GraphName, ParsedNodes))
+				if (!ParseNode(RootName, GEngineIni, *GraphName, ParsedNodes) || !Hierarchy->HasAllFlags(ECacheStoreFlags::Query | ECacheStoreFlags::Store))
 				{
 					FString Entry;
 					if (!GConfig->DoesSectionExist(*GraphName, GEngineIni))
@@ -195,8 +195,10 @@ public:
 					else
 					{
 						UE_LOG(LogDerivedDataCache, Fatal,
-							TEXT("Unable to create default cache graph '%s' because no cache store was available."),
-							*GraphName);
+							TEXT("Unable to use default cache graph '%s' because there are no %s nodes available."),
+							*GraphName,
+							Hierarchy->HasAllFlags(ECacheStoreFlags::Query) ? TEXT("writable") :
+							Hierarchy->HasAllFlags(ECacheStoreFlags::Store) ? TEXT("readable") : TEXT("readable or writable"));
 					}
 				}
 			}
@@ -894,6 +896,12 @@ private:
 		}
 	}
 
+	bool HasAllFlags(ECacheStoreFlags Flags) const final
+	{
+		check(Hierarchy);
+		return Hierarchy->HasAllFlags(Flags);
+	}
+
 	ICacheStoreStats* CreateStats(ILegacyCacheStore* CacheStore, ECacheStoreFlags Flags, FStringView Type, FStringView Name, FStringView Path) final
 	{
 		check(Hierarchy);
diff --git a/Engine/Source/Developer/DerivedDataCache/Private/DerivedDataCacheStore.h b/Engine/Source/Developer/DerivedDataCache/Private/DerivedDataCacheStore.h
--- a/Engine/Source/Developer/DerivedDataCache/Private/DerivedDataCacheStore.h
+++ b/Engine/Source/Developer/DerivedDataCache/Private/DerivedDataCacheStore.h
@@ -196,6 +196,9 @@ public:
 	virtual void SetFlags(ILegacyCacheStore* CacheStore, ECacheStoreFlags Flags) = 0;
 	virtual void RemoveNotSafe(ILegacyCacheStore* CacheStore) = 0;
 
+	/** Returns true if the combined flags of the owned cache stores contain all of these flags. */
+	virtual bool HasAllFlags(ECacheStoreFlags Flags) const = 0;
+
 	virtual ICacheStoreStats* CreateStats(ILegacyCacheStore* CacheStore, ECacheStoreFlags Flags, FStringView Type, FStringView Name, FStringView Path = {}) = 0;
 	virtual void DestroyStats(ICacheStoreStats* Stats) = 0;
 
diff --git a/Engine/Source/Developer/DerivedDataCache/Private/DerivedDataCacheStoreHierarchy.cpp b/Engine/Source/Developer/DerivedDataCache/Private/DerivedDataCacheStoreHierarchy.cpp
--- a/Engine/Source/Developer/DerivedDataCache/Private/DerivedDataCacheStoreHierarchy.cpp
+++ b/Engine/Source/Developer/DerivedDataCache/Private/DerivedDataCacheStoreHierarchy.cpp
@@ -48,6 +48,8 @@ public:
 	void SetFlags(ILegacyCacheStore* CacheStore, ECacheStoreFlags Flags) final;
 	void RemoveNotSafe(ILegacyCacheStore* CacheStore) final;
 
+	bool HasAllFlags(ECacheStoreFlags Flags) const final;
+
 	ICacheStoreStats* CreateStats(ILegacyCacheStore* CacheStore, ECacheStoreFlags Flags, FStringView Type, FStringView Name, FStringView Path) final;
 	void DestroyStats(ICacheStoreStats* Stats) final;
 
@@ -268,11 +270,20 @@ FCacheStoreHierarchy::FCacheStoreHierarchy(ICacheStoreOwner*& OutOwner, TFunctio
 FCacheStoreHierarchy::~FCacheStoreHierarchy()
 {
 	// Delete nodes separately before Nodes is destroyed because destroying stats depends on it.
-	for (FCacheStoreNode& Node : ReverseIterate(Nodes))
+	while (!Nodes.IsEmpty())
 	{
+		const int32 NodeIndex = Nodes.Num() - 1;
+		FCacheStoreNode& Node = Nodes[NodeIndex];
+
 		Node.AsyncCache.Reset();
 		delete Node.Cache;
 
+		if (Nodes.Num() <= NodeIndex)
+		{
+			// The node removed itself in its destructor.
+			continue;
+		}
+
 		if (UNLIKELY(!Node.CacheStats.IsEmpty()))
 		{
 			FCacheStoreStats* Stats = Node.CacheStats[0];
@@ -280,6 +291,8 @@ FCacheStoreHierarchy::~FCacheStoreHierarchy()
 				"Leaked stats for {Type} cache store '{Name}' with path '{Path}'.",
 				Stats->Type, Stats->Name, Stats->Path);
 		}
+
+		Nodes.Pop(/*bAllowShrinking*/ false);
 	}
 }
 
@@ -307,10 +320,32 @@ void FCacheStoreHierarchy::RemoveNotSafe(ILegacyCacheStore* CacheStore)
 	FWriteScopeLock Lock(NodesLock);
 	FCacheStoreNode* Node = Algo::FindBy(Nodes, CacheStore, &FCacheStoreNode::Cache);
 	checkf(!!Node, TEXT("Attempting to remove a cache store that is not registered to the hierarchy."));
-	Nodes.RemoveAt(UE_PTRDIFF_TO_INT32(Node - Nodes.GetData()));
+	const int32 NodeIndex = UE_PTRDIFF_TO_INT32(Node - Nodes.GetData());
+	if (UNLIKELY(!Nodes[NodeIndex].CacheStats.IsEmpty()))
+	{
+		FCacheStoreStats* Stats = Nodes[NodeIndex].CacheStats[0];
+		UE_LOGFMT(LogDerivedDataCache, Fatal,
+			"Leaked stats for {Type} cache store '{Name}' with path '{Path}'.",
+			Stats->Type, Stats->Name, Stats->Path);
+	}
+	Nodes.RemoveAt(NodeIndex);
 	UpdateNodeFlags();
 }
 
+bool FCacheStoreHierarchy::HasAllFlags(ECacheStoreFlags Flags) const
+{
+	FReadScopeLock Lock(NodesLock);
+	ECacheStoreFlags CombinedFlags = ECacheStoreFlags::None;
+	for (const FCacheStoreNode& Node : Nodes)
+	{
+		if (Node.Cache != MemoryCache)
+		{
+			CombinedFlags |= Node.CacheFlags;
+		}
+	}
+	return EnumHasAllFlags(CombinedFlags, Flags);
+}
+
 void FCacheStoreHierarchy::UpdateNodeFlags()
 {
 	ECacheStoreNodeFlags StoreFlags = ECacheStoreNodeFlags::None;
