Change the way readiness of ZenCacheStore is evaluated and re-evaluated.

Instead of relying on a one-time blocking check with no reponse timeout, we now issue a blocking request with a 5 second idle time limit.  If it fails, the store will still be created, but it will go into the same asyncronous re-evaluation loop as when performance is below the acceptable threshold and re-evaluate at 30 second intervals until both:

- Health is Ok
- Performance criteria (if any has been configured and is in use) is met

At which point it is activated.  Health checks have been changed from the health/status endpoint to health/ready because we don't want to act if the server is running but not ready for requests (eg: during the time when it may be wiping data during a schema change.

The overall goal is that we don't force the entire session to execute without zenserver if zenserver was not ready at startup.

diff --git a/Engine/Source/Developer/DerivedDataCache/Private/ZenCacheStore.cpp b/Engine/Source/Developer/DerivedDataCache/Private/ZenCacheStore.cpp
--- a/Engine/Source/Developer/DerivedDataCache/Private/ZenCacheStore.cpp
+++ b/Engine/Source/Developer/DerivedDataCache/Private/ZenCacheStore.cpp
@@ -167,6 +167,7 @@ private:
 	void EnqueueAsyncRpc(IRequestOwner& Owner, FCbObject RequestObject, FOnRpcComplete&& OnComplete);
 	void EnqueueAsyncRpc(IRequestOwner& Owner, const FCbPackage& RequestPackage, FOnRpcComplete&& OnComplete);
 
+	void ActivatePerformanceEvaluationThread();
 	void ConditionalEvaluatePerformance();
 	void UpdateStatus();
 
@@ -195,16 +196,16 @@ private:
 	class FGetValueOp;
 	class FGetChunksOp;
 
-	enum class EHealthStatus
+	enum class EHealth
 	{
 		Unknown,
 		Ok,
 		Error,
 	};
 
-	using FOnHealthStatusComplete = TUniqueFunction<void(THttpUniquePtr<IHttpResponse>& HttpResponse, EHealthStatus HealthStatus)>;
-	class FHealthStatusReceiver;
-	class FAsyncHealthStatusReceiver;
+	using FOnHealthComplete = TUniqueFunction<void(THttpUniquePtr<IHttpResponse>& HttpResponse, EHealth Health)>;
+	class FHealthReceiver;
+	class FAsyncHealthReceiver;
 
 	class FCbPackageReceiver;
 	class FAsyncCbPackageReceiver;
@@ -1301,17 +1302,17 @@ private:
 	FOnCacheGetChunkComplete OnComplete;
 };
 
-class FZenCacheStore::FHealthStatusReceiver final : public IHttpReceiver
+class FZenCacheStore::FHealthReceiver final : public IHttpReceiver
 {
 public:
-	FHealthStatusReceiver(const FHealthStatusReceiver&) = delete;
-	FHealthStatusReceiver& operator=(const FHealthStatusReceiver&) = delete;
+	FHealthReceiver(const FHealthReceiver&) = delete;
+	FHealthReceiver& operator=(const FHealthReceiver&) = delete;
 
-	explicit FHealthStatusReceiver(EHealthStatus& OutHealthStatus, IHttpReceiver* InNext = nullptr)
-		: HealthStatus(OutHealthStatus)
+	explicit FHealthReceiver(EHealth& OutHealth, IHttpReceiver* InNext = nullptr)
+		: Health(OutHealth)
 		, Next(InNext)
 	{
-		HealthStatus = EHealthStatus::Unknown;
+		Health = EHealth::Unknown;
 	}
 
 private:
@@ -1325,38 +1326,38 @@ private:
 		FUtf8StringView ResponseStringView(reinterpret_cast<const UTF8CHAR*>(BodyArray.GetData()), IntCastChecked<int32>(BodyArray.Num()));
 		if (ResponseStringView == UTF8TEXTVIEW("OK!"))
 		{
-			HealthStatus = EHealthStatus::Ok;
+			Health = EHealth::Ok;
 		}
 		else
 		{
-			HealthStatus = EHealthStatus::Error;
+			Health = EHealth::Error;
 		}
 		return Next;
 	}
 
 private:
-	EHealthStatus& HealthStatus;
+	EHealth& Health;
 	IHttpReceiver* Next;
 	TArray64<uint8> BodyArray;
 	FHttpByteArrayReceiver BodyReceiver{ BodyArray, this };
 };
 
-class FZenCacheStore::FAsyncHealthStatusReceiver final : public FRequestBase, public IHttpReceiver
+class FZenCacheStore::FAsyncHealthReceiver final : public FRequestBase, public IHttpReceiver
 {
 public:
-	FAsyncHealthStatusReceiver(const FAsyncHealthStatusReceiver&) = delete;
-	FAsyncHealthStatusReceiver& operator=(const FAsyncHealthStatusReceiver&) = delete;
+	FAsyncHealthReceiver(const FAsyncHealthReceiver&) = delete;
+	FAsyncHealthReceiver& operator=(const FAsyncHealthReceiver&) = delete;
 
-	FAsyncHealthStatusReceiver(
+	FAsyncHealthReceiver(
 		THttpUniquePtr<IHttpRequest>&& InRequest,
 		IRequestOwner* InOwner,
 		Zen::FZenServiceInstance& InZenServiceInstance,
-		FOnHealthStatusComplete&& InOnHealthStatusComplete)
+		FOnHealthComplete&& InOnHealthComplete)
 		: Request(MoveTemp(InRequest))
 		, Owner(InOwner)
 		, ZenServiceInstance(InZenServiceInstance)
-		, BaseReceiver(HealthStatus, this)
-		, OnHealthStatusComplete(MoveTemp(InOnHealthStatusComplete))
+		, BaseReceiver(Health, this)
+		, OnHealthComplete(MoveTemp(InOnHealthComplete))
 	{
 		Request->SendAsync(this, Response);
 	}
@@ -1400,19 +1401,19 @@ private:
 		{
 			if (Self->ShouldRecoverAndRetry(LocalResponse) && Self->ZenServiceInstance.TryRecovery())
 			{
-				new FAsyncHealthStatusReceiver(MoveTemp(Self->Request), Self->Owner, Self->ZenServiceInstance, MoveTemp(Self->OnHealthStatusComplete));
+				new FAsyncHealthReceiver(MoveTemp(Self->Request), Self->Owner, Self->ZenServiceInstance, MoveTemp(Self->OnHealthComplete));
 				return;
 			}
 
 			Self->Request.Reset();
-			if (Self->OnHealthStatusComplete)
+			if (Self->OnHealthComplete)
 			{
 				// Launch a task for the completion function since it can execute arbitrary code.
-				Self->Owner->LaunchTask(TEXT("ZenHealthStatusComplete"), [Self = TRefCountPtr(Self)]
+				Self->Owner->LaunchTask(TEXT("ZenHealthComplete"), [Self = TRefCountPtr(Self)]
 				{
 					// Ensuring that the OnRpcComplete method is destroyed by the time we exit this method by moving it to a local scope variable
-					FOnHealthStatusComplete LocalOnComplete = MoveTemp(Self->OnHealthStatusComplete);
-					LocalOnComplete(Self->Response, Self->HealthStatus);
+					FOnHealthComplete LocalOnComplete = MoveTemp(Self->OnHealthComplete);
+					LocalOnComplete(Self->Response, Self->Health);
 				});
 			}
 		});
@@ -1425,9 +1426,9 @@ private:
 	TRefCountPtr<IHttpResponseMonitor> Monitor;
 	IRequestOwner* Owner;
 	Zen::FZenServiceInstance& ZenServiceInstance;
-	EHealthStatus HealthStatus;
-	FHealthStatusReceiver BaseReceiver;
-	FOnHealthStatusComplete OnHealthStatusComplete;
+	EHealth Health;
+	FHealthReceiver BaseReceiver;
+	FOnHealthComplete OnHealthComplete;
 };
 
 class FZenCacheStore::FCbPackageReceiver final : public IHttpReceiver
@@ -1601,61 +1602,100 @@ void FZenCacheStore::Initialize(const FZenCacheStoreParams& Params)
 {
 	LastPerformanceEvaluationTicks.store(FDateTime::UtcNow().GetTicks(), std::memory_order_relaxed);
 	Namespace = Params.Namespace;
-	if (IsServiceReady())
-	{
-		RpcUri << ZenService.GetInstance().GetURL() << ANSITEXTVIEW("/z$/$rpc");
 
-		const uint32 MaxConnections = uint32(FMath::Clamp(FPlatformMisc::NumberOfCoresIncludingHyperthreads(), 8, 64));
-		constexpr uint32 RequestPoolSize = 128;
-		constexpr uint32 RequestPoolOverflowSize = 128;
+	RpcUri << ZenService.GetInstance().GetURL() << ANSITEXTVIEW("/z$/$rpc");
 
-		FHttpConnectionPoolParams ConnectionPoolParams;
-		ConnectionPoolParams.MaxConnections = MaxConnections;
-		ConnectionPoolParams.MinConnections = MaxConnections;
-		ConnectionPool = IHttpManager::Get().CreateConnectionPool(ConnectionPoolParams);
+	const uint32 MaxConnections = uint32(FMath::Clamp(FPlatformMisc::NumberOfCoresIncludingHyperthreads(), 8, 64));
+	constexpr uint32 RequestPoolSize = 128;
+	constexpr uint32 RequestPoolOverflowSize = 128;
 
-		FHttpClientParams ClientParams;
-		ClientParams.Version = EHttpVersion::V2;
-		ClientParams.MaxRequests = RequestPoolSize + RequestPoolOverflowSize;
-		ClientParams.MinRequests = RequestPoolSize;
-		ClientParams.LowSpeedLimit = 30;
-		ClientParams.LowSpeedTime = 60;
-		RequestQueue = FHttpRequestQueue(*ConnectionPool, ClientParams);
+	FHttpConnectionPoolParams ConnectionPoolParams;
+	ConnectionPoolParams.MaxConnections = MaxConnections;
+	ConnectionPoolParams.MinConnections = MaxConnections;
+	ConnectionPool = IHttpManager::Get().CreateConnectionPool(ConnectionPoolParams);
 
-		bIsUsable = true;
-		bIsLocalConnection = ZenService.GetInstance().IsServiceRunningLocally();
+	bool bReady = false;
+	{
+		// Issue a synchronous health/ready request with a limited idle time
+		FHttpClientParams ReadinessClientParams;
+		ReadinessClientParams.Version = EHttpVersion::V2;
+		ReadinessClientParams.MaxRequests = 1;
+		ReadinessClientParams.MinRequests = 1;
+		ReadinessClientParams.LowSpeedLimit = 1;
+		ReadinessClientParams.LowSpeedTime = 5; // 5 second idle time limit for the initial readiness check
+		THttpUniquePtr<IHttpClient> ReadinessClient = ConnectionPool->CreateClient(ReadinessClientParams);
+		THttpUniquePtr<IHttpRequest> ReadinessRequest = ReadinessClient->TryCreateRequest({});
+		TAnsiStringBuilder<256> StatusUri;
+		StatusUri << ZenService.GetInstance().GetURL() << ANSITEXTVIEW("/health/ready");
+		ReadinessRequest->SetUri(StatusUri);
+		ReadinessRequest->SetMethod(EHttpMethod::Get);
+		ReadinessRequest->AddAcceptType(EHttpMediaType::Text);
+		EHealth Health;
+		FHealthReceiver HealthReceiver(Health);
+		THttpUniquePtr<IHttpResponse> ReadinessResponse;
+		ReadinessRequest->Send(&HealthReceiver, ReadinessResponse);
+		bReady = Health == EHealth::Ok;
+	}
 
-		if (StoreOwner)
-		{
-			// Default to locally launched service getting the Local cache store flag.  Can be overridden by explicit value in config.
-			bool bLocal = Params.bLocal.Get(bIsLocalConnection);
+	FHttpClientParams ClientParams;
+	ClientParams.MaxRequests = RequestPoolSize + RequestPoolOverflowSize;
+	ClientParams.MinRequests = RequestPoolSize;
+	ClientParams.LowSpeedLimit = 1;
+	ClientParams.LowSpeedTime = 25;
+	RequestQueue = FHttpRequestQueue(*ConnectionPool, ClientParams);
 
-			// Default to non-locally launched service getting the Remote cache store flag.  Can be overridden by explicit value in config.
-			// In the future this could be extended to allow the Remote flag by default (even for locally launched instances) if they have upstreams configured.
-			bool bRemote = Params.bRemote.Get(!bIsLocalConnection);
+	bIsLocalConnection = ZenService.GetInstance().IsServiceRunningLocally();
+	bIsUsable = true;
+
+
+	if (StoreOwner)
+	{
+		// Default to locally launched service getting the Local cache store flag.  Can be overridden by explicit value in config.
+		bool bLocal = Params.bLocal.Get(bIsLocalConnection);
 
-			DeactivateAtMs = Params.DeactivateAtMs;
+		// Default to non-locally launched service getting the Remote cache store flag.  Can be overridden by explicit value in config.
+		// In the future this could be extended to allow the Remote flag by default (even for locally launched instances) if they have upstreams configured.
+		bool bRemote = Params.bRemote.Get(!bIsLocalConnection);
 
-			ECacheStoreFlags Flags = ECacheStoreFlags::Query;
-			Flags |= Params.bReadOnly ? ECacheStoreFlags::None : ECacheStoreFlags::Store;
-			Flags |= bLocal ? ECacheStoreFlags::Local : ECacheStoreFlags::None;
-			Flags |= bRemote ? ECacheStoreFlags::Remote : ECacheStoreFlags::None;
+		DeactivateAtMs = Params.DeactivateAtMs;
 
-			OperationalFlags = Flags;
+		ECacheStoreFlags Flags = ECacheStoreFlags::Query;
+		Flags |= Params.bReadOnly ? ECacheStoreFlags::None : ECacheStoreFlags::Store;
+		Flags |= bLocal ? ECacheStoreFlags::Local : ECacheStoreFlags::None;
+		Flags |= bRemote ? ECacheStoreFlags::Remote : ECacheStoreFlags::None;
 
-			StoreOwner->Add(this, Flags);
-			TStringBuilder<256> Path(InPlace, ZenService.GetInstance().GetURL(), TEXTVIEW(" ("), Namespace, TEXTVIEW(")"));
-			StoreStats = StoreOwner->CreateStats(this, Flags, TEXT("Zen"), *Params.Name, Path);
-			bTryEvaluatePerformance = !GIsBuildMachine && (StoreStats != nullptr) && (DeactivateAtMs > 0.0f);
+		OperationalFlags = Flags;
 
-			StoreStats->SetAttribute(TEXTVIEW("Namespace"), Namespace);
+		if (!bReady)
+		{
+			Flags = ECacheStoreFlags::None;
+			UE_LOG(LogDerivedDataCache, Display,
+				TEXT("%s: Readiness check failed. "
+					"It will be deactivated until responsiveness improves. "
+					"If this is consistent, consider disabling this cache store through "
+					"environment variables or other configuration."),
+				*GetName());
+			bDeactivatedForPerformance.store(true, std::memory_order_relaxed);
 		}
 
-		// Issue a request for stats as it will be fetched asynchronously and issuing now makes them available sooner for future callers.
-		Zen::FZenCacheStats ZenStats;
-		ZenService.GetInstance().GetCacheStats(ZenStats);
+		StoreOwner->Add(this, Flags);
+		TStringBuilder<256> Path(InPlace, ZenService.GetInstance().GetURL(), TEXTVIEW(" ("), Namespace, TEXTVIEW(")"));
+		StoreStats = StoreOwner->CreateStats(this, Flags, TEXT("Zen"), *Params.Name, Path);
+		bTryEvaluatePerformance = !GIsBuildMachine && (StoreStats != nullptr) && (DeactivateAtMs > 0.0f);
+
+		StoreStats->SetAttribute(TEXTVIEW("Namespace"), Namespace);
+
+		if (!bReady)
+		{
+			UpdateStatus();
+			ActivatePerformanceEvaluationThread();
+		}
 	}
 
+	// Issue a request for stats as it will be fetched asynchronously and issuing now makes them available sooner for future callers.
+	Zen::FZenCacheStats ZenStats;
+	ZenService.GetInstance().GetCacheStats(ZenStats);
+
 	MaxBatchPutKB = Params.MaxBatchPutKB;
 	CacheRecordBatchSize = Params.RecordBatchSize;
 	CacheChunksBatchSize = Params.ChunksBatchSize;
@@ -1699,6 +1739,57 @@ void FZenCacheStore::EnqueueAsyncRpc(IRequestOwner& Owner, const FCbPackage& Req
 	new FAsyncCbPackageReceiver(MoveTemp(Request), &Owner, ZenService.GetInstance(), MoveTemp(OnComplete));
 }
 
+void FZenCacheStore::ActivatePerformanceEvaluationThread()
+{
+	if (!PerformanceEvaluationThread.IsSet())
+	{
+		PerformanceEvaluationThread.Emplace(TEXT("ZenCacheStore Performance Evaluation"), [this]
+		{
+			while (!PerformanceEvaluationThreadShutdownEvent.WaitFor(FMonotonicTimeSpan::FromSeconds(30.0)))
+			{
+				IRequestOwner& Owner(PerformanceEvaluationRequestOwner);
+				FRequestBarrier Barrier(Owner);
+				THttpUniquePtr<IHttpRequest> Request = RequestQueue.CreateRequest({});
+				TAnsiStringBuilder<256> StatusUri;
+				StatusUri << ZenService.GetInstance().GetURL() << ANSITEXTVIEW("/health/ready");
+				Request->SetUri(StatusUri);
+				Request->SetMethod(EHttpMethod::Get);
+				Request->AddAcceptType(EHttpMediaType::Text);
+				new FAsyncHealthReceiver(MoveTemp(Request), &Owner, ZenService.GetInstance(), [this, StartTime = FMonotonicTimePoint::Now()](THttpUniquePtr<IHttpResponse>& HttpResponse, EHealth Health)
+					{
+						if (Health != EHealth::Ok)
+						{
+							// Any non-ok health means we hold off any possibility of reactivating and we don't add the lagency to the stats
+							// as the failure may be client-side and instantaneous which will create an artificially low latency measurement.
+							return;
+						}
+
+						double LatencySec = (HttpResponse->GetStats().StartTransferTime - HttpResponse->GetStats().ConnectTime);
+						StoreStats->AddLatency(StartTime, FMonotonicTimePoint::Now(), FMonotonicTimeSpan::FromSeconds(LatencySec));
+						if (!bTryEvaluatePerformance || (StoreStats->GetAverageLatency() * 1000 <= DeactivateAtMs))
+						{
+							if (PerformanceEvaluationThread.IsSet())
+							{
+								PerformanceEvaluationThreadShutdownEvent.Notify();
+							}
+
+							if (bDeactivatedForPerformance.load(std::memory_order_relaxed))
+							{
+								StoreOwner->SetFlags(this, OperationalFlags);
+								UE_LOG(LogDerivedDataCache, Display,
+									TEXT("%s: Performance has improved and meets minimum performance criteria. "
+										"It will be reactivated now."),
+									*GetName());
+								bDeactivatedForPerformance.store(false, std::memory_order_relaxed);
+								UpdateStatus();
+							}
+						}
+					});
+			}
+		});
+	}
+}
+
 void FZenCacheStore::ConditionalEvaluatePerformance()
 {
 	if (!bTryEvaluatePerformance)
@@ -1735,7 +1826,7 @@ void FZenCacheStore::ConditionalEvaluatePerformance()
 		if (!bDeactivatedForPerformance.load(std::memory_order_relaxed))
 		{
 			StoreOwner->SetFlags(this, ECacheStoreFlags::None);
-			UE_LOG(LogDerivedDataCache, Warning,
+			UE_LOG(LogDerivedDataCache, Display,
 				TEXT("%s: Performance does not meet minimum criteria. "
 					"It will be deactivated until performance measurements improve. "
 					"If this is consistent, consider disabling this cache store through "
@@ -1745,51 +1836,7 @@ void FZenCacheStore::ConditionalEvaluatePerformance()
 			UpdateStatus();
 		}
 
-		if (!PerformanceEvaluationThread.IsSet())
-		{
-			PerformanceEvaluationThread.Emplace(TEXT("ZenCacheStore Performance Evaluation"), [this]
-			{
-				while (!PerformanceEvaluationThreadShutdownEvent.WaitFor(FMonotonicTimeSpan::FromSeconds(30.0)))
-				{
-					IRequestOwner& Owner(PerformanceEvaluationRequestOwner);
-					FRequestBarrier Barrier(Owner);
-					THttpUniquePtr<IHttpRequest> Request = RequestQueue.CreateRequest({});
-					TAnsiStringBuilder<256> StatusUri;
-					StatusUri << ZenService.GetInstance().GetURL() << ANSITEXTVIEW("/health/status");
-					Request->SetUri(StatusUri);
-					Request->SetMethod(EHttpMethod::Get);
-					Request->AddAcceptType(EHttpMediaType::Text);
-					new FAsyncHealthStatusReceiver(MoveTemp(Request), &Owner, ZenService.GetInstance(), [this, StartTime = FMonotonicTimePoint::Now()](THttpUniquePtr<IHttpResponse>& HttpResponse, EHealthStatus HealthStatus)
-						{
-							if (HealthStatus != EHealthStatus::Ok)
-							{
-								return;
-							}
-
-							double LatencySec = (HttpResponse->GetStats().StartTransferTime - HttpResponse->GetStats().ConnectTime);
-							StoreStats->AddLatency(StartTime, FMonotonicTimePoint::Now(), FMonotonicTimeSpan::FromSeconds(LatencySec));
-							if (StoreStats->GetAverageLatency() * 1000 <= DeactivateAtMs)
-							{
-								if (PerformanceEvaluationThread.IsSet())
-								{
-									PerformanceEvaluationThreadShutdownEvent.Notify();
-								}
-
-								if (bDeactivatedForPerformance.load(std::memory_order_relaxed))
-								{
-									StoreOwner->SetFlags(this, OperationalFlags);
-									UE_LOG(LogDerivedDataCache, Display,
-										TEXT("%s: Performance has improved and meets minimum performance criteria. "
-											"It will be reactivated now."),
-										*GetName());
-									bDeactivatedForPerformance.store(false, std::memory_order_relaxed);
-									UpdateStatus();
-								}
-							}
-						});
-				}
-			});
-		}
+		ActivatePerformanceEvaluationThread();
 	}
 }
 
@@ -1799,7 +1846,7 @@ void FZenCacheStore::UpdateStatus()
 	{
 		if (bDeactivatedForPerformance.load(std::memory_order_relaxed))
 		{
-			StoreStats->SetStatus(ECacheStoreStatusCode::Warning, NSLOCTEXT("DerivedDataCache", "DeactivatedForPerformance", "Deactivated for performance"));
+			StoreStats->SetStatus(ECacheStoreStatusCode::Warning, NSLOCTEXT("DerivedDataCache", "DeactivatedForPerformanceOrReadiness", "Deactivated for performance or readiness"));
 		}
 		else
 		{
