Initialize an out variable to satisfy static analysis.

diff --git a/Engine/Source/Developer/DerivedDataCache/Private/ZenCacheStore.cpp b/Engine/Source/Developer/DerivedDataCache/Private/ZenCacheStore.cpp
--- a/Engine/Source/Developer/DerivedDataCache/Private/ZenCacheStore.cpp
+++ b/Engine/Source/Developer/DerivedDataCache/Private/ZenCacheStore.cpp
@@ -1630,7 +1630,7 @@ void FZenCacheStore::Initialize(const FZenCacheStoreParams& Params)
 		ReadinessRequest->SetUri(StatusUri);
 		ReadinessRequest->SetMethod(EHttpMethod::Get);
 		ReadinessRequest->AddAcceptType(EHttpMediaType::Text);
-		EHealth Health;
+		EHealth Health = EHealth::Unknown;
 		FHealthReceiver HealthReceiver(Health);
 		THttpUniquePtr<IHttpResponse> ReadinessResponse;
 		ReadinessRequest->Send(&HealthReceiver, ReadinessResponse);
