For local Zen DDC connections show the path where data is store in the Cache Statistics tab for DDC instead of local ip address
For Zen DDC connection show the storage size

diff --git a/Engine/Source/Developer/DerivedDataCache/Private/DerivedDataCacheStats.cpp b/Engine/Source/Developer/DerivedDataCache/Private/DerivedDataCacheStats.cpp
--- a/Engine/Source/Developer/DerivedDataCache/Private/DerivedDataCacheStats.cpp
+++ b/Engine/Source/Developer/DerivedDataCache/Private/DerivedDataCacheStats.cpp
@@ -317,6 +317,12 @@ double FCacheStoreStats::GetAverageLatency()
 	return AverageLatency.GetValue(FMonotonicTimePoint::Now());
 }
 
+void FCacheStoreStats::SetTotalPhysicalSize(uint64 InTotalPhysicalSize)
+{
+	TUniqueLock Lock(Mutex);
+	TotalPhysicalSize = InTotalPhysicalSize;
+}
+
 ///////////////////////////////////////////////////////////////////////////////////////////////////////////////////////
 
 FCacheBucketStats& FCacheStats::GetBucket(FCacheBucket Bucket)
diff --git a/Engine/Source/Developer/DerivedDataCache/Private/DerivedDataCacheStats.h b/Engine/Source/Developer/DerivedDataCache/Private/DerivedDataCacheStats.h
--- a/Engine/Source/Developer/DerivedDataCache/Private/DerivedDataCacheStats.h
+++ b/Engine/Source/Developer/DerivedDataCache/Private/DerivedDataCacheStats.h
@@ -162,6 +162,8 @@ public:
 	FMonotonicTimeSpan OtherThreadTime;
 	/** Count of requests that accessed the cache store. */
 	FRequestCounter RequestCount;
+	/** Total size of the data in the cache store, ~0ull indicate that the information is not available */
+	uint64 TotalPhysicalSize = ~0ull;
 
 	FTimeAveragedStat AverageLatency;
 	FTimeAveragedStat AveragePhysicalReadSize;
@@ -182,6 +184,7 @@ private:
 	void AddRequest(const FCacheStoreRequestStats& Stats) final;
 	void AddLatency(FMonotonicTimePoint StartTime, FMonotonicTimePoint EndTime, FMonotonicTimeSpan Latency) final;
 	double GetAverageLatency() final;
+	void SetTotalPhysicalSize(uint64 TotalPhysicalSize) final;
 };
 
 class FCacheStats final
diff --git a/Engine/Source/Developer/DerivedDataCache/Private/DerivedDataCacheStore.h b/Engine/Source/Developer/DerivedDataCache/Private/DerivedDataCacheStore.h
--- a/Engine/Source/Developer/DerivedDataCache/Private/DerivedDataCacheStore.h
+++ b/Engine/Source/Developer/DerivedDataCache/Private/DerivedDataCacheStore.h
@@ -311,6 +311,8 @@ public:
 
 	/** Gets the average latency value for the current time in seconds. */
 	virtual double GetAverageLatency() = 0;
+
+	virtual void SetTotalPhysicalSize(uint64 TotalPhysicalSize) = 0;
 };
 
 template <typename RequestRangeType, typename OnCompleteType>
diff --git a/Engine/Source/Developer/DerivedDataCache/Private/DerivedDataCacheStoreHierarchy.cpp b/Engine/Source/Developer/DerivedDataCache/Private/DerivedDataCacheStoreHierarchy.cpp
--- a/Engine/Source/Developer/DerivedDataCache/Private/DerivedDataCacheStoreHierarchy.cpp
+++ b/Engine/Source/Developer/DerivedDataCache/Private/DerivedDataCacheStoreHierarchy.cpp
@@ -1574,6 +1574,8 @@ static void ConvertToLegacyStats(FDerivedDataCacheStatsNode& OutNode, FCacheStor
 	OutNode.SpeedStats.LatencyMS = Stats.AverageLatency.GetValue(Now) * 1000.0;
 	OutNode.SpeedStats.ReadSpeedMBs = Stats.AveragePhysicalReadSize.GetRate(Now) / 1024.0 / 1024.0;
 	OutNode.SpeedStats.WriteSpeedMBs = Stats.AveragePhysicalWriteSize.GetRate(Now) / 1024.0 / 1024.0;
+
+	OutNode.SetTotalPhysicalSize(Stats.TotalPhysicalSize);
 }
 
 void FCacheStoreHierarchy::LegacyStats(FDerivedDataCacheStatsNode& OutNode)
diff --git a/Engine/Source/Developer/DerivedDataCache/Private/ZenCacheStore.cpp b/Engine/Source/Developer/DerivedDataCache/Private/ZenCacheStore.cpp
--- a/Engine/Source/Developer/DerivedDataCache/Private/ZenCacheStore.cpp
+++ b/Engine/Source/Developer/DerivedDataCache/Private/ZenCacheStore.cpp
@@ -2,6 +2,7 @@
 
 #include "DerivedDataLegacyCacheStore.h"
 #include "Experimental/ZenServerInterface.h"
+#include "Experimental/ZenStatistics.h"
 
 #if UE_WITH_ZEN
 
@@ -169,6 +170,7 @@ private:
 
 	void ActivatePerformanceEvaluationThread();
 	void ConditionalEvaluatePerformance();
+	void ConditionalUpdateStorageSize();
 	void UpdateStatus();
 
 	template <typename T, typename... ArgTypes>
@@ -228,6 +230,7 @@ private:
 	std::atomic<int64> LastPerformanceEvaluationTicks;
 	TOptional<FThread> PerformanceEvaluationThread;
 	FManualResetEvent PerformanceEvaluationThreadShutdownEvent;
+	std::atomic<int64> LastStorageSizeUpdateTicks;
 	float DeactivateAtMs = -1.0f;
 	ECacheStoreFlags OperationalFlags;
 	FRequestOwner PerformanceEvaluationRequestOwner;
@@ -310,6 +313,7 @@ public:
 			Request.Stats.Latency = AverageLatency;
 		}
 		CacheStore.ConditionalEvaluatePerformance();
+		CacheStore.ConditionalUpdateStorageSize();
 	}
 
 	void IssueRequests()
@@ -510,6 +514,7 @@ public:
 			Request.Stats.Latency = AverageLatency;
 		}
 		CacheStore.ConditionalEvaluatePerformance();
+		CacheStore.ConditionalUpdateStorageSize();
 	}
 
 	void IssueRequests()
@@ -716,6 +721,7 @@ public:
 			Request.Stats.Latency = AverageLatency;
 		}
 		CacheStore.ConditionalEvaluatePerformance();
+		CacheStore.ConditionalUpdateStorageSize();
 	}
 
 	void IssueRequests()
@@ -906,6 +912,7 @@ public:
 			Request.Stats.Latency = AverageLatency;
 		}
 		CacheStore.ConditionalEvaluatePerformance();
+		CacheStore.ConditionalUpdateStorageSize();
 	}
 
 	void IssueRequests()
@@ -1105,6 +1112,7 @@ public:
 			Request.Stats.Latency = AverageLatency;
 		}
 		CacheStore.ConditionalEvaluatePerformance();
+		CacheStore.ConditionalUpdateStorageSize();
 	}
 
 	void IssueRequests()
@@ -1601,6 +1609,7 @@ FZenCacheStore::~FZenCacheStore()
 void FZenCacheStore::Initialize(const FZenCacheStoreParams& Params)
 {
 	LastPerformanceEvaluationTicks.store(FDateTime::UtcNow().GetTicks(), std::memory_order_relaxed);
+	LastStorageSizeUpdateTicks.store(FDateTime::UtcNow().GetTicks(), std::memory_order_relaxed);
 	Namespace = Params.Namespace;
 
 	RpcUri << ZenService.GetInstance().GetURL() << ANSITEXTVIEW("/z$/$rpc");
@@ -1679,7 +1688,7 @@ void FZenCacheStore::Initialize(const FZenCacheStoreParams& Params)
 		}
 
 		StoreOwner->Add(this, Flags);
-		TStringBuilder<256> Path(InPlace, ZenService.GetInstance().GetURL(), TEXTVIEW(" ("), Namespace, TEXTVIEW(")"));
+		TStringBuilder<256> Path(InPlace, ZenService.GetInstance().GetPath(), TEXTVIEW(" ("), Namespace, TEXTVIEW(")"));
 		StoreStats = StoreOwner->CreateStats(this, Flags, TEXT("Zen"), *Params.Name, Path);
 		bTryEvaluatePerformance = !GIsBuildMachine && (StoreStats != nullptr) && (DeactivateAtMs > 0.0f);
 
@@ -1790,6 +1799,47 @@ void FZenCacheStore::ActivatePerformanceEvaluationThread()
 	}
 }
 
+void FZenCacheStore::ConditionalUpdateStorageSize()
+{
+	if (!StoreStats)
+	{
+		return;
+	}
+		
+	// Look for an opportunity to measure and evaluate if storage size is acceptable.
+	int64 LocalStorageSizeUpdateTicks = LastStorageSizeUpdateTicks.load(std::memory_order_relaxed);
+	FTimespan TimespanSinceLastStorageSizeUpdate = FDateTime::UtcNow() - FDateTime(LocalStorageSizeUpdateTicks);
+
+	if (TimespanSinceLastStorageSizeUpdate < FTimespan::FromSeconds(30))
+	{
+		return;
+	}
+
+	if (!LastStorageSizeUpdateTicks.compare_exchange_strong(LocalStorageSizeUpdateTicks, FDateTime::UtcNow().GetTicks()))
+	{
+		return;
+	}
+
+	bool PhysicalSizeIsValid = false;
+	double PhysicalSize = 0.0;
+	Zen::FZenCacheStats ZenCacheStats;
+	if (ZenService.GetInstance().GetCacheStats(ZenCacheStats))
+	{
+		PhysicalSize += ZenCacheStats.General.Size.Disk + static_cast<double>(ZenCacheStats.CID.Size.Total);
+		PhysicalSizeIsValid = true;
+	}
+	Zen::FZenProjectStats ZenProjectStats;
+	if (ZenService.GetInstance().GetProjectStats(ZenProjectStats))
+	{
+		PhysicalSize += ZenProjectStats.General.Size.Disk;
+		PhysicalSizeIsValid = true;
+	}
+	if (PhysicalSizeIsValid)
+	{
+		StoreStats->SetTotalPhysicalSize(static_cast<uint64>(PhysicalSize));
+	}
+}
+
 void FZenCacheStore::ConditionalEvaluatePerformance()
 {
 	if (!bTryEvaluatePerformance)
diff --git a/Engine/Source/Developer/DerivedDataCache/Public/DerivedDataCacheUsageStats.h b/Engine/Source/Developer/DerivedDataCache/Public/DerivedDataCacheUsageStats.h
--- a/Engine/Source/Developer/DerivedDataCache/Public/DerivedDataCacheUsageStats.h
+++ b/Engine/Source/Developer/DerivedDataCache/Public/DerivedDataCacheUsageStats.h
@@ -178,6 +178,9 @@ public:
 		}
 	}
 
+	uint64 GetTotalPhysicalSize() const { return TotalPhysicalSize; }
+	void SetTotalPhysicalSize(uint64 InTotalPhysicalSize) { TotalPhysicalSize = InTotalPhysicalSize; }
+
 public:
 	void GatherLegacyUsageStats(TMap<FString, FDerivedDataCacheUsageStats>& UsageStatsMap, FString&& GraphPath) const
 	{
@@ -213,6 +216,7 @@ protected:
 	FString CacheType;
 	FString CacheName;
 	FString CacheStatusText;
+	uint64 TotalPhysicalSize = ~0ull;
 	EDerivedDataCacheStatus CacheStatus;
 	bool bIsLocal;
 };
diff --git a/Engine/Source/Developer/Zen/Public/Experimental/ZenServerInterface.h b/Engine/Source/Developer/Zen/Public/Experimental/ZenServerInterface.h
--- a/Engine/Source/Developer/Zen/Public/Experimental/ZenServerInterface.h
+++ b/Engine/Source/Developer/Zen/Public/Experimental/ZenServerInterface.h
@@ -3,6 +3,7 @@
 #pragma once
 
 #include "Async/Future.h"
+#include "Async/Mutex.h"
 #include "Containers/StringFwd.h"
 #include "Containers/StringView.h"
 #include "Containers/UnrealString.h"
@@ -177,6 +178,7 @@ public:
 	UE_API FZenServiceInstance(FServiceSettings&& InSettings);
 	UE_API ~FZenServiceInstance();
 
+	UE_API const FString GetPath() const;
 	inline const TCHAR* GetURL() const { return *URL; }
 	inline const TCHAR* GetHostName() const { return *HostName; }
 	inline uint16 GetPort() const { return Port; }
@@ -211,13 +213,15 @@ private:
 
 	mutable TPimplPtr<class FZenHttpRequest> CacheStatsHttpRequest;
 	mutable TFuture<FZenCacheStats> CacheStatsRequest;
-	mutable FZenCacheStats LastCacheStats;
 	mutable uint64 LastCacheStatsTime = 0;
+	mutable FMutex LastCacheStatsMutex;
+	mutable FZenCacheStats LastCacheStats;
 
 	mutable TPimplPtr<class FZenHttpRequest> ProjectStatsHttpRequest;
 	mutable TFuture<FZenProjectStats> ProjectStatsRequest;
-	mutable FZenProjectStats LastProjectStats;
 	mutable uint64 LastProjectStatsTime = 0;
+	mutable FMutex LastProjectStatsMutex;
+	mutable FZenProjectStats LastProjectStats;
 
 	FServiceSettings Settings;
 	FString URL;
diff --git a/Engine/Source/Editor/DerivedDataEditor/Private/SDerivedDataDialogs.cpp b/Engine/Source/Editor/DerivedDataEditor/Private/SDerivedDataDialogs.cpp
--- a/Engine/Source/Editor/DerivedDataEditor/Private/SDerivedDataDialogs.cpp
+++ b/Engine/Source/Editor/DerivedDataEditor/Private/SDerivedDataDialogs.cpp
@@ -636,13 +636,23 @@ TSharedRef<SWidget> SDerivedDataCacheStatisticsDialog::GetGridPanel()
 		.Margin(TitleMargin)
 		.ColorAndOpacity(TitleColor)
 		.Font(TitleFont)
-		.Text(LOCTEXT("Name", "Name"))
+		.Text(LOCTEXT("Storage Size", "Storage Size"))
+	];
+
+	Panel->AddSlot(10, Row)
+	[
+		SNew(STextBlock)
+		.Margin(TitleMargin)
+		.ColorAndOpacity(TitleColor)
+		.Font(TitleFont)
+		.Text(LOCTEXT("Path", "Path"))
 	];
 
 	Row++;
 
 	double SumTotalGetMB = 0.0;
 	double SumTotalPutMB = 0.0;
+	uint64 SumTotalPhysicalSize = ~0ull;
 
 	for (TSharedRef<const FDerivedDataCacheStatsNode> Node : LeafUsageStats)
 	{
@@ -667,6 +677,19 @@ TSharedRef<SWidget> SDerivedDataCacheStatisticsDialog::GetGridPanel()
 		SumTotalGetMB += TotalGetMB;
 		SumTotalPutMB += TotalPutMB;
 
+		const uint64 TotalPhysicalSize = Node->GetTotalPhysicalSize();
+		if (TotalPhysicalSize != ~0ull)
+		{
+			if (SumTotalPhysicalSize == ~0ull)
+			{
+				SumTotalPhysicalSize = TotalPhysicalSize;
+			}
+			else
+			{
+				SumTotalPhysicalSize += TotalPhysicalSize;
+			}
+		}
+
 		TSharedPtr<SImage> StatusIcon;
 		switch (Node->GetCacheStatus())
 		{
@@ -788,6 +811,13 @@ TSharedRef<SWidget> SDerivedDataCacheStatisticsDialog::GetGridPanel()
 		}
 
 		Panel->AddSlot(9, Row)
+		[
+			SNew(STextBlock)
+			.Margin(DefaultMargin)
+			.Text(Node->GetTotalPhysicalSize() == ~0ull ? LOCTEXT("N/A", "N/A") : FText::AsMemory(Node->GetTotalPhysicalSize(), EMemoryUnitStandard::IEC))
+		];
+
+		Panel->AddSlot(10, Row)
 		[
 			SNew(STextBlock)
 			.Margin(DefaultMargin)
@@ -826,6 +856,16 @@ TSharedRef<SWidget> SDerivedDataCacheStatisticsDialog::GetGridPanel()
 		.Text(FText::FromString(SingleDecimalFormat(SumTotalPutMB)))
 	];
 
+	Panel->AddSlot(9, Row)
+	.HAlign(HAlign_Right)
+	[
+		SNew(STextBlock)
+		.Margin(TitleMargin)
+		.ColorAndOpacity(TitleColor)
+		.Font(TitleFont)
+		.Text(SumTotalPhysicalSize == ~0ull ? LOCTEXT("N/A", "N/A") : FText::AsMemory(SumTotalPhysicalSize, EMemoryUnitStandard::IEC))
+	];
+
 #else
 	Panel->AddSlot(0, 0)
 	[
