Change the internal use of cache flags in the DDC hierarchy to have the flag usage be atomic and avoid needing to take a write lock of the nodes lock when the flags are being set.  This is to avoid an issue where recursive read locks can lead to deadlocks in the presence of a write lock.  A future change is planned to alter the locks to ones that support recursive use, at which point this change will no longer be necessary.

diff --git a/Engine/Source/Developer/DerivedDataCache/Private/DerivedDataCacheStoreHierarchy.cpp b/Engine/Source/Developer/DerivedDataCache/Private/DerivedDataCacheStoreHierarchy.cpp
--- a/Engine/Source/Developer/DerivedDataCache/Private/DerivedDataCacheStoreHierarchy.cpp
+++ b/Engine/Source/Developer/DerivedDataCache/Private/DerivedDataCacheStoreHierarchy.cpp
@@ -18,6 +18,7 @@
 #include "DerivedDataRequest.h"
 #include "DerivedDataRequestOwner.h"
 #include "HAL/CriticalSection.h"
+#include "HAL/PlatformAtomics.h"
 #include "Logging/StructuredLog.h"
 #include "MemoryCacheStore.h"
 #include "Misc/EnumClassFlags.h"
@@ -127,6 +128,16 @@ private:
 
 	struct FCacheStoreNode
 	{
+		inline void SetCacheFlags(ECacheStoreFlags NewCacheFlags)
+		{
+			using InternalType = std::make_signed_t<std::underlying_type_t<ECacheStoreFlags>>;
+			FPlatformAtomics::AtomicStore(reinterpret_cast<InternalType*>(&CacheFlags), static_cast<InternalType>(NewCacheFlags));
+		}
+		inline ECacheStoreFlags GetCacheFlags() const 
+		{
+			return (ECacheStoreFlags)static_cast<std::underlying_type_t<ECacheStoreFlags>>(FPlatformAtomics::AtomicRead(reinterpret_cast<const std::make_signed_t<std::underlying_type_t<ECacheStoreFlags>>*>(&CacheFlags)));
+		}
+
 		ILegacyCacheStore* Cache{};
 		ECacheStoreFlags CacheFlags{};
 		ECacheStoreNodeFlags NodeFlags{};
@@ -135,7 +146,7 @@ private:
 	};
 
 	mutable FRWLock NodesLock;
-	ECacheStoreNodeFlags CombinedNodeFlags{};
+	std::atomic<ECacheStoreNodeFlags> CombinedNodeFlags{};
 	TArray<FCacheStoreNode, TInlineAllocator<8>> Nodes;
 	IMemoryCacheStore* MemoryCache;
 	FCacheStats CacheStats;
@@ -308,10 +319,12 @@ void FCacheStoreHierarchy::Add(ILegacyCacheStore* CacheStore, ECacheStoreFlags F
 
 void FCacheStoreHierarchy::SetFlags(ILegacyCacheStore* CacheStore, ECacheStoreFlags Flags)
 {
-	FWriteScopeLock Lock(NodesLock);
+	// Doing this as a read lock and relying on atomic access to the Nodes' CacheFlags and the hierarchy's CombinedNodeFlags
+	// to prevent a partial read/write from two concurrent executions.
+	FReadScopeLock Lock(NodesLock);
 	FCacheStoreNode* Node = Algo::FindBy(Nodes, CacheStore, &FCacheStoreNode::Cache);
 	checkf(!!Node, TEXT("Attempting to set flags on a cache store that is not registered to the hierarchy."));
-	Node->CacheFlags = Flags;
+	Node->SetCacheFlags(Flags);
 	UpdateNodeFlags();
 }
 
@@ -340,7 +353,7 @@ bool FCacheStoreHierarchy::HasAllFlags(ECacheStoreFlags Flags) const
 	{
 		if (Node.Cache != MemoryCache)
 		{
-			CombinedFlags |= Node.CacheFlags;
+			CombinedFlags |= Node.GetCacheFlags();
 		}
 	}
 	return EnumHasAllFlags(CombinedFlags, Flags);
@@ -348,16 +361,19 @@ bool FCacheStoreHierarchy::HasAllFlags(ECacheStoreFlags Flags) const
 
 void FCacheStoreHierarchy::UpdateNodeFlags()
 {
+	ECacheStoreNodeFlags OriginalCombinedFlags = CombinedNodeFlags.load();
+
 	ECacheStoreNodeFlags StoreFlags = ECacheStoreNodeFlags::None;
 	for (int32 Index = 0, Count = Nodes.Num(); Index < Count; ++Index)
 	{
 		FCacheStoreNode& Node = Nodes[Index];
 		Node.NodeFlags = StoreFlags;
-		if (EnumHasAllFlags(Node.CacheFlags, ECacheStoreFlags::Store | ECacheStoreFlags::Local))
+		ECacheStoreFlags CacheFlags = Node.GetCacheFlags();
+		if (EnumHasAllFlags(CacheFlags, ECacheStoreFlags::Store | ECacheStoreFlags::Local))
 		{
 			StoreFlags |= ECacheStoreNodeFlags::HasStoreLocalNode;
 		}
-		if (EnumHasAllFlags(Node.CacheFlags, ECacheStoreFlags::Store | ECacheStoreFlags::Remote))
+		if (EnumHasAllFlags(CacheFlags, ECacheStoreFlags::Store | ECacheStoreFlags::Remote))
 		{
 			StoreFlags |= ECacheStoreNodeFlags::HasStoreRemoteNode;
 		}
@@ -368,17 +384,21 @@ void FCacheStoreHierarchy::UpdateNodeFlags()
 	{
 		FCacheStoreNode& Node = Nodes[Index];
 		Node.NodeFlags |= QueryFlags;
-		if (EnumHasAllFlags(Node.CacheFlags, ECacheStoreFlags::Query | ECacheStoreFlags::Local))
+		ECacheStoreFlags CacheFlags = Node.GetCacheFlags();
+		if (EnumHasAllFlags(CacheFlags, ECacheStoreFlags::Query | ECacheStoreFlags::Local))
 		{
 			QueryFlags |= ECacheStoreNodeFlags::HasQueryLocalNode;
 		}
-		if (EnumHasAllFlags(Node.CacheFlags, ECacheStoreFlags::Query | ECacheStoreFlags::Remote))
+		if (EnumHasAllFlags(CacheFlags, ECacheStoreFlags::Query | ECacheStoreFlags::Remote))
 		{
 			QueryFlags |= ECacheStoreNodeFlags::HasQueryRemoteNode;
 		}
 	}
 
-	CombinedNodeFlags = StoreFlags | QueryFlags;
+	if (!CombinedNodeFlags.compare_exchange_strong(OriginalCombinedFlags, StoreFlags | QueryFlags))
+	{
+		UpdateNodeFlags();
+	}
 }
 
 ICacheStoreStats* FCacheStoreHierarchy::CreateStats(ILegacyCacheStore* CacheStore, ECacheStoreFlags Flags, FStringView Type, FStringView Name, FStringView Path)
@@ -640,7 +660,7 @@ void FCacheStoreHierarchy::TPutBatch<Params>::Begin(
 	IRequestOwner& InOwner,
 	FOnPutComplete&& InOnComplete)
 {
-	if (InRequests.IsEmpty() || !EnumHasAnyFlags(InHierarchy.CombinedNodeFlags, ECacheStoreNodeFlags::HasStoreNode))
+	if (InRequests.IsEmpty() || !EnumHasAnyFlags(InHierarchy.CombinedNodeFlags.load(), ECacheStoreNodeFlags::HasStoreNode))
 	{
 		return CompleteWithStatus(InRequests, InOnComplete, EStatus::Error);
 	}
@@ -688,7 +708,8 @@ bool FCacheStoreHierarchy::TPutBatch<Params>::DispatchGetRequests()
 	NodeGetIndex = NodePutIndex;
 
 	const FCacheStoreNode& Node = Hierarchy.Nodes[NodeGetIndex];
-	if (!EnumHasAnyFlags(Node.CacheFlags, ECacheStoreFlags::StopPutStore))
+	ECacheStoreFlags CacheFlags = Node.GetCacheFlags();
+	if (!EnumHasAnyFlags(CacheFlags, ECacheStoreFlags::StopPutStore))
 	{
 		return false;
 	}
@@ -699,7 +720,7 @@ bool FCacheStoreHierarchy::TPutBatch<Params>::DispatchGetRequests()
 	int32 RequestIndex = 0;
 	for (const FPutRequest& Request : Requests)
 	{
-		if (!States[RequestIndex].bStop && CanQuery(GetCombinedPolicy(Request.Policy), Node.CacheFlags))
+		if (!States[RequestIndex].bStop && CanQuery(GetCombinedPolicy(Request.Policy), CacheFlags))
 		{
 			NodeRequests.Add(MakeGetRequest(Request, RequestIndex));
 		}
@@ -746,7 +767,8 @@ template <typename Params>
 bool FCacheStoreHierarchy::TPutBatch<Params>::DispatchPutRequests()
 {
 	const FCacheStoreNode& Node = Hierarchy.Nodes[NodePutIndex];
-	if (!EnumHasAnyFlags(Node.CacheFlags, ECacheStoreFlags::Store))
+	ECacheStoreFlags CacheFlags = Node.GetCacheFlags();
+	if (!EnumHasAnyFlags(CacheFlags, ECacheStoreFlags::Store))
 	{
 		return false;
 	}
@@ -762,7 +784,7 @@ bool FCacheStoreHierarchy::TPutBatch<Params>::DispatchPutRequests()
 	for (const FPutRequest& Request : Requests)
 	{
 		const FRequestState& State = States[RequestIndex];
-		if (!State.bStop && CanStore(GetCombinedPolicy(Request.Policy), Node.CacheFlags))
+		if (!State.bStop && CanStore(GetCombinedPolicy(Request.Policy), CacheFlags))
 		{
 			(State.bFinished ? AsyncNodeRequests : NodeRequests).Add_GetRef(Request).UserData = uint64(RequestIndex);
 		}
@@ -801,7 +823,7 @@ void FCacheStoreHierarchy::TPutBatch<Params>::CompletePutRequest(FPutResponse&&
 		bool bCanQuery;
 		{
 			FReadScopeLock Lock(Hierarchy.NodesLock);
-			bCanQuery = EnumHasAnyFlags(Hierarchy.Nodes[NodePutIndex].CacheFlags, ECacheStoreFlags::Query);
+			bCanQuery = EnumHasAnyFlags(Hierarchy.Nodes[NodePutIndex].GetCacheFlags(), ECacheStoreFlags::Query);
 		}
 		if (bCanQuery)
 		{
@@ -900,7 +922,7 @@ void FCacheStoreHierarchy::TGetBatch<Params>::Begin(
 	IRequestOwner& InOwner,
 	FOnGetComplete&& InOnComplete)
 {
-	if (InRequests.IsEmpty() || !EnumHasAnyFlags(InHierarchy.CombinedNodeFlags, ECacheStoreNodeFlags::HasQueryNode))
+	if (InRequests.IsEmpty() || !EnumHasAnyFlags(InHierarchy.CombinedNodeFlags.load(), ECacheStoreNodeFlags::HasQueryNode))
 	{
 		return CompleteWithStatus(InRequests, InOnComplete, EStatus::Error);
 	}
@@ -931,6 +953,7 @@ void FCacheStoreHierarchy::TGetBatch<Params>::DispatchRequests()
 		}
 
 		const FCacheStoreNode& Node = Hierarchy.Nodes[NodeIndex];
+		ECacheStoreFlags CacheFlags = Node.GetCacheFlags();
 
 		uint64 StateIndex = 0;
 		for (FState& State : States)
@@ -944,12 +967,12 @@ void FCacheStoreHierarchy::TGetBatch<Params>::DispatchRequests()
 			const FGetResponse& Response = State.Response;
 			if (Response.Status == EStatus::Ok)
 			{
-				if (HasResponseData(Response) && CanStore(GetCombinedPolicy(Request.Policy), Node.CacheFlags))
+				if (HasResponseData(Response) && CanStore(GetCombinedPolicy(Request.Policy), CacheFlags))
 				{
 					AsyncNodeRequests.Add(MakePutRequest(Response, Request));
 					++State.NodeIndex;
 				}
-				else if (EnumHasAnyFlags(Node.CacheFlags, ECacheStoreFlags::StopGetStore) && CanQuery(GetCombinedPolicy(Request.Policy), Node.CacheFlags))
+				else if (EnumHasAnyFlags(CacheFlags, ECacheStoreFlags::StopGetStore) && CanQuery(GetCombinedPolicy(Request.Policy), CacheFlags))
 				{
 					NodeRequests.Add({Request.Name, Request.Key, AddPolicy(Request.Policy, ECachePolicy::SkipData), StateIndex});
 				}
@@ -960,7 +983,7 @@ void FCacheStoreHierarchy::TGetBatch<Params>::DispatchRequests()
 			}
 			else
 			{
-				if (const ECachePolicy CombinedPolicy = GetCombinedPolicy(Request.Policy); CanQuery(CombinedPolicy, Node.CacheFlags))
+				if (const ECachePolicy CombinedPolicy = GetCombinedPolicy(Request.Policy); CanQuery(CombinedPolicy, CacheFlags))
 				{
 					auto Policy = Request.Policy;
 					if (CanStoreIfOk(CombinedPolicy, Node.NodeFlags))
@@ -1052,7 +1075,7 @@ void FCacheStoreHierarchy::TGetBatch<Params>::CompleteRequest(FGetResponse&& Res
 		for (int32 PutNodeIndex = 0; PutNodeIndex < PreviousNodeIndex; ++PutNodeIndex)
 		{
 			const FCacheStoreNode& PutNode = Hierarchy.Nodes[PutNodeIndex];
-			if (CanStore(GetCombinedPolicy(State.Request.Policy), PutNode.CacheFlags))
+			if (CanStore(GetCombinedPolicy(State.Request.Policy), PutNode.GetCacheFlags()))
 			{
 				FRequestBarrier AsyncBarrier(AsyncOwner);
 				Invoke(Put(), PutNode.AsyncCache, MakeArrayView(&PutRequest, 1), AsyncOwner, [](auto&&){});
@@ -1068,7 +1091,7 @@ void FCacheStoreHierarchy::TGetBatch<Params>::CompleteRequest(FGetResponse&& Res
 
 	if (Response.Status == EStatus::Ok)
 	{
-		if (EnumHasAnyFlags(Node.CacheFlags, ECacheStoreFlags::StopGetStore))
+		if (EnumHasAnyFlags(Node.GetCacheFlags(), ECacheStoreFlags::StopGetStore))
 		{
 			// Never store to later nodes.
 			State.Request.Policy = RemovePolicy(State.Request.Policy, ECachePolicy::Default);
@@ -1443,7 +1466,7 @@ void FCacheStoreHierarchy::FGetChunksBatch::Begin(
 	IRequestOwner& InOwner,
 	FOnCacheGetChunkComplete&& InOnComplete)
 {
-	if (InRequests.IsEmpty() || !EnumHasAnyFlags(InHierarchy.CombinedNodeFlags, ECacheStoreNodeFlags::HasQueryNode))
+	if (InRequests.IsEmpty() || !EnumHasAnyFlags(InHierarchy.CombinedNodeFlags.load(), ECacheStoreNodeFlags::HasQueryNode))
 	{
 		return CompleteWithStatus(InRequests, InOnComplete, EStatus::Error);
 	}
@@ -1463,12 +1486,13 @@ void FCacheStoreHierarchy::FGetChunksBatch::DispatchRequests()
 	for (const int32 NodeCount = Hierarchy.Nodes.Num(); NodeIndex < NodeCount && !Owner.IsCanceled(); ++NodeIndex)
 	{
 		const FCacheStoreNode& Node = Hierarchy.Nodes[NodeIndex];
+		ECacheStoreFlags CacheFlags = Node.GetCacheFlags();
 
 		uint64 StateIndex = 0;
 		for (const FState& State : States)
 		{
 			const FCacheGetChunkRequest& Request = State.Request;
-			if (State.Status == EStatus::Error && CanQuery(Request.Policy, Node.CacheFlags))
+			if (State.Status == EStatus::Error && CanQuery(Request.Policy, CacheFlags))
 			{
 				NodeRequests.Add_GetRef(Request).UserData = StateIndex;
 			}
