DDC: Exposed metadata more extensively in the build API

diff --git a/Engine/Source/Developer/DerivedDataCache/Private/DerivedDataBuildJobContext.cpp b/Engine/Source/Developer/DerivedDataCache/Private/DerivedDataBuildJobContext.cpp
--- a/Engine/Source/Developer/DerivedDataCache/Private/DerivedDataBuildJobContext.cpp
+++ b/Engine/Source/Developer/DerivedDataCache/Private/DerivedDataBuildJobContext.cpp
@@ -11,6 +11,7 @@
 #include "DerivedDataRequestOwner.h"
 #include "DerivedDataValue.h"
 #include "Hash/Blake3.h"
+#include "Misc/AsciiSet.h"
 #include "Misc/StringBuilder.h"
 #include "UObject/NameTypes.h"
 
@@ -101,6 +102,14 @@ FSharedBuffer FBuildJobContext::FindInput(FUtf8StringView Key) const
 	return FSharedBuffer();
 }
 
+void FBuildJobContext::AddMeta(FUtf8StringView Key, const FCbField& Meta)
+{
+	constexpr FAsciiSet Valid("ABCDEFGHIJKLMNOPQRSTUVWXYZabcdefghijklmnopqrstuvwxyz0123456789");
+	checkf(!Key.IsEmpty() && FAsciiSet::HasOnly(Key, Valid),
+		TEXT("Key '%s' for metadata must be alphanumeric and non-empty."), *WriteToString<32>(Key));
+	OutputBuilder.AddMeta(Key, Meta);
+}
+
 void FBuildJobContext::AddValue(const FValueId& Id, const FValue& Value)
 {
 	OutputBuilder.AddValue(Id, Value);
diff --git a/Engine/Source/Developer/DerivedDataCache/Private/DerivedDataBuildJobContext.h b/Engine/Source/Developer/DerivedDataCache/Private/DerivedDataBuildJobContext.h
--- a/Engine/Source/Developer/DerivedDataCache/Private/DerivedDataBuildJobContext.h
+++ b/Engine/Source/Developer/DerivedDataCache/Private/DerivedDataBuildJobContext.h
@@ -59,6 +59,8 @@ private:
 	FCbObject FindConstant(FUtf8StringView Key) const final;
 	FSharedBuffer FindInput(FUtf8StringView Key) const final;
 
+	void AddMeta(FUtf8StringView Key, const FCbField& Meta) final;
+
 	void AddValue(const FValueId& Id, const FValue& Value) final;
 	void AddValue(const FValueId& Id, const FCompressedBuffer& Buffer) final;
 	void AddValue(const FValueId& Id, const FCompositeBuffer& Buffer, uint64 BlockSize) final;
diff --git a/Engine/Source/Developer/DerivedDataCache/Private/DerivedDataBuildOutput.cpp b/Engine/Source/Developer/DerivedDataCache/Private/DerivedDataBuildOutput.cpp
--- a/Engine/Source/Developer/DerivedDataCache/Private/DerivedDataBuildOutput.cpp
+++ b/Engine/Source/Developer/DerivedDataCache/Private/DerivedDataBuildOutput.cpp
@@ -94,13 +94,14 @@ public:
 	{
 		checkf(!Name.IsEmpty(), TEXT("A build output requires a non-empty name."));
 		AssertValidBuildFunctionName(Function, Name);
+		MetaWriter.BeginObject();
 		MessageWriter.BeginArray();
 		LogWriter.BeginArray();
 	}
 
 	~FBuildOutputBuilderInternal() final = default;
 
-	void SetMeta(FCbObject&& InMeta) final { Meta = MoveTemp(InMeta); Meta.MakeOwned(); }
+	void AddMeta(FUtf8StringView Key, const FCbField& Meta) final;
 	void AddValue(const FValueId& Id, const FValue& Value) final;
 	void AddMessage(const FBuildOutputMessage& Message) final;
 	void AddLog(const FBuildOutputLog& Log) final;
@@ -109,8 +110,8 @@ public:
 
 	FSharedString Name;
 	FUtf8SharedString Function;
-	FCbObject Meta;
 	TArray<FValueWithId> Values;
+	FCbWriter MetaWriter;
 	FCbWriter MessageWriter;
 	FCbWriter LogWriter;
 	bool bHasMessages = false;
@@ -144,6 +145,7 @@ public:
 	const FUtf8SharedString& GetFunction() const final { return Function; }
 
 	const FCbObject& GetMeta() const final { return Meta; }
+	bool FindMeta(FUtf8StringView Key, FCbField& OutMeta) const final;
 
 	const FValueWithId& GetValue(const FValueId& Id) const final;
 	TConstArrayView<FValueWithId> GetValues() const final { return Values; }
@@ -234,6 +236,12 @@ FBuildOutputInternal::FBuildOutputInternal(
 	bOutIsValid = TryLoad();
 }
 
+bool FBuildOutputInternal::FindMeta(FUtf8StringView Key, FCbField& OutMeta) const
+{
+	OutMeta = Meta.Find(Key);
+	return !!OutMeta;
+}
+
 const FValueWithId& FBuildOutputInternal::GetValue(const FValueId& Id) const
 {
 	const int32 Index = Algo::BinarySearchBy(Values, Id, &FValueWithId::GetId);
@@ -368,6 +376,12 @@ void FBuildOutputInternal::Save(FCacheRecordBuilder& RecordBuilder) const
 
 ///////////////////////////////////////////////////////////////////////////////////////////////////////////////////////
 
+void FBuildOutputBuilderInternal::AddMeta(FUtf8StringView Key, const FCbField& Meta)
+{
+	checkf(!Key.IsEmpty(), TEXT("Key for metadata must be non-empty."));
+	MetaWriter.AddField(Key, Meta);
+}
+
 void FBuildOutputBuilderInternal::AddValue(const FValueId& Id, const FValue& Value)
 {
 	checkf(Id.IsValid(), TEXT("Null value added in output for build of '%s' by %s."),
@@ -409,6 +423,7 @@ FBuildOutput FBuildOutputBuilderInternal::Build()
 	{
 		Values.Empty();
 	}
+	MetaWriter.EndObject();
 	MessageWriter.EndArray();
 	LogWriter.EndArray();
 	FCbObject Output;
@@ -429,7 +444,7 @@ FBuildOutput FBuildOutputBuilderInternal::Build()
 		Writer.EndObject();
 		Output = Writer.Save().AsObject();
 	}
-	return CreateBuildOutput(new FBuildOutputInternal(MoveTemp(Name), MoveTemp(Function), MoveTemp(Meta), MoveTemp(Output), MoveTemp(Values)));
+	return CreateBuildOutput(new FBuildOutputInternal(MoveTemp(Name), MoveTemp(Function), MetaWriter.Save().AsObject(), MoveTemp(Output), MoveTemp(Values)));
 }
 
 ///////////////////////////////////////////////////////////////////////////////////////////////////////////////////////
diff --git a/Engine/Source/Developer/DerivedDataCache/Public/DerivedDataBuildFunction.h b/Engine/Source/Developer/DerivedDataCache/Public/DerivedDataBuildFunction.h
--- a/Engine/Source/Developer/DerivedDataCache/Public/DerivedDataBuildFunction.h
+++ b/Engine/Source/Developer/DerivedDataCache/Public/DerivedDataBuildFunction.h
@@ -7,6 +7,7 @@
 #include "DerivedDataSharedStringFwd.h"
 #include "Memory/MemoryFwd.h"
 
+class FCbField;
 class FCbObject;
 class FCompressedBuffer;
 
@@ -91,6 +92,13 @@ public:
 	/** Returns the input with the matching key, or a null buffer if not found. */
 	virtual FSharedBuffer FindInput(FUtf8StringView Key) const = 0;
 
+	/**
+	 * Adds metadata to the build output. Must have a unique non-empty alphanumeric key.
+	 *
+	 * Metadata is not required to be deterministic or based only on inputs.
+	 */
+	virtual void AddMeta(FUtf8StringView Key, const FCbField& Meta) = 0;
+
 	/** Adds a value to the build output. Must have a non-null buffer and a unique ID. */
 	virtual void AddValue(const FValueId& Id, const FValue& Value) = 0;
 	virtual void AddValue(const FValueId& Id, const FCompressedBuffer& Buffer) = 0;
diff --git a/Engine/Source/Developer/DerivedDataCache/Public/DerivedDataBuildOutput.h b/Engine/Source/Developer/DerivedDataCache/Public/DerivedDataBuildOutput.h
--- a/Engine/Source/Developer/DerivedDataCache/Public/DerivedDataBuildOutput.h
+++ b/Engine/Source/Developer/DerivedDataCache/Public/DerivedDataBuildOutput.h
@@ -15,6 +15,7 @@
 
 #define UE_API DERIVEDDATACACHE_API
 
+class FCbField;
 class FCbObject;
 class FCbWriter;
 
@@ -39,6 +40,7 @@ public:
 	virtual const FSharedString& GetName() const = 0;
 	virtual const FUtf8SharedString& GetFunction() const = 0;
 	virtual const FCbObject& GetMeta() const = 0;
+	virtual bool FindMeta(FUtf8StringView Key, FCbField& OutMeta) const = 0;
 	virtual const FValueWithId& GetValue(const FValueId& Id) const = 0;
 	virtual TConstArrayView<FValueWithId> GetValues() const = 0;
 	virtual TConstArrayView<FBuildOutputMessage> GetMessages() const = 0;
@@ -57,7 +59,7 @@ class IBuildOutputBuilderInternal
 {
 public:
 	virtual ~IBuildOutputBuilderInternal() = default;
-	virtual void SetMeta(FCbObject&& Meta) = 0;
+	virtual void AddMeta(FUtf8StringView Key, const FCbField& Meta) = 0;
 	virtual void AddValue(const FValueId& Id, const FValue& Value) = 0;
 	virtual void AddMessage(const FBuildOutputMessage& Message) = 0;
 	virtual void AddLog(const FBuildOutputLog& Log) = 0;
@@ -130,6 +132,9 @@ public:
 	/** Returns the optional metadata. */
 	inline const FCbObject& GetMeta() const { return Output->GetMeta(); }
 
+	/** Finds the metadata matching the key. Returns true if found. */
+	inline bool FindMeta(FUtf8StringView Key, FCbField& OutMeta) const { return Output->FindMeta(Key, OutMeta); }
+
 	/** Returns the value matching the ID. Null if no match. Buffer is null if skipped. */
 	inline const FValueWithId& GetValue(const FValueId& Id) const { return Output->GetValue(Id); }
 
@@ -193,10 +198,10 @@ private:
 class FBuildOutputBuilder
 {
 public:
-	/** Set the metadata for the output. Holds a reference and is cloned if not owned. */
-	inline void SetMeta(FCbObject&& Meta)
+	/** Add metadata to the output. The key must be unique in this output. */
+	inline void AddMeta(FUtf8StringView Key, const FCbField& Meta)
 	{
-		return OutputBuilder->SetMeta(MoveTemp(Meta));
+		return OutputBuilder->AddMeta(Key, Meta);
 	}
 
 	/** Add a value to the output. The ID must be unique in this output. */
