Adjusting logs related to zen usage to:
- Remove repeated log related to fetching cache stats
- Ensure that the Zen cache usage has log lines explaining the status of the caches it attempts to connect to (success or failure)

diff --git a/Engine/Source/Developer/DerivedDataCache/Private/ZenCacheStore.cpp b/Engine/Source/Developer/DerivedDataCache/Private/ZenCacheStore.cpp
--- a/Engine/Source/Developer/DerivedDataCache/Private/ZenCacheStore.cpp
+++ b/Engine/Source/Developer/DerivedDataCache/Private/ZenCacheStore.cpp
@@ -109,7 +109,7 @@ public:
 
 	~FZenCacheStore() final;
 
-	inline FString GetName() const { return ZenService.GetInstance().GetURL(); }
+	inline const FString& GetName() const { return NodeName; }
 
 	/**
 	 * Checks if cache service is usable (reachable and accessible).
@@ -212,6 +212,7 @@ private:
 	class FCbPackageReceiver;
 	class FAsyncCbPackageReceiver;
 
+	FString NodeName;
 	FString Namespace;
 	UE::Zen::FScopeZenService ZenService;
 	ICacheStoreOwner* StoreOwner = nullptr;
@@ -1316,8 +1317,9 @@ public:
 	FHealthReceiver(const FHealthReceiver&) = delete;
 	FHealthReceiver& operator=(const FHealthReceiver&) = delete;
 
-	explicit FHealthReceiver(EHealth& OutHealth, IHttpReceiver* InNext = nullptr)
+	explicit FHealthReceiver(EHealth& OutHealth, FString* OutResponseBody = nullptr, IHttpReceiver* InNext = nullptr)
 		: Health(OutHealth)
+		, ResponseBody(OutResponseBody)
 		, Next(InNext)
 	{
 		Health = EHealth::Unknown;
@@ -1340,11 +1342,17 @@ private:
 		{
 			Health = EHealth::Error;
 		}
+
+		if (ResponseBody)
+		{
+			*ResponseBody = *WriteToString<64>(ResponseStringView);
+		}
 		return Next;
 	}
 
 private:
 	EHealth& Health;
+	FString* ResponseBody;
 	IHttpReceiver* Next;
 	TArray64<uint8> BodyArray;
 	FHttpByteArrayReceiver BodyReceiver{ BodyArray, this };
@@ -1364,7 +1372,7 @@ public:
 		: Request(MoveTemp(InRequest))
 		, Owner(InOwner)
 		, ZenServiceInstance(InZenServiceInstance)
-		, BaseReceiver(Health, this)
+		, BaseReceiver(Health, nullptr, this)
 		, OnHealthComplete(MoveTemp(InOnHealthComplete))
 	{
 		Request->SendAsync(this, Response);
@@ -1608,6 +1616,7 @@ FZenCacheStore::~FZenCacheStore()
 
 void FZenCacheStore::Initialize(const FZenCacheStoreParams& Params)
 {
+	NodeName = Params.Name;
 	LastPerformanceEvaluationTicks.store(FDateTime::UtcNow().GetTicks(), std::memory_order_relaxed);
 	LastStorageSizeUpdateTicks.store(FDateTime::UtcNow().GetTicks(), std::memory_order_relaxed);
 	Namespace = Params.Namespace;
@@ -1640,10 +1649,36 @@ void FZenCacheStore::Initialize(const FZenCacheStoreParams& Params)
 		ReadinessRequest->SetMethod(EHttpMethod::Get);
 		ReadinessRequest->AddAcceptType(EHttpMediaType::Text);
 		EHealth Health = EHealth::Unknown;
-		FHealthReceiver HealthReceiver(Health);
+		FString ResponseString;
+		FHealthReceiver HealthReceiver(Health, &ResponseString);
 		THttpUniquePtr<IHttpResponse> ReadinessResponse;
 		ReadinessRequest->Send(&HealthReceiver, ReadinessResponse);
 		bReady = Health == EHealth::Ok; // -V547
+
+		if (ReadinessResponse->GetErrorCode() == EHttpErrorCode::None &&
+			(ReadinessResponse->GetStatusCode() >= 200 && ReadinessResponse->GetStatusCode() <= 299))
+		{
+			UE_LOG(LogDerivedDataCache, Display,
+				TEXT("%s: Using ZenServer HTTP service at %s with namespace %s status: %s."),
+				*GetName(), ZenService.GetInstance().GetURL(), *Namespace, *ResponseString);
+		}
+		else
+		{
+			if (ZenService.GetInstance().IsServiceRunningLocally())
+			{
+				UE_LOG(LogDerivedDataCache, Warning,
+					TEXT("%s: Unable to reach ZenServer HTTP service at %s with namespace %s. Status: %d . Response: %s"),
+					*GetName(), ZenService.GetInstance().GetURL(), *Namespace, ReadinessResponse->GetStatusCode(),
+					*ResponseString);
+			}
+			else
+			{
+				UE_LOG(LogDerivedDataCache, Display,
+					TEXT("%s: Unable to reach ZenServer HTTP service at %s with namespace %s. Status: %d . Response: %s"),
+					*GetName(), ZenService.GetInstance().GetURL(), *Namespace, ReadinessResponse->GetStatusCode(),
+					*ResponseString);
+			}
+		}
 	}
 
 	FHttpClientParams ClientParams;
diff --git a/Engine/Source/Developer/Zen/Private/ZenServerInterface.cpp b/Engine/Source/Developer/Zen/Private/ZenServerInterface.cpp
--- a/Engine/Source/Developer/Zen/Private/ZenServerInterface.cpp
+++ b/Engine/Source/Developer/Zen/Private/ZenServerInterface.cpp
@@ -2384,24 +2384,23 @@ FZenServiceInstance::GetCacheStats(FZenCacheStats& Stats)
 #endif
 		// We've not got any requests in flight and we've met a given time requirement for requests
 	CacheStatsRequest = Async(ThreadPool, [Request]
-		{
+	{
 		check(Request != nullptr);
 		Request->Reset();
 
-			TArray64<uint8> GetBuffer;
+		TArray64<uint8> GetBuffer;
 		FZenHttpRequest::Result Result = Request->PerformBlockingDownload(TEXTVIEW("/stats/z$"), &GetBuffer, Zen::EContentType::CbObject);
 
-			FZenCacheStats Stats;
+		FZenCacheStats Stats;
 
 		if (Result == Zen::FZenHttpRequest::Result::Success && Request->GetResponseCode() == 200)
-			{
-			UE_LOG(LogZenServiceInstance, Display, TEXT("FZenServiceInstance::GetCacheStats (updating stats from response)"));
-				FCbFieldView RootView(GetBuffer.GetData());
-				Stats.bIsValid = LoadFromCompactBinary(RootView, Stats);
-			}
+		{
+			FCbFieldView RootView(GetBuffer.GetData());
+			Stats.bIsValid = LoadFromCompactBinary(RootView, Stats);
+		}
 
-			return Stats;
-		});
+		return Stats;
+	});
 
 	return Stats.bIsValid;
 }
