Fix log format convention for HttpCacheStore to match ZenCacheStore - log lines should be prefixed by node name, not domain or URL.

diff --git a/Engine/Source/Developer/DerivedDataCache/Private/HttpCacheStore.cpp b/Engine/Source/Developer/DerivedDataCache/Private/HttpCacheStore.cpp
--- a/Engine/Source/Developer/DerivedDataCache/Private/HttpCacheStore.cpp
+++ b/Engine/Source/Developer/DerivedDataCache/Private/HttpCacheStore.cpp
@@ -422,6 +422,7 @@ public:
 	}
 
 private:
+	FString NodeName;
 	FString Domain;
 	FString Namespace;
 	FString OAuthProvider;
@@ -834,7 +835,7 @@ public:
 	FHealthCheckOp(FHttpCacheStore& CacheStore, IHttpClient& Client)
 		: Operation(Client.TryCreateRequest({}))
 		, Owner(EPriority::High)
-		, Domain(*CacheStore.Domain)
+		, NodeName(*CacheStore.NodeName)
 	{
 		Operation.SetUri(WriteToAnsiString<256>(CacheStore.EffectiveDomain, ANSITEXTVIEW("/health/ready")));
 		Operation.SendAsync(Owner, []{});
@@ -846,13 +847,13 @@ public:
 		const FString Body = Operation.GetBodyAsString();
 		if (Operation.GetStatusCode() == 200)
 		{
-			UE_LOG(LogDerivedDataCache, Display, TEXT("%s: HTTP DDC: %s"), Domain, *Body);
+			UE_LOG(LogDerivedDataCache, Display, TEXT("%s: HTTP DDC: %s"), NodeName, *Body);
 			return true;
 		}
 		else
 		{
 			UE_LOG(LogDerivedDataCache, Warning, TEXT("%s: Unable to reach HTTP DDC at %s. %s"),
-				Domain, *WriteToString<256>(Operation), *Body);
+				NodeName, *WriteToString<256>(Operation), *Body);
 			return false;
 		}
 	}
@@ -860,7 +861,7 @@ public:
 private:
 	FHttpOperation Operation;
 	FRequestOwner Owner;
-	const TCHAR* Domain;
+	const TCHAR* NodeName;
 };
 
 //----------------------------------------------------------------------------------------------------------
@@ -1055,7 +1056,7 @@ void FHttpCacheStore::FPutPackageOp::BeginPutBlobs(FCbPackage&& Package, FCacheP
 		if (Response.Status == EStatus::Error)
 		{
 			UE_LOG(LogDerivedDataCache, Display, TEXT("%s: Failed to put reference object for put of %s from '%s'"),
-				*CacheStore.Domain, *WriteToString<96>(Key), *Name);
+				*CacheStore.NodeName, *WriteToString<96>(Key), *Name);
 		}
 		EndPut(Response.Status);
 		return;
@@ -1103,7 +1104,7 @@ void FHttpCacheStore::FPutPackageOp::BeginPutBlobs(FCbPackage&& Package, FCacheP
 				bExpectedHashesSerialized = true;
 			}
 			UE_LOG(LogDerivedDataCache, Display, TEXT("%s: Server reported needed hash '%s' that is outside the set of expected hashes (%s) for put of %s from '%s'"),
-				*CacheStore.Domain, *WriteToString<96>(NeededBlobHash), *ExpectedHashes, *WriteToString<96>(Key), *Name);
+				*CacheStore.NodeName, *WriteToString<96>(NeededBlobHash), *ExpectedHashes, *WriteToString<96>(Key), *Name);
 		}
 	}
 
@@ -1173,7 +1174,7 @@ void FHttpCacheStore::FPutPackageOp::EndPutBlob(FHttpOperation* Operation, uint6
 		{
 			const uint32 FailedBlobUploads = TotalBlobUploads - LocalSuccessfulBlobUploads;
 			UE_LOG(LogDerivedDataCache, Log, TEXT("%s: Failed to put %d/%d blobs for put of %s from '%s'"),
-				*CacheStore.Domain, FailedBlobUploads, TotalBlobUploads, *WriteToString<96>(Key), *Name);
+				*CacheStore.NodeName, FailedBlobUploads, TotalBlobUploads, *WriteToString<96>(Key), *Name);
 			EndPut(EStatus::Error);
 		}
 	}
@@ -1184,7 +1185,7 @@ void FHttpCacheStore::FPutPackageOp::EndPutRefFinalize(FCachePutRefResponse&& Re
 	if (Response.Status == EStatus::Error)
 	{
 		UE_LOG(LogDerivedDataCache, Log, TEXT("%s: Failed to finalize reference object for put of %s from '%s'"),
-			*CacheStore.Domain, *WriteToString<96>(Key), *Name);
+			*CacheStore.NodeName, *WriteToString<96>(Key), *Name);
 	}
 
 	EndPut(Response.Status);
@@ -1289,7 +1290,7 @@ void FHttpCacheStore::FGetRecordOp::GetRecordOnly(const FCacheKey& InKey, const
 	{
 		UE_LOG(LogDerivedDataCache, VeryVerbose,
 			TEXT("%s: Skipped get of %s from '%s' because this cache store is not available"),
-			*CacheStore.Domain, *WriteToString<96>(Key), *Name);
+			*CacheStore.NodeName, *WriteToString<96>(Key), *Name);
 		return InOnComplete({FCacheRecordBuilder(Key).Build(), EStatus::Error});
 	}
 
@@ -1297,14 +1298,14 @@ void FHttpCacheStore::FGetRecordOp::GetRecordOnly(const FCacheKey& InKey, const
 	if (!EnumHasAnyFlags(RecordPolicy, ECachePolicy::QueryRemote))
 	{
 		UE_LOG(LogDerivedDataCache, VeryVerbose, TEXT("%s: Skipped get of %s from '%s' due to cache policy"),
-			*CacheStore.Domain, *WriteToString<96>(Key), *Name);
+			*CacheStore.NodeName, *WriteToString<96>(Key), *Name);
 		return InOnComplete({FCacheRecordBuilder(Key).Build(), EStatus::Error});
 	}
 
 	if (CacheStore.DebugOptions.ShouldSimulateGetMiss(Key))
 	{
 		UE_LOG(LogDerivedDataCache, Verbose, TEXT("%s: Simulated miss for get of %s from '%s'"),
-			*CacheStore.Domain, *WriteToString<96>(Key), *Name);
+			*CacheStore.NodeName, *WriteToString<96>(Key), *Name);
 		return InOnComplete({FCacheRecordBuilder(Key).Build(), EStatus::Error});
 	}
 
@@ -1368,7 +1369,7 @@ void FHttpCacheStore::FGetRecordOp::EndGetRef(TUniquePtr<FHttpOperation> Operati
 	if (StatusCode < 200 || StatusCode > 204)
 	{
 		UE_LOG(LogDerivedDataCache, Verbose, TEXT("%s: Cache miss with missing package for %s from '%s'"),
-			*CacheStore.Domain, *WriteToString<96>(Key), *Name);
+			*CacheStore.NodeName, *WriteToString<96>(Key), *Name);
 		return;
 	}
 
@@ -1377,7 +1378,7 @@ void FHttpCacheStore::FGetRecordOp::EndGetRef(TUniquePtr<FHttpOperation> Operati
 	if (ValidateCompactBinary(Body, ECbValidateMode::Default) != ECbValidateError::None)
 	{
 		UE_LOG(LogDerivedDataCache, Log, TEXT("%s: Cache miss with invalid package for %s from '%s'"),
-			*CacheStore.Domain, *WriteToString<96>(Key), *Name);
+			*CacheStore.NodeName, *WriteToString<96>(Key), *Name);
 		return;
 	}
 
@@ -1387,7 +1388,7 @@ void FHttpCacheStore::FGetRecordOp::EndGetRef(TUniquePtr<FHttpOperation> Operati
 	if (Record.IsNull())
 	{
 		UE_LOG(LogDerivedDataCache, Log, TEXT("%s: Cache miss with record load failure for %s from '%s'"),
-			*CacheStore.Domain, *WriteToString<96>(Key), *Name);
+			*CacheStore.NodeName, *WriteToString<96>(Key), *Name);
 		return;
 	}
 
@@ -1625,7 +1626,7 @@ void FHttpCacheStore::FGetRecordOp::EndGetValue(FHttpOperation& Operation, const
 		{
 			UE_LOG(LogDerivedDataCache, Display,
 				TEXT("%s: Cache miss with corrupted value %s with hash %s for %s from '%s'"),
-				*CacheStore.Domain, *WriteToString<32>(Value.GetId()), *WriteToString<48>(Value.GetRawHash()),
+				*CacheStore.NodeName, *WriteToString<32>(Value.GetId()), *WriteToString<48>(Value.GetRawHash()),
 				*WriteToString<96>(Key), *Name);
 			OnComplete({Value, EStatus::Error});
 		}
@@ -1638,7 +1639,7 @@ void FHttpCacheStore::FGetRecordOp::EndGetValue(FHttpOperation& Operation, const
 	{
 		UE_LOG(LogDerivedDataCache, Verbose,
 			TEXT("%s: Cache miss with missing value %s with hash %s for %s from '%s'"),
-			*CacheStore.Domain, *WriteToString<32>(Value.GetId()), *WriteToString<48>(Value.GetRawHash()),
+			*CacheStore.NodeName, *WriteToString<32>(Value.GetId()), *WriteToString<48>(Value.GetRawHash()),
 			*WriteToString<96>(Key), *Name);
 		OnComplete({Value, EStatus::Error});
 	}
@@ -1746,7 +1747,7 @@ void FHttpCacheStore::FGetRecordOp::EndGetValuesExist(FHttpOperation* Operation,
 				{
 					UE_LOG(LogDerivedDataCache, Verbose,
 						TEXT("%s: Cache exists miss with missing value %s with hash %s for %s from '%s'"),
-						*CacheStore.Domain, *WriteToString<32>(Value.GetId()),
+						*CacheStore.NodeName, *WriteToString<32>(Value.GetId()),
 						*WriteToString<48>(Value.GetRawHash()), *WriteToString<96>(Key), *Name);
 					OnComplete({Value, EStatus::Error});
 					It.RemoveCurrentSwap();
@@ -1762,7 +1763,7 @@ void FHttpCacheStore::FGetRecordOp::EndGetValuesExist(FHttpOperation* Operation,
 	{
 		UE_LOG(LogDerivedDataCache, Verbose,
 			TEXT("%s: %s value %s with hash %s for %s from '%s'"),
-			*CacheStore.Domain, DefaultMessage, *WriteToString<32>(Value.GetId()),
+			*CacheStore.NodeName, DefaultMessage, *WriteToString<32>(Value.GetId()),
 			*WriteToString<48>(Value.GetRawHash()), *WriteToString<96>(Key), *Name);
 		OnComplete({Value, DefaultStatus});
 	}
@@ -1848,7 +1849,7 @@ void FHttpCacheStore::FGetValueOp::BeginGetRef(TUniquePtr<FHttpOperation>&& Oper
 	if (UNLIKELY(!Operation))
 	{
 		UE_LOG(LogDerivedDataCache, Verbose, TEXT("%s: Cache miss with failed with canceled request for %s from '%s'"),
-			*CacheStore.Domain, *WriteToString<96>(Key), *Name);
+			*CacheStore.NodeName, *WriteToString<96>(Key), *Name);
 		EndGet({Name, Key, {}, EStatus::Canceled});
 		return;
 	}
@@ -1892,7 +1893,7 @@ void FHttpCacheStore::FGetValueOp::EndGetRef(FHttpOperation& Operation)
 	if (StatusCode < 200 || StatusCode > 204)
 	{
 		UE_LOG(LogDerivedDataCache, Verbose, TEXT("%s: Cache miss with failed HTTP request for %s from '%s'"),
-			*CacheStore.Domain, *WriteToString<96>(Key), *Name);
+			*CacheStore.NodeName, *WriteToString<96>(Key), *Name);
 		return EndGet({Name, Key, {}, EStatus::Error});
 	}
 
@@ -1903,7 +1904,7 @@ void FHttpCacheStore::FGetValueOp::EndGetRef(FHttpOperation& Operation)
 		if (ValidateCompactBinary(Body, ECbValidateMode::Default) != ECbValidateError::None)
 		{
 			UE_LOG(LogDerivedDataCache, Display, TEXT("%s: Cache miss with invalid package for %s from '%s'"),
-				*CacheStore.Domain, *WriteToString<96>(Key), *Name);
+				*CacheStore.NodeName, *WriteToString<96>(Key), *Name);
 			return EndGet({Name, Key, {}, EStatus::Error});
 		}
 
@@ -1913,7 +1914,7 @@ void FHttpCacheStore::FGetValueOp::EndGetRef(FHttpOperation& Operation)
 		if (RawHash.IsZero() || RawSize == MAX_uint64)
 		{
 			UE_LOG(LogDerivedDataCache, Display, TEXT("%s: Cache miss with invalid value for %s from '%s'"),
-				*CacheStore.Domain, *WriteToString<96>(Key), *Name);
+				*CacheStore.NodeName, *WriteToString<96>(Key), *Name);
 			return EndGet({Name, Key, {}, EStatus::Error});
 		}
 
@@ -1940,7 +1941,7 @@ void FHttpCacheStore::FGetValueOp::EndGetRef(FHttpOperation& Operation)
 		if (!CompressedBuffer)
 		{
 			UE_LOG(LogDerivedDataCache, Display, TEXT("%s: Cache miss with invalid package for %s from '%s'"),
-				*CacheStore.Domain, *WriteToString<96>(Key), *Name);
+				*CacheStore.NodeName, *WriteToString<96>(Key), *Name);
 			return EndGet({Name, Key, {}, EStatus::Error});
 		}
 
@@ -2009,7 +2010,7 @@ void FHttpCacheStore::FExistsBatchOp::Exists(TConstArrayView<FCacheGetValueReque
 		{
 			UE_LOG(LogDerivedDataCache, VeryVerbose,
 				TEXT("%s: Skipped exists check of %s from '%s' because this cache store is not available"),
-				*CacheStore.Domain, *WriteToString<96>(Request.Key), *Request.Name);
+				*CacheStore.NodeName, *WriteToString<96>(Request.Key), *Request.Name);
 			OnComplete(Request.MakeResponse(EStatus::Error));
 			continue;
 		}
@@ -2017,7 +2018,7 @@ void FHttpCacheStore::FExistsBatchOp::Exists(TConstArrayView<FCacheGetValueReque
 		if (!EnumHasAnyFlags(Request.Policy, ECachePolicy::QueryRemote))
 		{
 			UE_LOG(LogDerivedDataCache, VeryVerbose, TEXT("%s: Skipped exists check of %s from '%s' due to cache policy"),
-				*CacheStore.Domain, *WriteToString<96>(Request.Key), *Request.Name);
+				*CacheStore.NodeName, *WriteToString<96>(Request.Key), *Request.Name);
 			OnComplete(Request.MakeResponse(EStatus::Error));
 			continue;
 		}
@@ -2025,7 +2026,7 @@ void FHttpCacheStore::FExistsBatchOp::Exists(TConstArrayView<FCacheGetValueReque
 		if (CacheStore.DebugOptions.ShouldSimulateGetMiss(Request.Key))
 		{
 			UE_LOG(LogDerivedDataCache, Verbose, TEXT("%s: Simulated miss for get of %s from '%s'"),
-				*CacheStore.Domain, *WriteToString<96>(Request.Key), *Request.Name);
+				*CacheStore.NodeName, *WriteToString<96>(Request.Key), *Request.Name);
 			OnComplete(Request.MakeResponse(EStatus::Error));
 			continue;
 		}
@@ -2075,7 +2076,7 @@ void FHttpCacheStore::FExistsBatchOp::BeginExists(TUniquePtr<FHttpOperation>&& O
 		for (const FCacheGetValueRequest& Request : Requests)
 		{
 			UE_LOG(LogDerivedDataCache, Verbose, TEXT("%s: Cache miss with canceled request for %s from '%s'"),
-				*CacheStore.Domain, *WriteToString<96>(Request.Key), *Request.Name);
+				*CacheStore.NodeName, *WriteToString<96>(Request.Key), *Request.Name);
 			RequestStats.Bucket = Request.Key.Bucket;
 			EndRequest(Request, {}, EStatus::Canceled);
 		}
@@ -2127,7 +2128,7 @@ void FHttpCacheStore::FExistsBatchOp::EndExists(FHttpOperation& Operation)
 		for (const FCacheGetValueRequest& Request : Requests)
 		{
 			UE_LOG(LogDerivedDataCache, Verbose, TEXT("%s: Cache miss with failed HTTP request for %s from '%s'"),
-				*CacheStore.Domain, *WriteToString<96>(Request.Key), *Request.Name);
+				*CacheStore.NodeName, *WriteToString<96>(Request.Key), *Request.Name);
 			RequestStats.Bucket = Request.Key.Bucket;
 			EndRequest(Request, {}, EStatus::Error);
 		}
@@ -2141,7 +2142,7 @@ void FHttpCacheStore::FExistsBatchOp::EndExists(FHttpOperation& Operation)
 		for (const FCacheGetValueRequest& Request : Requests)
 		{
 			UE_LOG(LogDerivedDataCache, Log, TEXT("%s: Cache miss with corrupt response for %s from '%s'."),
-				*CacheStore.Domain, *WriteToString<96>(Request.Key), *Request.Name);
+				*CacheStore.NodeName, *WriteToString<96>(Request.Key), *Request.Name);
 			RequestStats.Bucket = Request.Key.Bucket;
 			EndRequest(Request, {}, EStatus::Error);
 		}
@@ -2157,11 +2158,11 @@ void FHttpCacheStore::FExistsBatchOp::EndExists(FHttpOperation& Operation)
 	{
 		UE_LOG(LogDerivedDataCache, Log,
 			TEXT("%s: Cache exists returned unexpected quantity of results (expected %d, got %d)."),
-			*CacheStore.Domain, Requests.Num(), Results.Num());
+			*CacheStore.NodeName, Requests.Num(), Results.Num());
 		for (const FCacheGetValueRequest& Request : Requests)
 		{
 			UE_LOG(LogDerivedDataCache, Display, TEXT("%s: Cache miss with invalid response for %s from '%s'"),
-				*CacheStore.Domain, *WriteToString<96>(Request.Key), *Request.Name);
+				*CacheStore.NodeName, *WriteToString<96>(Request.Key), *Request.Name);
 			RequestStats.Bucket = Request.Key.Bucket;
 			EndRequest(Request, {}, EStatus::Error);
 		}
@@ -2178,7 +2179,7 @@ void FHttpCacheStore::FExistsBatchOp::EndExists(FHttpOperation& Operation)
 		if (OpId >= (uint32)Requests.Num())
 		{
 			UE_LOG(LogDerivedDataCache, Display, TEXT("%s: Encountered invalid opId %d while querying %d values"),
-				*CacheStore.Domain, OpId, Requests.Num());
+				*CacheStore.NodeName, OpId, Requests.Num());
 			continue;
 		}
 
@@ -2188,7 +2189,7 @@ void FHttpCacheStore::FExistsBatchOp::EndExists(FHttpOperation& Operation)
 		if (StatusCode < 200 || StatusCode > 204)
 		{
 			UE_LOG(LogDerivedDataCache, Verbose, TEXT("%s: Cache miss with unsuccessful response code %d for %s from '%s'"),
-				*CacheStore.Domain, StatusCode, *WriteToString<96>(Request.Key), *Request.Name);
+				*CacheStore.NodeName, StatusCode, *WriteToString<96>(Request.Key), *Request.Name);
 			EndRequest(Request, {}, EStatus::Error);
 			continue;
 		}
@@ -2198,7 +2199,7 @@ void FHttpCacheStore::FExistsBatchOp::EndExists(FHttpOperation& Operation)
 		if (RawHash.IsZero() || RawSize == MAX_uint64)
 		{
 			UE_LOG(LogDerivedDataCache, Display, TEXT("%s: Cache miss with invalid value for %s from '%s'"),
-				*CacheStore.Domain, *WriteToString<96>(Request.Key), *Request.Name);
+				*CacheStore.NodeName, *WriteToString<96>(Request.Key), *Request.Name);
 			EndRequest(Request, {}, EStatus::Error);
 			continue;
 		}
@@ -2221,7 +2222,8 @@ void FHttpCacheStore::FExistsBatchOp::EndRequest(const FCacheGetValueRequest& Re
 ///////////////////////////////////////////////////////////////////////////////////////////////////////////////////////
 
 FHttpCacheStore::FHttpCacheStore(const FHttpCacheStoreParams& Params, ICacheStoreOwner* Owner)
-	: Domain(Params.Host)
+	: NodeName(Params.Name)
+	, Domain(Params.Host)
 	, Namespace(Params.Namespace)
 	, OAuthProvider(Params.OAuthProvider)
 	, OAuthClientId(Params.OAuthClientId)
@@ -2245,7 +2247,7 @@ FHttpCacheStore::FHttpCacheStore(const FHttpCacheStoreParams& Params, ICacheStor
 	{
 		// Store the URI with the canonical name to pin to one region when using DNS-based region selection.
 		UE_LOG(LogDerivedDataCache, Display,
-			TEXT("%s: Pinned to %hs based on DNS canonical name."), *Domain, *ResolvedDomain);
+			TEXT("%s: Pinned to %hs based on DNS canonical name."), *NodeName, *ResolvedDomain);
 		EffectiveDomain.Reset();
 		EffectiveDomain.Append(ResolvedDomain);
 	}
@@ -2392,7 +2394,7 @@ bool FHttpCacheStore::AcquireAccessToken(IHttpClient* Client)
 {
 	if (Domain.StartsWith(TEXT("http://localhost")))
 	{
-		UE_LOG(LogDerivedDataCache, Log, TEXT("%s: Skipping authorization for connection to localhost."), *Domain);
+		UE_LOG(LogDerivedDataCache, Log, TEXT("%s: Skipping authorization for connection to localhost."), *NodeName);
 		return true;
 	}
 
@@ -2478,14 +2480,14 @@ bool FHttpCacheStore::AcquireAccessToken(IHttpClient* Client)
 					ResponseObject->TryGetNumberField(TEXT("expires_in"), ExpiryTimeSeconds))
 				{
 					UE_LOG(LogDerivedDataCache, Display,
-						TEXT("%s: Logged in to HTTP DDC services. Expires in %.0f seconds."), *Domain, ExpiryTimeSeconds);
+						TEXT("%s: Logged in to HTTP DDC services. Expires in %.0f seconds."), *NodeName, ExpiryTimeSeconds);
 					SetAccessTokenAndUnlock(Lock, AccessTokenString, ExpiryTimeSeconds);
 					return true;
 				}
 			}
 		}
 
-		UE_LOG(LogDerivedDataCache, Warning, TEXT("%s: Failed to log in to HTTP services with request %s."), *Domain, *WriteToString<256>(Operation));
+		UE_LOG(LogDerivedDataCache, Warning, TEXT("%s: Failed to log in to HTTP services with request %s."), *NodeName, *WriteToString<256>(Operation));
 		FailedLoginAttempts++;
 		return false;
 	}
@@ -2507,25 +2509,25 @@ bool FHttpCacheStore::AcquireAccessToken(IHttpClient* Client)
 			const double ExpiryTimeSeconds = (TokenExpiresAt - FDateTime::UtcNow()).GetTotalSeconds();
 			UE_LOG(LogDerivedDataCache, Display,
 				TEXT("%s: OidcToken: Logged in to HTTP DDC services. Expires at %s which is in %.0f seconds."),
-				*Domain, *TokenExpiresAt.ToString(), ExpiryTimeSeconds);
+				*NodeName, *TokenExpiresAt.ToString(), ExpiryTimeSeconds);
 			SetAccessTokenAndUnlock(Lock, AccessTokenString, ExpiryTimeSeconds);
 			return true;
 		}
 		else if (DesktopPlatform)
 		{
-			UE_LOG(LogDerivedDataCache, Warning, TEXT("%s: OidcToken: Failed to log in to HTTP services."), *Domain);
+			UE_LOG(LogDerivedDataCache, Warning, TEXT("%s: OidcToken: Failed to log in to HTTP services."), *NodeName);
 			FailedLoginAttempts++;
 			return false;
 		}
 		else
 		{
-			UE_LOG(LogDerivedDataCache, Warning, TEXT("%s: OidcToken: Use of OAuthProviderIdentifier requires that the target depend on DesktopPlatform."), *Domain);
+			UE_LOG(LogDerivedDataCache, Warning, TEXT("%s: OidcToken: Use of OAuthProviderIdentifier requires that the target depend on DesktopPlatform."), *NodeName);
 			FailedLoginAttempts++;
 			return false;
 		}
 	}
 
-	UE_LOG(LogDerivedDataCache, Warning, TEXT("%s: No available configuration to acquire an access token."), *Domain);
+	UE_LOG(LogDerivedDataCache, Warning, TEXT("%s: No available configuration to acquire an access token."), *NodeName);
 	FailedLoginAttempts++;
 	return false;
 }
@@ -2663,7 +2665,7 @@ void FHttpCacheStore::PutCacheRecordAsync(IRequestOwner& Owner, const FCachePutR
 	{
 		UE_LOG(LogDerivedDataCache, VeryVerbose,
 			TEXT("%s: Skipped put of %s from '%s' because this cache store is read-only"),
-			*Domain, *WriteToString<96>(Key), *Request.Name);
+			*NodeName, *WriteToString<96>(Key), *Request.Name);
 		return OnComplete(Request.MakeResponse(EStatus::Error));
 	}
 
@@ -2672,14 +2674,14 @@ void FHttpCacheStore::PutCacheRecordAsync(IRequestOwner& Owner, const FCachePutR
 	if (!EnumHasAnyFlags(RecordPolicy, ECachePolicy::StoreRemote))
 	{
 		UE_LOG(LogDerivedDataCache, VeryVerbose, TEXT("%s: Skipped put of %s from '%s' due to cache policy"),
-			*Domain, *WriteToString<96>(Key), *Request.Name);
+			*NodeName, *WriteToString<96>(Key), *Request.Name);
 		return OnComplete(Request.MakeResponse(EStatus::Error));
 	}
 
 	if (DebugOptions.ShouldSimulatePutMiss(Key))
 	{
 		UE_LOG(LogDerivedDataCache, Verbose, TEXT("%s: Simulated miss for put of %s from '%s'"),
-			*Domain, *WriteToString<96>(Key), *Request.Name);
+			*NodeName, *WriteToString<96>(Key), *Request.Name);
 		return OnComplete(Request.MakeResponse(EStatus::Error));
 	}
 
@@ -2720,7 +2722,7 @@ void FHttpCacheStore::PutCacheValueAsync(IRequestOwner& Owner, const FCachePutVa
 	{
 		UE_LOG(LogDerivedDataCache, VeryVerbose,
 			TEXT("%s: Skipped put of %s from '%s' because this cache store is read-only"),
-			*Domain, *WriteToString<96>(Request.Key), *Request.Name);
+			*NodeName, *WriteToString<96>(Request.Key), *Request.Name);
 		return OnComplete(Request.MakeResponse(EStatus::Error));
 	}
 
@@ -2728,14 +2730,14 @@ void FHttpCacheStore::PutCacheValueAsync(IRequestOwner& Owner, const FCachePutVa
 	if (!EnumHasAnyFlags(Request.Policy, ECachePolicy::StoreRemote))
 	{
 		UE_LOG(LogDerivedDataCache, VeryVerbose, TEXT("%s: Skipped put of %s from '%s' due to cache policy"),
-			*Domain, *WriteToString<96>(Request.Key), *Request.Name);
+			*NodeName, *WriteToString<96>(Request.Key), *Request.Name);
 		return OnComplete(Request.MakeResponse(EStatus::Error));
 	}
 
 	if (DebugOptions.ShouldSimulatePutMiss(Request.Key))
 	{
 		UE_LOG(LogDerivedDataCache, Verbose, TEXT("%s: Simulated miss for put of %s from '%s'"),
-			*Domain, *WriteToString<96>(Request.Key), *Request.Name);
+			*NodeName, *WriteToString<96>(Request.Key), *Request.Name);
 		return OnComplete(Request.MakeResponse(EStatus::Error));
 	}
 
@@ -2784,7 +2786,7 @@ void FHttpCacheStore::GetCacheValueAsync(
 	{
 		UE_LOG(LogDerivedDataCache, VeryVerbose,
 			TEXT("%s: Skipped get of %s from '%s' because this cache store is not available"),
-			*Domain, *WriteToString<96>(Key), *Name);
+			*NodeName, *WriteToString<96>(Key), *Name);
 		OnComplete({Name, Key, {}, UserData, EStatus::Error});
 		return;
 	}
@@ -2793,7 +2795,7 @@ void FHttpCacheStore::GetCacheValueAsync(
 	if (!EnumHasAnyFlags(Policy, ECachePolicy::QueryRemote))
 	{
 		UE_LOG(LogDerivedDataCache, VeryVerbose, TEXT("%s: Skipped get of %s from '%s' due to cache policy"),
-			*Domain, *WriteToString<96>(Key), *Name);
+			*NodeName, *WriteToString<96>(Key), *Name);
 		OnComplete({Name, Key, {}, UserData, EStatus::Error});
 		return;
 	}
@@ -2801,7 +2803,7 @@ void FHttpCacheStore::GetCacheValueAsync(
 	if (DebugOptions.ShouldSimulateGetMiss(Key))
 	{
 		UE_LOG(LogDerivedDataCache, Verbose, TEXT("%s: Simulated miss for get of %s from '%s'"),
-			*Domain, *WriteToString<96>(Key), *Name);
+			*NodeName, *WriteToString<96>(Key), *Name);
 		OnComplete({Name, Key, {}, UserData, EStatus::Error});
 		return;
 	}
@@ -2864,7 +2866,7 @@ void FHttpCacheStore::FinishChunkRequest(
 		const uint64 RawOffset = FMath::Min(Value.GetRawSize(), Request.RawOffset);
 		const uint64 RawSize = FMath::Min(Value.GetRawSize() - RawOffset, Request.RawSize);
 		UE_LOG(LogDerivedDataCache, Verbose, TEXT("%s: Cache hit for %s from '%s'"),
-			*Domain, *WriteToString<96>(Request.Key, '/', Request.Id), *Request.Name);
+			*NodeName, *WriteToString<96>(Request.Key, '/', Request.Id), *Request.Name);
 		FSharedBuffer Buffer;
 		const bool bExistsOnly = EnumHasAnyFlags(Request.Policy, ECachePolicy::SkipData);
 		if (!bExistsOnly)
@@ -2882,7 +2884,7 @@ void FHttpCacheStore::FinishChunkRequest(
 	else
 	{
 		UE_LOG(LogDerivedDataCache, Verbose, TEXT("%s: Cache miss for %s from '%s'"),
-			*Domain, *WriteToString<96>(Request.Key, '/', Request.Id), *Request.Name);
+			*NodeName, *WriteToString<96>(Request.Key, '/', Request.Id), *Request.Name);
 
 		SharedOnComplete.Get()(Request.MakeResponse(Status));
 	}
@@ -3176,7 +3178,7 @@ void FHttpCacheStore::GetValue(
 			{
 				TRACE_COUNTER_INCREMENT(HttpDDC_GetHit);
 				UE_LOG(LogDerivedDataCache, Verbose, TEXT("%s: Cache hit for %s from '%s'"),
-					*Domain, *WriteToString<96>(Response.Key), *Response.Name);
+					*NodeName, *WriteToString<96>(Response.Key), *Response.Name);
 			}
 			OnComplete(MoveTemp(Response));
 		});
@@ -3196,14 +3198,14 @@ void FHttpCacheStore::GetValue(
 					// With inline fetching, expect we will always have a value we can use.
 					// Even SkipData/Exists can rely on the blob existing if the ref is reported to exist.
 					UE_LOG(LogDerivedDataCache, Log, TEXT("%s: Cache miss due to inlining failure for %s from '%s'"),
-						*Domain, *WriteToString<96>(Response.Key), *Response.Name);
+						*NodeName, *WriteToString<96>(Response.Key), *Response.Name);
 				}
 
 				if (Response.Status == EStatus::Ok)
 				{
 					TRACE_COUNTER_INCREMENT(HttpDDC_GetHit);
 					UE_LOG(LogDerivedDataCache, Verbose, TEXT("%s: Cache hit for %s from '%s'"),
-						*Domain, *WriteToString<96>(Response.Key), *Response.Name);
+						*NodeName, *WriteToString<96>(Response.Key), *Response.Name);
 				}
 
 				SharedOnComplete.Get()(MoveTemp(Response));
