[ZenCacheStore]
 - Fix race condition on counters

diff --git a/Engine/Source/Developer/DerivedDataCache/Private/ZenCacheStore.cpp b/Engine/Source/Developer/DerivedDataCache/Private/ZenCacheStore.cpp
--- a/Engine/Source/Developer/DerivedDataCache/Private/ZenCacheStore.cpp
+++ b/Engine/Source/Developer/DerivedDataCache/Private/ZenCacheStore.cpp
@@ -35,14 +35,14 @@
 #include "ZenBackendUtils.h"
 #include "ZenSerialization.h"
 
-TRACE_DECLARE_INT_COUNTER(ZenDDC_Get, TEXT("ZenDDC Get"));
-TRACE_DECLARE_INT_COUNTER(ZenDDC_GetHit, TEXT("ZenDDC Get Hit"));
-TRACE_DECLARE_INT_COUNTER(ZenDDC_Put, TEXT("ZenDDC Put"));
-TRACE_DECLARE_INT_COUNTER(ZenDDC_PutHit, TEXT("ZenDDC Put Hit"));
-TRACE_DECLARE_INT_COUNTER(ZenDDC_BytesReceived, TEXT("ZenDDC Bytes Received"));
-TRACE_DECLARE_INT_COUNTER(ZenDDC_BytesSent, TEXT("ZenDDC Bytes Sent"));
-TRACE_DECLARE_INT_COUNTER(ZenDDC_CacheRecordRequestCountInFlight, TEXT("ZenDDC CacheRecord Request Count"));
-TRACE_DECLARE_INT_COUNTER(ZenDDC_ChunkRequestCountInFlight, TEXT("ZenDDC Chunk Request Count"));
+TRACE_DECLARE_ATOMIC_INT_COUNTER(ZenDDC_Get, TEXT("ZenDDC Get"));
+TRACE_DECLARE_ATOMIC_INT_COUNTER(ZenDDC_GetHit, TEXT("ZenDDC Get Hit"));
+TRACE_DECLARE_ATOMIC_INT_COUNTER(ZenDDC_Put, TEXT("ZenDDC Put"));
+TRACE_DECLARE_ATOMIC_INT_COUNTER(ZenDDC_PutHit, TEXT("ZenDDC Put Hit"));
+TRACE_DECLARE_ATOMIC_INT_COUNTER(ZenDDC_BytesReceived, TEXT("ZenDDC Bytes Received"));
+TRACE_DECLARE_ATOMIC_INT_COUNTER(ZenDDC_BytesSent, TEXT("ZenDDC Bytes Sent"));
+TRACE_DECLARE_ATOMIC_INT_COUNTER(ZenDDC_CacheRecordRequestCountInFlight, TEXT("ZenDDC CacheRecord Request Count"));
+TRACE_DECLARE_ATOMIC_INT_COUNTER(ZenDDC_ChunkRequestCountInFlight, TEXT("ZenDDC Chunk Request Count"));
 
 namespace UE::DerivedData
 {
