DDC: Fix race conditions on counters in the Jupiter and S3 cache stores

#rnx

diff --git a/Engine/Source/Developer/DerivedDataCache/Private/HttpCacheStore.cpp b/Engine/Source/Developer/DerivedDataCache/Private/HttpCacheStore.cpp
--- a/Engine/Source/Developer/DerivedDataCache/Private/HttpCacheStore.cpp
+++ b/Engine/Source/Developer/DerivedDataCache/Private/HttpCacheStore.cpp
@@ -73,12 +73,12 @@
 namespace UE::DerivedData
 {
 
-TRACE_DECLARE_INT_COUNTER(HttpDDC_Get, TEXT("HttpDDC Get"));
-TRACE_DECLARE_INT_COUNTER(HttpDDC_GetHit, TEXT("HttpDDC Get Hit"));
-TRACE_DECLARE_INT_COUNTER(HttpDDC_Put, TEXT("HttpDDC Put"));
-TRACE_DECLARE_INT_COUNTER(HttpDDC_PutHit, TEXT("HttpDDC Put Hit"));
-TRACE_DECLARE_INT_COUNTER(HttpDDC_BytesReceived, TEXT("HttpDDC Bytes Received"));
-TRACE_DECLARE_INT_COUNTER(HttpDDC_BytesSent, TEXT("HttpDDC Bytes Sent"));
+TRACE_DECLARE_ATOMIC_INT_COUNTER(HttpDDC_Get, TEXT("HttpDDC Get"));
+TRACE_DECLARE_ATOMIC_INT_COUNTER(HttpDDC_GetHit, TEXT("HttpDDC Get Hit"));
+TRACE_DECLARE_ATOMIC_INT_COUNTER(HttpDDC_Put, TEXT("HttpDDC Put"));
+TRACE_DECLARE_ATOMIC_INT_COUNTER(HttpDDC_PutHit, TEXT("HttpDDC Put Hit"));
+TRACE_DECLARE_ATOMIC_INT_COUNTER(HttpDDC_BytesReceived, TEXT("HttpDDC Bytes Received"));
+TRACE_DECLARE_ATOMIC_INT_COUNTER(HttpDDC_BytesSent, TEXT("HttpDDC Bytes Sent"));
 
 static bool ShouldAbortForShutdown()
 {
diff --git a/Engine/Source/Developer/DerivedDataCache/Private/S3CacheStore.cpp b/Engine/Source/Developer/DerivedDataCache/Private/S3CacheStore.cpp
--- a/Engine/Source/Developer/DerivedDataCache/Private/S3CacheStore.cpp
+++ b/Engine/Source/Developer/DerivedDataCache/Private/S3CacheStore.cpp
@@ -58,8 +58,8 @@
 namespace UE::DerivedData
 {
 
-TRACE_DECLARE_INT_COUNTER(S3DDC_Get, TEXT("S3DDC Get"));
-TRACE_DECLARE_INT_COUNTER(S3DDC_GetHit, TEXT("S3DDC Get Hit"));
+TRACE_DECLARE_ATOMIC_INT_COUNTER(S3DDC_Get, TEXT("S3DDC Get"));
+TRACE_DECLARE_ATOMIC_INT_COUNTER(S3DDC_GetHit, TEXT("S3DDC Get Hit"));
 
 ///////////////////////////////////////////////////////////////////////////////////////////////////////////////////////
 
