DDC: Added multi-process support to the saving of cache replays

diff --git a/Engine/Source/Developer/DerivedDataCache/Private/DerivedDataCacheReplay.cpp b/Engine/Source/Developer/DerivedDataCache/Private/DerivedDataCacheReplay.cpp
--- a/Engine/Source/Developer/DerivedDataCache/Private/DerivedDataCacheReplay.cpp
+++ b/Engine/Source/Developer/DerivedDataCache/Private/DerivedDataCacheReplay.cpp
@@ -5,6 +5,7 @@
 #include "Algo/AnyOf.h"
 #include "Compression/CompressedBuffer.h"
 #include "Containers/StringView.h"
+#include "CoreGlobals.h"
 #include "DerivedDataCacheKeyFilter.h"
 #include "DerivedDataCacheMethod.h"
 #include "DerivedDataCachePrivate.h"
@@ -15,6 +16,7 @@
 #include "HAL/FileManager.h"
 #include "HAL/PlatformTime.h"
 #include "Misc/CommandLine.h"
+#include "Misc/CoreDelegates.h"
 #include "Misc/Optional.h"
 #include "Misc/Parse.h"
 #include "Misc/ScopeLock.h"
@@ -100,6 +102,20 @@ private:
 	FMutableMemoryView RawBlockTail;
 	FCriticalSection Lock;
 	TOptional<FCacheReplayReader> ReplayReader;
+
+#if WITH_EDITOR
+private:
+	// Support for multi-process replay merging.
+
+	void WorkerCreated(const FMultiprocessCreatedContext& Context);
+	void WorkerDetached(const FMultiprocessDetachedContext& Context);
+	void MergeWorkerReplays();
+
+	FDelegateHandle WorkerCreatedHandle;
+	FDelegateHandle WorkerDetachedHandle;
+	/** Value tracks whether the worker is attached. */
+	TMap<int32, bool> WorkerIdToState;
+#endif // WITH_EDITOR
 };
 
 FCacheStoreReplay::FCacheStoreReplay(
@@ -115,6 +131,17 @@ FCacheStoreReplay::FCacheStoreReplay(
 {
 	if (!ReplayPath.IsEmpty())
 	{
+	#if WITH_EDITOR
+		if (const int32 WorkerId = UE::GetMultiprocessId())
+		{
+			ReplayPath.Appendf(TEXT(".worker%d.tmp"), WorkerId);
+		}
+		else
+		{
+			FCoreDelegates::OnMultiprocessWorkerCreated.AddRaw(this, &FCacheStoreReplay::WorkerCreated);
+			FCoreDelegates::OnMultiprocessWorkerDetached.AddRaw(this, &FCacheStoreReplay::WorkerDetached);
+		}
+	#endif // WITH_EDITOR
 		ReplayAr.Reset(IFileManager::Get().CreateFileWriter(*ReplayPath, FILEWRITE_NoFail));
 	}
 	if (CompressionBlockSize)
@@ -129,6 +156,15 @@ FCacheStoreReplay::~FCacheStoreReplay()
 {
 	FlushToArchive();
 	delete InnerCache;
+	
+#if WITH_EDITOR
+	if (!UE::GetMultiprocessId())
+	{
+		FCoreDelegates::OnMultiprocessWorkerDetached.Remove(WorkerDetachedHandle);
+		FCoreDelegates::OnMultiprocessWorkerCreated.Remove(WorkerCreatedHandle);
+		MergeWorkerReplays();
+	}
+#endif // WITH_EDITOR
 }
 
 template <typename RequestType>
@@ -262,6 +298,56 @@ void FCacheStoreReplay::GetChunks(
 	InnerCache->GetChunks(Requests, Owner, MoveTemp(OnComplete));
 }
 
+#if WITH_EDITOR
+void FCacheStoreReplay::WorkerCreated(const FMultiprocessCreatedContext& Context)
+{
+	WorkerIdToState.FindOrAdd(Context.Id) = true;
+}
+
+void FCacheStoreReplay::WorkerDetached(const FMultiprocessDetachedContext& Context)
+{
+	WorkerIdToState.FindOrAdd(Context.Id) = false;
+}
+
+void FCacheStoreReplay::MergeWorkerReplays()
+{
+	if (WorkerIdToState.IsEmpty())
+	{
+		return;
+	}
+	for (const TPair<int32, bool> Worker : WorkerIdToState)
+	{
+		const int32 WorkerId = Worker.Key;
+		FString WorkerReplayPath = ReplayPath;
+		WorkerReplayPath.Appendf(TEXT(".worker%d.tmp"), WorkerId);
+		if (Worker.Value)
+		{
+			UE_LOG(LogDerivedDataCache, Error,
+				TEXT("Replay: Skipped replay file '%s' because its worker has not detached."), *WorkerReplayPath);
+			continue;
+		}
+		if (TUniquePtr<FArchive> WorkerAr{IFileManager::Get().CreateFileReader(*WorkerReplayPath)})
+		{
+			int64 RemainingSize = WorkerAr->TotalSize();
+			constexpr int64 MaxBufferSize = 64 * 1024 * 1024;
+			const int64 BufferSize = FMath::Min(RemainingSize, MaxBufferSize);
+			TUniquePtr<uint8[]> Buffer(new uint8[BufferSize]);
+			while (RemainingSize > 0)
+			{
+				const int64 BlockSize = FMath::Min(RemainingSize, BufferSize);
+				WorkerAr->Serialize(Buffer.Get(), BlockSize);
+				ReplayAr->Serialize(Buffer.Get(), BlockSize);
+				RemainingSize -= BlockSize;
+			}
+			WorkerAr = nullptr;
+			IFileManager::Get().Delete(*WorkerReplayPath);
+			UE_LOG(LogDerivedDataCache, Display,
+				TEXT("Replay: Merged replay file '%s' from a worker process."), *WorkerReplayPath);
+		}
+	}
+}
+#endif // WITH_EDITOR
+
 ///////////////////////////////////////////////////////////////////////////////////////////////////////////////////////
 
 class FCacheReplayReader::FState
