Additional fix for hangs during DDC's WaitForQuiescence phase.  This one was caused by PUT operations to the Cloud DDC becoming stuck perpetually.  Addressing not with a time limit (we want to allow PUTs of large blobs that can take a long time), but with a generous rate limit so that if the connection is idle for a long time, we abandon it.

diff --git a/Engine/Source/Developer/DerivedDataCache/Private/HttpCacheStore.cpp b/Engine/Source/Developer/DerivedDataCache/Private/HttpCacheStore.cpp
--- a/Engine/Source/Developer/DerivedDataCache/Private/HttpCacheStore.cpp
+++ b/Engine/Source/Developer/DerivedDataCache/Private/HttpCacheStore.cpp
@@ -2279,9 +2279,9 @@ FHttpCacheStore::FHttpCacheStore(const FHttpCacheStoreParams& Params, ICacheStor
 		ClientParams.MinRequests = UE_HTTPDDC_GET_REQUEST_POOL_SIZE;
 		GetRequestQueue.Initialize(*ConnectionPool, ClientParams);
 
-		// Disabling rate limits during PUT operations as the cause too many spurious failures to put blobs or finalize refs
-		ClientParams.LowSpeedLimit = 0;
-		ClientParams.LowSpeedTime = 0;
+		// Giving very generous rate limits during PUT operations as they cause too many spurious failures to put blobs or finalize refs
+		ClientParams.LowSpeedLimit = 1;
+		ClientParams.LowSpeedTime = 60;
 		ClientParams.MaxRequests = UE_HTTPDDC_PUTREF_REQUEST_POOL_SIZE;
 		ClientParams.MinRequests = UE_HTTPDDC_PUTREF_REQUEST_POOL_SIZE;
 		PutRefRequestQueue.Initialize(*ConnectionPool, ClientParams);
