Allow printing of "Any" content type responses.

diff --git a/Engine/Source/Developer/DerivedDataCache/Private/HttpCacheStore.cpp b/Engine/Source/Developer/DerivedDataCache/Private/HttpCacheStore.cpp
--- a/Engine/Source/Developer/DerivedDataCache/Private/HttpCacheStore.cpp
+++ b/Engine/Source/Developer/DerivedDataCache/Private/HttpCacheStore.cpp
@@ -779,6 +779,7 @@ FString FHttpCacheStore::FHttpOperation::GetBodyAsString() const
 	EHttpMediaType ContentType = GetContentType();
 	switch (ContentType)
 	{
+	case EHttpMediaType::Any:
 	case EHttpMediaType::Text:
 	case EHttpMediaType::Json:
 	case EHttpMediaType::Yaml:
