Tweak to the  computation of latency for HttpCacheStore.

diff --git a/Engine/Source/Developer/DerivedDataCache/Private/HttpCacheStore.cpp b/Engine/Source/Developer/DerivedDataCache/Private/HttpCacheStore.cpp
--- a/Engine/Source/Developer/DerivedDataCache/Private/HttpCacheStore.cpp
+++ b/Engine/Source/Developer/DerivedDataCache/Private/HttpCacheStore.cpp
@@ -822,7 +822,7 @@ void FHttpCacheStore::FHttpOperation::GetStats(FRequestStats& OutStats) const
 	OutStats.PhysicalWriteSize += Stats.SendSize;
 	if (const EHttpMethod Method = Response->GetMethod(); Method == EHttpMethod::Get || Method == EHttpMethod::Head)
 	{
-		OutStats.AddLatency(FMonotonicTimeSpan::FromSeconds(Stats.StartTransferTime));
+		OutStats.AddLatency(FMonotonicTimeSpan::FromSeconds(Stats.StartTransferTime - Stats.ConnectTime));
 	}
 }
 
