Provide a configurable limit to the maximum number of requests per connection when using a protocol that does  multiplexing (HTTP2/HTTP3).

diff --git a/Engine/Source/Developer/DerivedDataCache/Private/Http/CurlHttpClient.cpp b/Engine/Source/Developer/DerivedDataCache/Private/Http/CurlHttpClient.cpp
--- a/Engine/Source/Developer/DerivedDataCache/Private/Http/CurlHttpClient.cpp
+++ b/Engine/Source/Developer/DerivedDataCache/Private/Http/CurlHttpClient.cpp
@@ -412,6 +412,10 @@ FCurlHttpConnectionPool::FCurlHttpConnectionPool(FCurlHttpManager& InManager, co
 	const uint32 MinConnections = Params.MinConnections ? Params.MinConnections : FMath::Min(Params.MaxConnections, DefaultConnectionCount);
 	curl_multi_setopt(CurlMulti, CURLMOPT_MAXCONNECTS, MinConnections);
 	curl_multi_setopt(CurlMulti, CURLMOPT_MAX_TOTAL_CONNECTIONS, MaxConnections);
+	if (Params.MaxRequestsPerConnection)
+	{
+		curl_multi_setopt(CurlMulti, CURLMOPT_MAX_CONCURRENT_STREAMS, (long)FMath::Min(Params.MaxRequestsPerConnection, (uint32)MAX_int32));
+	}
 	curl_multi_setopt(CurlMulti, CURLMOPT_PIPELINING, CURLPIPE_MULTIPLEX);
 }
 
@@ -630,7 +634,7 @@ void FCurlHttpClient::SetDefaultOptions(CURL* Curl, FCurlHttpHeaders& Headers) c
 
 	curl_easy_setopt(Curl, CURLOPT_HTTP_VERSION, ConvertVersion(Params.Version));
 
-	curl_easy_setopt(Curl, CURLOPT_DNS_CACHE_TIMEOUT, FMath::Min<long>(Params.DnsCacheTimeout, MAX_int32));
+	curl_easy_setopt(Curl, CURLOPT_DNS_CACHE_TIMEOUT, (long)FMath::Min(Params.DnsCacheTimeout, (uint32)MAX_int32));
 	curl_easy_setopt(Curl, CURLOPT_CONNECTTIMEOUT_MS, long(Params.ConnectTimeout));
 	curl_easy_setopt(Curl, CURLOPT_LOW_SPEED_LIMIT, long(Params.LowSpeedLimit));
 	curl_easy_setopt(Curl, CURLOPT_LOW_SPEED_TIME, long(Params.LowSpeedTime));
diff --git a/Engine/Source/Developer/DerivedDataCache/Private/Http/HttpClient.h b/Engine/Source/Developer/DerivedDataCache/Private/Http/HttpClient.h
--- a/Engine/Source/Developer/DerivedDataCache/Private/Http/HttpClient.h
+++ b/Engine/Source/Developer/DerivedDataCache/Private/Http/HttpClient.h
@@ -146,6 +146,9 @@ struct FHttpConnectionPoolParams final
 	/** Minimum number of concurrent connections maintained by the pool for reuse. Use 0 for the default limit. */
 	uint32 MinConnections = 0;
 
+	/** Maximum number of concurrent requests in flight for a connection when multiplexing. Use 0 for the default limit. */
+	uint32 MaxRequestsPerConnection = 0;
+
 	/** Allow requests to send async using the pool. Requests will block on execution when this is disabled. */
 	bool bAllowAsync = true;
 };
