DDC: Eliminated most blocking during get requests in HttpCacheStore

diff --git a/Engine/Source/Developer/DerivedDataCache/Private/HttpCacheStore.cpp b/Engine/Source/Developer/DerivedDataCache/Private/HttpCacheStore.cpp
--- a/Engine/Source/Developer/DerivedDataCache/Private/HttpCacheStore.cpp
+++ b/Engine/Source/Developer/DerivedDataCache/Private/HttpCacheStore.cpp
@@ -1312,9 +1312,8 @@ void FHttpCacheStore::FGetRecordOp::GetRecordOnly(const FCacheKey& InKey, const
 	OnRecordComplete = MoveTemp(InOnComplete);
 	RequestStats.Bucket = Key.Bucket;
 
-	TUniquePtr<FHttpOperation> Operation = CacheStore.WaitForHttpOperation(EOperationCategory::Get);
-	TRefCountPtr Self(this);
 	RequestTimer.Stop();
+	CacheStore.WaitForHttpOperationAsync(Owner, EOperationCategory::Get, [Self = TRefCountPtr(this)](TUniquePtr<FHttpOperation>&& Operation)
 	{
 		if (UNLIKELY(!Operation))
 		{
@@ -1336,7 +1335,7 @@ void FHttpCacheStore::FGetRecordOp::GetRecordOnly(const FCacheKey& InKey, const
 			Operation->GetStats(Self->RequestStats);
 			Self->EndGetRef(MoveTemp(Operation));
 		});
-	}
+	});
 }
 
 void FHttpCacheStore::FGetRecordOp::EndGetRef(TUniquePtr<FHttpOperation> Operation)
@@ -1554,11 +1553,10 @@ void FHttpCacheStore::FGetRecordOp::GetValues(TConstArrayView<FValueWithId> Valu
 			continue;
 		}
 
-		TUniquePtr<FHttpOperation> Operation = CacheStore.WaitForHttpOperation(EOperationCategory::Get);
-		TRefCountPtr Self(this);
+		CacheStore.WaitForHttpOperationAsync(Owner, EOperationCategory::Get, [Self = TRefCountPtr(this), SharedOnComplete, Value](TUniquePtr<FHttpOperation>&& Operation)
 		{
 			Self->BeginGetValue(MoveTemp(Operation), Value, SharedOnComplete);
-		}
+		});
 	}
 }
 
@@ -1664,14 +1662,13 @@ void FHttpCacheStore::FGetRecordOp::GetValuesExist(TConstArrayView<FValueWithId>
 	}
 
 	FRequestTimer RequestTimer(RequestStats);
+	RequestTimer.Stop();
 
 	FRequestBarrier Barrier(Owner);
-	TUniquePtr<FHttpOperation> Operation = CacheStore.WaitForHttpOperation(EOperationCategory::Get);
-	TRefCountPtr Self(this);
-	RequestTimer.Stop();
+	CacheStore.WaitForHttpOperationAsync(Owner, EOperationCategory::Get, [Self = TRefCountPtr(this), Values = MoveTemp(QueryValues), OnComplete = MoveTemp(OnComplete)](TUniquePtr<FHttpOperation>&& Operation) mutable
 	{
-		Self->BeginGetValuesExist(MoveTemp(Operation), MoveTemp(QueryValues), MoveTemp(OnComplete));
-	}
+		Self->BeginGetValuesExist(MoveTemp(Operation), MoveTemp(Values), MoveTemp(OnComplete));
+	});
 }
 
 void FHttpCacheStore::FGetRecordOp::BeginGetValuesExist(TUniquePtr<FHttpOperation>&& Operation, TArray<FValueWithId>&& Values, FOnValueComplete&& OnComplete)
@@ -1836,12 +1833,11 @@ void FHttpCacheStore::FGetValueOp::Get(const FCacheKey& InKey, ECachePolicy InPo
 	Policy = InPolicy;
 	OnComplete = MoveTemp(InOnComplete);
 
-	TUniquePtr<FHttpOperation> Operation = CacheStore.WaitForHttpOperation(EOperationCategory::Get);
-	TRefCountPtr Self(this);
 	RequestTimer.Stop();
+	CacheStore.WaitForHttpOperationAsync(Owner, EOperationCategory::Get, [Self = TRefCountPtr(this)](TUniquePtr<FHttpOperation>&& Operation)
 	{
 		Self->BeginGetRef(MoveTemp(Operation));
-	}
+	});
 }
 
 void FHttpCacheStore::FGetValueOp::BeginGetRef(TUniquePtr<FHttpOperation>&& Operation)
@@ -2061,12 +2057,11 @@ void FHttpCacheStore::FExistsBatchOp::Exists(TConstArrayView<FCacheGetValueReque
 	BodyWriter.EndObject();
 	FCbFieldIterator Body = BodyWriter.Save();
 
-	TUniquePtr<FHttpOperation> Operation = CacheStore.WaitForHttpOperation(EOperationCategory::Get);
-	TRefCountPtr Self(this);
 	RequestTimer.Stop();
+	CacheStore.WaitForHttpOperationAsync(Owner, EOperationCategory::Get, [Self = TRefCountPtr(this), Body = MoveTemp(Body)](TUniquePtr<FHttpOperation>&& Operation) mutable
 	{
 		Self->BeginExists(MoveTemp(Operation), MoveTemp(Body));
-	}
+	});
 }
 
 void FHttpCacheStore::FExistsBatchOp::BeginExists(TUniquePtr<FHttpOperation>&& Operation, FCbFieldIterator&& Body)
