Fixes issue where negative legacy could be accumulated in DDC request statistics.

- If a connection hasn't been established we now return a latency of 0 seconds rather than a negative number
- When calculating latency we use the pretransfer time from curl which fits closer to the intention of measuring latency, specifically the time between asking for bytes and receiving the first byte.
- Adds GetLatency utility method to `FHttpResponseStats` as we calculate latency in multiple places manually and all contained the same error

diff --git a/Engine/Source/Developer/DerivedDataCache/Private/Http/CurlHttpClient.cpp b/Engine/Source/Developer/DerivedDataCache/Private/Http/CurlHttpClient.cpp
--- a/Engine/Source/Developer/DerivedDataCache/Private/Http/CurlHttpClient.cpp
+++ b/Engine/Source/Developer/DerivedDataCache/Private/Http/CurlHttpClient.cpp
@@ -1135,6 +1135,7 @@ void FCurlHttpResponse::SetComplete(CURLcode Code)
 	Stats.NameResolveTime = GetDoubleInfo(CURLINFO_NAMELOOKUP_TIME);
 	Stats.ConnectTime = GetDoubleInfo(CURLINFO_CONNECT_TIME);
 	Stats.TlsConnectTime = GetDoubleInfo(CURLINFO_APPCONNECT_TIME);
+	Stats.PreTransferTime = GetDoubleInfo(CURLINFO_PRETRANSFER_TIME);
 	Stats.StartTransferTime = GetDoubleInfo(CURLINFO_STARTTRANSFER_TIME);
 	Stats.TotalTime = GetDoubleInfo(CURLINFO_TOTAL_TIME);
 
diff --git a/Engine/Source/Developer/DerivedDataCache/Private/Http/HttpClient.h b/Engine/Source/Developer/DerivedDataCache/Private/Http/HttpClient.h
--- a/Engine/Source/Developer/DerivedDataCache/Private/Http/HttpClient.h
+++ b/Engine/Source/Developer/DerivedDataCache/Private/Http/HttpClient.h
@@ -384,10 +384,18 @@ struct FHttpResponseStats final
 	double ConnectTime = 0.0;
 	/** Time from the start of the request until the TLS connection was complete, in seconds. */
 	double TlsConnectTime = 0.0;
+	/** Time from the start of the request until the request was sent but before the first byte was received, in seconds. */
+	double PreTransferTime = 0.0;
 	/** Time from the start of the request until the first byte was received, in seconds. */
 	double StartTransferTime = 0.0;
 	/** Time from the start of the request until response was complete, in seconds. */
 	double TotalTime = 0.0;
+
+	/** Time from when the request was sent until the first byte was received, in seconds. If no byte has been received, latency is 0. */
+	double GetLatency() const
+	{
+		return FMath::Max(0.0f, StartTransferTime - PreTransferTime);
+	}
 };
 
 /**
diff --git a/Engine/Source/Developer/DerivedDataCache/Private/HttpCacheStore.cpp b/Engine/Source/Developer/DerivedDataCache/Private/HttpCacheStore.cpp
--- a/Engine/Source/Developer/DerivedDataCache/Private/HttpCacheStore.cpp
+++ b/Engine/Source/Developer/DerivedDataCache/Private/HttpCacheStore.cpp
@@ -822,7 +822,7 @@ void FHttpCacheStore::FHttpOperation::GetStats(FRequestStats& OutStats) const
 	OutStats.PhysicalWriteSize += Stats.SendSize;
 	if (const EHttpMethod Method = Response->GetMethod(); Method == EHttpMethod::Get || Method == EHttpMethod::Head)
 	{
-		OutStats.AddLatency(FMonotonicTimeSpan::FromSeconds(Stats.StartTransferTime - Stats.ConnectTime));
+		OutStats.AddLatency(FMonotonicTimeSpan::FromSeconds(Stats.GetLatency()));
 	}
 }
 
diff --git a/Engine/Source/Developer/DerivedDataCache/Private/ZenCacheStore.cpp b/Engine/Source/Developer/DerivedDataCache/Private/ZenCacheStore.cpp
--- a/Engine/Source/Developer/DerivedDataCache/Private/ZenCacheStore.cpp
+++ b/Engine/Source/Developer/DerivedDataCache/Private/ZenCacheStore.cpp
@@ -570,7 +570,7 @@ public:
 			{
 				FRequestTimer RequestTimer(Requests[0].Stats);
 				const FHttpResponseStats& ResponseStats = HttpResponse->GetStats();
-				Requests[0].Stats.Latency = FMonotonicTimeSpan::FromSeconds(ResponseStats.StartTransferTime - ResponseStats.ConnectTime);
+				Requests[0].Stats.Latency = FMonotonicTimeSpan::FromSeconds(ResponseStats.GetLatency());
 
 				int32 RequestIndex = 0;
 				if (HttpResponse->GetErrorCode() == EHttpErrorCode::None && HttpResponse->GetStatusCode() >= 200 && HttpResponse->GetStatusCode() <= 299)
@@ -967,7 +967,7 @@ public:
 			{
 				FRequestTimer RequestTimer(Requests[0].Stats);
 				const FHttpResponseStats& ResponseStats = HttpResponse->GetStats();
-				Requests[0].Stats.Latency = FMonotonicTimeSpan::FromSeconds(ResponseStats.StartTransferTime - ResponseStats.ConnectTime);
+				Requests[0].Stats.Latency = FMonotonicTimeSpan::FromSeconds(ResponseStats.GetLatency());
 
 				int32 RequestIndex = 0;
 				if (HttpResponse->GetErrorCode() == EHttpErrorCode::None && HttpResponse->GetStatusCode() >= 200 && HttpResponse->GetStatusCode() <= 299)
@@ -1187,7 +1187,7 @@ public:
 			{
 				FRequestTimer RequestTimer(Requests[0].Stats);
 				const FHttpResponseStats& ResponseStats = HttpResponse->GetStats();
-				Requests[0].Stats.Latency = FMonotonicTimeSpan::FromSeconds(ResponseStats.StartTransferTime - ResponseStats.ConnectTime);
+				Requests[0].Stats.Latency = FMonotonicTimeSpan::FromSeconds(ResponseStats.GetLatency());
 
 				int32 RequestIndex = 0;
 				if (HttpResponse->GetErrorCode() == EHttpErrorCode::None && HttpResponse->GetStatusCode() >= 200 && HttpResponse->GetStatusCode() <= 299)
@@ -1255,6 +1255,11 @@ public:
 									}
 								}
 							}
+							else
+							{
+								FCbFieldView HashView2 = ResultObject[ANSITEXTVIEW("RawHash")];
+								RawHash = HashView2.AsHash();
+							}
 						}
 						Succeeded ? OnHit(RequestWithStats, MoveTemp(RawHash), RawSize, MoveTemp(RequestedBytes)) : OnMiss(RequestWithStats);
 					}
@@ -1834,7 +1839,7 @@ void FZenCacheStore::ActivatePerformanceEvaluationThread()
 							return;
 						}
 
-						double LatencySec = (HttpResponse->GetStats().StartTransferTime - HttpResponse->GetStats().ConnectTime);
+						const double LatencySec = HttpResponse->GetStats().GetLatency();
 						StoreStats->AddLatency(StartTime, FMonotonicTimePoint::Now(), FMonotonicTimeSpan::FromSeconds(LatencySec));
 						if (!bTryEvaluatePerformance || (StoreStats->GetAverageLatency() * 1000 <= DeactivateAtMs))
 						{
