DDC: Fixed cancellation for FHttpCacheStoreRequestQueue

#rnx

diff --git a/Engine/Source/Developer/DerivedDataCache/Private/HttpCacheStore.cpp b/Engine/Source/Developer/DerivedDataCache/Private/HttpCacheStore.cpp
--- a/Engine/Source/Developer/DerivedDataCache/Private/HttpCacheStore.cpp
+++ b/Engine/Source/Developer/DerivedDataCache/Private/HttpCacheStore.cpp
@@ -192,10 +192,7 @@ public:
 			{
 				if (THttpUniquePtr<IHttpRequest> Request = Client->TryCreateRequest({}))
 				{
-					if (FQueueRequest* Waiter = Queue.Pop())
-					{
-						Waiter->Complete(MoveTemp(Request));
-					}
+					TryGiveRequestToQueue(MoveTemp(Request));
 				}
 			}
 		};
@@ -214,11 +211,7 @@ public:
 
 		while (THttpUniquePtr<IHttpRequest> Request = Client->TryCreateRequest(Params))
 		{
-			if (FQueueRequest* Waiter = Queue.Pop())
-			{
-				Waiter->Complete(MoveTemp(Request));
-			}
-			else
+			if (!TryGiveRequestToQueue(MoveTemp(Request)))
 			{
 				OnRequest(MoveTemp(Request));
 				return;
@@ -229,11 +222,7 @@ public:
 
 		while (THttpUniquePtr<IHttpRequest> Request = Client->TryCreateRequest(Params))
 		{
-			if (FQueueRequest* Waiter = Queue.Pop())
-			{
-				Waiter->Complete(MoveTemp(Request));
-			}
-			else
+			if (!TryGiveRequestToQueue(MoveTemp(Request)))
 			{
 				return;
 			}
@@ -241,6 +230,18 @@ public:
 	}
 
 private:
+	bool TryGiveRequestToQueue(THttpUniquePtr<IHttpRequest>&& Request)
+	{
+		while (FQueueRequest* Waiter = Queue.Pop())
+		{
+			if (Waiter->TryClaimRequest(MoveTemp(Request)))
+			{
+				return true;
+			}
+		}
+		return false;
+	}
+
 	class FQueueRequest : FRequestBase
 	{
 	public:
@@ -248,31 +249,42 @@ private:
 			: Owner(InOwner)
 			, OnRequest(MoveTemp(InOnRequest))
 		{
+			AddRef(); // Release() is called by ClaimRequest()
 			Owner.Begin(this);
 		}
 
-		void Complete(THttpUniquePtr<IHttpRequest>&& Request)
+		bool TryClaimRequest(THttpUniquePtr<IHttpRequest>&& Request)
+		{
+			ON_SCOPE_EXIT { Release(); };
+			return TryComplete(MoveTemp(Request));
+		}
+
+	private:
+		bool TryComplete(THttpUniquePtr<IHttpRequest>&& Request)
 		{
 			if (bComplete.exchange(true))
 			{
-				OnComplete.Wait();
-				return;
+				return false;
 			}
 			Owner.End(this, [this](THttpUniquePtr<IHttpRequest>&& Request)
 			{
 				OnRequest(MoveTemp(Request));
 				OnComplete.Notify();
 			}, MoveTemp(Request));
+			return true;
 		}
 
-	private:
 		void SetPriority(EPriority Priority) final
 		{
 		}
 
 		void Cancel() final
 		{
-			Complete({});
+			if (!TryComplete({}))
+			{
+				TRACE_CPUPROFILER_EVENT_SCOPE(HttpDDC_CancelOperation);
+				OnComplete.Wait();
+			}
 		}
 
 		void Wait() final
