Add option to bypass proxies to the derived data http client (default is no bypass).  Zen cache store client uses it and defaults to bypassing proxies (but can be overridden in INI). UE Cloud DDC client uses it and defaults to NOT bypassing proxies (but can be overridden in INI).

diff --git a/Engine/Source/Developer/DerivedDataCache/Private/Http/CurlHttpClient.cpp b/Engine/Source/Developer/DerivedDataCache/Private/Http/CurlHttpClient.cpp
--- a/Engine/Source/Developer/DerivedDataCache/Private/Http/CurlHttpClient.cpp
+++ b/Engine/Source/Developer/DerivedDataCache/Private/Http/CurlHttpClient.cpp
@@ -644,6 +644,10 @@ void FCurlHttpClient::SetDefaultOptions(CURL* Curl, FCurlHttpHeaders& Headers) c
 		(Params.bFollow301Post ? CURL_REDIR_POST_301 : 0) |
 		(Params.bFollow302Post ? CURL_REDIR_POST_302 : 0) |
 		(Params.bFollow303Post ? CURL_REDIR_POST_303 : 0)));
+	if (Params.bBypassProxy)
+	{
+		curl_easy_setopt(Curl, CURLOPT_NOPROXY, "*");
+	}
 	curl_easy_setopt(Curl, CURLOPT_VERBOSE, long(Params.bVerbose));
 }
 
diff --git a/Engine/Source/Developer/DerivedDataCache/Private/Http/HttpClient.h b/Engine/Source/Developer/DerivedDataCache/Private/Http/HttpClient.h
--- a/Engine/Source/Developer/DerivedDataCache/Private/Http/HttpClient.h
+++ b/Engine/Source/Developer/DerivedDataCache/Private/Http/HttpClient.h
@@ -207,6 +207,9 @@ struct FHttpClientParams final
 	bool bFollow302Post : 1;
 	bool bFollow303Post : 1;
 
+	/** Avoid use of configured proxy on the client. */
+	bool bBypassProxy : 1;
+
 	/** Verbose logging for requests created by the client. */
 	bool bVerbose : 1;
 
@@ -215,6 +218,7 @@ struct FHttpClientParams final
 		, bFollow301Post(false)
 		, bFollow302Post(false)
 		, bFollow303Post(false)
+		, bBypassProxy(false)
 		, bVerbose(false)
 	{
 	}
diff --git a/Engine/Source/Developer/DerivedDataCache/Private/HttpCacheStore.cpp b/Engine/Source/Developer/DerivedDataCache/Private/HttpCacheStore.cpp
--- a/Engine/Source/Developer/DerivedDataCache/Private/HttpCacheStore.cpp
+++ b/Engine/Source/Developer/DerivedDataCache/Private/HttpCacheStore.cpp
@@ -362,6 +362,7 @@ struct FHttpCacheStoreParams
 
 	bool bResolveHostCanonicalName = true;
 	bool bReadOnly = false;
+	bool bBypassProxy = false;
 
 	void Parse(const TCHAR* NodeName, const TCHAR* Config);
 };
@@ -469,6 +470,7 @@ private:
 
 	bool bIsUsable = false;
 	bool bReadOnly = false;
+	bool bBypassProxy = false;
 
 	static inline FHttpCacheStore* AnyInstance = nullptr;
 
@@ -2241,6 +2243,7 @@ FHttpCacheStore::FHttpCacheStore(const FHttpCacheStoreParams& Params, ICacheStor
 	, AuthScheme(Params.AuthScheme)
 	, StoreOwner(Owner)
 	, bReadOnly(Params.bReadOnly)
+	, bBypassProxy(Params.bBypassProxy)
 {
 	TRACE_CPUPROFILER_EVENT_SCOPE(HttpDDC_Construct);
 
@@ -2382,6 +2385,7 @@ FHttpClientParams FHttpCacheStore::GetDefaultClientParams() const
 	ClientParams.TlsLevel = EHttpTlsLevel::All;
 	ClientParams.bFollowRedirects = true;
 	ClientParams.bFollow302Post = true;
+	ClientParams.bBypassProxy = bBypassProxy;
 
 	EHttpVersion HttpVersionEnum = EHttpVersion::V2;
 	TryLexFromString(HttpVersionEnum, HttpVersion);
@@ -3381,6 +3385,7 @@ void FHttpCacheStoreParams::Parse(const TCHAR* NodeName, const TCHAR* Config)
 	// Cache Params
 
 	FParse::Bool(Config, TEXT("ReadOnly="), bReadOnly);
+	FParse::Bool(Config, TEXT("BypassProxy="), bBypassProxy);
 }
 
 } // UE::DerivedData
diff --git a/Engine/Source/Developer/DerivedDataCache/Private/ZenCacheStore.cpp b/Engine/Source/Developer/DerivedDataCache/Private/ZenCacheStore.cpp
--- a/Engine/Source/Developer/DerivedDataCache/Private/ZenCacheStore.cpp
+++ b/Engine/Source/Developer/DerivedDataCache/Private/ZenCacheStore.cpp
@@ -81,6 +81,7 @@ struct FZenCacheStoreParams
 	TOptional<bool> bRemote;
 	bool bFlush = false;
 	bool bReadOnly = false;
+	bool bBypassProxy = true;
 
 	void Parse(const TCHAR* NodeName, const TCHAR* Config);
 };
@@ -1659,6 +1660,7 @@ void FZenCacheStore::Initialize(const FZenCacheStoreParams& Params)
 		ReadinessClientParams.MinRequests = 1;
 		ReadinessClientParams.LowSpeedLimit = 1;
 		ReadinessClientParams.LowSpeedTime = 5; // 5 second idle time limit for the initial readiness check
+		ReadinessClientParams.bBypassProxy = Params.bBypassProxy;
 		THttpUniquePtr<IHttpClient> ReadinessClient = ConnectionPool->CreateClient(ReadinessClientParams);
 		THttpUniquePtr<IHttpRequest> ReadinessRequest = ReadinessClient->TryCreateRequest({});
 		TAnsiStringBuilder<256> StatusUri;
@@ -1704,6 +1706,7 @@ void FZenCacheStore::Initialize(const FZenCacheStoreParams& Params)
 	ClientParams.MinRequests = RequestPoolSize;
 	ClientParams.LowSpeedLimit = 1;
 	ClientParams.LowSpeedTime = 25;
+	ClientParams.bBypassProxy = Params.bBypassProxy;
 	RequestQueue = FHttpRequestQueue(*ConnectionPool, ClientParams);
 
 	bIsLocalConnection = ZenService.GetInstance().IsServiceRunningLocally() || ZenService.GetInstance().GetServiceSettings().IsAutoLaunch();
@@ -2079,6 +2082,7 @@ void FZenCacheStoreParams::Parse(const TCHAR* NodeName, const TCHAR* Config)
 	// Sandbox and flush configuration for use in Cold/Warm type use cases
 	FParse::Value(Config, TEXT("Sandbox="), Sandbox);
 	FParse::Bool(Config, TEXT("Flush="), bFlush);
+	FParse::Bool(Config, TEXT("BypassProxy="), bBypassProxy);
 
 	// Performance deactivation
 	FParse::Value(Config, TEXT("DeactivateAt="), DeactivateAtMs);
