Ignore cancelations when tracking DDC usage stats as our legacy stat tracking system doesn't represent cancellations and instead gives a false perception about how many cache misses are occurring. The cancellations do have an overhead, but in practice most cancellations are done before network traffic has begun and has very little cost, such that we aren't at a major lossby not tracking this time.

diff --git a/Engine/Source/Developer/DerivedDataCache/Private/DerivedDataCacheStats.cpp b/Engine/Source/Developer/DerivedDataCache/Private/DerivedDataCacheStats.cpp
--- a/Engine/Source/Developer/DerivedDataCache/Private/DerivedDataCacheStats.cpp
+++ b/Engine/Source/Developer/DerivedDataCache/Private/DerivedDataCacheStats.cpp
@@ -237,6 +237,7 @@ void FCacheStoreStats::AddRequest(const FCacheStoreRequestStats& Stats)
 	using EStatType = FCallStats::EStatType;
 
 	const EHitOrMiss HitOrMiss = Stats.Status == EStatus::Ok ? EHitOrMiss::Hit : EHitOrMiss::Miss;
+	const bool bIsCanceled = Stats.Status == EStatus::Canceled;
 	const bool bIsGet = Stats.Op != ECacheStoreRequestOp::Put;
 	const bool bIsInGameThread = IsInGameThread();
 #endif
@@ -254,10 +255,15 @@ void FCacheStoreStats::AddRequest(const FCacheStoreRequestStats& Stats)
 		BucketStats.OtherThreadTime += Stats.OtherThreadTime;
 
 	#if ENABLE_COOK_STATS
-		FCallStats& CallStats = bIsGet ? BucketStats.GetStats : BucketStats.PutStats;
-		CallStats.Accumulate(HitOrMiss, EStatType::Cycles, int64(Stats.MainThreadTime.ToSeconds() / FPlatformTime::GetSecondsPerCycle64()), /*bIsInGameThread*/ true);
-		CallStats.Accumulate(HitOrMiss, EStatType::Cycles, int64(Stats.OtherThreadTime.ToSeconds() / FPlatformTime::GetSecondsPerCycle64()), /*bIsInGameThread*/ false);
-		CallStats.Accumulate(HitOrMiss, EStatType::Bytes, bIsGet ? Stats.PhysicalReadSize : Stats.PhysicalWriteSize, bIsInGameThread);
+		// FCallStats is a legacy stat codepath that isn't aware of cancellations. In order to avoid skewing miss% rates
+		// by accumulating cancellation requests as misses, we simply ignore cancellation requests for stat tracking
+		if (!bIsCanceled)
+		{
+			FCallStats& CallStats = bIsGet ? BucketStats.GetStats : BucketStats.PutStats;
+			CallStats.Accumulate(HitOrMiss, EStatType::Cycles, int64(Stats.MainThreadTime.ToSeconds() / FPlatformTime::GetSecondsPerCycle64()), /*bIsInGameThread*/ true);
+			CallStats.Accumulate(HitOrMiss, EStatType::Cycles, int64(Stats.OtherThreadTime.ToSeconds() / FPlatformTime::GetSecondsPerCycle64()), /*bIsInGameThread*/ false);
+			CallStats.Accumulate(HitOrMiss, EStatType::Bytes, bIsGet ? Stats.PhysicalReadSize : Stats.PhysicalWriteSize, bIsInGameThread);
+		}
 	#endif
 	}
 
@@ -293,11 +299,16 @@ void FCacheStoreStats::AddRequest(const FCacheStoreRequestStats& Stats)
 	}
 
 #if ENABLE_COOK_STATS
-	FCookStats::CallStats& CallStats = bIsGet ? GetStats : PutStats;
-	CallStats.Accumulate(HitOrMiss, EStatType::Counter, 1, bIsInGameThread);
-	CallStats.Accumulate(HitOrMiss, EStatType::Cycles, int64(Stats.MainThreadTime.ToSeconds() / FPlatformTime::GetSecondsPerCycle64()), /*bIsInGameThread*/ true);
-	CallStats.Accumulate(HitOrMiss, EStatType::Cycles, int64(Stats.OtherThreadTime.ToSeconds() / FPlatformTime::GetSecondsPerCycle64()), /*bIsInGameThread*/ false);
-	CallStats.Accumulate(HitOrMiss, EStatType::Bytes, bIsGet ? Stats.PhysicalReadSize : Stats.PhysicalWriteSize, bIsInGameThread);
+	// FCallStats is a legacy stat codepath that isn't aware of cancellations. In order to avoid skewing miss% rates
+	// by accumulating cancellation requests as misses, we simply ignore cancellation requests for stat tracking
+	if (!bIsCanceled)
+	{
+		FCookStats::CallStats& CallStats = bIsGet ? GetStats : PutStats;
+		CallStats.Accumulate(HitOrMiss, EStatType::Counter, 1, bIsInGameThread);
+		CallStats.Accumulate(HitOrMiss, EStatType::Cycles, int64(Stats.MainThreadTime.ToSeconds() / FPlatformTime::GetSecondsPerCycle64()), /*bIsInGameThread*/ true);
+		CallStats.Accumulate(HitOrMiss, EStatType::Cycles, int64(Stats.OtherThreadTime.ToSeconds() / FPlatformTime::GetSecondsPerCycle64()), /*bIsInGameThread*/ false);
+		CallStats.Accumulate(HitOrMiss, EStatType::Bytes, bIsGet ? Stats.PhysicalReadSize : Stats.PhysicalWriteSize, bIsInGameThread);
+	}
 #endif
 }
 
diff --git a/Engine/Source/Developer/DerivedDataCache/Private/DerivedDataCacheStoreAsync.cpp b/Engine/Source/Developer/DerivedDataCache/Private/DerivedDataCacheStoreAsync.cpp
--- a/Engine/Source/Developer/DerivedDataCache/Private/DerivedDataCacheStoreAsync.cpp
+++ b/Engine/Source/Developer/DerivedDataCache/Private/DerivedDataCacheStoreAsync.cpp
@@ -158,6 +158,12 @@ void FCacheStoreAsync::Execute(
 			{
 				COOK_STAT(Timer.AddHit(0));
 			}
+			else if (Response.Status == EStatus::Canceled)
+			{
+				// The request was cancelled so only track the cycles lost, rather than adding a miss
+				COOK_STAT(Timer.TrackCyclesOnly());
+			}
+
 			OnComplete(MoveTemp(Response));
 			Private::AddToAsyncTaskCounter(-1);
 		});
diff --git a/Engine/Source/Developer/DerivedDataCache/Private/DerivedDataCacheStoreHierarchy.cpp b/Engine/Source/Developer/DerivedDataCache/Private/DerivedDataCacheStoreHierarchy.cpp
--- a/Engine/Source/Developer/DerivedDataCache/Private/DerivedDataCacheStoreHierarchy.cpp
+++ b/Engine/Source/Developer/DerivedDataCache/Private/DerivedDataCacheStoreHierarchy.cpp
@@ -1637,9 +1637,16 @@ void FCacheStoreHierarchy::RecordStats(FCacheBucket Bucket, ERequestType Type, E
 	using EStatType = FCallStats::EStatType;
 
 	const bool bIsInGameThread = IsInGameThread();
+	const bool bIsCanceled = Status == EStatus::Canceled;
 	const EHitOrMiss HitOrMiss = Status == EStatus::Ok ? EHitOrMiss::Hit : EHitOrMiss::Miss;
-	FCallStats& CallStats = (Op == ERequestOp::Put) ? BucketStats.PutStats : BucketStats.GetStats;
-	CallStats.Accumulate(HitOrMiss, EStatType::Counter, 1, bIsInGameThread);
+
+	// FCallStats is a legacy stat codepath that isn't aware of cancellations. In order to avoid skewing miss% rates
+	// by accumulating cancellation requests as misses, we simply ignore cancellation requests for stat tracking
+	if (!bIsCanceled)
+	{
+		FCallStats& CallStats = (Op == ERequestOp::Put) ? BucketStats.PutStats : BucketStats.GetStats;
+		CallStats.Accumulate(HitOrMiss, EStatType::Counter, 1, bIsInGameThread);
+	}
 #endif
 }
 
