DDC: Fixed a race in the legacy CachedDataProbablyExistsBatch

diff --git a/Engine/Source/Developer/DerivedDataCache/Private/DerivedDataCache.cpp b/Engine/Source/Developer/DerivedDataCache/Private/DerivedDataCache.cpp
--- a/Engine/Source/Developer/DerivedDataCache/Private/DerivedDataCache.cpp
+++ b/Engine/Source/Developer/DerivedDataCache/Private/DerivedDataCache.cpp
@@ -7,7 +7,9 @@
 #include "AnalyticsEventAttribute.h"
 #include "Async/AsyncWork.h"
 #include "Async/InheritedContext.h"
+#include "Async/Mutex.h"
 #include "Async/TaskGraphInterfaces.h"
+#include "Async/UniqueLock.h"
 #include "Containers/Map.h"
 #include "DDCCleanup.h"
 #include "DerivedDataBackendInterface.h"
@@ -799,6 +801,7 @@ public:
 
 	virtual TBitArray<> CachedDataProbablyExistsBatch(TConstArrayView<FString> CacheKeys) override
 	{
+		FMutex ResultMutex;
 		TBitArray<> Result(false, CacheKeys.Num());
 		if (!CacheKeys.IsEmpty())
 		{
@@ -820,8 +823,10 @@ public:
 				}
 				FRequestOwner BlockingOwner(EPriority::Blocking);
 				Backend->GetRoot().LegacyGet(LegacyRequests, BlockingOwner,
-					[&Result](FLegacyCacheGetResponse&& Response)
+					[&Result, &ResultMutex](FLegacyCacheGetResponse&& Response)
 					{
+						// Lock because it is not safe to write bits in the same word from different threads.
+						TUniqueLock Lock(ResultMutex);
 						Result[int32(Response.UserData)] = Response.Status == EStatus::Ok;
 					});
 				BlockingOwner.Wait();
