DDC: Eliminated blocking tasks used by the legacy cache API

diff --git a/Engine/Source/Developer/DerivedDataCache/Private/DerivedDataCache.cpp b/Engine/Source/Developer/DerivedDataCache/Private/DerivedDataCache.cpp
--- a/Engine/Source/Developer/DerivedDataCache/Private/DerivedDataCache.cpp
+++ b/Engine/Source/Developer/DerivedDataCache/Private/DerivedDataCache.cpp
@@ -376,191 +376,285 @@ void LaunchTaskInCacheThreadPool(IRequestOwner& Owner, TUniqueFunction<void ()>&
 	LaunchTaskInThreadPool(Owner, GCacheThreadPool, MoveTemp(TaskBody));
 }
 
-/**
- * Implementation of the derived data cache
- * This API is fully threadsafe
-**/
-class FDerivedDataCache final
-	: public FDerivedDataCacheInterface
-	, public ICache
-	, public ICacheStoreMaintainer
-	, public IDDCCleanup
+class FLegacyFetchOrBuildTask
 {
-
-	/** 
-	 * Async worker that checks the cache backend and if that fails, calls the deriver to build the data and then puts the results to the cache
-	**/
-	friend class FBuildAsyncWorker;
-	class FBuildAsyncWorker : public FNonAbandonableTask
-	{
-	public:
-		/** 
-		 * Constructor for async task 
-		 * @param	InDataDeriver	plugin to produce cache key and in the event of a miss, return the data.
-		 * @param	InCacheKey		Complete cache key for this data.
-		**/
-		FBuildAsyncWorker(FDerivedDataBackend* InBackend, FDerivedDataPluginInterface* InDataDeriver, const TCHAR* InCacheKey, FStringView InDebugContext, bool bInSynchronousForStats)
-		: bSuccess(false)
-		, bSynchronousForStats(bInSynchronousForStats)
-		, bDataWasBuilt(false)
-		, Backend(InBackend)
-		, DataDeriver(InDataDeriver)
-		, CacheKey(InCacheKey)
+public:
+	FLegacyFetchOrBuildTask(
+		FDerivedDataBackend* InBackend,
+		FStringView InDebugContext,
+		const TCHAR* InCacheKey,
+		FDerivedDataPluginInterface* InDataDeriver,
+		EPriority InPriority)
+		: Backend(InBackend)
 		, DebugContext(InDebugContext)
+		, CacheKey(InCacheKey)
+		, DataDeriver(InDataDeriver)
+		, Owner(InPriority)
+	{
+	}
+
+	~FLegacyFetchOrBuildTask()
+	{
+		Owner.Wait();
+		delete DataDeriver;
+		DataDeriver = nullptr;
+	}
+
+	/** Start an async fetch and build. Call WaitAsync() before accessing any outputs. */
+	void StartAsync()
+	{
+		Backend->AddToAsyncCompletionCounter(1);
+		BeginGet();
+	}
+
+	/** Poll whether an async fetch and build is complete. */
+	bool PollAsync() const
+	{
+		return Owner.Poll();
+	}
+
+	/** Wait for an async fetch and build. */
+	void WaitAsync()
+	{
+		Owner.Wait();
+
+		if (bNeedsSyncBuild)
 		{
+			{
+				FRequestBarrier Barrier(Owner);
+				ExecuteBuild(Status == EStatus::Ok ? EBuildMode::Verify : EBuildMode::Normal);
+			}
+			Owner.Wait();
 		}
+	}
 
-		/** Async worker that checks the cache backend and if that fails, calls the deriver to build the data and then puts the results to the cache **/
-		void DoWork()
-		{
-			TRACE_CPUPROFILER_EVENT_SCOPE(DDC_DoWork);
+	/** Execute the fetch and build synchronously. */
+	void ExecuteSync()
+	{
+		StartAsync();
+		WaitAsync();
+	}
 
-			const int64 NumBeforeDDC = Data.Num();
-			bool bGetResult = false;
-			{
-				TRACE_CPUPROFILER_EVENT_SCOPE(DDC_Get);
+	inline TArray64<uint8>& GetData() { return Data; }
+	inline EStatus GetStatus() const { return Status; }
+	inline bool GetDataWasBuilt() const { return bDataWasBuilt; }
 
-				INC_DWORD_STAT(STAT_DDC_NumGets);
-				STAT(double ThisTime = 0);
-				{
-					SCOPE_SECONDS_COUNTER(ThisTime);
-					FLegacyCacheGetRequest LegacyRequest;
-					LegacyRequest.Name = DebugContext;
-					LegacyRequest.Key = FLegacyCacheKey(CacheKey, Backend->GetMaxKeyLength());
-					FRequestOwner BlockingOwner(EPriority::Blocking);
-					Backend->GetRoot().LegacyGet({LegacyRequest}, BlockingOwner,
-						[this, &bGetResult](FLegacyCacheGetResponse&& Response)
-						{
-							const uint64 RawSize = Response.Value.GetRawSize();
-							bGetResult = Response.Status == EStatus::Ok && RawSize > 0 && RawSize < MAX_int64;
-							if (bGetResult)
-							{
-								const FCompositeBuffer& RawData = Response.Value.GetRawData();
-								Data.Reset(int64(RawSize));
-								for (const FSharedBuffer& Segment : RawData.GetSegments())
-								{
-									Data.Append(static_cast<const uint8*>(Segment.GetData()), int64(Segment.GetSize()));
-								}
-								UE_CLOG(Data.Num() != int64(RawSize), LogDerivedDataCache, Display,
-									TEXT("Copied %" INT64_FMT " bytes when %" INT64_FMT " bytes were expected for %s from '%s'"),
-									Data.Num(), int64(RawSize), *CacheKey, *Response.Name);
-							}
-						});
-					BlockingOwner.Wait();
-				}
-				INC_FLOAT_STAT_BY(STAT_DDC_SyncGetTime, bSynchronousForStats ? (float)ThisTime : 0.0f);
-			}
+private:
+	enum class EBuildMode { Normal, Verify };
+	enum class EBuildThread { Unknown, Caller };
 
-			if (bGetResult && Data.Num())
-			{
-				if (GVerifyDDC && DataDeriver && DataDeriver->IsDeterministic())
+	void BeginGet()
+	{
+		TRACE_CPUPROFILER_EVENT_SCOPE(DDC_Get);
+		ContinueCounter.store(2, std::memory_order_relaxed);
+
+		INC_DWORD_STAT(STAT_DDC_NumGets);
+		STAT(double ThisTime = 0);
+		{
+			SCOPE_SECONDS_COUNTER(ThisTime);
+			FRequestBarrier Barrier(Owner);
+			FLegacyCacheGetRequest Request;
+			Request.Name = DebugContext;
+			Request.Key = FLegacyCacheKey(CacheKey, Backend->GetMaxKeyLength());
+			Backend->GetRoot().LegacyGet({Request}, Owner,
+				[this](FLegacyCacheGetResponse&& Response)
 				{
-					TArray<uint8> CmpData;
-					DataDeriver->Build(CmpData);
-					const int64 NumInDDC = Data.Num() - NumBeforeDDC;
-					const int64 NumGenerated = CmpData.Num();
-					
-					bool bMatchesInSize = NumGenerated == NumInDDC;
-					bool bDifferentMemory = true;
-					int32 DifferentOffset = 0;
-					if (bMatchesInSize)
+					const uint64 RawSize = Response.Value.GetRawSize();
+					if (Response.Status == EStatus::Ok && (RawSize == 0 || RawSize > MAX_int64))
+					{
+						Response.Status = EStatus::Error;
+					}
+					if (Response.Status == EStatus::Ok)
 					{
-						bDifferentMemory = false;
-						for (int32 i = 0; i < NumGenerated; i++)
+						const FCompositeBuffer& RawData = Response.Value.GetRawData();
+						Data.Reset(int64(RawSize));
+						for (const FSharedBuffer& Segment : RawData.GetSegments())
 						{
-							if (CmpData[i] != Data[i])
-							{
-								bDifferentMemory = true;
-								DifferentOffset = i;
-								break;
-							}
+							Data.Append(static_cast<const uint8*>(Segment.GetData()), int64(Segment.GetSize()));
 						}
+						UE_CLOG(Data.Num() != int64(RawSize), LogDerivedDataCache, Display,
+							TEXT("Copied %" INT64_FMT " bytes when %" INT64_FMT " bytes were expected for %s from '%s'"),
+							Data.Num(), int64(RawSize), *CacheKey, *Response.Name);
 					}
-
-					if (!bMatchesInSize || bDifferentMemory)
+					Status = Response.Status;
+					if (ContinueCounter.fetch_sub(1, std::memory_order_release) == 1)
 					{
-						FString ErrMsg = FString::Printf(TEXT("There is a mismatch between the DDC data and the generated data for plugin (%s) for asset (%s). BytesInDDC:%" INT64_FMT ", BytesGenerated:%" INT64_FMT ", bDifferentMemory:%d, offset:%d"), DataDeriver->GetPluginName(), *DataDeriver->GetDebugContextString(), NumInDDC, NumGenerated, bDifferentMemory, DifferentOffset);
-						ensureMsgf(false, TEXT("%s"), *ErrMsg);
-						UE_LOG(LogDerivedDataCache, Error, TEXT("%s"), *ErrMsg );
+						EndGet(EBuildThread::Unknown);
 					}
-				}
+				});
+		}
+		INC_FLOAT_STAT_BY(STAT_DDC_SyncGetTime, Owner.GetPriority() == EPriority::Blocking ? (float)ThisTime : 0.0f);
 
-				check(Data.Num());
-				bSuccess = true;
-				delete DataDeriver;
-				DataDeriver = NULL;
-			}
-			else if (DataDeriver)
-			{
-				{
-					TRACE_CPUPROFILER_EVENT_SCOPE(DDC_Build);
+		if (Owner.GetPriority() == EPriority::Blocking)
+		{
+			// Wait here to allow blocking requests to continue on this thread even when the cache needed to
+			// switch threads to process the request.
+			Owner.Wait();
+		}
 
-					INC_DWORD_STAT(STAT_DDC_NumBuilds);
-					STAT(double ThisTime = 0);
-					{
-						SCOPE_SECONDS_COUNTER(ThisTime);
-						TArray<uint8> Data32;
-						bSuccess = DataDeriver->Build(Data32);
-						Data = TArray64<uint8>(MoveTemp(Data32));
-						bDataWasBuilt = true;
-					}
-					INC_FLOAT_STAT_BY(STAT_DDC_SyncBuildTime, bSynchronousForStats ? (float)ThisTime : 0.0f);
-				}
-				delete DataDeriver;
-				DataDeriver = NULL;
-				if (bSuccess)
-				{
-					check(Data.Num());
+		if (ContinueCounter.fetch_sub(1, std::memory_order_acquire) == 1)
+		{
+			EndGet(EBuildThread::Caller);
+		}
+	}
 
-					TRACE_CPUPROFILER_EVENT_SCOPE(DDC_Put);
+	void EndGet(EBuildThread BuildThread)
+	{
+		if (Status == EStatus::Ok && GVerifyDDC && DataDeriver && DataDeriver->IsDeterministic())
+		{
+			BeginBuild(EBuildMode::Verify, BuildThread);
+		}
+		else if (Status == EStatus::Error && DataDeriver)
+		{
+			bDataWasBuilt = true;
+			BeginBuild(EBuildMode::Normal, BuildThread);
+		}
+		else
+		{
+			EndTask();
+		}
+	}
 
-					INC_DWORD_STAT(STAT_DDC_NumPuts);
-					STAT(double ThisTime = 0);
+	void BeginBuild(EBuildMode BuildMode, EBuildThread BuildThread)
+	{
+		if (DataDeriver->IsBuildThreadsafe())
+		{
+			Owner.LaunchTask(*DebugContext, [this, BuildMode] { ExecuteBuild(BuildMode); });
+		}
+		else if (BuildThread == EBuildThread::Caller)
+		{
+			ExecuteBuild(BuildMode);
+		}
+		else
+		{
+			bNeedsSyncBuild = true;
+		}
+	}
+
+	void ExecuteBuild(EBuildMode BuildMode)
+	{
+		if (Owner.IsCanceled())
+		{
+			Status = EStatus::Canceled;
+			EndTask();
+			return;
+		}
+
+		TArray64<uint8> BuildData;
+
+		{
+			TRACE_CPUPROFILER_EVENT_SCOPE(DDC_Build);
+
+			INC_DWORD_STAT(STAT_DDC_NumBuilds);
+			STAT(double ThisTime = 0);
+			{
+				SCOPE_SECONDS_COUNTER(ThisTime);
+				TArray<uint8> Data32;
+				Status = DataDeriver->Build(Data32) ? EStatus::Ok : EStatus::Error;
+				BuildData = TArray64<uint8>(MoveTemp(Data32));
+			}
+			INC_FLOAT_STAT_BY(STAT_DDC_SyncBuildTime, Owner.GetPriority() == EPriority::Blocking || bNeedsSyncBuild ? (float)ThisTime : 0.0f);
+		}
+
+		EndBuild(BuildMode, MoveTemp(BuildData));
+	}
+
+	void EndBuild(EBuildMode BuildMode, TArray64<uint8>&& BuildData)
+	{
+		if (Status == EStatus::Ok && BuildMode == EBuildMode::Verify)
+		{
+			const bool bMatchesInSize = Data.Num() == BuildData.Num();
+			bool bDifferentMemory = !bMatchesInSize;
+			int64 DifferentOffset = 0;
+			if (bMatchesInSize)
+			{
+				for (int64 Index = 0, Count = Data.Num(); Index < Count; ++Index)
+				{
+					if (Data[Index] != BuildData[Index])
 					{
-						SCOPE_SECONDS_COUNTER(ThisTime);
-						FLegacyCachePutRequest LegacyRequest;
-						LegacyRequest.Name = DebugContext;
-						LegacyRequest.Key = FLegacyCacheKey(CacheKey, Backend->GetMaxKeyLength());
-						LegacyRequest.Value = FLegacyCacheValue(FCompositeBuffer(FSharedBuffer::Clone(MakeMemoryView(Data))));
-						FRequestOwner AsyncOwner(EPriority::Normal);
-						Backend->GetRoot().LegacyPut({LegacyRequest}, AsyncOwner, [](auto&&){});
-						AsyncOwner.KeepAlive();
+						bDifferentMemory = true;
+						DifferentOffset = Index;
+						break;
 					}
-					INC_FLOAT_STAT_BY(STAT_DDC_PutTime, bSynchronousForStats ? (float)ThisTime : 0.0f);
 				}
 			}
-			if (!bSuccess)
+
+			if (!bMatchesInSize || bDifferentMemory)
 			{
-				Data.Empty();
+				FString ErrMsg = FString::Printf(TEXT("There is a mismatch between the DDC data and the generated data for plugin (%s) for asset (%s). BytesInDDC:%" INT64_FMT ", BytesGenerated:%" INT64_FMT ", bDifferentMemory:%d, Offset:%" INT64_FMT),
+					DataDeriver->GetPluginName(), *DataDeriver->GetDebugContextString(), Data.Num(), BuildData.Num(), bDifferentMemory, DifferentOffset);
+				ensureMsgf(false, TEXT("%s"), *ErrMsg);
+				UE_LOG(LogDerivedDataCache, Error, TEXT("%s"), *ErrMsg);
 			}
-			Backend->AddToAsyncCompletionCounter(-1);
 		}
 
-		FORCEINLINE TStatId GetStatId() const
+		delete DataDeriver;
+		DataDeriver = nullptr;
+
+		if (Status == EStatus::Ok && BuildMode == EBuildMode::Normal)
 		{
-			RETURN_QUICK_DECLARE_CYCLE_STAT(FBuildAsyncWorker, STATGROUP_ThreadPoolAsyncTasks);
+			Data = MoveTemp(BuildData);
+			BeginAsyncPut();
 		}
 
-		/** true in the case of a cache hit, otherwise the result of the deriver build call **/
-		bool							bSuccess;
-		/** true if we should record the timing **/
-		bool							bSynchronousForStats;
-		/** true if we had to build the data */
-		bool							bDataWasBuilt;
-		/** Backend graph to execute against. */
-		FDerivedDataBackend*			Backend;
-		/** Data dervier we are operating on **/
-		FDerivedDataPluginInterface*	DataDeriver;
-		/** Cache key associated with this build **/
-		FString							CacheKey;
-		/** Context from the caller */
-		FSharedString					DebugContext;
-		/** Data to return to caller, later **/
-		TArray64<uint8>					Data;
-	};
+		EndTask();
+	}
 
-public:
+	void BeginAsyncPut()
+	{
+		check(Data.Num());
 
+		TRACE_CPUPROFILER_EVENT_SCOPE(DDC_Put);
+
+		INC_DWORD_STAT(STAT_DDC_NumPuts);
+		STAT(double ThisTime = 0);
+		{
+			SCOPE_SECONDS_COUNTER(ThisTime);
+			FLegacyCachePutRequest Request;
+			Request.Name = DebugContext;
+			Request.Key = FLegacyCacheKey(CacheKey, Backend->GetMaxKeyLength());
+			Request.Value = FLegacyCacheValue(FCompositeBuffer(FSharedBuffer::Clone(MakeMemoryView(Data))));
+			FRequestOwner AsyncOwner(EPriority::Normal);
+			Backend->GetRoot().LegacyPut({Request}, AsyncOwner, [](auto&&) {});
+			AsyncOwner.KeepAlive();
+		}
+		INC_FLOAT_STAT_BY(STAT_DDC_PutTime, Owner.GetPriority() == EPriority::Blocking ? (float)ThisTime : 0.0f);
+	}
+
+	void EndTask()
+	{
+		if (Status != EStatus::Ok)
+		{
+			Data.Empty();
+		}
+		Backend->AddToAsyncCompletionCounter(-1);
+	}
+
+	FDerivedDataBackend* Backend;
+	FSharedString DebugContext;
+	FString CacheKey;
+	FDerivedDataPluginInterface* DataDeriver;
+	FRequestOwner Owner;
+	TArray64<uint8> Data;
+	EStatus Status = EStatus::Error;
+	bool bNeedsSyncBuild = false;
+	bool bDataWasBuilt = false;
+
+	/** Counter to control where to continue execution from. Avoids confusing scoping of timers. */
+	std::atomic<uint8> ContinueCounter = 0;
+};
+
+/**
+ * Implementation of the derived data cache
+ * This API is fully threadsafe
+**/
+class FDerivedDataCache final
+	: public FDerivedDataCacheInterface
+	, public ICache
+	, public ICacheStoreMaintainer
+	, public IDDCCleanup
+{
+public:
 	/** Constructor, called once to cereate a singleton **/
 	FDerivedDataCache()
 		: CurrentHandle(19248) // we will skip some potential handles to catch errors
@@ -591,11 +685,6 @@ public:
 	{
 		WaitForQuiescence(true);
 		FScopeLock ScopeLock(&SynchronizationObject);
-		for (TMap<uint32,FAsyncTask<FBuildAsyncWorker>*>::TIterator It(PendingTasks); It; ++It)
-		{
-			It.Value()->EnsureCompletion();
-			delete It.Value();
-		}
 		PendingTasks.Empty();
 		delete Backend;
 	}
@@ -606,38 +695,28 @@ public:
 		check(DataDeriver);
 		FString CacheKey = FDerivedDataCache::BuildCacheKey(DataDeriver);
 		UE_LOG(LogDerivedDataCache, VeryVerbose, TEXT("GetSynchronous %s from '%s'"), *CacheKey, *DataDeriver->GetDebugContextString());
-		FAsyncTask<FBuildAsyncWorker> PendingTask(Backend, DataDeriver, *CacheKey, DataDeriver->GetDebugContextString(), true);
-		AddToAsyncCompletionCounter(1);
-		PendingTask.StartSynchronousTask(EQueuedWorkPriority::Normal, EQueuedWorkFlags::DoNotRunInsideBusyWait);
-		OutData = TArray<uint8>(MoveTemp(PendingTask.GetTask().Data));
+		FLegacyFetchOrBuildTask PendingTask(Backend, DataDeriver->GetDebugContextString(), *CacheKey, DataDeriver, EPriority::Blocking);
+		PendingTask.ExecuteSync();
+		OutData = TArray<uint8>(MoveTemp(PendingTask.GetData()));
 		if (bDataWasBuilt)
 		{
-			*bDataWasBuilt = PendingTask.GetTask().bDataWasBuilt;
+			*bDataWasBuilt = PendingTask.GetDataWasBuilt();
 		}
-		return PendingTask.GetTask().bSuccess;
+		return PendingTask.GetStatus() == EStatus::Ok;
 	}
 
 	virtual uint32 GetAsynchronous(FDerivedDataPluginInterface* DataDeriver) override
 	{
 		DDC_SCOPE_CYCLE_COUNTER(DDC_GetAsynchronous);
-		FScopeLock ScopeLock(&SynchronizationObject);
 		const uint32 Handle = NextHandle();
 		FString CacheKey = FDerivedDataCache::BuildCacheKey(DataDeriver);
 		UE_LOG(LogDerivedDataCache, VeryVerbose, TEXT("GetAsynchronous %s from '%s', Handle %d"), *CacheKey, *DataDeriver->GetDebugContextString(), Handle);
-		const bool bSync = !DataDeriver->IsBuildThreadsafe();
-		FAsyncTask<FBuildAsyncWorker>* AsyncTask = new FAsyncTask<FBuildAsyncWorker>(Backend, DataDeriver, *CacheKey, DataDeriver->GetDebugContextString(), bSync);
+		TUniquePtr<FLegacyFetchOrBuildTask> AsyncTask = MakeUnique<FLegacyFetchOrBuildTask>(Backend, DataDeriver->GetDebugContextString(), *CacheKey, DataDeriver, EPriority::Blocking);
+		FLegacyFetchOrBuildTask* LocalAsyncTask = AsyncTask.Get();
+		FScopeLock ScopeLock(&SynchronizationObject);
 		check(!PendingTasks.Contains(Handle));
-		PendingTasks.Add(Handle,AsyncTask);
-		AddToAsyncCompletionCounter(1);
-		if (!bSync)
-		{
-			AsyncTask->StartBackgroundTask(DataDeriver->GetCustomThreadPool(), EQueuedWorkPriority::Normal, EQueuedWorkFlags::DoNotRunInsideBusyWait);
-		}
-		else
-		{
-			AsyncTask->StartSynchronousTask(EQueuedWorkPriority::Normal, EQueuedWorkFlags::DoNotRunInsideBusyWait);
-		}
-		// Must return a valid handle
+		PendingTasks.Add(Handle, MoveTemp(AsyncTask));
+		LocalAsyncTask->StartAsync();
 		check(Handle != 0);
 		return Handle;
 	}
@@ -645,13 +724,16 @@ public:
 	virtual bool PollAsynchronousCompletion(uint32 Handle) override
 	{
 		DDC_SCOPE_CYCLE_COUNTER(DDC_PollAsynchronousCompletion);
-		FAsyncTask<FBuildAsyncWorker>* AsyncTask = NULL;
+		FLegacyFetchOrBuildTask* AsyncTask = nullptr;
 		{
 			FScopeLock ScopeLock(&SynchronizationObject);
-			AsyncTask = PendingTasks.FindRef(Handle);
+			if (TUniquePtr<FLegacyFetchOrBuildTask>* LocalAsyncTask = PendingTasks.Find(Handle))
+			{
+				AsyncTask = LocalAsyncTask->Get();
+			}
 		}
 		check(AsyncTask);
-		return AsyncTask->IsDone();
+		return AsyncTask->PollAsync();
 	}
 
 	virtual void WaitAsynchronousCompletion(uint32 Handle) override
@@ -660,13 +742,16 @@ public:
 		STAT(double ThisTime = 0);
 		{
 			SCOPE_SECONDS_COUNTER(ThisTime);
-			FAsyncTask<FBuildAsyncWorker>* AsyncTask = NULL;
+			FLegacyFetchOrBuildTask* AsyncTask = nullptr;
 			{
 				FScopeLock ScopeLock(&SynchronizationObject);
-				AsyncTask = PendingTasks.FindRef(Handle);
+				if (TUniquePtr<FLegacyFetchOrBuildTask>* LocalAsyncTask = PendingTasks.Find(Handle))
+				{
+					AsyncTask = LocalAsyncTask->Get();
+				}
 			}
 			check(AsyncTask);
-			AsyncTask->EnsureCompletion();
+			AsyncTask->WaitAsync();
 			UE_LOG(LogDerivedDataCache, Verbose, TEXT("WaitAsynchronousCompletion, Handle %d"), Handle);
 		}
 		INC_FLOAT_STAT_BY(STAT_DDC_ASyncWaitTime,(float)ThisTime);
@@ -676,27 +761,25 @@ public:
 	bool GetAsynchronousResultsByHandle(uint32 Handle, DataType& OutData, bool* bOutDataWasBuilt)
 	{
 		DDC_SCOPE_CYCLE_COUNTER(DDC_GetAsynchronousResults);
-		FAsyncTask<FBuildAsyncWorker>* AsyncTask = NULL;
+		TUniquePtr<FLegacyFetchOrBuildTask> AsyncTask;
 		{
 			FScopeLock ScopeLock(&SynchronizationObject);
-			PendingTasks.RemoveAndCopyValue(Handle,AsyncTask);
+			PendingTasks.RemoveAndCopyValue(Handle, AsyncTask);
 		}
 		check(AsyncTask);
-		const bool bDataWasBuilt = AsyncTask->GetTask().bDataWasBuilt;
+		const bool bDataWasBuilt = AsyncTask->GetDataWasBuilt();
 		if (bOutDataWasBuilt)
 		{
 			*bOutDataWasBuilt = bDataWasBuilt;
 		}
-		if (!AsyncTask->GetTask().bSuccess)
+		if (AsyncTask->GetStatus() != EStatus::Ok)
 		{
 			UE_LOG(LogDerivedDataCache, Verbose, TEXT("GetAsynchronousResults, bDataWasBuilt: %d, Handle %d, FAILED"), (int32)bDataWasBuilt, Handle);
-			delete AsyncTask;
 			return false;
 		}
 
 		UE_LOG(LogDerivedDataCache, Verbose, TEXT("GetAsynchronousResults, bDataWasBuilt: %d, Handle %d, SUCCESS"), (int32)bDataWasBuilt, Handle);
-		OutData = DataType(MoveTemp(AsyncTask->GetTask().Data));
-		delete AsyncTask;
+		OutData = DataType(MoveTemp(AsyncTask->GetData()));
 		check(OutData.Num());
 		return true;
 	}
@@ -716,11 +799,10 @@ public:
 	{
 		DDC_SCOPE_CYCLE_COUNTER(DDC_GetSynchronous_Data);
 		UE_LOG(LogDerivedDataCache, VeryVerbose, TEXT("GetSynchronous %s from '%.*s'"), CacheKey, DebugContext.Len(), DebugContext.GetData());
-		FAsyncTask<FBuildAsyncWorker> PendingTask(Backend, nullptr, CacheKey, DebugContext, true);
-		AddToAsyncCompletionCounter(1);
-		PendingTask.StartSynchronousTask(EQueuedWorkPriority::Normal, EQueuedWorkFlags::DoNotRunInsideBusyWait);
-		OutData = DataType(MoveTemp(PendingTask.GetTask().Data));
-		return PendingTask.GetTask().bSuccess;
+		FLegacyFetchOrBuildTask PendingTask(Backend, DebugContext, CacheKey, nullptr, EPriority::Blocking);
+		PendingTask.ExecuteSync();
+		OutData = DataType(MoveTemp(PendingTask.GetData()));
+		return PendingTask.GetStatus() == EStatus::Ok;
 	}
 
 	virtual bool GetSynchronous(const TCHAR* CacheKey, TArray<uint8>& OutData, FStringView DebugContext) override
@@ -736,15 +818,15 @@ public:
 	virtual uint32 GetAsynchronous(const TCHAR* CacheKey, FStringView DebugContext) override
 	{
 		DDC_SCOPE_CYCLE_COUNTER(DDC_GetAsynchronous_Handle);
-		FScopeLock ScopeLock(&SynchronizationObject);
 		const uint32 Handle = NextHandle();
 		UE_LOG(LogDerivedDataCache, VeryVerbose, TEXT("GetAsynchronous %s from '%.*s', Handle %d"), CacheKey, DebugContext.Len(), DebugContext.GetData(), Handle);
-		FAsyncTask<FBuildAsyncWorker>* AsyncTask = new FAsyncTask<FBuildAsyncWorker>(Backend, nullptr, CacheKey, DebugContext, false);
+		TUniquePtr<FLegacyFetchOrBuildTask> AsyncTask = MakeUnique<FLegacyFetchOrBuildTask>(Backend, DebugContext, CacheKey, nullptr, EPriority::Normal);
+		FLegacyFetchOrBuildTask* LocalAsyncTask = AsyncTask.Get();
+		FScopeLock ScopeLock(&SynchronizationObject);
 		check(!PendingTasks.Contains(Handle));
-		PendingTasks.Add(Handle, AsyncTask);
-		AddToAsyncCompletionCounter(1);
-		// This request is I/O only, doesn't do any processing, send it to the I/O only thread-pool to avoid wasting worker threads on long I/O waits.
-		AsyncTask->StartBackgroundTask(GCacheThreadPool, EQueuedWorkPriority::Normal, EQueuedWorkFlags::DoNotRunInsideBusyWait);
+		PendingTasks.Add(Handle, MoveTemp(AsyncTask));
+		LocalAsyncTask->StartAsync();
+		check(Handle != 0);
 		return Handle;
 	}
 
@@ -1074,7 +1156,7 @@ private:
 	/** Object used for synchronization via a scoped lock **/
 	FCriticalSection			SynchronizationObject;
 	/** Map of handle to pending task **/
-	TMap<uint32,FAsyncTask<FBuildAsyncWorker>*>	PendingTasks;
+	TMap<uint32, TUniquePtr<FLegacyFetchOrBuildTask>> PendingTasks;
 
 	/** Cache notification delegate */
 	FOnDDCNotification DDCNotificationEvent;
diff --git a/Engine/Source/Developer/DerivedDataCache/Public/DerivedDataPluginInterface.h b/Engine/Source/Developer/DerivedDataCache/Public/DerivedDataPluginInterface.h
--- a/Engine/Source/Developer/DerivedDataCache/Public/DerivedDataPluginInterface.h
+++ b/Engine/Source/Developer/DerivedDataCache/Public/DerivedDataPluginInterface.h
@@ -41,6 +41,7 @@ public:
 	* If IsBuildThreadsafe() returns true, then this is the thread pool that is used to run the background task.
 	* @return	The threadpool to run the task with
 	**/
+	UE_DEPRECATED(5.5, "GetCustomThreadPool() is no longer supported.")
 	virtual FQueuedThreadPool* GetCustomThreadPool() const { return GThreadPool; }
 
 	/** Indicated that this plugin generates deterministic data. This is used for DDC verification */
