DDC: Added a missing request barrier

#rnx

diff --git a/Engine/Source/Developer/DerivedDataCache/Private/DerivedDataCache.cpp b/Engine/Source/Developer/DerivedDataCache/Private/DerivedDataCache.cpp
--- a/Engine/Source/Developer/DerivedDataCache/Private/DerivedDataCache.cpp
+++ b/Engine/Source/Developer/DerivedDataCache/Private/DerivedDataCache.cpp
@@ -494,6 +494,7 @@ private:
 
 		if (ContinueCounter.fetch_sub(1, std::memory_order_acquire) == 1)
 		{
+			FRequestBarrier Barrier(Owner);
 			EndGet(EBuildThread::Caller);
 		}
 	}
