DDC: Extracted FHttpCacheStoreRequestQueue for reuse by Zen

#rnx

diff --git a/Engine/Source/Developer/DerivedDataCache/Private/DerivedDataHttpRequestQueue.cpp b/Engine/Source/Developer/DerivedDataCache/Private/DerivedDataHttpRequestQueue.cpp
new file mode 100644
--- /dev/null
+++ b/Engine/Source/Developer/DerivedDataCache/Private/DerivedDataHttpRequestQueue.cpp
@@ -0,0 +1,132 @@
+// Copyright Epic Games, Inc. All Rights Reserved.
+
+#include "DerivedDataHttpRequestQueue.h"
+
+#include "Async/ManualResetEvent.h"
+#include "DerivedDataRequest.h"
+#include "DerivedDataRequestOwner.h"
+#include "Misc/ScopeExit.h"
+
+namespace UE::DerivedData
+{
+
+class FHttpRequestQueue::FQueueRequest : FRequestBase
+{
+public:
+	FQueueRequest(IRequestOwner& InOwner, FOnRequest&& InOnRequest)
+		: Owner(InOwner)
+		, OnRequest(MoveTemp(InOnRequest))
+	{
+		AddRef(); // Release() is called by ClaimRequest()
+		Owner.Begin(this);
+	}
+
+	bool TryClaimRequest(THttpUniquePtr<IHttpRequest>&& Request)
+	{
+		ON_SCOPE_EXIT{ Release(); };
+		return TryComplete(MoveTemp(Request));
+	}
+
+private:
+	bool TryComplete(THttpUniquePtr<IHttpRequest>&& Request)
+	{
+		if (bComplete.exchange(true))
+		{
+			return false;
+		}
+		Owner.End(this, [this](THttpUniquePtr<IHttpRequest>&& Request)
+		{
+			OnRequest(MoveTemp(Request));
+			OnComplete.Notify();
+		}, MoveTemp(Request));
+		return true;
+	}
+
+	void SetPriority(EPriority Priority) final
+	{
+	}
+
+	void Cancel() final
+	{
+		if (!TryComplete({}))
+		{
+			TRACE_CPUPROFILER_EVENT_SCOPE(HttpDDC_CancelOperation);
+			OnComplete.Wait();
+		}
+	}
+
+	void Wait() final
+	{
+		TRACE_CPUPROFILER_EVENT_SCOPE(HttpDDC_WaitOperation);
+		OnComplete.Wait();
+	}
+
+	IRequestOwner& Owner;
+	FOnRequest OnRequest;
+	FManualResetEvent OnComplete;
+	std::atomic<bool> bComplete = false;
+};
+
+void FHttpRequestQueue::Initialize(IHttpConnectionPool& ConnectionPool, const FHttpClientParams& ClientParams)
+{
+	FHttpClientParams QueueParams = ClientParams;
+	QueueParams.OnDestroyRequest = [this, OnDestroyRequest = MoveTemp(QueueParams.OnDestroyRequest)]
+		{
+			if (OnDestroyRequest)
+			{
+				OnDestroyRequest();
+			}
+			if (!Queue.IsEmpty())
+			{
+				if (THttpUniquePtr<IHttpRequest> Request = Client->TryCreateRequest({}))
+				{
+					TryGiveRequestToQueue(MoveTemp(Request));
+				}
+			}
+		};
+	Client = ConnectionPool.CreateClient(QueueParams);
+}
+
+void FHttpRequestQueue::CreateRequestAsync(IRequestOwner& Owner, const FHttpRequestParams& Params, FOnRequest&& OnRequest)
+{
+	if (Params.bIgnoreMaxRequests)
+	{
+		THttpUniquePtr<IHttpRequest> Request = Client->TryCreateRequest(Params);
+		checkf(Request, TEXT("IHttpClient::TryCreateRequest returned null in spite of bIgnoreMaxRequests."));
+		OnRequest(MoveTemp(Request));
+		return;
+	}
+
+	while (THttpUniquePtr<IHttpRequest> Request = Client->TryCreateRequest(Params))
+	{
+		if (!TryGiveRequestToQueue(MoveTemp(Request)))
+		{
+			OnRequest(MoveTemp(Request));
+			return;
+		}
+	}
+
+	Queue.Push(new FQueueRequest(Owner, MoveTemp(OnRequest)));
+
+	while (THttpUniquePtr<IHttpRequest> Request = Client->TryCreateRequest(Params))
+	{
+		if (!TryGiveRequestToQueue(MoveTemp(Request)))
+		{
+			return;
+		}
+	}
+}
+
+bool FHttpRequestQueue::TryGiveRequestToQueue(THttpUniquePtr<IHttpRequest>&& Request)
+{
+	while (FQueueRequest* Waiter = Queue.Pop())
+	{
+		if (Waiter->TryClaimRequest(MoveTemp(Request)))
+		{
+			return true;
+		}
+	}
+	return false;
+}
+
+} // UE::DerivedData
diff --git a/Engine/Source/Developer/DerivedDataCache/Private/DerivedDataHttpRequestQueue.h b/Engine/Source/Developer/DerivedDataCache/Private/DerivedDataHttpRequestQueue.h
new file mode 100644
--- /dev/null
+++ b/Engine/Source/Developer/DerivedDataCache/Private/DerivedDataHttpRequestQueue.h
@@ -0,0 +1,33 @@
+// Copyright Epic Games, Inc. All Rights Reserved.
+
+#pragma once
+
+#include "Containers/LockFreeList.h"
+#include "Http/HttpClient.h"
+
+template <typename FuncType> class TUniqueFunction;
+
+namespace UE::DerivedData { class IRequestOwner; }
+
+namespace UE::DerivedData
+{
+
+class FHttpRequestQueue
+{
+public:
+	using FOnRequest = TUniqueFunction<void(THttpUniquePtr<IHttpRequest>&& Request)>;
+
+	void Initialize(IHttpConnectionPool& ConnectionPool, const FHttpClientParams& ClientParams);
+
+	void CreateRequestAsync(IRequestOwner& Owner, const FHttpRequestParams& Params, FOnRequest&& OnRequest);
+
+private:
+	bool TryGiveRequestToQueue(THttpUniquePtr<IHttpRequest>&& Request);
+
+	class FQueueRequest;
+
+	THttpUniquePtr<IHttpClient> Client;
+	TLockFreePointerListFIFO<FQueueRequest, 0> Queue;
+};
+
+} // UE::DerivedData
diff --git a/Engine/Source/Developer/DerivedDataCache/Private/HttpCacheStore.cpp b/Engine/Source/Developer/DerivedDataCache/Private/HttpCacheStore.cpp
--- a/Engine/Source/Developer/DerivedDataCache/Private/HttpCacheStore.cpp
+++ b/Engine/Source/Developer/DerivedDataCache/Private/HttpCacheStore.cpp
@@ -17,6 +17,7 @@
 #include "DerivedDataCacheRecord.h"
 #include "DerivedDataCacheUsageStats.h"
 #include "DerivedDataChunk.h"
+#include "DerivedDataHttpRequestQueue.h"
 #include "DerivedDataRequest.h"
 #include "DerivedDataRequestOwner.h"
 #include "DerivedDataValue.h"
@@ -174,135 +175,6 @@ static bool TryResolveCanonicalHost(const FAnsiStringView Uri, FAnsiStringBuilde
 	return false;
 }
 
-class FHttpCacheStoreRequestQueue
-{
-public:
-	using FOnRequest = TUniqueFunction<void (THttpUniquePtr<IHttpRequest>&& Request)>;
-
-	void Initialize(IHttpConnectionPool& ConnectionPool, const FHttpClientParams& ClientParams)
-	{
-		FHttpClientParams QueueParams = ClientParams;
-		QueueParams.OnDestroyRequest = [this, OnDestroyRequest = MoveTemp(QueueParams.OnDestroyRequest)]
-		{
-			if (OnDestroyRequest)
-			{
-				OnDestroyRequest();
-			}
-			if (!Queue.IsEmpty())
-			{
-				if (THttpUniquePtr<IHttpRequest> Request = Client->TryCreateRequest({}))
-				{
-					TryGiveRequestToQueue(MoveTemp(Request));
-				}
-			}
-		};
-		Client = ConnectionPool.CreateClient(QueueParams);
-	}
-
-	void CreateRequestAsync(IRequestOwner& Owner, const FHttpRequestParams& Params, FOnRequest&& OnRequest)
-	{
-		if (Params.bIgnoreMaxRequests)
-		{
-			THttpUniquePtr<IHttpRequest> Request = Client->TryCreateRequest(Params);
-			checkf(Request, TEXT("IHttpClient::TryCreateRequest returned null in spite of bIgnoreMaxRequests."));
-			OnRequest(MoveTemp(Request));
-			return;
-		}
-
-		while (THttpUniquePtr<IHttpRequest> Request = Client->TryCreateRequest(Params))
-		{
-			if (!TryGiveRequestToQueue(MoveTemp(Request)))
-			{
-				OnRequest(MoveTemp(Request));
-				return;
-			}
-		}
-
-		Queue.Push(new FQueueRequest(Owner, MoveTemp(OnRequest)));
-
-		while (THttpUniquePtr<IHttpRequest> Request = Client->TryCreateRequest(Params))
-		{
-			if (!TryGiveRequestToQueue(MoveTemp(Request)))
-			{
-				return;
-			}
-		}
-	}
-
-private:
-	bool TryGiveRequestToQueue(THttpUniquePtr<IHttpRequest>&& Request)
-	{
-		while (FQueueRequest* Waiter = Queue.Pop())
-		{
-			if (Waiter->TryClaimRequest(MoveTemp(Request)))
-			{
-				return true;
-			}
-		}
-		return false;
-	}
-
-	class FQueueRequest : FRequestBase
-	{
-	public:
-		FQueueRequest(IRequestOwner& InOwner, FOnRequest&& InOnRequest)
-			: Owner(InOwner)
-			, OnRequest(MoveTemp(InOnRequest))
-		{
-			AddRef(); // Release() is called by ClaimRequest()
-			Owner.Begin(this);
-		}
-
-		bool TryClaimRequest(THttpUniquePtr<IHttpRequest>&& Request)
-		{
-			ON_SCOPE_EXIT { Release(); };
-			return TryComplete(MoveTemp(Request));
-		}
-
-	private:
-		bool TryComplete(THttpUniquePtr<IHttpRequest>&& Request)
-		{
-			if (bComplete.exchange(true))
-			{
-				return false;
-			}
-			Owner.End(this, [this](THttpUniquePtr<IHttpRequest>&& Request)
-			{
-				OnRequest(MoveTemp(Request));
-				OnComplete.Notify();
-			}, MoveTemp(Request));
-			return true;
-		}
-
-		void SetPriority(EPriority Priority) final
-		{
-		}
-
-		void Cancel() final
-		{
-			if (!TryComplete({}))
-			{
-				TRACE_CPUPROFILER_EVENT_SCOPE(HttpDDC_CancelOperation);
-				OnComplete.Wait();
-			}
-		}
-
-		void Wait() final
-		{
-			TRACE_CPUPROFILER_EVENT_SCOPE(HttpDDC_WaitOperation);
-			OnComplete.Wait();
-		}
-
-		IRequestOwner& Owner;
-		FOnRequest OnRequest;
-		FManualResetEvent OnComplete;
-		std::atomic<bool> bComplete = false;
-	};
-
-	THttpUniquePtr<IHttpClient> Client;
-	TLockFreePointerListFIFO<FQueueRequest, 0> Queue;
-};
-
 /**
  * Encapsulation for access token shared by all requests.
  */
@@ -455,10 +327,10 @@ private:
 	FDerivedDataCacheUsageStats UsageStats;
 	FBackendDebugOptions DebugOptions;
 	THttpUniquePtr<IHttpConnectionPool> ConnectionPool;
-	FHttpCacheStoreRequestQueue GetRequestQueue;
-	FHttpCacheStoreRequestQueue PutRefRequestQueue;
-	FHttpCacheStoreRequestQueue PutBlobsRequestQueue;
-	FHttpCacheStoreRequestQueue PutFinalizeRequestQueue;
+	FHttpRequestQueue GetRequestQueue;
+	FHttpRequestQueue PutRefRequestQueue;
+	FHttpRequestQueue PutBlobsRequestQueue;
+	FHttpRequestQueue PutFinalizeRequestQueue;
 
 	FCriticalSection AccessCs;
 	TUniquePtr<FHttpAccessToken> Access;
@@ -489,7 +361,7 @@ private:
 
 	class FHttpOperation;
 
-	FHttpCacheStoreRequestQueue& PickRequestQueue(EOperationCategory Category);
+	FHttpRequestQueue& PickRequestQueue(EOperationCategory Category);
 
 	/** Invokes the callback when an operation is available, or with null if canceled. */
 	void WaitForHttpOperationAsync(IRequestOwner& Owner, EOperationCategory Category, TUniqueFunction<void (TUniquePtr<FHttpOperation>&&)>&& OnOperation);
@@ -2590,7 +2462,7 @@ void FHttpCacheStore::SetAccessTokenAndUnlock(FScopeLock& Lock, FStringView Toke
 	}
 }
 
-FHttpCacheStoreRequestQueue& FHttpCacheStore::PickRequestQueue(EOperationCategory Category)
+FHttpRequestQueue& FHttpCacheStore::PickRequestQueue(EOperationCategory Category)
 {
 	switch (Category)
 	{
@@ -2635,7 +2507,7 @@ void FHttpCacheStore::WaitForHttpOperationAsync(IRequestOwner& Owner, EOperation
 void FHttpCacheStore::WaitForHttpRequestAsync(IRequestOwner& Owner, EOperationCategory Category, TUniqueFunction<void (THttpUniquePtr<IHttpRequest>&&)>&& OnRequest)
 {
 	FHttpRequestParams Params;
-	FHttpCacheStoreRequestQueue& RequestQueue = PickRequestQueue(Category);
+	FHttpRequestQueue& RequestQueue = PickRequestQueue(Category);
 	RequestQueue.CreateRequestAsync(Owner, Params, MoveTemp(OnRequest));
 }
 
diff --git a/Engine/Source/Developer/DerivedDataCache/Private/ZenCacheStore.cpp b/Engine/Source/Developer/DerivedDataCache/Private/ZenCacheStore.cpp
--- a/Engine/Source/Developer/DerivedDataCache/Private/ZenCacheStore.cpp
+++ b/Engine/Source/Developer/DerivedDataCache/Private/ZenCacheStore.cpp
@@ -219,7 +219,7 @@ private:
 	ICacheStoreOwner* StoreOwner = nullptr;
 	ICacheStoreStats* StoreStats = nullptr;
 	THttpUniquePtr<IHttpConnectionPool> ConnectionPool;
-	FHttpRequestQueue RequestQueue;
+	UE::FHttpRequestQueue RequestQueue;
 	bool bIsUsable = false;
 	bool bIsLocalConnection = false;
 	bool bTryEvaluatePerformance = false;
@@ -1710,7 +1710,7 @@ void FZenCacheStore::Initialize(const FZenCacheStoreParams& Params)
 	ClientParams.LowSpeedLimit = 1;
 	ClientParams.LowSpeedTime = 25;
 	ClientParams.bBypassProxy = Params.bBypassProxy;
-	RequestQueue = FHttpRequestQueue(*ConnectionPool, ClientParams);
+	RequestQueue = UE::FHttpRequestQueue(*ConnectionPool, ClientParams);
 
 	bIsLocalConnection = ZenService.GetInstance().IsServiceRunningLocally() || ZenService.GetInstance().GetServiceSettings().IsAutoLaunch();
 	bIsUsable = true;
