DDC: Eliminated most blocking during requests in ZenCacheStore

diff --git a/Engine/Source/Developer/DerivedDataCache/Private/ZenCacheStore.cpp b/Engine/Source/Developer/DerivedDataCache/Private/ZenCacheStore.cpp
--- a/Engine/Source/Developer/DerivedDataCache/Private/ZenCacheStore.cpp
+++ b/Engine/Source/Developer/DerivedDataCache/Private/ZenCacheStore.cpp
@@ -14,6 +14,7 @@
 #include "DerivedDataCacheRecord.h"
 #include "DerivedDataCacheUsageStats.h"
 #include "DerivedDataChunk.h"
+#include "DerivedDataHttpRequestQueue.h"
 #include "DerivedDataRequest.h"
 #include "DerivedDataRequestOwner.h"
 #include "Experimental/ZenStatistics.h"
@@ -164,8 +165,8 @@ private:
 	bool IsServiceReady();
 
 	static FCompositeBuffer SaveRpcPackage(const FCbPackage& Package);
-	THttpUniquePtr<IHttpRequest> CreateRpcRequest();
-	using FOnRpcComplete = TUniqueFunction<void(THttpUniquePtr<IHttpResponse>& HttpResponse, FCbPackage& Response)>;
+	void CreateRpcRequest(IRequestOwner& Owner, FHttpRequestQueue::FOnRequest&& OnRequest);
+	using FOnRpcComplete = TUniqueFunction<void(const THttpUniquePtr<IHttpResponse>& HttpResponse, const FCbPackage& Response)>;
 	void EnqueueAsyncRpc(IRequestOwner& Owner, FCbObject RequestObject, FOnRpcComplete&& OnComplete);
 	void EnqueueAsyncRpc(IRequestOwner& Owner, const FCbPackage& RequestPackage, FOnRpcComplete&& OnComplete);
 
@@ -206,7 +207,7 @@ private:
 		Error,
 	};
 
-	using FOnHealthComplete = TUniqueFunction<void(THttpUniquePtr<IHttpResponse>& HttpResponse, EHealth Health)>;
+	using FOnHealthComplete = TUniqueFunction<void(const THttpUniquePtr<IHttpResponse>& HttpResponse, EHealth Health)>;
 	class FHealthReceiver;
 	class FAsyncHealthReceiver;
 
@@ -219,7 +220,7 @@ private:
 	ICacheStoreOwner* StoreOwner = nullptr;
 	ICacheStoreStats* StoreStats = nullptr;
 	THttpUniquePtr<IHttpConnectionPool> ConnectionPool;
-	UE::FHttpRequestQueue RequestQueue;
+	FHttpRequestQueue RequestQueue;
 	bool bIsUsable = false;
 	bool bIsLocalConnection = false;
 	bool bTryEvaluatePerformance = false;
@@ -362,14 +363,16 @@ public:
 			BatchWriter.EndObject();
 			BatchPackage.SetObject(BatchWriter.Save().AsObject());
 
-			auto OnRpcComplete = [this, OpRef = TRefCountPtr<FPutOp>(this), Batch](THttpUniquePtr<IHttpResponse>& HttpResponse, FCbPackage& Response)
+			auto OnRpcComplete = [this, OpRef = TRefCountPtr<FPutOp>(this), Batch](const THttpUniquePtr<IHttpResponse>& HttpResponse, const FCbPackage& Response)
 			{
+				const bool bCanceled = !HttpResponse || HttpResponse->GetErrorCode() == EHttpErrorCode::Canceled;
+
 				FRequestTimer RequestTimer(Requests[0].Stats);
 				// Latency can't be measured for Put operations because it is intertwined with upload time.
 				Requests[0].Stats.Latency = FMonotonicTimeSpan::Infinity();
 
 				int32 RequestIndex = 0;
-				if (HttpResponse->GetErrorCode() == EHttpErrorCode::None && HttpResponse->GetStatusCode() >= 200 && HttpResponse->GetStatusCode() <= 299)
+				if (!bCanceled && HttpResponse->GetErrorCode() == EHttpErrorCode::None && HttpResponse->GetStatusCode() >= 200 && HttpResponse->GetStatusCode() <= 299)
 				{
 					const FCbObject& ResponseObj = Response.GetObject();
 					for (FCbField ResponseField : ResponseObj[ANSITEXTVIEW("Result")])
@@ -399,7 +402,7 @@ public:
 							*CacheStore.GetName(), Batch.Num(), RequestIndex, *WriteToString<256>(*HttpResponse));
 					}
 				}
-				else if ((HttpResponse->GetErrorCode() != EHttpErrorCode::Canceled) && (HttpResponse->GetStatusCode() != 404))
+				else if (!bCanceled && (HttpResponse->GetStatusCode() != 404))
 				{
 					UE_LOG(LogDerivedDataCache, Display,
 						TEXT("%s: Error response received from PutCacheRecords RPC: from %s"),
@@ -408,7 +411,7 @@ public:
 
 				for (const TRequestWithStats<FCachePutRequest>& RequestWithStats : Batch.RightChop(RequestIndex))
 				{
-					if (HttpResponse->GetErrorCode() == EHttpErrorCode::Canceled)
+					if (bCanceled)
 					{
 						OnCanceled(RequestWithStats);
 					}
@@ -567,14 +570,19 @@ public:
 			BatchRequest.EndObject();
 
 			FGetOp* OriginalOp = this;
-			auto OnRpcComplete = [this, OpRef = TRefCountPtr<FGetOp>(OriginalOp), Batch](THttpUniquePtr<IHttpResponse>& HttpResponse, FCbPackage& Response)
+			auto OnRpcComplete = [this, OpRef = TRefCountPtr<FGetOp>(OriginalOp), Batch](const THttpUniquePtr<IHttpResponse>& HttpResponse, const FCbPackage& Response)
 			{
+				const bool bCanceled = !HttpResponse || HttpResponse->GetErrorCode() == EHttpErrorCode::Canceled;
+
 				FRequestTimer RequestTimer(Requests[0].Stats);
-				const FHttpResponseStats& ResponseStats = HttpResponse->GetStats();
-				Requests[0].Stats.Latency = FMonotonicTimeSpan::FromSeconds(ResponseStats.GetLatency());
+				if (HttpResponse)
+				{
+					const FHttpResponseStats& ResponseStats = HttpResponse->GetStats();
+					Requests[0].Stats.Latency = FMonotonicTimeSpan::FromSeconds(ResponseStats.GetLatency());
+				}
 
 				int32 RequestIndex = 0;
-				if (HttpResponse->GetErrorCode() == EHttpErrorCode::None && HttpResponse->GetStatusCode() >= 200 && HttpResponse->GetStatusCode() <= 299)
+				if (!bCanceled && HttpResponse->GetErrorCode() == EHttpErrorCode::None && HttpResponse->GetStatusCode() >= 200 && HttpResponse->GetStatusCode() <= 299)
 				{
 					const FCbObject& ResponseObj = Response.GetObject();
 						
@@ -609,7 +617,7 @@ public:
 							*CacheStore.GetName(), Batch.Num(), RequestIndex, *WriteToString<256>(*HttpResponse));
 					}
 				}
-				else if ((HttpResponse->GetErrorCode() != EHttpErrorCode::Canceled) && (HttpResponse->GetStatusCode() != 404))
+				else if (!bCanceled && (HttpResponse->GetStatusCode() != 404))
 				{
 					UE_LOG(LogDerivedDataCache, Display,
 						TEXT("%s: Error response received from GetCacheRecords RPC: from %s"),
@@ -618,7 +626,7 @@ public:
 					
 				for (const TRequestWithStats<FCacheGetRequest>& RequestWithStats : Batch.RightChop(RequestIndex))
 				{
-					if (HttpResponse->GetErrorCode() == EHttpErrorCode::Canceled)
+					if (bCanceled)
 					{
 						OnCanceled(RequestWithStats);
 					}
@@ -772,14 +780,16 @@ public:
 			BatchWriter.EndObject();
 			BatchPackage.SetObject(BatchWriter.Save().AsObject());
 
-			auto OnRpcComplete = [this, OpRef = TRefCountPtr<FPutValueOp>(this), Batch](THttpUniquePtr<IHttpResponse>& HttpResponse, FCbPackage& Response)
+			auto OnRpcComplete = [this, OpRef = TRefCountPtr<FPutValueOp>(this), Batch](const THttpUniquePtr<IHttpResponse>& HttpResponse, const FCbPackage& Response)
 			{
+				const bool bCanceled = !HttpResponse || HttpResponse->GetErrorCode() == EHttpErrorCode::Canceled;
+
 				FRequestTimer RequestTimer(Requests[0].Stats);
 				// Latency can't be measured for Put operations because it is intertwined with upload time.
 				Requests[0].Stats.Latency = FMonotonicTimeSpan::Infinity();
 
 				int32 RequestIndex = 0;
-				if (HttpResponse->GetErrorCode() == EHttpErrorCode::None && HttpResponse->GetStatusCode() >= 200 && HttpResponse->GetStatusCode() <= 299)
+				if (!bCanceled && HttpResponse->GetErrorCode() == EHttpErrorCode::None && HttpResponse->GetStatusCode() >= 200 && HttpResponse->GetStatusCode() <= 299)
 				{
 					const FCbObject& ResponseObj = Response.GetObject();
 					for (FCbField ResponseField : ResponseObj[ANSITEXTVIEW("Result")])
@@ -808,7 +818,7 @@ public:
 							*CacheStore.GetName(), Batch.Num(), RequestIndex, *WriteToString<256>(*HttpResponse));
 					}
 				}
-				else if ((HttpResponse->GetErrorCode() != EHttpErrorCode::Canceled) && (HttpResponse->GetStatusCode() != 404))
+				else if (!bCanceled && (HttpResponse->GetStatusCode() != 404))
 				{
 					UE_LOG(LogDerivedDataCache, Display,
 						TEXT("%s: Error response received from PutCacheValues RPC: from %s"),
@@ -817,7 +827,7 @@ public:
 
 				for (const TRequestWithStats<FCachePutValueRequest>& RequestWithStats : Batch.RightChop(RequestIndex))
 				{
-					if (HttpResponse->GetErrorCode() == EHttpErrorCode::Canceled)
+					if (bCanceled)
 					{
 						OnCanceled(RequestWithStats);
 					}
@@ -964,14 +974,19 @@ public:
 			BatchRequest.EndObject();
 
 			FGetValueOp* OriginalOp = this;
-			auto OnRpcComplete = [this, OpRef = TRefCountPtr<FGetValueOp>(OriginalOp), Batch](THttpUniquePtr<IHttpResponse>& HttpResponse, FCbPackage& Response)
+			auto OnRpcComplete = [this, OpRef = TRefCountPtr<FGetValueOp>(OriginalOp), Batch](const THttpUniquePtr<IHttpResponse>& HttpResponse, const FCbPackage& Response)
 			{
+				const bool bCanceled = !HttpResponse || HttpResponse->GetErrorCode() == EHttpErrorCode::Canceled;
+
 				FRequestTimer RequestTimer(Requests[0].Stats);
-				const FHttpResponseStats& ResponseStats = HttpResponse->GetStats();
-				Requests[0].Stats.Latency = FMonotonicTimeSpan::FromSeconds(ResponseStats.GetLatency());
+				if (HttpResponse)
+				{
+					const FHttpResponseStats& ResponseStats = HttpResponse->GetStats();
+					Requests[0].Stats.Latency = FMonotonicTimeSpan::FromSeconds(ResponseStats.GetLatency());
+				}
 
 				int32 RequestIndex = 0;
-				if (HttpResponse->GetErrorCode() == EHttpErrorCode::None && HttpResponse->GetStatusCode() >= 200 && HttpResponse->GetStatusCode() <= 299)
+				if (!bCanceled && HttpResponse->GetErrorCode() == EHttpErrorCode::None && HttpResponse->GetStatusCode() >= 200 && HttpResponse->GetStatusCode() <= 299)
 				{
 					const FCbObject& ResponseObj = Response.GetObject();
 
@@ -1020,7 +1035,7 @@ public:
 							*CacheStore.GetName(), Batch.Num(), RequestIndex, *WriteToString<256>(*HttpResponse));
 					}
 				}
-				else if ((HttpResponse->GetErrorCode() != EHttpErrorCode::Canceled) && (HttpResponse->GetStatusCode() != 404))
+				else if (!bCanceled && (HttpResponse->GetStatusCode() != 404))
 				{
 					UE_LOG(LogDerivedDataCache, Display,
 						TEXT("%s: Error response received from GetCacheValues RPC: from %s"),
@@ -1029,7 +1044,7 @@ public:
 
 				for (const TRequestWithStats<FCacheGetValueRequest>& RequestWithStats : Batch.RightChop(RequestIndex))
 				{
-					if (HttpResponse->GetErrorCode() == EHttpErrorCode::Canceled)
+					if (bCanceled)
 					{
 						OnCanceled(RequestWithStats);
 					}
@@ -1184,14 +1199,19 @@ public:
 			BatchRequest.EndObject();
 
 			FGetChunksOp* OriginalOp = this;
-			auto OnRpcComplete = [this, OpRef = TRefCountPtr<FGetChunksOp>(OriginalOp), Batch](THttpUniquePtr<IHttpResponse>& HttpResponse, FCbPackage& Response)
+			auto OnRpcComplete = [this, OpRef = TRefCountPtr<FGetChunksOp>(OriginalOp), Batch](const THttpUniquePtr<IHttpResponse>& HttpResponse, const FCbPackage& Response)
 			{
+				const bool bCanceled = !HttpResponse || HttpResponse->GetErrorCode() == EHttpErrorCode::Canceled;
+
 				FRequestTimer RequestTimer(Requests[0].Stats);
-				const FHttpResponseStats& ResponseStats = HttpResponse->GetStats();
-				Requests[0].Stats.Latency = FMonotonicTimeSpan::FromSeconds(ResponseStats.GetLatency());
+				if (HttpResponse)
+				{
+					const FHttpResponseStats& ResponseStats = HttpResponse->GetStats();
+					Requests[0].Stats.Latency = FMonotonicTimeSpan::FromSeconds(ResponseStats.GetLatency());
+				}
 
 				int32 RequestIndex = 0;
-				if (HttpResponse->GetErrorCode() == EHttpErrorCode::None && HttpResponse->GetStatusCode() >= 200 && HttpResponse->GetStatusCode() <= 299)
+				if (!bCanceled && HttpResponse->GetErrorCode() == EHttpErrorCode::None && HttpResponse->GetStatusCode() >= 200 && HttpResponse->GetStatusCode() <= 299)
 				{
 					const FCbObject& ResponseObj = Response.GetObject();
 
@@ -1265,7 +1285,7 @@ public:
 						Succeeded ? OnHit(RequestWithStats, MoveTemp(RawHash), RawSize, MoveTemp(RequestedBytes)) : OnMiss(RequestWithStats);
 					}
 				}
-				else if ((HttpResponse->GetErrorCode() != EHttpErrorCode::Canceled) && (HttpResponse->GetStatusCode() != 404))
+				else if (!bCanceled && (HttpResponse->GetStatusCode() != 404))
 				{
 					UE_LOG(LogDerivedDataCache, Display,
 						TEXT("%s: Error response received from GetChunks RPC: from %s"),
@@ -1274,7 +1294,7 @@ public:
 
 				for (const TRequestWithStats<FCacheGetChunkRequest>& RequestWithStats : Batch.RightChop(RequestIndex))
 				{
-					if (HttpResponse->GetErrorCode() == EHttpErrorCode::Canceled)
+					if (bCanceled)
 					{
 						OnCanceled(RequestWithStats);
 					}
@@ -1710,7 +1730,7 @@ void FZenCacheStore::Initialize(const FZenCacheStoreParams& Params)
 	ClientParams.LowSpeedLimit = 1;
 	ClientParams.LowSpeedTime = 25;
 	ClientParams.bBypassProxy = Params.bBypassProxy;
-	RequestQueue = UE::FHttpRequestQueue(*ConnectionPool, ClientParams);
+	RequestQueue.Initialize(*ConnectionPool, ClientParams);
 
 	bIsLocalConnection = ZenService.GetInstance().IsServiceRunningLocally() || ZenService.GetInstance().GetServiceSettings().IsAutoLaunch();
 	bIsUsable = true;
@@ -1795,29 +1815,48 @@ FCompositeBuffer FZenCacheStore::SaveRpcPackage(const FCbPackage& Package)
 	return FCompositeBuffer(FSharedBuffer::TakeOwnership(Memory.ReleaseOwnership(), PackageMemorySize, FMemory::Free));
 }
 
-THttpUniquePtr<IHttpRequest> FZenCacheStore::CreateRpcRequest()
+void FZenCacheStore::CreateRpcRequest(IRequestOwner& Owner, FHttpRequestQueue::FOnRequest&& OnRequest)
 {
-	THttpUniquePtr<IHttpRequest> Request = RequestQueue.CreateRequest({});
-	Request->SetUri(RpcUri);
-	Request->SetMethod(EHttpMethod::Post);
-	Request->AddAcceptType(EHttpMediaType::CbPackage);
-	return Request;
+	RequestQueue.CreateRequestAsync(Owner, {}, [this, OnRequest = MoveTemp(OnRequest)](THttpUniquePtr<IHttpRequest>&& Request)
+	{
+		if (Request)
+		{
+			Request->SetUri(RpcUri);
+			Request->SetMethod(EHttpMethod::Post);
+			Request->AddAcceptType(EHttpMediaType::CbPackage);
+		}
+		OnRequest(MoveTemp(Request));
+	});
 }
 
 void FZenCacheStore::EnqueueAsyncRpc(IRequestOwner& Owner, FCbObject RequestObject, FOnRpcComplete&& OnComplete)
 {
-	THttpUniquePtr<IHttpRequest> Request = CreateRpcRequest();
-	Request->SetContentType(EHttpMediaType::CbObject);
-	Request->SetBody(RequestObject.GetBuffer().MakeOwned());
-	new FAsyncCbPackageReceiver(MoveTemp(Request), &Owner, ZenService.GetInstance(), MoveTemp(OnComplete));
+	CreateRpcRequest(Owner, [this, &Owner, RequestObject, OnComplete = MoveTemp(OnComplete)](THttpUniquePtr<IHttpRequest>&& Request) mutable
+	{
+		if (UNLIKELY(!Request))
+		{
+			OnComplete({}, {});
+			return;
+		}
+		Request->SetContentType(EHttpMediaType::CbObject);
+		Request->SetBody(RequestObject.GetBuffer().MakeOwned());
+		new FAsyncCbPackageReceiver(MoveTemp(Request), &Owner, ZenService.GetInstance(), MoveTemp(OnComplete));
+	});
 }
 
 void FZenCacheStore::EnqueueAsyncRpc(IRequestOwner& Owner, const FCbPackage& RequestPackage, FOnRpcComplete&& OnComplete)
 {
-	THttpUniquePtr<IHttpRequest> Request = CreateRpcRequest();
-	Request->SetContentType(EHttpMediaType::CbPackage);
-	Request->SetBody(SaveRpcPackage(RequestPackage));
-	new FAsyncCbPackageReceiver(MoveTemp(Request), &Owner, ZenService.GetInstance(), MoveTemp(OnComplete));
+	CreateRpcRequest(Owner, [this, &Owner, RequestPackage, OnComplete = MoveTemp(OnComplete)](THttpUniquePtr<IHttpRequest>&& Request) mutable
+	{
+		if (UNLIKELY(!Request))
+		{
+			OnComplete({}, {});
+			return;
+		}
+		Request->SetContentType(EHttpMediaType::CbPackage);
+		Request->SetBody(SaveRpcPackage(RequestPackage));
+		new FAsyncCbPackageReceiver(MoveTemp(Request), &Owner, ZenService.GetInstance(), MoveTemp(OnComplete));
+	});
 }
 
 void FZenCacheStore::ActivatePerformanceEvaluationThread()
@@ -1830,13 +1869,18 @@ void FZenCacheStore::ActivatePerformanceEvaluationThread()
 			{
 				IRequestOwner& Owner(PerformanceEvaluationRequestOwner);
 				FRequestBarrier Barrier(Owner);
-				THttpUniquePtr<IHttpRequest> Request = RequestQueue.CreateRequest({});
-				TAnsiStringBuilder<256> StatusUri;
-				StatusUri << ZenService.GetInstance().GetURL() << ANSITEXTVIEW("/health/ready");
-				Request->SetUri(StatusUri);
-				Request->SetMethod(EHttpMethod::Get);
-				Request->AddAcceptType(EHttpMediaType::Text);
-				new FAsyncHealthReceiver(MoveTemp(Request), &Owner, ZenService.GetInstance(), [this, StartTime = FMonotonicTimePoint::Now()](THttpUniquePtr<IHttpResponse>& HttpResponse, EHealth Health)
+				RequestQueue.CreateRequestAsync(Owner, {}, [this, &Owner](THttpUniquePtr<IHttpRequest>&& Request)
+				{
+					if (UNLIKELY(!Request))
+					{
+						return;
+					}
+					TAnsiStringBuilder<256> StatusUri;
+					StatusUri << ZenService.GetInstance().GetURL() << ANSITEXTVIEW("/health/ready");
+					Request->SetUri(StatusUri);
+					Request->SetMethod(EHttpMethod::Get);
+					Request->AddAcceptType(EHttpMediaType::Text);
+					new FAsyncHealthReceiver(MoveTemp(Request), &Owner, ZenService.GetInstance(), [this, StartTime = FMonotonicTimePoint::Now()](const THttpUniquePtr<IHttpResponse>& HttpResponse, EHealth Health)
 					{
 						if (Health != EHealth::Ok)
 						{
@@ -1866,6 +1910,7 @@ void FZenCacheStore::ActivatePerformanceEvaluationThread()
 							}
 						}
 					});
+				});
 			}
 		});
 	}
