DDC: Added detection and reporting of Cancel() and Wait() deadlocks in FRequestOwner

diff --git a/Engine/Source/Developer/DerivedDataCache/Private/DerivedDataCacheStoreHierarchy.cpp b/Engine/Source/Developer/DerivedDataCache/Private/DerivedDataCacheStoreHierarchy.cpp
--- a/Engine/Source/Developer/DerivedDataCache/Private/DerivedDataCacheStoreHierarchy.cpp
+++ b/Engine/Source/Developer/DerivedDataCache/Private/DerivedDataCacheStoreHierarchy.cpp
@@ -182,22 +182,15 @@ public:
 	{
 		if (!Owner)
 		{
-			NewOwner.BeginBarrier(ERequestBarrierFlags::Priority);
 			Owner = &NewOwner;
+			Barrier.Emplace(NewOwner, ERequestBarrierFlags::Priority);
 		}
 		check(Owner == &NewOwner);
 	}
 
-	~FDynamicRequestBarrier()
-	{
-		if (Owner)
-		{
-			Owner->EndBarrier(ERequestBarrierFlags::Priority);
-		}
-	}
-
 private:
 	IRequestOwner* Owner = nullptr;
+	TOptional<FRequestBarrier> Barrier;
 };
 
 class FCacheStoreHierarchy::FBatchBase
diff --git a/Engine/Source/Developer/DerivedDataCache/Private/DerivedDataRequestOwner.cpp b/Engine/Source/Developer/DerivedDataCache/Private/DerivedDataRequestOwner.cpp
--- a/Engine/Source/Developer/DerivedDataCache/Private/DerivedDataRequestOwner.cpp
+++ b/Engine/Source/Developer/DerivedDataCache/Private/DerivedDataRequestOwner.cpp
@@ -26,8 +26,8 @@ public:
 	void Begin(IRequest* Request) final;
 	TRefCountPtr<IRequest> End(IRequest* Request) final;
 
-	void BeginBarrier(ERequestBarrierFlags Flags) final;
-	void EndBarrier(ERequestBarrierFlags Flags) final;
+	void BeginBarrier(const FRequestBarrier& Barrier) final;
+	void EndBarrier(const FRequestBarrier& Barrier) final;
 
 	inline EPriority GetPriority() const final { return Priority.load(std::memory_order_relaxed); }
 	inline bool IsCanceled() const final { return bIsCanceled.load(std::memory_order_relaxed); }
@@ -108,24 +108,24 @@ TRefCountPtr<IRequest> FRequestOwnerShared::End(IRequest* Request)
 	return RequestRef;
 }
 
-void FRequestOwnerShared::BeginBarrier(ERequestBarrierFlags Flags)
+void FRequestOwnerShared::BeginBarrier(const FRequestBarrier& Barrier)
 {
 	AddRef();
 	FWriteScopeLock WriteLock(Lock);
 	++BarrierCount;
-	if (EnumHasAnyFlags(Flags, ERequestBarrierFlags::Priority))
+	if (EnumHasAnyFlags(Barrier.GetFlags(), ERequestBarrierFlags::Priority))
 	{
 		++PriorityBarrierCount;
 	}
 }
 
-void FRequestOwnerShared::EndBarrier(ERequestBarrierFlags Flags)
+void FRequestOwnerShared::EndBarrier(const FRequestBarrier& Barrier)
 {
 	ON_SCOPE_EXIT { Release(); };
 	bool bNotifyBarrier = false;
 	{
 		FWriteScopeLock WriteLock(Lock);
-		if (EnumHasAnyFlags(Flags, ERequestBarrierFlags::Priority) && --PriorityBarrierCount == 0)
+		if (EnumHasAnyFlags(Barrier.GetFlags(), ERequestBarrierFlags::Priority) && --PriorityBarrierCount == 0)
 		{
 			bPriorityChangedInBarrier = false;
 		}
@@ -191,6 +191,10 @@ void FRequestOwnerShared::Cancel()
 		}
 		else
 		{
+			checkf(!FRequestBarrier::HasBarrierForOwnerOnCallingThread(*this),
+				TEXT("Called Cancel() on a request owner while the calling thread holds a FRequestBarrier for it. "
+					 "This will cause a deadlock. To destroy a request owner from within a completion callback, "
+					 "first call KeepAlive() on the request owner to disable the auto-cancel mechanism."));
 			BarrierEvent.Wait(BarrierToken);
 		}
 	}
@@ -226,6 +230,9 @@ void FRequestOwnerShared::Wait()
 		}
 		else
 		{
+			checkf(!FRequestBarrier::HasBarrierForOwnerOnCallingThread(*this),
+				TEXT("Called Wait() on a request owner while the calling thread holds a FRequestBarrier for it. "
+					 "This will cause a deadlock. Narrow the scope of the barrier or move the wait outside of it."));
 			BarrierEvent.Wait(BarrierToken);
 		}
 	}
@@ -307,4 +314,35 @@ void IRequestOwner::LaunchTask(const TCHAR* DebugName, TUniqueFunction<void ()>&
 	TaskEvent.Trigger();
 }
 
+static thread_local FRequestBarrier* GRequestBarriersByThread;
+
+FRequestBarrier::FRequestBarrier(IRequestOwner& InOwner, ERequestBarrierFlags InFlags)
+	: Owner(InOwner)
+	, NextOnThread(GRequestBarriersByThread)
+	, Flags(InFlags)
+{
+	GRequestBarriersByThread = this;
+	Owner.BeginBarrier(*this);
+}
+
+FRequestBarrier::~FRequestBarrier()
+{
+	Owner.EndBarrier(*this);
+	checkf(GRequestBarriersByThread == this,
+		TEXT("Barriers on a thread must be destroyed in the reverse of their creation order."));
+	GRequestBarriersByThread = NextOnThread;
+}
+
+bool FRequestBarrier::HasBarrierForOwnerOnCallingThread(IRequestOwner& QueryOwner)
+{
+	for (FRequestBarrier* It = GRequestBarriersByThread; It; It = It->NextOnThread)
+	{
+		if (&It->Owner == &QueryOwner)
+		{
+			return true;
+		}
+	}
+	return false;
+}
+
 } // UE::DerivedData
diff --git a/Engine/Source/Developer/DerivedDataCache/Public/DerivedDataRequestOwner.h b/Engine/Source/Developer/DerivedDataCache/Public/DerivedDataRequestOwner.h
--- a/Engine/Source/Developer/DerivedDataCache/Public/DerivedDataRequestOwner.h
+++ b/Engine/Source/Developer/DerivedDataCache/Public/DerivedDataRequestOwner.h
@@ -13,6 +13,7 @@
 
 template <typename FuncType> class TUniqueFunction;
 
+namespace UE::DerivedData { class FRequestBarrier; }
 namespace UE::DerivedData { class IRequest; }
 namespace UE::DerivedData { enum class EPriority : uint8; }
 
@@ -75,8 +76,8 @@ public:
 	inline TRefCountPtr<IRequest> End(IRequest* Request, CallbackType&& Callback, CallbackArgTypes&&... CallbackArgs);
 
 	/** See FRequestBarrier. */
-	virtual void BeginBarrier(ERequestBarrierFlags Flags) = 0;
-	virtual void EndBarrier(ERequestBarrierFlags Flags) = 0;
+	UE_INTERNAL virtual void BeginBarrier(const FRequestBarrier& Barrier) = 0;
+	UE_INTERNAL virtual void EndBarrier(const FRequestBarrier& Barrier) = 0;
 
 	/** Returns the priority that new requests are expected to inherit. */
 	virtual EPriority GetPriority() const = 0;
@@ -159,23 +160,21 @@ private:
 class FRequestBarrier
 {
 public:
-	inline explicit FRequestBarrier(IRequestOwner& InOwner, ERequestBarrierFlags InFlags = ERequestBarrierFlags::None)
-		: Owner(InOwner)
-		, Flags(InFlags)
-	{
-		Owner.BeginBarrier(Flags);
-	}
-
-	inline ~FRequestBarrier()
-	{
-		Owner.EndBarrier(Flags);
-	}
+	UE_API explicit FRequestBarrier(IRequestOwner& Owner, ERequestBarrierFlags Flags = ERequestBarrierFlags::None);
+	UE_API ~FRequestBarrier();
 
 	FRequestBarrier(const FRequestBarrier&) = delete;
 	FRequestBarrier& operator=(const FRequestBarrier&) = delete;
 
+	UE_INTERNAL inline ERequestBarrierFlags GetFlags() const { return Flags; }
+
+	/** Returns true if the calling thread has an active barrier for the owner. */
+	UE_INTERNAL static bool HasBarrierForOwnerOnCallingThread(IRequestOwner& QueryOwner);
+
 private:
+	friend IRequestOwner;
 	IRequestOwner& Owner;
+	FRequestBarrier* NextOnThread;
 	ERequestBarrierFlags Flags;
 };
 
