DDC: Added explicit configuration of the file system cache as local or remote

Windows has a fallback to detect a remote cache.

diff --git a/Engine/Source/Developer/DerivedDataCache/Private/FileSystemCacheStore.cpp b/Engine/Source/Developer/DerivedDataCache/Private/FileSystemCacheStore.cpp
--- a/Engine/Source/Developer/DerivedDataCache/Private/FileSystemCacheStore.cpp
+++ b/Engine/Source/Developer/DerivedDataCache/Private/FileSystemCacheStore.cpp
@@ -52,6 +52,10 @@
 
 #include <atomic>
 
+#if PLATFORM_WINDOWS
+#include "Windows/WindowsHWrapper.h"
+#endif
+
 namespace UE::DerivedData
 {
 
@@ -821,6 +825,8 @@ struct FFileSystemCacheStoreParams
 	bool bDeleteOnly = false;
 	/** If true, do not write any files in this cache store. */
 	bool bReadOnly = false;
+	/** If true, this cache store is considered to be remote. */
+	bool bRemote = false;
 	/** If true, block on maintenance of this cache store on startup. */
 	bool bClean = false;
 	/** If true, delete everything in this cache store on startup. */
@@ -956,8 +962,8 @@ private:
 	ICacheStoreOwner& StoreOwner;
 	ICacheStoreStats* StoreStats = nullptr;
 
-	/** Speed class of this cache. */
-	EBackendSpeedClass SpeedClass = EBackendSpeedClass::Unknown;
+	/** If true, this cache store is considered to be remote. */
+	bool		bRemote;
 	/** If true, do not attempt to write to this cache. */
 	bool		bReadOnly;
 	/** If true, always update file timestamps on access. */
@@ -1039,9 +1045,8 @@ ILegacyCacheStore* FFileSystemCacheStore::TryCreate(
 			return nullptr;
 		}
 
-		// Skip creation of the cache store if the shared cache directory does not exist.
-		bool bShared = Params.CacheName == TEXTVIEW("Shared");
-		if (!bShared || IFileManager::Get().DirectoryExists(*Params.CachePath))
+		// Skip creation of the cache store if the remote cache directory does not exist.
+		if (!Params.bRemote || IFileManager::Get().DirectoryExists(*Params.CachePath))
 		{
 			if (ILegacyCacheStore* RedirectedStore = TryRedirection(Params, Owner, Graph))
 			{
@@ -1089,12 +1094,28 @@ FFileSystemCacheStore::FFileSystemCacheStore(
 	const FDerivedDataCacheSpeedStats& InSpeedStats)
 	: CachePath(Params.CachePath)
 	, StoreOwner(Owner)
+	, bRemote(Params.bRemote)
 	, bReadOnly(Params.bReadOnly)
 	, bTouch(Params.bTouch)
 	, MaxFileAgeInDays(Params.MaxFileAgeInDays)
 	, SpeedStats(InSpeedStats)
 	, DeactivateAtMs(Params.DeactivateAtMs)
 {
+#if PLATFORM_WINDOWS
+	if (!bRemote)
+	{
+		// Query for remote file systems because the Remote option is new and may not be set consistently,
+		// and there is a performance penalty when treating a remote cache as local.
+		if (HANDLE CacheHandle = CreateFile(*CachePath, GENERIC_READ, FILE_SHARE_DELETE | FILE_SHARE_READ | FILE_SHARE_WRITE,
+			nullptr, OPEN_EXISTING, FILE_FLAG_BACKUP_SEMANTICS, nullptr); CacheHandle != INVALID_HANDLE_VALUE)
+		{
+			FILE_REMOTE_PROTOCOL_INFO RemoteProtocolInfo;
+			bRemote = !!GetFileInformationByHandleEx(CacheHandle, FileRemoteProtocolInfo, &RemoteProtocolInfo, sizeof(RemoteProtocolInfo));
+			CloseHandle(CacheHandle);
+		}
+	}
+#endif // PLATFORM_WINDOWS
+
 	LastPerformanceEvaluationTicks.store(FDateTime::UtcNow().GetTicks(), std::memory_order_relaxed);
 
 	bool bReadTestPassed = SpeedStats.ReadSpeedMBs > 0.0;
@@ -1107,6 +1128,7 @@ FFileSystemCacheStore::FFileSystemCacheStore(
 	bDeactivatedForPerformance.store(bLocalDeactivatedForPerformance, std::memory_order_relaxed);
 
 	// classify and report on these times
+	EBackendSpeedClass SpeedClass;
 	if (SpeedStats.LatencyMS < 1)
 	{
 		SpeedClass = EBackendSpeedClass::Local;
@@ -1240,7 +1262,7 @@ FFileSystemCacheStore::FFileSystemCacheStore(
 	ECacheStoreFlags Flags = ECacheStoreFlags::None;
 	Flags |= Params.bDeleteOnly ? ECacheStoreFlags::None : ECacheStoreFlags::Query;
 	Flags |= Params.bDeleteOnly || bReadOnly ? ECacheStoreFlags::None : ECacheStoreFlags::Store;
-	Flags |= SpeedClass == EBackendSpeedClass::Local ? ECacheStoreFlags::Local : ECacheStoreFlags::Remote;
+	Flags |= bRemote ? ECacheStoreFlags::Remote : ECacheStoreFlags::Local;
 
 	StoreOwner.Add(this, Flags);
 	if (!Params.bDeleteOnly)
@@ -1826,7 +1848,7 @@ bool FFileSystemCacheStore::PutCacheRecord(
 	const ECachePolicy RecordPolicy = Policy.GetRecordPolicy();
 
 	// Skip the request if storing to the cache is disabled.
-	const ECachePolicy StoreFlag = SpeedClass == EBackendSpeedClass::Local ? ECachePolicy::StoreLocal : ECachePolicy::StoreRemote;
+	const ECachePolicy StoreFlag = bRemote ? ECachePolicy::StoreRemote : ECachePolicy::StoreLocal;
 	if (!EnumHasAnyFlags(RecordPolicy, StoreFlag))
 	{
 		UE_LOG(LogDerivedDataCache, VeryVerbose, TEXT("%s: Skipped put of %s from '%.*s' due to cache policy"),
@@ -1846,7 +1868,7 @@ bool FFileSystemCacheStore::PutCacheRecord(
 
 	// Check if there is an existing record package.
 	FCbPackage ExistingPackage;
-	const ECachePolicy QueryFlag = SpeedClass == EBackendSpeedClass::Local ? ECachePolicy::QueryLocal : ECachePolicy::QueryRemote;
+	const ECachePolicy QueryFlag = bRemote ? ECachePolicy::QueryRemote : ECachePolicy::QueryLocal;
 	bool bReplaceExisting = !EnumHasAnyFlags(RecordPolicy, QueryFlag);
 	bool bSavePackage = bReplaceExisting;
 	if (const bool bLoadPackage = !bReplaceExisting || !Algo::AllOf(Record.GetValues(), &FValue::HasData))
@@ -1985,7 +2007,7 @@ FOptionalCacheRecord FFileSystemCacheStore::GetCacheRecordOnly(
 	FRequestStats& Stats)
 {
 	// Skip the request if querying the cache is disabled.
-	const ECachePolicy QueryFlag = SpeedClass == EBackendSpeedClass::Local ? ECachePolicy::QueryLocal : ECachePolicy::QueryRemote;
+	const ECachePolicy QueryFlag = bRemote ? ECachePolicy::QueryRemote : ECachePolicy::QueryLocal;
 	const bool bLocalDeactivatedForPerformance = IsDeactivatedForPerformance();
 	if (bLocalDeactivatedForPerformance || !EnumHasAnyFlags(Policy.GetRecordPolicy(), QueryFlag))
 	{
@@ -2125,7 +2147,7 @@ bool FFileSystemCacheStore::PutCacheValue(
 	}
 
 	// Skip the request if storing to the cache is disabled.
-	const ECachePolicy StoreFlag = SpeedClass == EBackendSpeedClass::Local ? ECachePolicy::StoreLocal : ECachePolicy::StoreRemote;
+	const ECachePolicy StoreFlag = bRemote ? ECachePolicy::StoreRemote : ECachePolicy::StoreLocal;
 	if (!EnumHasAnyFlags(Policy, StoreFlag))
 	{
 		UE_LOG(LogDerivedDataCache, VeryVerbose, TEXT("%s: Skipped put of %s from '%.*s' due to cache policy"),
@@ -2144,7 +2166,7 @@ bool FFileSystemCacheStore::PutCacheValue(
 	FCbPackage ExistingPackage;
 	TStringBuilder<256> Path;
 	BuildCachePackagePath(Key, Path);
-	const ECachePolicy QueryFlag = SpeedClass == EBackendSpeedClass::Local ? ECachePolicy::QueryLocal : ECachePolicy::QueryRemote;
+	const ECachePolicy QueryFlag = bRemote ? ECachePolicy::QueryRemote : ECachePolicy::QueryLocal;
 	bool bReplaceExisting = !EnumHasAnyFlags(Policy, QueryFlag);
 	bool bSavePackage = bReplaceExisting;
 	if (const bool bLoadPackage = !bReplaceExisting || !Value.HasData())
@@ -2263,7 +2285,7 @@ bool FFileSystemCacheStore::GetCacheValueOnly(
 	FRequestStats& Stats)
 {
 	// Skip the request if querying the cache is disabled.
-	const ECachePolicy QueryFlag = SpeedClass == EBackendSpeedClass::Local ? ECachePolicy::QueryLocal : ECachePolicy::QueryRemote;
+	const ECachePolicy QueryFlag = bRemote ? ECachePolicy::QueryRemote : ECachePolicy::QueryLocal;
 	const bool bLocalDeactivatedForPerformance = IsDeactivatedForPerformance();
 	if (bLocalDeactivatedForPerformance || !EnumHasAnyFlags(Policy, QueryFlag))
 	{
@@ -2965,6 +2987,9 @@ void FFileSystemCacheStoreParams::Parse(const TCHAR* Name, const TCHAR* Config)
 {
 	CacheName = Name;
 
+	// Default remote behavior based on historical cache names.
+	bRemote = CacheName == TEXTVIEW("Shared");
+
 	FString Key;
 
 	// Path Params
@@ -3035,6 +3060,7 @@ void FFileSystemCacheStoreParams::Parse(const TCHAR* Name, const TCHAR* Config)
 	FParse::Bool(Config, TEXT("PromptIfMissing="), bPromptIfMissing);
 	FParse::Bool(Config, TEXT("DeleteOnly="), bDeleteOnly);
 	FParse::Bool(Config, TEXT("ReadOnly="), bReadOnly);
+	FParse::Bool(Config, TEXT("Remote="), bRemote);
 	FParse::Bool(Config, TEXT("Clean="), bClean);
 	FParse::Bool(Config, TEXT("Flush="), bFlush);
 
