Change UE custom http headers in DDC Http layer to use X- prefix.  Print session identifier to the log from the Jupiter cache store. Add commandlet name to User-Agent header, and add IsBuildMachine as a separate Http  header.

diff --git a/Engine/Source/Developer/DerivedDataCache/Private/Http/CurlHttpClient.cpp b/Engine/Source/Developer/DerivedDataCache/Private/Http/CurlHttpClient.cpp
--- a/Engine/Source/Developer/DerivedDataCache/Private/Http/CurlHttpClient.cpp
+++ b/Engine/Source/Developer/DerivedDataCache/Private/Http/CurlHttpClient.cpp
@@ -358,7 +358,7 @@ private:
 
 FCurlHttpManager::FCurlHttpManager()
 {
-	// User-Agent: UnrealEngine/X.Y.Z-<CL> (<Platform>; <Config> <TargetType>; <BranchName>) <AppName> (<ProjectName>)
+	// User-Agent: UnrealEngine/X.Y.Z-<CL> (<Platform>; <Config> <TargetType>; <BranchName>) <AppName>[CommandletName] (<ProjectName>)
 	const FEngineVersion& Version = FEngineVersion::Current();
 	UserAgent << ANSITEXTVIEW("User-Agent: UnrealEngine/")
 		<< Version.GetMajor() << '.' << Version.GetMinor() << '.' << Version.GetPatch() << '-' << Version.GetChangelist()
@@ -366,6 +366,15 @@ FCurlHttpManager::FCurlHttpManager()
 		<< ANSITEXTVIEW("; ") << LexToString(FApp::GetBuildConfiguration()) << ' ' << LexToString(FApp::GetBuildTargetType())
 		<< ANSITEXTVIEW("; ") << FApp::GetBranchName()
 		<< ANSITEXTVIEW(") ") << FApp::GetName();
+
+	if (IsRunningCommandlet())
+	{
+		FString CommandletName;
+		FParse::Value(FCommandLine::Get(), TEXT("Run="), CommandletName);
+		CommandletName.ToLowerInline();
+		UserAgent << '[' << CommandletName << ']';
+	}
+
 	if (FApp::HasProjectName() && FApp::GetName() != FApp::GetProjectName())
 	{
 		UserAgent << ANSITEXTVIEW(" (") << FApp::GetProjectName() << ')';
@@ -383,8 +392,9 @@ void FCurlHttpManager::SetDefaultOptions(CURL* Curl, FCurlHttpHeaders& Headers)
 	curl_easy_setopt(Curl, CURLOPT_NOSIGNAL, 1L);
 	curl_easy_setopt(Curl, CURLOPT_USERAGENT, UserAgent.GetData());
 
-	Headers.AddHeader(*WriteToAnsiString<64>(ANSITEXTVIEW("UE-Session: "), SessionId));
-	Headers.AddHeader(*WriteToAnsiString<32>(ANSITEXTVIEW("UE-Request: "), RequestId.fetch_add(1, std::memory_order_relaxed)));
+	Headers.AddHeader(*WriteToAnsiString<64>(ANSITEXTVIEW("X-UE-IsBuildMachine: "), GIsBuildMachine));
+	Headers.AddHeader(*WriteToAnsiString<64>(ANSITEXTVIEW("X-UE-Session: "), SessionId));
+	Headers.AddHeader(*WriteToAnsiString<32>(ANSITEXTVIEW("X-UE-Request: "), RequestId.fetch_add(1, std::memory_order_relaxed)));
 
 	// Remove the Expect: 100-Continue header that curl adds by default because it adds latency.
 	Headers.AddHeader("Expect:");
diff --git a/Engine/Source/Developer/DerivedDataCache/Private/HttpCacheStore.cpp b/Engine/Source/Developer/DerivedDataCache/Private/HttpCacheStore.cpp
--- a/Engine/Source/Developer/DerivedDataCache/Private/HttpCacheStore.cpp
+++ b/Engine/Source/Developer/DerivedDataCache/Private/HttpCacheStore.cpp
@@ -2133,6 +2133,8 @@ FHttpCacheStore::FHttpCacheStore(const FHttpCacheStoreParams& Params, ICacheStor
 		EffectiveDomain.Append(ResolvedDomain);
 	}
 
+	UE_LOG(LogDerivedDataCache, Display, TEXT("%s: Using session id %s."), *NodeName, *WriteToString<64>(FApp::GetSessionObjectId()));
+
 #if WITH_SSL
 	if (!Params.HostPinnedPublicKeys.IsEmpty() && EffectiveDomain.ToView().StartsWith(ANSITEXTVIEW("https://")))
 	{
