Revert prefix for custom UE custom http headers in DDC Http layer

diff --git a/Engine/Source/Developer/DerivedDataCache/Private/Http/CurlHttpClient.cpp b/Engine/Source/Developer/DerivedDataCache/Private/Http/CurlHttpClient.cpp
--- a/Engine/Source/Developer/DerivedDataCache/Private/Http/CurlHttpClient.cpp
+++ b/Engine/Source/Developer/DerivedDataCache/Private/Http/CurlHttpClient.cpp
@@ -392,9 +392,9 @@ void FCurlHttpManager::SetDefaultOptions(CURL* Curl, FCurlHttpHeaders& Headers)
 	curl_easy_setopt(Curl, CURLOPT_NOSIGNAL, 1L);
 	curl_easy_setopt(Curl, CURLOPT_USERAGENT, UserAgent.GetData());
 
-	Headers.AddHeader(*WriteToAnsiString<64>(ANSITEXTVIEW("X-UE-IsBuildMachine: "), GIsBuildMachine));
-	Headers.AddHeader(*WriteToAnsiString<64>(ANSITEXTVIEW("X-UE-Session: "), SessionId));
-	Headers.AddHeader(*WriteToAnsiString<32>(ANSITEXTVIEW("X-UE-Request: "), RequestId.fetch_add(1, std::memory_order_relaxed)));
+	Headers.AddHeader(*WriteToAnsiString<64>(ANSITEXTVIEW("UE-IsBuildMachine: "), GIsBuildMachine));
+	Headers.AddHeader(*WriteToAnsiString<64>(ANSITEXTVIEW("UE-Session: "), SessionId));
+	Headers.AddHeader(*WriteToAnsiString<32>(ANSITEXTVIEW("UE-Request: "), RequestId.fetch_add(1, std::memory_order_relaxed)));
 
 	// Remove the Expect: 100-Continue header that curl adds by default because it adds latency.
 	Headers.AddHeader("Expect:");
