DDC: Added FCacheKeyBuilder because it has been requested several times

diff --git a/Engine/Source/Developer/DerivedDataCache/Public/DerivedDataCacheKeyBuilder.h b/Engine/Source/Developer/DerivedDataCache/Public/DerivedDataCacheKeyBuilder.h
new file mode 100644
--- /dev/null
+++ b/Engine/Source/Developer/DerivedDataCache/Public/DerivedDataCacheKeyBuilder.h
@@ -0,0 +1,79 @@
+// Copyright Epic Games, Inc. All Rights Reserved.
+
+#pragma once
+
+#include "DerivedDataCacheKey.h"
+#include "Misc/Guid.h"
+#include "Serialization/Archive.h"
+#include "UObject/NameTypes.h"
+#include <type_traits>
+
+namespace UE::DerivedData
+{
+
+/**
+* A type that builds a cache key from arbitrary values serialized to it.
+*
+* FCacheKey GetKey(UAsset* Asset)
+* {
+*     FCacheKeyBuilder Builder;
+*     Builder << FGuid(TEXT("a3ae79ff-6a89-4124-afd6-dc095e000488"));
+*     Builder << ThirdPartyLibraryVersion;
+*     Builder << Asset->GetBulkData().GetPayloadId();
+*     static const FCacheBucket Bucket(ANSITEXTVIEW(""));
+*     return Builder.Build(Bucket);
+* }
+*/
+class FCacheKeyBuilder final : public FArchive
+{
+public:
+	inline FCacheKeyBuilder();
+
+	inline FCacheKey Build(FCacheBucket Bucket) const;
+
+	inline FString GetArchiveName() const final { return TEXT("FCacheKeyBuilder"); }
+
+	inline void Serialize(void* Data, int64 Num) final;
+
+	using FArchive::operator<<;
+	FArchive& operator<<(FName& Name) final;
+
+	template <typename ArgType>
+	inline FCacheKeyBuilder& operator<<(ArgType&& Arg)
+	{
+		FArchive& Ar = *this;
+		Ar << const_cast<std::decay_t<ArgType>&>(Arg);
+		return *this;
+	}
+
+private:
+	FIoHashBuilder HashBuilder;
+};
+
+inline FCacheKeyBuilder::FCacheKeyBuilder()
+{
+	SetIsSaving(true);
+	SetIsPersistent(false);
+
+	// Hash this version to provide a way to invalidate keys that were created using the key builder.
+	static FGuid BaseVersion(TEXT("7ad57ac2-c657-4c11-890c-6d9a2d88dd33"));
+	HashBuilder.Update(&BaseVersion, sizeof(FGuid));
+}
+
+inline FCacheKey FCacheKeyBuilder::Build(FCacheBucket Bucket) const
+{
+	return {Bucket, HashBuilder.Finalize()};
+}
+
+inline void FCacheKeyBuilder::Serialize(void* Data, int64 Num)
+{
+	HashBuilder.Update(Data, uint64(Num));
+}
+
+inline FArchive& FCacheKeyBuilder::operator<<(FName& Name)
+{
+	FString NameString = Name.ToString();
+	return *this << NameString;
+}
+
+} // UE::DerivedData
