DDC: Fixed cache policy merging for chunk requests in HttpCacheStore

diff --git a/Engine/Source/Developer/DerivedDataCache/Private/DerivedDataCachePolicy.cpp b/Engine/Source/Developer/DerivedDataCache/Private/DerivedDataCachePolicy.cpp
--- a/Engine/Source/Developer/DerivedDataCache/Private/DerivedDataCachePolicy.cpp
+++ b/Engine/Source/Developer/DerivedDataCache/Private/DerivedDataCachePolicy.cpp
@@ -122,6 +122,12 @@ static ECachePolicy ParseCachePolicy(const TStringView<CharType> Text)
 namespace UE::DerivedData
 {
 
+ECachePolicy CombineCachePolicy(const ECachePolicy A, const ECachePolicy B)
+{
+	constexpr ECachePolicy NegativeFlags = ECachePolicy::SkipData | ECachePolicy::SkipMeta;
+	return ((A | B) & ~NegativeFlags) | ((A & B) & NegativeFlags);
+}
+
 FAnsiStringBuilderBase& operator<<(FAnsiStringBuilderBase& Builder, ECachePolicy Policy) { return Private::CachePolicyToString(Builder, Policy); }
 FWideStringBuilderBase& operator<<(FWideStringBuilderBase& Builder, ECachePolicy Policy) { return Private::CachePolicyToString(Builder, Policy); }
 FUtf8StringBuilderBase& operator<<(FUtf8StringBuilderBase& Builder, ECachePolicy Policy) { return Private::CachePolicyToString(Builder, Policy); }
@@ -252,12 +258,8 @@ FCacheRecordPolicy FCacheRecordPolicyBuilder::Build()
 	FCacheRecordPolicy Policy(BasePolicy);
 	if (Shared)
 	{
-		const auto Add = [](const ECachePolicy A, const ECachePolicy B)
-		{
-			return ((A | B) & ~ECachePolicy::SkipData) | ((A & B) & ECachePolicy::SkipData);
-		};
 		const TConstArrayView<FCacheValuePolicy> Values = Shared->GetValuePolicies();
-		Policy.RecordPolicy = Algo::TransformAccumulate(Values, &FCacheValuePolicy::Policy, BasePolicy, Add);
+		Policy.RecordPolicy = Algo::TransformAccumulate(Values, &FCacheValuePolicy::Policy, BasePolicy, CombineCachePolicy);
 		Policy.Shared = MoveTemp(Shared);
 	}
 	return Policy;
diff --git a/Engine/Source/Developer/DerivedDataCache/Private/HttpCacheStore.cpp b/Engine/Source/Developer/DerivedDataCache/Private/HttpCacheStore.cpp
--- a/Engine/Source/Developer/DerivedDataCache/Private/HttpCacheStore.cpp
+++ b/Engine/Source/Developer/DerivedDataCache/Private/HttpCacheStore.cpp
@@ -2804,13 +2804,13 @@ void FHttpCacheStore::GetChunkGroupAsync(
 		return;
 	}
 
-	ECachePolicy GroupPolicy(ECachePolicy::None);
+	ECachePolicy GroupPolicy = ECachePolicy::SkipData | ECachePolicy::SkipMeta;
 	TArray<FCacheGetChunkRequest> RequestGroup;
 	RequestGroup.Reserve(static_cast<int>(EndRequest - StartRequest));
 	for (const FCacheGetChunkRequest* Request = StartRequest; Request != EndRequest; ++Request)
 	{
 		RequestGroup.Add(*Request);
-		GroupPolicy |= Request->Policy;
+		GroupPolicy = CombineCachePolicy(GroupPolicy, Request->Policy);
 	}
 
 	if (StartRequest->Id.IsValid())
diff --git a/Engine/Source/Developer/DerivedDataCache/Public/DerivedDataCachePolicy.h b/Engine/Source/Developer/DerivedDataCache/Public/DerivedDataCachePolicy.h
--- a/Engine/Source/Developer/DerivedDataCache/Public/DerivedDataCachePolicy.h
+++ b/Engine/Source/Developer/DerivedDataCache/Public/DerivedDataCachePolicy.h
@@ -101,6 +101,16 @@ enum class ECachePolicy : uint32
 
 ENUM_CLASS_FLAGS(ECachePolicy);
 
+/**
+ * Returns a cache policy that encompasses the operations requested by the both arguments.
+ *
+ * Examples:
+ * - (Local, Remote | SkipData) -> Default
+ * - (Default, Local | SkipMeta) -> Default
+ * - (Query, StoreRemote) -> Query | StoreRemote
+ */
+UE_API ECachePolicy CombineCachePolicy(ECachePolicy A, ECachePolicy B);
+
 /** Append a non-empty text version of the policy to the builder. */
 UE_API FAnsiStringBuilderBase& operator<<(FAnsiStringBuilderBase& Builder, ECachePolicy Policy);
 UE_API FWideStringBuilderBase& operator<<(FWideStringBuilderBase& Builder, ECachePolicy Policy);
