Fixed an unnecessary array copy in FCacheKeyFilter::Parse caused by incompatible allocators.

diff --git a/Engine/Source/Developer/DerivedDataCache/Private/DerivedDataCacheKeyFilter.cpp b/Engine/Source/Developer/DerivedDataCache/Private/DerivedDataCacheKeyFilter.cpp
--- a/Engine/Source/Developer/DerivedDataCache/Private/DerivedDataCacheKeyFilter.cpp
+++ b/Engine/Source/Developer/DerivedDataCache/Private/DerivedDataCacheKeyFilter.cpp
@@ -68,7 +68,7 @@ FCacheKeyFilter FCacheKeyFilter::Parse(const TCHAR* const Config, const TCHAR* c
 {
 	using namespace UE::DerivedData::Private;
 
-	TArray<FCacheKeyFilterState::FTypeRate, TInlineAllocator<16>> Types;
+	TArray<FCacheKeyFilterState::FTypeRate, TInlineAllocator<2>> Types;
 
 	FString TypeConfigArray;
 	for (FStringView ConfigView(Config), PrefixView(Prefix); FParse::Value(ConfigView.GetData(), Prefix, TypeConfigArray);)
