[DDC] Change FDerivedDataCache::GetAsynchronous priority from Blocking to Normal to avoid long hitches in GameThread

#tests Editor, PIE
#rnx

diff --git a/Engine/Source/Developer/DerivedDataCache/Private/DerivedDataCache.cpp b/Engine/Source/Developer/DerivedDataCache/Private/DerivedDataCache.cpp
--- a/Engine/Source/Developer/DerivedDataCache/Private/DerivedDataCache.cpp
+++ b/Engine/Source/Developer/DerivedDataCache/Private/DerivedDataCache.cpp
@@ -712,7 +712,7 @@ public:
 		const uint32 Handle = NextHandle();
 		FString CacheKey = FDerivedDataCache::BuildCacheKey(DataDeriver);
 		UE_LOG(LogDerivedDataCache, VeryVerbose, TEXT("GetAsynchronous %s from '%s', Handle %d"), *CacheKey, *DataDeriver->GetDebugContextString(), Handle);
-		TUniquePtr<FLegacyFetchOrBuildTask> AsyncTask = MakeUnique<FLegacyFetchOrBuildTask>(Backend, DataDeriver->GetDebugContextString(), *CacheKey, DataDeriver, EPriority::Blocking);
+		TUniquePtr<FLegacyFetchOrBuildTask> AsyncTask = MakeUnique<FLegacyFetchOrBuildTask>(Backend, DataDeriver->GetDebugContextString(), *CacheKey, DataDeriver, EPriority::Normal);
 		FLegacyFetchOrBuildTask* LocalAsyncTask = AsyncTask.Get();
 		FScopeLock ScopeLock(&SynchronizationObject);
 		check(!PendingTasks.Contains(Handle));
