DDC: Added -DDC-ForceMemoryCache to allow the editor to boot when the cache is misconfigured

diff --git a/Engine/Source/Developer/DerivedDataCache/Private/DerivedDataBackends.cpp b/Engine/Source/Developer/DerivedDataCache/Private/DerivedDataBackends.cpp
--- a/Engine/Source/Developer/DerivedDataCache/Private/DerivedDataBackends.cpp
+++ b/Engine/Source/Developer/DerivedDataCache/Private/DerivedDataBackends.cpp
@@ -192,10 +192,23 @@ public:
 							TEXT("Unable to create default cache graph '%s' because the root node '%s' is missing."),
 							*GraphName, RootName);
 					}
+					else if (FParse::Param(FCommandLine::Get(), TEXT("DDC-ForceMemoryCache")))
+					{
+						IMemoryCacheStore* Memory = nullptr;
+						CreateMemoryCacheStore(Memory, TEXT("ForceMemoryCache"), TEXT(""), Hierarchy);
+						UE_LOG(LogDerivedDataCache, Error,
+							TEXT("Unable to use default cache graph '%s' because there are no %s nodes available."),
+							*GraphName,
+							Hierarchy->HasAllFlags(ECacheStoreFlags::Query) ? TEXT("writable") :
+							Hierarchy->HasAllFlags(ECacheStoreFlags::Store) ? TEXT("readable") : TEXT("readable or writable"));
+						check(Hierarchy->HasAllFlags(ECacheStoreFlags::Query | ECacheStoreFlags::Store));
+					}
 					else
 					{
 						UE_LOG(LogDerivedDataCache, Fatal,
-							TEXT("Unable to use default cache graph '%s' because there are no %s nodes available."),
+							TEXT("Unable to use default cache graph '%s' because there are no %s nodes available."
+								 "Add -DDC-ForceMemoryCache to the command line to bypass this if you need access "
+								 "to the editor settings to fix the cache configuration."),
 							*GraphName,
 							Hierarchy->HasAllFlags(ECacheStoreFlags::Query) ? TEXT("writable") :
 							Hierarchy->HasAllFlags(ECacheStoreFlags::Store) ? TEXT("readable") : TEXT("readable or writable"));
