DDC: Added FValueId::MakeIndexed to more easily support arrays of values

diff --git a/Engine/Source/Developer/DerivedDataCache/Private/DerivedDataValue.cpp b/Engine/Source/Developer/DerivedDataCache/Private/DerivedDataValue.cpp
--- a/Engine/Source/Developer/DerivedDataCache/Private/DerivedDataValue.cpp
+++ b/Engine/Source/Developer/DerivedDataCache/Private/DerivedDataValue.cpp
@@ -5,6 +5,7 @@
 #include "Compression/OodleDataCompression.h"
 #include "Containers/StringConv.h"
 #include "Hash/xxhash.h"
+#include "Misc/ByteSwap.h"
 #include "Serialization/CompactBinary.h"
 #include "Serialization/CompactBinarySerialization.h"
 #include "Serialization/CompactBinaryWriter.h"
@@ -62,6 +63,17 @@ FValueId FValueId::FromName(const FWideStringView Name)
 	return FValueId::FromName(FTCHARToUTF8(Name));
 }
 
+FValueId FValueId::MakeIndexed(int32 Index) const
+{
+	alignas(int32) ByteArray IndexedBytes;
+	// Copy the prefix from this ID.
+	FMemoryView Prefix = MakeMemoryView(Bytes).Left(sizeof(ByteArray) - sizeof(int32));
+	FMutableMemoryView IndexView = MakeMemoryView(IndexedBytes).CopyFrom(Prefix);
+	// Copy the index as big endian to ensure an indexed ID is sorted by index.
+	*(int32*)IndexView.GetData() = NETWORK_ORDER32(Index);
+	return FValueId(IndexedBytes);
+}
+
 FCbWriter& operator<<(FCbWriter& Writer, const FValueId& Id)
 {
 	Writer.AddObjectId(Id);
diff --git a/Engine/Source/Developer/DerivedDataCache/Public/DerivedDataValueId.h b/Engine/Source/Developer/DerivedDataCache/Public/DerivedDataValueId.h
--- a/Engine/Source/Developer/DerivedDataCache/Public/DerivedDataValueId.h
+++ b/Engine/Source/Developer/DerivedDataCache/Public/DerivedDataValueId.h
@@ -45,6 +45,9 @@ public:
 	[[nodiscard]] UE_API static FValueId FromName(FUtf8StringView Name);
 	[[nodiscard]] UE_API static FValueId FromName(FWideStringView Name);
 
+	/** Returns a copy of this ID with the last 4 bytes replaced by the index. Use for arrays of values. */
+	[[nodiscard]] UE_API FValueId MakeIndexed(int32 Index) const;
+
 	/** Returns a reference to the raw byte array for the ID. */
 	inline const ByteArray& GetBytes() const { return Bytes; }
 
