DDC - Support specifying multiple hosts and benchmark to pick which is the optimal to use

diff --git a/Engine/Source/Developer/DerivedDataCache/Private/Http/HttpHostBuilder.cpp b/Engine/Source/Developer/DerivedDataCache/Private/Http/HttpHostBuilder.cpp
new file mode 100644
--- /dev/null
+++ b/Engine/Source/Developer/DerivedDataCache/Private/Http/HttpHostBuilder.cpp
@@ -0,0 +1,292 @@
+// Copyright Epic Games, Inc. All Rights Reserved.
+
+#include "HttpHostBuilder.h"
+
+#include "Async/ManualResetEvent.h"
+#include "Async/Mutex.h"
+#include "Async/UniqueLock.h"
+#include "Containers/StringFwd.h"
+#include "CoreTypes.h"
+#include "HAL/PlatformTime.h"
+#include "Http/HttpClient.h"
+#include "Logging/LogMacros.h"
+#include "Memory/MemoryFwd.h"
+#include "Serialization/CompactBinary.h"
+#include "String/Find.h"
+#include "String/ParseTokens.h"
+#include "Templates/Function.h"
+#include "Templates/UniquePtr.h"
+
+namespace UE
+{
+
+DEFINE_LOG_CATEGORY_STATIC(LogHttpHostBuilder, Display, All);
+
+class FHttpBenchmarkReceiver final : public IHttpReceiver
+{
+public:
+	FHttpBenchmarkReceiver(const FHttpBenchmarkReceiver&) = delete;
+	FHttpBenchmarkReceiver& operator=(const FHttpBenchmarkReceiver&) = delete;
+
+	using FOnComplete = TUniqueFunction<void(IHttpResponse& HttpResponse, FString& Host)>;
+
+	explicit FHttpBenchmarkReceiver(const FStringView InHostEntry, FOnComplete&& InOperationComplete, IHttpReceiver* InNext = nullptr)
+		: HostEntry(InHostEntry)
+		, OperationComplete(MoveTemp(InOperationComplete))
+		, Next(InNext)
+	{
+	}
+
+private:
+	IHttpReceiver* OnCreate(IHttpResponse& Response) final
+	{
+		return this;
+	}
+
+	IHttpReceiver* OnComplete(IHttpResponse& Response) final
+	{
+		OperationComplete(Response, HostEntry);
+
+		return Next;
+	}
+
+private:
+	FString HostEntry;
+	FOnComplete OperationComplete;
+	IHttpReceiver* Next;
+};
+
+class FHttpBenchmarkOperation final : public IHttpReceiver
+{
+public:
+	FHttpBenchmarkOperation(const FHttpBenchmarkOperation&) = delete;
+	FHttpBenchmarkOperation& operator=(const FHttpBenchmarkOperation&) = delete;
+
+	using FOnComplete = TUniqueFunction<void(IHttpResponse& HttpResponse, FString& Host)>;
+
+	explicit FHttpBenchmarkOperation(const FStringView InHostEntry, THttpUniquePtr<IHttpRequest>&& InRequest, FOnComplete&& InOperationComplete)
+		: HostEntry(InHostEntry)
+		, OperationComplete(MoveTemp(InOperationComplete))
+		, Request(MoveTemp(InRequest))
+		, Response(nullptr)
+	{
+		Request->SetUri(WriteToAnsiString<256>(HostEntry, ANSITEXTVIEW("/health/ready")));
+	}
+
+	void SendAsync()
+	{
+		Request->SendAsync(this, Response);
+	}
+
+	void Cancel()
+	{
+		Response->Cancel();
+	}
+	
+private:
+	IHttpReceiver* OnCreate(IHttpResponse& InResponse) final
+	{
+		return this;
+	}
+
+	IHttpReceiver* OnComplete(IHttpResponse& InResponse) final
+	{
+		OperationComplete(InResponse, HostEntry);
+
+		return nullptr;
+	}
+
+private:
+	FString HostEntry;
+	FOnComplete OperationComplete;
+
+	THttpUniquePtr<IHttpRequest> Request;
+	THttpUniquePtr<IHttpResponse> Response;
+};
+
+void FHttpHostBuilder::AddFromEndpoint(const FAnsiStringView HostUrl, const FAnsiStringView AccessToken)
+{
+	FHttpConnectionPoolParams ConnectionPoolParams;
+	const THttpUniquePtr<IHttpConnectionPool> LocalConnectionPool = IHttpManager::Get().CreateConnectionPool(ConnectionPoolParams);
+	FHttpClientParams DefaultClientParams;
+	// we want to keep this timeout fairly low incase the host is not reachable 
+	DefaultClientParams.ConnectTimeout = 5000;
+	THttpUniquePtr<IHttpClient> LocalClient = LocalConnectionPool->CreateClient(DefaultClientParams);
+
+	FHttpRequestParams RequestParams;
+	RequestParams.bIgnoreMaxRequests = true;
+
+	TArray64<uint8> BodyArray;
+	FHttpByteArrayReceiver PeersReceiver(BodyArray);
+	
+	THttpUniquePtr<IHttpRequest> Request = LocalClient->TryCreateRequest(RequestParams);
+
+	if (!AccessToken.IsEmpty())
+	{
+		Request->AddHeader(ANSITEXTVIEW("Authorization"), WriteToAnsiString<1024>(AccessToken));
+	}
+	Request->AddHeader(ANSITEXTVIEW("Accept"), "application/x-ue-cb");
+	TStringBuilderWithBuffer Uri(WriteToAnsiString<256>(HostUrl, ANSITEXTVIEW("/status/peers")));
+	Request->SetUri(Uri);
+	THttpUniquePtr<IHttpResponse> Response;
+	Request->Send(&PeersReceiver, Response);
+
+	if (Response->GetStatusCode() != 200)
+	{
+		UE_LOG(LogHttpHostBuilder, Warning, TEXT("Unsuccesful attempt to fetch hosts by endpoint from host: '%hs'. Status code was: %d"), Uri.ToString(), Response->GetStatusCode());
+		return;
+	}
+
+	FSharedBuffer SharedBuffer = MakeSharedBufferFromArray(MoveTemp(BodyArray));
+	FCbObject PeersObject = FCbObject(SharedBuffer);
+
+	for (FCbField Peer : PeersObject["peers"].AsArray())
+	{
+		FCbObject PeerObject = Peer.AsObject();
+		for (FCbField Endpoint : PeerObject["endpoints"])
+		{
+			FString endpointString = FString(Endpoint.AsString());
+			if (HostCandidates.Find(endpointString) == INDEX_NONE)
+			{
+				HostCandidates.Emplace(endpointString);
+			}
+		}
+	}
+	
+	return;
+}
+
+FString FHttpHostBuilder::GetHostCandidatesString()
+{
+	return FString::Join(HostCandidates, TEXT(", "));
+}
+
+void FHttpHostBuilder::AddFromString(const FAnsiStringView HostList)
+{
+	FString(HostList).ParseIntoArray(HostCandidates, TEXT(";"));
+}
+
+void FHttpHostBuilder::AddFromString(const FString HostList)
+{
+	HostList.ParseIntoArray(HostCandidates, TEXT(";"));
+}
+
+bool FHttpHostBuilder::ResolveHost(double WarningTimeoutSeconds, double TimeoutSeconds, FAnsiStringBuilderBase& OutHost, double& OutLatency)
+{
+	if (HostCandidates.Num() == 0)
+	{
+		// no hosts have been added
+		return false;
+	}
+
+	if (HostCandidates.Num() == 1)
+	{
+		// if there is only one candidate we do not benchmark it and just return this straight away
+		OutHost.Reset();
+		OutHost.Append(HostCandidates.Top());
+
+		OutLatency = 0.0;
+		return true;
+	}
+
+	bool HostFound = BenchmarkHostList(HostCandidates, WarningTimeoutSeconds, TimeoutSeconds, OutHost, OutLatency);
+
+	// if no valid host was found, set the host to the first candidate, even if this will likely not work to connect to its still better then no options
+	if (!HostFound)
+	{
+		OutHost.Reset();
+		OutHost.Append(HostCandidates.Top());
+		OutLatency = 0.0;
+	}
+	return HostFound;
+}
+
+bool FHttpHostBuilder::BenchmarkHostList(const TArray<FString> InHostCandidates, double WarningTimeoutSeconds, double TimeoutSeconds, FAnsiStringBuilderBase& OutHost, double& OutLatency)
+{
+	const double StartTime = FPlatformTime::Seconds();
+
+	constexpr uint32 MaxTotalConnections = 8;
+	FHttpConnectionPoolParams ConnectionPoolParams;
+	ConnectionPoolParams.MaxConnections = MaxTotalConnections;
+	ConnectionPoolParams.MinConnections = MaxTotalConnections;
+	const THttpUniquePtr<IHttpConnectionPool> LocalConnectionPool = IHttpManager::Get().CreateConnectionPool(ConnectionPoolParams);
+	const FHttpClientParams DefaultClientParams;
+	THttpUniquePtr<IHttpClient> LocalClient = LocalConnectionPool->CreateClient(DefaultClientParams);
+
+	struct FLatencySortableHost
+	{
+		double Latency;
+		FString Host;
+	};
+	TArray<FLatencySortableHost> SortedHostList;
+	TArray<TUniquePtr<FHttpBenchmarkOperation>> Operations;
+
+	FMutex Mutex;
+	FManualResetEvent HostFound;
+	for (auto &&HostCandidate : InHostCandidates)
+	{
+		// Benchmark each of the hosts
+		FHttpRequestParams RequestParams;
+
+		TUniquePtr<FHttpBenchmarkOperation> Operation = MakeUnique<FHttpBenchmarkOperation>(HostCandidate, LocalClient->TryCreateRequest(RequestParams), [&SortedHostList, &Mutex, &HostFound](const IHttpResponse& Response, FString& Host)
+		{
+			if (Response.GetStatusCode() != 200)
+			{
+				return;
+			}
+			// grab stats and store in dictionary, needs to get the response for this
+			const FHttpResponseStats& Stats = Response.GetStats();
+
+			FLatencySortableHost HostLatencyRecord;
+			HostLatencyRecord.Latency = Stats.GetLatency();
+			HostLatencyRecord.Host = Host;
+
+			{
+				TUniqueLock Lock(Mutex);
+				SortedHostList.Add(HostLatencyRecord);
+			}
+
+			HostFound.Notify();
+		});
+
+		Operation->SendAsync();
+		Operations.Emplace(MoveTemp(Operation));
+	}
+	bool ValidHostFound = HostFound.WaitFor(FMonotonicTimeSpan::FromSeconds(WarningTimeoutSeconds));
+	if (!ValidHostFound)
+	{
+		// warn that benchmarking is taking a lot of time
+		UE_LOG(LogHttpHostBuilder, Warning, TEXT("HTTP Benchmarking is slow, continuing to wait to determine ideal host..."));
+
+		ValidHostFound = HostFound.WaitFor(FMonotonicTimeSpan::FromSeconds(TimeoutSeconds));
+		if (!ValidHostFound)
+		{
+			UE_LOG(LogHttpHostBuilder, Warning, TEXT("No valid host found while benchmarking after timeout was reached"));
+			return false;
+		}
+	}
+	
+	{
+		TUniqueLock Lock(Mutex);
+
+		SortedHostList.Sort([](const FLatencySortableHost& A, const FLatencySortableHost& B) {return A.Latency > B.Latency;});
+
+		OutHost.Reset();
+		OutHost.Append(SortedHostList[0].Host);
+		OutLatency = SortedHostList[0].Latency;
+	}
+
+	// cancel any outstanding benchmark request
+	for (auto &&Operation : Operations)
+	{
+		Operation->Cancel();
+	}
+
+	const double EndTime = FPlatformTime::Seconds();
+	const double BenchmarkingDuration = EndTime - StartTime;
+
+	UE_LOG(LogHttpHostBuilder, Display, TEXT("Resolved to using host '%hs' based on HTTP benchmark with a estimated latency of '%.0fms'. Spent %.0fms doing benchmarking."), *OutHost, OutLatency * 1000, BenchmarkingDuration * 1000);
+	return true;
+}
+
+} // UE
\ No newline at end of file
diff --git a/Engine/Source/Developer/DerivedDataCache/Private/Http/HttpHostBuilder.h b/Engine/Source/Developer/DerivedDataCache/Private/Http/HttpHostBuilder.h
new file mode 100644
--- /dev/null
+++ b/Engine/Source/Developer/DerivedDataCache/Private/Http/HttpHostBuilder.h
@@ -0,0 +1,32 @@
+// Copyright Epic Games, Inc. All Rights Reserved.
+
+#pragma once
+
+#include "CoreTypes.h"
+#include "Containers/ContainersFwd.h"
+#include "Containers/StringFwd.h"
+#include "Containers/StringView.h"
+
+namespace UE
+{
+
+class FHttpHostBuilder
+{
+public:
+    // accept a semicolon seperated list of hosts to pick from
+	void AddFromString(const FAnsiStringView HostList);
+	void AddFromString(FString HostList);
+
+	// a url that is expected to host a /status/peers endpoint with endpoints we should attempt to use
+	void AddFromEndpoint(const FAnsiStringView HostUrl, const FAnsiStringView AccessToken);
+
+	FString GetHostCandidatesString();
+
+	bool ResolveHost(double WarningTimeoutSeconds, double TimeoutSeconds, FAnsiStringBuilderBase& OutHost, double& OutLatency);
+private:
+	bool BenchmarkHostList(const TArray<FString> HostCandidates, double WarningTimeoutSeconds, double TimeoutSeconds, FAnsiStringBuilderBase& OutHost, double& OutLatency);
+
+	TArray<FString> HostCandidates;
+};
+
+} // UE
diff --git a/Engine/Source/Developer/DerivedDataCache/Private/HttpCacheStore.cpp b/Engine/Source/Developer/DerivedDataCache/Private/HttpCacheStore.cpp
--- a/Engine/Source/Developer/DerivedDataCache/Private/HttpCacheStore.cpp
+++ b/Engine/Source/Developer/DerivedDataCache/Private/HttpCacheStore.cpp
@@ -27,6 +27,7 @@
 #include "HAL/IConsoleManager.h"
 #include "HAL/PlatformProcess.h"
 #include "Http/HttpClient.h"
+#include "Http/HttpHostBuilder.h"
 #include "IO/IoHash.h"
 #include "Memory/SharedBuffer.h"
 #include "Misc/App.h"
@@ -220,6 +221,7 @@ struct FHttpCacheStoreParams
 {
 	FString Name;
 	FString Host;
+	FString DiscoveryHost;
 	FString HostPinnedPublicKeys;
 	FString Namespace;
 	FString HttpVersion;
@@ -2123,6 +2125,33 @@ FHttpCacheStore::FHttpCacheStore(const FHttpCacheStoreParams& Params, ICacheStor
 	while (Domain.RemoveFromEnd(TEXT("/")));
 
 	EffectiveDomain.Append(Domain);
+
+	TAnsiStringBuilder<256> ResolvedHost;
+	double ResolvedLatency;
+	FHttpHostBuilder HostBuilder;
+	HostBuilder.AddFromString(EffectiveDomain);
+	if (!Params.DiscoveryHost.IsEmpty())
+	{
+		TAnsiStringBuilder<256> DiscoveryHost;
+		DiscoveryHost.Append(Params.DiscoveryHost);
+		AcquireAccessToken();
+		HostBuilder.AddFromEndpoint(DiscoveryHost, WriteToAnsiString<1024>(*Access));
+	}
+	if (HostBuilder.ResolveHost(/* Warning timeout */ 1.0, 4.0 /* Max duration timeout*/, ResolvedHost, ResolvedLatency))
+	{
+		EffectiveDomain.Reset();
+		EffectiveDomain.Append(ResolvedHost);
+	}
+	else
+	{
+		// even if we fail to resolve a host to use the returned host will at least contain the first of the possible hosts which we can attempt to use
+		EffectiveDomain.Reset();
+		EffectiveDomain.Append(ResolvedHost);
+
+		FString HostCandidates = HostBuilder.GetHostCandidatesString();
+		UE_LOG(LogDerivedDataCache, Warning, TEXT("%s: Unable to resolve best host candidate to use, most likely none of the suggested hosts was reachable. Attempted hosts were: '%s' ."), *NodeName, *HostCandidates);
+	}
+
 	TAnsiStringBuilder<256> ResolvedDomain;
 	if (Params.bResolveHostCanonicalName && TryResolveCanonicalHost(EffectiveDomain, ResolvedDomain))
 	{
@@ -3156,6 +3185,8 @@ void FHttpCacheStoreParams::Parse(const TCHAR* NodeName, const TCHAR* Config)
 		}
 	}
 
+	FParse::Value(Config, TEXT("DiscoveryHost="), DiscoveryHost);
+
 	FParse::Value(Config, TEXT("HostPinnedPublicKeys="), HostPinnedPublicKeys);
 
 	FParse::Bool(Config, TEXT("ResolveHostCanonicalName="), bResolveHostCanonicalName);
@@ -3312,7 +3343,7 @@ ILegacyCacheStore* CreateHttpCacheStore(const TCHAR* NodeName, const TCHAR* Conf
 		UE_LOG(LogDerivedDataCache, Warning, TEXT("%s: Missing required parameter 'StructuredNamespace', falling back to '%s'"), NodeName, *Params.Namespace);
 	}
 
-	if (bValidParams && !Params.Host.StartsWith(TEXT("http://localhost")))
+	if (bValidParams && (!Params.Host.StartsWith(TEXT("http://localhost")) || !Params.Host.StartsWith(TEXT("https://localhost"))))
 	{
 		bool bValidOAuthAccessToken = !Params.OAuthAccessToken.IsEmpty();
 
diff --git a/Engine/Source/Developer/DerivedDataCache/Private/ZenCacheStore.cpp b/Engine/Source/Developer/DerivedDataCache/Private/ZenCacheStore.cpp
--- a/Engine/Source/Developer/DerivedDataCache/Private/ZenCacheStore.cpp
+++ b/Engine/Source/Developer/DerivedDataCache/Private/ZenCacheStore.cpp
@@ -21,6 +21,7 @@
 #include "HAL/FileManager.h"
 #include "HAL/Thread.h"
 #include "Http/HttpClient.h"
+#include "Http/HttpHostBuilder.h"
 #include "Math/UnrealMathUtility.h"
 #include "Misc/App.h"
 #include "Misc/ConfigCacheIni.h"
@@ -2169,6 +2170,29 @@ ILegacyCacheStore* CreateZenCacheStore(const TCHAR* NodeName, const TCHAR* Confi
 		UE_LOG(LogDerivedDataCache, Warning, TEXT("%s: Missing required parameter 'Namespace', falling back to '%s'"), NodeName, *Params.Namespace);
 	}
 
+	if (!Params.Host.IsEmpty())
+	{
+		// if a host is specified we check to see if that is potentially a list of multiple hosts and then benchmark them to find the best to use
+		// if the host is empty it just means to connect to localhost, no need to do any benchmarking then
+		
+		TAnsiStringBuilder<256> ResolvedHost;
+		double ResolvedLatency;
+		FHttpHostBuilder HostBuilder;
+		HostBuilder.AddFromString(Params.Host);
+
+		if (HostBuilder.ResolveHost(/* Warning timeout */ 1.0, 4.0 /* Max duration timeout*/, ResolvedHost, ResolvedLatency))
+		{
+			Params.Host = FString(ResolvedHost);
+		}
+		else
+		{
+			// even if we fail to resolve a host to use the returned host will at least contain the first of the possible hosts which we can attempt to use
+			Params.Host = FString(ResolvedHost);
+			FString HostCandidates = HostBuilder.GetHostCandidatesString();
+			UE_LOG(LogDerivedDataCache, Warning, TEXT("%s: Unable to resolve best host candidate to use, most likely none of the suggested hosts was reachable. Attempted hosts were: '%s' ."), NodeName, *HostCandidates);
+		}		
+	}
+
 	bool bHasSandbox = !Params.Sandbox.IsEmpty();
 	bool bUseLocalDataCachePathOverrides = !bHasSandbox;
 
