DDC - Fixed issue with http host builder were it would wait for the full timeout even if no valid host existed (say because internet connection was down so all dns lookups failed).

Also fixed a bunch of other post commit feedback.

diff --git a/Engine/Source/Developer/DerivedDataCache/Private/Http/HttpHostBuilder.cpp b/Engine/Source/Developer/DerivedDataCache/Private/Http/HttpHostBuilder.cpp
--- a/Engine/Source/Developer/DerivedDataCache/Private/Http/HttpHostBuilder.cpp
+++ b/Engine/Source/Developer/DerivedDataCache/Private/Http/HttpHostBuilder.cpp
@@ -2,6 +2,7 @@
 
 #include "HttpHostBuilder.h"
 
+#include "Algo/MinElement.h"
 #include "Async/ManualResetEvent.h"
 #include "Async/Mutex.h"
 #include "Async/UniqueLock.h"
@@ -114,6 +115,7 @@ void FHttpHostBuilder::AddFromEndpoint(const FAnsiStringView HostUrl, const FAns
 	THttpUniquePtr<IHttpClient> LocalClient = LocalConnectionPool->CreateClient(DefaultClientParams);
 
 	FHttpRequestParams RequestParams;
+	// setting ignore max requests guarantees that CreateRequest actually creates a request no matter how many requests are in flight
 	RequestParams.bIgnoreMaxRequests = true;
 
 	TArray64<uint8> BodyArray;
@@ -126,7 +128,7 @@ void FHttpHostBuilder::AddFromEndpoint(const FAnsiStringView HostUrl, const FAns
 		Request->AddHeader(ANSITEXTVIEW("Authorization"), WriteToAnsiString<1024>(AccessToken));
 	}
 	Request->AddHeader(ANSITEXTVIEW("Accept"), "application/x-ue-cb");
-	TStringBuilderWithBuffer Uri(WriteToAnsiString<256>(HostUrl, ANSITEXTVIEW("/status/peers")));
+	TAnsiStringBuilder<256> Uri(InPlace, HostUrl, ANSITEXTVIEW("/status/peers"));
 	Request->SetUri(Uri);
 	THttpUniquePtr<IHttpResponse> Response;
 	Request->Send(&PeersReceiver, Response);
@@ -145,30 +147,30 @@ void FHttpHostBuilder::AddFromEndpoint(const FAnsiStringView HostUrl, const FAns
 		FCbObject PeerObject = Peer.AsObject();
 		for (FCbField Endpoint : PeerObject["endpoints"])
 		{
-			FString endpointString = FString(Endpoint.AsString());
-			if (HostCandidates.Find(endpointString) == INDEX_NONE)
-			{
-				HostCandidates.Emplace(endpointString);
-			}
+			HostCandidates.AddUnique(FString(Endpoint.AsString()));
 		}
 	}
-	
-	return;
 }
 
-FString FHttpHostBuilder::GetHostCandidatesString()
+FString FHttpHostBuilder::GetHostCandidatesString() const
 {
 	return FString::Join(HostCandidates, TEXT(", "));
 }
 
 void FHttpHostBuilder::AddFromString(const FAnsiStringView HostList)
 {
-	FString(HostList).ParseIntoArray(HostCandidates, TEXT(";"));
+	String::ParseTokens(HostList, ';', [this](FAnsiStringView Host)
+	{
+		this->HostCandidates.Emplace(Host);
+	}, String::EParseTokensOptions::SkipEmpty | String::EParseTokensOptions::Trim);
 }
 
-void FHttpHostBuilder::AddFromString(const FString HostList)
+void FHttpHostBuilder::AddFromString(const FStringView HostList)
 {
-	HostList.ParseIntoArray(HostCandidates, TEXT(";"));
+	String::ParseTokens(HostList, ';', [this](FStringView Host)
+	{
+		this->HostCandidates.Emplace(Host);
+	}, String::EParseTokensOptions::SkipEmpty | String::EParseTokensOptions::Trim);
 }
 
 bool FHttpHostBuilder::ResolveHost(double WarningTimeoutSeconds, double TimeoutSeconds, FAnsiStringBuilderBase& OutHost, double& OutLatency)
@@ -183,25 +185,25 @@ bool FHttpHostBuilder::ResolveHost(double WarningTimeoutSeconds, double TimeoutS
 	{
 		// if there is only one candidate we do not benchmark it and just return this straight away
 		OutHost.Reset();
-		OutHost.Append(HostCandidates.Top());
+		OutHost.Append(HostCandidates[0]);
 
 		OutLatency = 0.0;
 		return true;
 	}
 
-	bool HostFound = BenchmarkHostList(HostCandidates, WarningTimeoutSeconds, TimeoutSeconds, OutHost, OutLatency);
+	bool bHostFound = BenchmarkHostList(HostCandidates, WarningTimeoutSeconds, TimeoutSeconds, OutHost, OutLatency);
 
 	// if no valid host was found, set the host to the first candidate, even if this will likely not work to connect to its still better then no options
-	if (!HostFound)
+	if (!bHostFound)
 	{
 		OutHost.Reset();
-		OutHost.Append(HostCandidates.Top());
+		OutHost.Append(HostCandidates[0]);
 		OutLatency = 0.0;
 	}
-	return HostFound;
+	return bHostFound;
 }
 
-bool FHttpHostBuilder::BenchmarkHostList(const TArray<FString> InHostCandidates, double WarningTimeoutSeconds, double TimeoutSeconds, FAnsiStringBuilderBase& OutHost, double& OutLatency)
+bool FHttpHostBuilder::BenchmarkHostList(TConstArrayView<FString> InHostCandidates, double WarningTimeoutSeconds, double TimeoutSeconds, FAnsiStringBuilderBase& OutHost, double& OutLatency)
 {
 	const double StartTime = FPlatformTime::Seconds();
 
@@ -222,16 +224,29 @@ bool FHttpHostBuilder::BenchmarkHostList(const TArray<FString> InHostCandidates,
 	TArray<TUniquePtr<FHttpBenchmarkOperation>> Operations;
 
 	FMutex Mutex;
-	FManualResetEvent HostFound;
-	for (auto &&HostCandidate : InHostCandidates)
+	FManualResetEvent HostBenchmarkDone;
+
+	int32 FailedBenchmarkAttempts = 0;
+	bool bAllHostsFailed = false;
+	int32 CountOfBenchmarksToRun = InHostCandidates.Num();
+	for (const FString& HostCandidate : InHostCandidates)
 	{
 		// Benchmark each of the hosts
 		FHttpRequestParams RequestParams;
+		// setting ignore max requests guarantees that CreateRequest actually creates a request no matter how many requests are in flight
+		RequestParams.bIgnoreMaxRequests = true;
 
-		TUniquePtr<FHttpBenchmarkOperation> Operation = MakeUnique<FHttpBenchmarkOperation>(HostCandidate, LocalClient->TryCreateRequest(RequestParams), [&SortedHostList, &Mutex, &HostFound](const IHttpResponse& Response, FString& Host)
+		TUniquePtr<FHttpBenchmarkOperation> Operation = MakeUnique<FHttpBenchmarkOperation>(HostCandidate, LocalClient->TryCreateRequest(RequestParams), [&SortedHostList, &Mutex, &HostBenchmarkDone, &FailedBenchmarkAttempts, &bAllHostsFailed, &CountOfBenchmarksToRun](const IHttpResponse& Response, FString& Host)
 		{
 			if (Response.GetStatusCode() != 200)
 			{
+				FailedBenchmarkAttempts += 1;
+				if (FailedBenchmarkAttempts >= CountOfBenchmarksToRun)
+				{
+					bAllHostsFailed = true;
+					HostBenchmarkDone.Notify();
+				}
+
 				return;
 			}
 			// grab stats and store in dictionary, needs to get the response for this
@@ -246,20 +261,28 @@ bool FHttpHostBuilder::BenchmarkHostList(const TArray<FString> InHostCandidates,
 				SortedHostList.Add(HostLatencyRecord);
 			}
 
-			HostFound.Notify();
+			HostBenchmarkDone.Notify();
 		});
 
 		Operation->SendAsync();
 		Operations.Emplace(MoveTemp(Operation));
 	}
-	bool ValidHostFound = HostFound.WaitFor(FMonotonicTimeSpan::FromSeconds(WarningTimeoutSeconds));
-	if (!ValidHostFound)
+	bool bValidHostFound = HostBenchmarkDone.WaitFor(FMonotonicTimeSpan::FromSeconds(WarningTimeoutSeconds));
+
+	// all benchmarks failed
+	if (bAllHostsFailed)
+	{
+		UE_LOG(LogHttpHostBuilder, Warning, TEXT("No valid host found as all benchmark attempts had errors"));
+		return false;
+	}
+
+	if (!bValidHostFound)
 	{
 		// warn that benchmarking is taking a lot of time
 		UE_LOG(LogHttpHostBuilder, Warning, TEXT("HTTP Benchmarking is slow, continuing to wait to determine ideal host..."));
 
-		ValidHostFound = HostFound.WaitFor(FMonotonicTimeSpan::FromSeconds(TimeoutSeconds));
-		if (!ValidHostFound)
+		bValidHostFound = HostBenchmarkDone.WaitFor(FMonotonicTimeSpan::FromSeconds(TimeoutSeconds));
+		if (!bValidHostFound)
 		{
 			UE_LOG(LogHttpHostBuilder, Warning, TEXT("No valid host found while benchmarking after timeout was reached"));
 			return false;
@@ -269,15 +292,15 @@ bool FHttpHostBuilder::BenchmarkHostList(const TArray<FString> InHostCandidates,
 	{
 		TUniqueLock Lock(Mutex);
 
-		SortedHostList.Sort([](const FLatencySortableHost& A, const FLatencySortableHost& B) {return A.Latency > B.Latency;});
+		FLatencySortableHost* MinElement = Algo::MinElementBy(SortedHostList, &FLatencySortableHost::Latency);
 
 		OutHost.Reset();
-		OutHost.Append(SortedHostList[0].Host);
-		OutLatency = SortedHostList[0].Latency;
+		OutHost.Append(MinElement->Host);
+		OutLatency = MinElement->Latency;
 	}
 
 	// cancel any outstanding benchmark request
-	for (auto &&Operation : Operations)
+	for (const TUniquePtr<FHttpBenchmarkOperation>& Operation : Operations)
 	{
 		Operation->Cancel();
 	}
diff --git a/Engine/Source/Developer/DerivedDataCache/Private/Http/HttpHostBuilder.h b/Engine/Source/Developer/DerivedDataCache/Private/Http/HttpHostBuilder.h
--- a/Engine/Source/Developer/DerivedDataCache/Private/Http/HttpHostBuilder.h
+++ b/Engine/Source/Developer/DerivedDataCache/Private/Http/HttpHostBuilder.h
@@ -14,17 +14,18 @@ class FHttpHostBuilder
 {
 public:
     // accept a semicolon seperated list of hosts to pick from
-	void AddFromString(const FAnsiStringView HostList);
-	void AddFromString(FString HostList);
+	void AddFromString(FAnsiStringView HostList);
+	void AddFromString(FStringView HostList);
 
 	// a url that is expected to host a /status/peers endpoint with endpoints we should attempt to use
-	void AddFromEndpoint(const FAnsiStringView HostUrl, const FAnsiStringView AccessToken);
+	void AddFromEndpoint(FAnsiStringView HostUrl, FAnsiStringView AccessToken);
 
-	FString GetHostCandidatesString();
+	// produces a string of all the host candidates used for display purposes
+	FString GetHostCandidatesString() const;
 
 	bool ResolveHost(double WarningTimeoutSeconds, double TimeoutSeconds, FAnsiStringBuilderBase& OutHost, double& OutLatency);
 private:
-	bool BenchmarkHostList(const TArray<FString> HostCandidates, double WarningTimeoutSeconds, double TimeoutSeconds, FAnsiStringBuilderBase& OutHost, double& OutLatency);
+	static bool BenchmarkHostList(TConstArrayView<FString> HostCandidates, double WarningTimeoutSeconds, double TimeoutSeconds, FAnsiStringBuilderBase& OutHost, double& OutLatency);
 
 	TArray<FString> HostCandidates;
 };
