DDC - Http benchmarking will now always set the resolved host to a one of the host candidates even if the benchmarking failed as this is what is expected by the calling code.
Also fixed an issue were pending operations were not cancelled when benchmarks timed out.
Hopefully resolves crashes on editor startup when the benchmark fails.

diff --git a/Engine/Source/Developer/DerivedDataCache/Private/Http/HttpHostBuilder.cpp b/Engine/Source/Developer/DerivedDataCache/Private/Http/HttpHostBuilder.cpp
--- a/Engine/Source/Developer/DerivedDataCache/Private/Http/HttpHostBuilder.cpp
+++ b/Engine/Source/Developer/DerivedDataCache/Private/Http/HttpHostBuilder.cpp
@@ -12,6 +12,7 @@
 #include "Http/HttpClient.h"
 #include "Logging/LogMacros.h"
 #include "Memory/MemoryFwd.h"
+#include "Misc/ScopeExit.h"
 #include "Serialization/CompactBinary.h"
 #include "String/Find.h"
 #include "String/ParseTokens.h"
@@ -273,36 +274,50 @@ bool FHttpHostBuilder::BenchmarkHostList(TConstArrayView<FString> InHostCandidat
 	if (bAllHostsFailed)
 	{
 		UE_LOG(LogHttpHostBuilder, Warning, TEXT("No valid host found as all benchmark attempts had errors"));
+
+		OutHost.Reset();
+		OutHost.Append(InHostCandidates[0]);
+		OutLatency = 0.0;
+
 		return false;
 	}
 
-	if (!bValidHostFound)
 	{
-		// warn that benchmarking is taking a lot of time
-		UE_LOG(LogHttpHostBuilder, Warning, TEXT("HTTP Benchmarking is slow, continuing to wait to determine ideal host..."));
+		ON_SCOPE_EXIT {
+			// cancel any outstanding benchmark request
+			for (const TUniquePtr<FHttpBenchmarkOperation>& Operation : Operations)
+			{
+				Operation->Cancel();
+			}
+		};
 
-		bValidHostFound = HostBenchmarkDone.WaitFor(FMonotonicTimeSpan::FromSeconds(TimeoutSeconds));
 		if (!bValidHostFound)
 		{
-			UE_LOG(LogHttpHostBuilder, Warning, TEXT("No valid host found while benchmarking after timeout was reached"));
-			return false;
-		}
-	}
-	
-	{
-		TUniqueLock Lock(Mutex);
+			// warn that benchmarking is taking a lot of time
+			UE_LOG(LogHttpHostBuilder, Warning, TEXT("HTTP Benchmarking is slow, continuing to wait to determine ideal host..."));
 
-		FLatencySortableHost* MinElement = Algo::MinElementBy(SortedHostList, &FLatencySortableHost::Latency);
+			bValidHostFound = HostBenchmarkDone.WaitFor(FMonotonicTimeSpan::FromSeconds(TimeoutSeconds));
+			if (!bValidHostFound)
+			{
+				UE_LOG(LogHttpHostBuilder, Warning, TEXT("No valid host found while benchmarking after timeout was reached"));
 
-		OutHost.Reset();
-		OutHost.Append(MinElement->Host);
-		OutLatency = MinElement->Latency;
-	}
+				OutHost.Reset();
+				OutHost.Append(InHostCandidates[0]);
+				OutLatency = 0.0;
+				
+				return false;
+			}
+		}
+		
+		{
+			TUniqueLock Lock(Mutex);
 
-	// cancel any outstanding benchmark request
-	for (const TUniquePtr<FHttpBenchmarkOperation>& Operation : Operations)
-	{
-		Operation->Cancel();
+			FLatencySortableHost* MinElement = Algo::MinElementBy(SortedHostList, &FLatencySortableHost::Latency);
+
+			OutHost.Reset();
+			OutHost.Append(MinElement->Host);
+			OutLatency = MinElement->Latency;
+		}
 	}
 
 	const double EndTime = FPlatformTime::Seconds();
