DDC - Added some null checking and better logging to try and track down occasional crash we see when benchmarking is running.

diff --git a/Engine/Source/Developer/DerivedDataCache/Private/Http/HttpHostBuilder.cpp b/Engine/Source/Developer/DerivedDataCache/Private/Http/HttpHostBuilder.cpp
--- a/Engine/Source/Developer/DerivedDataCache/Private/Http/HttpHostBuilder.cpp
+++ b/Engine/Source/Developer/DerivedDataCache/Private/Http/HttpHostBuilder.cpp
@@ -314,9 +314,20 @@ bool FHttpHostBuilder::BenchmarkHostList(TConstArrayView<FString> InHostCandidat
 
 			FLatencySortableHost* MinElement = Algo::MinElementBy(SortedHostList, &FLatencySortableHost::Latency);
 
-			OutHost.Reset();
-			OutHost.Append(MinElement->Host);
-			OutLatency = MinElement->Latency;
+			if (MinElement != nullptr)
+			{
+				OutHost.Reset();
+				OutHost.Append(MinElement->Host);
+				OutLatency = MinElement->Latency;
+			}
+			else
+			{
+				OutHost.Reset();
+				OutHost.Append(InHostCandidates[0]);
+				OutLatency = 0.0;
+
+				UE_LOG(LogHttpHostBuilder, Warning, TEXT("Failed to determine fastest host option because we had no valid options"));
+			}
 		}
 	}
 
