Cloud DDC - Fixed display of domain name when connected to cloud ddc in the editor window for DDC statistics, will now print the domain we resolved to instead of all the domain that we attempted to use.

Also added some comments in the http host builder to clarify some error handling.

diff --git a/Engine/Source/Developer/DerivedDataCache/Private/Http/HttpHostBuilder.cpp b/Engine/Source/Developer/DerivedDataCache/Private/Http/HttpHostBuilder.cpp
--- a/Engine/Source/Developer/DerivedDataCache/Private/Http/HttpHostBuilder.cpp
+++ b/Engine/Source/Developer/DerivedDataCache/Private/Http/HttpHostBuilder.cpp
@@ -322,6 +322,7 @@ bool FHttpHostBuilder::BenchmarkHostList(TConstArrayView<FString> InHostCandidat
 			}
 			else
 			{
+				// no hosts were found
 				OutHost.Reset();
 				OutHost.Append(InHostCandidates[0]);
 				OutLatency = 0.0;
diff --git a/Engine/Source/Developer/DerivedDataCache/Private/HttpCacheStore.cpp b/Engine/Source/Developer/DerivedDataCache/Private/HttpCacheStore.cpp
--- a/Engine/Source/Developer/DerivedDataCache/Private/HttpCacheStore.cpp
+++ b/Engine/Source/Developer/DerivedDataCache/Private/HttpCacheStore.cpp
@@ -2210,11 +2210,12 @@ FHttpCacheStore::FHttpCacheStore(const FHttpCacheStoreParams& Params, ICacheStor
 		{
 			const ECacheStoreFlags Flags = ECacheStoreFlags::Remote | ECacheStoreFlags::Query |
 				(Params.bReadOnly ? ECacheStoreFlags::None : ECacheStoreFlags::Store);
-			TStringBuilder<256> Path(InPlace, Domain, TEXTVIEW(" ("), Namespace, TEXTVIEW(")"));
+			TStringBuilder<256> Path(InPlace, ResolvedHost, TEXTVIEW(" ("), Namespace, TEXTVIEW(")"));
 			StoreOwner->Add(this, Flags);
 			StoreStats = StoreOwner->CreateStats(this, Flags, TEXTVIEW("Unreal Cloud DDC"), Params.Name, Path);
 
 			StoreStats->SetAttribute(TEXTVIEW("Domain"), Domain);
+			StoreStats->SetAttribute(TEXTVIEW("ResolvedDomain"), WriteToString<128>(ResolvedHost));
 			StoreStats->SetAttribute(TEXTVIEW("EffectiveDomain"), WriteToString<128>(EffectiveDomain));
 			StoreStats->SetAttribute(TEXTVIEW("Namespace"), Namespace);
 			StoreStats->SetAttribute(TEXTVIEW("LoginAttempts"), WriteToString<16>(LoginAttempts));
