Disabling a test case that was consistently failing because we have no time guarantee about when upstream propagation will occur.

diff --git a/Engine/Source/Developer/DerivedDataCache/Private/Tests/HttpCacheStoreTest.cpp b/Engine/Source/Developer/DerivedDataCache/Private/Tests/HttpCacheStoreTest.cpp
--- a/Engine/Source/Developer/DerivedDataCache/Private/Tests/HttpCacheStoreTest.cpp
+++ b/Engine/Source/Developer/DerivedDataCache/Private/Tests/HttpCacheStoreTest.cpp
@@ -762,11 +762,12 @@ bool FHttpCacheStoreTest::RunTest(const FString& Parameters)
 			ValidateValues(TEXT("SimpleValueZenAndDirect"), GetAndValidateValues(TEXT("SimpleValueZen"), PutValuesZen, ECachePolicy::Default), ReceivedValues, ECachePolicy::Default);
 			ValidateValues(TEXT("SimpleValueSkipDataZenAndDirect"), GetAndValidateValues(TEXT("SimpleValueSkipDataZen"), PutValuesZen, ECachePolicy::Default | ECachePolicy::SkipData), ReceivedValuesSkipData, ECachePolicy::Default | ECachePolicy::SkipData);
 		}
-		if (ZenIntermediarySiblingBackend)
-		{
-			GetAndValidateValues(TEXT("SimpleValueZen"), PutValues, ECachePolicy::Default, ZenIntermediarySiblingBackend.Get());
-			GetAndValidateValues(TEXT("SimpleValueSkipDataZen"), PutValues, ECachePolicy::Default | ECachePolicy::SkipData, ZenIntermediarySiblingBackend.Get());
-		}
+		// Disabled because we have no guarantee that the upstream server received the put in any particular timemframe.
+		//if (ZenIntermediarySiblingBackend)
+		//{
+		//	GetAndValidateValues(TEXT("SimpleValueZen"), PutValues, ECachePolicy::Default, ZenIntermediarySiblingBackend.Get());
+		//	GetAndValidateValues(TEXT("SimpleValueSkipDataZen"), PutValues, ECachePolicy::Default | ECachePolicy::SkipData, ZenIntermediarySiblingBackend.Get());
+		//}
 #endif // UE_HTTPCACHESTORETEST_USE_ZEN
 	}
 
