Improve logging and error messages when Zen DDC and Zen Store requests fails

diff --git a/Engine/Source/Developer/DerivedDataCache/Private/Http/HttpClient.cpp b/Engine/Source/Developer/DerivedDataCache/Private/Http/HttpClient.cpp
--- a/Engine/Source/Developer/DerivedDataCache/Private/Http/HttpClient.cpp
+++ b/Engine/Source/Developer/DerivedDataCache/Private/Http/HttpClient.cpp
@@ -144,6 +144,81 @@ bool TryLexFromString(EHttpMediaType& OutMediaType, const FAnsiStringView View)
 	return true;
 }
 
+static FAnsiStringView LexStatusCodeToString(int HttpCode)
+{
+	switch (HttpCode)
+	{
+		// 1xx Informational
+	case 100: return "Continue";
+	case 101: return "Switching Protocols";
+
+		// 2xx Success
+	case 200: return "OK";
+	case 201: return "Created";
+	case 202: return "Accepted";
+	case 204: return "No Content";
+	case 205: return "Reset Content";
+	case 206: return "Partial Content";
+
+		// 3xx Redirection
+	case 300: return "Multiple Choices";
+	case 301: return "Moved Permanently";
+	case 302: return "Found";
+	case 303: return "See Other";
+	case 304: return "Not Modified";
+	case 305: return "Use Proxy";
+	case 306: return "Switch Proxy";
+	case 307: return "Temporary Redirect";
+	case 308: return "Permanent Redirect";
+
+		// 4xx Client errors
+	case 400: return "Bad Request";
+	case 401: return "Unauthorized";
+	case 402: return "Payment Required";
+	case 403: return "Forbidden";
+	case 404: return "Not Found";
+	case 405: return "Method Not Allowed";
+	case 406: return "Not Acceptable";
+	case 407: return "Proxy Authentication Required";
+	case 408: return "Request Timeout";
+	case 409: return "Conflict";
+	case 410: return "Gone";
+	case 411: return "Length Required";
+	case 412: return "Precondition Failed";
+	case 413: return "Payload Too Large";
+	case 414: return "URI Too Long";
+	case 415: return "Unsupported Media Type";
+	case 416: return "Range Not Satisfiable";
+	case 417: return "Expectation Failed";
+	case 418: return "I'm a teapot";
+	case 421: return "Misdirected Request";
+	case 422: return "Unprocessable Entity";
+	case 423: return "Locked";
+	case 424: return "Failed Dependency";
+	case 425: return "Too Early";
+	case 426: return "Upgrade Required";
+	case 428: return "Precondition Required";
+	case 429: return "Too Many Requests";
+	case 431: return "Request Header Fields Too Large";
+
+		// 5xx Server errors
+	case 500: return "Internal Server Error";
+	case 501: return "Not Implemented";
+	case 502: return "Bad Gateway";
+	case 503: return "Service Unavailable";
+	case 504: return "Gateway Timeout";
+	case 505: return "HTTP Version Not Supported";
+	case 506: return "Variant Also Negotiates";
+	case 507: return "Insufficient Storage";
+	case 508: return "Loop Detected";
+	case 510: return "Not Extended";
+	case 511: return "Network Authentication Required";
+
+	default: return "Unknown Result";
+	}
+}
+
+
 ///////////////////////////////////////////////////////////////////////////////////////////////////////////////////////
 
 void IHttpRequest::SetContentType(const EHttpMediaType Type, const FAnsiStringView Param)
@@ -220,7 +295,7 @@ FStringBuilderBase& operator<<(FStringBuilderBase& Builder, const IHttpResponse&
 
 	if (const int32 StatusCode = Response.GetStatusCode(); StatusCode > 0)
 	{
-		Builder << TEXTVIEW(" -> ") << StatusCode;
+		Builder << TEXTVIEW(" -> ") << LexStatusCodeToString(StatusCode) << " (" << StatusCode << ")";
 	}
 
 	if (const FAnsiStringView Error = Response.GetError(); !Error.IsEmpty())
diff --git a/Engine/Source/Developer/DerivedDataCache/Private/ZenCacheStore.cpp b/Engine/Source/Developer/DerivedDataCache/Private/ZenCacheStore.cpp
--- a/Engine/Source/Developer/DerivedDataCache/Private/ZenCacheStore.cpp
+++ b/Engine/Source/Developer/DerivedDataCache/Private/ZenCacheStore.cpp
@@ -167,7 +167,7 @@ private:
 
 	static FCompositeBuffer SaveRpcPackage(const FCbPackage& Package);
 	void CreateRpcRequest(IRequestOwner& Owner, FHttpRequestQueue::FOnRequest&& OnRequest);
-	using FOnRpcComplete = TUniqueFunction<void(const THttpUniquePtr<IHttpResponse>& HttpResponse, const FCbPackage& Response)>;
+	using FOnRpcComplete = TUniqueFunction<void(const THttpUniquePtr<IHttpResponse>& HttpResponse, const FCbPackage& Response, FStringView ResponseText)>;
 	void EnqueueAsyncRpc(IRequestOwner& Owner, FCbObject RequestObject, FOnRpcComplete&& OnComplete);
 	void EnqueueAsyncRpc(IRequestOwner& Owner, const FCbPackage& RequestPackage, FOnRpcComplete&& OnComplete);
 
@@ -364,7 +364,7 @@ public:
 			BatchWriter.EndObject();
 			BatchPackage.SetObject(BatchWriter.Save().AsObject());
 
-			auto OnRpcComplete = [this, OpRef = TRefCountPtr<FPutOp>(this), Batch](const THttpUniquePtr<IHttpResponse>& HttpResponse, const FCbPackage& Response)
+			auto OnRpcComplete = [this, OpRef = TRefCountPtr<FPutOp>(this), Batch](const THttpUniquePtr<IHttpResponse>& HttpResponse, const FCbPackage& Response, FStringView ResponseText)
 			{
 				const bool bCanceled = !HttpResponse || HttpResponse->GetErrorCode() == EHttpErrorCode::Canceled;
 
@@ -406,8 +406,8 @@ public:
 				else if (!bCanceled && (HttpResponse->GetStatusCode() != 404))
 				{
 					UE_LOG(LogDerivedDataCache, Display,
-						TEXT("%s: Error response received from PutCacheRecords RPC: from %s"),
-						*CacheStore.GetName(), *WriteToString<256>(*HttpResponse));
+						TEXT("%s: Error response received from PutCacheRecords RPC: from %s. \"%.*s\""),
+						*CacheStore.GetName(), *WriteToString<256>(*HttpResponse), ResponseText.Len(), ResponseText.GetData());
 				}
 
 				for (const TRequestWithStats<FCachePutRequest>& RequestWithStats : Batch.RightChop(RequestIndex))
@@ -571,7 +571,7 @@ public:
 			BatchRequest.EndObject();
 
 			FGetOp* OriginalOp = this;
-			auto OnRpcComplete = [this, OpRef = TRefCountPtr<FGetOp>(OriginalOp), Batch](const THttpUniquePtr<IHttpResponse>& HttpResponse, const FCbPackage& Response)
+			auto OnRpcComplete = [this, OpRef = TRefCountPtr<FGetOp>(OriginalOp), Batch](const THttpUniquePtr<IHttpResponse>& HttpResponse, const FCbPackage& Response, FStringView ResponseText)
 			{
 				const bool bCanceled = !HttpResponse || HttpResponse->GetErrorCode() == EHttpErrorCode::Canceled;
 
@@ -621,8 +621,8 @@ public:
 				else if (!bCanceled && (HttpResponse->GetStatusCode() != 404))
 				{
 					UE_LOG(LogDerivedDataCache, Display,
-						TEXT("%s: Error response received from GetCacheRecords RPC: from %s"),
-						*CacheStore.GetName(), *WriteToString<256>(*HttpResponse));
+						TEXT("%s: Error response received from GetCacheRecords RPC: from %s.\"%.*s\""),
+						*CacheStore.GetName(), *WriteToString<256>(*HttpResponse), ResponseText.Len(), ResponseText.GetData());
 				}
 					
 				for (const TRequestWithStats<FCacheGetRequest>& RequestWithStats : Batch.RightChop(RequestIndex))
@@ -781,7 +781,7 @@ public:
 			BatchWriter.EndObject();
 			BatchPackage.SetObject(BatchWriter.Save().AsObject());
 
-			auto OnRpcComplete = [this, OpRef = TRefCountPtr<FPutValueOp>(this), Batch](const THttpUniquePtr<IHttpResponse>& HttpResponse, const FCbPackage& Response)
+			auto OnRpcComplete = [this, OpRef = TRefCountPtr<FPutValueOp>(this), Batch](const THttpUniquePtr<IHttpResponse>& HttpResponse, const FCbPackage& Response, FStringView ResponseText)
 			{
 				const bool bCanceled = !HttpResponse || HttpResponse->GetErrorCode() == EHttpErrorCode::Canceled;
 
@@ -822,8 +822,8 @@ public:
 				else if (!bCanceled && (HttpResponse->GetStatusCode() != 404))
 				{
 					UE_LOG(LogDerivedDataCache, Display,
-						TEXT("%s: Error response received from PutCacheValues RPC: from %s"),
-						*CacheStore.GetName(), *WriteToString<256>(*HttpResponse));
+						TEXT("%s: Error response received from PutCacheValues RPC: from %s. \"%.*s\""),
+						*CacheStore.GetName(), *WriteToString<256>(*HttpResponse), ResponseText.Len(), ResponseText.GetData());
 				}
 
 				for (const TRequestWithStats<FCachePutValueRequest>& RequestWithStats : Batch.RightChop(RequestIndex))
@@ -975,7 +975,7 @@ public:
 			BatchRequest.EndObject();
 
 			FGetValueOp* OriginalOp = this;
-			auto OnRpcComplete = [this, OpRef = TRefCountPtr<FGetValueOp>(OriginalOp), Batch](const THttpUniquePtr<IHttpResponse>& HttpResponse, const FCbPackage& Response)
+			auto OnRpcComplete = [this, OpRef = TRefCountPtr<FGetValueOp>(OriginalOp), Batch](const THttpUniquePtr<IHttpResponse>& HttpResponse, const FCbPackage& Response, FStringView ResponseText)
 			{
 				const bool bCanceled = !HttpResponse || HttpResponse->GetErrorCode() == EHttpErrorCode::Canceled;
 
@@ -1039,8 +1039,8 @@ public:
 				else if (!bCanceled && (HttpResponse->GetStatusCode() != 404))
 				{
 					UE_LOG(LogDerivedDataCache, Display,
-						TEXT("%s: Error response received from GetCacheValues RPC: from %s"),
-						*CacheStore.GetName(), *WriteToString<256>(*HttpResponse));
+						TEXT("%s: Error response received from GetCacheValues RPC: from %s. \"%.*s\""),
+						*CacheStore.GetName(), *WriteToString<256>(*HttpResponse), ResponseText.Len(), ResponseText.GetData());
 				}
 
 				for (const TRequestWithStats<FCacheGetValueRequest>& RequestWithStats : Batch.RightChop(RequestIndex))
@@ -1200,7 +1200,7 @@ public:
 			BatchRequest.EndObject();
 
 			FGetChunksOp* OriginalOp = this;
-			auto OnRpcComplete = [this, OpRef = TRefCountPtr<FGetChunksOp>(OriginalOp), Batch](const THttpUniquePtr<IHttpResponse>& HttpResponse, const FCbPackage& Response)
+			auto OnRpcComplete = [this, OpRef = TRefCountPtr<FGetChunksOp>(OriginalOp), Batch](const THttpUniquePtr<IHttpResponse>& HttpResponse, const FCbPackage& Response, FStringView ResponseText)
 			{
 				const bool bCanceled = !HttpResponse || HttpResponse->GetErrorCode() == EHttpErrorCode::Canceled;
 
@@ -1289,8 +1289,8 @@ public:
 				else if (!bCanceled && (HttpResponse->GetStatusCode() != 404))
 				{
 					UE_LOG(LogDerivedDataCache, Display,
-						TEXT("%s: Error response received from GetChunks RPC: from %s"),
-						*CacheStore.GetName(), *WriteToString<256>(*HttpResponse));
+						TEXT("%s: Error response received from GetChunks RPC: from %s. \"%.*s\""),
+						*CacheStore.GetName(), *WriteToString<256>(*HttpResponse), ResponseText.Len(), ResponseText.GetData());
 				}
 
 				for (const TRequestWithStats<FCacheGetChunkRequest>& RequestWithStats : Batch.RightChop(RequestIndex))
@@ -1499,6 +1499,11 @@ public:
 	{
 	}
 
+	const FMemoryView Body() const
+	{
+		return MakeMemoryView(BodyArray);
+	}
+
 private:
 	IHttpReceiver* OnCreate(IHttpResponse& Response) final
 	{
@@ -1509,16 +1514,28 @@ private:
 	{
 		if (Response.GetErrorCode() == EHttpErrorCode::None)
 		{
-			FMemoryView MemoryView = MakeMemoryView(BodyArray);
+			EHttpMediaType ContentType = Response.GetContentType();
+			switch (ContentType)
 			{
-				FMemoryReaderView Ar(MemoryView);
-				if (Zen::Http::TryLoadCbPackage(Package, Ar))
+				case EHttpMediaType::CbPackage:
 				{
-					return Next;
+					const FMemoryView MemoryView = Body();
+					{
+						FMemoryReaderView Ar(MemoryView);
+						if (Zen::Http::TryLoadCbPackage(Package, Ar))
+						{
+							BodyArray.Reset();
+							return Next;
+						}
+					}
+					FMemoryReaderView Ar(MemoryView);
+					if (Package.TryLoad(Ar))
+					{
+						BodyArray.Reset();
+					}
 				}
+				break;
 			}
-			FMemoryReaderView Ar(MemoryView);
-			Package.TryLoad(Ar);
 		}
 		return Next;
 	}
@@ -1601,7 +1618,8 @@ private:
 				{
 					// Ensuring that the OnRpcComplete method is destroyed by the time we exit this method by moving it to a local scope variable
 					FOnRpcComplete LocalOnComplete = MoveTemp(Self->OnRpcComplete);
-					LocalOnComplete(Self->Response, Self->Package);
+
+					LocalOnComplete(Self->Response, Self->Package, Self->GetPayloadAsString());
 				});
 			}
 		});
@@ -1609,6 +1627,31 @@ private:
 	}
 
 private:
+	FString GetPayloadAsString() const
+	{
+		EHttpMediaType ContentType = Response->GetContentType();
+		FMemoryView BodyView = BaseReceiver.Body();
+		switch (ContentType)
+		{
+			case EHttpMediaType::Text:
+			case EHttpMediaType::Yaml:
+			case EHttpMediaType::Json:
+				if (BodyView.GetSize() > 0)
+				{
+					TStringBuilder<128> PayloadText;
+					PayloadText << FAnsiStringView((const ANSICHAR*)BodyView.GetData(), (int)(BodyView.GetSize()));
+					return *PayloadText;
+				}
+				return {};
+			default:
+			{
+				TStringBuilder<128> SizeText;
+				SizeText << "Payload (" << LexToString(ContentType) << "): " << BodyView.GetSize() << " bytes";
+				return *SizeText;
+			}
+		}
+	}
+
 	THttpUniquePtr<IHttpRequest> Request;
 	THttpUniquePtr<IHttpResponse> Response;
 	TRefCountPtr<IHttpResponseMonitor> Monitor;
@@ -1836,7 +1879,7 @@ void FZenCacheStore::EnqueueAsyncRpc(IRequestOwner& Owner, FCbObject RequestObje
 	{
 		if (UNLIKELY(!Request))
 		{
-			OnComplete({}, {});
+			OnComplete({}, {}, {});
 			return;
 		}
 		Request->SetContentType(EHttpMediaType::CbObject);
@@ -1851,7 +1894,7 @@ void FZenCacheStore::EnqueueAsyncRpc(IRequestOwner& Owner, const FCbPackage& Req
 	{
 		if (UNLIKELY(!Request))
 		{
-			OnComplete({}, {});
+			OnComplete({}, {}, {});
 			return;
 		}
 		Request->SetContentType(EHttpMediaType::CbPackage);
diff --git a/Engine/Source/Developer/IoStoreUtilities/Private/ZenStoreHttpClient.cpp b/Engine/Source/Developer/IoStoreUtilities/Private/ZenStoreHttpClient.cpp
--- a/Engine/Source/Developer/IoStoreUtilities/Private/ZenStoreHttpClient.cpp
+++ b/Engine/Source/Developer/IoStoreUtilities/Private/ZenStoreHttpClient.cpp
@@ -505,6 +505,10 @@ TIoStatusOr<uint64> FZenStoreHttpClient::GetChunkSize(const FIoChunkId& Id)
 		const uint64 ChunkSize = ResponseObj["size"].AsUInt64(0);
 		return ChunkSize;
 	}
+	else if (int ResponseCode = Request->GetResponseCode(); ResponseCode != 404)
+	{
+		UE_LOG(LogZenStore, Verbose, TEXT("Request '%s' failed. Error: %s. Status: %s (%d). \"%s\""), *ChunkUri, *FString(Request->GetError()), *FString(UE::Zen::FZenHttpRequest::LexResponseCodeToString(ResponseCode)), ResponseCode, *Request->GetResponseAsString());
+	}
 	return FIoStatus(EIoErrorCode::NotFound);
 }
 
diff --git a/Engine/Source/Developer/Zen/Internal/ZenServerHttp.h b/Engine/Source/Developer/Zen/Internal/ZenServerHttp.h
--- a/Engine/Source/Developer/Zen/Internal/ZenServerHttp.h
+++ b/Engine/Source/Developer/Zen/Internal/ZenServerHttp.h
@@ -105,6 +105,8 @@ public:
 		return int(ResponseCode);
 	}
 
+	ZEN_API FAnsiStringView GetError() const;
+
 	inline const bool GetResponseFormatValid() const
 	{
 		return bResponseFormatValid;
@@ -177,15 +179,12 @@ public:
 	ZEN_API Result PerformRpc(FStringView Uri, FCbObjectView Request, FCbPackage &OutResponse);
 	ZEN_API Result PerformRpc(FStringView Uri, const FCbPackage& Request, FCbPackage& OutResponse);
 
-	/** Returns the response buffer as a string. Note that is the request is performed
+	/** Returns the response buffer as a string. Note that if the request is performed
 		with an external buffer as target buffer this string will be empty.
 	 */
-	inline FString GetResponseAsString() const
-	{
-		return GetAnsiBufferAsString(ResponseBuffer);
-	}
+	ZEN_API FString GetResponseAsString() const;
 
-	/** Returns the response buffer. Note that is the request is performed
+	/** Returns the response buffer. Note that if the request is performed
 	  * with an external buffer as target buffer this will be empty.
 	  */
 	inline const TArray64<uint8>& GetResponseBuffer() const
@@ -198,6 +197,8 @@ public:
 
 	ZEN_API bool GetHeader(const ANSICHAR* Header, FString& OutValue) const;
 
+	ZEN_API static FAnsiStringView LexResponseCodeToString(int HttpCode);
+
 private:
 #if defined(CURL_STRICTER)
 	CURL*					Curl = nullptr;
diff --git a/Engine/Source/Developer/Zen/Private/ZenServerHttp.cpp b/Engine/Source/Developer/Zen/Private/ZenServerHttp.cpp
--- a/Engine/Source/Developer/Zen/Private/ZenServerHttp.cpp
+++ b/Engine/Source/Developer/Zen/Private/ZenServerHttp.cpp
@@ -121,6 +121,15 @@ namespace UE::Zen {
 		bLogErrors = bInLogErrors;
 	}
 
+	FAnsiStringView FZenHttpRequest::GetError() const
+	{
+		if (CurlResult != CURLE_OK)
+		{
+			return curl_easy_strerror((CURLcode)CurlResult);
+		}
+		return {};
+	}
+
 	void FZenHttpRequest::AddHeader(FStringView Header, FStringView Value)
 	{
 		TStringBuilder<128> Sb;
@@ -389,6 +398,80 @@ namespace UE::Zen {
 		return false;
 	}
 
+	FAnsiStringView FZenHttpRequest::LexResponseCodeToString(int HttpCode)
+	{
+		switch (HttpCode)
+		{
+			// 1xx Informational
+		case 100: return "Continue";
+		case 101: return "Switching Protocols";
+
+			// 2xx Success
+		case 200: return "OK";
+		case 201: return "Created";
+		case 202: return "Accepted";
+		case 204: return "No Content";
+		case 205: return "Reset Content";
+		case 206: return "Partial Content";
+
+			// 3xx Redirection
+		case 300: return "Multiple Choices";
+		case 301: return "Moved Permanently";
+		case 302: return "Found";
+		case 303: return "See Other";
+		case 304: return "Not Modified";
+		case 305: return "Use Proxy";
+		case 306: return "Switch Proxy";
+		case 307: return "Temporary Redirect";
+		case 308: return "Permanent Redirect";
+
+			// 4xx Client errors
+		case 400: return "Bad Request";
+		case 401: return "Unauthorized";
+		case 402: return "Payment Required";
+		case 403: return "Forbidden";
+		case 404: return "Not Found";
+		case 405: return "Method Not Allowed";
+		case 406: return "Not Acceptable";
+		case 407: return "Proxy Authentication Required";
+		case 408: return "Request Timeout";
+		case 409: return "Conflict";
+		case 410: return "Gone";
+		case 411: return "Length Required";
+		case 412: return "Precondition Failed";
+		case 413: return "Payload Too Large";
+		case 414: return "URI Too Long";
+		case 415: return "Unsupported Media Type";
+		case 416: return "Range Not Satisfiable";
+		case 417: return "Expectation Failed";
+		case 418: return "I'm a teapot";
+		case 421: return "Misdirected Request";
+		case 422: return "Unprocessable Entity";
+		case 423: return "Locked";
+		case 424: return "Failed Dependency";
+		case 425: return "Too Early";
+		case 426: return "Upgrade Required";
+		case 428: return "Precondition Required";
+		case 429: return "Too Many Requests";
+		case 431: return "Request Header Fields Too Large";
+
+			// 5xx Server errors
+		case 500: return "Internal Server Error";
+		case 501: return "Not Implemented";
+		case 502: return "Bad Gateway";
+		case 503: return "Service Unavailable";
+		case 504: return "Gateway Timeout";
+		case 505: return "HTTP Version Not Supported";
+		case 506: return "Variant Also Negotiates";
+		case 507: return "Insufficient Storage";
+		case 508: return "Loop Detected";
+		case 510: return "Not Extended";
+		case 511: return "Network Authentication Required";
+
+		default: return "Unknown Result";
+		}
+	}
+
 
 	void FZenHttpRequest::LogResult(long InResult, const TCHAR* Uri, RequestVerb Verb) const
 	{
@@ -476,6 +559,11 @@ namespace UE::Zen {
 		return FString::ConstructFromPtrSize(TCHARData.Get(), TCHARData.Length());
 	}
 
+	FString FZenHttpRequest::GetResponseAsString() const
+	{
+		return GetAnsiBufferAsString(ResponseBuffer);
+	}
+
 	FCbObjectView FZenHttpRequest::GetResponseAsObject() const
 	{
 		return FCbObjectView(ResponseBuffer.GetData());
diff --git a/Engine/Source/Runtime/StorageServerClient/Private/StorageServerConnection.cpp b/Engine/Source/Runtime/StorageServerClient/Private/StorageServerConnection.cpp
--- a/Engine/Source/Runtime/StorageServerClient/Private/StorageServerConnection.cpp
+++ b/Engine/Source/Runtime/StorageServerClient/Private/StorageServerConnection.cpp
@@ -40,6 +40,7 @@ bool FStorageServerConnection::Initialize(TArrayView<const FString> HostAddresse
 		CurrentHostAddr = HostAddress;
 		if (HandshakeRequest())
 		{
+			UE_LOG(LogStorageServerConnection, Display, TEXT("Zen store connection established to %s:%i."), *CurrentHostAddr, Port);
 			return true;
 		}
 	}
@@ -415,14 +416,11 @@ TSharedPtr<FInternetAddr> FStorageServerConnection::StringToInternetAddr(const F
 
 bool FStorageServerConnection::HandshakeRequest()
 {
-	TAnsiStringBuilder<256> ResourceBuilder;
-	ResourceBuilder.Append(BaseURI);
-
 	// Handshakes are done with a limited connection timeout so that we can find out if the destination is unreachable in a timely manner.
 	const float ConnectionTimeoutSeconds = 5.0f;
 
 	IStorageServerHttpClient::FResult ResultTuple = HttpClient->RequestSync(
-		*ResourceBuilder,
+		BaseURI,
 		EStorageServerContentType::Unknown,
 		"GET",
 		TOptional<FIoBuffer>(),
diff --git a/Engine/Source/Runtime/StorageServerClient/Private/StorageServerPlatformFile.cpp b/Engine/Source/Runtime/StorageServerClient/Private/StorageServerPlatformFile.cpp
--- a/Engine/Source/Runtime/StorageServerClient/Private/StorageServerPlatformFile.cpp
+++ b/Engine/Source/Runtime/StorageServerClient/Private/StorageServerPlatformFile.cpp
@@ -887,6 +887,10 @@ void FStorageServerPlatformFile::InitializeConnection()
 		UE_LOG(LogStorageServerPlatformFile, Error, TEXT("Failed to initialize connection to %s"), *FString::Join(HostAddrs, TEXT("\n")));
 		FPlatformMisc::RequestExit(true);
 	}
+	else
+	{
+		UE_LOG(LogStorageServerPlatformFile, Warning, TEXT("Failed to initialize connection to %s"), *FString::Join(HostAddrs, TEXT("\n")));
+	}
 }
 
 bool FStorageServerPlatformFile::FileExists(const TCHAR* Filename)
