DDC: Fixed a race in the return value of the legacy function TryToPrefetch

diff --git a/Engine/Source/Developer/DerivedDataCache/Private/DerivedDataCache.cpp b/Engine/Source/Developer/DerivedDataCache/Private/DerivedDataCache.cpp
--- a/Engine/Source/Developer/DerivedDataCache/Private/DerivedDataCache.cpp
+++ b/Engine/Source/Developer/DerivedDataCache/Private/DerivedDataCache.cpp
@@ -944,15 +944,18 @@ public:
 				LegacyRequest.UserData = uint64(Index);
 				++Index;
 			}
-			bool bOk = true;
+			std::atomic<bool> bOk = true;
 			FRequestOwner BlockingOwner(EPriority::Blocking);
 			Backend->GetRoot().LegacyGet(LegacyRequests, BlockingOwner,
 				[&bOk](FLegacyCacheGetResponse&& Response)
 				{
-					bOk &= Response.Status == EStatus::Ok;
+					if (Response.Status != EStatus::Ok)
+					{
+						bOk.store(false, std::memory_order_relaxed);
+					}
 				});
 			BlockingOwner.Wait();
-			return bOk;
+			return bOk.load(std::memory_order_relaxed);
 		}
 		return true;
 	}
