Add support for sending HTTP Cache Requests through a Unix Socket

#rnx

diff --git a/Engine/Source/Developer/DerivedDataCache/Private/Http/CurlHttpClient.cpp b/Engine/Source/Developer/DerivedDataCache/Private/Http/CurlHttpClient.cpp
--- a/Engine/Source/Developer/DerivedDataCache/Private/Http/CurlHttpClient.cpp
+++ b/Engine/Source/Developer/DerivedDataCache/Private/Http/CurlHttpClient.cpp
@@ -4,6 +4,7 @@
 
 #include "Async/InheritedContext.h"
 #include "Async/ManualResetEvent.h"
+#include "Containers/AnsiString.h"
 #include "Containers/ConsumeAllMpmcQueue.h"
 #include "Containers/LockFreeList.h"
 #include "Containers/StringView.h"
@@ -238,6 +239,7 @@ public:
 
 	void Reset() final;
 	void SetUri(FAnsiStringView Uri) final;
+	void SetUnixSocketPath(FAnsiStringView Uri) final;
 	void SetMethod(EHttpMethod Method) final;
 	void SetBody(const FCompositeBuffer& Body) final;
 	void AddHeader(FAnsiStringView Name, FAnsiStringView Value) final;
@@ -264,6 +266,7 @@ private:
 
 	EHttpMethod Method;
 	FAnsiStringBuilderBase Uri;
+	FAnsiString UnixSocketPath;
 	FCurlHttpClient& Client;
 	FCurlHttpHeaders Headers;
 	FCompositeBuffer Body;
@@ -716,6 +719,7 @@ void FCurlHttpRequest::Reset()
 	curl_easy_reset(Curl);
 	Headers.Reset();
 	Uri.Reset();
+	UnixSocketPath.Empty();
 	Body.Reset();
 	SetDefaultOptions();
 }
@@ -752,6 +756,22 @@ void FCurlHttpRequest::SetUri(const FAnsiStringView InUri)
 	curl_easy_setopt(Curl, CURLOPT_URL, *Uri);
 }
 
+void FCurlHttpRequest::SetUnixSocketPath(const FAnsiStringView InSocketPath)
+{
+	CheckIdle(TEXT("SetUnixSocketPath"));
+	UnixSocketPath.Empty();
+	UnixSocketPath.Append(InSocketPath);
+
+	if (UnixSocketPath.Len() > 0)
+	{
+		curl_easy_setopt(Curl, CURLOPT_UNIX_SOCKET_PATH, *UnixSocketPath);
+	}
+	else
+	{
+		curl_easy_setopt(Curl, CURLOPT_UNIX_SOCKET_PATH, NULL);
+	}
+}
+
 void FCurlHttpRequest::SetMethod(const EHttpMethod InMethod)
 {
 	CheckIdle(TEXT("SetMethod"));
diff --git a/Engine/Source/Developer/DerivedDataCache/Private/Http/HttpClient.h b/Engine/Source/Developer/DerivedDataCache/Private/Http/HttpClient.h
--- a/Engine/Source/Developer/DerivedDataCache/Private/Http/HttpClient.h
+++ b/Engine/Source/Developer/DerivedDataCache/Private/Http/HttpClient.h
@@ -237,6 +237,9 @@ public:
 	/** Set the URI used for the request, including the scheme and authority. */
 	virtual void SetUri(FAnsiStringView Uri) = 0;
 
+	/** Set a Unix socket path that should be used for sending requests */
+	virtual void SetUnixSocketPath(FAnsiStringView SocketPath) = 0;
+
 	/** Set the method used for the request. Defaults to GET. */
 	virtual void SetMethod(EHttpMethod Method) = 0;
 
diff --git a/Engine/Source/Developer/DerivedDataCache/Private/HttpCacheStore.cpp b/Engine/Source/Developer/DerivedDataCache/Private/HttpCacheStore.cpp
--- a/Engine/Source/Developer/DerivedDataCache/Private/HttpCacheStore.cpp
+++ b/Engine/Source/Developer/DerivedDataCache/Private/HttpCacheStore.cpp
@@ -225,6 +225,7 @@ struct FHttpCacheStoreParams
 	FString HostPinnedPublicKeys;
 	FString Namespace;
 	FString HttpVersion;
+	FString UnixSocketPath;
 	FString OAuthProvider;
 	FString OAuthClientId;
 	FString OAuthSecret;
@@ -319,6 +320,7 @@ private:
 	FString OAuthProviderIdentifier;
 	FString OAuthAccessToken;
 	FString HttpVersion;
+	FString UnixSocketPath;
 	FString AuthScheme;
 
 	FAnsiStringBuilderBase EffectiveDomain;
@@ -428,6 +430,7 @@ public:
 	// Prepare Request
 
 	void SetUri(FAnsiStringView Uri) { Request->SetUri(Uri); }
+	void SetUnixSocketPath(FAnsiStringView SocketPath) { Request->SetUnixSocketPath(SocketPath); }
 	void SetMethod(EHttpMethod Method) { Request->SetMethod(Method); }
 	void AddHeader(FAnsiStringView Name, FAnsiStringView Value) { Request->AddHeader(Name, Value); }
 	void SetBody(const FCompositeBuffer& Body) { Request->SetBody(Body); }
@@ -725,6 +728,7 @@ public:
 		, NodeName(*CacheStore.NodeName)
 	{
 		Operation.SetUri(WriteToAnsiString<256>(CacheStore.EffectiveDomain, ANSITEXTVIEW("/health/ready")));
+		Operation.SetUnixSocketPath(StringCast<ANSICHAR>(*CacheStore.UnixSocketPath));
 		Operation.SendAsync(Owner, []{});
 	}
 
@@ -861,6 +865,7 @@ void FHttpCacheStore::FPutPackageOp::BeginPutRef(TUniquePtr<FHttpOperation> Oper
 
 	FHttpOperation& LocalOperation = *Operation;
 	LocalOperation.SetUri(RefsUri);
+	LocalOperation.SetUnixSocketPath(StringCast<ANSICHAR>(*CacheStore.UnixSocketPath));
 	if (bFinalize)
 	{
 		LocalOperation.SetMethod(EHttpMethod::Post);
@@ -1018,6 +1023,7 @@ void FHttpCacheStore::FPutPackageOp::BeginPutBlobs(FCbPackage&& Package, FCacheP
 
 			FHttpOperation& LocalOperation = *Operation;
 			LocalOperation.SetUri(WriteToAnsiString<256>(Self->CacheStore.EffectiveDomain, ANSITEXTVIEW("/api/v1/compressed-blobs/"), Self->CacheStore.Namespace, '/', Blob.GetRawHash()));
+			LocalOperation.SetUnixSocketPath(StringCast<ANSICHAR>(*Self->CacheStore.UnixSocketPath));
 			LocalOperation.SetMethod(EHttpMethod::Put);
 			LocalOperation.SetContentType(EHttpMediaType::CompressedBinary);
 			LocalOperation.SetBody(Blob.GetCompressed());
@@ -1213,6 +1219,7 @@ void FHttpCacheStore::FGetRecordOp::GetRecordOnly(const FCacheKey& InKey, const
 
 		FHttpOperation& LocalOperation = *Operation;
 		LocalOperation.SetUri(WriteToAnsiString<256>(Self->CacheStore.EffectiveDomain, ANSITEXTVIEW("/api/v1/refs/"), Self->CacheStore.Namespace, '/', Bucket, '/', Self->Key.Hash));
+		LocalOperation.SetUnixSocketPath(StringCast<ANSICHAR>(*Self->CacheStore.UnixSocketPath));
 		LocalOperation.SetMethod(EHttpMethod::Get);
 		LocalOperation.AddAcceptType(EHttpMediaType::CbObject);
 		LocalOperation.SetExpectedStatusCodes({404});
@@ -1460,6 +1467,7 @@ void FHttpCacheStore::FGetRecordOp::BeginGetValue(
 
 	FHttpOperation& LocalOperation = *Operation;
 	LocalOperation.SetUri(WriteToAnsiString<256>(CacheStore.EffectiveDomain, ANSITEXTVIEW("/api/v1/compressed-blobs/"), CacheStore.Namespace, '/', Value.GetRawHash()));
+	LocalOperation.SetUnixSocketPath(StringCast<ANSICHAR>(*CacheStore.UnixSocketPath));
 	LocalOperation.SetMethod(EHttpMethod::Get);
 	LocalOperation.AddAcceptType(EHttpMediaType::Any);
 	LocalOperation.SetExpectedStatusCodes({404});
@@ -1578,6 +1586,7 @@ void FHttpCacheStore::FGetRecordOp::BeginGetValuesExist(TUniquePtr<FHttpOperatio
 
 	FHttpOperation& LocalOperation = *Operation;
 	LocalOperation.SetUri(Uri);
+	LocalOperation.SetUnixSocketPath(StringCast<ANSICHAR>(*CacheStore.UnixSocketPath));
 	LocalOperation.SetMethod(EHttpMethod::Post);
 	LocalOperation.SetContentType(EHttpMediaType::FormUrlEncoded);
 	LocalOperation.AddAcceptType(EHttpMediaType::Json);
@@ -1746,6 +1755,7 @@ void FHttpCacheStore::FGetValueOp::BeginGetRef(TUniquePtr<FHttpOperation>&& Oper
 
 	FHttpOperation& LocalOperation = *Operation;
 	LocalOperation.SetUri(WriteToAnsiString<256>(CacheStore.EffectiveDomain, ANSITEXTVIEW("/api/v1/refs/"), CacheStore.Namespace, '/', Bucket, '/', Key.Hash));
+	LocalOperation.SetUnixSocketPath(StringCast<ANSICHAR>(*CacheStore.UnixSocketPath));
 	LocalOperation.SetMethod(EHttpMethod::Get);
 	if (bSkipData)
 	{
@@ -1969,6 +1979,7 @@ void FHttpCacheStore::FExistsBatchOp::BeginExists(TUniquePtr<FHttpOperation>&& O
 
 	FHttpOperation& LocalOperation = *Operation;
 	LocalOperation.SetUri(WriteToAnsiString<256>(CacheStore.EffectiveDomain, ANSITEXTVIEW("/api/v1/refs/"), CacheStore.Namespace));
+	LocalOperation.SetUnixSocketPath(StringCast<ANSICHAR>(*CacheStore.UnixSocketPath));
 	LocalOperation.SetMethod(EHttpMethod::Post);
 	LocalOperation.SetContentType(EHttpMediaType::CbObject);
 	LocalOperation.AddAcceptType(EHttpMediaType::CbObject);
@@ -2114,6 +2125,7 @@ FHttpCacheStore::FHttpCacheStore(const FHttpCacheStoreParams& Params, ICacheStor
 	, OAuthProviderIdentifier(Params.OAuthProviderIdentifier)
 	, OAuthAccessToken(Params.OAuthAccessToken)
 	, HttpVersion(Params.HttpVersion)
+	, UnixSocketPath(Params.UnixSocketPath)
 	, AuthScheme(Params.AuthScheme)
 	, StoreOwner(Owner)
 	, bReadOnly(Params.bReadOnly)
@@ -3212,6 +3224,25 @@ void FHttpCacheStoreParams::Parse(const TCHAR* NodeName, const TCHAR* Config)
 		}
 	}
 
+	// Unix Socket Params
+	FParse::Value(Config, TEXT("UnixSocketPath="), UnixSocketPath);
+	if (FParse::Value(Config, TEXT("EnvUnixSocketPathOverride="), OverrideName))
+	{
+		FString UnixSocketPathEnv = FPlatformMisc::GetEnvironmentVariable(*OverrideName);
+		if (!UnixSocketPathEnv.IsEmpty())
+		{
+			UnixSocketPath = UnixSocketPathEnv;
+			UE_LOG(LogDerivedDataCache, Log, TEXT("%s: Found environment override for UnixSocketPath %s=%s"), NodeName, *OverrideName, *UnixSocketPath);
+		}
+	}
+	if (FParse::Value(Config, TEXT("CommandLineUnixSocketPathOverride="), OverrideName))
+	{
+		if (FParse::Value(FCommandLine::Get(), *(OverrideName + TEXT("=")), UnixSocketPath))
+		{
+			UE_LOG(LogDerivedDataCache, Log, TEXT("%s: Found command line override for UnixSocketPath %s=%s"), NodeName, *OverrideName, *UnixSocketPath);
+		}
+	}
+
 	// Namespace Params
 
 	if (Namespace.IsEmpty())
