Fixing a bug where DDC-Verify when using multi-process cook would have cook workers delete the output directory at inappropriate times.  Only allow the cook director to delete the directory.

diff --git a/Engine/Source/Developer/DerivedDataCache/Private/DerivedDataBackends.cpp b/Engine/Source/Developer/DerivedDataCache/Private/DerivedDataBackends.cpp
--- a/Engine/Source/Developer/DerivedDataCache/Private/DerivedDataBackends.cpp
+++ b/Engine/Source/Developer/DerivedDataCache/Private/DerivedDataBackends.cpp
@@ -226,7 +226,10 @@ public:
 			FParse::Param(FCommandLine::Get(), TEXT("DDC-Verify")) ||
 			bVerifyFound)
 		{
-			IFileManager::Get().DeleteDirectory(*(FPaths::ProjectSavedDir() / TEXT("VerifyDDC/")), /*bRequireExists*/ false, /*bTree*/ true);
+			if (!UE::GetMultiprocessId())
+			{
+				IFileManager::Get().DeleteDirectory(*(FPaths::ProjectSavedDir() / TEXT("VerifyDDC/")), /*bRequireExists*/ false, /*bTree*/ true);
+			}
 			RootNode = CreateCacheStoreVerify(RootNode, /*bPutOnError*/ bVerifyFix);
 		}
 
