Allow the Jupiter DDC client to specify an explicit "allow overwrite" header (X-Jupiter-Allow-Overwrite) iff the caller's policy is to store without querying.  When "allow overwrite" is not specified, if the client receives a 409 status code during a put operation (either record or value), it will emit a display verbosity message about a non-deterministic value being detected that aligns with what the filesystem cache store emits in that situation.

diff --git a/Engine/Source/Developer/DerivedDataCache/Private/HttpCacheStore.cpp b/Engine/Source/Developer/DerivedDataCache/Private/HttpCacheStore.cpp
--- a/Engine/Source/Developer/DerivedDataCache/Private/HttpCacheStore.cpp
+++ b/Engine/Source/Developer/DerivedDataCache/Private/HttpCacheStore.cpp
@@ -772,7 +772,7 @@ public:
 		return new FPutPackageOp(CacheStore, Owner, Name);
 	}
 
-	void Put(const FCacheKey& Key, FCbPackage&& Package, FOnPackageComplete&& OnComplete);
+	void Put(const FCacheKey& Key, const FCacheRecordPolicy& Policy, FCbPackage&& Package, FOnPackageComplete&& OnComplete);
 
 	const FRequestStats& ReadStats() const { return RequestStats; }
 	FRequestStats& EditStats() { return RequestStats; }
@@ -802,10 +802,10 @@ private:
 
 	FPutPackageOp(FHttpCacheStore& CacheStore, IRequestOwner& Owner, const FSharedString& Name);
 
-	void BeginOperation(bool bFinalize, FOnCachePutRefComplete&& OnComplete);
+	void BeginOperation(bool bFinalize, bool bAllowOverwrite, FOnCachePutRefComplete&& OnComplete);
 
-	void BeginPutRef(TUniquePtr<FHttpOperation> Operation, bool bFinalize, FOnCachePutRefComplete&& OnComplete);
-	void EndPutRef(TUniquePtr<FHttpOperation> Operation, bool bFinalize, FOnCachePutRefComplete&& OnComplete);
+	void BeginPutRef(TUniquePtr<FHttpOperation> Operation, bool bFinalize, bool bAllowOverwrite, FOnCachePutRefComplete&& OnComplete);
+	void EndPutRef(TUniquePtr<FHttpOperation> Operation, bool bFinalize, bool bAllowOverwrite, FOnCachePutRefComplete&& OnComplete);
 
 	void BeginPutBlobs(FCbPackage&& Package, FCachePutRefResponse&& Response);
 	void EndPutBlob(FHttpOperation* Operation, uint64 LogicalSize);
@@ -823,27 +823,30 @@ FHttpCacheStore::FPutPackageOp::FPutPackageOp(FHttpCacheStore& InCacheStore, IRe
 	RequestStats.Name = Name;
 }
 
-void FHttpCacheStore::FPutPackageOp::Put(const FCacheKey& InKey, FCbPackage&& Package, FOnPackageComplete&& OnComplete)
+void FHttpCacheStore::FPutPackageOp::Put(const FCacheKey& InKey, const FCacheRecordPolicy& Policy, FCbPackage&& Package, FOnPackageComplete&& OnComplete)
 {
+	ECachePolicy RecordPolicy = Policy.GetRecordPolicy();
+	const bool bAllowOverwrite = !EnumHasAnyFlags(RecordPolicy, ECachePolicy::QueryRemote);
 	Key = InKey;
 	Object = Package.GetObject();
 	ObjectHash = Package.GetObjectHash();
 	OnPackageComplete = MoveTemp(OnComplete);
-	BeginOperation(/*bFinalize*/ false, [Self = TRefCountPtr(this), Package = MoveTemp(Package)](FCachePutRefResponse&& Response) mutable
+	BeginOperation(/*bFinalize*/ false, bAllowOverwrite, [Self = TRefCountPtr(this), Package = MoveTemp(Package)](FCachePutRefResponse&& Response) mutable
 	{
 		Self->BeginPutBlobs(MoveTemp(Package), MoveTemp(Response));
 	});
 }
 
-void FHttpCacheStore::FPutPackageOp::BeginOperation(bool bFinalize, FOnCachePutRefComplete&& OnComplete)
+void FHttpCacheStore::FPutPackageOp::BeginOperation(bool bFinalize, bool bAllowOverwrite, FOnCachePutRefComplete&& OnComplete)
 {
-	CacheStore.WaitForHttpOperationAsync(Owner, bFinalize ? EOperationCategory::PutFinalize : EOperationCategory::PutRef, [Self = TRefCountPtr(this), bFinalize, OnComplete = MoveTemp(OnComplete)](TUniquePtr<FHttpOperation>&& Operation) mutable
+	CacheStore.WaitForHttpOperationAsync(Owner, bFinalize ? EOperationCategory::PutFinalize : EOperationCategory::PutRef,
+	[Self = TRefCountPtr(this), bFinalize, bAllowOverwrite, OnComplete = MoveTemp(OnComplete)](TUniquePtr<FHttpOperation>&& Operation) mutable
 	{
-		Self->BeginPutRef(MoveTemp(Operation), bFinalize, MoveTemp(OnComplete));
+		Self->BeginPutRef(MoveTemp(Operation), bFinalize, bAllowOverwrite, MoveTemp(OnComplete));
 	});
 }
 
-void FHttpCacheStore::FPutPackageOp::BeginPutRef(TUniquePtr<FHttpOperation> Operation, bool bFinalize, FOnCachePutRefComplete&& OnComplete)
+void FHttpCacheStore::FPutPackageOp::BeginPutRef(TUniquePtr<FHttpOperation> Operation, bool bFinalize, bool bAllowOverwrite, FOnCachePutRefComplete&& OnComplete)
 {
 	if (UNLIKELY(!Operation))
 	{
@@ -877,26 +880,45 @@ void FHttpCacheStore::FPutPackageOp::BeginPutRef(TUniquePtr<FHttpOperation> Oper
 		LocalOperation.SetContentType(EHttpMediaType::CbObject);
 		LocalOperation.AddHeader(ANSITEXTVIEW("X-Jupiter-IoHash"), WriteToAnsiString<48>(ObjectHash));
 		LocalOperation.SetBody(Object.GetBuffer());
+
+		// Only specify the X-Jupiter-Allow-Overwrite header if we are specifically allowing overwrites.
+		// If we are not allowing overwrites, we leave the header unspecified.  This is both to avoid
+		// sending excess bytes with every put, but also to allow the server to apply its own default
+		// behavior instead of having the client specify the overwrite behavior explicitly.
+		if (bAllowOverwrite)
+		{
+			LocalOperation.AddHeader(ANSITEXTVIEW("X-Jupiter-Allow-Overwrite"), WriteToAnsiString<8>(bAllowOverwrite));
+		}
+		else
+		{
+			LocalOperation.SetExpectedStatusCodes({409});
+		}
 	}
 	LocalOperation.AddAcceptType(EHttpMediaType::Json);
 
 	RequestTimer.Stop();
-	LocalOperation.SendAsync(Owner, [Self = TRefCountPtr(this), Operation = MoveTemp(Operation), bFinalize, OnComplete = MoveTemp(OnComplete)]() mutable
+	LocalOperation.SendAsync(Owner, [Self = TRefCountPtr(this), Operation = MoveTemp(Operation), bFinalize, bAllowOverwrite, OnComplete = MoveTemp(OnComplete)]() mutable
 	{
 		Operation->GetStats(Self->RequestStats);
-		Self->EndPutRef(MoveTemp(Operation), bFinalize, MoveTemp(OnComplete));
+		Self->EndPutRef(MoveTemp(Operation), bFinalize, bAllowOverwrite, MoveTemp(OnComplete));
 	});
 }
 
 void FHttpCacheStore::FPutPackageOp::EndPutRef(
 	TUniquePtr<FHttpOperation> Operation,
 	bool bFinalize,
+	bool bAllowOverwrite,
 	FOnCachePutRefComplete&& OnComplete)
 {
 	TRACE_CPUPROFILER_EVENT_SCOPE(HttpDDC_PutPackage_EndPutRef);
 
 	if (const int32 StatusCode = Operation->GetStatusCode(); StatusCode < 200 || StatusCode > 204)
 	{
+		if (!bAllowOverwrite && (StatusCode == 409))
+		{
+			UE_LOG(LogDerivedDataCache, Display, TEXT("%s: Cache put found non-deterministic record for %s from '%s'"),
+				*CacheStore.NodeName, *WriteToString<96>(Key), *Name);
+		}
 		const EStatus Status = Operation->GetErrorCode() == EHttpErrorCode::Canceled ? EStatus::Canceled : EStatus::Error;
 		OnComplete({{}, Status});
 		return;
@@ -1058,7 +1080,7 @@ void FHttpCacheStore::FPutPackageOp::EndPutBlob(FHttpOperation* Operation, uint6
 		}
 		else if (LocalSuccessfulBlobUploads == TotalBlobUploads)
 		{
-			BeginOperation(/*bFinalize*/ true, [Self = TRefCountPtr(this)](FCachePutRefResponse&& Response)
+			BeginOperation(/*bFinalize*/ true, /*bAllowOverwrite*/ false, [Self = TRefCountPtr(this)](FCachePutRefResponse&& Response)
 			{
 				Self->EndPutRefFinalize(MoveTemp(Response));
 			});
@@ -2583,9 +2605,6 @@ void FHttpCacheStore::PutCacheRecordAsync(IRequestOwner& Owner, const FCachePutR
 		return OnComplete(Request.MakeResponse(EStatus::Error));
 	}
 
-	// TODO: Jupiter currently always overwrites.  It doesn't have a "write if not present" feature (for records or attachments),
-	//		 but would require one to implement all policy correctly.
-
 	TRefCountPtr<FPutPackageOp> Op = FPutPackageOp::New(*this, Owner, Request.Name);
 
 	FCbPackage Package;
@@ -2599,7 +2618,7 @@ void FHttpCacheStore::PutCacheRecordAsync(IRequestOwner& Owner, const FCachePutR
 		Package = Request.Record.Save();
 	}
 
-	Op->Put(Key, MoveTemp(Package), [Op, Request, OnComplete = MoveTemp(OnComplete)](FPutPackageOp::FResponse&& Response)
+	Op->Put(Key, Request.Policy, MoveTemp(Package), [Op, Request, OnComplete = MoveTemp(OnComplete)](FPutPackageOp::FResponse&& Response)
 	{
 		if (Response.Status == EStatus::Ok)
 		{
@@ -2639,9 +2658,6 @@ void FHttpCacheStore::PutCacheValueAsync(IRequestOwner& Owner, const FCachePutVa
 		return OnComplete(Request.MakeResponse(EStatus::Error));
 	}
 
-	// TODO: Jupiter currently always overwrites.  It doesn't have a "write if not present" feature (for records or attachments),
-	//		 but would require one to implement all policy correctly.
-
 	TRefCountPtr<FPutPackageOp> Op = FPutPackageOp::New(*this, Owner, Request.Name);
 
 	FCbPackage Package;
@@ -2663,7 +2679,7 @@ void FHttpCacheStore::PutCacheValueAsync(IRequestOwner& Owner, const FCachePutVa
 		Package.AddAttachment(FCbAttachment(Request.Value.GetData()));
 	}
 
-	Op->Put(Request.Key, MoveTemp(Package), [Op, Request, OnComplete = MoveTemp(OnComplete)](FPutPackageOp::FResponse&& Response)
+	Op->Put(Request.Key, Request.Policy, MoveTemp(Package), [Op, Request, OnComplete = MoveTemp(OnComplete)](FPutPackageOp::FResponse&& Response)
 	{
 		TRACE_COUNTER_ADD(HttpDDC_BytesReceived, Op->ReadStats().PhysicalReadSize);
 		TRACE_COUNTER_ADD(HttpDDC_BytesSent, Op->ReadStats().PhysicalWriteSize);
