DDC: Delete corrupted files in file system cache store when accessed through GetChunks

diff --git a/Engine/Source/Developer/DerivedDataCache/Private/DerivedDataCacheKey.cpp b/Engine/Source/Developer/DerivedDataCache/Private/DerivedDataCacheKey.cpp
--- a/Engine/Source/Developer/DerivedDataCache/Private/DerivedDataCacheKey.cpp
+++ b/Engine/Source/Developer/DerivedDataCache/Private/DerivedDataCacheKey.cpp
@@ -225,6 +225,11 @@ FCbWriter& operator<<(FCbWriter& Writer, const FCacheKey& Key)
 	return Writer;
 }
 
+void SerializeForLog(FCbWriter& Writer, const FCacheKey& Key)
+{
+	Writer.AddString(WriteToString<96>(Key));
+}
+
 bool LoadFromCompactBinary(const FCbFieldView Field, FCacheKey& OutKey)
 {
 	bool bOk = Field.IsObject();
diff --git a/Engine/Source/Developer/DerivedDataCache/Private/FileSystemCacheStore.cpp b/Engine/Source/Developer/DerivedDataCache/Private/FileSystemCacheStore.cpp
--- a/Engine/Source/Developer/DerivedDataCache/Private/FileSystemCacheStore.cpp
+++ b/Engine/Source/Developer/DerivedDataCache/Private/FileSystemCacheStore.cpp
@@ -26,6 +26,7 @@
 #include "HAL/Thread.h"
 #include "Hash/xxhash.h"
 #include "HashingArchiveProxy.h"
+#include "Logging/StructuredLog.h"
 #include "Misc/App.h"
 #include "Misc/CommandLine.h"
 #include "Misc/ConfigCacheIni.h"
@@ -930,7 +931,7 @@ private:
 		ECachePolicy Policy,
 		FValue& OutValue,
 		FRequestStats& Stats) const;
-	void GetCacheContent(
+	[[nodiscard]] bool GetCacheContent(
 		FStringView Name,
 		const FCacheKey& Key,
 		const FValueId& Id,
@@ -940,6 +941,9 @@ private:
 		TUniquePtr<FArchive>& OutArchive,
 		FRequestStats& Stats) const;
 
+	void DeleteCacheContent(FStringView Name, const FValue& Value) const;
+	void DeleteCacheContent(FStringView Name, const FValue& Value, FCompressedBufferReader& Reader, TUniquePtr<FArchive>& OutArchive) const;
+
 	void BuildCachePackagePath(const FCacheKey& CacheKey, FStringBuilderBase& Path) const;
 	void BuildCacheContentPath(const FIoHash& RawHash, FStringBuilderBase& Path) const;
 
@@ -1783,11 +1787,10 @@ void FFileSystemCacheStore::GetChunks(
 					{
 						if (const FValueWithId& ValueWithId = Record.Get().GetValue(Request.Id))
 						{
-							bHasValue = true;
 							Value = ValueWithId;
 							ValueId = Request.Id;
 							ValueKey = Request.Key;
-							GetCacheContent(Request.Name, Request.Key, ValueId, Value, Request.Policy, ValueReader, ValueAr, RequestStats);
+							bHasValue = GetCacheContent(Request.Name, Request.Key, ValueId, Value, Request.Policy, ValueReader, ValueAr, RequestStats);
 						}
 					}
 				}
@@ -1797,7 +1800,7 @@ void FFileSystemCacheStore::GetChunks(
 					bHasValue = GetCacheValueOnly(Request.Name, Request.Key, Request.Policy, Value, RequestStats);
 					if (bHasValue)
 					{
-						GetCacheContent(Request.Name, Request.Key, Request.Id, Value, Request.Policy, ValueReader, ValueAr, RequestStats);
+						bHasValue = GetCacheContent(Request.Name, Request.Key, Request.Id, Value, Request.Policy, ValueReader, ValueAr, RequestStats);
 					}
 				}
 			}
@@ -1812,6 +1815,13 @@ void FFileSystemCacheStore::GetChunks(
 				{
 					Buffer = ValueReader.Decompress(RawOffset, RawSize);
 					RequestStats.LogicalReadSize += Buffer.GetSize();
+					if (!Buffer)
+					{
+						UE_LOGFMT(LogDerivedDataCache, Display,
+							"{Cache}: Cache miss with corrupted value {Id} with hash {RawHash} for {Key} from '{Name}'",
+							CachePath, Request.Id, Value.GetRawHash(), Request.Key, Request.Name);
+						DeleteCacheContent(Request.Name, Value, ValueReader, ValueAr);
+					}
 				}
 				Status = bExistsOnly || Buffer.GetSize() == RawSize ? EStatus::Ok : EStatus::Error;
 			}
@@ -2466,6 +2476,7 @@ bool FFileSystemCacheStore::GetCacheContent(
 				TEXT("%s: Cache miss with corrupted value %s with hash %s for %s from '%.*s'"),
 				*CachePath, *WriteToString<16>(Id), *WriteToString<48>(RawHash),
 				*WriteToString<96>(Key), Name.Len(), Name.GetData());
+			DeleteCacheContent(Name, Value);
 			return false;
 		}
 	}
@@ -2477,7 +2488,7 @@ bool FFileSystemCacheStore::GetCacheContent(
 	return false;
 }
 
-void FFileSystemCacheStore::GetCacheContent(
+bool FFileSystemCacheStore::GetCacheContent(
 	const FStringView Name,
 	const FCacheKey& Key,
 	const FValueId& Id,
@@ -2509,7 +2520,7 @@ void FFileSystemCacheStore::GetCacheContent(
 
 	if (!EnumHasAnyFlags(Policy, ECachePolicy::Query))
 	{
-		return;
+		return true;
 	}
 
 	if (Value.HasData())
@@ -2519,7 +2530,7 @@ void FFileSystemCacheStore::GetCacheContent(
 			Reader.SetSource(Value.GetData());
 		}
 		OutArchive.Reset();
-		return;
+		return true;
 	}
 
 	const FIoHash& RawHash = Value.GetRawHash();
@@ -2534,7 +2545,7 @@ void FFileSystemCacheStore::GetCacheContent(
 			{
 				AccessLogWriter->Append(RawHash, Path);
 			}
-			return;
+			return true;
 		}
 	}
 	else
@@ -2553,15 +2564,14 @@ void FFileSystemCacheStore::GetCacheContent(
 				{
 					AccessLogWriter->Append(RawHash, Path);
 				}
-				return;
+				return true;
 			}
 			UE_LOG(LogDerivedDataCache, Display,
 				TEXT("%s: Cache miss with corrupted value %s with hash %s for %s from '%.*s'"),
 				*CachePath, *WriteToString<16>(Id), *WriteToString<48>(RawHash),
 				*WriteToString<96>(Key), Name.Len(), Name.GetData());
-			Reader.ResetSource();
-			OutArchive.Reset();
-			return;
+			DeleteCacheContent(Name, Value, Reader, OutArchive);
+			return false;
 		}
 	}
 
@@ -2569,6 +2579,34 @@ void FFileSystemCacheStore::GetCacheContent(
 		TEXT("%s: Cache miss with missing value %s with hash %s for %s from '%.*s'"),
 		*CachePath, *WriteToString<16>(Id), *WriteToString<48>(RawHash), *WriteToString<96>(Key),
 		Name.Len(), Name.GetData());
+	return false;
+}
+
+void FFileSystemCacheStore::DeleteCacheContent(const FStringView Name, const FValue& Value) const
+{
+	if (!bReadOnly)
+	{
+		TStringBuilder<256> Path;
+		BuildCacheContentPath(Value.GetRawHash(), Path);
+		if (IFileManager::Get().Delete(*Path, /*bRequireExists*/ false, /*bEvenReadOnly*/ false, /*bQuiet*/ true))
+		{
+			UE_LOGFMT(LogDerivedDataCache, VeryVerbose, "{Cache}: Deleted {Path} from '{Name}'", CachePath, Path, Name);
+		}
+	}
+}
+
+void FFileSystemCacheStore::DeleteCacheContent(
+	const FStringView Name,
+	const FValue& Value,
+	FCompressedBufferReader& Reader,
+	TUniquePtr<FArchive>& OutArchive) const
+{
+	if (OutArchive)
+	{
+		Reader.ResetSource();
+		OutArchive.Reset();
+	}
+	DeleteCacheContent(Name, Value);
 }
 
 void FFileSystemCacheStore::BuildCachePackagePath(const FCacheKey& CacheKey, FStringBuilderBase& Path) const
diff --git a/Engine/Source/Developer/DerivedDataCache/Public/DerivedDataCacheKey.h b/Engine/Source/Developer/DerivedDataCache/Public/DerivedDataCacheKey.h
--- a/Engine/Source/Developer/DerivedDataCache/Public/DerivedDataCacheKey.h
+++ b/Engine/Source/Developer/DerivedDataCache/Public/DerivedDataCacheKey.h
@@ -147,6 +147,8 @@ inline TStringBuilderBase<CharType>& operator<<(TStringBuilderBase<CharType>& Bu
 	return Builder << Key.Bucket << CharType('/') << Key.Hash;
 }
 
+UE_API void SerializeForLog(FCbWriter& Writer, const FCacheKey& Key);
+
 UE_API FCbWriter& operator<<(FCbWriter& Writer, const FCacheKey& Key);
 UE_API bool LoadFromCompactBinary(FCbFieldView Field, FCacheKey& OutKey);
 
