Avoid using the batch REF existence check operation from the jupiter ddc client when we are only checking existence of a single item.

diff --git a/Engine/Source/Developer/DerivedDataCache/Private/HttpCacheStore.cpp b/Engine/Source/Developer/DerivedDataCache/Private/HttpCacheStore.cpp
--- a/Engine/Source/Developer/DerivedDataCache/Private/HttpCacheStore.cpp
+++ b/Engine/Source/Developer/DerivedDataCache/Private/HttpCacheStore.cpp
@@ -3082,8 +3082,11 @@ void FHttpCacheStore::GetValue(
 	TRACE_CPUPROFILER_EVENT_SCOPE(HttpDDC_GetValue);
 	TRACE_COUNTER_ADD(HttpDDC_Get, Requests.Num());
 
+	// NOTE: Both the 'if' and the 'else' branch handle SkipData - see how FHttpCacheStore::FGetValueOp::BeginGetRef chooses to send
+	//		 or not send x-jupiter-inline as an accept type depending on whether the request has SkipData - we just choose to use
+	//		 the batch existence operation so we avoid multiple round trips when there is more than one existence check being requested.
 	const auto HasSkipData = [](const FCacheGetValueRequest& Request) { return EnumHasAnyFlags(Request.Policy, ECachePolicy::SkipData); };
-	if (Algo::AllOf(Requests, HasSkipData))
+	if ((Requests.Num() > 1) && Algo::AllOf(Requests, HasSkipData))
 	{
 		TRefCountPtr<FExistsBatchOp> Op = FExistsBatchOp::New(*this, Owner);
 		Op->Exists(Requests, [this, OnComplete = MoveTemp(OnComplete)](FCacheGetValueResponse&& Response)
