Forcing remote cache stores to be read-only when using a windows editor build that was compiled with clang to avoid having floating point computation differences cause the editor to write DDC key contents that differ from what an MSVC build would produce.

diff --git a/Engine/Source/Developer/DerivedDataCache/Private/FileSystemCacheStore.cpp b/Engine/Source/Developer/DerivedDataCache/Private/FileSystemCacheStore.cpp
--- a/Engine/Source/Developer/DerivedDataCache/Private/FileSystemCacheStore.cpp
+++ b/Engine/Source/Developer/DerivedDataCache/Private/FileSystemCacheStore.cpp
@@ -1191,6 +1191,18 @@ FFileSystemCacheStore::FFileSystemCacheStore(
 		bTouch = false;
 		//bReadOnly = true;
 	}
+	
+	// We are forcing remote cache stores to be read-only if the executable was clang compiled to ensure that differences in floating
+	// point precision don't cause writes to DDC keys with contents that differ from what an MSVC build would produce.
+#if PLATFORM_WINDOWS && WITH_EDITOR && defined(__clang__)
+	if (bRemote)
+	{
+		UE_LOG(LogDerivedDataCache, Display,
+			TEXT("%s: Forcing read-only mode for clang build of editor on windows."),
+			*CachePath);
+		bReadOnly = true;
+	}
+#endif
 
 	if (FString(FCommandLine::Get()).Contains(TEXT("Run=DerivedDataCache")))
 	{
diff --git a/Engine/Source/Developer/DerivedDataCache/Private/HttpCacheStore.cpp b/Engine/Source/Developer/DerivedDataCache/Private/HttpCacheStore.cpp
--- a/Engine/Source/Developer/DerivedDataCache/Private/HttpCacheStore.cpp
+++ b/Engine/Source/Developer/DerivedDataCache/Private/HttpCacheStore.cpp
@@ -2242,8 +2242,18 @@ FHttpCacheStore::FHttpCacheStore(const FHttpCacheStoreParams& Params, ICacheStor
 
 		if (StoreOwner)
 		{
+			bReadOnly = Params.bReadOnly;
+
+			// We are forcing remote cache stores to be read-only if the executable was clang compiled to ensure that differences in floating
+			// point precision don't cause writes to DDC keys with contents that differ from what an MSVC build would produce.
+#if PLATFORM_WINDOWS && WITH_EDITOR && defined(__clang__)
+			UE_LOG(LogDerivedDataCache, Display,
+				TEXT("%s: Forcing read-only mode for clang build of editor on windows."),
+				NodeName);
+			bReadOnly = true;
+#endif
 			const ECacheStoreFlags Flags = ECacheStoreFlags::Remote | ECacheStoreFlags::Query |
-				(Params.bReadOnly ? ECacheStoreFlags::None : ECacheStoreFlags::Store);
+				(bReadOnly ? ECacheStoreFlags::None : ECacheStoreFlags::Store);
 			TStringBuilder<256> Path(InPlace, ResolvedHost, TEXTVIEW(" ("), Namespace, TEXTVIEW(")"));
 			StoreOwner->Add(this, Flags);
 			StoreStats = StoreOwner->CreateStats(this, Flags, TEXTVIEW("Unreal Cloud DDC"), Params.Name, Path);
diff --git a/Engine/Source/Developer/DerivedDataCache/Private/ZenCacheStore.cpp b/Engine/Source/Developer/DerivedDataCache/Private/ZenCacheStore.cpp
--- a/Engine/Source/Developer/DerivedDataCache/Private/ZenCacheStore.cpp
+++ b/Engine/Source/Developer/DerivedDataCache/Private/ZenCacheStore.cpp
@@ -1789,6 +1789,20 @@ void FZenCacheStore::Initialize(const FZenCacheStoreParams& Params)
 		// In the future this could be extended to allow the Remote flag by default (even for locally launched instances) if they have upstreams configured.
 		bool bRemote = Params.bRemote.Get(!bIsLocalConnection);
 
+		bool bReadOnly = Params.bReadOnly;
+
+		// We are forcing remote cache stores to be read-only if the executable was clang compiled to ensure that differences in floating
+		// point precision don't cause writes to DDC keys with contents that differ from what an MSVC build would produce.
+#if PLATFORM_WINDOWS && WITH_EDITOR && defined(__clang__)
+		if (bRemote)
+		{
+			UE_LOG(LogDerivedDataCache, Display,
+				TEXT("%s: Forcing read-only mode for clang build of editor on windows."),
+				*GetName());
+			bReadOnly = true;
+		}
+#endif
+
 		DeactivateAtMs = Params.DeactivateAtMs;
 
 		ECacheStoreFlags Flags = ECacheStoreFlags::Query;
