Fix a missing FString conversion to C string when printing a log message on clang compiled windows editor.

diff --git a/Engine/Source/Developer/DerivedDataCache/Private/HttpCacheStore.cpp b/Engine/Source/Developer/DerivedDataCache/Private/HttpCacheStore.cpp
--- a/Engine/Source/Developer/DerivedDataCache/Private/HttpCacheStore.cpp
+++ b/Engine/Source/Developer/DerivedDataCache/Private/HttpCacheStore.cpp
@@ -2249,7 +2249,7 @@ FHttpCacheStore::FHttpCacheStore(const FHttpCacheStoreParams& Params, ICacheStor
 #if PLATFORM_WINDOWS && WITH_EDITOR && defined(__clang__)
 			UE_LOG(LogDerivedDataCache, Display,
 				TEXT("%s: Forcing read-only mode for clang build of editor on windows."),
-				NodeName);
+				*NodeName);
 			bReadOnly = true;
 #endif
 			const ECacheStoreFlags Flags = ECacheStoreFlags::Remote | ECacheStoreFlags::Query |
