DDC: Fix request dispatching in the hierarchy to track the state index correctly

diff --git a/Engine/Source/Developer/DerivedDataCache/Private/DerivedDataCacheStoreHierarchy.cpp b/Engine/Source/Developer/DerivedDataCache/Private/DerivedDataCacheStoreHierarchy.cpp
--- a/Engine/Source/Developer/DerivedDataCache/Private/DerivedDataCacheStoreHierarchy.cpp
+++ b/Engine/Source/Developer/DerivedDataCache/Private/DerivedDataCacheStoreHierarchy.cpp
@@ -951,6 +951,8 @@ void FCacheStoreHierarchy::TGetBatch<Params>::DispatchRequests()
 		uint64 StateIndex = 0;
 		for (FState& State : States)
 		{
+			ON_SCOPE_EXIT { ++StateIndex; };
+
 			if (NodeIndex < State.NodeIndex)
 			{
 				continue;
@@ -995,7 +997,6 @@ void FCacheStoreHierarchy::TGetBatch<Params>::DispatchRequests()
 					++State.NodeIndex;
 				}
 			}
-			++StateIndex;
 		}
 
 		if (!AsyncNodeRequests.IsEmpty())
