add cvar to allow keeping sandboxed cache zen server instances running after launching process exits

#rnx

diff --git a/Engine/Source/Developer/DerivedDataCache/Private/ZenCacheStore.cpp b/Engine/Source/Developer/DerivedDataCache/Private/ZenCacheStore.cpp
--- a/Engine/Source/Developer/DerivedDataCache/Private/ZenCacheStore.cpp
+++ b/Engine/Source/Developer/DerivedDataCache/Private/ZenCacheStore.cpp
@@ -19,6 +19,7 @@
 #include "DerivedDataRequestOwner.h"
 #include "Experimental/ZenStatistics.h"
 #include "HAL/FileManager.h"
+#include "HAL/IConsoleManager.h"
 #include "HAL/Thread.h"
 #include "Http/HttpClient.h"
 #include "Http/HttpHostBuilder.h"
@@ -46,6 +47,13 @@ TRACE_DECLARE_ATOMIC_INT_COUNTER(ZenDDC_BytesSent, TEXT("ZenDDC Bytes Sent"));
 TRACE_DECLARE_ATOMIC_INT_COUNTER(ZenDDC_CacheRecordRequestCountInFlight, TEXT("ZenDDC CacheRecord Request Count"));
 TRACE_DECLARE_ATOMIC_INT_COUNTER(ZenDDC_ChunkRequestCountInFlight, TEXT("ZenDDC Chunk Request Count"));
 
+static TAutoConsoleVariable<bool> CVarZenCacheLimitSandboxProcessLifetime(
+	TEXT("zen.cache.limitsandboxprocesslifetime"),
+	true,
+	TEXT("When true, Zen server process for sandbox cache instances will be shut down when the process which started it exits. When false, process will persist.\n")
+	TEXT("Default: true"),
+	ECVF_ReadOnly);
+
 namespace UE::DerivedData
 {
 
@@ -2277,8 +2285,8 @@ ILegacyCacheStore* CreateZenCacheStore(const TCHAR* NodeName, const TCHAR* Confi
 		AutoLaunchSettings.DataPath += Params.Sandbox;
 		AutoLaunchSettings.bIsDefaultSharedRunContext = false;
 
-		// The unique local instances will always limit process lifetime for now to avoid accumulating many of them
-		AutoLaunchSettings.bLimitProcessLifetime = true;
+		// The unique local instances will always limit process lifetime by default for now to avoid accumulating many of them
+		AutoLaunchSettings.bLimitProcessLifetime = CVarZenCacheLimitSandboxProcessLifetime.GetValueOnAnyThread();
 
 		// Flush the cache if requested.
 		uint32 MultiprocessId = UE::GetMultiprocessId();
