Addressing static analysis warnings caused by self-deleting data types.

diff --git a/Engine/Source/Developer/DerivedDataCache/Private/ZenCacheStore.cpp b/Engine/Source/Developer/DerivedDataCache/Private/ZenCacheStore.cpp
--- a/Engine/Source/Developer/DerivedDataCache/Private/ZenCacheStore.cpp
+++ b/Engine/Source/Developer/DerivedDataCache/Private/ZenCacheStore.cpp
@@ -1774,13 +1774,14 @@ void FZenCacheStore::EnqueueAsyncRpc(IRequestOwner& Owner, FCbObject RequestObje
 		}
 		Request->SetContentType(EHttpMediaType::CbObject);
 		Request->SetBody(RequestObject.GetBuffer().MakeOwned());
+		// Receiver will self-delete
 		FRequestAsyncCbPackageReceiver* Receiver = new FRequestAsyncCbPackageReceiver(
 			MoveTemp(Request),
 			&Owner,
 			ZenService.GetInstance(),
 			MoveTemp(OnComplete));
 		Receiver->SendAsync();
-	});
+	}); //-V773
 }
 
 void FZenCacheStore::EnqueueAsyncRpc(IRequestOwner& Owner, const FCbPackage& RequestPackage, FOnRpcComplete&& OnComplete)
@@ -1794,13 +1795,14 @@ void FZenCacheStore::EnqueueAsyncRpc(IRequestOwner& Owner, const FCbPackage& Req
 		}
 		Request->SetContentType(EHttpMediaType::CbPackage);
 		Request->SetBody(SaveRpcPackage(RequestPackage));
+		// Receiver will self-delete
 		FRequestAsyncCbPackageReceiver* Receiver = new FRequestAsyncCbPackageReceiver(
 			MoveTemp(Request),
 			&Owner,
 			ZenService.GetInstance(),
 			MoveTemp(OnComplete));
 		Receiver->SendAsync();
-	});
+	}); //-V773
 }
 
 void FZenCacheStore::ActivatePerformanceEvaluationThread()
diff --git a/Engine/Source/Developer/IoStoreUtilities/Private/ZenStoreHttpClient.cpp b/Engine/Source/Developer/IoStoreUtilities/Private/ZenStoreHttpClient.cpp
--- a/Engine/Source/Developer/IoStoreUtilities/Private/ZenStoreHttpClient.cpp
+++ b/Engine/Source/Developer/IoStoreUtilities/Private/ZenStoreHttpClient.cpp
@@ -370,6 +370,7 @@ void FZenStoreHttpClient::RpcAsync(FAnsiStringView Uri, FCbObject RpcRequestObje
 	Request->AddAcceptType(EHttpMediaType::CbPackage);
 	Request->SetContentType(EHttpMediaType::CbObject);
 	Request->SetBody(RpcRequestObject.GetBuffer().MakeOwned());
+	// PackageReceiver will delete itself in the completion callback
 	Zen::FAsyncCbPackageReceiver* PackageReceiver = new Zen::FAsyncCbPackageReceiver(MoveTemp(Request), ZenService.GetInstance(),
 		[ResponseCallback = MoveTemp(ResponseCallback)]
 		(Zen::FAsyncCbPackageReceiver* Receiver, const THttpUniquePtr<IHttpResponse>& HttpResponse, const FCbPackage& Response) mutable
@@ -383,7 +384,7 @@ void FZenStoreHttpClient::RpcAsync(FAnsiStringView Uri, FCbObject RpcRequestObje
 		},
 		/*MaxAttempts*/ 2);
 	PackageReceiver->SendAsync();
-}
+} //-V773
 
 bool FZenStoreHttpClient::Delete(Zen::FZenHttpRequest& Request, FStringView Uri)
 {
