Additional client-side validation for the zen cache store.

diff --git a/Engine/Source/Developer/DerivedDataCache/Private/ZenCacheStore.cpp b/Engine/Source/Developer/DerivedDataCache/Private/ZenCacheStore.cpp
--- a/Engine/Source/Developer/DerivedDataCache/Private/ZenCacheStore.cpp
+++ b/Engine/Source/Developer/DerivedDataCache/Private/ZenCacheStore.cpp
@@ -204,6 +204,15 @@ private:
 		RequestTypeAccessorType TypeAccessor,
 		const ERequestOp Op);
 
+	[[nodiscard]] bool ValidateAndPrepareReceivedValue(
+		FStringView Name,
+		const FCacheKey& Key,
+		const FValueId& Id,
+		const FValue& Value,
+		ECachePolicy Policy,
+		FValue& OutValue,
+		FRequestStats& Stats) const;
+
 	class FPutOp;
 	class FGetOp;
 	class FPutValueOp;
@@ -294,6 +303,62 @@ void FZenCacheStore::StartRequests(
 	}
 }
 
+bool FZenCacheStore::ValidateAndPrepareReceivedValue(
+		FStringView Name,
+		const FCacheKey& Key,
+		const FValueId& Id,
+		const FValue& Value,
+		ECachePolicy Policy,
+		FValue& OutValue,
+		FRequestStats& Stats) const
+{
+	if (!EnumHasAnyFlags(Policy, ECachePolicy::Query))
+	{
+		// Ensure any data accompanying the value is removed.
+		OutValue = Value.RemoveData();
+		return true;
+	}
+
+	if (EnumHasAnyFlags(Policy, ECachePolicy::SkipData))
+	{
+		// Ensure any data accompanying the value is removed.
+		OutValue = Value.RemoveData();
+		return true;
+	}
+
+	if (!Value.HasData())
+	{
+		UE_LOG(LogDerivedDataCache, Verbose,
+			TEXT("%s: Cache miss with missing value %s with hash %s for %s from '%.*s'"),
+			*GetName(), *WriteToString<16>(Id), *WriteToString<48>(Value.GetRawHash()), *WriteToString<96>(Key),
+			Name.Len(), Name.GetData());
+		return false;
+	}
+
+	if (Value.GetRawHash().IsZero() || Value.GetRawSize() == MAX_uint64)
+	{
+		UE_LOG(LogDerivedDataCache, Display, TEXT("%s: Cache miss with invalid value for %s from '%.*s'"),
+			*GetName(), *WriteToString<96>(Key), Name.Len(), Name.GetData());
+		OutValue = Value.RemoveData();
+		return false;
+	}
+
+	if (Value.GetData().GetRawHash() != Value.GetRawHash())
+	{
+		UE_LOG(LogDerivedDataCache, Display,
+			TEXT("%s: Cache miss with corrupted value %s with hash %s for %s from '%.*s'"),
+			*GetName(), *WriteToString<16>(Id), *WriteToString<48>(Value.GetRawHash()),
+			*WriteToString<96>(Key), Name.Len(), Name.GetData());
+		OutValue = Value.RemoveData();
+		return false;
+	}
+
+	OutValue = Value;
+	Stats.AddLogicalRead(Value);
+	Stats.PhysicalReadSize += Value.GetData().GetCompressedSize();
+	return true;
+}
+
 class FZenCacheStore::FPutOp final : public FThreadSafeRefCountedObject
 {
 public:
@@ -610,7 +675,7 @@ public:
 
 						if (CacheStore.DebugOptions.ShouldSimulateGetMiss(Key))
 						{
-							UE_LOG(LogDerivedDataCache, Verbose, TEXT("%s: Simulated miss for put of '%s' from '%s'"),
+							UE_LOG(LogDerivedDataCache, Verbose, TEXT("%s: Simulated miss for get of '%s' from '%s'"),
 								*CacheStore.GetName(), *WriteToString<96>(Key), *RequestWithStats.Request.Name);
 						}
 						else if (!RecordField.IsNull())
@@ -653,38 +718,59 @@ public:
 private:
 	void OnHit(const TRequestWithStats<FCacheGetRequest>& RequestWithStats, FCacheRecord&& Record)
 	{
-		UE_LOG(LogDerivedDataCache, Verbose, TEXT("%s: Cache hit for '%s' from '%s'"),
-			*CacheStore.GetName(), *WriteToString<96>(RequestWithStats.Request.Key), *RequestWithStats.Request.Name);
+		FCacheRecordBuilder RecordBuilder(RequestWithStats.Request.Key);
 
-		bool bInComplete = false;
-		if (const FCbObject& Meta = Record.GetMeta())
+		const ECachePolicy RecordPolicy = RequestWithStats.Request.Policy.GetRecordPolicy();
+		if (!EnumHasAnyFlags(RecordPolicy, ECachePolicy::SkipMeta))
 		{
-			RequestWithStats.Stats.PhysicalReadSize += Meta.GetSize();
+			if (FCbObject Meta = Record.GetMeta())
+			{
+				RequestWithStats.Stats.PhysicalReadSize += Meta.GetSize();
+				RecordBuilder.SetMeta(MoveTemp(Meta));
+			}
 		}
+
+		EStatus Status = EStatus::Ok;
+		const bool bPartialRecord = EnumHasAnyFlags(RecordPolicy, ECachePolicy::PartialRecord);
 		for (const FValueWithId& Value : Record.GetValues())
 		{
-			RequestWithStats.Stats.AddLogicalRead(Value);
-			RequestWithStats.Stats.PhysicalReadSize += Value.GetData().GetCompressedSize();
-			ECachePolicy ValuePolicy = RequestWithStats.Request.Policy.GetValuePolicy(Value.GetId());
-			if (EnumHasAnyFlags(ValuePolicy, ECachePolicy::SkipData))
+			const FValueId& Id = Value.GetId();
+			const ECachePolicy ValuePolicy = RequestWithStats.Request.Policy.GetValuePolicy(Value.GetId());
+			FValue Content;
+			if (CacheStore.ValidateAndPrepareReceivedValue(RequestWithStats.Request.Name,
+				RequestWithStats.Request.Key,
+				Id,
+				Value,
+				ValuePolicy,
+				Content,
+				RequestWithStats.Stats))
 			{
-				continue;
+				RecordBuilder.AddValue(Id, MoveTemp(Content));
 			}
-			if (Value.HasData())
+			else if (bPartialRecord)
 			{
-				continue;
+				Status = EStatus::Error;
+				RecordBuilder.AddValue(Value);
 			}
-			if (EnumHasAnyFlags(ValuePolicy, ECachePolicy::Query))
+			else
 			{
-				bInComplete = true;
+				// End immediately with error and no record
+				RequestWithStats.EndRequest(CacheStore, EStatus::Error);
+				OnComplete(RequestWithStats.Request.MakeResponse(EStatus::Error));
+				return;
+			}
 		}
+
+		if (Status == EStatus::Ok)
+		{
+			UE_LOG(LogDerivedDataCache, Verbose, TEXT("%s: Cache hit for '%s' from '%s'"),
+				*CacheStore.GetName(), *WriteToString<96>(RequestWithStats.Request.Key), *RequestWithStats.Request.Name);
 		}
-		EStatus Status = bInComplete ? EStatus::Error : EStatus::Ok;
 		RequestWithStats.EndRequest(CacheStore, Status);
 
 		TRACE_COUNTER_INCREMENT(ZenDDC_GetHit);
 		TRACE_COUNTER_ADD(ZenDDC_BytesReceived, int64(RequestWithStats.Stats.PhysicalReadSize));
-		OnComplete({RequestWithStats.Request.Name, MoveTemp(Record), RequestWithStats.Request.UserData, Status});
+		OnComplete({RequestWithStats.Request.Name, RecordBuilder.Build(), RequestWithStats.Request.UserData, Status});
 	}
 
 	void OnMiss(const TRequestWithStats<FCacheGetRequest>& RequestWithStats)
@@ -1071,11 +1157,22 @@ public:
 private:
 	void OnHit(const TRequestWithStats<FCacheGetValueRequest>& RequestWithStats, FValue&& Value)
 	{
+		if (!CacheStore.ValidateAndPrepareReceivedValue(RequestWithStats.Request.Name,
+			RequestWithStats.Request.Key,
+			{},
+			Value,
+			RequestWithStats.Request.Policy,
+			Value,
+			RequestWithStats.Stats))
+		{
+			// End immediately with error and no value
+			RequestWithStats.EndRequest(CacheStore, EStatus::Error);
+			OnComplete(RequestWithStats.Request.MakeResponse(EStatus::Error));
+			return;
+		}
+
 		UE_LOG(LogDerivedDataCache, Verbose, TEXT("%s: Cache hit for '%s' from '%s'"),
 			*CacheStore.GetName(), *WriteToString<96>(RequestWithStats.Request.Key), *RequestWithStats.Request.Name);
-
-		RequestWithStats.Stats.AddLogicalRead(Value);
-		RequestWithStats.Stats.PhysicalReadSize += Value.GetData().GetCompressedSize();
 		RequestWithStats.EndRequest(CacheStore, EStatus::Ok);
 
 		TRACE_COUNTER_INCREMENT(ZenDDC_GetHit);
