Fix gaps in HttpCacheStore DDC client-side validation and ensure that we verify the error code before we try to read status code.  Note that the server is not sending X-Jupiter-IoHash headers for FHttpCacheStore::FGetRecordOp::EndGetValue, but if sent, the client will use it for validation.

diff --git a/Engine/Source/Developer/DerivedDataCache/Private/HttpCacheStore.cpp b/Engine/Source/Developer/DerivedDataCache/Private/HttpCacheStore.cpp
--- a/Engine/Source/Developer/DerivedDataCache/Private/HttpCacheStore.cpp
+++ b/Engine/Source/Developer/DerivedDataCache/Private/HttpCacheStore.cpp
@@ -911,6 +911,17 @@ void FHttpCacheStore::FPutPackageOp::EndPutRef(
 {
 	TRACE_CPUPROFILER_EVENT_SCOPE(HttpDDC_PutPackage_EndPutRef);
 
+	if (const EHttpErrorCode ErrorCode = Operation->GetErrorCode(); ErrorCode != EHttpErrorCode::None)
+	{
+		if (ErrorCode != EHttpErrorCode::Canceled)
+		{
+			UE_LOG(LogDerivedDataCache, Display, TEXT("%s: Cache put failed due to error or retry exhaustion on record for %s from '%s'"),
+				*CacheStore.NodeName, *WriteToString<96>(Key), *Name);
+		}
+		OnComplete({{}, ErrorCode == EHttpErrorCode::Canceled ? EStatus::Canceled : EStatus::Error});
+		return;
+	}
+
 	if (const int32 StatusCode = Operation->GetStatusCode(); StatusCode < 200 || StatusCode > 204)
 	{
 		if (!bAllowOverwrite && (StatusCode == 409))
@@ -918,8 +929,7 @@ void FHttpCacheStore::FPutPackageOp::EndPutRef(
 			UE_LOG(LogDerivedDataCache, Display, TEXT("%s: Cache put found non-deterministic record for %s from '%s'"),
 				*CacheStore.NodeName, *WriteToString<96>(Key), *Name);
 		}
-		const EStatus Status = Operation->GetErrorCode() == EHttpErrorCode::Canceled ? EStatus::Canceled : EStatus::Error;
-		OnComplete({{}, Status});
+		OnComplete({{}, EStatus::Error});
 		return;
 	}
 
@@ -1062,7 +1072,7 @@ void FHttpCacheStore::FPutPackageOp::EndPutBlob(FHttpOperation* Operation, uint6
 	if (Operation)
 	{
 		const int32 StatusCode = Operation->GetStatusCode();
-		if (StatusCode >= 200 && StatusCode <= 204)
+		if (Operation->GetErrorCode() == EHttpErrorCode::None && StatusCode >= 200 && StatusCode <= 204)
 		{
 			SuccessfulBlobUploads.fetch_add(1, std::memory_order_relaxed);
 			TUniqueLock Lock(RequestStats.Mutex.Get());
@@ -1279,6 +1289,17 @@ void FHttpCacheStore::FGetRecordOp::EndGetRef(TUniquePtr<FHttpOperation> Operati
 		return;
 	}
 
+	if (const EHttpErrorCode ErrorCode = Operation->GetErrorCode(); ErrorCode != EHttpErrorCode::None)
+	{
+		if (ErrorCode != EHttpErrorCode::Canceled)
+		{
+			UE_LOG(LogDerivedDataCache, Display, TEXT("%s: Cache miss due to error or retry exhaustion on record for %s from '%s'"),
+				*CacheStore.NodeName, *WriteToString<96>(Key), *Name);
+		}
+		Status = ErrorCode == EHttpErrorCode::Canceled ? EStatus::Canceled : EStatus::Error;
+		return;
+	}
+
 	const int32 StatusCode = Operation->GetStatusCode();
 	if (StatusCode < 200 || StatusCode > 204)
 	{
@@ -1289,6 +1310,20 @@ void FHttpCacheStore::FGetRecordOp::EndGetRef(TUniquePtr<FHttpOperation> Operati
 
 	FSharedBuffer Body = Operation->GetBody();
 
+	if (FAnsiStringView ReceivedHashStr = Operation->GetHeader("X-Jupiter-IoHash"); !ReceivedHashStr.IsEmpty())
+	{
+		FIoHash ReceivedHash(ReceivedHashStr);
+		FIoHash ComputedHash = FIoHash::HashBuffer(Body.GetView());
+		if (ReceivedHash != ComputedHash)
+		{
+			UE_LOG(LogDerivedDataCache, Display,
+				TEXT("%s: Cache miss with corrupted record received hash %s when expected hash %s for %s from '%s'"),
+				*CacheStore.NodeName, *WriteToString<48>(ComputedHash),
+				*WriteToString<48>(ReceivedHash), *WriteToString<96>(Key), *Name);
+			return;
+		}
+	}
+
 	if (ValidateCompactBinary(Body, ECbValidateMode::Default) != ECbValidateError::None)
 	{
 		UE_LOG(LogDerivedDataCache, Log, TEXT("%s: Cache miss with invalid package for %s from '%s'"),
@@ -1507,7 +1542,7 @@ void FHttpCacheStore::FGetRecordOp::EndGetValue(FHttpOperation& Operation, const
 
 	bool bHit = false;
 	FCompressedBuffer CompressedBuffer;
-	if (Operation.GetStatusCode() == 200)
+	if (Operation.GetErrorCode() == EHttpErrorCode::None && Operation.GetStatusCode() == 200)
 	{
 		switch (Operation.GetContentType())
 		{
@@ -1532,6 +1567,21 @@ void FHttpCacheStore::FGetRecordOp::EndGetValue(FHttpOperation& Operation, const
 
 	if (bHit)
 	{
+		if (FAnsiStringView ReceivedHashStr = Operation.GetHeader("X-Jupiter-IoHash"); !ReceivedHashStr.IsEmpty())
+		{
+			FIoHash ReceivedHash(ReceivedHashStr);
+			FIoHash ComputedHash = FIoHash::HashBuffer(Operation.GetBody().GetView());
+			if (ReceivedHash != ComputedHash)
+			{
+				UE_LOG(LogDerivedDataCache, Display,
+					TEXT("%s: Cache miss with corrupted value %s received hash %s when expected hash %s for %s from '%s'"),
+					*CacheStore.NodeName, *WriteToString<32>(Value.GetId()), *WriteToString<48>(ComputedHash),
+					*WriteToString<48>(ReceivedHash), *WriteToString<96>(Key), *Name);
+				OnComplete({ Value, EStatus::Error });
+				return;
+			}
+		}
+
 		if (CompressedBuffer.GetRawHash() == Value.GetRawHash() && CompressedBuffer.GetRawSize() == Value.GetRawSize())
 		{
 			OnComplete({FValueWithId(Value.GetId(), MoveTemp(CompressedBuffer)), EStatus::Ok});
@@ -1551,10 +1601,18 @@ void FHttpCacheStore::FGetRecordOp::EndGetValue(FHttpOperation& Operation, const
 	}
 	else
 	{
-		UE_LOG(LogDerivedDataCache, Verbose,
-			TEXT("%s: Cache miss with missing value %s with hash %s for %s from '%s'"),
-			*CacheStore.NodeName, *WriteToString<32>(Value.GetId()), *WriteToString<48>(Value.GetRawHash()),
-			*WriteToString<96>(Key), *Name);
+		if (Operation.GetErrorCode() != EHttpErrorCode::None)
+		{
+			UE_LOG(LogDerivedDataCache, Display, TEXT("%s: Cache miss due to error or retry exhaustion on value %s for %s from '%s'"),
+				*CacheStore.NodeName, *WriteToString<32>(Value.GetId()), *WriteToString<96>(Key), *Name);
+		}
+		else
+		{
+			UE_LOG(LogDerivedDataCache, Verbose,
+				TEXT("%s: Cache miss with missing value %s with hash %s for %s from '%s'"),
+				*CacheStore.NodeName, *WriteToString<32>(Value.GetId()), *WriteToString<48>(Value.GetRawHash()),
+				*WriteToString<96>(Key), *Name);
+		}
 		OnComplete({Value, EStatus::Error});
 	}
 }
@@ -1638,6 +1696,11 @@ void FHttpCacheStore::FGetRecordOp::EndGetValuesExist(FHttpOperation* Operation,
 		DefaultMessage = TEXT("Cache exists miss with canceled request for");
 		DefaultStatus = EStatus::Canceled;
 	}
+	else if (Operation->GetErrorCode() != EHttpErrorCode::None)
+	{
+		DefaultMessage = TEXT("Cache exists miss with error or retry exhaustion during request for");
+		DefaultStatus = EStatus::Error;
+	}
 	else if (const int32 StatusCode = Operation->GetStatusCode(); StatusCode < 200 || StatusCode > 204)
 	{
 		DefaultMessage = TEXT("Cache exists miss with failed response for");
@@ -1803,6 +1866,16 @@ void FHttpCacheStore::FGetValueOp::EndGetRef(FHttpOperation& Operation)
 
 	const bool bSkipData = EnumHasAnyFlags(Policy, ECachePolicy::SkipData);
 
+	if (const EHttpErrorCode ErrorCode = Operation.GetErrorCode(); ErrorCode != EHttpErrorCode::None)
+	{
+		if (ErrorCode != EHttpErrorCode::Canceled)
+		{
+			UE_LOG(LogDerivedDataCache, Display, TEXT("%s: Cache miss due to error or retry exhaustion for %s from '%s'"),
+				*CacheStore.NodeName, *WriteToString<96>(Key), *Name);
+		}
+		return EndGet({Name, Key, {}, ErrorCode == EHttpErrorCode::Canceled ? EStatus::Canceled : EStatus::Error});
+	}
+
 	const int32 StatusCode = Operation.GetStatusCode();
 	if (StatusCode < 200 || StatusCode > 204)
 	{
@@ -1815,6 +1888,22 @@ void FHttpCacheStore::FGetValueOp::EndGetRef(FHttpOperation& Operation)
 
 	if (bSkipData)
 	{
+		if (FAnsiStringView ReceivedHashStr = Operation.GetHeader("X-Jupiter-IoHash"); !ReceivedHashStr.IsEmpty())
+		{
+			FRequestTimer RequestTimer(RequestStats);
+			FIoHash ReceivedHash(ReceivedHashStr);
+			FIoHash ComputedHash = FIoHash::HashBuffer(Body.GetView());
+			if (ReceivedHash != ComputedHash)
+			{
+				RequestTimer.Stop();
+				UE_LOG(LogDerivedDataCache, Display,
+					TEXT("%s: Cache miss with corrupted value reference received hash %s when expected hash %s for %s from '%s'"),
+					*CacheStore.NodeName, *WriteToString<48>(ComputedHash),
+					*WriteToString<48>(ReceivedHash), *WriteToString<96>(Key), *Name);
+				return EndGet({ Name, Key, {}, EStatus::Error });
+			}
+		}
+
 		if (ValidateCompactBinary(Body, ECbValidateMode::Default) != ECbValidateError::None)
 		{
 			UE_LOG(LogDerivedDataCache, Display, TEXT("%s: Cache miss with invalid package for %s from '%s'"),
@@ -1838,20 +1927,27 @@ void FHttpCacheStore::FGetValueOp::EndGetRef(FHttpOperation& Operation)
 	{
 		FCompressedBuffer CompressedBuffer = FCompressedBuffer::FromCompressed(Body);
 
-		if (!CompressedBuffer)
+		if (FAnsiStringView ReceivedHashStr = Operation.GetHeader("X-Jupiter-InlinePayloadHash"); !ReceivedHashStr.IsEmpty())
 		{
 			FRequestTimer RequestTimer(RequestStats);
-			if (FAnsiStringView ReceivedHashStr = Operation.GetHeader("X-Jupiter-InlinePayloadHash"); !ReceivedHashStr.IsEmpty())
+			FIoHash ReceivedHash(ReceivedHashStr);
+			FIoHash ComputedHash = FIoHash::HashBuffer(Body.GetView());
+			if (ReceivedHash != ComputedHash)
 			{
-				FIoHash ReceivedHash(ReceivedHashStr);
-				FIoHash ComputedHash = FIoHash::HashBuffer(Body.GetView());
-				if (ReceivedHash == ComputedHash)
-				{
-					CompressedBuffer = FCompressedBuffer::Compress(Body);
-				}
+				RequestTimer.Stop();
+				UE_LOG(LogDerivedDataCache, Display,
+					TEXT("%s: Cache miss with corrupted value received hash %s when expected hash %s for %s from '%s'"),
+					*CacheStore.NodeName, *WriteToString<48>(ComputedHash),
+					*WriteToString<48>(ReceivedHash), *WriteToString<96>(Key), *Name);
+				return EndGet({Name, Key, {}, EStatus::Error});
 			}
 		}
 
+		if (!CompressedBuffer)
+		{
+			CompressedBuffer = FCompressedBuffer::Compress(Body);
+		}
+
 		if (!CompressedBuffer)
 		{
 			UE_LOG(LogDerivedDataCache, Display, TEXT("%s: Cache miss with invalid package for %s from '%s'"),
@@ -2035,13 +2131,29 @@ void FHttpCacheStore::FExistsBatchOp::EndExists(FHttpOperation& Operation)
 	RequestStats.Type = ERequestType::Value;
 	RequestStats.Op = ERequestOp::Get;
 
+	if (const EHttpErrorCode ErrorCode = Operation.GetErrorCode(); ErrorCode != EHttpErrorCode::None)
+	{
+		RequestTimer.Stop();
+		for (const FCacheGetValueRequest& Request : Requests)
+		{
+			if (ErrorCode != EHttpErrorCode::Canceled)
+			{
+				UE_LOG(LogDerivedDataCache, Display, TEXT("%s: Cache miss with failed HTTP request due to error or retry exhaustion on record for %s from '%s'"),
+					*CacheStore.NodeName, *WriteToString<96>(Request.Key), *Request.Name);
+			}
+			RequestStats.Bucket = Request.Key.Bucket;
+			EndRequest(Request, {}, EStatus::Error);
+		}
+		return;
+	}
+
 	const int32 OverallStatusCode = Operation.GetStatusCode();
 	if (OverallStatusCode < 200 || OverallStatusCode > 204)
 	{
 		RequestTimer.Stop();
 		for (const FCacheGetValueRequest& Request : Requests)
 		{
-			UE_LOG(LogDerivedDataCache, Verbose, TEXT("%s: Cache miss with failed HTTP request for %s from '%s'"),
+			UE_LOG(LogDerivedDataCache, Verbose, TEXT("%s: Cache miss with failed non-success status code for %s from '%s'"),
 				*CacheStore.NodeName, *WriteToString<96>(Request.Key), *Request.Name);
 			RequestStats.Bucket = Request.Key.Bucket;
 			EndRequest(Request, {}, EStatus::Error);
