Add a -DDC-AllowNoActiveStores commandline argument for situations where we want to allow continued execution when there are no available DDC cache stores.

diff --git a/Engine/Source/Developer/DerivedDataCache/Private/DerivedDataBackends.cpp b/Engine/Source/Developer/DerivedDataCache/Private/DerivedDataBackends.cpp
--- a/Engine/Source/Developer/DerivedDataCache/Private/DerivedDataBackends.cpp
+++ b/Engine/Source/Developer/DerivedDataCache/Private/DerivedDataBackends.cpp
@@ -205,7 +205,8 @@ public:
 					}
 					else
 					{
-						UE_LOG(LogDerivedDataCache, Fatal,
+						UE_CLOG(!FParse::Param(FCommandLine::Get(), TEXT("DDC-AllowNoActiveStores")),
+							LogDerivedDataCache, Fatal,
 							TEXT("Unable to use default cache graph '%s' because there are no %s nodes available."
 								 "Add -DDC-ForceMemoryCache to the command line to bypass this if you need access "
 								 "to the editor settings to fix the cache configuration."),
