DDC: Added value hashes and sizes to replays

diff --git a/Engine/Source/Developer/DerivedDataCache/Private/DerivedDataCache.cpp b/Engine/Source/Developer/DerivedDataCache/Private/DerivedDataCache.cpp
--- a/Engine/Source/Developer/DerivedDataCache/Private/DerivedDataCache.cpp
+++ b/Engine/Source/Developer/DerivedDataCache/Private/DerivedDataCache.cpp
@@ -257,6 +257,13 @@ FCacheGetChunkResponse FCacheGetChunkRequest::MakeResponse(const EStatus Status)
 FCbWriter& operator<<(FCbWriter& Writer, const FCacheGetRequest& Request)
 {
 	Writer.BeginObject();
+	Private::SaveToCompactBinary(Writer, Request);
+	Writer.EndObject();
+	return Writer;
+}
+
+void Private::SaveToCompactBinary(FCbWriter& Writer, const FCacheGetRequest& Request)
+{
 	if (!Request.Name.IsEmpty())
 	{
 		Writer << ANSITEXTVIEW("Name") << Request.Name;
@@ -270,8 +277,6 @@ FCbWriter& operator<<(FCbWriter& Writer, const FCacheGetRequest& Request)
 	{
 		Writer << ANSITEXTVIEW("UserData") << Request.UserData;
 	}
-	Writer.EndObject();
-	return Writer;
 }
 
 bool LoadFromCompactBinary(FCbFieldView Field, FCacheGetRequest& Request)
@@ -287,6 +292,13 @@ bool LoadFromCompactBinary(FCbFieldView Field, FCacheGetRequest& Request)
 FCbWriter& operator<<(FCbWriter& Writer, const FCacheGetValueRequest& Request)
 {
 	Writer.BeginObject();
+	Private::SaveToCompactBinary(Writer, Request);
+	Writer.EndObject();
+	return Writer;
+}
+
+void Private::SaveToCompactBinary(FCbWriter& Writer, const FCacheGetValueRequest& Request)
+{
 	if (!Request.Name.IsEmpty())
 	{
 		Writer << ANSITEXTVIEW("Name") << MakeStringView(Request.Name);
@@ -300,8 +312,6 @@ FCbWriter& operator<<(FCbWriter& Writer, const FCacheGetValueRequest& Request)
 	{
 		Writer << ANSITEXTVIEW("UserData") << Request.UserData;
 	}
-	Writer.EndObject();
-	return Writer;
 }
 
 bool LoadFromCompactBinary(FCbFieldView Field, FCacheGetValueRequest& Request)
@@ -317,6 +327,13 @@ bool LoadFromCompactBinary(FCbFieldView Field, FCacheGetValueRequest& Request)
 FCbWriter& operator<<(FCbWriter& Writer, const FCacheGetChunkRequest& Request)
 {
 	Writer.BeginObject();
+	Private::SaveToCompactBinary(Writer, Request);
+	Writer.EndObject();
+	return Writer;
+}
+
+void Private::SaveToCompactBinary(FCbWriter& Writer, const FCacheGetChunkRequest& Request)
+{
 	if (!Request.Name.IsEmpty())
 	{
 		Writer << ANSITEXTVIEW("Name") << MakeStringView(Request.Name);
@@ -346,8 +363,6 @@ FCbWriter& operator<<(FCbWriter& Writer, const FCacheGetChunkRequest& Request)
 	{
 		Writer << ANSITEXTVIEW("UserData") << Request.UserData;
 	}
-	Writer.EndObject();
-	return Writer;
 }
 
 bool LoadFromCompactBinary(FCbFieldView Field, FCacheGetChunkRequest& OutRequest)
diff --git a/Engine/Source/Developer/DerivedDataCache/Private/DerivedDataCachePrivate.h b/Engine/Source/Developer/DerivedDataCache/Private/DerivedDataCachePrivate.h
--- a/Engine/Source/Developer/DerivedDataCache/Private/DerivedDataCachePrivate.h
+++ b/Engine/Source/Developer/DerivedDataCache/Private/DerivedDataCachePrivate.h
@@ -4,6 +4,7 @@
 
 #include "CoreTypes.h"
 
+class FCbWriter;
 class FDerivedDataCacheInterface;
 
 template <typename FuncType> class TUniqueFunction;
@@ -11,6 +12,9 @@ template <typename FuncType> class TUniqueFunction;
 namespace UE::DerivedData { class FCacheRecord; }
 namespace UE::DerivedData { class ICache; }
 namespace UE::DerivedData { class IRequestOwner; }
+namespace UE::DerivedData { struct FCacheGetChunkRequest; }
+namespace UE::DerivedData { struct FCacheGetRequest; }
+namespace UE::DerivedData { struct FCacheGetValueRequest; }
 
 namespace UE::DerivedData::Private
 {
@@ -21,6 +25,9 @@ int32 AddToAsyncTaskCounter(int32 Addend); // returns the previous value
 // Implemented in DerivedDataCache.cpp
 ICache* CreateCache(FDerivedDataCacheInterface** OutLegacyCache);
 void LaunchTaskInCacheThreadPool(IRequestOwner& Owner, TUniqueFunction<void ()>&& TaskBody);
+void SaveToCompactBinary(FCbWriter& Writer, const FCacheGetRequest& Request);
+void SaveToCompactBinary(FCbWriter& Writer, const FCacheGetValueRequest& Request);
+void SaveToCompactBinary(FCbWriter& Writer, const FCacheGetChunkRequest& Request);
 
 // Implemented in DerivedDataCacheRecord.cpp
 uint64 GetCacheRecordCompressedSize(const FCacheRecord& Record);
diff --git a/Engine/Source/Developer/DerivedDataCache/Private/DerivedDataCacheReplay.cpp b/Engine/Source/Developer/DerivedDataCache/Private/DerivedDataCacheReplay.cpp
--- a/Engine/Source/Developer/DerivedDataCache/Private/DerivedDataCacheReplay.cpp
+++ b/Engine/Source/Developer/DerivedDataCache/Private/DerivedDataCacheReplay.cpp
@@ -50,23 +50,23 @@ public:
 	~FCacheStoreReplay();
 
 	void Put(
-		const TConstArrayView<FCachePutRequest> Requests,
+		TConstArrayView<FCachePutRequest> Requests,
 		IRequestOwner& Owner,
 		FOnCachePutComplete&& OnComplete) final;
 	void Get(
-		const TConstArrayView<FCacheGetRequest> Requests,
+		TConstArrayView<FCacheGetRequest> Requests,
 		IRequestOwner& Owner,
 		FOnCacheGetComplete&& OnComplete) final;
 	void PutValue(
-		const TConstArrayView<FCachePutValueRequest> Requests,
+		TConstArrayView<FCachePutValueRequest> Requests,
 		IRequestOwner& Owner,
 		FOnCachePutValueComplete&& OnComplete) final;
 	void GetValue(
-		const TConstArrayView<FCacheGetValueRequest> Requests,
+		TConstArrayView<FCacheGetValueRequest> Requests,
 		IRequestOwner& Owner,
 		FOnCacheGetValueComplete&& OnComplete) final;
 	void GetChunks(
-		const TConstArrayView<FCacheGetChunkRequest> Requests,
+		TConstArrayView<FCacheGetChunkRequest> Requests,
 		IRequestOwner& Owner,
 		FOnCacheGetChunkComplete&& OnComplete) final;
 
@@ -85,9 +85,22 @@ public:
 		ReplayReader = MoveTemp(Reader);
 	}
 
-private:
+	ICacheStore* GetInnerCache() const
+	{
+		return InnerCache;
+	}
+
 	template <typename RequestType>
-	void SerializeRequests(TConstArrayView<RequestType> Requests, ECacheMethod Method, EPriority Priority);
+	void SerializeRequests(
+		TConstArrayView<RequestType> Requests,
+		TConstArrayView<TCacheResponseFor<RequestType>> Responses,
+		ECacheMethod Method,
+		EPriority Priority);
+
+private:
+	static void SaveValuesToCompactBinary(FCbWriter& Writer, const FCacheGetResponse& Response);
+	static void SaveValuesToCompactBinary(FCbWriter& Writer, const FCacheGetValueResponse& Response);
+	static void SaveValuesToCompactBinary(FCbWriter& Writer, const FCacheGetChunkResponse& Response);
 
 	void WriteBinaryToArchive(const FCompositeBuffer& RawBinary);
 	void WriteToArchive(FCbWriter& Writer);
@@ -118,6 +131,61 @@ private:
 #endif // WITH_EDITOR
 };
 
+///////////////////////////////////////////////////////////////////////////////////////////////////////////////////////
+
+template <typename RequestType>
+struct TCacheStoreReplayBatch
+{
+	using FRequest = RequestType;
+	using FResponse = TCacheResponseFor<FRequest>;
+	using FOnComplete = TCacheOnCompleteFor<FRequest>;
+
+	FCacheStoreReplay* Replay;
+	const TArray<FRequest> Requests;
+	TArray<FResponse> Responses;
+	FOnComplete OnComplete;
+	std::atomic<int32> ActiveRequests;
+	EPriority Priority;
+
+	explicit TCacheStoreReplayBatch(
+		FCacheStoreReplay* InReplay,
+		TConstArrayView<FRequest> InRequests,
+		IRequestOwner& Owner,
+		FOnComplete&& InOnComplete)
+		: Replay(InReplay)
+		, Requests(InRequests)
+		, OnComplete(MoveTemp(InOnComplete))
+		, ActiveRequests(Requests.Num())
+		, Priority(Owner.GetPriority())
+	{
+		Responses.SetNumUninitialized(ActiveRequests);
+
+		TArray<FRequest, TInlineAllocator<16>> ModifiedRequests(Requests);
+		int32 Index = 0;
+		for (FRequest& Request : ModifiedRequests)
+		{
+			Request.UserData = Index++;
+		}
+
+		(Replay->GetInnerCache()->*CacheFunctionFor<FRequest>)(ModifiedRequests, Owner, [this](FResponse&& Response)
+		{
+			const int32 Index = int32(Response.UserData);
+			Response.UserData = Requests[Index].UserData;
+			new(&Responses[Index]) FResponse(Response);
+
+			OnComplete(MoveTemp(Response));
+
+			if (ActiveRequests.fetch_sub(1, std::memory_order_acq_rel) == 1)
+			{
+				Replay->SerializeRequests<FRequest>(Requests, Responses, CacheMethodFor<FRequest>, Priority);
+				delete this;
+			}
+		});
+	}
+};
+
+///////////////////////////////////////////////////////////////////////////////////////////////////////////////////////
+
 FCacheStoreReplay::FCacheStoreReplay(
 	ILegacyCacheStore* InInnerCache,
 	FCacheKeyFilter InKeyFilter,
@@ -170,6 +238,7 @@ FCacheStoreReplay::~FCacheStoreReplay()
 template <typename RequestType>
 void FCacheStoreReplay::SerializeRequests(
 	const TConstArrayView<RequestType> Requests,
+	const TConstArrayView<TCacheResponseFor<RequestType>> Responses,
 	const ECacheMethod Method,
 	const EPriority Priority)
 {
@@ -191,11 +260,16 @@ void FCacheStoreReplay::SerializeRequests(
 	Writer << ANSITEXTVIEW("Method") << Method;
 	Writer << ANSITEXTVIEW("Priority") << Priority;
 	Writer.BeginArray(ANSITEXTVIEW("Requests"));
+	int32 Index = -1;
 	for (const RequestType& Request : Requests)
 	{
+		++Index;
 		if (KeyFilter.IsMatch(Request.Key))
 		{
-			Writer << Request;
+			Writer.BeginObject();
+			Private::SaveToCompactBinary(Writer, Request);
+			SaveValuesToCompactBinary(Writer, Responses[Index]);
+			Writer.EndObject();
 		}
 	}
 	Writer.EndArray();
@@ -211,6 +285,45 @@ void FCacheStoreReplay::SerializeRequests(
 	WriteToArchive(Writer);
 }
 
+void FCacheStoreReplay::SaveValuesToCompactBinary(FCbWriter& Writer, const FCacheGetResponse& Response)
+{
+	if (Response.Status == EStatus::Ok)
+	{
+		Writer.BeginArray(ANSITEXTVIEW("Values"));
+		for (const FValueWithId& Value : Response.Record.GetValues())
+		{
+			Writer.BeginObject();
+			Writer << ANSITEXTVIEW("Id") << Value.GetId();
+			Writer << ANSITEXTVIEW("RawHash") << Value.GetRawHash();
+			Writer << ANSITEXTVIEW("RawSize") << Value.GetRawSize();
+			Writer.EndObject();
+		}
+		Writer.EndArray();
+	}
+}
+
+void FCacheStoreReplay::SaveValuesToCompactBinary(FCbWriter& Writer, const FCacheGetValueResponse& Response)
+{
+	if (Response.Status == EStatus::Ok)
+	{
+		Writer.BeginObject(ANSITEXTVIEW("Value"));
+		Writer << ANSITEXTVIEW("RawHash") << Response.Value.GetRawHash();
+		Writer << ANSITEXTVIEW("RawSize") << Response.Value.GetRawSize();
+		Writer.EndObject();
+	}
+}
+
+void FCacheStoreReplay::SaveValuesToCompactBinary(FCbWriter& Writer, const FCacheGetChunkResponse& Response)
+{
+	if (Response.Status == EStatus::Ok)
+	{
+		Writer.BeginObject(ANSITEXTVIEW("Value"));
+		Writer << ANSITEXTVIEW("RawHash") << Response.RawHash;
+		Writer << ANSITEXTVIEW("RawSize") << Response.RawSize;
+		Writer.EndObject();
+	}
+}
+
 void FCacheStoreReplay::WriteBinaryToArchive(const FCompositeBuffer& RawBinary)
 {
 	const FValue CompressedBinary = FValue::Compress(RawBinary, RawBlock.GetSize());
@@ -268,8 +381,7 @@ void FCacheStoreReplay::Get(
 	IRequestOwner& Owner,
 	FOnCacheGetComplete&& OnComplete)
 {
-	SerializeRequests(Requests, ECacheMethod::Get, Owner.GetPriority());
-	InnerCache->Get(Requests, Owner, MoveTemp(OnComplete));
+	new TCacheStoreReplayBatch(this, Requests, Owner, MoveTemp(OnComplete));
 }
 
 void FCacheStoreReplay::PutValue(
@@ -285,8 +397,7 @@ void FCacheStoreReplay::GetValue(
 	IRequestOwner& Owner,
 	FOnCacheGetValueComplete&& OnComplete)
 {
-	SerializeRequests(Requests, ECacheMethod::GetValue, Owner.GetPriority());
-	InnerCache->GetValue(Requests, Owner, MoveTemp(OnComplete));
+	new TCacheStoreReplayBatch(this, Requests, Owner, MoveTemp(OnComplete));
 }
 
 void FCacheStoreReplay::GetChunks(
@@ -294,8 +405,7 @@ void FCacheStoreReplay::GetChunks(
 	IRequestOwner& Owner,
 	FOnCacheGetChunkComplete&& OnComplete)
 {
-	SerializeRequests(Requests, ECacheMethod::GetChunks, Owner.GetPriority());
-	InnerCache->GetChunks(Requests, Owner, MoveTemp(OnComplete));
+	new TCacheStoreReplayBatch(this, Requests, Owner, MoveTemp(OnComplete));
 }
 
 #if WITH_EDITOR
diff --git a/Engine/Source/Developer/DerivedDataCache/Private/DerivedDataCacheStore.h b/Engine/Source/Developer/DerivedDataCache/Private/DerivedDataCacheStore.h
--- a/Engine/Source/Developer/DerivedDataCache/Private/DerivedDataCacheStore.h
+++ b/Engine/Source/Developer/DerivedDataCache/Private/DerivedDataCacheStore.h
@@ -5,6 +5,7 @@
 #include "Async/Mutex.h"
 #include "Containers/ContainersFwd.h"
 #include "DerivedDataCache.h"
+#include "DerivedDataCacheMethod.h"
 #include "Misc/MonotonicTime.h"
 #include "Templates/DontCopy.h"
 
@@ -325,4 +326,74 @@ inline void CompleteWithStatus(RequestRangeType&& Requests, OnCompleteType&& OnC
 	}
 }
 
+namespace Private
+{
+
+template <typename RequestType>
+struct TCacheRequestTraits;
+
+template <>
+struct TCacheRequestTraits<FCachePutRequest>
+{
+	using FRequest = FCachePutRequest;
+	using FResponse = FCachePutResponse;
+	using FOnComplete = FOnCachePutComplete;
+	constexpr static ECacheMethod Method = ECacheMethod::Put;
+	constexpr static void (ICacheStore::*Function)(TConstArrayView<FRequest>, IRequestOwner&, FOnComplete&&) = &ICacheStore::Put;
+};
+
+template <>
+struct TCacheRequestTraits<FCacheGetRequest>
+{
+	using FRequest = FCacheGetRequest;
+	using FResponse = FCacheGetResponse;
+	using FOnComplete = FOnCacheGetComplete;
+	constexpr static ECacheMethod Method = ECacheMethod::Get;
+	constexpr static void (ICacheStore::*Function)(TConstArrayView<FRequest>, IRequestOwner&, FOnComplete&&) = &ICacheStore::Get;
+};
+
+template <>
+struct TCacheRequestTraits<FCachePutValueRequest>
+{
+	using FRequest = FCachePutValueRequest;
+	using FResponse = FCachePutValueResponse;
+	using FOnComplete = FOnCachePutValueComplete;
+	constexpr static ECacheMethod Method = ECacheMethod::PutValue;
+	constexpr static void (ICacheStore::*Function)(TConstArrayView<FRequest>, IRequestOwner&, FOnComplete&&) = &ICacheStore::PutValue;
+};
+
+template <>
+struct TCacheRequestTraits<FCacheGetValueRequest>
+{
+	using FRequest = FCacheGetValueRequest;
+	using FResponse = FCacheGetValueResponse;
+	using FOnComplete = FOnCacheGetValueComplete;
+	constexpr static ECacheMethod Method = ECacheMethod::GetValue;
+	constexpr static void (ICacheStore::*Function)(TConstArrayView<FRequest>, IRequestOwner&, FOnComplete&&) = &ICacheStore::GetValue;
+};
+
+template <>
+struct TCacheRequestTraits<FCacheGetChunkRequest>
+{
+	using FRequest = FCacheGetChunkRequest;
+	using FResponse = FCacheGetChunkResponse;
+	using FOnComplete = FOnCacheGetChunkComplete;
+	constexpr static ECacheMethod Method = ECacheMethod::GetChunks;
+	constexpr static void (ICacheStore::*Function)(TConstArrayView<FRequest>, IRequestOwner&, FOnComplete&&) = &ICacheStore::GetChunks;
+};
+
+} // Private
+
+/** Alias for the response type corresponding to the request type. */
+template <typename RequestType> using TCacheResponseFor = typename Private::TCacheRequestTraits<RequestType>::FResponse;
+
+/** Alias for the completion callback type corresponding to the request type. */
+template <typename RequestType> using TCacheOnCompleteFor = typename Private::TCacheRequestTraits<RequestType>::FOnComplete;
+
+/** Alias for the cache method corresponding to the request type. */
+template <typename RequestType> inline constexpr auto CacheMethodFor = Private::TCacheRequestTraits<RequestType>::Method;
+
+/** Alias for the cache function corresponding to the request type. */
+template <typename RequestType> inline constexpr auto CacheFunctionFor = Private::TCacheRequestTraits<RequestType>::Function;
+
 } // UE::DerivedData
