Adds Zen remote Read/Write/Latency speed information to Cook Stats.
#rnx

diff --git a/Engine/Source/Developer/DerivedDataCache/Private/DerivedDataCache.cpp b/Engine/Source/Developer/DerivedDataCache/Private/DerivedDataCache.cpp
--- a/Engine/Source/Developer/DerivedDataCache/Private/DerivedDataCache.cpp
+++ b/Engine/Source/Developer/DerivedDataCache/Private/DerivedDataCache.cpp
@@ -201,6 +201,9 @@ namespace UE::DerivedData::CookStats
 				TEXT("ZenRemoteGetHits"), ZenRemoteGetHits,
 				TEXT("ZenRemoteGetTotal"), ZenRemoteGetTotal,
 				TEXT("ZenRemoteGetHitPct"), SafeDivide(ZenRemoteGetHits, ZenRemoteGetTotal),
+				TEXT("ZenRemoteReadSpeed"), ZenRemoteSpeedStats.ReadSpeedMBs,
+				TEXT("ZenRemoteWriteSpeed"), ZenRemoteSpeedStats.WriteSpeedMBs,
+				TEXT("ZenRemoteLatency"), ZenRemoteSpeedStats.LatencyMS,
 				TEXT("CloudGetHits"), CloudGetHits,
 				TEXT("CloudGetTotal"), CloudGetTotal,
 				TEXT("CloudGetHitPct"), SafeDivide(CloudGetHits, CloudGetTotal),
