DDC: Fixed a crash in FTimeAveragedStat and added tests

- Fixed a crash when a value is added with StartTime == EndTime and StartTime is the start of an existing active time range.
- Fixed comparison of EndTime to consistently treat time ranges as [StartTime, EndTime).
- Invalidate cached averages whenever a new value is added.

diff --git a/Engine/Source/Developer/DerivedDataCache/Private/DerivedDataCacheStats.cpp b/Engine/Source/Developer/DerivedDataCache/Private/DerivedDataCacheStats.cpp
--- a/Engine/Source/Developer/DerivedDataCache/Private/DerivedDataCacheStats.cpp
+++ b/Engine/Source/Developer/DerivedDataCache/Private/DerivedDataCacheStats.cpp
@@ -53,21 +53,25 @@ void FTimeAveragedStat::Add(FMonotonicTimePoint StartTime, FMonotonicTimePoint E
 
 	if (EndTime <= LastTime - Period)
 	{
+		// Ignore because the value ended before the active time window.
 		return;
 	}
 	else
 	{
+		// Queue to subtract when EndTime leaves the active time window.
 		const int32 Index = Algo::LowerBoundBy(EndValues, EndTime, &FValue::Time);
 		EndValues.Insert({EndTime, Value}, Index);
 	}
 
 	if (StartTime < LastTime)
 	{
+		// Accumulate immediately because the value is within the active time window.
 		AccumulatedValue += Value;
 		++AccumulatedValueCount;
 	}
 	else
 	{
+		// Queue to add when StartTime enters the active time window.
 		const int32 Index = Algo::LowerBoundBy(StartValues, StartTime, &FValue::Time);
 		StartValues.Insert({StartTime, Value}, Index);
 	}
@@ -76,6 +80,7 @@ void FTimeAveragedStat::Add(FMonotonicTimePoint StartTime, FMonotonicTimePoint E
 	const int32 StartIndex = Algo::LowerBoundBy(ActiveRanges, StartTime, &FRange::EndTime);
 	if (StartIndex == Count)
 	{
+		// StartTime is later than existing ranges. Add a new range.
 		ActiveRanges.Add({StartTime, EndTime});
 	}
 	else
@@ -83,19 +88,20 @@ void FTimeAveragedStat::Add(FMonotonicTimePoint StartTime, FMonotonicTimePoint E
 		FRange& StartRange = ActiveRanges[StartIndex];
 		if (EndTime < StartRange.StartTime)
 		{
+			// EndTime is earlier than existing ranges. Add a new range.
 			ActiveRanges.Insert({StartTime, EndTime}, StartIndex);
 		}
-		else
+		else if (StartTime != EndTime)
 		{
-			if (StartTime < StartRange.StartTime)
-			{
-				StartRange.StartTime = StartTime;
-			}
+			// Extend the existing range earlier to include StartTime.
+			StartRange.StartTime = FMath::Min(StartTime, StartRange.StartTime);
 
-			const int32 EndIndex = Algo::LowerBoundBy(ActiveRanges, EndTime, &FRange::StartTime);
-			StartRange.EndTime = FMath::Max(EndTime, ActiveRanges[EndIndex - 1].EndTime);
+			// Extend the existing range later to include EndTime.
+			const int32 EndIndex = Algo::LowerBoundBy(ActiveRanges, EndTime, &FRange::StartTime) - 1;
+			StartRange.EndTime = FMath::Max(EndTime, ActiveRanges[EndIndex].EndTime);
 
-			ActiveRanges.RemoveAt(StartIndex + 1, EndIndex - StartIndex - 1, EAllowShrinking::No);
+			// Remove any ranges that have been subsumed by this range.
+			ActiveRanges.RemoveAt(StartIndex + 1, EndIndex - StartIndex, EAllowShrinking::No);
 		}
 	}
 
@@ -104,14 +110,18 @@ void FTimeAveragedStat::Add(FMonotonicTimePoint StartTime, FMonotonicTimePoint E
 	{
 		Update(EndTime);
 	}
+	else
+	{
+		bUpdatePending = true;
+	}
 }
 
 double FTimeAveragedStat::GetRate(FMonotonicTimePoint Time)
 {
 	TUniqueLock Lock(Mutex);
-	if (LastTime < Time)
+	if (LastTime < Time || bUpdatePending)
 	{
-		Update(Time);
+		Update(FMath::Max(LastTime, Time));
 	}
 	return AverageRate;
 }
@@ -119,15 +129,16 @@ double FTimeAveragedStat::GetRate(FMonotonicTimePoint Time)
 double FTimeAveragedStat::GetValue(FMonotonicTimePoint Time)
 {
 	TUniqueLock Lock(Mutex);
-	if (LastTime < Time)
+	if (LastTime < Time || bUpdatePending)
 	{
-		Update(Time);
+		Update(FMath::Max(LastTime, Time));
 	}
 	return AverageValue;
 }
 
 void FTimeAveragedStat::Update(FMonotonicTimePoint Time)
 {
+	bUpdatePending = false;
 	LastTime = Time;
 
 	const FMonotonicTimePoint StartTime = Time - Period;
@@ -155,7 +166,7 @@ void FTimeAveragedStat::Update(FMonotonicTimePoint Time)
 	int32 EndCount = 0;
 	for (const FValue& Value : EndValues)
 	{
-		if (StartTime <= Value.Time)
+		if (StartTime < Value.Time)
 		{
 			break;
 		}
diff --git a/Engine/Source/Developer/DerivedDataCache/Private/DerivedDataCacheStats.h b/Engine/Source/Developer/DerivedDataCache/Private/DerivedDataCacheStats.h
--- a/Engine/Source/Developer/DerivedDataCache/Private/DerivedDataCacheStats.h
+++ b/Engine/Source/Developer/DerivedDataCache/Private/DerivedDataCacheStats.h
@@ -97,6 +97,8 @@ private:
 	double AccumulatedValue = 0.0;
 	/** Number of values within the period at the most recent update time. */
 	int32 AccumulatedValueCount = 0;
+	/** Whether values have been added since the last update. */
+	bool bUpdatePending = false;
 	/** Mutex that can be used to synchronize access to the stat from multiple threads. */
 	FMutex Mutex;
 };
diff --git a/Engine/Source/Developer/DerivedDataCache/Tests/DerivedDataCacheStatsTest.cpp b/Engine/Source/Developer/DerivedDataCache/Tests/DerivedDataCacheStatsTest.cpp
new file mode 100644
--- /dev/null
+++ b/Engine/Source/Developer/DerivedDataCache/Tests/DerivedDataCacheStatsTest.cpp
@@ -0,0 +1,106 @@
+// Copyright Epic Games, Inc. All Rights Reserved.
+
+#include "DerivedDataCacheStats.h"
+
+#if WITH_TESTS
+
+#include "Tests/TestHarnessAdapter.h"
+
+namespace UE::DerivedData
+{
+
+TEST_CASE_NAMED(FDerivedDataCacheTimeAveragedStatTest, "DerivedData::Cache::Stats::TimeAveragedStat", "[DerivedData]")
+{
+	const auto MakeTime = [](double Seconds) { return FMonotonicTimePoint::FromSeconds(Seconds); };
+
+	FTimeAveragedStat Stat;
+	Stat.SetPeriod(FMonotonicTimeSpan::FromSeconds(8.0));
+
+	// Test StartTime after active window.
+	Stat.Add(MakeTime(0.0), MakeTime(8.0), 1024.0); // A
+	Stat.Add(MakeTime(16.0), MakeTime(32.0), 128.0); // B
+
+	// Test 128b over 8s -> 16b/s.
+	CHECK(Stat.GetRate(MakeTime(32.0)) == 16.0);   // (B)
+	CHECK(Stat.GetValue(MakeTime(32.0)) == 128.0); // (B)
+
+	// Test EndTime before active window.
+	Stat.Add(MakeTime(16.0), MakeTime(24.0), 128.0); // C
+
+	// Test 128b over 4s -> 32b/s.
+	CHECK(Stat.GetRate(MakeTime(36.0)) == 32.0);   // (B)
+	CHECK(Stat.GetValue(MakeTime(36.0)) == 128.0); // (B)
+
+	// Test a range that is entirely past existing ranges.
+	Stat.Add(MakeTime(40.0), MakeTime(48.0), 1024.0); // D
+
+	// Test 1024b over 8s -> 128b/s.
+	CHECK(Stat.GetRate(MakeTime(48.0)) == 128.0);   // (D)
+	CHECK(Stat.GetValue(MakeTime(48.0)) == 1024.0); // (D)
+
+	// Test that adding new values within the current window recalculates averages.
+	Stat.Add(MakeTime(38.0), MakeTime(46.0), 2048.0); // E
+
+	// Test 3072b over 8s -> 384b/s.
+	CHECK(Stat.GetRate(MakeTime(48.0)) == 384.0);   // (D, E)
+	CHECK(Stat.GetValue(MakeTime(48.0)) == 1536.0); // (D, E)
+
+	// Test 3072b over 6s -> 512b/s.
+	CHECK(Stat.GetRate(MakeTime(50.0)) == 512.0);   // (D, E)
+	CHECK(Stat.GetValue(MakeTime(50.0)) == 1536.0); // (D, E)
+
+	// Test a disjoint range that fits a previous range in the same window.
+	Stat.Add(MakeTime(52.0), MakeTime(54.0), 512.0); // F
+
+	// Test 1536b over 2+2s -> 384b/s.
+	CHECK(Stat.GetRate(MakeTime(54.0)) == 384.0);  // (D, F)
+	CHECK(Stat.GetValue(MakeTime(54.0)) == 768.0); // (D, F)
+
+	// Test a disjoint range in the middle of a previous range in the same window.
+	Stat.Add(MakeTime(50.5), MakeTime(51.0), 1536.0); // G
+
+	// Test 3072b over 1.5+0.5+2s -> 768b/s.
+	CHECK(Stat.GetRate(MakeTime(54.5)) == 768.0);   // (D, F, G)
+	CHECK(Stat.GetValue(MakeTime(54.5)) == 1024.0); // (D, F, G)
+
+	// Test a range that extends a previous range earlier in the same window.
+	Stat.Add(MakeTime(50.0), MakeTime(51.0), 1024.0); // H
+
+	// Test 4096b over 1+1+2s -> 1024b/s.
+	CHECK(Stat.GetRate(MakeTime(55.0)) == 1024.0);  // (D, F, G, H)
+	CHECK(Stat.GetValue(MakeTime(55.0)) == 1024.0); // (D, F, G, H)
+
+	// Test a range that overlaps two previous ranges in the same window and ends before the later one.
+	Stat.Add(MakeTime(46.5), MakeTime(50.5), 1536.0); // I
+
+	// Test 5632b over 3.5+2s -> 1024b/s.
+	CHECK(Stat.GetRate(MakeTime(55.5)) == 1024.0);  // (D, F, G, H, I)
+	CHECK(Stat.GetValue(MakeTime(55.5)) == 1126.4); // (D, F, G, H, I)
+
+	// Test a range that overlaps two previous ranges and ends after the later one.
+	Stat.Add(MakeTime(50.0), MakeTime(53.0), 512.0); // J
+
+	// Test 5120b over 5s -> 1024b/s.
+	CHECK(Stat.GetRate(MakeTime(57.0)) == 1024.0);  // (F, G, H, I, J)
+	CHECK(Stat.GetValue(MakeTime(57.0)) == 1024.0); // (F, G, H, I, J)
+
+	// Test an empty range that touches the start of the last existing range.
+	Stat.Add(MakeTime(70.0), MakeTime(74.0), 128.0); // K
+	Stat.Add(MakeTime(70.0), MakeTime(70.0), 128.0); // L
+
+	// Test 256b over 4s -> 64b/s.
+	CHECK(Stat.GetRate(MakeTime(74.0)) == 64.0);   // (K, L)
+	CHECK(Stat.GetValue(MakeTime(74.0)) == 128.0); // (K, L)
+
+	// Test an empty range that touches the start of the first existing range.
+	Stat.Add(MakeTime(75.0), MakeTime(76.0), 256.0); // M
+	Stat.Add(MakeTime(70.0), MakeTime(70.0), 128.0); // N
+
+	// Test 640b over 5s -> 128b/s.
+	CHECK(Stat.GetRate(MakeTime(76.0)) == 128.0);  // (K, L, M, N)
+	CHECK(Stat.GetValue(MakeTime(76.0)) == 160.0); // (K, L, M, N)
+}
+
+} // UE::DerivedData
+
+#endif // WITH_TESTS
