DDC: Serialize put requests into replays and replay them as gets

diff --git a/Engine/Source/Developer/DerivedDataCache/Private/DerivedDataCacheReplay.cpp b/Engine/Source/Developer/DerivedDataCache/Private/DerivedDataCacheReplay.cpp
--- a/Engine/Source/Developer/DerivedDataCache/Private/DerivedDataCacheReplay.cpp
+++ b/Engine/Source/Developer/DerivedDataCache/Private/DerivedDataCacheReplay.cpp
@@ -373,6 +373,19 @@ void FCacheStoreReplay::Put(
 	IRequestOwner& Owner,
 	FOnCachePutComplete&& OnComplete)
 {
+	// Serialize Put() as a Get() except for the Put method.
+	if (ReplayAr)
+	{
+		TArray<FCacheGetRequest, TInlineAllocator<16>> GetRequests;
+		TArray<FCacheGetResponse, TInlineAllocator<16>> GetResponses;
+		for (const FCachePutRequest& Request : Requests)
+		{
+			GetRequests.Add({Request.Name, Request.Record.GetKey(), Request.Policy, Request.UserData});
+			GetResponses.Add({Request.Name, Request.Record, Request.UserData, EStatus::Ok});
+		}
+		SerializeRequests<FCacheGetRequest>(GetRequests, GetResponses, ECacheMethod::Put, Owner.GetPriority());
+	}
+
 	InnerCache->Put(Requests, Owner, MoveTemp(OnComplete));
 }
 
@@ -389,6 +402,19 @@ void FCacheStoreReplay::PutValue(
 	IRequestOwner& Owner,
 	FOnCachePutValueComplete&& OnComplete)
 {
+	// Serialize PutValue() as a GetValue() except for the PutValue method.
+	if (ReplayAr)
+	{
+		TArray<FCacheGetValueRequest, TInlineAllocator<16>> GetRequests;
+		TArray<FCacheGetValueResponse, TInlineAllocator<16>> GetResponses;
+		for (const FCachePutValueRequest& Request : Requests)
+		{
+			GetRequests.Add({Request.Name, Request.Key, Request.Policy, Request.UserData});
+			GetResponses.Add({Request.Name, Request.Key, Request.Value, Request.UserData, EStatus::Ok});
+		}
+		SerializeRequests<FCacheGetValueRequest>(GetRequests, GetResponses, ECacheMethod::PutValue, Owner.GetPriority());
+	}
+
 	InnerCache->PutValue(Requests, Owner, MoveTemp(OnComplete));
 }
 
@@ -498,8 +524,8 @@ public:
 	};
 
 private:
-	template <typename RequestType, typename FunctionType>
-	bool DispatchRequests(FCbObjectView Object, ECacheMethod Method, FunctionType Function);
+	template <typename RequestType>
+	bool DispatchRequests(FCbObjectView Object, ECacheMethod Method);
 
 	void ApplyPolicyTransform(FCacheGetRequest& Request);
 	void ApplyPolicyTransform(FCacheGetValueRequest& Request);
@@ -563,11 +589,8 @@ void FCacheReplayReader::FState::WaitForAsyncReads()
 	}
 }
 
-template <typename RequestType, typename FunctionType>
-bool FCacheReplayReader::FState::DispatchRequests(
-	const FCbObjectView Object,
-	const ECacheMethod Method,
-	const FunctionType Function)
+template <typename RequestType>
+bool FCacheReplayReader::FState::DispatchRequests(const FCbObjectView Object, const ECacheMethod Method)
 {
 	if (!MethodFilter.IsMatch(Method))
 	{
@@ -618,15 +641,15 @@ bool FCacheReplayReader::FState::DispatchRequests(
 		{
 			// Owners with non-blocking priority can execute the request directly because it will be async.
 			FRequestBarrier Barrier(Owner);
-			Invoke(Function, TargetCache, Requests, Owner, [](auto&&){});
+			Invoke(CacheFunctionFor<RequestType>, TargetCache, Requests, Owner, [](TCacheResponseFor<RequestType>&&){});
 		}
 		else
 		{
 			// Owners with blocking priority launch a task to execute the blocking request to allow concurrent replay.
-			BlockingTasks.Add(Tasks::Launch(TEXT("CacheReplayTask"), [this, Requests = MoveTemp(Requests), Function]
+			BlockingTasks.Add(Tasks::Launch(TEXT("CacheReplayTask"), [this, Requests = MoveTemp(Requests)]
 			{
 				FRequestOwner BlockingOwner(EPriority::Blocking);
-				Invoke(Function, TargetCache, Requests, BlockingOwner, [](auto&&){});
+				Invoke(CacheFunctionFor<RequestType>, TargetCache, Requests, BlockingOwner, [](TCacheResponseFor<RequestType>&&){});
 				BlockingOwner.Wait();
 			}));
 		}
@@ -700,12 +723,13 @@ bool FCacheReplayReader::FState::ReadFromArchive(FArchive& ReplayAr, const uint6
 	// A scratch buffer for decompressed blocks of fields.
 	FUniqueBuffer BlockScratch;
 
+	bool bOk = false;
 	for (int64 Offset = ReplayAr.Tell(); Offset < ReplayAr.TotalSize(); Offset = ReplayAr.Tell())
 	{
 		FCbField Field = LoadCompactBinary(ReplayAr, Alloc);
 		if (!Field)
 		{
-			UE_LOG(LogDerivedDataCache, Warning,
+			UE_LOG(LogDerivedDataCache, Error,
 				TEXT("Replay: Failed to load compact binary at offset %" INT64_FMT ". "
 					 "Archive is at offset %" INT64_FMT " and has total size %" INT64_FMT "."),
 				Offset, ReplayAr.Tell(), ReplayAr.TotalSize());
@@ -722,7 +746,8 @@ bool FCacheReplayReader::FState::ReadFromArchive(FArchive& ReplayAr, const uint6
 					TEXT("Replay: Failed to load compressed buffer from binary field at offset %" INT64_FMT ". "
 						 "Archive is at offset %" INT64_FMT " and has total size %" INT64_FMT "."),
 					Offset, ReplayAr.Tell(), ReplayAr.TotalSize());
-				return false;
+				bOk = false;
+				continue;
 			}
 
 			const uint64 RawBlockSize = CompressedBuffer.GetRawSize();
@@ -738,13 +763,15 @@ bool FCacheReplayReader::FState::ReadFromArchive(FArchive& ReplayAr, const uint6
 					TEXT("Replay: Failed to decompress compressed buffer from binary field at offset %" INT64_FMT ". "
 						 "Archive is at offset %" INT64_FMT " and has total size %" INT64_FMT "."),
 					Offset, ReplayAr.Tell(), ReplayAr.TotalSize());
-				return false;
+				bOk = false;
+				continue;
 			}
 
 			FMemoryReaderView InnerAr(RawBlockView);
 			if (!ReadFromArchive(InnerAr, DefaultScratchSize))
 			{
-				return false;
+				bOk = false;
+				continue;
 			}
 		}
 
@@ -755,11 +782,11 @@ bool FCacheReplayReader::FState::ReadFromArchive(FArchive& ReplayAr, const uint6
 				TEXT("Replay: Failed to load cache request from object field at offset %" INT64_FMT ". "
 						"Archive is at offset %" INT64_FMT " and has total size %" INT64_FMT "."),
 				Offset, ReplayAr.Tell(), ReplayAr.TotalSize());
-			return false;
+			bOk = false;
+			continue;
 		}
 	}
-
-	return true;
+	return bOk;
 }
 
 bool FCacheReplayReader::FState::ReadFromObject(const FCbObjectView Object)
@@ -772,12 +799,14 @@ bool FCacheReplayReader::FState::ReadFromObject(const FCbObjectView Object)
 
 	switch (Method)
 	{
+	case ECacheMethod::Put:
 	case ECacheMethod::Get:
-		return DispatchRequests<FCacheGetRequest>(Object, Method, &ICacheStore::Get);
+		return DispatchRequests<FCacheGetRequest>(Object, Method);
+	case ECacheMethod::PutValue:
 	case ECacheMethod::GetValue:
-		return DispatchRequests<FCacheGetValueRequest>(Object, Method, &ICacheStore::GetValue);
+		return DispatchRequests<FCacheGetValueRequest>(Object, Method);
 	case ECacheMethod::GetChunks:
-		return DispatchRequests<FCacheGetChunkRequest>(Object, Method, &ICacheStore::GetChunks);
+		return DispatchRequests<FCacheGetChunkRequest>(Object, Method);
 	}
 
 	return false;
