DDC: Fixed replay ReadFromArchive always returning false

#rnx

diff --git a/Engine/Source/Developer/DerivedDataCache/Private/DerivedDataCacheReplay.cpp b/Engine/Source/Developer/DerivedDataCache/Private/DerivedDataCacheReplay.cpp
--- a/Engine/Source/Developer/DerivedDataCache/Private/DerivedDataCacheReplay.cpp
+++ b/Engine/Source/Developer/DerivedDataCache/Private/DerivedDataCacheReplay.cpp
@@ -723,7 +723,7 @@ bool FCacheReplayReader::FState::ReadFromArchive(FArchive& ReplayAr, const uint6
 	// A scratch buffer for decompressed blocks of fields.
 	FUniqueBuffer BlockScratch;
 
-	bool bOk = false;
+	bool bOk = true;
 	for (int64 Offset = ReplayAr.Tell(); Offset < ReplayAr.TotalSize(); Offset = ReplayAr.Tell())
 	{
 		FCbField Field = LoadCompactBinary(ReplayAr, Alloc);
