[Backout] - CL45484049
#rnx
Original CL Desc
-----------------------------------------------------------------
DDC: Serialize put requests into replays and replay them as gets

diff --git a/Engine/Source/Developer/DerivedDataCache/Private/DerivedDataCacheReplay.cpp b/Engine/Source/Developer/DerivedDataCache/Private/DerivedDataCacheReplay.cpp
--- a/Engine/Source/Developer/DerivedDataCache/Private/DerivedDataCacheReplay.cpp
+++ b/Engine/Source/Developer/DerivedDataCache/Private/DerivedDataCacheReplay.cpp
@@ -373,19 +373,6 @@ void FCacheStoreReplay::Put(
 	IRequestOwner& Owner,
 	FOnCachePutComplete&& OnComplete)
 {
-	// Serialize Put() as a Get() except for the Put method.
-	if (ReplayAr)
-	{
-		TArray<FCacheGetRequest, TInlineAllocator<16>> GetRequests;
-		TArray<FCacheGetResponse, TInlineAllocator<16>> GetResponses;
-		for (const FCachePutRequest& Request : Requests)
-		{
-			GetRequests.Add({Request.Name, Request.Record.GetKey(), Request.Policy, Request.UserData});
-			GetResponses.Add({Request.Name, Request.Record, Request.UserData, EStatus::Ok});
-		}
-		SerializeRequests<FCacheGetRequest>(GetRequests, GetResponses, ECacheMethod::Put, Owner.GetPriority());
-	}
-
 	InnerCache->Put(Requests, Owner, MoveTemp(OnComplete));
 }
 
@@ -402,19 +389,6 @@ void FCacheStoreReplay::PutValue(
 	IRequestOwner& Owner,
 	FOnCachePutValueComplete&& OnComplete)
 {
-	// Serialize PutValue() as a GetValue() except for the PutValue method.
-	if (ReplayAr)
-	{
-		TArray<FCacheGetValueRequest, TInlineAllocator<16>> GetRequests;
-		TArray<FCacheGetValueResponse, TInlineAllocator<16>> GetResponses;
-		for (const FCachePutValueRequest& Request : Requests)
-		{
-			GetRequests.Add({Request.Name, Request.Key, Request.Policy, Request.UserData});
-			GetResponses.Add({Request.Name, Request.Key, Request.Value, Request.UserData, EStatus::Ok});
-		}
-		SerializeRequests<FCacheGetValueRequest>(GetRequests, GetResponses, ECacheMethod::PutValue, Owner.GetPriority());
-	}
-
 	InnerCache->PutValue(Requests, Owner, MoveTemp(OnComplete));
 }
 
@@ -524,8 +498,8 @@ public:
 	};
 
 private:
-	template <typename RequestType>
-	bool DispatchRequests(FCbObjectView Object, ECacheMethod Method);
+	template <typename RequestType, typename FunctionType>
+	bool DispatchRequests(FCbObjectView Object, ECacheMethod Method, FunctionType Function);
 
 	void ApplyPolicyTransform(FCacheGetRequest& Request);
 	void ApplyPolicyTransform(FCacheGetValueRequest& Request);
@@ -589,8 +563,11 @@ void FCacheReplayReader::FState::WaitForAsyncReads()
 	}
 }
 
-template <typename RequestType>
-bool FCacheReplayReader::FState::DispatchRequests(const FCbObjectView Object, const ECacheMethod Method)
+template <typename RequestType, typename FunctionType>
+bool FCacheReplayReader::FState::DispatchRequests(
+	const FCbObjectView Object,
+	const ECacheMethod Method,
+	const FunctionType Function)
 {
 	if (!MethodFilter.IsMatch(Method))
 	{
@@ -641,15 +618,15 @@ bool FCacheReplayReader::FState::DispatchRequests(const FCbObjectView Object, co
 		{
 			// Owners with non-blocking priority can execute the request directly because it will be async.
 			FRequestBarrier Barrier(Owner);
-			Invoke(CacheFunctionFor<RequestType>, TargetCache, Requests, Owner, [](TCacheResponseFor<RequestType>&&){});
+			Invoke(Function, TargetCache, Requests, Owner, [](auto&&){});
 		}
 		else
 		{
 			// Owners with blocking priority launch a task to execute the blocking request to allow concurrent replay.
-			BlockingTasks.Add(Tasks::Launch(TEXT("CacheReplayTask"), [this, Requests = MoveTemp(Requests)]
+			BlockingTasks.Add(Tasks::Launch(TEXT("CacheReplayTask"), [this, Requests = MoveTemp(Requests), Function]
 			{
 				FRequestOwner BlockingOwner(EPriority::Blocking);
-				Invoke(CacheFunctionFor<RequestType>, TargetCache, Requests, BlockingOwner, [](TCacheResponseFor<RequestType>&&){});
+				Invoke(Function, TargetCache, Requests, BlockingOwner, [](auto&&){});
 				BlockingOwner.Wait();
 			}));
 		}
@@ -723,13 +700,12 @@ bool FCacheReplayReader::FState::ReadFromArchive(FArchive& ReplayAr, const uint6
 	// A scratch buffer for decompressed blocks of fields.
 	FUniqueBuffer BlockScratch;
 
-	bool bOk = false;
 	for (int64 Offset = ReplayAr.Tell(); Offset < ReplayAr.TotalSize(); Offset = ReplayAr.Tell())
 	{
 		FCbField Field = LoadCompactBinary(ReplayAr, Alloc);
 		if (!Field)
 		{
-			UE_LOG(LogDerivedDataCache, Error,
+			UE_LOG(LogDerivedDataCache, Warning,
 				TEXT("Replay: Failed to load compact binary at offset %" INT64_FMT ". "
 					 "Archive is at offset %" INT64_FMT " and has total size %" INT64_FMT "."),
 				Offset, ReplayAr.Tell(), ReplayAr.TotalSize());
@@ -746,8 +722,7 @@ bool FCacheReplayReader::FState::ReadFromArchive(FArchive& ReplayAr, const uint6
 					TEXT("Replay: Failed to load compressed buffer from binary field at offset %" INT64_FMT ". "
 						 "Archive is at offset %" INT64_FMT " and has total size %" INT64_FMT "."),
 					Offset, ReplayAr.Tell(), ReplayAr.TotalSize());
-				bOk = false;
-				continue;
+				return false;
 			}
 
 			const uint64 RawBlockSize = CompressedBuffer.GetRawSize();
@@ -763,15 +738,13 @@ bool FCacheReplayReader::FState::ReadFromArchive(FArchive& ReplayAr, const uint6
 					TEXT("Replay: Failed to decompress compressed buffer from binary field at offset %" INT64_FMT ". "
 						 "Archive is at offset %" INT64_FMT " and has total size %" INT64_FMT "."),
 					Offset, ReplayAr.Tell(), ReplayAr.TotalSize());
-				bOk = false;
-				continue;
+				return false;
 			}
 
 			FMemoryReaderView InnerAr(RawBlockView);
 			if (!ReadFromArchive(InnerAr, DefaultScratchSize))
 			{
-				bOk = false;
-				continue;
+				return false;
 			}
 		}
 
@@ -782,11 +755,11 @@ bool FCacheReplayReader::FState::ReadFromArchive(FArchive& ReplayAr, const uint6
 				TEXT("Replay: Failed to load cache request from object field at offset %" INT64_FMT ". "
 						"Archive is at offset %" INT64_FMT " and has total size %" INT64_FMT "."),
 				Offset, ReplayAr.Tell(), ReplayAr.TotalSize());
-			bOk = false;
-			continue;
+			return false;
 		}
 	}
-	return bOk;
+
+	return true;
 }
 
 bool FCacheReplayReader::FState::ReadFromObject(const FCbObjectView Object)
@@ -799,14 +772,12 @@ bool FCacheReplayReader::FState::ReadFromObject(const FCbObjectView Object)
 
 	switch (Method)
 	{
-	case ECacheMethod::Put:
 	case ECacheMethod::Get:
-		return DispatchRequests<FCacheGetRequest>(Object, Method);
-	case ECacheMethod::PutValue:
+		return DispatchRequests<FCacheGetRequest>(Object, Method, &ICacheStore::Get);
 	case ECacheMethod::GetValue:
-		return DispatchRequests<FCacheGetValueRequest>(Object, Method);
+		return DispatchRequests<FCacheGetValueRequest>(Object, Method, &ICacheStore::GetValue);
 	case ECacheMethod::GetChunks:
-		return DispatchRequests<FCacheGetChunkRequest>(Object, Method);
+		return DispatchRequests<FCacheGetChunkRequest>(Object, Method, &ICacheStore::GetChunks);
 	}
 
 	return false;
