DDC: Clarify documentation on FRequestOwner::Poll()

#rnx

diff --git a/Engine/Source/Developer/DerivedDataCache/Public/DerivedDataRequestOwner.h b/Engine/Source/Developer/DerivedDataCache/Public/DerivedDataRequestOwner.h
--- a/Engine/Source/Developer/DerivedDataCache/Public/DerivedDataRequestOwner.h
+++ b/Engine/Source/Developer/DerivedDataCache/Public/DerivedDataRequestOwner.h
@@ -126,7 +126,7 @@ public:
 	UE_API void Cancel();
 	/** Wait for any active and future requests or barriers in the owner. */
 	UE_API void Wait();
-	/** Poll whether the owner has any active requests or barriers. */
+	/** Poll the owner for active requests or barriers, returning false if it has either and true otherwise. */
 	UE_API bool Poll() const;
 
 	/** Returns whether the owner has been canceled, which new requests are expected to check. */
