DDC-Verify Telemetry

diff --git a/Engine/Plugins/Experimental/EditorTelemetry/Source/Private/EditorTelemetry.cpp b/Engine/Plugins/Experimental/EditorTelemetry/Source/Private/EditorTelemetry.cpp
--- a/Engine/Plugins/Experimental/EditorTelemetry/Source/Private/EditorTelemetry.cpp
+++ b/Engine/Plugins/Experimental/EditorTelemetry/Source/Private/EditorTelemetry.cpp
@@ -275,6 +275,39 @@ void FEditorTelemetry::RecordEvent_DDCSummary(const FString& Context, TArray<FAn
 #endif			
 }
 
+void FEditorTelemetry::RecordEvent_DDCVerification(const FString& Context, TArray<FAnalyticsEventAttribute> Attributes)
+{
+	const int SchemaVersion = 1;
+	Attributes.Emplace(TEXT("SchemaVersion"), SchemaVersion);
+	Attributes.Emplace(TEXT("Context"), Context);
+
+	FDerivedDataCacheVerificationStats DDCVerifyStats;
+	GetDerivedDataCacheRef().GatherVerificationStats(DDCVerifyStats);
+
+	for (const FDerivedDataCacheVerificationStat& Stat : DDCVerifyStats.Entries)
+	{
+		if (Stat.RecordsTested > 0)
+		{
+			TArray<FAnalyticsEventAttribute> EventAttributes = Attributes;
+			EventAttributes.Emplace(TEXT("Bucket"), Stat.BucketName);
+			EventAttributes.Emplace(TEXT("RecordsTested"), Stat.RecordsTested);
+			EventAttributes.Emplace(TEXT("RecordsMatchedCache"), Stat.RecordsMatchedCache);
+			EventAttributes.Emplace(TEXT("RecordsNotInCache"), Stat.RecordsNotInCache);
+			EventAttributes.Emplace(TEXT("ValuesTested"), Stat.ValuesTested);
+			EventAttributes.Emplace(TEXT("ValuesTestedBytes"), Stat.ValuesTestedBytes);
+			EventAttributes.Emplace(TEXT("ValuesMatchedCache"), Stat.ValuesMatchedCache);
+			EventAttributes.Emplace(TEXT("ValuesMatchedCacheBytes"), Stat.ValuesMatchedCacheBytes);
+			EventAttributes.Emplace(TEXT("ValuesNotInCache"), Stat.ValuesNotInCache);
+			EventAttributes.Emplace(TEXT("ValuesNotInCacheBytes"), Stat.ValuesNotInCacheBytes);
+			EventAttributes.Emplace(TEXT("ValuesChanged"), Stat.ValuesChanged);
+			EventAttributes.Emplace(TEXT("ValuesChangedBytes"), Stat.ValuesChangedBytes);
+			EventAttributes.Emplace(TEXT("ValuesNotInRequest"), Stat.ValuesNotInRequest);
+			EventAttributes.Emplace(TEXT("ValuesNotInRequestBytes"), Stat.ValuesNotInRequestBytes);
+			FStudioTelemetry::Get().RecordEvent(TEXT("Core.DDC.Verify"), EventAttributes);
+		}
+	}
+}
+
 void FEditorTelemetry::RecordEvent_Zen(const FString& Context, TArray<FAnalyticsEventAttribute> Attributes)
 {
 #if UE_WITH_ZEN
@@ -839,6 +872,7 @@ void FEditorTelemetry::StartSession()
 
 				FEditorTelemetry::RecordEvent_Cooking(Attributes);
 				FEditorTelemetry::RecordEvent_CoreSystems(TEXT("Cooking"), Attributes);
+				FEditorTelemetry::RecordEvent_DDCVerification(TEXT("Cooking"), Attributes);
 			});
 	}
 
@@ -875,6 +909,7 @@ void FEditorTelemetry::StartSession()
 
 		FEditorTelemetry::RecordEvent_Cooking(CookingSpan->GetAttributes());
 		FEditorTelemetry::RecordEvent_CoreSystems(TEXT("Cooking"), CookingSpan->GetAttributes());
+		FEditorTelemetry::RecordEvent_DDCVerification(TEXT("Cooking"), CookingSpan->GetAttributes());
 
 		FStudioTelemetry::Get().EndSpan(CookingSpan);
 	});
diff --git a/Engine/Plugins/Experimental/EditorTelemetry/Source/Public/EditorTelemetry.h b/Engine/Plugins/Experimental/EditorTelemetry/Source/Public/EditorTelemetry.h
--- a/Engine/Plugins/Experimental/EditorTelemetry/Source/Public/EditorTelemetry.h
+++ b/Engine/Plugins/Experimental/EditorTelemetry/Source/Public/EditorTelemetry.h
@@ -37,6 +37,7 @@ public:
 	UE_API void RecordEvent_CoreSystems(const FString& Context, TArray<FAnalyticsEventAttribute> Attributes = {});
 	UE_API void RecordEvent_DDCResource(const FString& Context, TArray<FAnalyticsEventAttribute> Attributes = {});
 	UE_API void RecordEvent_DDCSummary(const FString& Context, TArray<FAnalyticsEventAttribute> Attributes = {});
+	UE_API void RecordEvent_DDCVerification(const FString& Context, TArray<FAnalyticsEventAttribute> Attributes = {});
 	UE_API void RecordEvent_Zen(const FString& Context, TArray<FAnalyticsEventAttribute> Attributes = {});
 	UE_API void RecordEvent_VirtualAssets(const FString& Context, TArray<FAnalyticsEventAttribute> Attributes = {});
 	UE_API void RecordEvent_MemoryLLM(const FString& Context, TArray<FAnalyticsEventAttribute> Attributes = {});
diff --git a/Engine/Source/Developer/DerivedDataCache/Private/DerivedDataBackendInterface.h b/Engine/Source/Developer/DerivedDataCache/Private/DerivedDataBackendInterface.h
--- a/Engine/Source/Developer/DerivedDataCache/Private/DerivedDataBackendInterface.h
+++ b/Engine/Source/Developer/DerivedDataCache/Private/DerivedDataBackendInterface.h
@@ -135,6 +135,8 @@ public:
 	virtual TSharedRef<FDerivedDataCacheStatsNode> GatherUsageStats() const = 0;
 
 	virtual void GatherResourceStats(TArray<FDerivedDataCacheResourceStat>& DDCResourceStats) const = 0;
+
+	virtual void GatherVerificationStats(FDerivedDataCacheVerificationStats& OutStats) const = 0;
 };
 
 /** Lexical conversions from and to enums */
diff --git a/Engine/Source/Developer/DerivedDataCache/Private/DerivedDataBackends.cpp b/Engine/Source/Developer/DerivedDataCache/Private/DerivedDataBackends.cpp
--- a/Engine/Source/Developer/DerivedDataCache/Private/DerivedDataBackends.cpp
+++ b/Engine/Source/Developer/DerivedDataCache/Private/DerivedDataBackends.cpp
@@ -844,6 +844,14 @@ public:
 		}
 	}
 
+	virtual void GatherVerificationStats(FDerivedDataCacheVerificationStats& OutStats) const override
+	{
+		if (RootCache)
+		{
+			RootCache->GatherVerificationStats(OutStats);
+		}
+	}
+
 private:
 	void Add(ILegacyCacheStore* CacheStore, ECacheStoreFlags Flags) final
 	{
diff --git a/Engine/Source/Developer/DerivedDataCache/Private/DerivedDataCache.cpp b/Engine/Source/Developer/DerivedDataCache/Private/DerivedDataCache.cpp
--- a/Engine/Source/Developer/DerivedDataCache/Private/DerivedDataCache.cpp
+++ b/Engine/Source/Developer/DerivedDataCache/Private/DerivedDataCache.cpp
@@ -1146,6 +1146,11 @@ public:
 #endif
 	}
 
+	void GatherVerificationStats(FDerivedDataCacheVerificationStats& OutStats) const override
+	{
+		Backend->GatherVerificationStats(OutStats);
+	}
+
 	/** Get event delegate for data cache notifications */
 	virtual FOnDDCNotification& GetDDCNotificationEvent()
 	{
diff --git a/Engine/Source/Developer/DerivedDataCache/Private/DerivedDataCacheStoreVerify.cpp b/Engine/Source/Developer/DerivedDataCache/Private/DerivedDataCacheStoreVerify.cpp
--- a/Engine/Source/Developer/DerivedDataCache/Private/DerivedDataCacheStoreVerify.cpp
+++ b/Engine/Source/Developer/DerivedDataCache/Private/DerivedDataCacheStoreVerify.cpp
@@ -1,9 +1,12 @@
 // Copyright Epic Games, Inc. All Rights Reserved.
 
 #include "Algo/Compare.h"
+#include "Algo/Find.h"
+#include "Async/UniqueLock.h"
 #include "Containers/Set.h"
 #include "DerivedDataCacheKey.h"
 #include "DerivedDataCacheKeyFilter.h"
+#include "DerivedDataCacheUsageStats.h"
 #include "DerivedDataLegacyCacheStore.h"
 #include "HAL/CriticalSection.h"
 #include "Misc/CommandLine.h"
@@ -108,6 +111,8 @@ public:
 		return InnerCache->LegacyDebugOptions(Options);
 	}
 
+	void GatherVerificationStats(FDerivedDataCacheVerificationStats& OutStats) const final;
+
 private:
 	struct FVerifyPutState
 	{
@@ -135,7 +140,7 @@ private:
 	void GetDataComplete(IRequestOwner& Owner, FVerifyPutValueState* State, FCacheGetValueResponse&& Response);
 	void GetComplete(IRequestOwner& Owner, FVerifyPutValueState* State);
 
-	static bool CompareRecords(const FCacheRecord& PutRecord, const FCacheRecord& GetRecord, const FSharedString& Name);
+	bool CompareRecords(const FCacheRecord& PutRecord, const FCacheRecord& GetRecord, const FSharedString& Name);
 
 	static void LogChangedValue(
 		const FSharedString& Name,
@@ -146,7 +151,28 @@ private:
 		const FCompositeBuffer& NewRawData,
 		const FCompositeBuffer& OldRawData);
 
+	// Track per-bucket verification results
+	struct FTrackedEntryStats
+	{
+		/*
+			ModifyT = void(FDerivedDataCacheVerificationStat&), used to write to a bucket and maintain the write lock
+		*/
+		template <typename ModifyT>
+		void Modify(FCacheBucket Bucket, ModifyT DoModification)
+		{
+			TUniqueLock DoLock(Lock);
+
+			// Bucket names are resolved in GatherVerificationStats
+			FDerivedDataCacheVerificationStat& StatsRef = PerBucketStats.FindOrAdd(Bucket);		
+			DoModification(StatsRef);
+		}
+
+		TMap<FCacheBucket, FDerivedDataCacheVerificationStat> PerBucketStats;
+		mutable FMutex Lock;
+	};
+
 private:
+	FTrackedEntryStats BucketStats;
 	ILegacyCacheStore* InnerCache;
 	FCriticalSection AlreadyTestedLock;
 	TSet<FCacheKey> AlreadyTested;
@@ -154,6 +180,17 @@ private:
 	bool bPutOnError;
 };
 
+// A helper to get the total size of an array of values in a record
+uint64 GetValuesTotalSize(const TConstArrayView<FValueWithId>& Values)
+{
+	uint64 TotalSize = 0;
+	for (const FValueWithId& Value : Values)
+	{
+		TotalSize += Value.GetRawSize();
+	}
+	return TotalSize;
+}
+
 void FCacheStoreVerify::Put(
 	const TConstArrayView<FCachePutRequest> Requests,
 	IRequestOwner& Owner,
@@ -172,6 +209,15 @@ void FCacheStoreVerify::Put(
 				AlreadyTested.Add(Key, &bForward);
 			}
 			(bForward ? State->ForwardRequests : State->VerifyRequests).Add(Request);
+			if (!bForward)
+			{
+				BucketStats.Modify(Request.Record.GetKey().Bucket, [&Request](FDerivedDataCacheVerificationStat& BucketStat) 
+				{
+					BucketStat.RecordsTested++;
+					BucketStat.ValuesTested += Request.Record.GetValues().Num();
+					BucketStat.ValuesTestedBytes += GetValuesTotalSize(Request.Record.GetValues());
+				});
+			}
 		}
 	}
 
@@ -216,6 +262,12 @@ void FCacheStoreVerify::GetMetaComplete(IRequestOwner& Owner, FVerifyPutState* S
 				TEXT("Verify: Data in the cache matches newly generated data for %s from '%s'."),
 				*WriteToString<96>(Request.Record.GetKey()), *Request.Name);
 			State->OnComplete(Request.MakeResponse(EStatus::Ok));
+			BucketStats.Modify(Response.Record.GetKey().Bucket, [&Request](FDerivedDataCacheVerificationStat& BucketStat) 
+			{
+				BucketStat.RecordsMatchedCache++;
+				BucketStat.ValuesMatchedCache += Request.Record.GetValues().Num();
+				BucketStat.ValuesMatchedCacheBytes += GetValuesTotalSize(Request.Record.GetValues());
+			});
 		}
 		else
 		{
@@ -232,6 +284,12 @@ void FCacheStoreVerify::GetMetaComplete(IRequestOwner& Owner, FVerifyPutState* S
 		UE_LOG(LogDerivedDataCache, Warning,
 			TEXT("Verify: Cache did not contain a record for %s from '%s'."),
 			*WriteToString<96>(Request.Record.GetKey()), *Request.Name);
+		BucketStats.Modify(Response.Record.GetKey().Bucket, [&Request](FDerivedDataCacheVerificationStat& BucketStat) 
+		{
+			BucketStat.RecordsNotInCache++;
+			BucketStat.ValuesNotInCache += Request.Record.GetValues().Num();
+			BucketStat.ValuesNotInCacheBytes += GetValuesTotalSize(Request.Record.GetValues());
+		});
 		FWriteScopeLock Lock(State->Lock);
 		State->ForwardRequests.Add(MoveTemp(Request));
 	}
@@ -273,6 +331,12 @@ void FCacheStoreVerify::GetDataComplete(IRequestOwner& Owner, FVerifyPutState* S
 		UE_LOG(LogDerivedDataCache, Warning,
 			TEXT("Verify: Cache did not contain a record for %s from '%s'."),
 			*WriteToString<96>(Request.Record.GetKey()), *Request.Name);
+		BucketStats.Modify(Request.Record.GetKey().Bucket, [&Request](FDerivedDataCacheVerificationStat& BucketStat) 
+		{
+			BucketStat.RecordsNotInCache++;
+			BucketStat.ValuesNotInCache += Request.Record.GetValues().Num();
+			BucketStat.ValuesNotInCacheBytes += GetValuesTotalSize(Request.Record.GetValues());
+		});
 		FWriteScopeLock Lock(State->Lock);
 		State->ForwardRequests.Add(MoveTemp(Request));
 	}
@@ -336,6 +400,14 @@ void FCacheStoreVerify::PutValue(
 				AlreadyTested.Add(Request.Key, &bForward);
 			}
 			(bForward ? State->ForwardRequests : State->VerifyRequests).Add(Request);
+			if (!bForward)
+			{
+				BucketStats.Modify(Request.Key.Bucket, [&Request](FDerivedDataCacheVerificationStat& BucketStat) 
+				{
+					BucketStat.ValuesTested++;
+					BucketStat.ValuesTestedBytes += Request.Value.GetRawSize();
+				});
+			}
 		}
 	}
 
@@ -375,6 +447,11 @@ void FCacheStoreVerify::GetMetaComplete(IRequestOwner& Owner, FVerifyPutValueSta
 				TEXT("Verify: Data in the cache matches newly generated data for %s from '%s'."),
 				*WriteToString<96>(Request.Key), *Request.Name);
 			State->OnComplete(Request.MakeResponse(EStatus::Ok));
+			BucketStats.Modify(Response.Key.Bucket, [&Request](FDerivedDataCacheVerificationStat& BucketStat) 
+			{
+				BucketStat.ValuesMatchedCache++;
+				BucketStat.ValuesMatchedCacheBytes += Request.Value.GetRawSize();
+			});
 		}
 		else
 		{
@@ -390,6 +467,11 @@ void FCacheStoreVerify::GetMetaComplete(IRequestOwner& Owner, FVerifyPutValueSta
 		UE_LOG(LogDerivedDataCache, Display,
 			TEXT("Verify: Cache did not contain a value for %s from '%s'."),
 			*WriteToString<96>(Request.Key), *Request.Name);
+		BucketStats.Modify(Response.Key.Bucket, [&Request](FDerivedDataCacheVerificationStat& BucketStat) 
+		{
+			BucketStat.ValuesNotInCache++;
+			BucketStat.ValuesNotInCacheBytes += Request.Value.GetRawSize();
+		});
 		FWriteScopeLock Lock(State->Lock);
 		State->ForwardRequests.Add(MoveTemp(Request));
 	}
@@ -409,12 +491,22 @@ void FCacheStoreVerify::GetDataComplete(IRequestOwner& Owner, FVerifyPutValueSta
 				TEXT("Verify: Data in the cache matches newly generated data for %s from '%s'."),
 				*WriteToString<96>(Request.Key), *Request.Name);
 			State->OnComplete(Request.MakeResponse(EStatus::Ok));
+			BucketStats.Modify(Response.Key.Bucket, [&Request](FDerivedDataCacheVerificationStat& BucketStat) 
+			{
+				BucketStat.ValuesMatchedCache++;
+				BucketStat.ValuesMatchedCacheBytes += Request.Value.GetRawSize();
+			});
 		}
 		else
 		{
 			LogChangedValue(Request.Name, Request.Key, FValueId::Null,
 				Request.Value.GetRawHash(), Response.Value.GetRawHash(),
 				Request.Value.GetData().DecompressToComposite(), Response.Value.GetData().DecompressToComposite());
+			BucketStats.Modify(Response.Key.Bucket, [&Request](FDerivedDataCacheVerificationStat& BucketStat) 
+			{
+				BucketStat.ValuesChanged++;
+				BucketStat.ValuesChangedBytes += Request.Value.GetRawSize();
+			});
 			if (bPutOnError)
 			{
 				// Ask to overwrite existing values to potentially eliminate the mismatch.
@@ -436,6 +528,11 @@ void FCacheStoreVerify::GetDataComplete(IRequestOwner& Owner, FVerifyPutValueSta
 		UE_LOG(LogDerivedDataCache, Display,
 			TEXT("Verify: Cache did not contain a value for %s from '%s'."),
 			*WriteToString<96>(Request.Key), *Request.Name);
+		BucketStats.Modify(Response.Key.Bucket, [&Request](FDerivedDataCacheVerificationStat& BucketStat) 
+		{
+			BucketStat.ValuesNotInCache++;
+			BucketStat.ValuesNotInCacheBytes += Request.Value.GetRawSize();
+		});
 		FWriteScopeLock Lock(State->Lock);
 		State->ForwardRequests.Add(MoveTemp(Request));
 	}
@@ -519,19 +616,25 @@ bool FCacheStoreVerify::CompareRecords(const FCacheRecord& PutRecord, const FCac
 	const FValueWithId* GetIt = GetValues.GetData();
 	const FValueWithId* const PutEnd = PutIt + PutValues.Num();
 	const FValueWithId* const GetEnd = GetIt + GetValues.Num();
+	uint32 ValuesDifferent = 0, ValuesMatch = 0, NewValues = 0, MissingValues = 0;
+	uint64 ValuesDifferentBytes = 0, ValuesMatchBytes = 0, NewValuesBytes = 0, MissingValuesBytes = 0;
 
-	const auto LogNewValue = [&Name, &Key](const FValueWithId& Value)
+	const auto LogNewValue = [&Name, &Key, &NewValues, &NewValuesBytes](const FValueWithId& Value)
 	{
 		UE_LOG(LogDerivedDataCache, Error,
 			TEXT("Verify: Value %s with hash %s is in the new record but does not exist in the cache for %s from '%s'."),
 			*WriteToString<32>(Value.GetId()), *WriteToString<48>(Value.GetRawHash()), *WriteToString<96>(Key), *Name);
+		NewValues++;
+		NewValuesBytes += Value.GetRawSize();
 	};
 
-	const auto LogOldValue = [&Name, &Key](const FValueWithId& Value)
+	const auto LogOldValue = [&Name, &Key, &MissingValues, &MissingValuesBytes](const FValueWithId& Value)
 	{
 		UE_LOG(LogDerivedDataCache, Error,
 			TEXT("Verify: Value %s with hash %s is in the cache but does not exist in the new record for %s from '%s'."),
 			*WriteToString<32>(Value.GetId()), *WriteToString<48>(Value.GetRawHash()), *WriteToString<96>(Key), *Name);
+		MissingValues++;
+		MissingValuesBytes += Value.GetRawSize();
 	};
 
 	while (PutIt != PutEnd && GetIt != GetEnd)
@@ -544,6 +647,13 @@ bool FCacheStoreVerify::CompareRecords(const FCacheRecord& PutRecord, const FCac
 					PutIt->GetRawHash(), GetIt->GetRawHash(),
 					PutIt->GetData().DecompressToComposite(), GetIt->GetData().DecompressToComposite());
 				bEqual = false;
+				ValuesDifferent++;
+				ValuesDifferentBytes += PutIt->GetRawSize();
+			}
+			else
+			{
+				ValuesMatch++;
+				ValuesMatchBytes += PutIt->GetRawSize();
 			}
 			++PutIt;
 			++GetIt;
@@ -572,6 +682,19 @@ bool FCacheStoreVerify::CompareRecords(const FCacheRecord& PutRecord, const FCac
 		bEqual = false;
 	}
 
+	BucketStats.Modify(PutRecord.GetKey().Bucket, [&](FDerivedDataCacheVerificationStat& BucketStat) 
+	{
+		BucketStat.RecordsMatchedCache += bEqual ? 1 : 0;
+		BucketStat.ValuesMatchedCache += ValuesMatch;
+		BucketStat.ValuesMatchedCacheBytes += ValuesMatchBytes;
+		BucketStat.ValuesNotInCache += NewValues;
+		BucketStat.ValuesNotInCacheBytes += NewValuesBytes;
+		BucketStat.ValuesChanged += ValuesDifferent;
+		BucketStat.ValuesChangedBytes += ValuesDifferentBytes;
+		BucketStat.ValuesNotInRequest += MissingValues;
+		BucketStat.ValuesNotInRequestBytes += MissingValuesBytes;
+	});
+
 	return bEqual;
 }
 
@@ -625,6 +748,27 @@ void FCacheStoreVerify::LogChangedValue(
 	LogDataToFile(OldRawHash, OldRawData, TEXTVIEW(".fromcache"));
 }
 
+void FCacheStoreVerify::GatherVerificationStats(FDerivedDataCacheVerificationStats& OutStats) const
+{
+	// Buckets may share the same display name, combine their stats here
+	TUniqueLock DoLock(BucketStats.Lock);
+	for (const TPair<FCacheBucket, FDerivedDataCacheVerificationStat>& Entry : BucketStats.PerBucketStats)
+	{
+		TStringBuilder<64> BucketName;
+		Entry.Key.ToDisplayName(BucketName);
+		FDerivedDataCacheVerificationStat* StatToRecord = Algo::FindBy(OutStats.Entries, BucketName.ToView(), &FDerivedDataCacheVerificationStat::BucketName);
+		if (StatToRecord != nullptr)
+		{
+			*StatToRecord += Entry.Value;
+		}
+		else
+		{
+			FDerivedDataCacheVerificationStat& NewEntry = OutStats.Entries.Emplace_GetRef(Entry.Value);
+			NewEntry.BucketName = BucketName;
+		}
+	}
+}
+
 ILegacyCacheStore* CreateCacheStoreVerify(ILegacyCacheStore* InnerCache, bool bPutOnError)
 {
 	return new FCacheStoreVerify(InnerCache, bPutOnError);
diff --git a/Engine/Source/Developer/DerivedDataCache/Private/DerivedDataLegacyCacheStore.h b/Engine/Source/Developer/DerivedDataCache/Private/DerivedDataLegacyCacheStore.h
--- a/Engine/Source/Developer/DerivedDataCache/Private/DerivedDataLegacyCacheStore.h
+++ b/Engine/Source/Developer/DerivedDataCache/Private/DerivedDataLegacyCacheStore.h
@@ -9,6 +9,7 @@
 #include <atomic>
 
 class FDerivedDataCacheStatsNode;
+struct FDerivedDataCacheVerificationStats;
 
 namespace UE::DerivedData { class FLegacyCacheKey; }
 namespace UE::DerivedData { struct FBackendDebugOptions; }
@@ -51,6 +52,8 @@ public:
 	virtual void LegacyStats(FDerivedDataCacheStatsNode& OutNode) = 0;
 
 	virtual bool LegacyDebugOptions(FBackendDebugOptions& Options) = 0;
+
+	virtual void GatherVerificationStats(FDerivedDataCacheVerificationStats& OutStats) const { }
 };
 
 class Private::FLegacyCacheKeyShared final
diff --git a/Engine/Source/Developer/DerivedDataCache/Public/DerivedDataCacheInterface.h b/Engine/Source/Developer/DerivedDataCache/Public/DerivedDataCacheInterface.h
--- a/Engine/Source/Developer/DerivedDataCache/Public/DerivedDataCacheInterface.h
+++ b/Engine/Source/Developer/DerivedDataCache/Public/DerivedDataCacheInterface.h
@@ -14,6 +14,7 @@ class FDerivedDataCacheStatsNode;
 class IDDCCleanup;
 struct FDerivedDataCacheResourceStat;
 struct FDerivedDataCacheSummaryStats;
+struct FDerivedDataCacheVerificationStats;
 struct FAnalyticsEventAttribute;
 
 /** 
@@ -278,6 +279,11 @@ public:
 
 	virtual void GatherSummaryStats(FDerivedDataCacheSummaryStats& DDCSummaryStats) const = 0;
 
+	/**
+	 * Retrieve stats about the integrity of DDC data vs generated data
+	 */
+	UE_INTERNAL virtual void GatherVerificationStats(FDerivedDataCacheVerificationStats& OutStats) const = 0;
+
 	UE_DEPRECATED(5.4, "This is no longer used. Use GatherUsageStats/GatherResourceStats/GatherSummaryStats instead")
 	virtual void GatherAnalytics(TArray<FAnalyticsEventAttribute>& Attributes) const = 0;
 
diff --git a/Engine/Source/Developer/DerivedDataCache/Public/DerivedDataCacheUsageStats.h b/Engine/Source/Developer/DerivedDataCache/Public/DerivedDataCacheUsageStats.h
--- a/Engine/Source/Developer/DerivedDataCache/Public/DerivedDataCacheUsageStats.h
+++ b/Engine/Source/Developer/DerivedDataCache/Public/DerivedDataCacheUsageStats.h
@@ -309,6 +309,48 @@ struct FDerivedDataCacheResourceStatKeyFuncs : BaseKeyFuncs<FDerivedDataCacheRes
 	static uint32 GetKeyHash(const FString& Key) { return GetTypeHash(Key); }
 };
 
+// DDC-Verify collects data about a subset of newly generated derived data, and whether it matches existing cached values
+// This object is used to track statistics about verification results and is collected per DDC bucket
+struct FDerivedDataCacheVerificationStat
+{
+	FString BucketName;						// name of the bucket
+	uint32 RecordsTested = 0;				// num records that were tested against the cache
+	uint32 RecordsMatchedCache = 0;			// num records that match existing cache data
+	uint32 RecordsNotInCache = 0;			// records that have no existing data in the cache 
+	uint32 ValuesTested = 0;				// multiple values per record
+	uint64 ValuesTestedBytes = 0;			// total size of values tested
+	uint32 ValuesMatchedCache = 0;			// num values that match existing cache data
+	uint64 ValuesMatchedCacheBytes = 0;		// total size of matching data
+	uint32 ValuesNotInCache = 0;			// num values that have no existing data in the cache
+	uint64 ValuesNotInCacheBytes = 0;		// num bytes of new values written to the cache
+	uint32 ValuesChanged = 0;				// num values existing in the cache that don't match the new data
+	uint64 ValuesChangedBytes = 0;			// num bytes of changed values detected (note this tracks the entire value, not the individual changes)
+	uint32 ValuesNotInRequest = 0;			// num values that were in the cache before, but not in the new data
+	uint64 ValuesNotInRequestBytes = 0;		// num bytes of data we are no longer writing in values
+
+	const FDerivedDataCacheVerificationStat& operator+=(const FDerivedDataCacheVerificationStat& OtherStat)
+	{
+		RecordsTested += OtherStat.RecordsTested;
+		RecordsMatchedCache += OtherStat.RecordsMatchedCache;
+		RecordsNotInCache += OtherStat.RecordsNotInCache;
+		ValuesTested += OtherStat.ValuesTested;
+		ValuesTestedBytes += OtherStat.ValuesTestedBytes;
+		ValuesMatchedCache += OtherStat.ValuesMatchedCache;
+		ValuesMatchedCacheBytes += OtherStat.ValuesMatchedCacheBytes;
+		ValuesNotInCache += OtherStat.ValuesNotInCache;
+		ValuesNotInCacheBytes += OtherStat.ValuesNotInCacheBytes;
+		ValuesChanged += OtherStat.ValuesChanged;
+		ValuesChangedBytes += OtherStat.ValuesChangedBytes;
+		ValuesNotInRequest += OtherStat.ValuesNotInRequest;
+		ValuesNotInRequestBytes += OtherStat.ValuesNotInRequestBytes;
+		return *this;
+	}
+};
+
+struct FDerivedDataCacheVerificationStats
+{
+	TArray<FDerivedDataCacheVerificationStat> Entries;
+};
 
 DERIVEDDATACACHE_API void GatherDerivedDataCacheResourceStats(TArray<FDerivedDataCacheResourceStat>& DDCResourceStats);
 DERIVEDDATACACHE_API void GatherDerivedDataCacheSummaryStats(FDerivedDataCacheSummaryStats& DDCSummaryStats);
\ No newline at end of file
