DDC: Added ECachePolicy::NonDeterministic

diff --git a/Engine/Source/Developer/DerivedDataCache/Private/DerivedDataCachePolicy.cpp b/Engine/Source/Developer/DerivedDataCache/Private/DerivedDataCachePolicy.cpp
--- a/Engine/Source/Developer/DerivedDataCache/Private/DerivedDataCachePolicy.cpp
+++ b/Engine/Source/Developer/DerivedDataCache/Private/DerivedDataCachePolicy.cpp
@@ -42,6 +42,7 @@ static const FCachePolicyToText CachePolicyToText[]
 	{ ECachePolicy::SkipData,      ANSITEXTVIEW("SkipData") },
 	{ ECachePolicy::PartialRecord, ANSITEXTVIEW("PartialRecord") },
 	{ ECachePolicy::KeepAlive,     ANSITEXTVIEW("KeepAlive") },
+	{ ECachePolicy::NonDeterministic, ANSITEXTVIEW("NonDeterministic") },
 	// None must be last because it matches every policy.
 	{ ECachePolicy::None,          ANSITEXTVIEW("None") },
 };
@@ -51,7 +52,8 @@ static constexpr ECachePolicy CachePolicyKnownFlags
 	| ECachePolicy::SkipMeta
 	| ECachePolicy::SkipData
 	| ECachePolicy::PartialRecord
-	| ECachePolicy::KeepAlive;
+	| ECachePolicy::KeepAlive
+	| ECachePolicy::NonDeterministic;
 
 template <typename CharType>
 static TStringBuilderBase<CharType>& CachePolicyToString(TStringBuilderBase<CharType>& Builder, ECachePolicy Policy)
@@ -124,7 +126,7 @@ namespace UE::DerivedData
 
 ECachePolicy CombineCachePolicy(const ECachePolicy A, const ECachePolicy B)
 {
-	constexpr ECachePolicy NegativeFlags = ECachePolicy::SkipData | ECachePolicy::SkipMeta;
+	constexpr ECachePolicy NegativeFlags = ECachePolicy::SkipData | ECachePolicy::SkipMeta | ECachePolicy::NonDeterministic;
 	return ((A | B) & ~NegativeFlags) | ((A & B) & NegativeFlags);
 }
 
diff --git a/Engine/Source/Developer/DerivedDataCache/Private/DerivedDataCacheStoreVerify.cpp b/Engine/Source/Developer/DerivedDataCache/Private/DerivedDataCacheStoreVerify.cpp
--- a/Engine/Source/Developer/DerivedDataCache/Private/DerivedDataCacheStoreVerify.cpp
+++ b/Engine/Source/Developer/DerivedDataCache/Private/DerivedDataCacheStoreVerify.cpp
@@ -140,7 +140,11 @@ private:
 	void GetDataComplete(IRequestOwner& Owner, FVerifyPutValueState* State, FCacheGetValueResponse&& Response);
 	void GetComplete(IRequestOwner& Owner, FVerifyPutValueState* State);
 
-	bool CompareRecords(const FCacheRecord& PutRecord, const FCacheRecord& GetRecord, const FSharedString& Name);
+	bool CompareRecords(
+		const FCacheRecord& PutRecord,
+		const FCacheRecord& GetRecord,
+		const FCacheRecordPolicy& Policy,
+		const FSharedString& Name);
 
 	static void LogChangedValue(
 		const FSharedString& Name,
@@ -203,7 +207,7 @@ void FCacheStoreVerify::Put(
 		for (const FCachePutRequest& Request : Requests)
 		{
 			const FCacheKey& Key = Request.Record.GetKey();
-			bool bForward = !Filter.IsMatch(Key);
+			bool bForward = EnumHasAnyFlags(Request.Policy.GetRecordPolicy(), ECachePolicy::NonDeterministic) || !Filter.IsMatch(Key);
 			if (!bForward)
 			{
 				AlreadyTested.Add(Key, &bForward);
@@ -304,7 +308,7 @@ void FCacheStoreVerify::GetDataComplete(IRequestOwner& Owner, FVerifyPutState* S
 	if ((Response.Status == EStatus::Ok) ||
 		(Response.Status == EStatus::Error && !Response.Record.GetValues().IsEmpty()))
 	{
-		if (CompareRecords(Request.Record, Response.Record, Request.Name))
+		if (CompareRecords(Request.Record, Response.Record, Request.Policy, Request.Name))
 		{
 			UE_LOG(LogDerivedDataCache, Verbose,
 				TEXT("Verify: Data in the cache matches newly generated data for %s from '%s'."),
@@ -370,7 +374,8 @@ void FCacheStoreVerify::Get(
 		FScopeLock Lock(&AlreadyTestedLock);
 		for (const FCacheGetRequest& Request : Requests)
 		{
-			const bool bForward = !Filter.IsMatch(Request.Key) || AlreadyTested.Contains(Request.Key);
+			const bool bForward = EnumHasAnyFlags(Request.Policy.GetRecordPolicy(), ECachePolicy::NonDeterministic) ||
+				!Filter.IsMatch(Request.Key) || AlreadyTested.Contains(Request.Key);
 			(bForward ? ForwardRequests : VerifyRequests).Add(Request);
 		}
 	}
@@ -394,7 +399,7 @@ void FCacheStoreVerify::PutValue(
 		FScopeLock Lock(&AlreadyTestedLock);
 		for (const FCachePutValueRequest& Request : Requests)
 		{
-			bool bForward = !Filter.IsMatch(Request.Key);
+			bool bForward = EnumHasAnyFlags(Request.Policy, ECachePolicy::NonDeterministic) || !Filter.IsMatch(Request.Key);
 			if (!bForward)
 			{
 				AlreadyTested.Add(Request.Key, &bForward);
@@ -566,7 +571,8 @@ void FCacheStoreVerify::GetValue(
 		FScopeLock Lock(&AlreadyTestedLock);
 		for (const FCacheGetValueRequest& Request : Requests)
 		{
-			const bool bForward = !Filter.IsMatch(Request.Key) || AlreadyTested.Contains(Request.Key);
+			const bool bForward = EnumHasAnyFlags(Request.Policy, ECachePolicy::NonDeterministic) ||
+				!Filter.IsMatch(Request.Key) || AlreadyTested.Contains(Request.Key);
 			(bForward ? ForwardRequests : VerifyRequests).Add(Request);
 		}
 	}
@@ -592,7 +598,8 @@ void FCacheStoreVerify::GetChunks(
 		FScopeLock Lock(&AlreadyTestedLock);
 		for (const FCacheGetChunkRequest& Request : Requests)
 		{
-			const bool bForward = !Filter.IsMatch(Request.Key) || AlreadyTested.Contains(Request.Key);
+			const bool bForward = EnumHasAnyFlags(Request.Policy, ECachePolicy::NonDeterministic) ||
+				!Filter.IsMatch(Request.Key) || AlreadyTested.Contains(Request.Key);
 			(bForward ? ForwardRequests : VerifyRequests).Add(Request);
 		}
 	}
@@ -605,7 +612,11 @@ void FCacheStoreVerify::GetChunks(
 	}
 }
 
-bool FCacheStoreVerify::CompareRecords(const FCacheRecord& PutRecord, const FCacheRecord& GetRecord, const FSharedString& Name)
+bool FCacheStoreVerify::CompareRecords(
+	const FCacheRecord& PutRecord,
+	const FCacheRecord& GetRecord,
+	const FCacheRecordPolicy& Policy,
+	const FSharedString& Name)
 {
 	bool bEqual = true;
 
@@ -643,12 +654,25 @@ bool FCacheStoreVerify::CompareRecords(const FCacheRecord& PutRecord, const FCac
 		{
 			if (PutIt->GetRawHash() != GetIt->GetRawHash())
 			{
-				LogChangedValue(Name, Key, PutIt->GetId(),
-					PutIt->GetRawHash(), GetIt->GetRawHash(),
-					PutIt->GetData().DecompressToComposite(), GetIt->GetData().DecompressToComposite());
-				bEqual = false;
-				ValuesDifferent++;
-				ValuesDifferentBytes += PutIt->GetRawSize();
+				if (!EnumHasAnyFlags(Policy.GetValuePolicy(PutIt->GetId()), ECachePolicy::NonDeterministic))
+				{
+					LogChangedValue(Name, Key, PutIt->GetId(),
+						PutIt->GetRawHash(), GetIt->GetRawHash(),
+						PutIt->GetData().DecompressToComposite(), GetIt->GetData().DecompressToComposite());
+					bEqual = false;
+					ValuesDifferent++;
+					ValuesDifferentBytes += PutIt->GetRawSize();
+				}
+				else
+				{
+					UE_LOG(LogDerivedDataCache, Verbose,
+						TEXT("Verify: Value %s has hash %s in the newly generated data and hash %s in the cache for %s from '%s'. "
+							 "The value is expected to be non-deterministic."),
+						*WriteToString<32>(PutIt->GetId()), *WriteToString<48>(PutIt->GetRawHash()),
+						*WriteToString<48>(GetIt->GetRawHash()), *WriteToString<96>(Key), *Name);
+					ValuesMatch++;
+					ValuesMatchBytes += PutIt->GetRawSize();
+				}
 			}
 			else
 			{
diff --git a/Engine/Source/Developer/DerivedDataCache/Private/FileSystemCacheStore.cpp b/Engine/Source/Developer/DerivedDataCache/Private/FileSystemCacheStore.cpp
--- a/Engine/Source/Developer/DerivedDataCache/Private/FileSystemCacheStore.cpp
+++ b/Engine/Source/Developer/DerivedDataCache/Private/FileSystemCacheStore.cpp
@@ -1901,9 +1901,11 @@ bool FFileSystemCacheStore::PutCacheRecord(
 			// Save the new package if the existing package is invalid.
 			const FOptionalCacheRecord ExistingRecord = FCacheRecord::Load(ExistingPackage);
 			bSavePackage |= !ExistingRecord;
-			const auto MakeValueTuple = [](const FValueWithId& Value) -> TTuple<FValueId, FIoHash>
+			const auto MakeValueTuple = [&Policy](const FValueWithId& Value) -> TTuple<FValueId, FIoHash>
 			{
-				return MakeTuple(Value.GetId(), Value.GetRawHash());
+				return !EnumHasAnyFlags(Policy.GetValuePolicy(Value.GetId()), ECachePolicy::NonDeterministic)
+					? MakeTuple(Value.GetId(), Value.GetRawHash())
+					: TTuple<FValueId, FIoHash>{};
 			};
 			if (ExistingRecord && !Algo::CompareBy(ExistingRecord.Get().GetValues(), Record.GetValues(), MakeValueTuple))
 			{
@@ -2206,7 +2208,8 @@ bool FFileSystemCacheStore::PutCacheValue(
 					*CachePath, *WriteToString<96>(Key), Name.Len(), Name.GetData());
 				bSavePackage = true;
 			}
-			else if (!(RawHash == Value.GetRawHash() && RawSize == Value.GetRawSize()))
+			else if (!EnumHasAnyFlags(Policy, ECachePolicy::NonDeterministic) &&
+				!(RawHash == Value.GetRawHash() && RawSize == Value.GetRawSize()))
 			{
 				// Content differs between the existing value and the new value.
 				UE_LOG(LogDerivedDataCache, Display, TEXT("%s: Cache put found non-deterministic value "
diff --git a/Engine/Source/Developer/DerivedDataCache/Private/HttpCacheStore.cpp b/Engine/Source/Developer/DerivedDataCache/Private/HttpCacheStore.cpp
--- a/Engine/Source/Developer/DerivedDataCache/Private/HttpCacheStore.cpp
+++ b/Engine/Source/Developer/DerivedDataCache/Private/HttpCacheStore.cpp
@@ -6,6 +6,7 @@
 #if WITH_HTTP_DDC_BACKEND
 
 #include "Algo/AllOf.h"
+#include "Algo/NoneOf.h"
 #include "Algo/Transform.h"
 #include "Async/ManualResetEvent.h"
 #include "Async/Mutex.h"
@@ -801,13 +802,13 @@ private:
 
 	FPutPackageOp(FHttpCacheStore& CacheStore, IRequestOwner& Owner, const FSharedString& Name);
 
-	void BeginOperation(bool bFinalize, bool bAllowOverwrite, FOnCachePutRefComplete&& OnComplete);
+	void BeginOperation(bool bFinalize, bool bAllowOverwrite, bool bReportConflicts, FOnCachePutRefComplete&& OnComplete);
 
-	void BeginPutRef(TUniquePtr<FHttpOperation> Operation, bool bFinalize, bool bAllowOverwrite, FOnCachePutRefComplete&& OnComplete);
-	void EndPutRef(TUniquePtr<FHttpOperation> Operation, bool bFinalize, bool bAllowOverwrite, FOnCachePutRefComplete&& OnComplete);
+	void BeginPutRef(TUniquePtr<FHttpOperation> Operation, bool bFinalize, bool bAllowOverwrite, bool bReportConflicts, FOnCachePutRefComplete&& OnComplete);
+	void EndPutRef(TUniquePtr<FHttpOperation> Operation, bool bFinalize, bool bReportConflicts, FOnCachePutRefComplete&& OnComplete);
 
-	void BeginPutBlobs(FCbPackage&& Package, FCachePutRefResponse&& Response);
-	void EndPutBlob(FHttpOperation* Operation, uint64 LogicalSize);
+	void BeginPutBlobs(FCbPackage&& Package, FCachePutRefResponse&& Response, bool bReportConflicts);
+	void EndPutBlob(FHttpOperation* Operation, uint64 LogicalSize, bool bReportConflicts);
 
 	void EndPutRefFinalize(FCachePutRefResponse&& Response);
 
@@ -824,28 +825,34 @@ FHttpCacheStore::FPutPackageOp::FPutPackageOp(FHttpCacheStore& InCacheStore, IRe
 
 void FHttpCacheStore::FPutPackageOp::Put(const FCacheKey& InKey, const FCacheRecordPolicy& Policy, FCbPackage&& Package, FOnPackageComplete&& OnComplete)
 {
-	ECachePolicy RecordPolicy = Policy.GetRecordPolicy();
-	const bool bAllowOverwrite = !EnumHasAnyFlags(RecordPolicy, ECachePolicy::QueryRemote);
+	const bool bAllowOverwrite = !EnumHasAnyFlags(Policy.GetRecordPolicy(), ECachePolicy::QueryRemote);
+	const bool bReportConflicts = !EnumHasAnyFlags(Policy.GetBasePolicy(), ECachePolicy::NonDeterministic) &&
+		Algo::NoneOf(Policy.GetValuePolicies(), [](const FCacheValuePolicy& Value) { return EnumHasAnyFlags(Value.Policy, ECachePolicy::NonDeterministic); });
 	Key = InKey;
 	Object = Package.GetObject();
 	ObjectHash = Package.GetObjectHash();
 	OnPackageComplete = MoveTemp(OnComplete);
-	BeginOperation(/*bFinalize*/ false, bAllowOverwrite, [Self = TRefCountPtr(this), Package = MoveTemp(Package)](FCachePutRefResponse&& Response) mutable
+	BeginOperation(/*bFinalize*/ false, bAllowOverwrite, bReportConflicts, [Self = TRefCountPtr(this), Package = MoveTemp(Package), bReportConflicts](FCachePutRefResponse&& Response) mutable
 	{
-		Self->BeginPutBlobs(MoveTemp(Package), MoveTemp(Response));
+		Self->BeginPutBlobs(MoveTemp(Package), MoveTemp(Response), bReportConflicts);
 	});
 }
 
-void FHttpCacheStore::FPutPackageOp::BeginOperation(bool bFinalize, bool bAllowOverwrite, FOnCachePutRefComplete&& OnComplete)
+void FHttpCacheStore::FPutPackageOp::BeginOperation(bool bFinalize, bool bAllowOverwrite, bool bReportConflicts, FOnCachePutRefComplete&& OnComplete)
 {
 	CacheStore.WaitForHttpOperationAsync(Owner, bFinalize ? EOperationCategory::PutFinalize : EOperationCategory::PutRef,
-	[Self = TRefCountPtr(this), bFinalize, bAllowOverwrite, OnComplete = MoveTemp(OnComplete)](TUniquePtr<FHttpOperation>&& Operation) mutable
+	[Self = TRefCountPtr(this), bFinalize, bAllowOverwrite, bReportConflicts, OnComplete = MoveTemp(OnComplete)](TUniquePtr<FHttpOperation>&& Operation) mutable
 	{
-		Self->BeginPutRef(MoveTemp(Operation), bFinalize, bAllowOverwrite, MoveTemp(OnComplete));
+		Self->BeginPutRef(MoveTemp(Operation), bFinalize, bAllowOverwrite, bReportConflicts, MoveTemp(OnComplete));
 	});
 }
 
-void FHttpCacheStore::FPutPackageOp::BeginPutRef(TUniquePtr<FHttpOperation> Operation, bool bFinalize, bool bAllowOverwrite, FOnCachePutRefComplete&& OnComplete)
+void FHttpCacheStore::FPutPackageOp::BeginPutRef(
+	TUniquePtr<FHttpOperation> Operation,
+	bool bFinalize,
+	bool bAllowOverwrite,
+	bool bReportConflicts,
+	FOnCachePutRefComplete&& OnComplete)
 {
 	if (UNLIKELY(!Operation))
 	{
@@ -896,17 +903,17 @@ void FHttpCacheStore::FPutPackageOp::BeginPutRef(TUniquePtr<FHttpOperation> Oper
 	LocalOperation.AddAcceptType(EHttpMediaType::Json);
 
 	RequestTimer.Stop();
-	LocalOperation.SendAsync(Owner, [Self = TRefCountPtr(this), Operation = MoveTemp(Operation), bFinalize, bAllowOverwrite, OnComplete = MoveTemp(OnComplete)]() mutable
+	LocalOperation.SendAsync(Owner, [Self = TRefCountPtr(this), Operation = MoveTemp(Operation), bFinalize, bReportConflicts, OnComplete = MoveTemp(OnComplete)]() mutable
 	{
 		Operation->GetStats(Self->RequestStats);
-		Self->EndPutRef(MoveTemp(Operation), bFinalize, bAllowOverwrite, MoveTemp(OnComplete));
+		Self->EndPutRef(MoveTemp(Operation), bFinalize, bReportConflicts, MoveTemp(OnComplete));
 	});
 }
 
 void FHttpCacheStore::FPutPackageOp::EndPutRef(
 	TUniquePtr<FHttpOperation> Operation,
 	bool bFinalize,
-	bool bAllowOverwrite,
+	bool bReportConflicts,
 	FOnCachePutRefComplete&& OnComplete)
 {
 	TRACE_CPUPROFILER_EVENT_SCOPE(HttpDDC_PutPackage_EndPutRef);
@@ -924,12 +931,12 @@ void FHttpCacheStore::FPutPackageOp::EndPutRef(
 
 	if (const int32 StatusCode = Operation->GetStatusCode(); StatusCode < 200 || StatusCode > 204)
 	{
-		if (!bAllowOverwrite && (StatusCode == 409))
+		if (bReportConflicts && (StatusCode == 409))
 		{
 			UE_LOG(LogDerivedDataCache, Display, TEXT("%s: Cache put found non-deterministic record for %s from '%s'"),
 				*CacheStore.NodeName, *WriteToString<96>(Key), *Name);
 		}
-		OnComplete({{}, EStatus::Error});
+		OnComplete({{}, StatusCode == 409 ? EStatus::Ok : EStatus::Error});
 		return;
 	}
 
@@ -972,7 +979,7 @@ void FHttpCacheStore::FPutPackageOp::EndPutRef(
 	OnComplete({NeededBlobHashes, EStatus::Ok});
 }
 
-void FHttpCacheStore::FPutPackageOp::BeginPutBlobs(FCbPackage&& Package, FCachePutRefResponse&& Response)
+void FHttpCacheStore::FPutPackageOp::BeginPutBlobs(FCbPackage&& Package, FCachePutRefResponse&& Response, bool bReportConflicts)
 {
 	if (Response.Status != EStatus::Ok)
 	{
@@ -1044,11 +1051,11 @@ void FHttpCacheStore::FPutPackageOp::BeginPutBlobs(FCbPackage&& Package, FCacheP
 	FRequestBarrier Barrier(Owner);
 	for (const FCompressedBuffer& Blob : Blobs)
 	{
-		CacheStore.WaitForHttpOperationAsync(Owner, EOperationCategory::PutBlobs, [Self = TRefCountPtr(this), Blob](TUniquePtr<FHttpOperation>&& Operation)
+		CacheStore.WaitForHttpOperationAsync(Owner, EOperationCategory::PutBlobs, [Self = TRefCountPtr(this), Blob, bReportConflicts](TUniquePtr<FHttpOperation>&& Operation)
 		{
 			if (UNLIKELY(!Operation))
 			{
-				Self->EndPutBlob(nullptr, 0);
+				Self->EndPutBlob(nullptr, 0, bReportConflicts);
 				return;
 			}
 
@@ -1058,16 +1065,16 @@ void FHttpCacheStore::FPutPackageOp::BeginPutBlobs(FCbPackage&& Package, FCacheP
 			LocalOperation.SetMethod(EHttpMethod::Put);
 			LocalOperation.SetContentType(EHttpMediaType::CompressedBinary);
 			LocalOperation.SetBody(Blob.GetCompressed());
-			LocalOperation.SendAsync(Self->Owner, [Self, Operation = MoveTemp(Operation), LogicalSize = Blob.GetRawSize()]
+			LocalOperation.SendAsync(Self->Owner, [Self, Operation = MoveTemp(Operation), LogicalSize = Blob.GetRawSize(), bReportConflicts]
 			{
 				Operation->GetStats(Self->RequestStats);
-				Self->EndPutBlob(Operation.Get(), LogicalSize);
+				Self->EndPutBlob(Operation.Get(), LogicalSize, bReportConflicts);
 			});
 		});
 	}
 }
 
-void FHttpCacheStore::FPutPackageOp::EndPutBlob(FHttpOperation* Operation, uint64 LogicalSize)
+void FHttpCacheStore::FPutPackageOp::EndPutBlob(FHttpOperation* Operation, uint64 LogicalSize, bool bReportConflicts)
 {
 	if (Operation)
 	{
@@ -1089,7 +1096,7 @@ void FHttpCacheStore::FPutPackageOp::EndPutBlob(FHttpOperation* Operation, uint6
 		}
 		else if (LocalSuccessfulBlobUploads == TotalBlobUploads)
 		{
-			BeginOperation(/*bFinalize*/ true, /*bAllowOverwrite*/ false, [Self = TRefCountPtr(this)](FCachePutRefResponse&& Response)
+			BeginOperation(/*bFinalize*/ true, /*bAllowOverwrite*/ false, bReportConflicts, [Self = TRefCountPtr(this)](FCachePutRefResponse&& Response)
 			{
 				Self->EndPutRefFinalize(MoveTemp(Response));
 			});
diff --git a/Engine/Source/Developer/DerivedDataCache/Public/DerivedDataCachePolicy.h b/Engine/Source/Developer/DerivedDataCache/Public/DerivedDataCachePolicy.h
--- a/Engine/Source/Developer/DerivedDataCache/Public/DerivedDataCachePolicy.h
+++ b/Engine/Source/Developer/DerivedDataCache/Public/DerivedDataCachePolicy.h
@@ -97,6 +97,17 @@ enum class ECachePolicy : uint32
 	 * This flag is meant to be used when subsequent accesses will not tolerate a cache miss.
 	 */
 	KeepAlive       = 1 << 7,
+
+	/**
+	 * Values associated with this policy are known to be non-deterministic.
+	 *
+	 * This is meant to be applied to values that are expected to differ between builds.
+	 * Disables reporting of non-determinism in put requests and disables verification of
+	 * determinism for associated values when running in verify mode.
+	 *
+	 * DO NOT APPLY THIS FLAG TO VALUES THAT ARE WRITTEN TO COOKED PACKAGES.
+	 */
+	NonDeterministic = 1 << 8,
 };
 
 ENUM_CLASS_FLAGS(ECachePolicy);
@@ -130,7 +141,7 @@ struct FCacheValuePolicy
 	ECachePolicy Policy = ECachePolicy::Default;
 
 	/** Flags that are valid on a value policy. */
-	static constexpr ECachePolicy PolicyMask = ECachePolicy::Default | ECachePolicy::SkipData;
+	static constexpr ECachePolicy PolicyMask = ECachePolicy::Default | ECachePolicy::SkipData | ECachePolicy::NonDeterministic;
 };
 
 UE_API FCbWriter& operator<<(FCbWriter& Writer, const FCacheValuePolicy& Policy);
