Updates to the Zen DDC client code to take advantage of server-side changes to allow for more accurate reporting of non-determinism issues and correct status when requesting a partial record that doesn't have all values present on the server.

diff --git a/Engine/Source/Developer/DerivedDataCache/Private/ZenCacheStore.cpp b/Engine/Source/Developer/DerivedDataCache/Private/ZenCacheStore.cpp
--- a/Engine/Source/Developer/DerivedDataCache/Private/ZenCacheStore.cpp
+++ b/Engine/Source/Developer/DerivedDataCache/Private/ZenCacheStore.cpp
@@ -6,6 +6,7 @@
 
 #if UE_WITH_ZEN
 
+#include "Algo/Compare.h"
 #include "Async/ManualResetEvent.h"
 #include "Async/UniqueLock.h"
 #include "BatchView.h"
@@ -474,10 +475,12 @@ public:
 				if (!bCanceled && HttpResponse->GetErrorCode() == EHttpErrorCode::None && HttpResponse->GetStatusCode() >= 200 && HttpResponse->GetStatusCode() <= 299)
 				{
 					const FCbObject& ResponseObj = Response.GetObject();
+					FCbFieldIterator DetailsFieldIt = ResponseObj[ANSITEXTVIEW("Details")].AsArray().CreateIterator();
 					for (FCbField ResponseField : ResponseObj[ANSITEXTVIEW("Result")])
 					{
 						if (RequestIndex >= Batch.Num())
 						{
+							++DetailsFieldIt;
 							++RequestIndex;
 							continue;
 						}
@@ -486,13 +489,16 @@ public:
 
 						const FCacheKey& Key = RequestWithStats.Request.Record.GetKey();
 						bool bPutSucceeded = ResponseField.AsBool();
+
 						if (CacheStore.DebugOptions.ShouldSimulatePutMiss(Key))
 						{
 							UE_LOG(LogDerivedDataCache, Verbose, TEXT("%s: Simulated miss for put of %s from '%s'"),
 								*CacheStore.GetName(), *WriteToString<96>(Key), *RequestWithStats.Request.Name);
 							bPutSucceeded = false;
 						}
-						bPutSucceeded ? OnHit(RequestWithStats) : OnMiss(RequestWithStats);
+
+						bPutSucceeded ? OnHit(RequestWithStats, DetailsFieldIt.AsObject()) : OnMiss(RequestWithStats, DetailsFieldIt.AsObject());
+						++DetailsFieldIt;
 					}
 					if (RequestIndex != Batch.Num())
 					{
@@ -516,7 +522,7 @@ public:
 					}
 					else
 					{
-						OnMiss(RequestWithStats);
+						OnMiss(RequestWithStats, FCbObject());
 					}
 				}
 			};
@@ -542,8 +548,28 @@ private:
 		return EBatchView::Continue;
 	}
 
-	void OnHit(const TRequestWithStats<FCachePutRequest>& RequestWithStats)
+	void OnHit(const TRequestWithStats<FCachePutRequest>& RequestWithStats, FCbObject&& DetailsObject)
 	{
+		if (DetailsObject)
+		{
+			FCbObject ExistingRecordObject(DetailsObject[ANSITEXTVIEW("Record")].AsObject());
+			if (ExistingRecordObject)
+			{
+				FCbPackage ExistingRecordPackage(ExistingRecordObject);
+				const FOptionalCacheRecord ExistingRecord = FCacheRecord::Load(ExistingRecordPackage);
+				const auto MakeValueTuple = [&RequestWithStats](const FValueWithId& Value) -> TTuple<FValueId, FIoHash>
+				{
+					return !EnumHasAnyFlags(RequestWithStats.Request.Policy.GetValuePolicy(Value.GetId()), ECachePolicy::NonDeterministic)
+						? MakeTuple(Value.GetId(), Value.GetRawHash())
+						: TTuple<FValueId, FIoHash>{};
+				};
+				if (ExistingRecord && !Algo::CompareBy(ExistingRecord.Get().GetValues(), RequestWithStats.Request.Record.GetValues(), MakeValueTuple))
+				{
+					UE_LOG(LogDerivedDataCache, Display, TEXT("%s: Cache put found non-deterministic record for %s from '%s'"),
+						*CacheStore.GetName(), *WriteToString<96>(RequestWithStats.Request.Record.GetKey()), *RequestWithStats.Request.Name);
+				}
+			}
+		}
 		UE_LOG(LogDerivedDataCache, Verbose, TEXT("%s: Cache put complete for %s from '%s'"),
 			*CacheStore.GetName(), *WriteToString<96>(RequestWithStats.Request.Record.GetKey()), *RequestWithStats.Request.Name);
 
@@ -563,10 +589,19 @@ private:
 		OnComplete(RequestWithStats.Request.MakeResponse(EStatus::Ok));
 	}
 
-	void OnMiss(const TRequestWithStats<FCachePutRequest>& RequestWithStats)
+	void OnMiss(const TRequestWithStats<FCachePutRequest>& RequestWithStats, FCbObject&& DetailsObject)
 	{
-		UE_LOG(LogDerivedDataCache, Verbose, TEXT("%s: Cache put failed for '%s' from '%s'"),
-			*CacheStore.GetName(), *WriteToString<96>(RequestWithStats.Request.Record.GetKey()), *RequestWithStats.Request.Name);
+		FUtf8StringView DetailsMessage = DetailsObject[ANSITEXTVIEW("Message")].AsString();
+		if (DetailsObject && !DetailsMessage.IsEmpty())
+		{
+			UE_LOG(LogDerivedDataCache, Verbose, TEXT("%s: Cache put failed with message '%s' for '%s' from '%s'"),
+				*CacheStore.GetName(), *WriteToString<64>(DetailsMessage), *WriteToString<96>(RequestWithStats.Request.Record.GetKey()), *RequestWithStats.Request.Name);
+		}
+		else
+		{
+			UE_LOG(LogDerivedDataCache, Verbose, TEXT("%s: Cache put failed for '%s' from '%s'"),
+				*CacheStore.GetName(), *WriteToString<96>(RequestWithStats.Request.Record.GetKey()), *RequestWithStats.Request.Name);
+		}
 		RequestWithStats.EndRequest(CacheStore, EStatus::Error);
 		OnComplete(RequestWithStats.Request.MakeResponse(EStatus::Error));
 	}
@@ -684,11 +719,13 @@ public:
 				if (!bCanceled && HttpResponse->GetErrorCode() == EHttpErrorCode::None && HttpResponse->GetStatusCode() >= 200 && HttpResponse->GetStatusCode() <= 299)
 				{
 					const FCbObject& ResponseObj = Response.GetObject();
-						
+
+					FCbFieldViewIterator IncompleteFieldIt = ResponseObj[ANSITEXTVIEW("Incomplete")].AsArrayView().CreateViewIterator();
 					for (FCbField RecordField : ResponseObj[ANSITEXTVIEW("Result")])
 					{
 						if (RequestIndex >= Batch.Num())
 						{
+							++IncompleteFieldIt;
 							++RequestIndex;
 							continue;
 						}
@@ -707,7 +744,8 @@ public:
 						{
 							Record = FCacheRecord::Load(Response, RecordField.AsObject());
 						}
-						Record ? OnHit(RequestWithStats, MoveTemp(Record).Get()) : OnMiss(RequestWithStats);
+						Record ? OnHit(RequestWithStats, MoveTemp(Record).Get(), IncompleteFieldIt.AsBool()) : OnMiss(RequestWithStats);
+						++IncompleteFieldIt;
 					}
 					if (RequestIndex != Batch.Num())
 					{
@@ -741,7 +779,7 @@ public:
 	}
 
 private:
-	void OnHit(const TRequestWithStats<FCacheGetRequest>& RequestWithStats, FCacheRecord&& Record)
+	void OnHit(const TRequestWithStats<FCacheGetRequest>& RequestWithStats, FCacheRecord&& Record, bool bIncomplete)
 	{
 		FCacheRecordBuilder RecordBuilder(RequestWithStats.Request.Key);
 
@@ -757,6 +795,13 @@ private:
 
 		EStatus Status = EStatus::Ok;
 		const bool bPartialRecord = EnumHasAnyFlags(RecordPolicy, ECachePolicy::PartialRecord);
+		if (bPartialRecord && bIncomplete)
+		{
+			// The server will indicate that a record is incomplete because when requesting
+			// partial records and skipping data, the returned record is identical whether
+			// the data for the values is present on the server or not present on the server.
+			Status = EStatus::Error;
+		}
 		for (const FValueWithId& Value : Record.GetValues())
 		{
 			const FValueId& Id = Value.GetId();
@@ -912,10 +957,12 @@ public:
 				if (!bCanceled && HttpResponse->GetErrorCode() == EHttpErrorCode::None && HttpResponse->GetStatusCode() >= 200 && HttpResponse->GetStatusCode() <= 299)
 				{
 					const FCbObject& ResponseObj = Response.GetObject();
+					FCbFieldViewIterator DetailsFieldIt = ResponseObj[ANSITEXTVIEW("Details")].AsArrayView().CreateViewIterator();
 					for (FCbField ResponseField : ResponseObj[ANSITEXTVIEW("Result")])
 					{
 						if (RequestIndex >= Batch.Num())
 						{
+							++DetailsFieldIt;
 							++RequestIndex;
 							continue;
 						}
@@ -929,7 +976,8 @@ public:
 								*CacheStore.GetName(), *WriteToString<96>(RequestWithStats.Request.Key), *RequestWithStats.Request.Name);
 							bPutSucceeded = false;
 						}
-						bPutSucceeded ? OnHit(RequestWithStats) : OnMiss(RequestWithStats);
+						bPutSucceeded ? OnHit(RequestWithStats, DetailsFieldIt.AsObjectView()) : OnMiss(RequestWithStats, DetailsFieldIt.AsObjectView());
+						++DetailsFieldIt;
 					}
 					if (RequestIndex != Batch.Num())
 					{
@@ -953,7 +1001,7 @@ public:
 					}
 					else
 					{
-						OnMiss(RequestWithStats);
+						OnMiss(RequestWithStats, FCbObjectView());
 					}
 				}
 			};
@@ -975,8 +1023,19 @@ private:
 		return EBatchView::Continue;
 	}
 
-	void OnHit(const TRequestWithStats<FCachePutValueRequest>& RequestWithStats)
+	void OnHit(const TRequestWithStats<FCachePutValueRequest>& RequestWithStats, FCbObjectView DetailsObjectView)
 	{
+		if (DetailsObjectView && !EnumHasAnyFlags(RequestWithStats.Request.Policy, ECachePolicy::NonDeterministic))
+		{
+			const FValue& Value = RequestWithStats.Request.Value;
+			const FIoHash RawHash = DetailsObjectView["RawHash"].AsHash();
+			const uint64 RawSize = DetailsObjectView["RawSize"].AsUInt64(MAX_uint64);
+
+			UE_LOG(LogDerivedDataCache, Display, TEXT("%s: Cache put found non-deterministic value "
+				"with new hash %s and existing hash %s for %s from '%s'"),
+				*CacheStore.GetName(), *WriteToString<48>(Value.GetRawHash()), *WriteToString<48>(RawHash),
+				*WriteToString<96>(RequestWithStats.Request.Key), *RequestWithStats.Request.Name);
+		}
 		UE_LOG(LogDerivedDataCache, Verbose, TEXT("%s: Cache put complete for %s from '%s'"),
 			*CacheStore.GetName(), *WriteToString<96>(RequestWithStats.Request.Key), *RequestWithStats.Request.Name);
 
@@ -989,10 +1048,19 @@ private:
 		OnComplete(RequestWithStats.Request.MakeResponse(EStatus::Ok));
 	}
 
-	void OnMiss(const TRequestWithStats<FCachePutValueRequest>& RequestWithStats)
+	void OnMiss(const TRequestWithStats<FCachePutValueRequest>& RequestWithStats, FCbObjectView DetailsObjectView)
 	{
-		UE_LOG(LogDerivedDataCache, Verbose, TEXT("%s: Cache put failed for '%s' from '%s'"),
-			*CacheStore.GetName(), *WriteToString<96>(RequestWithStats.Request.Key), *RequestWithStats.Request.Name);
+		FUtf8StringView DetailsMessage = DetailsObjectView[ANSITEXTVIEW("Message")].AsString();
+		if (DetailsObjectView && !DetailsMessage.IsEmpty())
+		{
+			UE_LOG(LogDerivedDataCache, Verbose, TEXT("%s: Cache put failed for '%s' from '%s' with message '%s'"),
+				*CacheStore.GetName(), *WriteToString<96>(RequestWithStats.Request.Key), *RequestWithStats.Request.Name, *WriteToString<64>(DetailsMessage));
+		}
+		else
+		{
+			UE_LOG(LogDerivedDataCache, Verbose, TEXT("%s: Cache put failed for '%s' from '%s'"),
+				*CacheStore.GetName(), *WriteToString<96>(RequestWithStats.Request.Key), *RequestWithStats.Request.Name);
+		}
 		RequestWithStats.EndRequest(CacheStore, EStatus::Error);
 		OnComplete(RequestWithStats.Request.MakeResponse(EStatus::Error));
 	}
