DDC: Require that RawHash matches when present in GetChunks()

The optional RawHash was added as a hint but usage dictates that it would be safer to treat it as a requirement when present, otherwise callers may receive different data than they anticipate.

#FailedValidations Preflight

diff --git a/Engine/Source/Developer/DerivedDataCache/Private/FileSystemCacheStore.cpp b/Engine/Source/Developer/DerivedDataCache/Private/FileSystemCacheStore.cpp
--- a/Engine/Source/Developer/DerivedDataCache/Private/FileSystemCacheStore.cpp
+++ b/Engine/Source/Developer/DerivedDataCache/Private/FileSystemCacheStore.cpp
@@ -1833,7 +1833,11 @@ void FFileSystemCacheStore::GetChunks(
 						DeleteCacheContent(Request.Name, Value, ValueReader, ValueAr);
 					}
 				}
-				Status = bExistsOnly || Buffer.GetSize() == RawSize ? EStatus::Ok : EStatus::Error;
+				const bool bRawHashMatches = Request.RawHash.IsZero() || Request.RawHash == Value.GetRawHash();
+				Status = bRawHashMatches && (bExistsOnly || Buffer.GetSize() == RawSize) ? EStatus::Ok : EStatus::Error;
+				UE_CLOGFMT(!bRawHashMatches, LogDerivedDataCache, Verbose,
+					"{Cache}: Cache miss with mismatched value {Id} received hash {ReceivedHash} when expected hash {RawHash} for {Key} from '{Name}'",
+					CachePath, Request.Id, Value.GetRawHash(), Request.RawHash, Request.Key, Request.Name);
 			}
 		}
 		RequestStats.Status = Status;
diff --git a/Engine/Source/Developer/DerivedDataCache/Private/HttpCacheStore.cpp b/Engine/Source/Developer/DerivedDataCache/Private/HttpCacheStore.cpp
--- a/Engine/Source/Developer/DerivedDataCache/Private/HttpCacheStore.cpp
+++ b/Engine/Source/Developer/DerivedDataCache/Private/HttpCacheStore.cpp
@@ -30,6 +30,7 @@
 #include "Http/HttpClient.h"
 #include "Http/HttpHostBuilder.h"
 #include "IO/IoHash.h"
+#include "Logging/StructuredLog.h"
 #include "Memory/SharedBuffer.h"
 #include "Misc/App.h"
 #include "Misc/CommandLine.h"
@@ -2896,29 +2897,38 @@ void FHttpCacheStore::FinishChunkRequest(
 {
 	if (Status == EStatus::Ok)
 	{
-		const uint64 RawOffset = FMath::Min(Value.GetRawSize(), Request.RawOffset);
-		const uint64 RawSize = FMath::Min(Value.GetRawSize() - RawOffset, Request.RawSize);
-		UE_LOG(LogDerivedDataCache, Verbose, TEXT("%s: Cache hit for %s from '%s'"),
-			*NodeName, *WriteToString<96>(Request.Key, '/', Request.Id), *Request.Name);
-		FSharedBuffer Buffer;
-		const bool bExistsOnly = EnumHasAnyFlags(Request.Policy, ECachePolicy::SkipData);
-		if (!bExistsOnly)
+		if (Request.RawHash.IsZero() || Request.RawHash == Value.GetRawHash())
 		{
-			Buffer = ValueReader.Decompress(RawOffset, RawSize);
+			const uint64 RawOffset = FMath::Min(Value.GetRawSize(), Request.RawOffset);
+			const uint64 RawSize = FMath::Min(Value.GetRawSize() - RawOffset, Request.RawSize);
+			UE_LOG(LogDerivedDataCache, Verbose, TEXT("%s: Cache hit for %s from '%s'"),
+				*NodeName, *WriteToString<96>(Request.Key, '/', Request.Id), *Request.Name);
+			FSharedBuffer Buffer;
+			const bool bExistsOnly = EnumHasAnyFlags(Request.Policy, ECachePolicy::SkipData);
+			if (!bExistsOnly)
+			{
+				Buffer = ValueReader.Decompress(RawOffset, RawSize);
+			}
+			const EStatus ChunkStatus = bExistsOnly || Buffer.GetSize() == RawSize ? EStatus::Ok : EStatus::Error;
+			if (ChunkStatus == EStatus::Ok)
+			{
+				TRACE_COUNTER_INCREMENT(HttpDDC_GetHit);
+			}
+			SharedOnComplete.Get()({Request.Name, Request.Key, Request.Id, Request.RawOffset,
+				RawSize, Value.GetRawHash(), MoveTemp(Buffer), Request.UserData, ChunkStatus});
 		}
-		const EStatus ChunkStatus = bExistsOnly || Buffer.GetSize() == RawSize ? EStatus::Ok : EStatus::Error;
-		if (ChunkStatus == EStatus::Ok)
+		else
 		{
-			TRACE_COUNTER_INCREMENT(HttpDDC_GetHit);
+			UE_LOGFMT(LogDerivedDataCache, Verbose,
+				"{Cache}: Cache miss with mismatched value {Id} received hash {ReceivedHash} when expected hash {RawHash} for {Key} from '{Name}'",
+				NodeName, Request.Id, Value.GetRawHash(), Request.RawHash, Request.Key, Request.Name);
+			SharedOnComplete.Get()(Request.MakeResponse(EStatus::Error));
 		}
-		SharedOnComplete.Get()({ Request.Name, Request.Key, Request.Id, Request.RawOffset,
-			RawSize, Value.GetRawHash(), MoveTemp(Buffer), Request.UserData, ChunkStatus });
 	}
 	else
 	{
 		UE_LOG(LogDerivedDataCache, Verbose, TEXT("%s: Cache miss for %s from '%s'"),
 			*NodeName, *WriteToString<96>(Request.Key, '/', Request.Id), *Request.Name);
-
 		SharedOnComplete.Get()(Request.MakeResponse(Status));
 	}
 }
@@ -3026,7 +3036,8 @@ void FHttpCacheStore::GetChunkGroupAsync(
 
 					if (!bHasValue || IsValueDataReady(Value, Request.Policy))
 					{
-						FinishChunkRequest(Request, Response.Status, Value, NullReader, SharedOnComplete);
+						const EStatus Status = bHasValue ? Response.Status : EStatus::Error;
+						FinishChunkRequest(Request, Status, Value, NullReader, SharedOnComplete);
 					}
 					else
 					{
diff --git a/Engine/Source/Developer/DerivedDataCache/Private/MemoryCacheStore.cpp b/Engine/Source/Developer/DerivedDataCache/Private/MemoryCacheStore.cpp
--- a/Engine/Source/Developer/DerivedDataCache/Private/MemoryCacheStore.cpp
+++ b/Engine/Source/Developer/DerivedDataCache/Private/MemoryCacheStore.cpp
@@ -731,7 +731,8 @@ void FMemoryCacheStore::GetChunks(
 			}
 		}
 
-		if (bProcessHit && Request.RawOffset <= Value.GetRawSize())
+		if (bProcessHit && Request.RawOffset <= Value.GetRawSize() &&
+			(Request.RawHash.IsZero() || Request.RawHash == Value.GetRawHash()))
 		{
 			const uint64 RawSize = FMath::Min(Value.GetRawSize() - Request.RawOffset, Request.RawSize);
 			FSharedBuffer Buffer;
diff --git a/Engine/Source/Developer/DerivedDataCache/Private/PakFileCacheStore.cpp b/Engine/Source/Developer/DerivedDataCache/Private/PakFileCacheStore.cpp
--- a/Engine/Source/Developer/DerivedDataCache/Private/PakFileCacheStore.cpp
+++ b/Engine/Source/Developer/DerivedDataCache/Private/PakFileCacheStore.cpp
@@ -18,6 +18,7 @@
 #include "HAL/PlatformFile.h"
 #include "HAL/PlatformFileManager.h"
 #include "HashingArchiveProxy.h"
+#include "Logging/StructuredLog.h"
 #include "Misc/Compression.h"
 #include "Misc/FileHelper.h"
 #include "Misc/Paths.h"
@@ -854,7 +855,11 @@ void FPakFileCacheStore::GetChunks(
 					Buffer = ValueReader.Decompress(RawOffset, RawSize);
 					RequestStats.LogicalReadSize += Buffer.GetSize();
 				}
-				Status = bExistsOnly || Buffer.GetSize() == RawSize ? EStatus::Ok : EStatus::Error;
+				const bool bRawHashMatches = Request.RawHash.IsZero() || Request.RawHash == Value.GetRawHash();
+				Status = bRawHashMatches && (bExistsOnly || Buffer.GetSize() == RawSize) ? EStatus::Ok : EStatus::Error;
+				UE_CLOGFMT(!bRawHashMatches, LogDerivedDataCache, Verbose,
+					"{Cache}: Cache miss with mismatched value {Id} received hash {ReceivedHash} when expected hash {RawHash} for {Key} from '{Name}'",
+					CachePath, Request.Id, Value.GetRawHash(), Request.RawHash, Request.Key, Request.Name);
 			}
 		}
 		RequestStats.Status = Status;
diff --git a/Engine/Source/Developer/DerivedDataCache/Private/S3CacheStore.cpp b/Engine/Source/Developer/DerivedDataCache/Private/S3CacheStore.cpp
--- a/Engine/Source/Developer/DerivedDataCache/Private/S3CacheStore.cpp
+++ b/Engine/Source/Developer/DerivedDataCache/Private/S3CacheStore.cpp
@@ -23,6 +23,7 @@
 #include "HAL/PlatformFile.h"
 #include "HAL/PlatformFileManager.h"
 #include "HashingArchiveProxy.h"
+#include "Logging/StructuredLog.h"
 #include "Memory/SharedBuffer.h"
 #include "Misc/Base64.h"
 #include "Misc/Compression.h"
@@ -1842,7 +1843,11 @@ void FS3CacheStore::GetChunks(
 					Buffer = ValueReader.Decompress(RawOffset, RawSize);
 					RequestStats.LogicalReadSize += Buffer.GetSize();
 				}
-				Status = bExistsOnly || Buffer.GetSize() == RawSize ? EStatus::Ok : EStatus::Error;
+				const bool bRawHashMatches = Request.RawHash.IsZero() || Request.RawHash == Value.GetRawHash();
+				Status = bRawHashMatches && (bExistsOnly || Buffer.GetSize() == RawSize) ? EStatus::Ok : EStatus::Error;
+				UE_CLOGFMT(!bRawHashMatches, LogDerivedDataCache, Verbose,
+					"{Cache}: Cache miss with mismatched value {Id} received hash {ReceivedHash} when expected hash {RawHash} for {Key} from '{Name}'",
+					*GetName(), Request.Id, Value.GetRawHash(), Request.RawHash, Request.Key, Request.Name);
 			}
 		}
 		RequestStats.Status = Status;
