DDC: Use a recursive shared mutex in the hierarchy because some functions can be called recursively

diff --git a/Engine/Source/Developer/DerivedDataCache/Private/DerivedDataCacheStoreHierarchy.cpp b/Engine/Source/Developer/DerivedDataCache/Private/DerivedDataCacheStoreHierarchy.cpp
--- a/Engine/Source/Developer/DerivedDataCache/Private/DerivedDataCacheStoreHierarchy.cpp
+++ b/Engine/Source/Developer/DerivedDataCache/Private/DerivedDataCacheStoreHierarchy.cpp
@@ -3,10 +3,10 @@
 #include "Algo/Accumulate.h"
 #include "Algo/AllOf.h"
 #include "Algo/AnyOf.h"
-#include "Algo/Compare.h"
 #include "Algo/Find.h"
 #include "Algo/MinElement.h"
-#include "Async/Mutex.h"
+#include "Async/SharedLock.h"
+#include "Async/SharedRecursiveMutex.h"
 #include "Async/UniqueLock.h"
 #include "Containers/Array.h"
 #include "Containers/StaticArray.h"
@@ -17,12 +17,10 @@
 #include "DerivedDataLegacyCacheStore.h"
 #include "DerivedDataRequest.h"
 #include "DerivedDataRequestOwner.h"
-#include "HAL/CriticalSection.h"
 #include "HAL/PlatformAtomics.h"
 #include "Logging/StructuredLog.h"
 #include "MemoryCacheStore.h"
 #include "Misc/EnumClassFlags.h"
-#include "Misc/ScopeRWLock.h"
 #include "ProfilingDebugging/CookStats.h"
 #include "Serialization/CompactBinary.h"
 #include "Templates/Invoke.h"
@@ -145,7 +143,7 @@ private:
 		TArray<FCacheStoreStats*> CacheStats;
 	};
 
-	mutable FRWLock NodesLock;
+	mutable FSharedRecursiveMutex NodesMutex;
 	std::atomic<ECacheStoreNodeFlags> CombinedNodeFlags{};
 	TArray<FCacheStoreNode, TInlineAllocator<8>> Nodes;
 	IMemoryCacheStore* MemoryCache;
@@ -302,7 +300,7 @@ FCacheStoreHierarchy::~FCacheStoreHierarchy()
 
 void FCacheStoreHierarchy::Add(ILegacyCacheStore* CacheStore, ECacheStoreFlags Flags)
 {
-	FWriteScopeLock Lock(NodesLock);
+	TUniqueLock Lock(NodesMutex);
 	checkf(!Algo::FindBy(Nodes, CacheStore, &FCacheStoreNode::Cache),
 		TEXT("Attempting to add a cache store that was previously registered to the hierarchy."));
 	TUniquePtr<ILegacyCacheStore> AsyncCacheStore(CreateCacheStoreAsync(CacheStore, /*MemoryCache*/ nullptr, /*bDeleteInnerCache*/ false));
@@ -314,7 +312,7 @@ void FCacheStoreHierarchy::SetFlags(ILegacyCacheStore* CacheStore, ECacheStoreFl
 {
 	// Doing this as a read lock and relying on atomic access to the Nodes' CacheFlags and the hierarchy's CombinedNodeFlags
 	// to prevent a partial read/write from two concurrent executions.
-	FReadScopeLock Lock(NodesLock);
+	TSharedLock Lock(NodesMutex);
 	FCacheStoreNode* Node = Algo::FindBy(Nodes, CacheStore, &FCacheStoreNode::Cache);
 	checkf(!!Node, TEXT("Attempting to set flags on a cache store that is not registered to the hierarchy."));
 	Node->SetCacheFlags(Flags);
@@ -323,7 +321,7 @@ void FCacheStoreHierarchy::SetFlags(ILegacyCacheStore* CacheStore, ECacheStoreFl
 
 void FCacheStoreHierarchy::RemoveNotSafe(ILegacyCacheStore* CacheStore)
 {
-	FWriteScopeLock Lock(NodesLock);
+	TUniqueLock Lock(NodesMutex);
 	FCacheStoreNode* Node = Algo::FindBy(Nodes, CacheStore, &FCacheStoreNode::Cache);
 	checkf(!!Node, TEXT("Attempting to remove a cache store that is not registered to the hierarchy."));
 	const int32 NodeIndex = UE_PTRDIFF_TO_INT32(Node - Nodes.GetData());
@@ -340,7 +338,7 @@ void FCacheStoreHierarchy::RemoveNotSafe(ILegacyCacheStore* CacheStore)
 
 bool FCacheStoreHierarchy::HasAllFlags(ECacheStoreFlags Flags) const
 {
-	FReadScopeLock Lock(NodesLock);
+	TSharedLock Lock(NodesMutex);
 	ECacheStoreFlags CombinedFlags = ECacheStoreFlags::None;
 	for (const FCacheStoreNode& Node : Nodes)
 	{
@@ -396,7 +394,7 @@ void FCacheStoreHierarchy::UpdateNodeFlags()
 
 ICacheStoreStats* FCacheStoreHierarchy::CreateStats(ILegacyCacheStore* CacheStore, ECacheStoreFlags Flags, FStringView Type, FStringView Name, FStringView Path)
 {
-	FWriteScopeLock Lock(NodesLock);
+	TUniqueLock Lock(NodesMutex);
 	for (FCacheStoreNode& Node : Nodes)
 	{
 		if (Node.Cache == CacheStore)
@@ -419,7 +417,7 @@ ICacheStoreStats* FCacheStoreHierarchy::CreateStats(ILegacyCacheStore* CacheStor
 
 void FCacheStoreHierarchy::DestroyStats(ICacheStoreStats* Stats)
 {
-	FWriteScopeLock Lock(NodesLock);
+	TUniqueLock Lock(NodesMutex);
 	for (FCacheStoreNode& Node : Nodes)
 	{
 		if (Node.CacheStats.RemoveSingle((FCacheStoreStats*)Stats))
@@ -665,7 +663,7 @@ void FCacheStoreHierarchy::TPutBatch<Params>::Begin(
 template <typename Params>
 void FCacheStoreHierarchy::TPutBatch<Params>::DispatchRequests()
 {
-	FReadScopeLock Lock(Hierarchy.NodesLock);
+	TSharedLock Lock(Hierarchy.NodesMutex);
 	FDynamicRequestBarrier Barrier;
 
 	for (const int32 NodeCount = Hierarchy.Nodes.Num(); NodePutIndex < NodeCount && !BatchOwner.IsCanceled(); ++NodePutIndex)
@@ -815,7 +813,7 @@ void FCacheStoreHierarchy::TPutBatch<Params>::CompletePutRequest(FPutResponse&&
 		check(!State.bFinished);
 		bool bCanQuery;
 		{
-			FReadScopeLock Lock(Hierarchy.NodesLock);
+			TSharedLock Lock(Hierarchy.NodesMutex);
 			bCanQuery = EnumHasAnyFlags(Hierarchy.Nodes[NodePutIndex].GetCacheFlags(), ECacheStoreFlags::Query);
 		}
 		if (bCanQuery)
@@ -927,7 +925,7 @@ void FCacheStoreHierarchy::TGetBatch<Params>::Begin(
 template <typename Params>
 void FCacheStoreHierarchy::TGetBatch<Params>::DispatchRequests()
 {
-	FReadScopeLock Lock(Hierarchy.NodesLock);
+	TSharedLock Lock(Hierarchy.NodesMutex);
 	FDynamicRequestBarrier Barrier;
 
 	TArray<FGetRequest, TInlineAllocator<8>> NodeRequests;
@@ -1055,7 +1053,7 @@ void FCacheStoreHierarchy::TGetBatch<Params>::CompleteRequest(FGetResponse&& Res
 		return;
 	}
 
-	FReadScopeLock Lock(Hierarchy.NodesLock);
+	TSharedLock Lock(Hierarchy.NodesMutex);
 
 	const FCacheStoreNode& Node = Hierarchy.Nodes[State.NodeIndex];
 	const bool bFirstOk = State.Response.Status == EStatus::Ok && PreviousStatus == EStatus::Error;
@@ -1471,7 +1469,7 @@ void FCacheStoreHierarchy::FGetChunksBatch::Begin(
 
 void FCacheStoreHierarchy::FGetChunksBatch::DispatchRequests()
 {
-	FReadScopeLock Lock(Hierarchy.NodesLock);
+	TSharedLock Lock(Hierarchy.NodesMutex);
 	FDynamicRequestBarrier Barrier;
 
 	TArray<FCacheGetChunkRequest, TInlineAllocator<8>> NodeRequests;
@@ -1598,7 +1596,7 @@ static void ConvertToLegacyStats(FDerivedDataCacheStatsNode& OutNode, FCacheStor
 
 void FCacheStoreHierarchy::LegacyStats(FDerivedDataCacheStatsNode& OutNode)
 {
-	FReadScopeLock Lock(NodesLock);
+	TSharedLock Lock(NodesMutex);
 	OutNode.Children.Reserve(Nodes.Num());
 	for (const FCacheStoreNode& Node : Nodes)
 	{
