DDC: Added -DDC-ReplayLoadSync to wait for loaded replays to finish

This can be important because waiting for the replay on exit instead can lead to authentication failures for Cloud DDC if the replay takes longer than the auth expiry, and it is impossible to refresh a token once the engine has started shutting down.

diff --git a/Engine/Source/Developer/DerivedDataCache/Internal/DerivedDataCacheReplay.h b/Engine/Source/Developer/DerivedDataCache/Internal/DerivedDataCacheReplay.h
--- a/Engine/Source/Developer/DerivedDataCache/Internal/DerivedDataCacheReplay.h
+++ b/Engine/Source/Developer/DerivedDataCache/Internal/DerivedDataCacheReplay.h
@@ -37,6 +37,8 @@ public:
 	UE_API bool ReadFromArchive(FArchive& ReplayAr, uint64 ScratchSize = DefaultScratchSize);
 	UE_API bool ReadFromObject(FCbObjectView Object);
 
+	UE_API void WaitForAsyncReads();
+
 private:
 	class FState;
 	TPimplPtr<FState> State;
diff --git a/Engine/Source/Developer/DerivedDataCache/Private/DerivedDataCacheReplay.cpp b/Engine/Source/Developer/DerivedDataCache/Private/DerivedDataCacheReplay.cpp
--- a/Engine/Source/Developer/DerivedDataCache/Private/DerivedDataCacheReplay.cpp
+++ b/Engine/Source/Developer/DerivedDataCache/Private/DerivedDataCacheReplay.cpp
@@ -844,6 +844,11 @@ bool FCacheReplayReader::ReadFromObject(const FCbObjectView Object)
 	return State->ReadFromObject(Object);
 }
 
+void FCacheReplayReader::WaitForAsyncReads()
+{
+	State->WaitForAsyncReads();
+}
+
 void FCacheReplayReader::SetKeyFilter(FCacheKeyFilter KeyFilter)
 {
 	State->KeyFilter = MoveTemp(KeyFilter);
@@ -988,6 +993,11 @@ ILegacyCacheStore* TryCreateCacheStoreReplay(ILegacyCacheStore* InnerCache)
 			}
 		}
 
+		if (FParse::Param(CommandLine, TEXT("DDC-ReplayLoadSync")))
+		{
+			Reader.WaitForAsyncReads();
+		}
+
 		ReplayStore->SetReader(MoveTemp(Reader));
 	}
 
