DDC: Added a private API for reading and parsing DDC config from GConfig

This is a prerequisite for migrating to the a new format for configuring the DDC.

diff --git a/Engine/Programs/LowLevelTests/Config/DefaultEngine.ini b/Engine/Programs/LowLevelTests/Config/DefaultEngine.ini
new file mode 100644
--- /dev/null
+++ b/Engine/Programs/LowLevelTests/Config/DefaultEngine.ini
@@ -0,0 +1,42 @@
+[DerivedDataCacheGraphs]
+TestInvalidCacheGraph1=Invalid
+TestInvalidCacheGraph2="Invalid
+TestInvalidCacheGraph3="Invalid\
+TestInvalidCacheGraph4="Invalid\"
+TestInvalidCacheGraph5=(Invalid
+TestInvalidCacheGraph6=(Invalid,
+TestInvalidCacheStore1=(TestInvalidStore1)
+TestInvalidCacheStore2=(TestInvalidStore1(Key1=Value1))
+TestInvalidCacheStore3=(TestInvalidStore2(Key1=Value1))
+TestCacheGraphInlineStore1 = ( TestInlineStore ( Key1 = Value1 ( A = 1 , B = 2, C ) , Key2 = Value2, Key3 = D (123) ) ) ; Comment
+TestCacheGraphInlineStore2 = ( TestInlineStore = TestMissing ( Key1 = Value1 ( A = 1 , B = 2, C ) , Key2 = Value2, Key3 = D (123) ) ) ; Comment
+TestCacheGraph0=()
+TestCacheGraph1=(TestStore1)
+TestCacheGraph2=(TestStore1,TestStore2,TestStore3)
+TestCacheGraph3=(TestStore1,TestStore2,TestStore3())
+TestCacheGraph4 = ( TestStore1 , TestStore2 , TestStore3 ) ; Comment
+TestCacheGraph5 = ( TestStore1 , TestStore2 ( ) , TestStore3 ( Key1 = Value1 , Key2 = Value2 ) ) ; Comment
+TestCacheGraph6 = ( TestStore1 = ( Key1 = Value1 ) ) ; Comment
+TestCacheGraph7 = ( TestStore1 = TestStore5 ( Key1 = Value1 ) ) ; Comment
+TestCacheGraph8 = ( TestStore1 = TestStore6 ) ; Comment
+TestCacheGraph9 = ( TestStore1 ( Key1 = Value1 ) , Deprecated = "Use TestCacheGraph1." ) ; Comment
+
+[DerivedDataCacheStores]
+TestInvalidStore1=Invalid
+TestInvalidStore2=(Base=TestInvalidStore1,Key2=Value2)
+TestInvalidStore3=(Base=TestStore1,Base=TestStore2)
+TestStore1=(Key1=Store1Value,Key2=Store1Value,ValueOnly,(X="A=1, \"B\"=2, C=3"))
+TestStore2=(Base=TestStore1,Key1=Store2Value)
+TestStore3 = ( ValueOnly , Key1 = Store3Value , Key2 = Store3Value , Key3 = ( A = 1 , B = 2 , C ) ) ; Comment
+TestStore4 = ( Key1 = Store4Value , Base = TestStore2 ) ; Comment
+TestStore5 = ( Key1 = Store5Value , Key2 = Store5Value ) ; Comment
+TestStore6 = ( Base = TestStore1 ) ; Comment
+TestStore7 = ( Base = TestStore1 , Deprecated = "Use TestStore1." ) ; Comment
+
+[TestLegacyCacheGraph1]
+Root=(Type=KeyLength, Length=120, Inner=Verify)
+Verify=(Type=Verify, Inner=AsyncPut)
+AsyncPut=(Type=AsyncPut, Inner=Hierarchy)
+Hierarchy=(Type=Hierarchical, Inner=TestLegacyStore1, Inner=TestStore2, Inner=TestStore3)
+TestLegacyStore1=(Key1=LegacyStore1Value, Key2=LegacyStore1Value)
+TestStore2=(Key1=LegacyStore2Value)
diff --git a/Engine/Source/Developer/DerivedDataCache/Private/DerivedDataConfig.cpp b/Engine/Source/Developer/DerivedDataCache/Private/DerivedDataConfig.cpp
new file mode 100644
--- /dev/null
+++ b/Engine/Source/Developer/DerivedDataCache/Private/DerivedDataConfig.cpp
@@ -0,0 +1,519 @@
+// Copyright Epic Games, Inc. All Rights Reserved.
+
+#include "DerivedDataConfig.h"
+
+#include "CoreGlobals.h"
+#include "DerivedDataPrivate.h"
+#include "Logging/StructuredLog.h"
+#include "Misc/AsciiSet.h"
+#include "Misc/ConfigCacheIni.h"
+#include "Misc/StringBuilder.h"
+
+namespace UE::DerivedData::Config
+{
+
+constexpr static TCHAR ObjectStartChar = '(';
+constexpr static TCHAR ObjectEndChar = ')';
+constexpr static TCHAR FieldDelimiterChar = ',';
+constexpr static TCHAR ValueDelimiterChar = '=';
+constexpr static TCHAR QuoteChar = '"';
+constexpr static FAsciiSet ObjectStartSet{{ObjectStartChar, TEXT('\0')}};
+constexpr static FAsciiSet ObjectEndSet{{ObjectEndChar, TEXT('\0')}};
+constexpr static FAsciiSet FieldDelimiterSet{{FieldDelimiterChar, TEXT('\0')}};
+constexpr static FAsciiSet QuoteSet{{QuoteChar, TEXT('\0')}};
+constexpr static FAsciiSet EscapeSet = "\\";
+constexpr static FAsciiSet WhiteSpaceSet = " \t";
+constexpr static FAsciiSet FieldNameSet = "0123456789ABCDEFGHIJKLMNOPQRSTUVWXYZabcdefghijklmnopqrstuvwxyz_";
+
+constexpr static FStringView FieldDelimiter = TEXTVIEW(", ");
+
+FConfigIterator::FConfigIterator(FStringView Config)
+{
+	Config = FAsciiSet::TrimPrefixWith(Config, WhiteSpaceSet);
+	if (Config.StartsWith(ObjectStartChar))
+	{
+		Tail = Config.RightChop(1);
+		++*this;
+	}
+	else
+	{
+		bAtEnd = true;
+		bError = true;
+	}
+}
+
+FConfigIterator& FConfigIterator::operator++()
+{
+	if (bLastField)
+	{
+		bAtEnd = true;
+	}
+
+	if (bAtEnd)
+	{
+		return *this;
+	}
+
+	// Try to parse the optional field name.
+	FStringView Config = Tail;
+	Config = FAsciiSet::TrimPrefixWith(Config, WhiteSpaceSet);
+	Name = FAsciiSet::FindPrefixWith(Config, FieldNameSet);
+	if (!Name.IsEmpty())
+	{
+		FStringView Rest = Config.RightChop(Name.Len());
+		Rest = FAsciiSet::TrimPrefixWith(Rest, WhiteSpaceSet);
+		if (Rest.StartsWith(ValueDelimiterChar))
+		{
+			Config = Rest.RightChop(1);
+			Config = FAsciiSet::TrimPrefixWith(Config, WhiteSpaceSet);
+		}
+		else
+		{
+			Name.Reset();
+		}
+	}
+
+	// Try to parse the field value by skipping objects and strings.
+	FStringView NextTail = Config;
+	for (int32 ObjectCount = 0;;)
+	{
+		if (ObjectCount)
+		{
+			NextTail = FAsciiSet::TrimPrefixWithout(NextTail, ObjectStartSet | ObjectEndSet | QuoteSet);
+		}
+		else
+		{
+			NextTail = FAsciiSet::TrimPrefixWithout(NextTail, FieldDelimiterSet | ObjectStartSet | ObjectEndSet | QuoteSet);
+		}
+
+		if (NextTail.IsEmpty())
+		{
+			bAtEnd = true;
+			bError = true;
+			return *this;
+		}
+
+		const TCHAR Control = *NextTail.GetData();
+		bLastField = (Control == ObjectEndChar);
+		NextTail.RightChopInline(1);
+
+		if (bLastField || Control == FieldDelimiterChar)
+		{
+			if (ObjectCount == 0)
+			{
+				break;
+			}
+			--ObjectCount;
+		}
+		else if (Control == ObjectStartChar)
+		{
+			++ObjectCount;
+		}
+		else if (Control == QuoteChar)
+		{
+			for (bool bFoundEnd = false; !bFoundEnd;)
+			{
+				NextTail = FAsciiSet::TrimPrefixWithout(NextTail, QuoteSet | EscapeSet);
+				bFoundEnd = NextTail.IsEmpty() || NextTail.StartsWith(QuoteChar);
+				NextTail.RightChopInline(bFoundEnd ? 1 : 2);
+			}
+		}
+	}
+
+	Tail = NextTail;
+	Value = Config.LeftChop(NextTail.Len() + 1);
+	Value = FAsciiSet::TrimSuffixWith(Value, WhiteSpaceSet);
+	bAtEnd = Name.IsEmpty() && Value.IsEmpty();
+	return *this;
+}
+
+FStringView ParseQuotedString(FStringView String, FStringBuilderBase& OutScratch)
+{
+	const bool bStartsWithQuote = String.StartsWith(QuoteChar);
+	const bool bEndsWithQuote = String.EndsWith(QuoteChar);
+	if (!bStartsWithQuote && !bEndsWithQuote)
+	{
+		// No Quotes
+		return String;
+	}
+	String.RightChopInline(bStartsWithQuote ? 1 : 0);
+	String.LeftChopInline(bEndsWithQuote ? 1 : 0);
+
+	FStringView Prefix = FAsciiSet::FindPrefixWithout(String, EscapeSet);
+	if (Prefix.Len() == String.Len())
+	{
+		// No Escapes
+		return String;
+	}
+
+	OutScratch.Reset();
+	do
+	{
+		OutScratch.Append(Prefix);
+		String.RightChopInline(Prefix.Len() + 1);
+		if (String.Len() > 0)
+		{
+			switch (TCHAR C = String.GetData()[0])
+			{
+			case '\'': OutScratch.AppendChar('\''); break;
+			case '\"': OutScratch.AppendChar('\"'); break;
+			case '\?': OutScratch.AppendChar('\?'); break;
+			case '\\': OutScratch.AppendChar('\\'); break;
+			case 'a': OutScratch.AppendChar('\a'); break;
+			case 'b': OutScratch.AppendChar('\b'); break;
+			case 'f': OutScratch.AppendChar('\f'); break;
+			case 'n': OutScratch.AppendChar('\n'); break;
+			case 'r': OutScratch.AppendChar('\r'); break;
+			case 't': OutScratch.AppendChar('\t'); break;
+			case 'v': OutScratch.AppendChar('\v'); break;
+			default: OutScratch.AppendChar(C); break;
+			}
+			String.RightChopInline(1);
+		}
+		Prefix = FAsciiSet::FindPrefixWithout(String, EscapeSet);
+	}
+	while (!String.IsEmpty());
+	return OutScratch;
+}
+
+bool FindFirstFieldValue(FStringView Config, FStringView Name, FStringView& OutValue)
+{
+	FConfigIterator It(Config);
+	for (; It; ++It)
+	{
+		if (Name == It.GetName())
+		{
+			OutValue = It.GetValue();
+			return true;
+		}
+	}
+	return false;
+}
+
+bool VisitFieldValues(FStringView Config, FStringView Name, TFunctionWithContext<void (FStringView)> OnValue)
+{
+	FConfigIterator It(Config);
+	for (; It; ++It)
+	{
+		if (Name == It.GetName())
+		{
+			OnValue(It.GetValue());
+		}
+	}
+	return !It.HasError();
+}
+
+bool TrySplitNameAndConfig(FStringView Value, FStringView& OutName, FStringView& OutConfig)
+{
+	Value = FAsciiSet::TrimPrefixWith(Value, WhiteSpaceSet);
+	Value = FAsciiSet::TrimSuffixWith(Value, WhiteSpaceSet);
+	const FStringView Name = FAsciiSet::FindPrefixWith(Value, FieldNameSet);
+
+	Value.RightChopInline(Name.Len());
+	Value = FAsciiSet::TrimPrefixWith(Value, WhiteSpaceSet);
+
+	if (Value.IsEmpty() || (Value.StartsWith(ObjectStartChar) && Value.EndsWith(ObjectEndChar)))
+	{
+		OutName = Name;
+		OutConfig = Value;
+		return true;
+	}
+	if (!Name.IsEmpty())
+	{
+		OutName = Name;
+		OutConfig.Reset();
+		return true;
+	}
+	return false;
+}
+
+void AppendConfig(FStringBuilderBase& OutConfig, FStringView Name, FStringView Value)
+{
+	if (!Name.IsEmpty())
+	{
+		OutConfig.Append(Name);
+		OutConfig.AppendChar(ValueDelimiterChar);
+	}
+
+	if (FStringView ValueName, Config; TrySplitNameAndConfig(Value, ValueName, Config))
+	{
+		OutConfig.Append(ValueName);
+		if (!Config.IsEmpty())
+		{
+			FConfigIterator It(Config);
+			const int32 ResetSize = OutConfig.Len();
+			OutConfig.AppendChar(ObjectStartChar);
+			for (bool bDelimit = false; It; ++It, bDelimit = true)
+			{
+				if (bDelimit)
+				{
+					OutConfig.Append(FieldDelimiter);
+				}
+				AppendConfig(OutConfig, It.GetName(), It.GetValue());
+			}
+			OutConfig.AppendChar(ObjectEndChar);
+			if (It.HasError())
+			{
+				OutConfig.RemoveSuffix(OutConfig.Len() - ResetSize);
+				OutConfig.Append(Config);
+			}
+		}
+	}
+	else
+	{
+		OutConfig.Append(Value);
+	}
+}
+
+class FConfigBuilder
+{
+public:
+	inline explicit FConfigBuilder(FStringBuilderBase& OutConfig)
+		: Out(OutConfig)
+	{
+		Out.AppendChar(ObjectStartChar);
+	}
+
+	inline ~FConfigBuilder()
+	{
+		if (Out.ToView().EndsWith(FieldDelimiter, ESearchCase::CaseSensitive))
+		{
+			Out.RemoveSuffix(FieldDelimiter.Len());
+		}
+
+		Out.AppendChar(ObjectEndChar);
+
+		if (bError)
+		{
+			Out.RemoveSuffix(Out.Len() - StartIndex);
+		}
+	}
+
+	inline void Append(FStringView Name, FStringView Value)
+	{
+		AppendConfig(Out, Name, Value);
+		Out.Append(FieldDelimiter);
+	}
+
+	inline void SetError() { bError = true; }
+	inline bool HasError() const { return bError; }
+
+private:
+	FStringBuilderBase& Out;
+	const int32 StartIndex = Out.Len();
+	bool bError = false;
+};
+
+} // UE::DerivedData::Config
+
+namespace UE::DerivedData
+{
+
+bool TryFindCacheGraphConfig(const TCHAR* GraphName, FStringBuilderBase& OutConfig)
+{
+	// UE_DEPRECATED(5.8, "Remove the LegacyConfig block after the standard deprecation period.")
+	FString LegacyConfig;
+	if (GConfig->GetString(GraphName, TEXT("Root"), LegacyConfig, GEngineIni))
+	{
+		//UE_LOGFMT(LogDerivedDataCache, Warning,
+		//	"Configuring cache graph {GraphName} in config section [{GraphName}] of {IniName}.ini is deprecated. "
+		//	"Migrate the configuration to the [DerivedDataCacheGraphs] section of {IniName}.ini.",
+		//	("GraphName", GraphName), ("IniName", GEngineIni));
+		for (FStringView Type, Inner = TEXTVIEW("Root");;)
+		{
+			if (Config::FindFirstFieldValue(LegacyConfig, TEXTVIEW("Type"), Type))
+			{
+				if (Type == TEXTVIEW("Hierarchical"))
+				{
+					Config::FConfigBuilder Builder(OutConfig);
+					Config::FConfigIterator It(LegacyConfig);
+					for (; It; ++It)
+					{
+						if (It.GetName() == TEXTVIEW("Inner"))
+						{
+							Builder.Append({}, It.GetValue());
+						}
+					}
+					if (It.HasError())
+					{
+						Builder.SetError();
+					}
+					return !Builder.HasError();
+				}
+				if (Type == TEXTVIEW("AsyncPut") || Type == TEXTVIEW("KeyLength") || Type == TEXTVIEW("Verify"))
+				{
+					UE_LOGFMT(LogDerivedDataCache, Warning,
+						"The cache graph {GraphName} contains cache store {StoreName} of deprecated type {StoreType}.",
+						GraphName, Inner, Type);
+					if (Config::FindFirstFieldValue(LegacyConfig, TEXTVIEW("Inner"), Inner) &&
+						GConfig->GetString(GraphName, *WriteToString<64>(Inner), LegacyConfig, GEngineIni))
+					{
+						continue;
+					}
+				}
+			}
+			UE_LOGFMT(LogDerivedDataCache, Warning,
+				"Failed to parse cache store {StoreName} from cache graph {GraphName} in config section [{GraphName}] of {IniName}.ini. "
+				"Migrate the configuration to the [DerivedDataCacheGraphs] section of {IniName}.ini.",
+				("GraphName", GraphName), ("StoreName", Inner), ("IniName", GEngineIni));
+			break;
+		}
+	}
+
+	FString Config;
+	if (GConfig->GetString(TEXT("DerivedDataCacheGraphs"), GraphName, Config, GEngineIni))
+	{
+		Config::FConfigBuilder Builder(OutConfig);
+		Config::FConfigIterator It(Config);
+		for (; It; ++It)
+		{
+			if (It.GetName() == TEXTVIEW("Deprecated"))
+			{
+				TStringBuilder<128> MessageScratch;
+				UE_LOGFMT(LogDerivedDataCache, Warning, "The cache graph {GraphName} is deprecated. {Message}",
+					GraphName, Config::ParseQuotedString(It.GetValue(), MessageScratch));
+			}
+			else if (FStringView StoreName = It.GetName(); !StoreName.IsEmpty())
+			{
+				Builder.Append({}, StoreName);
+			}
+			else if (FStringView StoreConfig; Config::TrySplitNameAndConfig(It.GetValue(), StoreName, StoreConfig))
+			{
+				Builder.Append({}, StoreName);
+			}
+			else
+			{
+				Builder.SetError();
+			}
+		}
+		if (It.HasError())
+		{
+			Builder.SetError();
+		}
+		return !Builder.HasError();
+	}
+	return false;
+}
+
+class FCacheStoreConfigBuilder
+{
+public:
+	explicit FCacheStoreConfigBuilder(FStringBuilderBase& OutConfig);
+
+	bool TryFind(const TCHAR* StoreName, const TCHAR* GraphName);
+
+private:
+	bool TryQuery(FStringView StoreName, const TCHAR* GraphName = nullptr);
+	bool TryBuild(FStringView StoreName, FStringView Config);
+
+	Config::FConfigBuilder Builder;
+};
+
+FCacheStoreConfigBuilder::FCacheStoreConfigBuilder(FStringBuilderBase& OutConfig)
+	: Builder(OutConfig)
+{
+}
+
+bool FCacheStoreConfigBuilder::TryFind(const TCHAR* StoreName, const TCHAR* GraphName)
+{
+	const bool bOk = TryQuery(StoreName, GraphName);
+	if (!bOk)
+	{
+		Builder.SetError();
+	}
+	return bOk;
+}
+
+bool FCacheStoreConfigBuilder::TryQuery(FStringView StoreName, const TCHAR* GraphName)
+{
+	// UE_DEPRECATED(5.8, "Remove the LegacyConfig block after the standard deprecation period.")
+	FString LegacyConfig;
+	if (GraphName && GConfig->GetString(GraphName, *WriteToString<64>(StoreName), LegacyConfig, GEngineIni))
+	{
+		//UE_LOGFMT(LogDerivedDataCache, Warning,
+		//	"Configuring cache store {StoreName} in config section [{GraphName}] of {IniName}.ini is deprecated. "
+		//	"Migrate the configuration to the [DerivedDataCacheStores] section of {IniName}.ini.",
+		//	("GraphName", GraphName), ("StoreName", StoreName), ("IniName", GEngineIni));
+		return TryBuild(StoreName, LegacyConfig);
+	}
+
+	bool bOk = true;
+	bool bHasConfig = false;
+
+	FString GraphConfig;
+	if (GraphName && GConfig->GetString(TEXT("DerivedDataCacheGraphs"), GraphName, GraphConfig, GEngineIni))
+	{
+		Config::FConfigIterator It(GraphConfig);
+		for (; It; ++It)
+		{
+			FStringView GraphStoreName, GraphStoreConfig;
+			if (Config::TrySplitNameAndConfig(It.GetValue(), GraphStoreName, GraphStoreConfig) &&
+				(It.GetName() == StoreName || (It.GetName().IsEmpty() && GraphStoreName == StoreName)))
+			{
+				if (!GraphStoreConfig.IsEmpty())
+				{
+					bHasConfig = true;
+					bOk &= TryBuild(StoreName, GraphStoreConfig);
+				}
+				StoreName = GraphStoreName;
+				if (StoreName.IsEmpty())
+				{
+					return bOk && bHasConfig;
+				}
+				break;
+			}
+		}
+	}
+
+	FString StoreConfig;
+	if (GConfig->GetString(TEXT("DerivedDataCacheStores"), *WriteToString<64>(StoreName), StoreConfig, GEngineIni))
+	{
+		bHasConfig = true;
+		bOk &= TryBuild(StoreName, StoreConfig);
+	}
+
+	return bOk && bHasConfig;
+}
+
+bool FCacheStoreConfigBuilder::TryBuild(FStringView StoreName, FStringView Config)
+{
+	bool bOk = true;
+	FStringView BaseName;
+
+	Config::FConfigIterator It(Config);
+	for (; It; ++It)
+	{
+		if (It.GetName() == TEXTVIEW("Base"))
+		{
+			bOk &= BaseName.IsEmpty(); // Only one Base is allowed.
+			BaseName = It.GetValue();
+		}
+		else if (It.GetName() == TEXTVIEW("Deprecated"))
+		{
+			TStringBuilder<128> MessageScratch;
+			UE_LOGFMT(LogDerivedDataCache, Warning, "The cache store {StoreName} is deprecated. {Message}",
+				StoreName, Config::ParseQuotedString(It.GetValue(), MessageScratch));
+		}
+		else
+		{
+			Builder.Append(It.GetName(), It.GetValue());
+		}
+	}
+	bOk &= !It.HasError();
+
+	// Process the base after appending the current config to ensure overrides come first.
+	if (bOk && !BaseName.IsEmpty())
+	{
+		bOk &= TryQuery(*WriteToString<64>(BaseName));
+	}
+
+	return bOk;
+}
+
+bool TryFindCacheStoreConfig(const TCHAR* StoreName, const TCHAR* GraphName, FStringBuilderBase& OutConfig)
+{
+	FCacheStoreConfigBuilder Builder(OutConfig);
+	return Builder.TryFind(StoreName, GraphName);
+}
+
+} // UE::DerivedData
diff --git a/Engine/Source/Developer/DerivedDataCache/Private/DerivedDataConfig.h b/Engine/Source/Developer/DerivedDataCache/Private/DerivedDataConfig.h
new file mode 100644
--- /dev/null
+++ b/Engine/Source/Developer/DerivedDataCache/Private/DerivedDataConfig.h
@@ -0,0 +1,90 @@
+// Copyright Epic Games, Inc. All Rights Reserved.
+
+#pragma once
+
+#include "CoreTypes.h"
+#include "Containers/StringFwd.h"
+#include "Containers/StringView.h"
+#include "Templates/FunctionWithContext.h"
+
+namespace UE::DerivedData::Config
+{
+
+/**
+ * Iterates fields of a config string in the form (Name1=Value1, Name2=Value2, Name3=(X=1, Y=2, Z=3), ValueOnly).
+ *
+ * Names must contain only [A-Za-z0-9_].
+ * Strings with spaces or quotes must be quoted and escaped: (Description="A \"real\" description...")
+ * Values that are a name and config are explicitly supported: (Node=Point(X=1, Y=2, Z=3), ValueOnly(Name="Test"))
+ */
+class FConfigIterator
+{
+public:
+	explicit FConfigIterator(FStringView Config);
+
+	FConfigIterator& operator++();
+
+	inline bool HasError() const { return bError; }
+	inline explicit operator bool() const { return !bAtEnd; }
+	inline const FStringView& GetName() const { return Name; }
+	inline const FStringView& GetValue() const { return Value; }
+
+private:
+	FStringView Tail;
+	FStringView Name;
+	FStringView Value;
+	bool bAtEnd = false;
+	bool bError = false;
+	bool bLastField = false;
+};
+
+/**
+ * Returns the possibly-quoted string, with quotes removed and escapes processed.
+ *
+ * @param String       View of a string that may be surrounded by quotes and may contain escaped characters.
+ * @param OutScratch   Builder to use as scratch space when processing escape sequences.
+ * @return A view into String if there are no escapes, and a view into OutScratch otherwise.
+ */
+[[nodiscard]] FStringView ParseQuotedString(FStringView String, FStringBuilderBase& OutScratch);
+
+/** Finds the first field in Config matching Name case-insensitively. */
+[[nodiscard]] bool FindFirstFieldValue(FStringView Config, FStringView Name, FStringView& OutValue);
+
+/** Visits every field in Config matching Name case-insensitively. */
+[[nodiscard]] bool VisitFieldValues(FStringView Config, FStringView Name, TFunctionWithContext<void (FStringView)> OnValue);
+
+/** Splits "Name(Config)" to ("Name", "(Config)"), "Name" to ("Name", ""), "(Config)" to ("", "(Config)"). */
+[[nodiscard]] bool TrySplitNameAndConfig(FStringView Value, FStringView& OutName, FStringView& OutConfig);
+
+/** Appends "Name=Value", or "Name" if Value is empty, with extra whitespace removed from any valid nested objects. */
+void AppendConfig(FStringBuilderBase& OutConfig, FStringView Name, FStringView Value);
+
+} // UE::DerivedData::Config
+
+namespace UE::DerivedData
+{
+
+/**
+* Finds the configuration for a cache graph.
+*
+* @param GraphName   Name of the cache graph, which is compared case-insensitively.
+* @param OutConfig   Appended with the configuration for the cache graph if it is found.
+* @return True if the configuration was found, otherwise false.
+*/
+bool TryFindCacheGraphConfig(const TCHAR* GraphName, FStringBuilderBase& OutConfig);
+
+/**
+ * Finds the configuration for a cache store, optionally within a specific cache graph.
+ *
+ * The optional cache graph may contain configuration overrides for the cache store.
+ * The configuration is appended to OutConfig from most-derived to least-derived,
+ * such that FParse::Value() will see override values before base values.
+ *
+ * @param StoreName   Name of the cache store, which is compared case-insensitively.
+ * @param GraphName   Name of the cache graph, which is compared case-insensitively, or nullptr.
+ * @param OutConfig   Appended with the configuration for the cache store if it is found.
+ * @return True if the configuration was found, otherwise false.
+ */
+bool TryFindCacheStoreConfig(const TCHAR* StoreName, const TCHAR* GraphName, FStringBuilderBase& OutConfig);
+
+} // UE::DerivedData
diff --git a/Engine/Source/Developer/DerivedDataCache/Private/DerivedDataLegacyCacheStore.h b/Engine/Source/Developer/DerivedDataCache/Private/DerivedDataLegacyCacheStore.h
--- a/Engine/Source/Developer/DerivedDataCache/Private/DerivedDataLegacyCacheStore.h
+++ b/Engine/Source/Developer/DerivedDataCache/Private/DerivedDataLegacyCacheStore.h
@@ -3,8 +3,8 @@
 #pragma once
 
 #include "DerivedDataCacheStore.h"
+#include "DerivedDataPrivate.h"
 #include "HAL/CriticalSection.h"
-#include "Logging/LogMacros.h"
 #include "Templates/RefCounting.h"
 #include <atomic>
 
@@ -22,8 +22,6 @@ namespace UE::DerivedData { struct FLegacyCachePutResponse; }
 namespace UE::DerivedData::Private { class FLegacyCacheKeyShared; }
 namespace UE::DerivedData::Private { class FLegacyCacheValueShared; }
 
-DECLARE_LOG_CATEGORY_EXTERN(LogDerivedDataCache, Log, All);
-
 namespace UE::DerivedData
 {
 
diff --git a/Engine/Source/Developer/DerivedDataCache/Private/DerivedDataPrivate.h b/Engine/Source/Developer/DerivedDataCache/Private/DerivedDataPrivate.h
--- a/Engine/Source/Developer/DerivedDataCache/Private/DerivedDataPrivate.h
+++ b/Engine/Source/Developer/DerivedDataCache/Private/DerivedDataPrivate.h
@@ -3,6 +3,9 @@
 #pragma once
 
 #include "HAL/LowLevelMemTracker.h"
+#include "Logging/LogMacros.h"
+
+DECLARE_LOG_CATEGORY_EXTERN(LogDerivedDataCache, Log, All);
 
 namespace UE::DerivedData
 {
diff --git a/Engine/Source/Developer/DerivedDataCache/Tests/DerivedDataConfigTest.cpp b/Engine/Source/Developer/DerivedDataCache/Tests/DerivedDataConfigTest.cpp
new file mode 100644
--- /dev/null
+++ b/Engine/Source/Developer/DerivedDataCache/Tests/DerivedDataConfigTest.cpp
@@ -0,0 +1,133 @@
+// Copyright Epic Games, Inc. All Rights Reserved.
+
+#if WITH_LOW_LEVEL_TESTS
+
+#include "DerivedDataConfig.h"
+
+#include "Containers/Array.h"
+#include "TestHarness.h"
+#include "Tests/WarnFilterScope.h"
+
+namespace UE::DerivedData
+{
+
+TEST_CASE("DerivedData::Config", "[DerivedData]")
+{
+	SECTION("ParseQuotedString")
+	{
+		TStringBuilder<64> Scratch;
+		CHECK(Config::ParseQuotedString(TEXTVIEW(""), Scratch) == TEXTVIEW(""));
+		CHECK(Config::ParseQuotedString(TEXTVIEW(R"(A)"), Scratch) == TEXTVIEW(R"(A)"));
+		CHECK(Config::ParseQuotedString(TEXTVIEW(R"(A B C)"), Scratch) == TEXTVIEW(R"(A B C)"));
+		CHECK(Config::ParseQuotedString(TEXTVIEW(R"(A " C)"), Scratch) == TEXTVIEW(R"(A " C)"));
+		CHECK(Config::ParseQuotedString(TEXTVIEW(R"(\ A " C \)"), Scratch) == TEXTVIEW(R"(\ A " C \)"));
+		CHECK(Config::ParseQuotedString(TEXTVIEW(R"('A B C')"), Scratch) == TEXTVIEW(R"('A B C')"));
+		CHECK(Config::ParseQuotedString(TEXTVIEW(R"("A B C")"), Scratch) == TEXTVIEW(R"(A B C)"));
+		CHECK(Config::ParseQuotedString(TEXTVIEW(R"("A B C)"), Scratch) == TEXTVIEW(R"(A B C)"));
+		CHECK(Config::ParseQuotedString(TEXTVIEW(R"(A B C")"), Scratch) == TEXTVIEW(R"(A B C)"));
+		CHECK(Scratch.Len() == 0);
+		CHECK(Config::ParseQuotedString(TEXTVIEW(R"("A \"B\" C")"), Scratch) == TEXTVIEW(R"(A "B" C)"));
+		CHECK(Config::ParseQuotedString(TEXTVIEW(R"("A \"\" C")"), Scratch) == TEXTVIEW(R"(A "" C)"));
+		CHECK(Config::ParseQuotedString(TEXTVIEW(R"("\\")"), Scratch) == TEXTVIEW(R"(\)"));
+		CHECK(Config::ParseQuotedString(TEXTVIEW(R"("\")"), Scratch) == TEXTVIEW(R"()"));
+		CHECK(Config::ParseQuotedString(TEXTVIEW(R"("\'\"\?\\\a\b\f\n\r\t\v")"), Scratch) == TEXTVIEW("\'\"\?\\\a\b\f\n\r\t\v"));
+	}
+}
+
+TEST_CASE("DerivedData::CacheConfig", "[DerivedData]")
+{
+	SECTION("TryFindCacheGraphConfig")
+	{
+		const auto TestFindGraph = [](const TCHAR* GraphName, const TCHAR* ExpectedConfig) -> bool
+		{
+			CAPTURE(GraphName);
+			TStringBuilder<1024> Config;
+			bool bFound = TryFindCacheGraphConfig(GraphName, Config);
+			CHECK(Config.ToView() == MakeStringView(ExpectedConfig));
+			return bFound;
+		};
+
+		CHECK_FALSE(TestFindGraph(TEXT("TestMissingCacheGraph"), TEXT("")));
+		for (int32 Index = 1; Index <= 6; ++Index)
+		{
+			CHECK_FALSE(TestFindGraph(*WriteToString<32>(TEXTVIEW("TestInvalidCacheGraph"), Index), TEXT("")));
+		}
+		CHECK(TestFindGraph(TEXT("TestCacheGraph0"), TEXT("()")));
+		CHECK(TestFindGraph(TEXT("TestCacheGraph1"), TEXT("(TestStore1)")));
+		CHECK(TestFindGraph(TEXT("TestCacheGraph2"), TEXT("(TestStore1, TestStore2, TestStore3)")));
+		CHECK(TestFindGraph(TEXT("TestCacheGraph3"), TEXT("(TestStore1, TestStore2, TestStore3)")));
+		CHECK(TestFindGraph(TEXT("TestCacheGraph4"), TEXT("(TestStore1, TestStore2, TestStore3)")));
+		CHECK(TestFindGraph(TEXT("TestCacheGraph5"), TEXT("(TestStore1, TestStore2, TestStore3)")));
+		CHECK(TestFindGraph(TEXT("TestCacheGraph6"), TEXT("(TestStore1)")));
+		CHECK(TestFindGraph(TEXT("TestCacheGraph7"), TEXT("(TestStore1)")));
+		CHECK(TestFindGraph(TEXT("TestCacheGraph8"), TEXT("(TestStore1)")));
+
+		{
+			CHECK_LOG_SCOPE(LogDerivedDataCache, Warning, "The cache graph TestCacheGraph9 is deprecated. Use TestCacheGraph1.");
+			CHECK(TestFindGraph(TEXT("TestCacheGraph9"), TEXT("(TestStore1)")));
+		}
+
+		{
+			CHECK_LOG_SCOPE(LogDerivedDataCache, Warning, "The cache graph TestLegacyCacheGraph1 contains cache store Root of deprecated type KeyLength.");
+			CHECK_LOG_SCOPE(LogDerivedDataCache, Warning, "The cache graph TestLegacyCacheGraph1 contains cache store Verify of deprecated type Verify.");
+			CHECK_LOG_SCOPE(LogDerivedDataCache, Warning, "The cache graph TestLegacyCacheGraph1 contains cache store AsyncPut of deprecated type AsyncPut.");
+			CHECK(TestFindGraph(TEXT("TestLegacyCacheGraph1"), TEXT("(TestLegacyStore1, TestStore2, TestStore3)")));
+		}
+	}
+
+	SECTION("TryFindCacheStoreConfig")
+	{
+		const auto TestFindStore = [](const TCHAR* StoreName, const TCHAR* GraphName, const TCHAR* ExpectedConfig) -> bool
+		{
+			CAPTURE(StoreName, GraphName);
+			TStringBuilder<1024> Config;
+			bool bFound = TryFindCacheStoreConfig(StoreName, GraphName, Config);
+			CHECK(Config.ToView() == MakeStringView(ExpectedConfig));
+			return bFound;
+		};
+
+		CHECK_FALSE(TestFindStore(TEXT("TestMissing"), nullptr, TEXT("")));
+		CHECK_FALSE(TestFindStore(TEXT("TestMissing"), TEXT("TestCacheGraph1"), TEXT("")));
+
+		CHECK_FALSE(TestFindStore(TEXT("TestInvalidStore1"), nullptr, TEXT("")));
+		CHECK_FALSE(TestFindStore(TEXT("TestInvalidStore1"), TEXT("TestInvalidCacheStore1"), TEXT("")));
+		CHECK_FALSE(TestFindStore(TEXT("TestInvalidStore1"), TEXT("TestInvalidCacheStore2"), TEXT("")));
+		CHECK_FALSE(TestFindStore(TEXT("TestInvalidStore1"), TEXT("TestInvalidCacheStore3"), TEXT("")));
+		CHECK_FALSE(TestFindStore(TEXT("TestInvalidStore2"), TEXT("TestInvalidCacheStore3"), TEXT("")));
+		CHECK_FALSE(TestFindStore(TEXT("TestInvalidStore3"), nullptr, TEXT("")));
+
+		CHECK_FALSE(TestFindStore(TEXT("TestInlineStore"), nullptr, TEXT("")));
+		CHECK(TestFindStore(TEXT("TestInlineStore"), TEXT("TestCacheGraphInlineStore1"), TEXT("(Key1=Value1(A=1, B=2, C), Key2=Value2, Key3=D(123))")));
+		CHECK(TestFindStore(TEXT("TestInlineStore"), TEXT("TestCacheGraphInlineStore2"), TEXT("(Key1=Value1(A=1, B=2, C), Key2=Value2, Key3=D(123))")));
+		CHECK_FALSE(TestFindStore(TEXT("TestMissing"), TEXT("TestCacheGraphInlineStore2"), TEXT("")));
+
+		CHECK(TestFindStore(TEXT("TestStore1"), nullptr, TEXT("(Key1=Store1Value, Key2=Store1Value, ValueOnly, (X=\"A=1, \\\"B\\\"=2, C=3\"))")));
+		CHECK(TestFindStore(TEXT("TestStore1"), TEXT("TestCacheGraph3"), TEXT("(Key1=Store1Value, Key2=Store1Value, ValueOnly, (X=\"A=1, \\\"B\\\"=2, C=3\"))")));
+
+		CHECK(TestFindStore(TEXT("TestStore2"), nullptr, TEXT("(Key1=Store2Value, Key1=Store1Value, Key2=Store1Value, ValueOnly, (X=\"A=1, \\\"B\\\"=2, C=3\"))")));
+		CHECK(TestFindStore(TEXT("TestStore2"), TEXT("TestCacheGraph3"), TEXT("(Key1=Store2Value, Key1=Store1Value, Key2=Store1Value, ValueOnly, (X=\"A=1, \\\"B\\\"=2, C=3\"))")));
+
+		CHECK(TestFindStore(TEXT("TestStore3"), nullptr, TEXT("(ValueOnly, Key1=Store3Value, Key2=Store3Value, Key3=(A=1, B=2, C))")));
+		CHECK(TestFindStore(TEXT("TestStore3"), TEXT("TestCacheGraph5"), TEXT("(Key1=Value1, Key2=Value2, ValueOnly, Key1=Store3Value, Key2=Store3Value, Key3=(A=1, B=2, C))")));
+
+		CHECK(TestFindStore(TEXT("TestStore4"), nullptr, TEXT("(Key1=Store4Value, Key1=Store2Value, Key1=Store1Value, Key2=Store1Value, ValueOnly, (X=\"A=1, \\\"B\\\"=2, C=3\"))")));
+		CHECK(TestFindStore(TEXT("TestStore4"), TEXT("TestCacheGraph5"), TEXT("(Key1=Store4Value, Key1=Store2Value, Key1=Store1Value, Key2=Store1Value, ValueOnly, (X=\"A=1, \\\"B\\\"=2, C=3\"))")));
+
+		CHECK(TestFindStore(TEXT("TestStore1"), TEXT("TestCacheGraph6"), TEXT("(Key1=Value1)")));
+		CHECK(TestFindStore(TEXT("TestStore1"), TEXT("TestCacheGraph7"), TEXT("(Key1=Value1, Key1=Store5Value, Key2=Store5Value)")));
+		CHECK(TestFindStore(TEXT("TestStore1"), TEXT("TestCacheGraph8"), TEXT("(Key1=Store1Value, Key2=Store1Value, ValueOnly, (X=\"A=1, \\\"B\\\"=2, C=3\"))")));
+
+		{
+			CHECK_LOG_SCOPE(LogDerivedDataCache, Warning, "The cache store TestStore7 is deprecated. Use TestStore1.");
+			CHECK(TestFindStore(TEXT("TestStore7"), nullptr, TEXT("(Key1=Store1Value, Key2=Store1Value, ValueOnly, (X=\"A=1, \\\"B\\\"=2, C=3\"))")));
+		}
+
+		CHECK(TestFindStore(TEXT("TestLegacyStore1"), TEXT("TestLegacyCacheGraph1"), TEXT("(Key1=LegacyStore1Value, Key2=LegacyStore1Value)")));
+		CHECK(TestFindStore(TEXT("TestStore2"), TEXT("TestLegacyCacheGraph1"), TEXT("(Key1=LegacyStore2Value)")));
+		CHECK(TestFindStore(TEXT("TestStore3"), TEXT("TestLegacyCacheGraph1"), TEXT("(ValueOnly, Key1=Store3Value, Key2=Store3Value, Key3=(A=1, B=2, C))")));
+	}
+}
+
+} // UE::DerivedData
+
+#endif // WITH_LOW_LEVEL_TESTS
