DDC: Fixed a dangling string view in TryFindCacheGraphConfig

diff --git a/Engine/Source/Developer/DerivedDataCache/Private/DerivedDataConfig.cpp b/Engine/Source/Developer/DerivedDataCache/Private/DerivedDataConfig.cpp
--- a/Engine/Source/Developer/DerivedDataCache/Private/DerivedDataConfig.cpp
+++ b/Engine/Source/Developer/DerivedDataCache/Private/DerivedDataConfig.cpp
@@ -313,13 +313,13 @@ bool TryFindCacheGraphConfig(const TCHAR* GraphName, FStringBuilderBase& OutConf
 {
 	// UE_DEPRECATED(5.8, "Remove the LegacyConfig block after the standard deprecation period.")
 	FString LegacyConfig;
-	if (GConfig->GetString(GraphName, TEXT("Root"), LegacyConfig, GEngineIni))
+	if (FString StoreName = TEXT("Root"); GConfig->GetString(GraphName, *StoreName, LegacyConfig, GEngineIni))
 	{
 		//UE_LOGFMT(LogDerivedDataCache, Warning,
 		//	"Configuring cache graph {GraphName} in config section [{GraphName}] of {IniName}.ini is deprecated. "
 		//	"Migrate the configuration to the [DerivedDataCacheGraphs] section of {IniName}.ini.",
 		//	("GraphName", GraphName), ("IniName", GEngineIni));
-		for (FStringView Type, Inner = TEXTVIEW("Root");;)
+		for (FStringView Type;;)
 		{
 			if (Config::FindFirstFieldValue(LegacyConfig, TEXTVIEW("Type"), Type))
 			{
@@ -344,18 +344,21 @@ bool TryFindCacheGraphConfig(const TCHAR* GraphName, FStringBuilderBase& OutConf
 				{
 					UE_LOGFMT(LogDerivedDataCache, Warning,
 						"The cache graph {GraphName} contains cache store {StoreName} of deprecated type {StoreType}.",
-						GraphName, Inner, Type);
-					if (Config::FindFirstFieldValue(LegacyConfig, TEXTVIEW("Inner"), Inner) &&
-						GConfig->GetString(GraphName, *WriteToString<64>(Inner), LegacyConfig, GEngineIni))
+						GraphName, StoreName, Type);
+					if (FStringView Inner; Config::FindFirstFieldValue(LegacyConfig, TEXTVIEW("Inner"), Inner))
 					{
-						continue;
+						StoreName = Inner;
+						if (GConfig->GetString(GraphName, *StoreName, LegacyConfig, GEngineIni))
+						{
+							continue;
+						}
 					}
 				}
 			}
 			UE_LOGFMT(LogDerivedDataCache, Warning,
 				"Failed to parse cache store {StoreName} from cache graph {GraphName} in config section [{GraphName}] of {IniName}.ini. "
 				"Migrate the configuration to the [DerivedDataCacheGraphs] section of {IniName}.ini.",
-				("GraphName", GraphName), ("StoreName", Inner), ("IniName", GEngineIni));
+				("GraphName", GraphName), ("StoreName", StoreName), ("IniName", GEngineIni));
 			break;
 		}
 	}
