DDC: Swapped to LogCategory.h in DerivedDataPrivate.h to minimize includes

#rnx

diff --git a/Engine/Source/Developer/DerivedDataCache/Private/DerivedDataPrivate.h b/Engine/Source/Developer/DerivedDataCache/Private/DerivedDataPrivate.h
--- a/Engine/Source/Developer/DerivedDataCache/Private/DerivedDataPrivate.h
+++ b/Engine/Source/Developer/DerivedDataCache/Private/DerivedDataPrivate.h
@@ -3,7 +3,7 @@
 #pragma once
 
 #include "HAL/LowLevelMemTracker.h"
-#include "Logging/LogMacros.h"
+#include "Logging/LogCategory.h"
 
 DECLARE_LOG_CATEGORY_EXTERN(LogDerivedDataCache, Log, All);
 
