DDC: Updated graph creation to use the new config API

Graphs are now configured in the [DerivedDataCacheGraphs] section of the Engine config.
Stores are now configured in the [DerivedDataCacheStores] section of the Engine config.

Graphs are an ordered list of stores. They support Name=OtherName to remap store names and Deprecated="DeprecationMessage'.
Stores are configured the same as before, with the addition of Base=BaseStoreName to inherit and Deprecated="DeprecationMessage" to deprecate.

Deprecated AsyncPut, KeyLength, and Verify graph nodes.

[DerivedDataCacheGraphs]
Default=(Hierarchical, Pak, EnginePak, ZenLocal, Local, ZenShared, Shared, Cloud)
Installed=(Pak, CompressedPak, EnginePak=InstalledEnginePak, EnterprisePak, ZenLocal, Local=InstalledLocal, ZenShared, Shared, Cloud)

[DerivedDataCacheStores]
Local=(Type=FileSystem, DeleteOnly=true, UnusedFileAge=8, Path=%ENGINEDIR%DerivedDataCache, EnvPathOverride=UE-LocalDataCachePath, EditorOverrideSetting=LocalDerivedDataCache, CommandLineOverride=LocalDataCachePath)
Pak=(Type=ReadPak, Filename="%GAMEDIR%DerivedDataCache/DDC.ddp")
EnginePak=(Type=ReadPak, Filename="%ENGINEDIR%DerivedDataCache/DDC.ddp")
ZenLocal=(Type=Zen, ServerID=Local)
InstalledLocal=(Base=Local, Path="%ENGINEVERSIONAGNOSTICUSERDIR%DerivedDataCache")

diff --git a/Engine/Source/Developer/DerivedDataCache/Private/DerivedDataBackends.cpp b/Engine/Source/Developer/DerivedDataCache/Private/DerivedDataBackends.cpp
--- a/Engine/Source/Developer/DerivedDataCache/Private/DerivedDataBackends.cpp
+++ b/Engine/Source/Developer/DerivedDataCache/Private/DerivedDataBackends.cpp
@@ -8,6 +8,7 @@
 #include "DerivedDataCacheReplay.h"
 #include "DerivedDataCacheStore.h"
 #include "DerivedDataCacheUsageStats.h"
+#include "DerivedDataConfig.h"
 #include "DerivedDataRequestOwner.h"
 #include "HAL/CriticalSection.h"
 #include "HAL/FileManager.h"
@@ -15,6 +16,7 @@
 #include "HAL/PlatformFileManager.h"
 #include "HAL/ThreadSafeCounter.h"
 #include "Internationalization/FastDecimalFormat.h"
+#include "Logging/StructuredLog.h"
 #include "Math/BasicMathExpressionEvaluator.h"
 #include "Math/UnitConversion.h"
 #include "MemoryCacheStore.h"
@@ -89,7 +91,6 @@ public:
 	FDerivedDataBackendGraph()
 		: RootCache(nullptr)
 		, MemoryCache(nullptr)
-		, BootCache(nullptr)
 		, WritePakCache(nullptr)
 		, Hierarchy(nullptr)
 		, bUsingSharedDDC(false)
@@ -119,13 +120,11 @@ public:
 		FParsedNodeMap ParsedNodes;
 		TGuardValue ParsedNodeMapGuard(ActiveParsedNodeMap, &ParsedNodes);
 
-		ILegacyCacheStore* RootNode = nullptr;
-
 		// Create the graph using ini settings. The string "default" forwards creation to use the default graph.
 
-		if (!FParse::Value(FCommandLine::Get(), TEXT("-DDC="), GraphName))
+		if (!FParse::Value(FCommandLine::Get(), TEXT("-DDC="), ActiveGraphName))
 		{
-			GraphName = GDerivedDataCacheGraphName.GetValueOnGameThread();
+			ActiveGraphName = GDerivedDataCacheGraphName.GetValueOnGameThread();
 		}
 
 		// A DDC graph of "None" is used by build worker programs that use the DDC build code paths but avoid use of the DDC cache
@@ -136,126 +135,106 @@ public:
 		// To avoid this confusion (and until we can use the build without the cache) it is restricted to use in programs only
 		// and not the editor.
 #if IS_PROGRAM
-		if (GraphName == TEXT("None"))
+		if (ActiveGraphName == TEXT("None"))
 		{
 			UE_LOG(LogDerivedDataCache, Display, TEXT("Requested cache graph of 'None'. Every cache operation will fail."));
 		}
 		else
 #endif
 		{
-			const TCHAR* const RootName = TEXT("Root");
-
-			if (!GraphName.IsEmpty() && GraphName != TEXT("Default"))
+			const auto ResetGraph = [this, &ParsedNodes]
 			{
-				RootNode = CreateCacheStoreHierarchy(Hierarchy, [this](IMemoryCacheStore*& OutCache) { GetMemoryCache(OutCache); });
+				// Destroy any cache stores that have been created.
+				delete RootCache;
+				RootCache = nullptr;
+				ParsedNodes.Empty();
+				MemoryCache = nullptr;
+				WritePakCache = nullptr;
+				Hierarchy = nullptr;
+				ReadPakCache.Empty();
+				Directories.Empty();
+			};
+
+			if (!ActiveGraphName.IsEmpty() && ActiveGraphName != TEXT("Default"))
+			{
+				if (!ParseGraph(*ActiveGraphName, ParsedNodes) || !Hierarchy->HasAllFlags(ECacheStoreFlags::Query | ECacheStoreFlags::Store))
+				{
+					ResetGraph();
+					UE_LOGFMT(LogDerivedDataCache, Warning,
+						"Failed to create graph {GraphName}. Attempting to create the default graph...", ActiveGraphName);
+				}
+			}
 
-				if (!ParseNode(RootName, GEngineIni, *GraphName, ParsedNodes) || !Hierarchy->HasAllFlags(ECacheStoreFlags::Query | ECacheStoreFlags::Store))
+			// UE_DEPRECATED(5.8, "Remove the legacy default block after the standard deprecation period.")
+			if (!Hierarchy)
+			{
+				// Try to use the legacy default graph.
+				ActiveGraphName = FApp::IsEngineInstalled() ? TEXT("InstalledDerivedDataBackendGraph") : TEXT("DerivedDataBackendGraph");
+				if (!ParseGraph(*ActiveGraphName, ParsedNodes) || !Hierarchy->HasAllFlags(ECacheStoreFlags::Query | ECacheStoreFlags::Store))
 				{
-					// Destroy any cache stores that have been created.
-					delete RootNode;
-					RootNode = nullptr;
-					ParsedNodes.Empty();
-					MemoryCache = nullptr;
-					BootCache = nullptr;
-					WritePakCache = nullptr;
-					Hierarchy = nullptr;
-					ReadPakCache.Empty();
-					Directories.Empty();
-					bAsyncFound = false;
-					bBootFound = false;
-					bHierarchyFound = false;
-					bVerifyFound = false;
-					bVerifyFix = false;
-					UE_LOG(LogDerivedDataCache, Warning,
-						TEXT("Unable to create or use cache graph '%s'. Reverting to the default graph."), *GraphName);
+					ResetGraph();
 				}
 			}
 
 			if (!Hierarchy)
 			{
 				// Try to use the default graph.
-				GraphName = FApp::IsEngineInstalled() ? TEXT("InstalledDerivedDataBackendGraph") : TEXT("DerivedDataBackendGraph");
-				RootNode = CreateCacheStoreHierarchy(Hierarchy, [this](IMemoryCacheStore*& OutCache) { GetMemoryCache(OutCache); });
-
-				if (!ParseNode(RootName, GEngineIni, *GraphName, ParsedNodes) || !Hierarchy->HasAllFlags(ECacheStoreFlags::Query | ECacheStoreFlags::Store))
+				ActiveGraphName = FApp::IsEngineInstalled() ? TEXT("Installed") : TEXT("Default");
+				if (!ParseGraph(*ActiveGraphName, ParsedNodes) || !Hierarchy->HasAllFlags(ECacheStoreFlags::Query | ECacheStoreFlags::Store))
 				{
-					FString Entry;
-					if (!GConfig->DoesSectionExist(*GraphName, GEngineIni))
-					{
-						UE_LOG(LogDerivedDataCache, Fatal,
-							TEXT("Unable to create default cache graph '%s' because its config section is missing in the '%s' config."),
-							*GraphName, *GEngineIni);
-					}
-					else if (!GConfig->GetString(*GraphName, RootName, Entry, GEngineIni) || !Entry.Len())
+					TStringBuilder<256> GraphConfig;
+					if (!TryFindCacheGraphConfig(*ActiveGraphName, GraphConfig))
 					{
-						UE_LOG(LogDerivedDataCache, Fatal,
-							TEXT("Unable to create default cache graph '%s' because the root node '%s' is missing."),
-							*GraphName, RootName);
+						UE_LOGFMT(LogDerivedDataCache, Fatal,
+							"Failed to create default cache graph '{GraphName}' because its config is missing in the '{IniName}' config.",
+							ActiveGraphName, GEngineIni);
 					}
 					else if (FParse::Param(FCommandLine::Get(), TEXT("DDC-ForceMemoryCache")))
 					{
 						IMemoryCacheStore* Memory = nullptr;
 						CreateMemoryCacheStore(Memory, TEXT("ForceMemoryCache"), TEXT(""), Hierarchy);
-						UE_LOG(LogDerivedDataCache, Error,
-							TEXT("Unable to use default cache graph '%s' because there are no %s nodes available."),
-							*GraphName,
-							Hierarchy->HasAllFlags(ECacheStoreFlags::Query) ? TEXT("writable") :
-							Hierarchy->HasAllFlags(ECacheStoreFlags::Store) ? TEXT("readable") : TEXT("readable or writable"));
+						UE_LOGFMT(LogDerivedDataCache, Error,
+							"Unable to use default cache graph '{GraphName}' because there are no {MissingTypes} nodes available.",
+							ActiveGraphName,
+							Hierarchy->HasAllFlags(ECacheStoreFlags::Query) ? "writable" :
+							Hierarchy->HasAllFlags(ECacheStoreFlags::Store) ? "readable" : "readable or writable");
 						check(Hierarchy->HasAllFlags(ECacheStoreFlags::Query | ECacheStoreFlags::Store));
 					}
 					else
 					{
-						UE_CLOG(!FParse::Param(FCommandLine::Get(), TEXT("DDC-AllowNoActiveStores")),
+						UE_CLOGFMT(!FParse::Param(FCommandLine::Get(), TEXT("DDC-AllowNoActiveStores")),
 							LogDerivedDataCache, Fatal,
-							TEXT("Unable to use default cache graph '%s' because there are no %s nodes available."
-								 "Add -DDC-ForceMemoryCache to the command line to bypass this if you need access "
-								 "to the editor settings to fix the cache configuration."),
-							*GraphName,
-							Hierarchy->HasAllFlags(ECacheStoreFlags::Query) ? TEXT("writable") :
-							Hierarchy->HasAllFlags(ECacheStoreFlags::Store) ? TEXT("readable") : TEXT("readable or writable"));
+							"Unable to use default cache graph '{GraphName}' because there are no {MissingTypes} nodes available."
+							"Add -DDC-ForceMemoryCache to the command line to bypass this if you need access "
+							"to the editor settings to fix the cache configuration.",
+							*ActiveGraphName,
+							Hierarchy->HasAllFlags(ECacheStoreFlags::Query) ? "writable" :
+							Hierarchy->HasAllFlags(ECacheStoreFlags::Store) ? "readable" : "readable or writable");
 					}
 				}
 			}
 		}
 
 		// Async must exist in the graph.
-		RootNode = CreateCacheStoreAsync(RootNode, MemoryCache, /*bDeleteInnerCache*/ true);
+		RootCache = CreateCacheStoreAsync(RootCache, MemoryCache, /*bDeleteInnerCache*/ true);
 
-		// Create a Verify node when using -DDC-Verify[=Type1[@Rate2][+Type2[@Rate2]...]] or the graph has a verify node.
+		// Create a Verify node when using -DDC-Verify[=Type1[@Rate2][+Type2[@Rate2]...]].
 		if (FString VerifyArg;
 			FParse::Value(FCommandLine::Get(), TEXT("-DDC-Verify="), VerifyArg) ||
-			FParse::Param(FCommandLine::Get(), TEXT("DDC-Verify")) ||
-			bVerifyFound)
+			FParse::Param(FCommandLine::Get(), TEXT("DDC-Verify")))
 		{
 			if (!UE::GetMultiprocessId())
 			{
 				IFileManager::Get().DeleteDirectory(*(FPaths::ProjectSavedDir() / TEXT("VerifyDDC/")), /*bRequireExists*/ false, /*bTree*/ true);
 			}
-			RootNode = CreateCacheStoreVerify(RootNode, /*bPutOnError*/ bVerifyFix);
+			RootCache = CreateCacheStoreVerify(RootCache, /*bPutOnError*/ false);
 		}
 
 		// Create a Replay node when requested on the command line.
-		if (ILegacyCacheStore* ReplayNode = TryCreateCacheStoreReplay(RootNode))
+		if (ILegacyCacheStore* ReplayNode = TryCreateCacheStoreReplay(RootCache))
 		{
-			RootNode = ReplayNode;
+			RootCache = ReplayNode;
 		}
-
-		if (MaxKeyLength == 0)
-		{
-			MaxKeyLength = MAX_BACKEND_KEY_LENGTH;
-		}
-
-		RootCache = RootNode;
-	}
-
-	/**
-	 * Helper function to get the value of parsed bool as the return value
-	 **/
-	bool GetParsedBool( const TCHAR* Stream, const TCHAR* Match ) const
-	{
-		bool bValue = 0;
-		FParse::Bool( Stream, Match, bValue );
-		return bValue;
 	}
 
 	/**
@@ -267,23 +246,29 @@ public:
 	 * @param InParsedNodes Map of parsed nodes and their names to be able to find already parsed nodes
 	 * @return Derived data backend interface instance created from ini settings
 	 */
-	bool ParseNode(const FString& NodeName, const FString& IniFilename, const TCHAR* IniSection, FParsedNodeMap& InParsedNodes)
+	bool ParseNode(const TCHAR* StoreName, const TCHAR* GraphName, FParsedNodeMap& InParsedNodes)
 	{
-		if (const FParsedNode* ParsedNode = InParsedNodes.Find(NodeName))
+		if (const FParsedNode* ParsedNode = InParsedNodes.Find(StoreName))
+		{
+			UE_LOGFMT(LogDerivedDataCache, Display, "Node {StoreName} was referenced more than once in cache graph {GraphName}. Nodes may not be shared.", StoreName, GraphName);
+			return false;
+		}
+
+		TStringBuilder<256> StoreConfig;
+		if (!TryFindCacheStoreConfig(StoreName, GraphName, StoreConfig))
 		{
-			UE_LOG(LogDerivedDataCache, Display, TEXT("Node %s was referenced more than once in the graph. Nodes may not be shared."), *NodeName);
+			UE_LOGFMT(LogDerivedDataCache, Log, "Failed to find config for cache store {StoreName} in cache graph {GraphName}.", StoreName, GraphName);
 			return false;
 		}
 
 		ICacheStoreOwner* NodeOwner = this;
-		FString Entry;
-		if (GConfig->GetString(IniSection, *NodeName, Entry, IniFilename))
+		FString Entry(StoreConfig);
 		{
 			Entry.TrimStartInline();
 			Entry.RemoveFromStart(TEXT("("));
 			Entry.RemoveFromEnd(TEXT(")"));
 
-			TGuardValue NodeNameGuard(ActiveNodeName, NodeName);
+			TGuardValue NodeNameGuard(ActiveNodeName, FString(StoreName));
 			TGuardValue NodeConfigGuard(ActiveNodeConfig, Entry);
 
 			FString	NodeType;
@@ -291,61 +276,41 @@ public:
 			{
 				if (NodeType == TEXT("FileSystem"))
 				{
-					return ParseFileSystemCache(*NodeName, *Entry, InParsedNodes);
-				}
-				else if (NodeType == TEXT("Boot"))
-				{
-					return ParseBootCache(*NodeName, *Entry, InParsedNodes);
+					return ParseFileSystemCache(StoreName, *Entry, InParsedNodes);
 				}
 				else if (NodeType == TEXT("Memory"))
 				{
-					return ParseMemoryCache(*NodeName, *Entry, InParsedNodes);
-				}
-				else if (NodeType == TEXT("Hierarchical"))
-				{
-					return ParseHierarchyNode(*NodeName, *Entry, IniFilename, IniSection, InParsedNodes);
-				}
-				else if (NodeType == TEXT("KeyLength"))
-				{
-					return ParseKeyLength(*NodeName, *Entry, IniFilename, IniSection, InParsedNodes);
-				}
-				else if (NodeType == TEXT("AsyncPut"))
-				{
-					return ParseAsyncNode(*NodeName, *Entry, IniFilename, IniSection, InParsedNodes);
-				}
-				else if (NodeType == TEXT("Verify"))
-				{
-					return ParseVerify(*NodeName, *Entry, IniFilename, IniSection, InParsedNodes);
+					return ParseMemoryCache(StoreName, *Entry, InParsedNodes);
 				}
 				else if (NodeType == TEXT("ReadPak"))
 				{
-					return ParsePak(*NodeName, *Entry, /*bWriting*/ false, InParsedNodes);
+					return ParsePak(StoreName, *Entry, /*bWriting*/ false, InParsedNodes);
 				}
 				else if (NodeType == TEXT("WritePak"))
 				{
-					return ParsePak(*NodeName, *Entry, /*bWriting*/ true, InParsedNodes);
+					return ParsePak(StoreName, *Entry, /*bWriting*/ true, InParsedNodes);
 				}
 				else if (NodeType == TEXT("S3"))
 				{
-					if (ILegacyCacheStore* CacheStore = CreateS3CacheStore(*NodeName, *Entry, *this))
+					if (ILegacyCacheStore* CacheStore = CreateS3CacheStore(StoreName, *Entry, *this))
 					{
-						InParsedNodes.Add(NodeName, CacheStore);
+						InParsedNodes.Add(StoreName, CacheStore);
 						return true;
 					}
 				}
 				else if (NodeType == TEXT("Cloud") || NodeType == TEXT("Http"))
 				{
-					if (ILegacyCacheStore* CacheStore = CreateHttpCacheStore(*NodeName, *Entry, this))
+					if (ILegacyCacheStore* CacheStore = CreateHttpCacheStore(StoreName, *Entry, this))
 					{
-						InParsedNodes.Add(NodeName, CacheStore);
+						InParsedNodes.Add(StoreName, CacheStore);
 						return true;
 					}
 				}
 				else if (NodeType == TEXT("Zen"))
 				{
-					if (ILegacyCacheStore* CacheStore = CreateZenCacheStore(*NodeName, *Entry, this))
+					if (ILegacyCacheStore* CacheStore = CreateZenCacheStore(StoreName, *Entry, this))
 					{
-						InParsedNodes.Add(NodeName,CacheStore);
+						InParsedNodes.Add(StoreName, CacheStore);
 						return true;
 					}
 				}
@@ -367,8 +332,9 @@ public:
 	{
 		ILegacyCacheStore* PakNode = nullptr;
 		FString PakFilename;
-		FParse::Value( Entry, TEXT("Filename="), PakFilename );
-		bool bCompressed = GetParsedBool(Entry, TEXT("Compressed="));
+		FParse::Value(Entry, TEXT("Filename="), PakFilename);
+		bool bCompressed = false;
+		FParse::Bool(Entry, TEXT("Compressed="), bCompressed);
 
 		if (PakFilename.IsEmpty())
 		{
@@ -411,140 +377,41 @@ public:
 	}
 
 	/**
-	 * Creates Verify wrapper interface from ini settings.
+	 * Creates Hierarchical interface from ini settings.
 	 *
 	 * @param NodeName Node name.
 	 * @param Entry Node definition.
 	 * @param IniFilename ini filename.
 	 * @param IniSection ini section containing graph definition
 	 * @param InParsedNodes map of nodes and their names which have already been parsed
-	 * @return Verify wrapper backend interface instance or nullptr if unsuccessful
+	 * @return Hierarchical backend interface instance or nullptr if unsuccessful
 	 */
-	bool ParseVerify(const TCHAR* NodeName, const TCHAR* Entry, const FString& IniFilename, const TCHAR* IniSection, FParsedNodeMap& InParsedNodes)
+	bool ParseGraph(const TCHAR* GraphName, FParsedNodeMap& InParsedNodes)
 	{
-		if (bVerifyFound)
-		{
-			UE_LOG(LogDerivedDataCache, Warning, TEXT("Unable to create %s Verify because only one Verify node is supported."), NodeName);
-			return false;
-		}
+		RootCache = CreateCacheStoreHierarchy(Hierarchy, [this](IMemoryCacheStore*& OutCache) { GetMemoryCache(OutCache); });
 
-		bVerifyFound = true;
-		bVerifyFix = GetParsedBool(Entry, TEXT("Fix="));
-
-		FString InnerName;
-		if (FParse::Value(Entry, TEXT("Inner="), InnerName) && ParseNode(InnerName, IniFilename, IniSection, InParsedNodes))
+		TStringBuilder<256> GraphConfig;
+		if (!TryFindCacheGraphConfig(GraphName, GraphConfig))
 		{
-			return true;
-		}
-
-		UE_LOG(LogDerivedDataCache, Warning, TEXT("Unable to find inner node %s for Verify node %s. Verify node will not be created."), *InnerName, NodeName);
-		return false;
-	}
-
-	/**
-	 * Creates AsyncPut wrapper interface from ini settings.
-	 *
-	 * @param NodeName Node name.
-	 * @param Entry Node definition.
-	 * @param IniFilename ini filename.
-	 * @param IniSection ini section containing graph definition
-	 * @param InParsedNodes map of nodes and their names which have already been parsed
-	 * @return AsyncPut wrapper backend interface instance or nullptr if unsuccessful
-	 */
-	bool ParseAsyncNode(const TCHAR* NodeName, const TCHAR* Entry, const FString& IniFilename, const TCHAR* IniSection, FParsedNodeMap& InParsedNodes)
-	{
-		if (bAsyncFound)
-		{
-			UE_LOG(LogDerivedDataCache, Warning, TEXT("Unable to create %s AsyncPut because only one AsyncPut node is supported."), NodeName);
 			return false;
 		}
 
-		bAsyncFound = true;
-		FString InnerName;
-		if (FParse::Value(Entry, TEXT("Inner="), InnerName) && ParseNode(InnerName, IniFilename, IniSection, InParsedNodes))
-		{
-			return true;
-		}
-
-		UE_LOG(LogDerivedDataCache, Warning, TEXT("Unable to find inner node %s for AsyncPut node %s. AsyncPut node will not be created."), *InnerName, NodeName);
-		return false;
-	}
-
-	/**
-	 * Creates KeyLength wrapper interface from ini settings.
-	 *
-	 * @param NodeName Node name.
-	 * @param Entry Node definition.
-	 * @param IniFilename ini filename.
-	 * @param IniSection ini section containing graph definition
-	 * @param InParsedNodes map of nodes and their names which have already been parsed
-	 * @return KeyLength wrapper backend interface instance or nullptr if unsuccessful
-	 */
-	bool ParseKeyLength(const TCHAR* NodeName, const TCHAR* Entry, const FString& IniFilename, const TCHAR* IniSection, FParsedNodeMap& InParsedNodes)
-	{
-		if (MaxKeyLength)
+		bool bParsed = false;
+		Config::FConfigIterator It(GraphConfig);
+		for (; It; ++It)
 		{
-			UE_LOG(LogDerivedDataCache, Warning, TEXT("Node %s is disabled because there may be only one key length node."), NodeName);
-			return false;
+			bParsed |= ParseNode(*WriteToString<64>(It.GetValue()), GraphName, InParsedNodes);
 		}
 
-		FString InnerName;
-		if (FParse::Value(Entry, TEXT("Inner="), InnerName) && ParseNode(InnerName, IniFilename, IniSection, InParsedNodes))
+		if (It.HasError())
 		{
-			int32 KeyLength = MAX_BACKEND_KEY_LENGTH;
-			FParse::Value(Entry, TEXT("Length="), KeyLength);
-			MaxKeyLength = FMath::Clamp(KeyLength, 0, MAX_BACKEND_KEY_LENGTH);
-			return true;
-		}
-
-		UE_LOG(LogDerivedDataCache, Warning, TEXT("Unable to find inner node %s for KeyLength node %s. KeyLength node will not be created."), *InnerName, NodeName);
-		return false;
-	}
-
-	/**
-	 * Creates Hierarchical interface from ini settings.
-	 *
-	 * @param NodeName Node name.
-	 * @param Entry Node definition.
-	 * @param IniFilename ini filename.
-	 * @param IniSection ini section containing graph definition
-	 * @param InParsedNodes map of nodes and their names which have already been parsed
-	 * @return Hierarchical backend interface instance or nullptr if unsuccessful
-	 */
-	bool ParseHierarchyNode(const TCHAR* NodeName, const TCHAR* Entry, const FString& IniFilename, const TCHAR* IniSection, FParsedNodeMap& InParsedNodes)
-	{
-		if (bHierarchyFound)
-		{
-			UE_LOG(LogDerivedDataCache, Warning, TEXT("Node %s is disabled because there may be only one hierarchy node. "
-				"Confirm there is only one hierarchy in the cache graph and that it is inside of any async node."), NodeName);
+			UE_LOGFMT(LogDerivedDataCache, Warning, "Failed to parse at cache graph {GraphName}.", GraphName);
 			return false;
 		}
 
-		const TCHAR* InnerMatch = TEXT("Inner=");
-		const int32 InnerMatchLength = FCString::Strlen(InnerMatch);
-
-		bHierarchyFound = true;
-		bool bParsed = false;
-		FString InnerName;
-		while (FParse::Value(Entry, InnerMatch, InnerName))
-		{
-			if (ParseNode(InnerName, IniFilename, IniSection, InParsedNodes))
-			{
-				bParsed = true;
-			}
-			else
-			{
-				UE_LOG(LogDerivedDataCache, Log, TEXT("Unable to find inner node %s for hierarchy %s."), *InnerName, NodeName);
-			}
-
-			// Move the Entry pointer forward so that we can find more children
-			Entry = FCString::Strifind(Entry, InnerMatch);
-			Entry += InnerMatchLength;
-		}
-
 		if (!bParsed)
 		{
-			UE_LOG(LogDerivedDataCache, Warning, TEXT("Hierarchical cache %s has no inner backends and will not be created."), NodeName);
+			UE_LOGFMT(LogDerivedDataCache, Warning, "Failed to create at least one cache store for cache graph {GraphName}.", GraphName);
 			return false;
 		}
 
@@ -576,38 +443,6 @@ public:
 		return false;
 	}
 
-	/**
-	 * Creates Boot data cache interface from ini settings.
-	 *
-	 * @param NodeName Node name.
-	 * @param Entry Node definition.
-	 * @param OutFilename filename specified for the cache
-	 * @return Boot data cache backend interface instance or nullptr if unsuccessful
-	 */
-	bool ParseBootCache(const TCHAR* NodeName, const TCHAR* Entry, FParsedNodeMap& InParsedNodes)
-	{
-		UE_LOG(LogDerivedDataCache, Display, TEXT("Boot nodes are deprecated. Please remove the Boot node from the cache graph."));
-		if (bBootFound)
-		{
-			UE_LOG(LogDerivedDataCache, Warning, TEXT("Unable to create %s Boot cache because only one Boot node is supported."), NodeName);
-			return false;
-		}
-
-		bBootFound = true;
-
-	#if WITH_EDITOR
-		// Only allow boot cache with the editor. We don't want other tools and utilities (eg. SCW) writing to the same file.
-		CreateMemoryCacheStore(BootCache, TEXT("Boot"), *WriteToString<128>(TEXT("-Boot "), Entry), this);
-	#endif
-
-		if (!!BootCache)
-		{
-			InParsedNodes.Add(NodeName, BootCache);
-		}
-
-		return !!BootCache;
-	}
-
 	/**
 	 * Creates Memory data cache interface from ini settings.
 	 *
@@ -660,16 +495,12 @@ public:
 
 	virtual int32 GetMaxKeyLength() const override
 	{
-		return MaxKeyLength;
+		return MAX_BACKEND_KEY_LENGTH;
 	}
 
 	virtual void NotifyBootComplete() override
 	{
 		check(RootCache);
-		if (BootCache)
-		{
-			BootCache->Disable();
-		}
 	}
 
 	virtual void WaitForQuiescence(bool bShutdown) override
@@ -753,12 +584,12 @@ public:
 
 	virtual const TCHAR* GetGraphName() const override
 	{
-		return *GraphName;
+		return *ActiveGraphName;
 	}
 
 	virtual const TCHAR* GetDefaultGraphName() const override
 	{
-		return FApp::IsEngineInstalled() ? TEXT("InstalledDerivedDataBackendGraph") : TEXT("DerivedDataBackendGraph");
+		return FApp::IsEngineInstalled() ? TEXT("Installed") : TEXT("Default");
 	}
 
 	virtual void AddToAsyncCompletionCounter(int32 Addend) override
@@ -962,7 +793,7 @@ private:
 			return *ParsedNode;
 		}
 
-		if (!ParseNode(Name, GEngineIni, *GraphName, *ActiveParsedNodeMap))
+		if (!ParseNode(Name, *ActiveGraphName, *ActiveParsedNodeMap))
 		{
 			return nullptr;
 		}
@@ -1083,7 +914,7 @@ private:
 
 	static inline FDerivedDataBackendGraph*			StaticGraph;
 
-	FString											GraphName;
+	FString											ActiveGraphName;
 	FString											ReadPakFilename;
 	FString											WritePakFilename;
 
@@ -1094,7 +925,6 @@ private:
 
 	/** Instances of backend interfaces which exist in only one copy */
 	IMemoryCacheStore* MemoryCache;
-	IMemoryCacheStore* BootCache;
 	IPakFileCacheStore* WritePakCache;
 	ICacheStoreOwner* Hierarchy;
 	/** Support for multiple read only pak files. */
@@ -1105,8 +935,6 @@ private:
 
 	FBackendDebugOptions DefaultDebugOptions;
 
-	int32 MaxKeyLength = 0;
-
 	bool bHasDefaultDebugOptions = false;
 
 	/** Whether a shared cache is in use */
@@ -1126,12 +954,6 @@ private:
 	FParsedNodeMap* ActiveParsedNodeMap = nullptr;
 	FString ActiveNodeName;
 	FString ActiveNodeConfig;
-
-	bool bAsyncFound = false;
-	bool bBootFound = false;
-	bool bHierarchyFound = false;
-	bool bVerifyFound = false;
-	bool bVerifyFix = false;
 };
 
 } // UE::DerivedData
