DDC: Fixed a missing include on platforms without support for S3 DDC

#rnx

diff --git a/Engine/Source/Developer/DerivedDataCache/Private/S3CacheStore.cpp b/Engine/Source/Developer/DerivedDataCache/Private/S3CacheStore.cpp
--- a/Engine/Source/Developer/DerivedDataCache/Private/S3CacheStore.cpp
+++ b/Engine/Source/Developer/DerivedDataCache/Private/S3CacheStore.cpp
@@ -3,6 +3,8 @@
 #include "DerivedDataCacheStore.h"
 #include "DerivedDataLegacyCacheStore.h"
 
+#include "Logging/LogMacros.h"
+
 #if WITH_S3_DDC_BACKEND
 
 #if PLATFORM_MICROSOFT
