DDC: Fix for false-positive static analysis warning about a null Hierarchy

#rnx

diff --git a/Engine/Source/Developer/DerivedDataCache/Private/DerivedDataBackends.cpp b/Engine/Source/Developer/DerivedDataCache/Private/DerivedDataBackends.cpp
--- a/Engine/Source/Developer/DerivedDataCache/Private/DerivedDataBackends.cpp
+++ b/Engine/Source/Developer/DerivedDataCache/Private/DerivedDataBackends.cpp
@@ -166,7 +166,7 @@ public:
 			}
 
 			// UE_DEPRECATED(5.8, "Remove the legacy default block after the standard deprecation period.")
-			if (!Hierarchy)
+			if (!RootCache)
 			{
 				// Try to use the legacy default graph.
 				ActiveGraphName = FApp::IsEngineInstalled() ? TEXT("InstalledDerivedDataBackendGraph") : TEXT("DerivedDataBackendGraph");
@@ -176,7 +176,7 @@ public:
 				}
 			}
 
-			if (!Hierarchy)
+			if (!RootCache)
 			{
 				// Try to use the default graph.
 				ActiveGraphName = FApp::IsEngineInstalled() ? TEXT("Installed") : TEXT("Default");
