DDC: Removed unused short key and legacy mode functionality

diff --git a/Engine/Source/Developer/DerivedDataCache/Private/DerivedDataBackendInterface.h b/Engine/Source/Developer/DerivedDataCache/Private/DerivedDataBackendInterface.h
--- a/Engine/Source/Developer/DerivedDataCache/Private/DerivedDataBackendInterface.h
+++ b/Engine/Source/Developer/DerivedDataCache/Private/DerivedDataBackendInterface.h
@@ -39,16 +39,6 @@ LLM_DECLARE_TAG_API(UntaggedDDCResult, DERIVEDDATACACHE_API);
 namespace UE::DerivedData
 {
 
-enum class EBackendLegacyMode
-{
-	/** Use only GetValue/PutValue. */
-	ValueOnly,
-	/** Use GetValue/PutValue with a fallback to GetCachedData+PutValue for misses. */
-	ValueWithLegacyFallback,
-	/** Use only GetCachedData/PutCachedData. */
-	LegacyOnly,
-};
-
 /**
  * Speed classes. Higher values are faster so > / < comparisons make sense.
  */
@@ -99,13 +89,10 @@ public:
 	 */
 	virtual ILegacyCacheStore& GetRoot() = 0;
 
-	virtual int32 GetMaxKeyLength() const = 0;
-
 	//--------------------
 	// System Interface, copied from FDerivedDataCacheInterface
 	//--------------------
 
-	virtual void NotifyBootComplete() = 0;
 	virtual void AddToAsyncCompletionCounter(int32 Addend) = 0;
 	virtual bool AnyAsyncRequestsRemaining() = 0;
 	virtual bool IsShuttingDown() = 0;
@@ -141,42 +128,6 @@ public:
 
 /** Lexical conversions from and to enums */
 
-[[nodiscard]] inline const TCHAR* LexToString(const EBackendLegacyMode Value)
-{
-	switch (Value)
-	{
-	case EBackendLegacyMode::ValueOnly:
-		return TEXT("ValueOnly");
-	case EBackendLegacyMode::ValueWithLegacyFallback:
-		return TEXT("ValueWithLegacyFallback");
-	case EBackendLegacyMode::LegacyOnly:
-		return TEXT("LegacyOnly");
-	default:
-		checkNoEntry();
-		return TEXT("Unknown");
-	}
-}
-
-[[nodiscard]] inline bool TryLexFromString(EBackendLegacyMode& OutValue, const FStringView String)
-{
-	if (String == TEXTVIEW("ValueOnly"))
-	{
-		OutValue = EBackendLegacyMode::ValueOnly;
-		return true;
-	}
-	if (String == TEXTVIEW("ValueWithLegacyFallback"))
-	{
-		OutValue = EBackendLegacyMode::ValueWithLegacyFallback;
-		return true;
-	}
-	if (String == TEXTVIEW("LegacyOnly"))
-	{
-		OutValue = EBackendLegacyMode::LegacyOnly;
-		return true;
-	}
-	return false;
-}
-
 inline const TCHAR* LexToString(EBackendSpeedClass SpeedClass)
 {
 	switch (SpeedClass)
diff --git a/Engine/Source/Developer/DerivedDataCache/Private/DerivedDataBackends.cpp b/Engine/Source/Developer/DerivedDataCache/Private/DerivedDataBackends.cpp
--- a/Engine/Source/Developer/DerivedDataCache/Private/DerivedDataBackends.cpp
+++ b/Engine/Source/Developer/DerivedDataCache/Private/DerivedDataBackends.cpp
@@ -39,7 +39,6 @@ DEFINE_LOG_CATEGORY(LogDerivedDataCache);
 LLM_DEFINE_TAG(UntaggedDDCResult);
 LLM_DEFINE_TAG(DDCBackend);
 
-#define MAX_BACKEND_KEY_LENGTH (120)
 #define LOCTEXT_NAMESPACE "DerivedDataBackendGraph"
 
 static TAutoConsoleVariable<FString> GDerivedDataCacheGraphName(
@@ -493,16 +492,6 @@ public:
 		return *RootCache;
 	}
 
-	virtual int32 GetMaxKeyLength() const override
-	{
-		return MAX_BACKEND_KEY_LENGTH;
-	}
-
-	virtual void NotifyBootComplete() override
-	{
-		check(RootCache);
-	}
-
 	virtual void WaitForQuiescence(bool bShutdown) override
 	{
 		double StartTime = FPlatformTime::Seconds();
diff --git a/Engine/Source/Developer/DerivedDataCache/Private/DerivedDataCache.cpp b/Engine/Source/Developer/DerivedDataCache/Private/DerivedDataCache.cpp
--- a/Engine/Source/Developer/DerivedDataCache/Private/DerivedDataCache.cpp
+++ b/Engine/Source/Developer/DerivedDataCache/Private/DerivedDataCache.cpp
@@ -473,7 +473,7 @@ private:
 			FRequestBarrier Barrier(Owner);
 			FLegacyCacheGetRequest Request;
 			Request.Name = DebugContext;
-			Request.Key = FLegacyCacheKey(CacheKey, Backend->GetMaxKeyLength());
+			Request.Key = FLegacyCacheKey(CacheKey);
 			Backend->GetRoot().LegacyGet({Request}, Owner,
 				[this](FLegacyCacheGetResponse&& Response)
 				{
@@ -631,7 +631,7 @@ private:
 			SCOPE_SECONDS_COUNTER(ThisTime);
 			FLegacyCachePutRequest Request;
 			Request.Name = DebugContext;
-			Request.Key = FLegacyCacheKey(CacheKey, Backend->GetMaxKeyLength());
+			Request.Key = FLegacyCacheKey(CacheKey);
 			Request.Value = FLegacyCacheValue(FCompositeBuffer(FSharedBuffer::Clone(MakeMemoryView(Data))));
 			FRequestOwner AsyncOwner(EPriority::Normal);
 			Backend->GetRoot().LegacyPut({Request}, AsyncOwner, [](auto&&) {});
@@ -858,7 +858,7 @@ public:
 			SCOPE_SECONDS_COUNTER(ThisTime);
 			FLegacyCachePutRequest LegacyRequest;
 			LegacyRequest.Name = DebugContext;
-			LegacyRequest.Key = FLegacyCacheKey(CacheKey, Backend->GetMaxKeyLength());
+			LegacyRequest.Key = FLegacyCacheKey(CacheKey);
 			LegacyRequest.Value = FLegacyCacheValue(FCompositeBuffer(FSharedBuffer::Clone(MakeMemoryView(Data))));
 			FRequestOwner AsyncOwner(EPriority::Normal);
 			Backend->GetRoot().LegacyPut({LegacyRequest}, AsyncOwner, [](auto&&){});
@@ -872,7 +872,7 @@ public:
 	{
 		DDC_SCOPE_CYCLE_COUNTER(DDC_MarkTransient);
 		FLegacyCacheDeleteRequest LegacyRequest;
-		LegacyRequest.Key = FLegacyCacheKey(CacheKey, Backend->GetMaxKeyLength());
+		LegacyRequest.Key = FLegacyCacheKey(CacheKey);
 		LegacyRequest.Name = LegacyRequest.Key.GetFullKey();
 		LegacyRequest.bTransient = true;
 		FRequestOwner BlockingOwner(EPriority::Blocking);
@@ -889,7 +889,7 @@ public:
 		{
 			SCOPE_SECONDS_COUNTER(ThisTime);
 			FLegacyCacheGetRequest LegacyRequest;
-			LegacyRequest.Key = FLegacyCacheKey(CacheKey, Backend->GetMaxKeyLength());
+			LegacyRequest.Key = FLegacyCacheKey(CacheKey);
 			LegacyRequest.Name = LegacyRequest.Key.GetFullKey();
 			LegacyRequest.Policy = ECachePolicy::Query | ECachePolicy::SkipData;
 			FRequestOwner BlockingOwner(EPriority::Blocking);
@@ -917,7 +917,7 @@ public:
 				for (const FString& CacheKey : CacheKeys)
 				{
 					FLegacyCacheGetRequest& LegacyRequest = LegacyRequests.AddDefaulted_GetRef();
-					LegacyRequest.Key = FLegacyCacheKey(CacheKey, Backend->GetMaxKeyLength());
+					LegacyRequest.Key = FLegacyCacheKey(CacheKey);
 					LegacyRequest.Name = LegacyRequest.Key.GetFullKey();
 					LegacyRequest.Policy = ECachePolicy::Query | ECachePolicy::SkipData;
 					LegacyRequest.UserData = uint64(Index);
@@ -957,7 +957,7 @@ public:
 			{
 				FLegacyCacheGetRequest& LegacyRequest = LegacyRequests.AddDefaulted_GetRef();
 				LegacyRequest.Name = Name;
-				LegacyRequest.Key = FLegacyCacheKey(CacheKey, Backend->GetMaxKeyLength());
+				LegacyRequest.Key = FLegacyCacheKey(CacheKey);
 				LegacyRequest.Policy = ECachePolicy::Default | ECachePolicy::SkipData;
 				LegacyRequest.UserData = uint64(Index);
 				++Index;
@@ -980,8 +980,6 @@ public:
 
 	void NotifyBootComplete() override
 	{
-		DDC_SCOPE_CYCLE_COUNTER(DDC_NotifyBootComplete);
-		Backend->NotifyBootComplete();
 	}
 
 	void AddToAsyncCompletionCounter(int32 Addend) override
diff --git a/Engine/Source/Developer/DerivedDataCache/Private/DerivedDataLegacyCacheStore.cpp b/Engine/Source/Developer/DerivedDataCache/Private/DerivedDataLegacyCacheStore.cpp
--- a/Engine/Source/Developer/DerivedDataCache/Private/DerivedDataLegacyCacheStore.cpp
+++ b/Engine/Source/Developer/DerivedDataCache/Private/DerivedDataLegacyCacheStore.cpp
@@ -2,15 +2,7 @@
 
 #include "DerivedDataLegacyCacheStore.h"
 
-#include "Algo/AllOf.h"
-#include "DerivedDataBackendInterface.h"
-#include "DerivedDataCacheInterface.h"
-#include "Misc/Crc.h"
-#include "Misc/ScopeRWLock.h"
-#include "Misc/SecureHash.h"
-#include "Misc/StringBuilder.h"
-#include "ProfilingDebugging/CpuProfilerTrace.h"
-#include "String/BytesToHex.h"
+#include "Async/UniqueLock.h"
 
 namespace UE::DerivedData
 {
@@ -83,99 +75,12 @@ void ILegacyCacheStore::LegacyDelete(
 	CompleteWithStatus(Requests, OnComplete, EStatus::Error);
 }
 
-Private::FLegacyCacheKeyShared::FLegacyCacheKeyShared(const FStringView InFullKey, const int32 InMaxKeyLength)
-	: FullKey(InFullKey)
-	, MaxKeyLength(InMaxKeyLength)
+FLegacyCacheKey::FLegacyCacheKey(const FStringView InFullKey)
+	: Key(ConvertLegacyCacheKey(InFullKey))
+	, FullKey(InFullKey)
 {
 }
 
-const FSharedString& Private::FLegacyCacheKeyShared::GetShortKey()
-{
-	if (FReadScopeLock ReadLock(Lock); !ShortKey.IsEmpty())
-	{
-		return ShortKey;
-	}
-
-	TRACE_CPUPROFILER_EVENT_SCOPE(FLegacyCacheKey::GetShortKey);
-
-	FWriteScopeLock WriteLock(Lock);
-	if (!ShortKey.IsEmpty())
-	{
-		return ShortKey;
-	}
-
-	const FStringView LocalFullKey = FullKey;
-	checkf(Algo::AllOf(LocalFullKey, FDerivedDataCacheInterface::IsValidCacheChar),
-		TEXT("Invalid characters in cache key %s. Use SanitizeCacheKey or BuildCacheKey to create valid keys."), *FullKey);
-
-	if (LocalFullKey.Len() <= MaxKeyLength)
-	{
-		ShortKey = FullKey;
-		return ShortKey;
-	}
-
-	const auto FullKeyUCS2 = StringCast<UCS2CHAR>(LocalFullKey.GetData(), LocalFullKey.Len());
-	const int32 FullKeyLength = FullKeyUCS2.Length();
-	const uint32 CRCofPayload(FCrc::MemCrc32(FullKeyUCS2.Get(), FullKeyLength * sizeof(UCS2CHAR)));
-
-	uint8 Hash[FSHA1::DigestSize];
-	FSHA1 HashState;
-	HashState.Update((const uint8*)&FullKeyLength, sizeof(int32));
-	HashState.Update((const uint8*)&CRCofPayload, sizeof(uint32));
-	HashState.Update((const uint8*)FullKeyUCS2.Get(), FullKeyLength * sizeof(UCS2CHAR));
-	HashState.Final();
-	HashState.GetHash(Hash);
-
-	TStringBuilder<256> ShortKeyBuilder;
-	ShortKeyBuilder << LocalFullKey.Left(MaxKeyLength - FSHA1::DigestSize * 2 - 2) << TEXT("__");
-	String::BytesToHex(Hash, ShortKeyBuilder);
-
-	check(ShortKeyBuilder.Len() == MaxKeyLength && ShortKeyBuilder.Len() > 0);
-	ShortKey = ShortKeyBuilder;
-	return ShortKey;
-}
-
-FLegacyCacheKey::FLegacyCacheKey(const FStringView FullKey, const int32 MaxKeyLength)
-	: Key(ConvertLegacyCacheKey(FullKey))
-	, Shared(new Private::FLegacyCacheKeyShared(FullKey, MaxKeyLength))
-{
-}
-
-bool FLegacyCacheKey::ReadValueTrailer(FCompositeBuffer& Value) const
-{
-	if (HasShortKey())
-	{
-		FUtf8StringBuilderBase FullKey;
-		FullKey << GetFullKey();
-		const FMemoryView CompareKey(*FullKey, (FullKey.Len() + 1) * sizeof(UTF8CHAR));
-		if (Value.GetSize() < CompareKey.GetSize())
-		{
-			UE_LOG(LogDerivedDataCache, Display, TEXT("FLegacyCacheKey: Hash collision or short value for key %s."), *GetFullKey());
-			return false;
-		}
-		FUniqueBuffer CopyBuffer;
-		const uint64 KeyOffset = Value.GetSize() - CompareKey.GetSize();
-		const FMemoryView ValueKey = Value.ViewOrCopyRange(KeyOffset, CompareKey.GetSize(), CopyBuffer);
-		if (!CompareKey.EqualBytes(ValueKey))
-		{
-			UE_LOG(LogDerivedDataCache, Display, TEXT("FLegacyCacheKey: Hash collision for key %s."), *GetFullKey());
-			return false;
-		}
-		Value = Value.Mid(0, KeyOffset);
-	}
-	return true;
-}
-
-void FLegacyCacheKey::WriteValueTrailer(FCompositeBuffer& Value) const
-{
-	if (HasShortKey())
-	{
-		FUtf8StringBuilderBase FullKey;
-		FullKey << GetFullKey();
-		Value = FCompositeBuffer(Value, FSharedBuffer::Clone(*FullKey, (FullKey.Len() + 1) * sizeof(UTF8CHAR)));
-	}
-}
-
 FLegacyCacheValue::FLegacyCacheValue(const FValue& Value)
 {
 	if (Value.HasData() || !Value.GetRawHash().IsZero())
@@ -204,7 +109,7 @@ Private::FLegacyCacheValueShared::FLegacyCacheValueShared(const FCompositeBuffer
 
 const FValue& Private::FLegacyCacheValueShared::GetValue()
 {
-	FWriteScopeLock WriteLock(Lock);
+	TUniqueLock Lock(Mutex);
 	if (!Value.HasData() && !RawData.IsNull())
 	{
 		Value = FValue::Compress(RawData);
@@ -214,7 +119,7 @@ const FValue& Private::FLegacyCacheValueShared::GetValue()
 
 const FCompositeBuffer& Private::FLegacyCacheValueShared::GetRawData()
 {
-	FWriteScopeLock WriteLock(Lock);
+	TUniqueLock Lock(Mutex);
 	if (RawData.IsNull() && Value.HasData())
 	{
 		RawData = Value.GetData().DecompressToComposite();
diff --git a/Engine/Source/Developer/DerivedDataCache/Private/DerivedDataLegacyCacheStore.h b/Engine/Source/Developer/DerivedDataCache/Private/DerivedDataLegacyCacheStore.h
--- a/Engine/Source/Developer/DerivedDataCache/Private/DerivedDataLegacyCacheStore.h
+++ b/Engine/Source/Developer/DerivedDataCache/Private/DerivedDataLegacyCacheStore.h
@@ -2,9 +2,9 @@
 
 #pragma once
 
+#include "Async/Mutex.h"
 #include "DerivedDataCacheStore.h"
 #include "DerivedDataPrivate.h"
-#include "HAL/CriticalSection.h"
 #include "Templates/RefCounting.h"
 #include <atomic>
 
@@ -19,7 +19,6 @@ namespace UE::DerivedData { struct FLegacyCacheGetRequest; }
 namespace UE::DerivedData { struct FLegacyCacheGetResponse; }
 namespace UE::DerivedData { struct FLegacyCachePutRequest; }
 namespace UE::DerivedData { struct FLegacyCachePutResponse; }
-namespace UE::DerivedData::Private { class FLegacyCacheKeyShared; }
 namespace UE::DerivedData::Private { class FLegacyCacheValueShared; }
 
 namespace UE::DerivedData
@@ -54,56 +53,18 @@ public:
 	virtual void GatherVerificationStats(FDerivedDataCacheVerificationStats& OutStats) const { }
 };
 
-class Private::FLegacyCacheKeyShared final
-{
-public:
-	FLegacyCacheKeyShared(FStringView FullKey, int32 MaxKeyLength);
-
-	FLegacyCacheKeyShared(const FLegacyCacheKeyShared&) = delete;
-	FLegacyCacheKeyShared& operator=(const FLegacyCacheKeyShared&) = delete;
-
-	inline bool HasShortKey() const { return FullKey.Len() > MaxKeyLength; }
-	inline const FSharedString& GetFullKey() const { return FullKey; }
-	const FSharedString& GetShortKey();
-
-	inline void AddRef()
-	{
-		ReferenceCount.fetch_add(1, std::memory_order_relaxed);
-	}
-
-	inline void Release()
-	{
-		if (ReferenceCount.fetch_sub(1, std::memory_order_acq_rel) == 1)
-		{
-			delete this;
-		}
-	}
-
-private:
-	FSharedString FullKey;
-	FSharedString ShortKey;
-	int32 MaxKeyLength;
-	std::atomic<uint32> ReferenceCount{0};
-	FRWLock Lock;
-};
-
 class FLegacyCacheKey
 {
 public:
 	FLegacyCacheKey() = default;
-	FLegacyCacheKey(FStringView FullKey, int32 MaxKeyLength);
+	explicit FLegacyCacheKey(FStringView FullKey);
 
 	[[nodiscard]] inline const FCacheKey& GetKey() const { return Key; }
-	[[nodiscard]] inline const FSharedString& GetFullKey() const { return Shared ? Shared->GetFullKey() : FSharedString::Empty; }
-	[[nodiscard]] inline const FSharedString& GetShortKey() const { return Shared ? Shared->GetShortKey() : FSharedString::Empty; }
-	[[nodiscard]] inline bool HasShortKey() const { return Shared ? Shared->HasShortKey() : false; }
-
-	[[nodiscard]] bool ReadValueTrailer(FCompositeBuffer& Value) const;
-	void WriteValueTrailer(FCompositeBuffer& Value) const;
+	[[nodiscard]] inline const FSharedString& GetFullKey() const { return FullKey; }
 
 private:
 	FCacheKey Key;
-	TRefCountPtr<Private::FLegacyCacheKeyShared> Shared;
+	FSharedString FullKey;
 };
 
 class Private::FLegacyCacheValueShared final
@@ -138,7 +99,7 @@ private:
 	FValue Value;
 	FCompositeBuffer RawData;
 	std::atomic<uint32> ReferenceCount{0};
-	FRWLock Lock;
+	FMutex Mutex;
 };
 
 class FLegacyCacheValue
