DDC: Added BuildLegacyCacheKey to support migration from the legacy API

diff --git a/Engine/Source/Developer/DerivedDataCache/Private/DerivedDataCacheKey.cpp b/Engine/Source/Developer/DerivedDataCache/Private/DerivedDataCacheKey.cpp
--- a/Engine/Source/Developer/DerivedDataCache/Private/DerivedDataCacheKey.cpp
+++ b/Engine/Source/Developer/DerivedDataCache/Private/DerivedDataCacheKey.cpp
@@ -2,6 +2,9 @@
 
 #include "DerivedDataCacheKey.h"
 
+#include "Async/SharedLock.h"
+#include "Async/SharedMutex.h"
+#include "Async/UniqueLock.h"
 #include "Containers/Map.h"
 #include "Containers/Set.h"
 #include "Containers/StringConv.h"
@@ -9,7 +12,7 @@
 #include "Hash/xxhash.h"
 #include "IO/IoHash.h"
 #include "Math/UnrealMathUtility.h"
-#include "Misc/ScopeRWLock.h"
+#include "Misc/AsciiSet.h"
 #include "Misc/StringBuilder.h"
 #include "Serialization/CompactBinary.h"
 #include "Serialization/CompactBinarySerialization.h"
@@ -99,7 +102,7 @@ public:
 	inline void SetDisplayName(FCacheBucket Bucket, FStringView DisplayName);
 
 private:
-	FRWLock Lock;
+	FSharedMutex Mutex;
 	TSet<FCacheBucketOwner, FCacheBucketOwnerKeyFuncs> Buckets;
 	TMap<FCacheBucket, FString> DisplayNames;
 };
@@ -114,7 +117,7 @@ inline FCacheBucket FCacheBuckets::FindOrAdd(const TStringView<CharType> Name)
 	if (NameView.Len() <= FCacheBucket::MaxNameLen)
 	{
 		Hash = FCacheBucketOwnerKeyFuncs::GetKeyHash(NameView);
-		FReadScopeLock ReadLock(Lock);
+		TSharedLock ReadLock(Mutex);
 		if (const FCacheBucketOwner* Bucket = Buckets.FindByHash(Hash, NameView))
 		{
 			return *Bucket;
@@ -122,13 +125,13 @@ inline FCacheBucket FCacheBuckets::FindOrAdd(const TStringView<CharType> Name)
 	}
 
 	FCacheBucketOwner LocalBucket(NameView);
-	FWriteScopeLock WriteLock(Lock);
+	TUniqueLock WriteLock(Mutex);
 	return Buckets.FindOrAddByHash(Hash, MoveTemp(LocalBucket));
 }
 
 void FCacheBuckets::GetDisplayName(FCacheBucket Bucket, FStringBuilderBase& OutDisplayName)
 {
-	if (FReadScopeLock ReadLock(Lock); const FString* DisplayName = DisplayNames.Find(Bucket))
+	if (TSharedLock ReadLock(Mutex); const FString* DisplayName = DisplayNames.Find(Bucket))
 	{
 		OutDisplayName << *DisplayName;
 		return;
@@ -143,7 +146,7 @@ void FCacheBuckets::GetDisplayName(FCacheBucket Bucket, FStringBuilderBase& OutD
 
 void FCacheBuckets::SetDisplayName(FCacheBucket Bucket, FStringView DisplayName)
 {
-	if (FReadScopeLock ReadLock(Lock); const FString* ExistingDisplayName = DisplayNames.Find(Bucket))
+	if (TSharedLock ReadLock(Mutex); const FString* ExistingDisplayName = DisplayNames.Find(Bucket))
 	{
 		if (*ExistingDisplayName == DisplayName)
 		{
@@ -151,7 +154,7 @@ void FCacheBuckets::SetDisplayName(FCacheBucket Bucket, FStringView DisplayName)
 		}
 	}
 
-	FWriteScopeLock WriteLock(Lock);
+	TUniqueLock WriteLock(Mutex);
 	DisplayNames.Emplace(Bucket, DisplayName);
 }
 
@@ -251,4 +254,35 @@ FCacheKey ConvertLegacyCacheKey(const FStringView Key)
 	return {Bucket, FIoHash::HashBuffer(MakeMemoryView(Utf8Key))};
 }
 
+FCacheKey BuildLegacyCacheKey(FStringView PluginName, FStringView VersionString, FStringView PluginSpecificCacheKeySuffix)
+{
+	const auto AppendSanitized = [](FStringBuilderBase& Out, FStringView In)
+	{
+		constexpr FAsciiSet ValidChars = "ABCDEFGHIJKLMNOPQRSTUVWXYZabcdefghijklmnopqrstuvwxyz0123456789_$";
+		for (;;)
+		{
+			FStringView ValidPrefix = FAsciiSet::FindPrefixWith(In, ValidChars);
+			Out.Append(ValidPrefix);
+			In.RemovePrefix(ValidPrefix.Len());
+
+			if (In.IsEmpty())
+			{
+				break;
+			}
+
+			Out.Appendf(TEXT("$%x"), uint32(In.GetData()[0]));
+			In.RemovePrefix(1);
+		}
+	};
+
+	TStringBuilder<1024> Out;
+	Out.Reserve(PluginName.Len() + 1 + VersionString.Len() + 1 + PluginSpecificCacheKeySuffix.Len());
+	AppendSanitized(Out, PluginName);
+	Out.AppendChar(TEXT('_'));
+	AppendSanitized(Out, VersionString);
+	Out.AppendChar(TEXT('_'));
+	AppendSanitized(Out, PluginSpecificCacheKeySuffix);
+	return ConvertLegacyCacheKey(Out);
+}
+
 } // UE::DerivedData
diff --git a/Engine/Source/Developer/DerivedDataCache/Public/DerivedDataCacheKey.h b/Engine/Source/Developer/DerivedDataCache/Public/DerivedDataCacheKey.h
--- a/Engine/Source/Developer/DerivedDataCache/Public/DerivedDataCacheKey.h
+++ b/Engine/Source/Developer/DerivedDataCache/Public/DerivedDataCacheKey.h
@@ -157,10 +157,19 @@ UE_API bool LoadFromCompactBinary(FCbFieldView Field, FCacheKey& OutKey);
 /**
  * Constructs a cache key from a legacy cache key.
  *
+ * Keys converted with this function are the same keys used internally by the legacy cache API.
+ *
  * Legacy keys are typically of the form DATATYPENAME_VersionGuid_TypeSpecificSuffix.
  */
 UE_API FCacheKey ConvertLegacyCacheKey(FStringView Key);
 
+/**
+ * Builds a cache key from the same components as the legacy BuildCacheKey API.
+ *
+ * Equivalent to ConvertLegacyCacheKey(FDerivedDataCacheInterface::BuildCacheKey(PluginName, VersionString, PluginSpecificCacheKeySuffix)).
+ */
+UE_API FCacheKey BuildLegacyCacheKey(FStringView PluginName, FStringView VersionString, FStringView PluginSpecificCacheKeySuffix);
+
 ///////////////////////////////////////////////////////////////////////////////////////////////////////////////////////
 
 } // UE::DerivedData
