DDC: Extended legacy config parsing to handle deprecated nodes anywhere in the graph

diff --git a/Engine/Programs/LowLevelTests/Config/DefaultEngine.ini b/Engine/Programs/LowLevelTests/Config/DefaultEngine.ini
--- a/Engine/Programs/LowLevelTests/Config/DefaultEngine.ini
+++ b/Engine/Programs/LowLevelTests/Config/DefaultEngine.ini
@@ -32,11 +32,12 @@ TestStore4 = ( Key1 = Store4Value , Base = TestStore2 ) ; Comment
 TestStore5 = ( Key1 = Store5Value , Key2 = Store5Value ) ; Comment
 TestStore6 = ( Base = TestStore1 ) ; Comment
 TestStore7 = ( Base = TestStore1 , Deprecated = "Use TestStore1." ) ; Comment
+TestStore8 = ( Type = Cloud ) ; Comment
 
 [TestLegacyCacheGraph1]
 Root=(Type=KeyLength, Length=120, Inner=Verify)
 Verify=(Type=Verify, Inner=AsyncPut)
 AsyncPut=(Type=AsyncPut, Inner=Hierarchy)
-Hierarchy=(Type=Hierarchical, Inner=TestLegacyStore1, Inner=TestStore2, Inner=TestStore3)
-TestLegacyStore1=(Key1=LegacyStore1Value, Key2=LegacyStore1Value)
-TestStore2=(Key1=LegacyStore2Value)
+Hierarchy=(Type=Hierarchical, Inner=TestLegacyStore1, Inner=TestStore2, Inner=TestStore8)
+TestLegacyStore1=(Type=Zen, Key1=LegacyStore1Value, Key2=LegacyStore1Value)
+TestStore2=(Type=Zen, Key1=LegacyStore2Value)
diff --git a/Engine/Source/Developer/DerivedDataCache/Private/DerivedDataConfig.cpp b/Engine/Source/Developer/DerivedDataCache/Private/DerivedDataConfig.cpp
--- a/Engine/Source/Developer/DerivedDataCache/Private/DerivedDataConfig.cpp
+++ b/Engine/Source/Developer/DerivedDataCache/Private/DerivedDataConfig.cpp
@@ -309,58 +309,61 @@ private:
 namespace UE::DerivedData
 {
 
-bool TryFindCacheGraphConfig(const TCHAR* GraphName, FStringBuilderBase& OutConfig)
+static bool TryConvertLegacyCacheStoreConfig(const TCHAR* GraphName, FStringView StoreName, Config::FConfigBuilder& Builder);
+
+static bool TryConvertLegacyCacheGraphConfig(const TCHAR* GraphName, FStringView StoreName, FStringView StoreConfig, Config::FConfigBuilder& Builder)
 {
-	// UE_DEPRECATED(5.8, "Remove the LegacyConfig block after the standard deprecation period.")
-	FString LegacyConfig;
-	if (FString StoreName = TEXT("Root"); GConfig->GetString(GraphName, *StoreName, LegacyConfig, GEngineIni))
+	if (FStringView StoreType; Config::FindFirstFieldValue(StoreConfig, TEXTVIEW("Type"), StoreType))
 	{
-		//UE_LOGFMT(LogDerivedDataCache, Warning,
-		//	"Configuring cache graph {GraphName} in config section [{GraphName}] of {IniName}.ini is deprecated. "
-		//	"Migrate the configuration to the [DerivedDataCacheGraphs] section of {IniName}.ini.",
-		//	("GraphName", GraphName), ("IniName", GEngineIni));
-		for (FStringView Type;;)
+		bool bIsHierarchy = (StoreType == TEXTVIEW("Hierarchical"));
+		if (bIsHierarchy || (StoreType == TEXTVIEW("AsyncPut") || StoreType == TEXTVIEW("KeyLength") || StoreType == TEXTVIEW("Verify")))
 		{
-			if (Config::FindFirstFieldValue(LegacyConfig, TEXTVIEW("Type"), Type))
+			UE_CLOGFMT(!bIsHierarchy, LogDerivedDataCache, Warning,
+				"The cache graph {GraphName} contains cache store {StoreName} of deprecated type {StoreType}.",
+				GraphName, StoreName, StoreType);
+			int32 InnerCount = 0;
+			Config::FConfigIterator It(StoreConfig);
+			for (; It && (bIsHierarchy || InnerCount == 0); ++It)
 			{
-				if (Type == TEXTVIEW("Hierarchical"))
-				{
-					Config::FConfigBuilder Builder(OutConfig);
-					Config::FConfigIterator It(LegacyConfig);
-					for (; It; ++It)
-					{
-						if (It.GetName() == TEXTVIEW("Inner"))
-						{
-							Builder.Append({}, It.GetValue());
-						}
-					}
-					if (It.HasError())
-					{
-						Builder.SetError();
-					}
-					return !Builder.HasError();
-				}
-				if (Type == TEXTVIEW("AsyncPut") || Type == TEXTVIEW("KeyLength") || Type == TEXTVIEW("Verify"))
+				if (It.GetName() == TEXTVIEW("Inner"))
 				{
-					UE_LOGFMT(LogDerivedDataCache, Warning,
-						"The cache graph {GraphName} contains cache store {StoreName} of deprecated type {StoreType}.",
-						GraphName, StoreName, Type);
-					if (FStringView Inner; Config::FindFirstFieldValue(LegacyConfig, TEXTVIEW("Inner"), Inner))
+					++InnerCount;
+					if (!TryConvertLegacyCacheStoreConfig(GraphName, It.GetValue(), Builder))
 					{
-						StoreName = Inner;
-						if (GConfig->GetString(GraphName, *StoreName, LegacyConfig, GEngineIni))
-						{
-							continue;
-						}
+						return false;
 					}
 				}
 			}
-			UE_LOGFMT(LogDerivedDataCache, Warning,
-				"Failed to parse cache store {StoreName} from cache graph {GraphName} in config section [{GraphName}] of {IniName}.ini. "
-				"Migrate the configuration to the [DerivedDataCacheGraphs] section of {IniName}.ini.",
-				("GraphName", GraphName), ("StoreName", StoreName), ("IniName", GEngineIni));
-			break;
+			if (!bIsHierarchy && InnerCount == 0)
+			{
+				UE_LOGFMT(LogDerivedDataCache, Warning,
+					"Failed to parse inner cache store in {StoreName} from cache graph {GraphName}.", StoreName, GraphName);
+				return false;
+			}
+			return !It.HasError();
 		}
+		Builder.Append({}, StoreName);
+		return true;
+	}
+	UE_LOGFMT(LogDerivedDataCache, Warning,
+		"Failed to parse cache store {StoreName} from cache graph {GraphName} in config section [{GraphName}] of {IniName}.ini. "
+		"Migrate the configuration to the [DerivedDataCacheGraphs] section of {IniName}.ini.",
+		("GraphName", GraphName), ("StoreName", StoreName), ("IniName", GEngineIni));
+	return false;
+}
+
+bool TryFindCacheGraphConfig(const TCHAR* GraphName, FStringBuilderBase& OutConfig)
+{
+	// UE_DEPRECATED(5.8, "Remove the LegacyConfig block after the standard deprecation period.")
+	FString LegacyConfig;
+	if (const TCHAR* StoreName = TEXT("Root"); GConfig->GetString(GraphName, StoreName, LegacyConfig, GEngineIni))
+	{
+		//UE_LOGFMT(LogDerivedDataCache, Warning,
+		//	"Configuring cache graph {GraphName} in config section [{GraphName}] of {IniName}.ini is deprecated. "
+		//	"Migrate the configuration to the [DerivedDataCacheGraphs] section of {IniName}.ini.",
+		//	("GraphName", GraphName), ("IniName", GEngineIni));
+		Config::FConfigBuilder Builder(OutConfig);
+		return TryConvertLegacyCacheGraphConfig(GraphName, StoreName, LegacyConfig, Builder);
 	}
 
 	FString Config;
@@ -401,7 +404,7 @@ bool TryFindCacheGraphConfig(const TCHAR* GraphName, FStringBuilderBase& OutConf
 class FCacheStoreConfigBuilder
 {
 public:
-	explicit FCacheStoreConfigBuilder(FStringBuilderBase& OutConfig);
+	explicit FCacheStoreConfigBuilder(FStringBuilderBase& OutConfig, bool bLogLegacyConfig = true);
 
 	bool TryFind(const TCHAR* StoreName, const TCHAR* GraphName);
 
@@ -410,10 +413,12 @@ private:
 	bool TryBuild(FStringView StoreName, FStringView Config);
 
 	Config::FConfigBuilder Builder;
+	bool bLogLegacyConfig;
 };
 
-FCacheStoreConfigBuilder::FCacheStoreConfigBuilder(FStringBuilderBase& OutConfig)
+FCacheStoreConfigBuilder::FCacheStoreConfigBuilder(FStringBuilderBase& OutConfig, bool bInLogLegacyConfig)
 	: Builder(OutConfig)
+	, bLogLegacyConfig(bInLogLegacyConfig)
 {
 }
 
@@ -433,7 +438,7 @@ bool FCacheStoreConfigBuilder::TryQuery(FStringView StoreName, const TCHAR* Grap
 	FString LegacyConfig;
 	if (GraphName && GConfig->GetString(GraphName, *WriteToString<64>(StoreName), LegacyConfig, GEngineIni))
 	{
-		//UE_LOGFMT(LogDerivedDataCache, Warning,
+		//UE_CLOGFMT(bLogLegacyConfig, LogDerivedDataCache, Warning,
 		//	"Configuring cache store {StoreName} in config section [{GraphName}] of {IniName}.ini is deprecated. "
 		//	"Migrate the configuration to the [DerivedDataCacheStores] section of {IniName}.ini.",
 		//	("GraphName", GraphName), ("StoreName", StoreName), ("IniName", GEngineIni));
@@ -519,4 +524,11 @@ bool TryFindCacheStoreConfig(const TCHAR* StoreName, const TCHAR* GraphName, FSt
 	return Builder.TryFind(StoreName, GraphName);
 }
 
+static bool TryConvertLegacyCacheStoreConfig(const TCHAR* GraphName, FStringView StoreName, Config::FConfigBuilder& Builder)
+{
+	TStringBuilder<256> StoreConfig;
+	bool bOk = FCacheStoreConfigBuilder(StoreConfig, /*bLogLegacyConfig*/ false).TryFind(*WriteToString<64>(StoreName), GraphName);
+	return bOk && TryConvertLegacyCacheGraphConfig(GraphName, StoreName, StoreConfig, Builder);
+}
+
 } // UE::DerivedData
diff --git a/Engine/Source/Developer/DerivedDataCache/Tests/DerivedDataConfigTest.cpp b/Engine/Source/Developer/DerivedDataCache/Tests/DerivedDataConfigTest.cpp
--- a/Engine/Source/Developer/DerivedDataCache/Tests/DerivedDataConfigTest.cpp
+++ b/Engine/Source/Developer/DerivedDataCache/Tests/DerivedDataConfigTest.cpp
@@ -71,7 +71,7 @@ TEST_CASE("DerivedData::CacheConfig", "[DerivedData]")
 			CHECK_LOG_SCOPE(LogDerivedDataCache, Warning, "The cache graph TestLegacyCacheGraph1 contains cache store Root of deprecated type KeyLength.");
 			CHECK_LOG_SCOPE(LogDerivedDataCache, Warning, "The cache graph TestLegacyCacheGraph1 contains cache store Verify of deprecated type Verify.");
 			CHECK_LOG_SCOPE(LogDerivedDataCache, Warning, "The cache graph TestLegacyCacheGraph1 contains cache store AsyncPut of deprecated type AsyncPut.");
-			CHECK(TestFindGraph(TEXT("TestLegacyCacheGraph1"), TEXT("(TestLegacyStore1, TestStore2, TestStore3)")));
+			CHECK(TestFindGraph(TEXT("TestLegacyCacheGraph1"), TEXT("(TestLegacyStore1, TestStore2, TestStore8)")));
 		}
 	}
 
@@ -122,9 +122,9 @@ TEST_CASE("DerivedData::CacheConfig", "[DerivedData]")
 			CHECK(TestFindStore(TEXT("TestStore7"), nullptr, TEXT("(Key1=Store1Value, Key2=Store1Value, ValueOnly, (X=\"A=1, \\\"B\\\"=2, C=3\"))")));
 		}
 
-		CHECK(TestFindStore(TEXT("TestLegacyStore1"), TEXT("TestLegacyCacheGraph1"), TEXT("(Key1=LegacyStore1Value, Key2=LegacyStore1Value)")));
-		CHECK(TestFindStore(TEXT("TestStore2"), TEXT("TestLegacyCacheGraph1"), TEXT("(Key1=LegacyStore2Value)")));
-		CHECK(TestFindStore(TEXT("TestStore3"), TEXT("TestLegacyCacheGraph1"), TEXT("(ValueOnly, Key1=Store3Value, Key2=Store3Value, Key3=(A=1, B=2, C))")));
+		CHECK(TestFindStore(TEXT("TestLegacyStore1"), TEXT("TestLegacyCacheGraph1"), TEXT("(Type=Zen, Key1=LegacyStore1Value, Key2=LegacyStore1Value)")));
+		CHECK(TestFindStore(TEXT("TestStore2"), TEXT("TestLegacyCacheGraph1"), TEXT("(Type=Zen, Key1=LegacyStore2Value)")));
+		CHECK(TestFindStore(TEXT("TestStore8"), TEXT("TestLegacyCacheGraph1"), TEXT("(Type=Cloud)")));
 	}
 }
 
