DDC: Added custom cache graph support

A custom cache graph can be provided on the command line with, e.g., -DDC=(EnginePak,Cloud), or in the same format with the DDC.Graph CVar. Any configuration that is valid in a graph is supported here, such as -DDC=(ZenLocal,Cloud(ReadOnly=true)) to create a graph with ZenLocal and a read-only Cloud.

diff --git a/Engine/Source/Developer/DerivedDataCache/Private/DerivedDataBackends.cpp b/Engine/Source/Developer/DerivedDataCache/Private/DerivedDataBackends.cpp
--- a/Engine/Source/Developer/DerivedDataCache/Private/DerivedDataBackends.cpp
+++ b/Engine/Source/Developer/DerivedDataCache/Private/DerivedDataBackends.cpp
@@ -44,7 +44,7 @@ LLM_DEFINE_TAG(DDCBackend);
 static TAutoConsoleVariable<FString> GDerivedDataCacheGraphName(
 	TEXT("DDC.Graph"),
 	TEXT("Default"),
-	TEXT("Name of the graph to use for the Derived Data Cache."),
+	TEXT("Name or config of the graph to use for the Derived Data Cache."),
 	ECVF_ReadOnly);
 
 namespace UE::DerivedData::Private
@@ -121,7 +121,7 @@ public:
 
 		// Create the graph using ini settings. The string "default" forwards creation to use the default graph.
 
-		if (!FParse::Value(FCommandLine::Get(), TEXT("-DDC="), ActiveGraphName))
+		if (!FParse::Value(FCommandLine::Get(), TEXT("-DDC="), ActiveGraphName, /*bShouldStopOnSeparator*/ false))
 		{
 			ActiveGraphName = GDerivedDataCacheGraphName.GetValueOnGameThread();
 		}
@@ -156,11 +156,11 @@ public:
 
 			if (!ActiveGraphName.IsEmpty() && ActiveGraphName != TEXT("Default"))
 			{
-				if (!ParseGraph(*ActiveGraphName, ParsedNodes) || !Hierarchy->HasAllFlags(ECacheStoreFlags::Query | ECacheStoreFlags::Store))
+				if (!ParseGraph(ActiveGraphName, ParsedNodes) || !Hierarchy->HasAllFlags(ECacheStoreFlags::Query | ECacheStoreFlags::Store))
 				{
 					ResetGraph();
 					UE_LOGFMT(LogDerivedDataCache, Warning,
-						"Failed to create graph {GraphName}. Attempting to create the default graph...", ActiveGraphName);
+						"Failed to create graph '{Graph}'. Attempting to create the default graph...", ActiveGraphName);
 				}
 			}
 
@@ -169,7 +169,7 @@ public:
 			{
 				// Try to use the legacy default graph.
 				ActiveGraphName = FApp::IsEngineInstalled() ? TEXT("InstalledDerivedDataBackendGraph") : TEXT("DerivedDataBackendGraph");
-				if (!ParseGraph(*ActiveGraphName, ParsedNodes) || !Hierarchy->HasAllFlags(ECacheStoreFlags::Query | ECacheStoreFlags::Store))
+				if (!ParseGraph(ActiveGraphName, ParsedNodes) || !Hierarchy->HasAllFlags(ECacheStoreFlags::Query | ECacheStoreFlags::Store))
 				{
 					ResetGraph();
 				}
@@ -179,13 +179,13 @@ public:
 			{
 				// Try to use the default graph.
 				ActiveGraphName = FApp::IsEngineInstalled() ? TEXT("Installed") : TEXT("Default");
-				if (!ParseGraph(*ActiveGraphName, ParsedNodes) || !Hierarchy->HasAllFlags(ECacheStoreFlags::Query | ECacheStoreFlags::Store))
+				if (!ParseGraph(ActiveGraphName, ParsedNodes) || !Hierarchy->HasAllFlags(ECacheStoreFlags::Query | ECacheStoreFlags::Store))
 				{
 					TStringBuilder<256> GraphConfig;
 					if (!TryFindCacheGraphConfig(*ActiveGraphName, GraphConfig))
 					{
 						UE_LOGFMT(LogDerivedDataCache, Fatal,
-							"Failed to create default cache graph '{GraphName}' because its config is missing in the '{IniName}' config.",
+							"Failed to create default cache graph '{Graph}' because its config is missing in the '{Config}' config.",
 							ActiveGraphName, GEngineIni);
 					}
 					else if (FParse::Param(FCommandLine::Get(), TEXT("DDC-ForceMemoryCache")))
@@ -193,7 +193,7 @@ public:
 						IMemoryCacheStore* Memory = nullptr;
 						CreateMemoryCacheStore(Memory, TEXT("ForceMemoryCache"), TEXT(""), Hierarchy);
 						UE_LOGFMT(LogDerivedDataCache, Error,
-							"Unable to use default cache graph '{GraphName}' because there are no {MissingTypes} nodes available.",
+							"Unable to use default cache graph '{Graph}' because there are no {MissingTypes} nodes available.",
 							ActiveGraphName,
 							Hierarchy->HasAllFlags(ECacheStoreFlags::Query) ? "writable" :
 							Hierarchy->HasAllFlags(ECacheStoreFlags::Store) ? "readable" : "readable or writable");
@@ -203,7 +203,7 @@ public:
 					{
 						UE_CLOGFMT(!FParse::Param(FCommandLine::Get(), TEXT("DDC-AllowNoActiveStores")),
 							LogDerivedDataCache, Fatal,
-							"Unable to use default cache graph '{GraphName}' because there are no {MissingTypes} nodes available."
+							"Unable to use default cache graph '{Graph}' because there are no {MissingTypes} nodes available."
 							"Add -DDC-ForceMemoryCache to the command line to bypass this if you need access "
 							"to the editor settings to fix the cache configuration.",
 							*ActiveGraphName,
@@ -245,18 +245,18 @@ public:
 	 * @param InParsedNodes Map of parsed nodes and their names to be able to find already parsed nodes
 	 * @return Derived data backend interface instance created from ini settings
 	 */
-	bool ParseNode(const TCHAR* StoreName, const TCHAR* GraphName, FParsedNodeMap& InParsedNodes)
+	bool ParseNode(FString StoreName, FStringView GraphNameOrConfig, FParsedNodeMap& InParsedNodes)
 	{
 		if (const FParsedNode* ParsedNode = InParsedNodes.Find(StoreName))
 		{
-			UE_LOGFMT(LogDerivedDataCache, Display, "Node {StoreName} was referenced more than once in cache graph {GraphName}. Nodes may not be shared.", StoreName, GraphName);
+			UE_LOGFMT(LogDerivedDataCache, Display, "Node '{Store}' was referenced more than once in cache graph '{Graph}'. Nodes may not be shared.", StoreName, GraphNameOrConfig);
 			return false;
 		}
 
 		TStringBuilder<256> StoreConfig;
-		if (!TryFindCacheStoreConfig(StoreName, GraphName, StoreConfig))
+		if (!TryFindCacheStoreConfig(StoreName, GraphNameOrConfig, StoreConfig))
 		{
-			UE_LOGFMT(LogDerivedDataCache, Log, "Failed to find config for cache store {StoreName} in cache graph {GraphName}.", StoreName, GraphName);
+			UE_LOGFMT(LogDerivedDataCache, Log, "Failed to find config for cache store '{Store}' in cache graph '{Graph}'.", StoreName, GraphNameOrConfig);
 			return false;
 		}
 
@@ -267,7 +267,7 @@ public:
 			Entry.RemoveFromStart(TEXT("("));
 			Entry.RemoveFromEnd(TEXT(")"));
 
-			TGuardValue NodeNameGuard(ActiveNodeName, FString(StoreName));
+			TGuardValue NodeNameGuard(ActiveNodeName, StoreName);
 			TGuardValue NodeConfigGuard(ActiveNodeConfig, Entry);
 
 			FString	NodeType;
@@ -275,23 +275,23 @@ public:
 			{
 				if (NodeType == TEXT("FileSystem"))
 				{
-					return ParseFileSystemCache(StoreName, *Entry, InParsedNodes);
+					return ParseFileSystemCache(*StoreName, *Entry, InParsedNodes);
 				}
 				else if (NodeType == TEXT("Memory"))
 				{
-					return ParseMemoryCache(StoreName, *Entry, InParsedNodes);
+					return ParseMemoryCache(*StoreName, *Entry, InParsedNodes);
 				}
 				else if (NodeType == TEXT("ReadPak"))
 				{
-					return ParsePak(StoreName, *Entry, /*bWriting*/ false, InParsedNodes);
+					return ParsePak(*StoreName, *Entry, /*bWriting*/ false, InParsedNodes);
 				}
 				else if (NodeType == TEXT("WritePak"))
 				{
-					return ParsePak(StoreName, *Entry, /*bWriting*/ true, InParsedNodes);
+					return ParsePak(*StoreName, *Entry, /*bWriting*/ true, InParsedNodes);
 				}
 				else if (NodeType == TEXT("S3"))
 				{
-					if (ILegacyCacheStore* CacheStore = CreateS3CacheStore(StoreName, *Entry, *this))
+					if (ILegacyCacheStore* CacheStore = CreateS3CacheStore(*StoreName, *Entry, *this))
 					{
 						InParsedNodes.Add(StoreName, CacheStore);
 						return true;
@@ -299,7 +299,7 @@ public:
 				}
 				else if (NodeType == TEXT("Cloud") || NodeType == TEXT("Http"))
 				{
-					if (ILegacyCacheStore* CacheStore = CreateHttpCacheStore(StoreName, *Entry, this))
+					if (ILegacyCacheStore* CacheStore = CreateHttpCacheStore(*StoreName, *Entry, this))
 					{
 						InParsedNodes.Add(StoreName, CacheStore);
 						return true;
@@ -307,7 +307,7 @@ public:
 				}
 				else if (NodeType == TEXT("Zen"))
 				{
-					if (ILegacyCacheStore* CacheStore = CreateZenCacheStore(StoreName, *Entry, this))
+					if (ILegacyCacheStore* CacheStore = CreateZenCacheStore(*StoreName, *Entry, this))
 					{
 						InParsedNodes.Add(StoreName, CacheStore);
 						return true;
@@ -385,12 +385,12 @@ public:
 	 * @param InParsedNodes map of nodes and their names which have already been parsed
 	 * @return Hierarchical backend interface instance or nullptr if unsuccessful
 	 */
-	bool ParseGraph(const TCHAR* GraphName, FParsedNodeMap& InParsedNodes)
+	bool ParseGraph(FStringView GraphNameOrConfig, FParsedNodeMap& InParsedNodes)
 	{
 		RootCache = CreateCacheStoreHierarchy(Hierarchy, [this](IMemoryCacheStore*& OutCache) { GetMemoryCache(OutCache); });
 
 		TStringBuilder<256> GraphConfig;
-		if (!TryFindCacheGraphConfig(GraphName, GraphConfig))
+		if (!TryFindCacheGraphConfig(GraphNameOrConfig, GraphConfig))
 		{
 			return false;
 		}
@@ -399,18 +399,18 @@ public:
 		Config::FConfigIterator It(GraphConfig);
 		for (; It; ++It)
 		{
-			bParsed |= ParseNode(*WriteToString<64>(It.GetValue()), GraphName, InParsedNodes);
+			bParsed |= ParseNode(FString(It.GetValue()), GraphNameOrConfig, InParsedNodes);
 		}
 
 		if (It.HasError())
 		{
-			UE_LOGFMT(LogDerivedDataCache, Warning, "Failed to parse at cache graph {GraphName}.", GraphName);
+			UE_LOGFMT(LogDerivedDataCache, Warning, "Failed to parse cache graph '{Graph}'.", GraphNameOrConfig);
 			return false;
 		}
 
 		if (!bParsed)
 		{
-			UE_LOGFMT(LogDerivedDataCache, Warning, "Failed to create at least one cache store for cache graph {GraphName}.", GraphName);
+			UE_LOGFMT(LogDerivedDataCache, Warning, "Failed to create at least one cache store for cache graph '{Graph}'.", GraphNameOrConfig);
 			return false;
 		}
 
@@ -782,13 +782,12 @@ private:
 			return *ParsedNode;
 		}
 
-		if (!ParseNode(Name, *ActiveGraphName, *ActiveParsedNodeMap))
+		if (!ParseNode(Name, ActiveGraphName, *ActiveParsedNodeMap))
 		{
 			return nullptr;
 		}
 
-		const FParsedNode* ParsedNode = ActiveParsedNodeMap->Find(Name);
-		if (ParsedNode)
+		if (const FParsedNode* ParsedNode = ActiveParsedNodeMap->Find(Name))
 		{
 			return *ParsedNode;
 		}
diff --git a/Engine/Source/Developer/DerivedDataCache/Private/DerivedDataConfig.cpp b/Engine/Source/Developer/DerivedDataCache/Private/DerivedDataConfig.cpp
--- a/Engine/Source/Developer/DerivedDataCache/Private/DerivedDataConfig.cpp
+++ b/Engine/Source/Developer/DerivedDataCache/Private/DerivedDataConfig.cpp
@@ -309,9 +309,9 @@ private:
 namespace UE::DerivedData
 {
 
-static bool TryConvertLegacyCacheStoreConfig(const TCHAR* GraphName, FStringView StoreName, Config::FConfigBuilder& Builder);
+static bool TryConvertLegacyCacheStoreConfig(FStringView StoreName, FStringView GraphName, Config::FConfigBuilder& Builder);
 
-static bool TryConvertLegacyCacheGraphConfig(const TCHAR* GraphName, FStringView StoreName, FStringView StoreConfig, Config::FConfigBuilder& Builder)
+static bool TryConvertLegacyCacheGraphConfig(FStringView GraphName, FStringView StoreName, FStringView StoreConfig, Config::FConfigBuilder& Builder)
 {
 	if (FStringView StoreType; Config::FindFirstFieldValue(StoreConfig, TEXTVIEW("Type"), StoreType))
 	{
@@ -319,7 +319,7 @@ static bool TryConvertLegacyCacheGraphConfig(const TCHAR* GraphName, FStringView
 		if (bIsHierarchy || (StoreType == TEXTVIEW("AsyncPut") || StoreType == TEXTVIEW("KeyLength") || StoreType == TEXTVIEW("Verify")))
 		{
 			UE_CLOGFMT(!bIsHierarchy, LogDerivedDataCache, Warning,
-				"The cache graph {GraphName} contains cache store {StoreName} of deprecated type {StoreType}.",
+				"The cache graph '{Graph}' contains cache store '{Store}' of deprecated type '{StoreType}'.",
 				GraphName, StoreName, StoreType);
 			int32 InnerCount = 0;
 			Config::FConfigIterator It(StoreConfig);
@@ -328,7 +328,7 @@ static bool TryConvertLegacyCacheGraphConfig(const TCHAR* GraphName, FStringView
 				if (It.GetName() == TEXTVIEW("Inner"))
 				{
 					++InnerCount;
-					if (!TryConvertLegacyCacheStoreConfig(GraphName, It.GetValue(), Builder))
+					if (!TryConvertLegacyCacheStoreConfig(It.GetValue(), GraphName, Builder))
 					{
 						return false;
 					}
@@ -337,7 +337,7 @@ static bool TryConvertLegacyCacheGraphConfig(const TCHAR* GraphName, FStringView
 			if (!bIsHierarchy && InnerCount == 0)
 			{
 				UE_LOGFMT(LogDerivedDataCache, Warning,
-					"Failed to parse inner cache store in {StoreName} from cache graph {GraphName}.", StoreName, GraphName);
+					"Failed to parse inner cache store in '{Store}' from cache graph '{Graph}'.", StoreName, GraphName);
 				return false;
 			}
 			return !It.HasError();
@@ -346,59 +346,76 @@ static bool TryConvertLegacyCacheGraphConfig(const TCHAR* GraphName, FStringView
 		return true;
 	}
 	UE_LOGFMT(LogDerivedDataCache, Warning,
-		"Failed to parse cache store {StoreName} from cache graph {GraphName} in config section [{GraphName}] of {IniName}.ini. "
-		"Migrate the configuration to the [DerivedDataCacheGraphs] section of {IniName}.ini.",
-		("GraphName", GraphName), ("StoreName", StoreName), ("IniName", GEngineIni));
+		"Failed to parse cache store '{Store}' from cache graph '{Graph}' in config section [{Graph}] of the '{Config}' config. "
+		"Migrate the configuration to the [DerivedDataCacheGraphs] section of the '{Config}' config.",
+		("Graph", GraphName), ("Store", StoreName), ("Config", GEngineIni));
 	return false;
 }
 
-bool TryFindCacheGraphConfig(const TCHAR* GraphName, FStringBuilderBase& OutConfig)
+bool TryFindCacheGraphConfig(FStringView GraphNameOrConfig, FStringBuilderBase& OutConfig)
 {
+	FStringView GraphName;
+	FStringView GraphConfig;
+	if (!Config::TrySplitNameAndConfig(GraphNameOrConfig, GraphName, GraphConfig))
+	{
+		return false;
+	}
+
+	if (!GraphName.IsEmpty() && !GraphConfig.IsEmpty())
+	{
+		return false;
+	}
+
 	// UE_DEPRECATED(5.8, "Remove the LegacyConfig block after the standard deprecation period.")
-	FString LegacyConfig;
-	if (const TCHAR* StoreName = TEXT("Root"); GConfig->GetString(GraphName, StoreName, LegacyConfig, GEngineIni))
+	if (FString LegacyConfig; !GraphName.IsEmpty() && GConfig->GetString(*WriteToString<64>(GraphName), TEXT("Root"), LegacyConfig, GEngineIni))
 	{
 		//UE_LOGFMT(LogDerivedDataCache, Warning,
-		//	"Configuring cache graph {GraphName} in config section [{GraphName}] of {IniName}.ini is deprecated. "
-		//	"Migrate the configuration to the [DerivedDataCacheGraphs] section of {IniName}.ini.",
-		//	("GraphName", GraphName), ("IniName", GEngineIni));
+		//	"Configuring cache graph '{Graph}' in config section [{Graph}] of the '{Config}' config is deprecated. "
+		//	"Migrate the configuration to the [DerivedDataCacheGraphs] section of the '{Config}' config.",
+		//	("Graph", GraphNameOrConfig), ("Config", GEngineIni));
 		Config::FConfigBuilder Builder(OutConfig);
-		return TryConvertLegacyCacheGraphConfig(GraphName, StoreName, LegacyConfig, Builder);
+		return TryConvertLegacyCacheGraphConfig(GraphName, TEXTVIEW("Root"), LegacyConfig, Builder);
 	}
 
-	FString Config;
-	if (GConfig->GetString(TEXT("DerivedDataCacheGraphs"), GraphName, Config, GEngineIni))
+	FString GraphConfigStorage;
+	if (!GraphName.IsEmpty() && GConfig->GetString(TEXT("DerivedDataCacheGraphs"), *WriteToString<64>(GraphName), GraphConfigStorage, GEngineIni))
 	{
-		Config::FConfigBuilder Builder(OutConfig);
-		Config::FConfigIterator It(Config);
-		for (; It; ++It)
+		GraphConfig = GraphConfigStorage;
+	}
+
+	if (GraphConfig.IsEmpty())
+	{
+		return false;
+	}
+
+	Config::FConfigBuilder Builder(OutConfig);
+	Config::FConfigIterator It(GraphConfig);
+	for (; It; ++It)
+	{
+		if (It.GetName() == TEXTVIEW("Deprecated"))
 		{
-			if (It.GetName() == TEXTVIEW("Deprecated"))
-			{
-				TStringBuilder<128> MessageScratch;
-				UE_LOGFMT(LogDerivedDataCache, Warning, "The cache graph {GraphName} is deprecated. {Message}",
-					GraphName, Config::ParseQuotedString(It.GetValue(), MessageScratch));
-			}
-			else if (FStringView StoreName = It.GetName(); !StoreName.IsEmpty())
-			{
-				Builder.Append({}, StoreName);
-			}
-			else if (FStringView StoreConfig; Config::TrySplitNameAndConfig(It.GetValue(), StoreName, StoreConfig))
-			{
-				Builder.Append({}, StoreName);
-			}
-			else
-			{
-				Builder.SetError();
-			}
+			TStringBuilder<128> MessageScratch;
+			UE_LOGFMT(LogDerivedDataCache, Warning, "The cache graph '{Graph}' is deprecated. {Message}",
+				GraphName, Config::ParseQuotedString(It.GetValue(), MessageScratch));
 		}
-		if (It.HasError())
+		else if (FStringView StoreName = It.GetName(); !StoreName.IsEmpty())
+		{
+			Builder.Append({}, StoreName);
+		}
+		else if (FStringView StoreConfig; Config::TrySplitNameAndConfig(It.GetValue(), StoreName, StoreConfig))
+		{
+			Builder.Append({}, StoreName);
+		}
+		else
 		{
 			Builder.SetError();
 		}
-		return !Builder.HasError();
 	}
-	return false;
+	if (It.HasError())
+	{
+		Builder.SetError();
+	}
+	return !Builder.HasError();
 }
 
 class FCacheStoreConfigBuilder
@@ -406,11 +423,11 @@ class FCacheStoreConfigBuilder
 public:
 	explicit FCacheStoreConfigBuilder(FStringBuilderBase& OutConfig, bool bLogLegacyConfig = true);
 
-	bool TryFind(const TCHAR* StoreName, const TCHAR* GraphName);
+	bool TryFind(FStringView StoreName, FStringView GraphNameOrConfig);
 
 private:
-	bool TryQuery(FStringView StoreName, const TCHAR* GraphName = nullptr);
-	bool TryBuild(FStringView StoreName, FStringView Config);
+	bool TryQuery(FStringView StoreName, FStringView GraphNameOrConfig = {});
+	bool TryBuild(FStringView StoreName, FStringView StoreConfig);
 
 	Config::FConfigBuilder Builder;
 	bool bLogLegacyConfig;
@@ -422,9 +439,9 @@ FCacheStoreConfigBuilder::FCacheStoreConfigBuilder(FStringBuilderBase& OutConfig
 {
 }
 
-bool FCacheStoreConfigBuilder::TryFind(const TCHAR* StoreName, const TCHAR* GraphName)
+bool FCacheStoreConfigBuilder::TryFind(FStringView StoreName, FStringView GraphNameOrConfig)
 {
-	const bool bOk = TryQuery(StoreName, GraphName);
+	const bool bOk = TryQuery(StoreName, GraphNameOrConfig);
 	if (!bOk)
 	{
 		Builder.SetError();
@@ -432,24 +449,41 @@ bool FCacheStoreConfigBuilder::TryFind(const TCHAR* StoreName, const TCHAR* Grap
 	return bOk;
 }
 
-bool FCacheStoreConfigBuilder::TryQuery(FStringView StoreName, const TCHAR* GraphName)
+bool FCacheStoreConfigBuilder::TryQuery(FStringView StoreName, FStringView GraphNameOrConfig)
 {
+	FStringView GraphName;
+	FStringView GraphConfig;
+	if (!Config::TrySplitNameAndConfig(GraphNameOrConfig, GraphName, GraphConfig))
+	{
+		return false;
+	}
+
+	if (!GraphName.IsEmpty() && !GraphConfig.IsEmpty())
+	{
+		return false;
+	}
+
 	// UE_DEPRECATED(5.8, "Remove the LegacyConfig block after the standard deprecation period.")
 	FString LegacyConfig;
-	if (GraphName && GConfig->GetString(GraphName, *WriteToString<64>(StoreName), LegacyConfig, GEngineIni))
+	if (!GraphName.IsEmpty() && GConfig->GetString(*WriteToString<64>(GraphName), *WriteToString<64>(StoreName), LegacyConfig, GEngineIni))
 	{
 		//UE_CLOGFMT(bLogLegacyConfig, LogDerivedDataCache, Warning,
-		//	"Configuring cache store {StoreName} in config section [{GraphName}] of {IniName}.ini is deprecated. "
-		//	"Migrate the configuration to the [DerivedDataCacheStores] section of {IniName}.ini.",
-		//	("GraphName", GraphName), ("StoreName", StoreName), ("IniName", GEngineIni));
+		//	"Configuring cache store '{Store}' in config section [{Graph}] of the '{Config}' config is deprecated. "
+		//	"Migrate the configuration to the [DerivedDataCacheStores] section of the '{Config}' config.",
+		//	("Graph", GraphNameOrConfig), ("Store", StoreName), ("Config", GEngineIni));
 		return TryBuild(StoreName, LegacyConfig);
 	}
 
 	bool bOk = true;
 	bool bHasConfig = false;
 
-	FString GraphConfig;
-	if (GraphName && GConfig->GetString(TEXT("DerivedDataCacheGraphs"), GraphName, GraphConfig, GEngineIni))
+	FString GraphConfigStorage;
+	if (!GraphName.IsEmpty() && GConfig->GetString(TEXT("DerivedDataCacheGraphs"), *WriteToString<64>(GraphName), GraphConfigStorage, GEngineIni))
+	{
+		GraphConfig = GraphConfigStorage;
+	}
+
+	if (!GraphConfig.IsEmpty())
 	{
 		Config::FConfigIterator It(GraphConfig);
 		for (; It; ++It)
@@ -483,12 +517,12 @@ bool FCacheStoreConfigBuilder::TryQuery(FStringView StoreName, const TCHAR* Grap
 	return bOk && bHasConfig;
 }
 
-bool FCacheStoreConfigBuilder::TryBuild(FStringView StoreName, FStringView Config)
+bool FCacheStoreConfigBuilder::TryBuild(FStringView StoreName, FStringView StoreConfig)
 {
 	bool bOk = true;
 	FStringView BaseName;
 
-	Config::FConfigIterator It(Config);
+	Config::FConfigIterator It(StoreConfig);
 	for (; It; ++It)
 	{
 		if (It.GetName() == TEXTVIEW("Base"))
@@ -499,7 +533,7 @@ bool FCacheStoreConfigBuilder::TryBuild(FStringView StoreName, FStringView Confi
 		else if (It.GetName() == TEXTVIEW("Deprecated"))
 		{
 			TStringBuilder<128> MessageScratch;
-			UE_LOGFMT(LogDerivedDataCache, Warning, "The cache store {StoreName} is deprecated. {Message}",
+			UE_LOGFMT(LogDerivedDataCache, Warning, "The cache store '{Store}' is deprecated. {Message}",
 				StoreName, Config::ParseQuotedString(It.GetValue(), MessageScratch));
 		}
 		else
@@ -512,22 +546,22 @@ bool FCacheStoreConfigBuilder::TryBuild(FStringView StoreName, FStringView Confi
 	// Process the base after appending the current config to ensure overrides come first.
 	if (bOk && !BaseName.IsEmpty())
 	{
-		bOk &= TryQuery(*WriteToString<64>(BaseName));
+		bOk &= TryQuery(BaseName);
 	}
 
 	return bOk;
 }
 
-bool TryFindCacheStoreConfig(const TCHAR* StoreName, const TCHAR* GraphName, FStringBuilderBase& OutConfig)
+bool TryFindCacheStoreConfig(FStringView StoreName, FStringView GraphNameOrConfig, FStringBuilderBase& OutConfig)
 {
 	FCacheStoreConfigBuilder Builder(OutConfig);
-	return Builder.TryFind(StoreName, GraphName);
+	return Builder.TryFind(StoreName, GraphNameOrConfig);
 }
 
-static bool TryConvertLegacyCacheStoreConfig(const TCHAR* GraphName, FStringView StoreName, Config::FConfigBuilder& Builder)
+static bool TryConvertLegacyCacheStoreConfig(FStringView StoreName, FStringView GraphName, Config::FConfigBuilder& Builder)
 {
 	TStringBuilder<256> StoreConfig;
-	bool bOk = FCacheStoreConfigBuilder(StoreConfig, /*bLogLegacyConfig*/ false).TryFind(*WriteToString<64>(StoreName), GraphName);
+	bool bOk = FCacheStoreConfigBuilder(StoreConfig, /*bLogLegacyConfig*/ false).TryFind(StoreName, GraphName);
 	return bOk && TryConvertLegacyCacheGraphConfig(GraphName, StoreName, StoreConfig, Builder);
 }
 
diff --git a/Engine/Source/Developer/DerivedDataCache/Private/DerivedDataConfig.h b/Engine/Source/Developer/DerivedDataCache/Private/DerivedDataConfig.h
--- a/Engine/Source/Developer/DerivedDataCache/Private/DerivedDataConfig.h
+++ b/Engine/Source/Developer/DerivedDataCache/Private/DerivedDataConfig.h
@@ -65,13 +65,17 @@ namespace UE::DerivedData
 {
 
 /**
-* Finds the configuration for a cache graph.
-*
-* @param GraphName   Name of the cache graph, which is compared case-insensitively.
-* @param OutConfig   Appended with the configuration for the cache graph if it is found.
-* @return True if the configuration was found, otherwise false.
-*/
-bool TryFindCacheGraphConfig(const TCHAR* GraphName, FStringBuilderBase& OutConfig);
+ * Finds the configuration for a cache graph.
+ *
+ * Written in the normalized format "(StoreName1, StoreName2, StoreName3)".
+ * Store remapping and inline configuration is stripped from the output.
+ * Example: "(StoreName1=OtherName(Key=Value))" outputs "(StoreName1)".
+ *
+ * @param GraphNameOrConfig   Name of the cache graph, which is compared case-insensitively, or a graph configuration.
+ * @param OutConfig           Appended with the configuration for the cache graph if it is found.
+ * @return True if the configuration was found, otherwise false.
+ */
+bool TryFindCacheGraphConfig(FStringView GraphNameOrConfig, FStringBuilderBase& OutConfig);
 
 /**
  * Finds the configuration for a cache store, optionally within a specific cache graph.
@@ -80,11 +84,15 @@ bool TryFindCacheGraphConfig(const TCHAR* GraphName, FStringBuilderBase& OutConf
  * The configuration is appended to OutConfig from most-derived to least-derived,
  * such that FParse::Value() will see override values before base values.
  *
- * @param StoreName   Name of the cache store, which is compared case-insensitively.
- * @param GraphName   Name of the cache graph, which is compared case-insensitively, or nullptr.
- * @param OutConfig   Appended with the configuration for the cache store if it is found.
+ * Written in the normalized format "(Type=Value, Key1=Value1, Key2=Value2, Value3)".
+ * Store inheritance is stripped from the output.
+ * Example: "(Base=OtherStore, Key1 = Value1, Value2)" outputs "(Key1=Value1, Value2, KeyFromOtherStore=OtherValue)".
+ *
+ * @param StoreName           Name of the cache store, which is compared case-insensitively.
+ * @param GraphNameOrConfig   Name of the cache graph, which is compared case-insensitively, or a graph configuration.
+ * @param OutConfig           Appended with the configuration for the cache store if it is found.
  * @return True if the configuration was found, otherwise false.
  */
-bool TryFindCacheStoreConfig(const TCHAR* StoreName, const TCHAR* GraphName, FStringBuilderBase& OutConfig);
+bool TryFindCacheStoreConfig(FStringView StoreName, FStringView GraphNameOrConfig, FStringBuilderBase& OutConfig);
 
 } // UE::DerivedData
diff --git a/Engine/Source/Developer/DerivedDataCache/Tests/DerivedDataConfigTest.cpp b/Engine/Source/Developer/DerivedDataCache/Tests/DerivedDataConfigTest.cpp
--- a/Engine/Source/Developer/DerivedDataCache/Tests/DerivedDataConfigTest.cpp
+++ b/Engine/Source/Developer/DerivedDataCache/Tests/DerivedDataConfigTest.cpp
@@ -38,11 +38,11 @@ TEST_CASE("DerivedData::CacheConfig", "[DerivedData]")
 {
 	SECTION("TryFindCacheGraphConfig")
 	{
-		const auto TestFindGraph = [](const TCHAR* GraphName, const TCHAR* ExpectedConfig) -> bool
+		const auto TestFindGraph = [](const TCHAR* GraphNameOrConfig, const TCHAR* ExpectedConfig) -> bool
 		{
-			CAPTURE(GraphName);
+			CAPTURE(GraphNameOrConfig);
 			TStringBuilder<1024> Config;
-			bool bFound = TryFindCacheGraphConfig(GraphName, Config);
+			bool bFound = TryFindCacheGraphConfig(GraphNameOrConfig, Config);
 			CHECK(Config.ToView() == MakeStringView(ExpectedConfig));
 			return bFound;
 		};
@@ -63,25 +63,32 @@ TEST_CASE("DerivedData::CacheConfig", "[DerivedData]")
 		CHECK(TestFindGraph(TEXT("TestCacheGraph8"), TEXT("(TestStore1)")));
 
 		{
-			CHECK_LOG_SCOPE(LogDerivedDataCache, Warning, "The cache graph TestCacheGraph9 is deprecated. Use TestCacheGraph1.");
+			CHECK_LOG_SCOPE(LogDerivedDataCache, Warning, "The cache graph 'TestCacheGraph9' is deprecated. Use TestCacheGraph1.");
 			CHECK(TestFindGraph(TEXT("TestCacheGraph9"), TEXT("(TestStore1)")));
 		}
 
 		{
-			CHECK_LOG_SCOPE(LogDerivedDataCache, Warning, "The cache graph TestLegacyCacheGraph1 contains cache store Root of deprecated type KeyLength.");
-			CHECK_LOG_SCOPE(LogDerivedDataCache, Warning, "The cache graph TestLegacyCacheGraph1 contains cache store Verify of deprecated type Verify.");
-			CHECK_LOG_SCOPE(LogDerivedDataCache, Warning, "The cache graph TestLegacyCacheGraph1 contains cache store AsyncPut of deprecated type AsyncPut.");
+			CHECK_LOG_SCOPE(LogDerivedDataCache, Warning, "The cache graph 'TestLegacyCacheGraph1' contains cache store 'Root' of deprecated type 'KeyLength'.");
+			CHECK_LOG_SCOPE(LogDerivedDataCache, Warning, "The cache graph 'TestLegacyCacheGraph1' contains cache store 'Verify' of deprecated type 'Verify'.");
+			CHECK_LOG_SCOPE(LogDerivedDataCache, Warning, "The cache graph 'TestLegacyCacheGraph1' contains cache store 'AsyncPut' of deprecated type 'AsyncPut'.");
 			CHECK(TestFindGraph(TEXT("TestLegacyCacheGraph1"), TEXT("(TestLegacyStore1, TestStore2, TestStore8)")));
 		}
+
+		CHECK(TestFindGraph(TEXT("(TestStore1)"), TEXT("(TestStore1)")));
+		CHECK(TestFindGraph(TEXT(" ( TestStore1 , TestStore2 ) "), TEXT("(TestStore1, TestStore2)")));
+		CHECK(TestFindGraph(TEXT(" ( TestStore1 ( Key1 = CustomValue ) , TestStore2(CustomValue1, Key2=CustomValue2) ) "), TEXT("(TestStore1, TestStore2)")));
+
+		CHECK_FALSE(TestFindGraph(TEXT("TestCacheGraph1(TestStore1)"), TEXT("")));
+		CHECK_FALSE(TestFindGraph(TEXT("TestCacheGraph1(TestStore1,TestStore2)"), TEXT("")));
 	}
 
 	SECTION("TryFindCacheStoreConfig")
 	{
-		const auto TestFindStore = [](const TCHAR* StoreName, const TCHAR* GraphName, const TCHAR* ExpectedConfig) -> bool
+		const auto TestFindStore = [](const TCHAR* StoreName, const TCHAR* GraphNameOrConfig, const TCHAR* ExpectedConfig) -> bool
 		{
-			CAPTURE(StoreName, GraphName);
+			CAPTURE(StoreName, GraphNameOrConfig);
 			TStringBuilder<1024> Config;
-			bool bFound = TryFindCacheStoreConfig(StoreName, GraphName, Config);
+			bool bFound = TryFindCacheStoreConfig(StoreName, GraphNameOrConfig, Config);
 			CHECK(Config.ToView() == MakeStringView(ExpectedConfig));
 			return bFound;
 		};
@@ -118,13 +125,27 @@ TEST_CASE("DerivedData::CacheConfig", "[DerivedData]")
 		CHECK(TestFindStore(TEXT("TestStore1"), TEXT("TestCacheGraph8"), TEXT("(Key1=Store1Value, Key2=Store1Value, ValueOnly, (X=\"A=1, \\\"B\\\"=2, C=3\"))")));
 
 		{
-			CHECK_LOG_SCOPE(LogDerivedDataCache, Warning, "The cache store TestStore7 is deprecated. Use TestStore1.");
+			CHECK_LOG_SCOPE(LogDerivedDataCache, Warning, "The cache store 'TestStore7' is deprecated. Use TestStore1.");
 			CHECK(TestFindStore(TEXT("TestStore7"), nullptr, TEXT("(Key1=Store1Value, Key2=Store1Value, ValueOnly, (X=\"A=1, \\\"B\\\"=2, C=3\"))")));
 		}
 
 		CHECK(TestFindStore(TEXT("TestLegacyStore1"), TEXT("TestLegacyCacheGraph1"), TEXT("(Type=Zen, Key1=LegacyStore1Value, Key2=LegacyStore1Value)")));
 		CHECK(TestFindStore(TEXT("TestStore2"), TEXT("TestLegacyCacheGraph1"), TEXT("(Type=Zen, Key1=LegacyStore2Value)")));
 		CHECK(TestFindStore(TEXT("TestStore8"), TEXT("TestLegacyCacheGraph1"), TEXT("(Type=Cloud)")));
+
+		CHECK(TestFindStore(TEXT("TestStore1"), TEXT("(TestStore1)"),
+			TEXT("(Key1=Store1Value, Key2=Store1Value, ValueOnly, (X=\"A=1, \\\"B\\\"=2, C=3\"))")));
+		CHECK(TestFindStore(TEXT("TestStore1"), TEXT(" ( TestStore1 , TestStore2 ) "),
+			TEXT("(Key1=Store1Value, Key2=Store1Value, ValueOnly, (X=\"A=1, \\\"B\\\"=2, C=3\"))")));
+		CHECK(TestFindStore(TEXT("TestStore2"), TEXT(" ( TestStore1 , TestStore2 ) "),
+			TEXT("(Key1=Store2Value, Key1=Store1Value, Key2=Store1Value, ValueOnly, (X=\"A=1, \\\"B\\\"=2, C=3\"))")));
+		CHECK(TestFindStore(TEXT("TestStore1"), TEXT(" ( TestStore1 ( Key1 = CustomValue ) , TestStore2(CustomValue1, Key2=CustomValue2) ) "),
+			TEXT("(Key1=CustomValue, Key1=Store1Value, Key2=Store1Value, ValueOnly, (X=\"A=1, \\\"B\\\"=2, C=3\"))")));
+		CHECK(TestFindStore(TEXT("TestStore2"), TEXT(" ( TestStore1 ( Key1 = CustomValue ) , TestStore2(CustomValue1, Key2=CustomValue2) ) "),
+			TEXT("(CustomValue1, Key2=CustomValue2, Key1=Store2Value, Key1=Store1Value, Key2=Store1Value, ValueOnly, (X=\"A=1, \\\"B\\\"=2, C=3\"))")));
+
+		CHECK_FALSE(TestFindStore(TEXT("TestStore1"), TEXT("TestCacheGraph1(TestStore1)"), TEXT("")));
+		CHECK_FALSE(TestFindStore(TEXT("TestStore1"), TEXT("TestCacheGraph1(TestStore1,TestStore2)"), TEXT("")));
 	}
 }
 
