DDC: Fixed Config::TrySplitNameAndConfig

It was incorrectly returning true for inputs that do not match the expected format.

#rnx

diff --git a/Engine/Source/Developer/DerivedDataCache/Private/DerivedDataConfig.cpp b/Engine/Source/Developer/DerivedDataCache/Private/DerivedDataConfig.cpp
--- a/Engine/Source/Developer/DerivedDataCache/Private/DerivedDataConfig.cpp
+++ b/Engine/Source/Developer/DerivedDataCache/Private/DerivedDataConfig.cpp
@@ -212,18 +212,14 @@ bool TrySplitNameAndConfig(FStringView Value, FStringView& OutName, FStringView&
 	Value.RightChopInline(Name.Len());
 	Value = FAsciiSet::TrimPrefixWith(Value, WhiteSpaceSet);
 
-	if (Value.IsEmpty() || (Value.StartsWith(ObjectStartChar) && Value.EndsWith(ObjectEndChar)))
+	const bool bHasValue = Value.StartsWith(ObjectStartChar) && Value.EndsWith(ObjectEndChar);
+	const bool bHasName = !Name.IsEmpty() && (bHasValue || Value.IsEmpty());
+	if (bHasName || bHasValue)
 	{
 		OutName = Name;
 		OutConfig = Value;
 		return true;
 	}
-	if (!Name.IsEmpty())
-	{
-		OutName = Name;
-		OutConfig.Reset();
-		return true;
-	}
 	return false;
 }
 
@@ -356,7 +352,7 @@ bool TryFindCacheGraphConfig(FStringView GraphNameOrConfig, FStringBuilderBase&
 {
 	FStringView GraphName;
 	FStringView GraphConfig;
-	if (!Config::TrySplitNameAndConfig(GraphNameOrConfig, GraphName, GraphConfig))
+	if (!GraphNameOrConfig.IsEmpty() && !Config::TrySplitNameAndConfig(GraphNameOrConfig, GraphName, GraphConfig))
 	{
 		return false;
 	}
@@ -453,7 +449,7 @@ bool FCacheStoreConfigBuilder::TryQuery(FStringView StoreName, FStringView Graph
 {
 	FStringView GraphName;
 	FStringView GraphConfig;
-	if (!Config::TrySplitNameAndConfig(GraphNameOrConfig, GraphName, GraphConfig))
+	if (!GraphNameOrConfig.IsEmpty() && !Config::TrySplitNameAndConfig(GraphNameOrConfig, GraphName, GraphConfig))
 	{
 		return false;
 	}
diff --git a/Engine/Source/Developer/DerivedDataCache/Tests/DerivedDataConfigTest.cpp b/Engine/Source/Developer/DerivedDataCache/Tests/DerivedDataConfigTest.cpp
--- a/Engine/Source/Developer/DerivedDataCache/Tests/DerivedDataConfigTest.cpp
+++ b/Engine/Source/Developer/DerivedDataCache/Tests/DerivedDataConfigTest.cpp
@@ -32,6 +32,35 @@ TEST_CASE("DerivedData::Config", "[DerivedData]")
 		CHECK(Config::ParseQuotedString(TEXTVIEW(R"("\")"), Scratch) == TEXTVIEW(R"()"));
 		CHECK(Config::ParseQuotedString(TEXTVIEW(R"("\'\"\?\\\a\b\f\n\r\t\v")"), Scratch) == TEXTVIEW("\'\"\?\\\a\b\f\n\r\t\v"));
 	}
+
+	SECTION("TrySplitNameAndConfig")
+	{
+		const auto TestTrySplitNameAndConfig = [](const TCHAR* Input, const TCHAR* ExpectedName, const TCHAR* ExpectedConfig)
+		{
+			CAPTURE(Input);
+			FStringView Name, Config;
+			bool bOk = Config::TrySplitNameAndConfig(Input, Name, Config);
+			CHECK(bOk == (ExpectedName && ExpectedConfig));
+			if (bOk)
+			{
+				CHECK(Name == ExpectedName);
+				CHECK(Config == ExpectedConfig);
+			}
+		};
+		TestTrySplitNameAndConfig(TEXT(""), nullptr, nullptr);
+		TestTrySplitNameAndConfig(TEXT("-"), nullptr, nullptr);
+		TestTrySplitNameAndConfig(TEXT("A-B"), nullptr, nullptr);
+		TestTrySplitNameAndConfig(TEXT("A-B(C)"), nullptr, nullptr);
+		TestTrySplitNameAndConfig(TEXT("A.B"), nullptr, nullptr);
+		TestTrySplitNameAndConfig(TEXT("A.B(C)"), nullptr, nullptr);
+		TestTrySplitNameAndConfig(TEXT("ABC"), TEXT("ABC"), TEXT(""));
+		TestTrySplitNameAndConfig(TEXT("ABC()"), TEXT("ABC"), TEXT("()"));
+		TestTrySplitNameAndConfig(TEXT(" ABC "), TEXT("ABC"), TEXT(""));
+		TestTrySplitNameAndConfig(TEXT(" ABC () "), TEXT("ABC"), TEXT("()"));
+		TestTrySplitNameAndConfig(TEXT("()"), TEXT(""), TEXT("()"));
+		TestTrySplitNameAndConfig(TEXT(" () "), TEXT(""), TEXT("()"));
+		TestTrySplitNameAndConfig(TEXT("(A=1, B=2)"), TEXT(""), TEXT("(A=1, B=2)"));
+	}
 }
 
 TEST_CASE("DerivedData::CacheConfig", "[DerivedData]")
