DDC: Robustly handle legacy config without new config and vice versa

#rnx

diff --git a/Engine/Source/Developer/DerivedDataCache/Private/DerivedDataBackends.cpp b/Engine/Source/Developer/DerivedDataCache/Private/DerivedDataBackends.cpp
--- a/Engine/Source/Developer/DerivedDataCache/Private/DerivedDataBackends.cpp
+++ b/Engine/Source/Developer/DerivedDataCache/Private/DerivedDataBackends.cpp
@@ -165,10 +165,11 @@ public:
 			}
 
 			// UE_DEPRECATED(5.8, "Remove the legacy default block after the standard deprecation period.")
-			if (!RootCache)
+			const TCHAR* LegacyDefaultGraphName = FApp::IsEngineInstalled() ? TEXT("InstalledDerivedDataBackendGraph") : TEXT("DerivedDataBackendGraph");
+			if (!RootCache && GConfig->DoesSectionExist(LegacyDefaultGraphName, GEngineIni))
 			{
 				// Try to use the legacy default graph.
-				ActiveGraphName = FApp::IsEngineInstalled() ? TEXT("InstalledDerivedDataBackendGraph") : TEXT("DerivedDataBackendGraph");
+				ActiveGraphName = LegacyDefaultGraphName;
 				if (!ParseGraph(ActiveGraphName, ParsedNodes) || !Hierarchy->HasAllFlags(ECacheStoreFlags::Query | ECacheStoreFlags::Store))
 				{
 					ResetGraph();
@@ -182,7 +183,8 @@ public:
 				if (!ParseGraph(ActiveGraphName, ParsedNodes) || !Hierarchy->HasAllFlags(ECacheStoreFlags::Query | ECacheStoreFlags::Store))
 				{
 					TStringBuilder<256> GraphConfig;
-					if (!TryFindCacheGraphConfig(*ActiveGraphName, GraphConfig))
+					if (!TryFindCacheGraphConfig(ActiveGraphName, GraphConfig) &&
+						!TryFindCacheGraphConfig(LegacyDefaultGraphName, GraphConfig))
 					{
 						UE_LOGFMT(LogDerivedDataCache, Fatal,
 							"Failed to create default cache graph '{Graph}' because its config is missing in the '{Config}' config.",
@@ -203,7 +205,7 @@ public:
 					{
 						UE_CLOGFMT(!FParse::Param(FCommandLine::Get(), TEXT("DDC-AllowNoActiveStores")),
 							LogDerivedDataCache, Fatal,
-							"Unable to use default cache graph '{Graph}' because there are no {MissingTypes} nodes available."
+							"Unable to use default cache graph '{Graph}' because there are no {MissingTypes} nodes available. "
 							"Add -DDC-ForceMemoryCache to the command line to bypass this if you need access "
 							"to the editor settings to fix the cache configuration.",
 							*ActiveGraphName,
