DDC: Enable deprecation warnings for legacy cache config

#rnx

diff --git a/Engine/Source/Developer/DerivedDataCache/Private/DerivedDataConfig.cpp b/Engine/Source/Developer/DerivedDataCache/Private/DerivedDataConfig.cpp
--- a/Engine/Source/Developer/DerivedDataCache/Private/DerivedDataConfig.cpp
+++ b/Engine/Source/Developer/DerivedDataCache/Private/DerivedDataConfig.cpp
@@ -365,10 +365,10 @@ bool TryFindCacheGraphConfig(FStringView GraphNameOrConfig, FStringBuilderBase&
 	// UE_DEPRECATED(5.8, "Remove the LegacyConfig block after the standard deprecation period.")
 	if (FString LegacyConfig; !GraphName.IsEmpty() && GConfig->GetString(*WriteToString<64>(GraphName), TEXT("Root"), LegacyConfig, GEngineIni))
 	{
-		//UE_LOGFMT(LogDerivedDataCache, Warning,
-		//	"Configuring cache graph '{Graph}' in config section [{Graph}] of the '{Config}' config is deprecated. "
-		//	"Migrate the configuration to the [DerivedDataCacheGraphs] section of the '{Config}' config.",
-		//	("Graph", GraphNameOrConfig), ("Config", GEngineIni));
+		UE_LOGFMT(LogDerivedDataCache, Warning,
+			"Configuring cache graph '{Graph}' in config section [{Graph}] of the '{Config}' config is deprecated. "
+			"Migrate the configuration to the [DerivedDataCacheGraphs] section of the '{Config}' config.",
+			("Graph", GraphNameOrConfig), ("Config", GEngineIni));
 		Config::FConfigBuilder Builder(OutConfig);
 		return TryConvertLegacyCacheGraphConfig(GraphName, TEXTVIEW("Root"), LegacyConfig, Builder);
 	}
@@ -463,10 +463,10 @@ bool FCacheStoreConfigBuilder::TryQuery(FStringView StoreName, FStringView Graph
 	FString LegacyConfig;
 	if (!GraphName.IsEmpty() && GConfig->GetString(*WriteToString<64>(GraphName), *WriteToString<64>(StoreName), LegacyConfig, GEngineIni))
 	{
-		//UE_CLOGFMT(bLogLegacyConfig, LogDerivedDataCache, Warning,
-		//	"Configuring cache store '{Store}' in config section [{Graph}] of the '{Config}' config is deprecated. "
-		//	"Migrate the configuration to the [DerivedDataCacheStores] section of the '{Config}' config.",
-		//	("Graph", GraphNameOrConfig), ("Store", StoreName), ("Config", GEngineIni));
+		UE_CLOGFMT(bLogLegacyConfig, LogDerivedDataCache, Warning,
+			"Configuring cache store '{Store}' in config section [{Graph}] of the '{Config}' config is deprecated. "
+			"Migrate the configuration to the [DerivedDataCacheStores] section of the '{Config}' config.",
+			("Graph", GraphNameOrConfig), ("Store", StoreName), ("Config", GEngineIni));
 		return TryBuild(StoreName, LegacyConfig);
 	}
 
diff --git a/Engine/Source/Developer/DerivedDataCache/Tests/DerivedDataConfigTest.cpp b/Engine/Source/Developer/DerivedDataCache/Tests/DerivedDataConfigTest.cpp
--- a/Engine/Source/Developer/DerivedDataCache/Tests/DerivedDataConfigTest.cpp
+++ b/Engine/Source/Developer/DerivedDataCache/Tests/DerivedDataConfigTest.cpp
@@ -97,6 +97,7 @@ TEST_CASE("DerivedData::CacheConfig", "[DerivedData]")
 		}
 
 		{
+			CHECK_LOG_SCOPE(LogDerivedDataCache, Warning, "Configuring cache graph 'TestLegacyCacheGraph1' in config section [TestLegacyCacheGraph1] of the 'Engine' config is deprecated. Migrate the configuration to the [DerivedDataCacheGraphs] section of the 'Engine' config.");
 			CHECK_LOG_SCOPE(LogDerivedDataCache, Warning, "The cache graph 'TestLegacyCacheGraph1' contains cache store 'Root' of deprecated type 'KeyLength'.");
 			CHECK_LOG_SCOPE(LogDerivedDataCache, Warning, "The cache graph 'TestLegacyCacheGraph1' contains cache store 'Verify' of deprecated type 'Verify'.");
 			CHECK_LOG_SCOPE(LogDerivedDataCache, Warning, "The cache graph 'TestLegacyCacheGraph1' contains cache store 'AsyncPut' of deprecated type 'AsyncPut'.");
@@ -158,8 +159,14 @@ TEST_CASE("DerivedData::CacheConfig", "[DerivedData]")
 			CHECK(TestFindStore(TEXT("TestStore7"), nullptr, TEXT("(Key1=Store1Value, Key2=Store1Value, ValueOnly, (X=\"A=1, \\\"B\\\"=2, C=3\"))")));
 		}
 
-		CHECK(TestFindStore(TEXT("TestLegacyStore1"), TEXT("TestLegacyCacheGraph1"), TEXT("(Type=Zen, Key1=LegacyStore1Value, Key2=LegacyStore1Value)")));
-		CHECK(TestFindStore(TEXT("TestStore2"), TEXT("TestLegacyCacheGraph1"), TEXT("(Type=Zen, Key1=LegacyStore2Value)")));
+		{
+			CHECK_LOG_SCOPE(LogDerivedDataCache, Warning, "Configuring cache store 'TestLegacyStore1' in config section [TestLegacyCacheGraph1] of the 'Engine' config is deprecated. Migrate the configuration to the [DerivedDataCacheStores] section of the 'Engine' config.");
+			CHECK(TestFindStore(TEXT("TestLegacyStore1"), TEXT("TestLegacyCacheGraph1"), TEXT("(Type=Zen, Key1=LegacyStore1Value, Key2=LegacyStore1Value)")));
+		}
+		{
+			CHECK_LOG_SCOPE(LogDerivedDataCache, Warning, "Configuring cache store 'TestStore2' in config section [TestLegacyCacheGraph1] of the 'Engine' config is deprecated. Migrate the configuration to the [DerivedDataCacheStores] section of the 'Engine' config.");
+			CHECK(TestFindStore(TEXT("TestStore2"), TEXT("TestLegacyCacheGraph1"), TEXT("(Type=Zen, Key1=LegacyStore2Value)")));
+		}
 		CHECK(TestFindStore(TEXT("TestStore8"), TEXT("TestLegacyCacheGraph1"), TEXT("(Type=Cloud)")));
 
 		CHECK(TestFindStore(TEXT("TestStore1"), TEXT("(TestStore1)"),
