DDC: Moved global state and functions into the graph

diff --git a/Engine/Source/Developer/DerivedDataCache/Internal/DerivedDataCacheReplay.h b/Engine/Source/Developer/DerivedDataCache/Internal/DerivedDataCacheReplay.h
--- a/Engine/Source/Developer/DerivedDataCache/Internal/DerivedDataCacheReplay.h
+++ b/Engine/Source/Developer/DerivedDataCache/Internal/DerivedDataCacheReplay.h
@@ -12,6 +12,7 @@ class FCbObjectView;
 
 namespace UE::DerivedData { class FCacheKeyFilter; }
 namespace UE::DerivedData { class FCacheMethodFilter; }
+namespace UE::DerivedData { class ICacheStoreOwner; }
 namespace UE::DerivedData { class ILegacyCacheStore; }
 namespace UE::DerivedData { enum class ECacheMethod : uint8; }
 namespace UE::DerivedData { enum class ECachePolicy : uint32; }
@@ -32,7 +33,7 @@ public:
 	UE_API void SetPolicyTransform(ECachePolicy AddFlags, ECachePolicy RemoveFlags);
 	UE_API void SetPriorityOverride(EPriority Priority);
 
-	UE_API void ReadFromFileAsync(const TCHAR* ReplayPath, uint64 ScratchSize = DefaultScratchSize);
+	UE_API void ReadFromFileAsync(const TCHAR* ReplayPath, ICacheStoreOwner* Owner, uint64 ScratchSize = DefaultScratchSize);
 	UE_API bool ReadFromFile(const TCHAR* ReplayPath, uint64 ScratchSize = DefaultScratchSize);
 	UE_API bool ReadFromArchive(FArchive& ReplayAr, uint64 ScratchSize = DefaultScratchSize);
 	UE_API bool ReadFromObject(FCbObjectView Object);
diff --git a/Engine/Source/Developer/DerivedDataCache/Private/DerivedDataBackends.cpp b/Engine/Source/Developer/DerivedDataCache/Private/DerivedDataBackends.cpp
--- a/Engine/Source/Developer/DerivedDataCache/Private/DerivedDataBackends.cpp
+++ b/Engine/Source/Developer/DerivedDataCache/Private/DerivedDataBackends.cpp
@@ -1,8 +1,9 @@
 // Copyright Epic Games, Inc. All Rights Reserved.
 
+#include "DerivedDataBackendInterface.h"
+
 #include "CoreTypes.h"
 #include "Containers/StringView.h"
-#include "DerivedDataBackendInterface.h"
 #include "DerivedDataCacheMethod.h"
 #include "DerivedDataCachePrivate.h"
 #include "DerivedDataCacheReplay.h"
@@ -47,22 +48,10 @@ static TAutoConsoleVariable<FString> GDerivedDataCacheGraphName(
 	TEXT("Name or config of the graph to use for the Derived Data Cache."),
 	ECVF_ReadOnly);
 
-namespace UE::DerivedData::Private
-{
-
-static std::atomic<int32> GAsyncTaskCounter;
-
-int32 AddToAsyncTaskCounter(int32 Addend)
-{
-	return GAsyncTaskCounter.fetch_add(Addend);
-}
-
-} // UE::DerivedData::Private
-
 namespace UE::DerivedData
 {
 
-ILegacyCacheStore* CreateCacheStoreAsync(ILegacyCacheStore* InnerBackend, IMemoryCacheStore* MemoryCache, bool bDeleteInnerCache);
+ILegacyCacheStore* CreateCacheStoreAsync(ICacheStoreOwner* Owner, ILegacyCacheStore* InnerBackend, IMemoryCacheStore* MemoryCache, bool bDeleteInnerCache);
 ILegacyCacheStore* CreateCacheStoreHierarchy(ICacheStoreOwner*& OutOwner, TFunctionRef<void (IMemoryCacheStore*&)> MemoryCacheCreator);
 ILegacyCacheStore* CreateCacheStoreThrottle(ILegacyCacheStore* InnerCache, uint32 LatencyMS, uint32 MaxBytesPerSecond);
 ILegacyCacheStore* CreateCacheStoreVerify(ILegacyCacheStore* InnerCache, bool bPutOnError);
@@ -72,7 +61,7 @@ void CreateMemoryCacheStore(IMemoryCacheStore*& OutCache, const TCHAR* Name, con
 IPakFileCacheStore* CreatePakFileCacheStore(const TCHAR* Name, const TCHAR* Filename, bool bWriting, bool bCompressed, ICacheStoreOwner* Owner);
 ILegacyCacheStore* CreateS3CacheStore(const TCHAR* Name, const TCHAR* Config, ICacheStoreOwner& Owner);
 ILegacyCacheStore* CreateZenCacheStore(const TCHAR* Name, const TCHAR* Config, ICacheStoreOwner* Owner);
-ILegacyCacheStore* TryCreateCacheStoreReplay(ILegacyCacheStore* InnerCache);
+ILegacyCacheStore* TryCreateCacheStoreReplay(ILegacyCacheStore* InnerCache, ICacheStoreOwner* Owner);
 
 /**
  * This class is used to create a singleton that represents the derived data cache hierarchy and all of the wrappers necessary
@@ -93,7 +82,6 @@ public:
 		, WritePakCache(nullptr)
 		, Hierarchy(nullptr)
 		, bUsingSharedDDC(false)
-		, bIsShuttingDown(false)
 		, MountPakCommand(
 			TEXT("DDC.MountPak"),
 			*LOCTEXT("CommandText_DDCMountPak", "Mounts read-only pak file").ToString(),
@@ -217,7 +205,7 @@ public:
 		}
 
 		// Async must exist in the graph.
-		RootCache = CreateCacheStoreAsync(RootCache, MemoryCache, /*bDeleteInnerCache*/ true);
+		RootCache = CreateCacheStoreAsync(Hierarchy, RootCache, MemoryCache, /*bDeleteInnerCache*/ true);
 
 		// Create a Verify node when using -DDC-Verify[=Type1[@Rate2][+Type2[@Rate2]...]].
 		if (FString VerifyArg;
@@ -232,7 +220,7 @@ public:
 		}
 
 		// Create a Replay node when requested on the command line.
-		if (ILegacyCacheStore* ReplayNode = TryCreateCacheStoreReplay(RootCache))
+		if (ILegacyCacheStore* ReplayNode = TryCreateCacheStoreReplay(RootCache, Hierarchy))
 		{
 			RootCache = ReplayNode;
 		}
@@ -496,24 +484,13 @@ public:
 
 	virtual void WaitForQuiescence(bool bShutdown) override
 	{
-		double StartTime = FPlatformTime::Seconds();
-		double LastPrint = StartTime;
-
 		if (bShutdown)
 		{
-			bIsShuttingDown.store(true, std::memory_order_relaxed);
+			SetShuttingDown();
 		}
 
-		while (const int32 AsyncCompletionCounter = Private::AddToAsyncTaskCounter(0))
-		{
-			check(AsyncCompletionCounter > 0);
-			FPlatformProcess::Sleep(0.1f);
-			if (FPlatformTime::Seconds() - LastPrint > 5.0)
-			{
-				UE_LOG(LogDerivedDataCache, Log, TEXT("Waited %ds for derived data cache to finish..."), int32(FPlatformTime::Seconds() - StartTime));
-				LastPrint = FPlatformTime::Seconds();
-			}
-		}
+		WaitForIdle();
+
 		if (bShutdown)
 		{
 			FString MergePaks;
@@ -585,17 +562,12 @@ public:
 
 	virtual void AddToAsyncCompletionCounter(int32 Addend) override
 	{
-		verify(Private::AddToAsyncTaskCounter(Addend) + Addend >= 0);
+		verify(AddToAsyncTaskCounter(Addend) + Addend >= 0);
 	}
 
 	virtual bool AnyAsyncRequestsRemaining() override
 	{
-		return Private::AddToAsyncTaskCounter(0) > 0;
-	}
-
-	virtual bool IsShuttingDown() override
-	{
-		return bIsShuttingDown.load(std::memory_order_relaxed);
+		return AddToAsyncTaskCounter(0) > 0;
 	}
 
 	virtual void GetDirectories(TArray<FString>& OutResults) override
@@ -749,6 +721,12 @@ private:
 		return Hierarchy->HasAllFlags(Flags);
 	}
 
+	int32 AddToAsyncTaskCounter(int32 Value) final
+	{
+		check(Hierarchy);
+		return Hierarchy->AddToAsyncTaskCounter(Value);
+	}
+
 	ICacheStoreStats* CreateStats(ILegacyCacheStore* CacheStore, ECacheStoreFlags Flags, FStringView Type, FStringView Name, FStringView Path) final
 	{
 		check(Hierarchy);
@@ -771,6 +749,30 @@ private:
 		Hierarchy->LegacyResourceStats(OutStats);
 	}
 
+	void SetShuttingDown() final
+	{
+		check(Hierarchy);
+		Hierarchy->SetShuttingDown();
+	}
+
+	bool IsShuttingDown() final
+	{
+		check(Hierarchy);
+		return Hierarchy->IsShuttingDown();
+	}
+
+	bool IsShuttingDown() const final
+	{
+		check(Hierarchy);
+		return Hierarchy->IsShuttingDown();
+	}
+
+	void WaitForIdle() const final
+	{
+		check(Hierarchy);
+		Hierarchy->WaitForIdle();
+	}
+
 	ILegacyCacheStore* FindOrCreate(const TCHAR* Name) final
 	{
 		if (!ActiveParsedNodeMap)
@@ -898,7 +900,7 @@ private:
 			Replay.SetPriorityOverride(Priority);
 		}
 
-		Replay.ReadFromFileAsync(*Args[0]);
+		Replay.ReadFromFileAsync(*Args[0], Hierarchy);
 		Replays.Add(MoveTemp(Replay));
 	}
 
@@ -930,9 +932,6 @@ private:
 	/** Whether a shared cache is in use */
 	bool bUsingSharedDDC;
 
-	/** Whether a shutdown is pending */
-	std::atomic<bool> bIsShuttingDown;
-
 	/** MountPak console command */
 	FAutoConsoleCommand MountPakCommand;
 	/** UnmountPak console command */
diff --git a/Engine/Source/Developer/DerivedDataCache/Private/DerivedDataCachePrivate.h b/Engine/Source/Developer/DerivedDataCache/Private/DerivedDataCachePrivate.h
--- a/Engine/Source/Developer/DerivedDataCache/Private/DerivedDataCachePrivate.h
+++ b/Engine/Source/Developer/DerivedDataCache/Private/DerivedDataCachePrivate.h
@@ -19,9 +19,6 @@ namespace UE::DerivedData { struct FCacheGetValueRequest; }
 namespace UE::DerivedData::Private
 {
 
-// Implemented in DerivedDataBackends.cpp
-int32 AddToAsyncTaskCounter(int32 Addend); // returns the previous value
-
 // Implemented in DerivedDataCache.cpp
 ICache* CreateCache(FDerivedDataCacheInterface** OutLegacyCache);
 void LaunchTaskInCacheThreadPool(IRequestOwner& Owner, TUniqueFunction<void ()>&& TaskBody);
diff --git a/Engine/Source/Developer/DerivedDataCache/Private/DerivedDataCacheReplay.cpp b/Engine/Source/Developer/DerivedDataCache/Private/DerivedDataCacheReplay.cpp
--- a/Engine/Source/Developer/DerivedDataCache/Private/DerivedDataCacheReplay.cpp
+++ b/Engine/Source/Developer/DerivedDataCache/Private/DerivedDataCacheReplay.cpp
@@ -496,7 +496,7 @@ public:
 
 	void WaitForAsyncReads();
 
-	void ReadFromFileAsync(const TCHAR* ReplayPath, uint64 ScratchSize);
+	void ReadFromFileAsync(const TCHAR* ReplayPath, ICacheStoreOwner* Owner, uint64 ScratchSize);
 	bool ReadFromFile(const TCHAR* ReplayPath, uint64 ScratchSize);
 	bool ReadFromArchive(FArchive& ReplayAr, uint64 ScratchSize);
 	bool ReadFromObject(FCbObjectView Object);
@@ -680,14 +680,14 @@ void FCacheReplayReader::FState::ApplyPolicyTransform(FCacheGetChunkRequest& Req
 	EnumRemoveFlags(Request.Policy, PolicyFlagsToRemove);
 }
 
-void FCacheReplayReader::FState::ReadFromFileAsync(const TCHAR* const ReplayPath, const uint64 ScratchSize)
+void FCacheReplayReader::FState::ReadFromFileAsync(const TCHAR* const ReplayPath, ICacheStoreOwner* Owner, const uint64 ScratchSize)
 {
-	Private::AddToAsyncTaskCounter(1);
-	ReadAsyncPipe.Launch(TEXT("CacheReplayReadFromFileAsync"), [this, ReplayPath = FString(ReplayPath), ScratchSize]
+	Owner->AddToAsyncTaskCounter(1);
+	ReadAsyncPipe.Launch(TEXT("CacheReplayReadFromFileAsync"), [this, ReplayPath = FString(ReplayPath), Owner, ScratchSize]
 	{
 		TRACE_CPUPROFILER_EVENT_SCOPE(ReplayDDC_ReadFromFileAsync);
 		ReadFromFile(*ReplayPath, ScratchSize);
-		Private::AddToAsyncTaskCounter(-1);
+		Owner->AddToAsyncTaskCounter(-1);
 	});
 }
 
@@ -818,9 +818,9 @@ FCacheReplayReader::FCacheReplayReader(ILegacyCacheStore* const TargetCache)
 	State->TargetCache = TargetCache;
 }
 
-void FCacheReplayReader::ReadFromFileAsync(const TCHAR* ReplayPath, const uint64 ScratchSize)
+void FCacheReplayReader::ReadFromFileAsync(const TCHAR* ReplayPath, ICacheStoreOwner* Owner, const uint64 ScratchSize)
 {
-	return State->ReadFromFileAsync(ReplayPath, ScratchSize);
+	return State->ReadFromFileAsync(ReplayPath, Owner, ScratchSize);
 }
 
 bool FCacheReplayReader::ReadFromFile(const TCHAR* ReplayPath, const uint64 ScratchSize)
@@ -938,7 +938,7 @@ static void ParseReplayPolicyTransform(const TCHAR* const CommandLine, FCacheRep
 
 ///////////////////////////////////////////////////////////////////////////////////////////////////////////////////////
 
-ILegacyCacheStore* TryCreateCacheStoreReplay(ILegacyCacheStore* InnerCache)
+ILegacyCacheStore* TryCreateCacheStoreReplay(ILegacyCacheStore* InnerCache, ICacheStoreOwner* Owner)
 {
 	const TCHAR* const CommandLine = FCommandLine::Get();
 	const bool bHasReplayLoad = String::FindFirst(CommandLine, TEXT("-DDC-ReplayLoad=")) != INDEX_NONE;
@@ -989,7 +989,7 @@ ILegacyCacheStore* TryCreateCacheStoreReplay(ILegacyCacheStore* InnerCache)
 			FString ReplayLoadPath;
 			if (FParse::Value(*Token, TEXT("-DDC-ReplayLoad="), ReplayLoadPath))
 			{
-				Reader.ReadFromFileAsync(*ReplayLoadPath);
+				Reader.ReadFromFileAsync(*ReplayLoadPath, Owner);
 			}
 		}
 
diff --git a/Engine/Source/Developer/DerivedDataCache/Private/DerivedDataCacheStore.h b/Engine/Source/Developer/DerivedDataCache/Private/DerivedDataCacheStore.h
--- a/Engine/Source/Developer/DerivedDataCache/Private/DerivedDataCacheStore.h
+++ b/Engine/Source/Developer/DerivedDataCache/Private/DerivedDataCacheStore.h
@@ -201,10 +201,16 @@ public:
 	/** Returns true if the combined flags of the owned cache stores contain all of these flags. */
 	virtual bool HasAllFlags(ECacheStoreFlags Flags) const = 0;
 
+	virtual int32 AddToAsyncTaskCounter(int32 Value) = 0;
+
 	virtual ICacheStoreStats* CreateStats(ILegacyCacheStore* CacheStore, ECacheStoreFlags Flags, FStringView Type, FStringView Name, FStringView Path = {}) = 0;
 	virtual void DestroyStats(ICacheStoreStats* Stats) = 0;
 
 	virtual void LegacyResourceStats(TArray<FDerivedDataCacheResourceStat>& OutStats) const = 0;
+
+	virtual void SetShuttingDown() = 0;
+	virtual bool IsShuttingDown() const = 0;
+	virtual void WaitForIdle() const = 0;
 };
 
 class ICacheStoreGraph
diff --git a/Engine/Source/Developer/DerivedDataCache/Private/DerivedDataCacheStoreAsync.cpp b/Engine/Source/Developer/DerivedDataCache/Private/DerivedDataCacheStoreAsync.cpp
--- a/Engine/Source/Developer/DerivedDataCache/Private/DerivedDataCacheStoreAsync.cpp
+++ b/Engine/Source/Developer/DerivedDataCache/Private/DerivedDataCacheStoreAsync.cpp
@@ -25,7 +25,7 @@ namespace UE::DerivedData
 class FCacheStoreAsync final : public ILegacyCacheStore
 {
 public:
-	FCacheStoreAsync(ILegacyCacheStore* InnerCache, IMemoryCacheStore* MemoryCache, bool bDeleteInnerCache);
+	FCacheStoreAsync(ICacheStoreOwner* StoreOwner, ILegacyCacheStore* InnerCache, IMemoryCacheStore* MemoryCache, bool bDeleteInnerCache);
 	~FCacheStoreAsync() final;
 
 	virtual void Put(
@@ -112,14 +112,16 @@ private:
 		OnCompleteType&& OnComplete,
 		OnExecuteType&& OnExecute);
 
+	ICacheStoreOwner* StoreOwner;
 	ILegacyCacheStore* InnerCache;
 	IMemoryCacheStore* MemoryCache;
 	FDerivedDataCacheUsageStats UsageStats;
 	bool bDeleteInnerCache;
 };
 
-FCacheStoreAsync::FCacheStoreAsync(ILegacyCacheStore* InInnerCache, IMemoryCacheStore* InMemoryCache, bool bInDeleteInnerCache)
-	: InnerCache(InInnerCache)
+FCacheStoreAsync::FCacheStoreAsync(ICacheStoreOwner* InStoreOwner, ILegacyCacheStore* InInnerCache, IMemoryCacheStore* InMemoryCache, bool bInDeleteInnerCache)
+	: StoreOwner(InStoreOwner)
+	, InnerCache(InInnerCache)
 	, MemoryCache(InMemoryCache)
 	, bDeleteInnerCache(bInDeleteInnerCache)
 {
@@ -165,11 +167,11 @@ void FCacheStoreAsync::Execute(
 			}
 
 			OnComplete(MoveTemp(Response));
-			Private::AddToAsyncTaskCounter(-1);
+			StoreOwner->AddToAsyncTaskCounter(-1);
 		});
 	};
 
-	Private::AddToAsyncTaskCounter(Requests.Num());
+	StoreOwner->AddToAsyncTaskCounter(Requests.Num());
 	if (Owner.GetPriority() == EPriority::Blocking)
 	{
 		return ExecuteWithStats(Requests, Owner, MoveTemp(OnComplete));
@@ -189,14 +191,14 @@ void FCacheStoreAsync::Execute(
 			else
 			{
 				CompleteWithStatus(Requests, MoveTemp(OnComplete), EStatus::Canceled);
-				Private::AddToAsyncTaskCounter(-Requests.Num());
+				StoreOwner->AddToAsyncTaskCounter(-Requests.Num());
 			}
 		});
 }
 
-ILegacyCacheStore* CreateCacheStoreAsync(ILegacyCacheStore* InnerCache, IMemoryCacheStore* MemoryCache, bool bDeleteInnerCache)
+ILegacyCacheStore* CreateCacheStoreAsync(ICacheStoreOwner* Owner, ILegacyCacheStore* InnerCache, IMemoryCacheStore* MemoryCache, bool bDeleteInnerCache)
 {
-	return new FCacheStoreAsync(InnerCache, MemoryCache, bDeleteInnerCache);
+	return new FCacheStoreAsync(Owner, InnerCache, MemoryCache, bDeleteInnerCache);
 }
 
 } // UE::DerivedData
diff --git a/Engine/Source/Developer/DerivedDataCache/Private/DerivedDataCacheStoreHierarchy.cpp b/Engine/Source/Developer/DerivedDataCache/Private/DerivedDataCacheStoreHierarchy.cpp
--- a/Engine/Source/Developer/DerivedDataCache/Private/DerivedDataCacheStoreHierarchy.cpp
+++ b/Engine/Source/Developer/DerivedDataCache/Private/DerivedDataCacheStoreHierarchy.cpp
@@ -18,6 +18,7 @@
 #include "DerivedDataRequest.h"
 #include "DerivedDataRequestOwner.h"
 #include "HAL/PlatformAtomics.h"
+#include "HAL/PlatformProcess.h"
 #include "Logging/StructuredLog.h"
 #include "MemoryCacheStore.h"
 #include "Misc/EnumClassFlags.h"
@@ -33,7 +34,7 @@ void GatherDerivedDataCacheResourceStats(TArray<FDerivedDataCacheResourceStat>&
 namespace UE::DerivedData
 {
 
-ILegacyCacheStore* CreateCacheStoreAsync(ILegacyCacheStore* InnerCache, IMemoryCacheStore* MemoryCache, bool bDeleteInnerCache);
+ILegacyCacheStore* CreateCacheStoreAsync(ICacheStoreOwner* Owner, ILegacyCacheStore* InnerCache, IMemoryCacheStore* MemoryCache, bool bDeleteInnerCache);
 
 ///////////////////////////////////////////////////////////////////////////////////////////////////////////////////////
 
@@ -49,11 +50,17 @@ public:
 
 	bool HasAllFlags(ECacheStoreFlags Flags) const final;
 
+	int32 AddToAsyncTaskCounter(int32 Value) final;
+
 	ICacheStoreStats* CreateStats(ILegacyCacheStore* CacheStore, ECacheStoreFlags Flags, FStringView Type, FStringView Name, FStringView Path) final;
 	void DestroyStats(ICacheStoreStats* Stats) final;
 
 	void LegacyResourceStats(TArray<FDerivedDataCacheResourceStat>& OutStats) const final;
 
+	void SetShuttingDown() final { bShuttingDown.store(true, std::memory_order_relaxed); }
+	bool IsShuttingDown() const final { return bShuttingDown.load(std::memory_order_relaxed); }
+	void WaitForIdle() const final;
+
 	void Put(
 		TConstArrayView<FCachePutRequest> Requests,
 		IRequestOwner& Owner,
@@ -148,6 +155,8 @@ private:
 	TArray<FCacheStoreNode, TInlineAllocator<8>> Nodes;
 	IMemoryCacheStore* MemoryCache;
 	FCacheStats CacheStats;
+	std::atomic<int32> AsyncTaskCounter = 0;
+	std::atomic<bool> bShuttingDown = false;
 };
 
 ///////////////////////////////////////////////////////////////////////////////////////////////////////////////////////
@@ -303,7 +312,7 @@ void FCacheStoreHierarchy::Add(ILegacyCacheStore* CacheStore, ECacheStoreFlags F
 	TUniqueLock Lock(NodesMutex);
 	checkf(!Algo::FindBy(Nodes, CacheStore, &FCacheStoreNode::Cache),
 		TEXT("Attempting to add a cache store that was previously registered to the hierarchy."));
-	TUniquePtr<ILegacyCacheStore> AsyncCacheStore(CreateCacheStoreAsync(CacheStore, /*MemoryCache*/ nullptr, /*bDeleteInnerCache*/ false));
+	TUniquePtr<ILegacyCacheStore> AsyncCacheStore(CreateCacheStoreAsync(this, CacheStore, /*MemoryCache*/ nullptr, /*bDeleteInnerCache*/ false));
 	Nodes.Add({CacheStore, Flags, {}, MoveTemp(AsyncCacheStore)});
 	UpdateNodeFlags();
 }
@@ -350,6 +359,11 @@ bool FCacheStoreHierarchy::HasAllFlags(ECacheStoreFlags Flags) const
 	return EnumHasAllFlags(CombinedFlags, Flags);
 }
 
+int32 FCacheStoreHierarchy::AddToAsyncTaskCounter(int32 Value)
+{
+	return AsyncTaskCounter.fetch_add(Value);
+}
+
 void FCacheStoreHierarchy::UpdateNodeFlags()
 {
 	ECacheStoreNodeFlags OriginalCombinedFlags = CombinedNodeFlags.load();
@@ -533,6 +547,24 @@ void FCacheStoreHierarchy::LegacyResourceStats(TArray<FDerivedDataCacheResourceS
 #endif
 }
 
+void FCacheStoreHierarchy::WaitForIdle() const
+{
+	const double StartTime = FPlatformTime::Seconds();
+	double LastPrint = StartTime;
+	while (const int32 Counter = AsyncTaskCounter.load(std::memory_order_relaxed))
+	{
+		check(Counter > 0);
+		FPlatformProcess::Sleep(0.1f);
+		const double CurrentTime = FPlatformTime::Seconds();
+		if (CurrentTime - LastPrint > 5.0)
+		{
+			UE_LOGFMT(LogDerivedDataCache, Log, "Waited {WaitTime}s for Derived Data Dache to finish {TaskCount} tasks...",
+				int32(CurrentTime - StartTime), Counter);
+			LastPrint = CurrentTime;
+		}
+	}
+}
+
 FCacheRecordPolicy FCacheStoreHierarchy::AddPolicy(const FCacheRecordPolicy& BasePolicy, ECachePolicy Policy)
 {
 	return BasePolicy.Transform([Policy](ECachePolicy P) { return P | Policy; });
diff --git a/Engine/Source/Developer/DerivedDataCache/Private/HttpCacheStore.cpp b/Engine/Source/Developer/DerivedDataCache/Private/HttpCacheStore.cpp
--- a/Engine/Source/Developer/DerivedDataCache/Private/HttpCacheStore.cpp
+++ b/Engine/Source/Developer/DerivedDataCache/Private/HttpCacheStore.cpp
@@ -84,11 +84,6 @@ TRACE_DECLARE_ATOMIC_INT_COUNTER(HttpDDC_PutHit, TEXT("HttpDDC Put Hit"));
 TRACE_DECLARE_ATOMIC_INT_COUNTER(HttpDDC_BytesReceived, TEXT("HttpDDC Bytes Received"));
 TRACE_DECLARE_ATOMIC_INT_COUNTER(HttpDDC_BytesSent, TEXT("HttpDDC Bytes Sent"));
 
-static bool ShouldAbortForShutdown()
-{
-	return !GIsBuildMachine && FDerivedDataBackend::Get().IsShuttingDown();
-}
-
 static bool IsValueDataReady(FValue& Value, const ECachePolicy Policy)
 {
 	if (!EnumHasAnyFlags(Policy, ECachePolicy::Query))
@@ -425,8 +420,9 @@ public:
 	FHttpOperation(const FHttpOperation&) = delete;
 	FHttpOperation& operator=(const FHttpOperation&) = delete;
 
-	explicit FHttpOperation(THttpUniquePtr<IHttpRequest>&& InRequest)
+	explicit FHttpOperation(THttpUniquePtr<IHttpRequest>&& InRequest, const ICacheStoreOwner* InStoreOwner)
 		: Request(MoveTemp(InRequest))
+		, StoreOwner(InStoreOwner)
 	{
 	}
 
@@ -446,6 +442,8 @@ public:
 	void Send();
 	void SendAsync(IRequestOwner& Owner, TUniqueFunction<void ()>&& OnComplete);
 
+	bool ShouldAbortForShutdown() const { return !GIsBuildMachine && StoreOwner && StoreOwner->IsShuttingDown(); }
+
 	// Consume Response
 
 	int32 GetStatusCode() const { return Response->GetStatusCode(); }
@@ -472,6 +470,7 @@ private:
 	THttpUniquePtr<IHttpResponse> Response;
 	TArray<int32, TInlineAllocator<4>> ExpectedStatusCodes;
 	uint32 AttemptCount = 0;
+	const ICacheStoreOwner* StoreOwner = nullptr;
 };
 
 class FHttpCacheStore::FHttpOperation::FHttpOperationReceiver final : public IHttpReceiver
@@ -512,7 +511,7 @@ private:
 
 	bool ShouldRetry(IHttpResponse& LocalResponse) const
 	{
-		if (Operation->AttemptCount >= UE_HTTPDDC_MAX_ATTEMPTS || ShouldAbortForShutdown())
+		if (Operation->AttemptCount >= UE_HTTPDDC_MAX_ATTEMPTS || Operation->ShouldAbortForShutdown())
 		{
 			return false;
 		}
@@ -725,8 +724,8 @@ void FHttpCacheStore::FHttpOperation::GetStats(FRequestStats& OutStats) const
 class FHttpCacheStore::FHealthCheckOp final
 {
 public:
-	FHealthCheckOp(FHttpCacheStore& CacheStore, IHttpClient& Client)
-		: Operation(Client.TryCreateRequest({}))
+	FHealthCheckOp(FHttpCacheStore& CacheStore, IHttpClient& Client, const ICacheStoreOwner* StoreOwner)
+		: Operation(Client.TryCreateRequest({}), StoreOwner)
 		, Owner(EPriority::High)
 		, NodeName(*CacheStore.NodeName)
 	{
@@ -2340,7 +2339,7 @@ FHttpCacheStore::FHttpCacheStore(const FHttpCacheStoreParams& Params, ICacheStor
 	FHttpClientParams ClientParams = GetDefaultClientParams();
 
 	THttpUniquePtr<IHttpClient> Client = ConnectionPool->CreateClient(ClientParams);
-	FHealthCheckOp HealthCheck(*this, *Client);
+	FHealthCheckOp HealthCheck(*this, *Client, StoreOwner);
 	if (AcquireAccessToken(Client.Get()) && HealthCheck.IsReady())
 	{
 		ClientParams.MaxRequests = UE_HTTPDDC_GET_REQUEST_POOL_SIZE;
@@ -2518,7 +2517,7 @@ bool FHttpCacheStore::AcquireAccessToken(IHttpClient* Client)
 
 		FHttpRequestParams RequestParams;
 		RequestParams.bIgnoreMaxRequests = true;
-		FHttpOperation Operation(Client->TryCreateRequest(RequestParams));
+		FHttpOperation Operation(Client->TryCreateRequest(RequestParams), StoreOwner);
 		Operation.SetUri(StringCast<ANSICHAR>(*OAuthProvider));
 
 		if (OAuthProvider.StartsWith(TEXT("http://localhost")))
@@ -2690,7 +2689,7 @@ void FHttpCacheStore::WaitForHttpOperationAsync(IRequestOwner& Owner, EOperation
 			Request->AddHeader(ANSITEXTVIEW("Authorization"), WriteToAnsiString<1024>(*Access));
 		}
 
-		OnOperation(MakeUnique<FHttpOperation>(MoveTemp(Request)));
+		OnOperation(MakeUnique<FHttpOperation>(MoveTemp(Request), StoreOwner));
 	});
 }
 
diff --git a/Engine/Source/Developer/DerivedDataCache/Private/PakFileCacheStore.cpp b/Engine/Source/Developer/DerivedDataCache/Private/PakFileCacheStore.cpp
--- a/Engine/Source/Developer/DerivedDataCache/Private/PakFileCacheStore.cpp
+++ b/Engine/Source/Developer/DerivedDataCache/Private/PakFileCacheStore.cpp
@@ -260,7 +260,11 @@ FPakFileCacheStore::~FPakFileCacheStore()
 
 void FPakFileCacheStore::Close()
 {
-	FDerivedDataBackend::Get().WaitForQuiescence();
+	if (StoreOwner)
+	{
+		StoreOwner->WaitForIdle();
+	}
+
 	if (!bClosed)
 	{
 		if (bWriting)
