DDC: Moved the ICache implementation into the graph

#rnx

diff --git a/Engine/Source/Developer/DerivedDataCache/Private/DerivedDataBackendInterface.h b/Engine/Source/Developer/DerivedDataCache/Private/DerivedDataBackendInterface.h
--- a/Engine/Source/Developer/DerivedDataCache/Private/DerivedDataBackendInterface.h
+++ b/Engine/Source/Developer/DerivedDataCache/Private/DerivedDataBackendInterface.h
@@ -69,17 +69,17 @@ struct FBackendDebugOptions
 	bool ShouldSimulateGetMiss(const FCacheKey& Key);
 };
 
-class FDerivedDataBackend
+class FDerivedDataBackend : public ICache
 {
 public:
-	static FDerivedDataBackend* Create();
+	static FDerivedDataBackend* CreateStatic();
 
 	/**
 	 * Singleton to retrieve the GLOBAL backend
 	 *
 	 * @return Reference to the global cache backend
 	 */
-	static FDerivedDataBackend& Get();
+	static FDerivedDataBackend& GetStatic();
 
 	virtual ~FDerivedDataBackend() = default;
 
diff --git a/Engine/Source/Developer/DerivedDataCache/Private/DerivedDataBackends.cpp b/Engine/Source/Developer/DerivedDataCache/Private/DerivedDataBackends.cpp
--- a/Engine/Source/Developer/DerivedDataCache/Private/DerivedDataBackends.cpp
+++ b/Engine/Source/Developer/DerivedDataCache/Private/DerivedDataBackends.cpp
@@ -2,8 +2,10 @@
 
 #include "DerivedDataBackendInterface.h"
 
+#include "Algo/AllOf.h"
 #include "CoreTypes.h"
 #include "Containers/StringView.h"
+#include "DerivedDataCacheMaintainer.h"
 #include "DerivedDataCacheMethod.h"
 #include "DerivedDataCachePrivate.h"
 #include "DerivedDataCacheReplay.h"
@@ -11,6 +13,7 @@
 #include "DerivedDataCacheUsageStats.h"
 #include "DerivedDataConfig.h"
 #include "DerivedDataRequestOwner.h"
+#include "Features/IModularFeatures.h"
 #include "HAL/CriticalSection.h"
 #include "HAL/FileManager.h"
 #include "HAL/IConsoleManager.h"
@@ -32,6 +35,7 @@
 #include "Modules/ModuleManager.h"
 #include "PakFileCacheStore.h"
 #include "ProfilingDebugging/CookStats.h"
+#include "ProfilingDebugging/CpuProfilerTrace.h"
 #include "Serialization/CompactBinaryPackage.h"
 #include "String/Find.h"
 #include <atomic>
@@ -67,7 +71,11 @@ ILegacyCacheStore* TryCreateCacheStoreReplay(ILegacyCacheStore* InnerCache, ICac
  * This class is used to create a singleton that represents the derived data cache hierarchy and all of the wrappers necessary
  * ideally this would be data driven and the backends would be plugins...
  */
-class FDerivedDataBackendGraph final : public FDerivedDataBackend, ICacheStoreOwner, ICacheStoreGraph
+class FDerivedDataBackendGraph
+	: public FDerivedDataBackend
+	, public ICacheStoreOwner
+	, public ICacheStoreGraph
+	, public ICacheStoreMaintainer
 {
 public:
 	using FParsedNode = ILegacyCacheStore*;
@@ -224,6 +232,8 @@ public:
 		{
 			RootCache = ReplayNode;
 		}
+
+		CacheStoreMaintainers = IModularFeatures::Get().GetModularFeatureImplementations<ICacheStoreMaintainer>(FeatureName);
 	}
 
 	/**
@@ -575,7 +585,7 @@ public:
 		OutResults = Directories;
 	}
 
-	static FORCEINLINE FDerivedDataBackendGraph& Get()
+	static FORCEINLINE FDerivedDataBackendGraph& GetStatic()
 	{
 		check(StaticGraph);
 		return *StaticGraph;
@@ -904,6 +914,73 @@ private:
 		Replays.Add(MoveTemp(Replay));
 	}
 
+	// ICache Interface
+
+	void Put(
+		TConstArrayView<FCachePutRequest> Requests,
+		IRequestOwner& Owner,
+		FOnCachePutComplete&& OnComplete) final
+	{
+		TRACE_CPUPROFILER_EVENT_SCOPE(DDC_Put);
+		return RootCache->Put(Requests, Owner, OnComplete ? MoveTemp(OnComplete) : FOnCachePutComplete([](auto&&) {}));
+	}
+
+	void Get(
+		TConstArrayView<FCacheGetRequest> Requests,
+		IRequestOwner& Owner,
+		FOnCacheGetComplete&& OnComplete) final
+	{
+		TRACE_CPUPROFILER_EVENT_SCOPE(DDC_Get);
+		return RootCache->Get(Requests, Owner, OnComplete ? MoveTemp(OnComplete) : FOnCacheGetComplete([](auto&&) {}));
+	}
+
+	void PutValue(
+		TConstArrayView<FCachePutValueRequest> Requests,
+		IRequestOwner& Owner,
+		FOnCachePutValueComplete&& OnComplete) final
+	{
+		TRACE_CPUPROFILER_EVENT_SCOPE(DDC_PutValue);
+		return RootCache->PutValue(Requests, Owner, OnComplete ? MoveTemp(OnComplete) : FOnCachePutValueComplete([](auto&&) {}));
+	}
+
+	void GetValue(
+		TConstArrayView<FCacheGetValueRequest> Requests,
+		IRequestOwner& Owner,
+		FOnCacheGetValueComplete&& OnComplete) final
+	{
+		TRACE_CPUPROFILER_EVENT_SCOPE(DDC_GetValue);
+		return RootCache->GetValue(Requests, Owner, OnComplete ? MoveTemp(OnComplete) : FOnCacheGetValueComplete([](auto&&) {}));
+	}
+
+	void GetChunks(
+		TConstArrayView<FCacheGetChunkRequest> Requests,
+		IRequestOwner& Owner,
+		FOnCacheGetChunkComplete&& OnComplete) final
+	{
+		TRACE_CPUPROFILER_EVENT_SCOPE(DDC_GetChunks);
+		return RootCache->GetChunks(Requests, Owner, OnComplete ? MoveTemp(OnComplete) : FOnCacheGetChunkComplete([](auto&&) {}));
+	}
+
+	ICacheStoreMaintainer& GetMaintainer() final
+	{
+		return *this;
+	}
+
+	// ICacheStoreMaintainer Interface
+
+	bool IsIdle() const final
+	{
+		return Algo::AllOf(CacheStoreMaintainers, &ICacheStoreMaintainer::IsIdle);
+	}
+
+	void BoostPriority() final
+	{
+		for (ICacheStoreMaintainer* Maintainer : CacheStoreMaintainers)
+		{
+			Maintainer->BoostPriority();
+		}
+	}
+
 	static inline FDerivedDataBackendGraph*			StaticGraph;
 
 	FString											ActiveGraphName;
@@ -943,6 +1020,8 @@ private:
 	FParsedNodeMap* ActiveParsedNodeMap = nullptr;
 	FString ActiveNodeName;
 	FString ActiveNodeConfig;
+
+	TArray<ICacheStoreMaintainer*> CacheStoreMaintainers;
 };
 
 } // UE::DerivedData
@@ -950,14 +1029,14 @@ private:
 namespace UE::DerivedData
 {
 
-FDerivedDataBackend* FDerivedDataBackend::Create()
+FDerivedDataBackend* FDerivedDataBackend::CreateStatic()
 {
 	return new FDerivedDataBackendGraph();
 }
 
-FDerivedDataBackend& FDerivedDataBackend::Get()
+FDerivedDataBackend& FDerivedDataBackend::GetStatic()
 {
-	return FDerivedDataBackendGraph::Get();
+	return FDerivedDataBackendGraph::GetStatic();
 }
 
 ///////////////////////////////////////////////////////////////////////////////////////////////////////////////////////
diff --git a/Engine/Source/Developer/DerivedDataCache/Private/DerivedDataCache.cpp b/Engine/Source/Developer/DerivedDataCache/Private/DerivedDataCache.cpp
--- a/Engine/Source/Developer/DerivedDataCache/Private/DerivedDataCache.cpp
+++ b/Engine/Source/Developer/DerivedDataCache/Private/DerivedDataCache.cpp
@@ -3,7 +3,6 @@
 #include "DerivedDataCache.h"
 #include "DerivedDataCacheInterface.h"
 
-#include "Algo/AllOf.h"
 #include "AnalyticsEventAttribute.h"
 #include "Async/AsyncWork.h"
 #include "Async/InheritedContext.h"
@@ -23,7 +22,6 @@
 #include "DerivedDataRequestOwner.h"
 #include "DerivedDataThreadPoolTask.h"
 #include "Experimental/ZenServerInterface.h"
-#include "Features/IModularFeatures.h"
 #include "HAL/ThreadSafeCounter.h"
 #include "Misc/CommandLine.h"
 #include "Misc/CoreMisc.h"
@@ -173,7 +171,7 @@ namespace UE::DerivedData::CookStats
 			const int64 RootPutTotal = RootPutHits + RootPutMisses;
 
 			AddStat(TEXT("DDC.Summary"), FCookStatsManager::CreateKeyValueArray(
-				TEXT("BackEnd"), FDerivedDataBackend::Get().GetGraphName(),
+				TEXT("BackEnd"), FDerivedDataBackend::GetStatic().GetGraphName(),
 				TEXT("HasLocalCache"), LocalNode || ZenLocalNode,
 				TEXT("HasSharedCache"), SharedNode || ZenRemoteNode,
 				TEXT("HasCloudCache"), !!CloudNode,
@@ -667,11 +665,7 @@ private:
  * Implementation of the derived data cache
  * This API is fully threadsafe
 **/
-class FDerivedDataCache final
-	: public FDerivedDataCacheInterface
-	, public ICache
-	, public ICacheStoreMaintainer
-	, public IDDCCleanup
+class FDerivedDataCache final : public FDerivedDataCacheInterface, public IDDCCleanup
 {
 public:
 	/** Constructor, called once to cereate a singleton **/
@@ -690,9 +684,7 @@ public:
 #endif
 		}
 
-		Backend = FDerivedDataBackend::Create();
-
-		CacheStoreMaintainers = IModularFeatures::Get().GetModularFeatureImplementations<ICacheStoreMaintainer>(FeatureName);
+		Backend = FDerivedDataBackend::CreateStatic();
 
 		GVerifyDDC = FParse::Param(FCommandLine::Get(), TEXT("VerifyDDC"));
 
@@ -1021,14 +1013,14 @@ public:
 
 	virtual bool IsFinished() const override
 	{
-		return IsIdle();
+		return Backend->GetMaintainer().IsIdle();
 	}
 
 	virtual void WaitBetweenDeletes(bool bWait) override
 	{
 		if (!bWait)
 		{
-			BoostPriority();
+			Backend->GetMaintainer().BoostPriority();
 		}
 	}
 
@@ -1155,6 +1147,11 @@ public:
 		return DDCNotificationEvent;
 	}
 
+	ICache* GetCache() const
+	{
+		return Backend;
+	}
+
 protected:
 	uint32 NextHandle()
 	{
@@ -1185,88 +1182,14 @@ private:
 
 	/** Cache notification delegate */
 	FOnDDCNotification DDCNotificationEvent;
-
-public:
-	// ICache Interface
-
-	void Put(
-		TConstArrayView<FCachePutRequest> Requests,
-		IRequestOwner& Owner,
-		FOnCachePutComplete&& OnComplete) final
-	{
-		DDC_SCOPE_CYCLE_COUNTER(DDC_Put);
-		return Backend->GetRoot().Put(Requests, Owner, OnComplete ? MoveTemp(OnComplete) : FOnCachePutComplete([](auto&&) {}));
-	}
-
-	void Get(
-		TConstArrayView<FCacheGetRequest> Requests,
-		IRequestOwner& Owner,
-		FOnCacheGetComplete&& OnComplete) final
-	{
-		DDC_SCOPE_CYCLE_COUNTER(DDC_Get);
-		return Backend->GetRoot().Get(Requests, Owner, OnComplete ? MoveTemp(OnComplete) : FOnCacheGetComplete([](auto&&) {}));
-	}
-
-	void PutValue(
-		TConstArrayView<FCachePutValueRequest> Requests,
-		IRequestOwner& Owner,
-		FOnCachePutValueComplete&& OnComplete) final
-	{
-		DDC_SCOPE_CYCLE_COUNTER(DDC_PutValue);
-		return Backend->GetRoot().PutValue(Requests, Owner, OnComplete ? MoveTemp(OnComplete) : FOnCachePutValueComplete([](auto&&) {}));
-	}
-
-	void GetValue(
-		TConstArrayView<FCacheGetValueRequest> Requests,
-		IRequestOwner& Owner,
-		FOnCacheGetValueComplete&& OnComplete) final
-	{
-		DDC_SCOPE_CYCLE_COUNTER(DDC_GetValue);
-		return Backend->GetRoot().GetValue(Requests, Owner, OnComplete ? MoveTemp(OnComplete) : FOnCacheGetValueComplete([](auto&&) {}));
-	}
-
-	void GetChunks(
-		TConstArrayView<FCacheGetChunkRequest> Requests,
-		IRequestOwner& Owner,
-		FOnCacheGetChunkComplete&& OnComplete) final
-	{
-		DDC_SCOPE_CYCLE_COUNTER(DDC_GetChunks);
-		return Backend->GetRoot().GetChunks(Requests, Owner, OnComplete ? MoveTemp(OnComplete) : FOnCacheGetChunkComplete([](auto&&) {}));
-	}
-
-	ICacheStoreMaintainer& GetMaintainer() final
-	{
-		return *this;
-	}
-
-	// ICacheStoreMaintainer Interface
-
-	bool IsIdle() const final
-	{
-		return Algo::AllOf(CacheStoreMaintainers, &ICacheStoreMaintainer::IsIdle);
-	}
-
-	void BoostPriority() final
-	{
-		for (ICacheStoreMaintainer* Maintainer : CacheStoreMaintainers)
-		{
-			Maintainer->BoostPriority();
-		}
-	}
-
-private:
-	TArray<ICacheStoreMaintainer*> CacheStoreMaintainers;
 };
 
-ICache* CreateCache(FDerivedDataCacheInterface** OutLegacyCache)
+void CreateCache(ICache*& OutCache, FDerivedDataCacheInterface*& OutLegacyCache)
 {
 	LLM_SCOPE_BYTAG(DerivedDataCache);
 	FDerivedDataCache* Cache = new FDerivedDataCache;
-	if (OutLegacyCache)
-	{
-		*OutLegacyCache = Cache;
-	}
-	return Cache;
+	OutCache = Cache->GetCache();
+	OutLegacyCache = Cache;
 }
 
 } // UE::DerivedData::Private
diff --git a/Engine/Source/Developer/DerivedDataCache/Private/DerivedDataCacheModule.cpp b/Engine/Source/Developer/DerivedDataCache/Private/DerivedDataCacheModule.cpp
--- a/Engine/Source/Developer/DerivedDataCache/Private/DerivedDataCacheModule.cpp
+++ b/Engine/Source/Developer/DerivedDataCache/Private/DerivedDataCacheModule.cpp
@@ -6,6 +6,7 @@
 #include "DerivedDataBuild.h"
 #include "DerivedDataBuildPrivate.h"
 #include "DerivedDataCache.h"
+#include "DerivedDataCacheInterface.h"
 #include "DerivedDataCachePrivate.h"
 #include "DerivedDataPrivate.h"
 #include "HAL/CriticalSection.h"
@@ -15,6 +16,9 @@
 namespace UE::DerivedData::Private
 {
 
+// Implemented in DerivedDataCache.cpp
+void CreateCache(ICache*& OutCache, FDerivedDataCacheInterface*& OutLegacyCache);
+
 static FDerivedDataCacheInterface* GDerivedDataLegacyCache;
 static ICache* GDerivedDataCache;
 static IBuild* GDerivedDataBuild;
@@ -39,8 +43,9 @@ public:
 		FScopeLock Lock(&CreateLock);
 		if (!GDerivedDataCache)
 		{
-			GDerivedDataCache = CreateCache(&GDerivedDataLegacyCache);
+			CreateCache(GDerivedDataCache, GDerivedDataLegacyCache);
 			check(GDerivedDataCache);
+			check(GDerivedDataLegacyCache);
 		}
 	}
 
@@ -65,7 +70,7 @@ public:
 	{
 		delete GDerivedDataBuild;
 		GDerivedDataBuild = nullptr;
-		delete GDerivedDataCache;
+		delete GDerivedDataLegacyCache;
 		GDerivedDataCache = nullptr;
 		GDerivedDataLegacyCache = nullptr;
 	}
diff --git a/Engine/Source/Developer/DerivedDataCache/Private/DerivedDataCachePrivate.h b/Engine/Source/Developer/DerivedDataCache/Private/DerivedDataCachePrivate.h
--- a/Engine/Source/Developer/DerivedDataCache/Private/DerivedDataCachePrivate.h
+++ b/Engine/Source/Developer/DerivedDataCache/Private/DerivedDataCachePrivate.h
@@ -20,7 +20,6 @@ namespace UE::DerivedData::Private
 {
 
 // Implemented in DerivedDataCache.cpp
-ICache* CreateCache(FDerivedDataCacheInterface** OutLegacyCache);
 void LaunchTaskInCacheThreadPool(IRequestOwner& Owner, TUniqueFunction<void ()>&& TaskBody);
 void SaveToCompactBinary(FCbWriter& Writer, const FCacheGetRequest& Request);
 void SaveToCompactBinary(FCbWriter& Writer, const FCacheGetValueRequest& Request);
