DDC: Allow creating multiple cache graphs

diff --git a/Engine/Source/Developer/DerivedDataCache/Private/DerivedDataBackends.cpp b/Engine/Source/Developer/DerivedDataCache/Private/DerivedDataBackends.cpp
--- a/Engine/Source/Developer/DerivedDataCache/Private/DerivedDataBackends.cpp
+++ b/Engine/Source/Developer/DerivedDataCache/Private/DerivedDataBackends.cpp
@@ -81,47 +81,33 @@
 	using FParsedNode = ILegacyCacheStore*;
 	using FParsedNodeMap = TMap<FString, FParsedNode>;
 
-	/**
-	 * constructor, builds the cache tree
-	 */
 	FDerivedDataBackendGraph()
-		: RootCache(nullptr)
-		, MemoryCache(nullptr)
-		, WritePakCache(nullptr)
-		, Hierarchy(nullptr)
-		, bUsingSharedDDC(false)
-		, MountPakCommand(
-			TEXT("DDC.MountPak"),
-			*LOCTEXT("CommandText_DDCMountPak", "Mounts read-only pak file").ToString(),
-			FConsoleCommandWithArgsDelegate::CreateRaw(this, &FDerivedDataBackendGraph::MountPakCommandHandler))
-		, UnmountPakCommand(
-			TEXT("DDC.UnmountPak"),
-			*LOCTEXT("CommandText_DDCUnmountPak", "Unmounts read-only pak file").ToString(),
-			FConsoleCommandWithArgsDelegate::CreateRaw(this, &FDerivedDataBackendGraph::UnmountPakCommandHandler))
-		, LoadReplayCommand(
-			TEXT("DDC.LoadReplay"),
-			*LOCTEXT("CommandText_DDCLoadReplay", "Loads a cache replay file created by -DDC-ReplaySave=<Path>").ToString(),
-			FConsoleCommandWithArgsDelegate::CreateRaw(this, &FDerivedDataBackendGraph::LoadReplayCommandHandler))
 	{
-		check(!StaticGraph);
-		StaticGraph = this;
-
-		check(IsInGameThread()); // we pretty much need this to be initialized from the main thread...it uses GConfig, etc
+		// Needs to be initialized from the main thread because it uses config and console variables
+		check(IsInGameThread());
 		check(GConfig && GConfig->IsReadyForUse());
 
 		bHasDefaultDebugOptions = FBackendDebugOptions::ParseFromTokens(DefaultDebugOptions, TEXT("All"), FCommandLine::Get());
+	}
 
-		RootCache = nullptr;
-		FParsedNodeMap ParsedNodes;
-		TGuardValue ParsedNodeMapGuard(ActiveParsedNodeMap, &ParsedNodes);
-
-		// Create the graph using ini settings. The string "default" forwards creation to use the default graph.
-
-		if (!FParse::Value(FCommandLine::Get(), TEXT("-DDC="), ActiveGraphName, /*bShouldStopOnSeparator*/ false))
+	static FDerivedDataBackendGraph* TryCreate(FStringView GraphNameOrConfig)
+	{
+		FDerivedDataBackendGraph* Graph = new FDerivedDataBackendGraph();
+		if (Graph->CreateGraph(GraphNameOrConfig))
 		{
-			ActiveGraphName = GDerivedDataCacheGraphName.GetValueOnGameThread();
+			Graph->PostCreateGraph();
+			return Graph;
 		}
+		delete Graph;
+		return nullptr;
+	}
 
+protected:
+	bool CreateGraph(FStringView GraphNameOrConfig)
+	{
+		ActiveGraphName = GraphNameOrConfig;
+
+#if IS_PROGRAM
 		// A DDC graph of "None" is used by build worker programs that use the DDC build code paths but avoid use of the DDC cache
 		// code paths. Unfortunately the cache must currently be instantiated as part of initializing the build code, so we have
 		// to disable the cache portion by injecting "-DDC=None" to the commandline args during process startup. This mode is not
@@ -129,89 +115,32 @@
 		// can lead to confusion when people attempt to use "-DDC=None" in the editor expecting it to behave like "-DDC=Cold".
 		// To avoid this confusion (and until we can use the build without the cache) it is restricted to use in programs only
 		// and not the editor.
-#if IS_PROGRAM
-		if (ActiveGraphName == TEXT("None"))
+		if (GraphNameOrConfig == TEXTVIEW("None"))
 		{
 			UE_LOG(LogDerivedDataCache, Display, TEXT("Requested cache graph of 'None'. Every cache operation will fail."));
+			return true;
 		}
-		else
 #endif
-		{
-			const auto ResetGraph = [this, &ParsedNodes]
-			{
-				// Destroy any cache stores that have been created.
-				delete RootCache;
-				RootCache = nullptr;
-				ParsedNodes.Empty();
-				MemoryCache = nullptr;
-				WritePakCache = nullptr;
-				Hierarchy = nullptr;
-				ReadPakCache.Empty();
-				Directories.Empty();
-			};
 
-			if (!ActiveGraphName.IsEmpty() && ActiveGraphName != TEXT("Default"))
-			{
-				if (!ParseGraph(ActiveGraphName, ParsedNodes) || !Hierarchy->HasAllFlags(ECacheStoreFlags::Query | ECacheStoreFlags::Store))
-				{
-					ResetGraph();
-					UE_LOGFMT(LogDerivedDataCache, Warning,
-						"Failed to create graph '{Graph}'. Attempting to create the default graph...", ActiveGraphName);
-				}
-			}
+		FParsedNodeMap ParsedNodes;
+		TGuardValue ParsedNodeMapGuard(ActiveParsedNodeMap, &ParsedNodes);
+		return ParseGraph(ActiveGraphName, ParsedNodes) && Hierarchy->HasAllFlags(ECacheStoreFlags::Query | ECacheStoreFlags::Store);
+	}
 
-			// UE_DEPRECATED(5.8, "Remove the legacy default block after the standard deprecation period.")
-			const TCHAR* LegacyDefaultGraphName = FApp::IsEngineInstalled() ? TEXT("InstalledDerivedDataBackendGraph") : TEXT("DerivedDataBackendGraph");
-			if (!RootCache && GConfig->DoesSectionExist(LegacyDefaultGraphName, GEngineIni))
-			{
-				// Try to use the legacy default graph.
-				ActiveGraphName = LegacyDefaultGraphName;
-				if (!ParseGraph(ActiveGraphName, ParsedNodes) || !Hierarchy->HasAllFlags(ECacheStoreFlags::Query | ECacheStoreFlags::Store))
-				{
-					ResetGraph();
-				}
-			}
+	void ResetGraph()
+	{
+		delete RootCache;
+		RootCache = nullptr;
+		MemoryCache = nullptr;
+		WritePakCache = nullptr;
+		Hierarchy = nullptr;
+		ReadPakCache.Empty();
+		Directories.Empty();
+		ActiveGraphName.Empty();
+	}
 
-			if (!RootCache)
-			{
-				// Try to use the default graph.
-				ActiveGraphName = FApp::IsEngineInstalled() ? TEXT("Installed") : TEXT("Default");
-				if (!ParseGraph(ActiveGraphName, ParsedNodes) || !Hierarchy->HasAllFlags(ECacheStoreFlags::Query | ECacheStoreFlags::Store))
-				{
-					TStringBuilder<256> GraphConfig;
-					if (!TryFindCacheGraphConfig(ActiveGraphName, GraphConfig) &&
-						!TryFindCacheGraphConfig(LegacyDefaultGraphName, GraphConfig))
-					{
-						UE_LOGFMT(LogDerivedDataCache, Fatal,
-							"Failed to create default cache graph '{Graph}' because its config is missing in the '{Config}' config.",
-							ActiveGraphName, GEngineIni);
-					}
-					else if (FParse::Param(FCommandLine::Get(), TEXT("DDC-ForceMemoryCache")))
-					{
-						IMemoryCacheStore* Memory = nullptr;
-						CreateMemoryCacheStore(Memory, TEXT("ForceMemoryCache"), TEXT(""), Hierarchy);
-						UE_LOGFMT(LogDerivedDataCache, Error,
-							"Unable to use default cache graph '{Graph}' because there are no {MissingTypes} nodes available.",
-							ActiveGraphName,
-							Hierarchy->HasAllFlags(ECacheStoreFlags::Query) ? "writable" :
-							Hierarchy->HasAllFlags(ECacheStoreFlags::Store) ? "readable" : "readable or writable");
-						check(Hierarchy->HasAllFlags(ECacheStoreFlags::Query | ECacheStoreFlags::Store));
-					}
-					else
-					{
-						UE_CLOGFMT(!FParse::Param(FCommandLine::Get(), TEXT("DDC-AllowNoActiveStores")),
-							LogDerivedDataCache, Fatal,
-							"Unable to use default cache graph '{Graph}' because there are no {MissingTypes} nodes available. "
-							"Add -DDC-ForceMemoryCache to the command line to bypass this if you need access "
-							"to the editor settings to fix the cache configuration.",
-							*ActiveGraphName,
-							Hierarchy->HasAllFlags(ECacheStoreFlags::Query) ? "writable" :
-							Hierarchy->HasAllFlags(ECacheStoreFlags::Store) ? "readable" : "readable or writable");
-					}
-				}
-			}
-		}
-
+	void PostCreateGraph()
+	{
 		// Async must exist in the graph.
 		RootCache = CreateCacheStoreAsync(Hierarchy, RootCache, MemoryCache, /*bDeleteInnerCache*/ true);
 
@@ -233,9 +162,13 @@
 			RootCache = ReplayNode;
 		}
 
+		// Note: This will contain maintainers for any caches that existed prior to this point.
+		//       This is not robust unless later caches are destroyed before earlier caches.
+		//       This restriction is unreasonable and will be solved later.
 		CacheStoreMaintainers = IModularFeatures::Get().GetModularFeatureImplementations<ICacheStoreMaintainer>(FeatureName);
 	}
 
+private:
 	/**
 	 * Parses backend graph node from ini settings
 	 *
@@ -478,20 +411,29 @@
 		}
 	}
 
+protected:
 	virtual ~FDerivedDataBackendGraph()
 	{
-		check(StaticGraph == this);
-		Replays.Empty();
 		delete RootCache;
-		StaticGraph = nullptr;
 	}
 
-	ILegacyCacheStore& GetRoot() override
+	bool HasRoot() const
 	{
+		return !!RootCache;
+	}
+
+	ILegacyCacheStore& GetRoot() final
+	{
 		check(RootCache);
 		return *RootCache;
 	}
 
+	ICacheStoreOwner* GetOwner() const
+	{
+		return Hierarchy;
+	}
+
+private:
 	virtual void WaitForQuiescence(bool bShutdown) override
 	{
 		if (bShutdown)
@@ -585,12 +527,7 @@
 		OutResults = Directories;
 	}
 
-	static FORCEINLINE FDerivedDataBackendGraph& GetStatic()
-	{
-		check(StaticGraph);
-		return *StaticGraph;
-	}
-
+protected:
 	virtual ILegacyCacheStore* MountPakFile(const TCHAR* PakFilename) override
 	{
 		// Assumptions: there's at least one read-only pak backend in the hierarchy
@@ -630,6 +567,7 @@
 		return false;
 	}
 
+private:
 	virtual TSharedRef<FDerivedDataCacheStatsNode> GatherUsageStats() const override
 	{
 		TSharedRef<FDerivedDataCacheStatsNode> Stats = MakeShared<FDerivedDataCacheStatsNode>();
@@ -656,7 +594,6 @@
 		}
 	}
 
-private:
 	void Add(ILegacyCacheStore* CacheStore, ECacheStoreFlags Flags) final
 	{
 		check(Hierarchy);
@@ -808,26 +745,220 @@
 		return nullptr;
 	}
 
-	/** MountPak console command handler. */
-	void UnmountPakCommandHandler(const TArray<FString>& Args)
+	// ICache Interface
+
+	void Put(
+		TConstArrayView<FCachePutRequest> Requests,
+		IRequestOwner& Owner,
+		FOnCachePutComplete&& OnComplete) final
 	{
-		if (Args.Num() < 1)
-		{
-			UE_LOG(LogDerivedDataCache, Log, TEXT("Usage: DDC.MountPak PakFilename"));
-			return;
-		}
-		UnmountPakFile(*Args[0]);
+		TRACE_CPUPROFILER_EVENT_SCOPE(DDC_Put);
+		return RootCache->Put(Requests, Owner, OnComplete ? MoveTemp(OnComplete) : FOnCachePutComplete([](auto&&) {}));
 	}
 
-	/** UnmountPak console command handler. */
+	void Get(
+		TConstArrayView<FCacheGetRequest> Requests,
+		IRequestOwner& Owner,
+		FOnCacheGetComplete&& OnComplete) final
+	{
+		TRACE_CPUPROFILER_EVENT_SCOPE(DDC_Get);
+		return RootCache->Get(Requests, Owner, OnComplete ? MoveTemp(OnComplete) : FOnCacheGetComplete([](auto&&) {}));
+	}
+
+	void PutValue(
+		TConstArrayView<FCachePutValueRequest> Requests,
+		IRequestOwner& Owner,
+		FOnCachePutValueComplete&& OnComplete) final
+	{
+		TRACE_CPUPROFILER_EVENT_SCOPE(DDC_PutValue);
+		return RootCache->PutValue(Requests, Owner, OnComplete ? MoveTemp(OnComplete) : FOnCachePutValueComplete([](auto&&) {}));
+	}
+
+	void GetValue(
+		TConstArrayView<FCacheGetValueRequest> Requests,
+		IRequestOwner& Owner,
+		FOnCacheGetValueComplete&& OnComplete) final
+	{
+		TRACE_CPUPROFILER_EVENT_SCOPE(DDC_GetValue);
+		return RootCache->GetValue(Requests, Owner, OnComplete ? MoveTemp(OnComplete) : FOnCacheGetValueComplete([](auto&&) {}));
+	}
+
+	void GetChunks(
+		TConstArrayView<FCacheGetChunkRequest> Requests,
+		IRequestOwner& Owner,
+		FOnCacheGetChunkComplete&& OnComplete) final
+	{
+		TRACE_CPUPROFILER_EVENT_SCOPE(DDC_GetChunks);
+		return RootCache->GetChunks(Requests, Owner, OnComplete ? MoveTemp(OnComplete) : FOnCacheGetChunkComplete([](auto&&) {}));
+	}
+
+	ICacheStoreMaintainer& GetMaintainer() final
+	{
+		return *this;
+	}
+
+	// ICacheStoreMaintainer Interface
+
+	bool IsIdle() const final
+	{
+		return Algo::AllOf(CacheStoreMaintainers, &ICacheStoreMaintainer::IsIdle);
+	}
+
+	void BoostPriority() final
+	{
+		for (ICacheStoreMaintainer* Maintainer : CacheStoreMaintainers)
+		{
+			Maintainer->BoostPriority();
+		}
+	}
+
+	FString											ActiveGraphName;
+	FString											ReadPakFilename;
+	FString											WritePakFilename;
+
+	/** Root of the graph */
+	ILegacyCacheStore* RootCache = nullptr;
+
+	TMap<ILegacyCacheStore*, ILegacyCacheStore*> ThrottleNodes;
+
+	/** Instances of backend interfaces which exist in only one copy */
+	IMemoryCacheStore* MemoryCache = nullptr;
+	IPakFileCacheStore* WritePakCache = nullptr;
+	ICacheStoreOwner* Hierarchy = nullptr;
+	/** Support for multiple read only pak files. */
+	TArray<IPakFileCacheStore*> ReadPakCache;
+
+	/** List of directories used by the DDC */
+	TArray<FString> Directories;
+
+	FBackendDebugOptions DefaultDebugOptions;
+	bool bHasDefaultDebugOptions = false;
+
+	/** Whether a shared cache is in use */
+	bool bUsingSharedDDC = false;
+
+	FParsedNodeMap* ActiveParsedNodeMap = nullptr;
+	FString ActiveNodeName;
+	FString ActiveNodeConfig;
+
+	TArray<ICacheStoreMaintainer*> CacheStoreMaintainers;
+};
+
+class FDerivedDataBackendStaticGraph final : public FDerivedDataBackendGraph
+{
+public:
+	FDerivedDataBackendStaticGraph()
+		: MountPakCommand(
+			TEXT("DDC.MountPak"),
+			*LOCTEXT("CommandText_DDCMountPak", "Mounts read-only pak file").ToString(),
+			FConsoleCommandWithArgsDelegate::CreateRaw(this, &FDerivedDataBackendStaticGraph::MountPakCommandHandler))
+		, UnmountPakCommand(
+			TEXT("DDC.UnmountPak"),
+			*LOCTEXT("CommandText_DDCUnmountPak", "Unmounts read-only pak file").ToString(),
+			FConsoleCommandWithArgsDelegate::CreateRaw(this, &FDerivedDataBackendStaticGraph::UnmountPakCommandHandler))
+		, LoadReplayCommand(
+			TEXT("DDC.LoadReplay"),
+			*LOCTEXT("CommandText_DDCLoadReplay", "Loads a cache replay file created by -DDC-ReplaySave=<Path>").ToString(),
+			FConsoleCommandWithArgsDelegate::CreateRaw(this, &FDerivedDataBackendStaticGraph::LoadReplayCommandHandler))
+	{
+		check(!StaticGraph);
+		StaticGraph = this;
+
+		FString GraphName;
+		if (!FParse::Value(FCommandLine::Get(), TEXT("-DDC="), GraphName, /*bShouldStopOnSeparator*/ false))
+		{
+			GraphName = GDerivedDataCacheGraphName.GetValueOnGameThread();
+		}
+
+		const TCHAR* const DefaultGraphName = FApp::IsEngineInstalled() ? TEXT("Installed") : TEXT("Default");
+		const TCHAR* const LegacyDefaultGraphName = FApp::IsEngineInstalled() ? TEXT("InstalledDerivedDataBackendGraph") : TEXT("DerivedDataBackendGraph");
+
+		// Try to use the graph from DDC.Graph or -DDC=...
+		if (!GraphName.IsEmpty() && GraphName != DefaultGraphName && !CreateGraph(GraphName))
+		{
+			ResetGraph();
+			UE_LOGFMT(LogDerivedDataCache, Warning,
+				"Failed to create graph '{Graph}'. Attempting to create the default graph...", GraphName);
+		}
+
+		// UE_DEPRECATED(5.8, "Remove the legacy default block after the standard deprecation period.")
+		// Try to use the legacy default graph.
+		if (!HasRoot() && GConfig->DoesSectionExist(LegacyDefaultGraphName, GEngineIni) && !CreateGraph(LegacyDefaultGraphName))
+		{
+			ResetGraph();
+		}
+
+		// Try to use the default graph.
+		if (!HasRoot() && !CreateGraph(DefaultGraphName))
+		{
+			ICacheStoreOwner* Owner = GetOwner();
+			TStringBuilder<256> GraphConfig;
+			if (!TryFindCacheGraphConfig(DefaultGraphName, GraphConfig) &&
+				!TryFindCacheGraphConfig(LegacyDefaultGraphName, GraphConfig))
+			{
+				UE_LOGFMT(LogDerivedDataCache, Fatal,
+					"Failed to create default cache graph '{Graph}' because its config is missing in the '{Config}' config.",
+					DefaultGraphName, GEngineIni);
+			}
+			else if (FParse::Param(FCommandLine::Get(), TEXT("DDC-ForceMemoryCache")))
+			{
+				UE_LOGFMT(LogDerivedDataCache, Error,
+					"Unable to use default cache graph '{Graph}' because there are no {MissingTypes} nodes available.",
+					DefaultGraphName,
+					Owner->HasAllFlags(ECacheStoreFlags::Query) ? "writable" :
+					Owner->HasAllFlags(ECacheStoreFlags::Store) ? "readable" : "readable or writable");
+				IMemoryCacheStore* Memory = nullptr;
+				CreateMemoryCacheStore(Memory, TEXT("ForceMemoryCache"), TEXT(""), Owner);
+				check(Owner->HasAllFlags(ECacheStoreFlags::Query | ECacheStoreFlags::Store));
+			}
+			else
+			{
+				UE_CLOGFMT(!FParse::Param(FCommandLine::Get(), TEXT("DDC-AllowNoActiveStores")),
+					LogDerivedDataCache, Fatal,
+					"Unable to use default cache graph '{Graph}' because there are no {MissingTypes} nodes available. "
+					"Add -DDC-ForceMemoryCache to the command line to bypass this if you need access "
+					"to the editor settings to fix the cache configuration.",
+					DefaultGraphName,
+					Owner->HasAllFlags(ECacheStoreFlags::Query) ? "writable" :
+					Owner->HasAllFlags(ECacheStoreFlags::Store) ? "readable" : "readable or writable");
+			}
+		}
+
+		PostCreateGraph();
+	}
+
+	~FDerivedDataBackendStaticGraph() final
+	{
+		check(StaticGraph == this);
+		StaticGraph = nullptr;
+		Replays.Empty();
+	}
+
+	static FORCEINLINE FDerivedDataBackendGraph& GetStatic()
+	{
+		check(StaticGraph);
+		return *StaticGraph;
+	}
+
+private:
 	void MountPakCommandHandler(const TArray<FString>& Args)
 	{
 		if (Args.Num() < 1)
 		{
-			UE_LOG(LogDerivedDataCache, Log, TEXT("Usage: DDC.UnmountPak PakFilename"));
+			UE_LOG(LogDerivedDataCache, Log, TEXT("Usage: DDC.MountPak PakFilename"));
 			return;
 		}
 		MountPakFile(*Args[0]);
 	}
 
+	void UnmountPakCommandHandler(const TArray<FString>& Args)
+	{
+		if (Args.Num() < 1)
+		{
+			UE_LOG(LogDerivedDataCache, Log, TEXT("Usage: DDC.UnmountPak PakFilename"));
+			return;
+		}
+		UnmountPakFile(*Args[0]);
+	}
+
 	void LoadReplayCommandHandler(const TArray<FString>& Args)
@@ -848,7 +979,7 @@
 		TStringBuilder<512> JoinedArgs;
 		JoinedArgs.Join(MakeArrayView(Args).RightChop(1), TEXT(' '));
 
-		FCacheReplayReader Replay(RootCache);
+		FCacheReplayReader Replay(&GetRoot());
 
 		// Parse Key Filter
 		const bool bDefaultMatch = String::FindFirst(*JoinedArgs, TEXT("Types="), ESearchCase::IgnoreCase) == INDEX_NONE;
@@ -910,118 +1041,18 @@
 			Replay.SetPriorityOverride(Priority);
 		}
 
-		Replay.ReadFromFileAsync(*Args[0], Hierarchy);
+		Replay.ReadFromFileAsync(*Args[0], GetOwner());
 		Replays.Add(MoveTemp(Replay));
 	}
 
-	// ICache Interface
+private:
+	inline static FDerivedDataBackendStaticGraph* StaticGraph;
 
-	void Put(
-		TConstArrayView<FCachePutRequest> Requests,
-		IRequestOwner& Owner,
-		FOnCachePutComplete&& OnComplete) final
-	{
-		TRACE_CPUPROFILER_EVENT_SCOPE(DDC_Put);
-		return RootCache->Put(Requests, Owner, OnComplete ? MoveTemp(OnComplete) : FOnCachePutComplete([](auto&&) {}));
-	}
-
-	void Get(
-		TConstArrayView<FCacheGetRequest> Requests,
-		IRequestOwner& Owner,
-		FOnCacheGetComplete&& OnComplete) final
-	{
-		TRACE_CPUPROFILER_EVENT_SCOPE(DDC_Get);
-		return RootCache->Get(Requests, Owner, OnComplete ? MoveTemp(OnComplete) : FOnCacheGetComplete([](auto&&) {}));
-	}
-
-	void PutValue(
-		TConstArrayView<FCachePutValueRequest> Requests,
-		IRequestOwner& Owner,
-		FOnCachePutValueComplete&& OnComplete) final
-	{
-		TRACE_CPUPROFILER_EVENT_SCOPE(DDC_PutValue);
-		return RootCache->PutValue(Requests, Owner, OnComplete ? MoveTemp(OnComplete) : FOnCachePutValueComplete([](auto&&) {}));
-	}
-
-	void GetValue(
-		TConstArrayView<FCacheGetValueRequest> Requests,
-		IRequestOwner& Owner,
-		FOnCacheGetValueComplete&& OnComplete) final
-	{
-		TRACE_CPUPROFILER_EVENT_SCOPE(DDC_GetValue);
-		return RootCache->GetValue(Requests, Owner, OnComplete ? MoveTemp(OnComplete) : FOnCacheGetValueComplete([](auto&&) {}));
-	}
-
-	void GetChunks(
-		TConstArrayView<FCacheGetChunkRequest> Requests,
-		IRequestOwner& Owner,
-		FOnCacheGetChunkComplete&& OnComplete) final
-	{
-		TRACE_CPUPROFILER_EVENT_SCOPE(DDC_GetChunks);
-		return RootCache->GetChunks(Requests, Owner, OnComplete ? MoveTemp(OnComplete) : FOnCacheGetChunkComplete([](auto&&) {}));
-	}
-
-	ICacheStoreMaintainer& GetMaintainer() final
-	{
-		return *this;
-	}
-
-	// ICacheStoreMaintainer Interface
-
-	bool IsIdle() const final
-	{
-		return Algo::AllOf(CacheStoreMaintainers, &ICacheStoreMaintainer::IsIdle);
-	}
-
-	void BoostPriority() final
-	{
-		for (ICacheStoreMaintainer* Maintainer : CacheStoreMaintainers)
-		{
-			Maintainer->BoostPriority();
-		}
-	}
-
-	static inline FDerivedDataBackendGraph*			StaticGraph;
-
-	FString											ActiveGraphName;
-	FString											ReadPakFilename;
-	FString											WritePakFilename;
-
-	/** Root of the graph */
-	ILegacyCacheStore*					RootCache;
-
-	TMap<ILegacyCacheStore*, ILegacyCacheStore*> ThrottleNodes;
-
-	/** Instances of backend interfaces which exist in only one copy */
-	IMemoryCacheStore* MemoryCache;
-	IPakFileCacheStore* WritePakCache;
-	ICacheStoreOwner* Hierarchy;
-	/** Support for multiple read only pak files. */
-	TArray<IPakFileCacheStore*> ReadPakCache;
-
-	/** List of directories used by the DDC */
-	TArray<FString> Directories;
-
-	FBackendDebugOptions DefaultDebugOptions;
-
-	bool bHasDefaultDebugOptions = false;
-
-	/** Whether a shared cache is in use */
-	bool bUsingSharedDDC;
-
-	/** MountPak console command */
 	FAutoConsoleCommand MountPakCommand;
-	/** UnmountPak console command */
 	FAutoConsoleCommand UnmountPakCommand;
 
 	FAutoConsoleCommand LoadReplayCommand;
 	TArray<FCacheReplayReader> Replays;
-
-	FParsedNodeMap* ActiveParsedNodeMap = nullptr;
-	FString ActiveNodeName;
-	FString ActiveNodeConfig;
-
-	TArray<ICacheStoreMaintainer*> CacheStoreMaintainers;
 };
 
 } // UE::DerivedData
@@ -1031,12 +1062,12 @@
 
 FDerivedDataBackend* FDerivedDataBackend::CreateStatic()
 {
-	return new FDerivedDataBackendGraph();
+	return new FDerivedDataBackendStaticGraph();
 }
 
 FDerivedDataBackend& FDerivedDataBackend::GetStatic()
 {
-	return FDerivedDataBackendGraph::GetStatic();
+	return FDerivedDataBackendStaticGraph::GetStatic();
 }
 
 ///////////////////////////////////////////////////////////////////////////////////////////////////////////////////////
@@ -1134,4 +1165,14 @@
 
 } // UE::DerivedData
 
+namespace UE::DerivedData::Private
+{
+
+ICache* TryCreateCache(FStringView GraphNameOrConfig)
+{
+	return FDerivedDataBackendGraph::TryCreate(GraphNameOrConfig);
+}
+
+} // UE::DerivedData::Private
+
 #undef LOCTEXT_NAMESPACE
diff --git a/Engine/Source/Developer/DerivedDataCache/Private/DerivedDataCacheModule.cpp b/Engine/Source/Developer/DerivedDataCache/Private/DerivedDataCacheModule.cpp
--- a/Engine/Source/Developer/DerivedDataCache/Private/DerivedDataCacheModule.cpp
+++ b/Engine/Source/Developer/DerivedDataCache/Private/DerivedDataCacheModule.cpp
@@ -19,6 +19,9 @@ namespace UE::DerivedData::Private
 // Implemented in DerivedDataCache.cpp
 void CreateCache(ICache*& OutCache, FDerivedDataCacheInterface*& OutLegacyCache);
 
+// Implemented in DerivedDataBackends.cpp
+ICache* TryCreateCache(FStringView GraphNameOrConfig);
+
 static FDerivedDataCacheInterface* GDerivedDataLegacyCache;
 static ICache* GDerivedDataCache;
 static IBuild* GDerivedDataBuild;
@@ -106,7 +109,7 @@ ICache& GetCache()
 	{
 		return *Cache;
 	}
-	checkf(IsInGameThread(), TEXT("The derived data cache must be created on the main thread."));
+	check(IsInGameThread());
 	if (Private::FDerivedDataCacheModule* Module = Private::GetModule())
 	{
 		Module->CreateCacheOnce();
@@ -121,13 +124,23 @@ ICache* TryGetCache()
 	return Private::GDerivedDataCache;
 }
 
+ICache* CreateCache(FStringView GraphNameOrConfig)
+{
+	check(IsInGameThread());
+	if (Private::GetModule())
+	{
+		return Private::TryCreateCache(GraphNameOrConfig);
+	}
+	return nullptr;
+}
+
 IBuild& GetBuild()
 {
 	if (IBuild* Build = Private::GDerivedDataBuild)
 	{
 		return *Build;
 	}
-	checkf(IsInGameThread(), TEXT("The derived data build system must be created on the main thread."));
+	check(IsInGameThread());
 	if (Private::FDerivedDataCacheModule* Module = Private::GetModule())
 	{
 		Module->CreateBuildOnce();
diff --git a/Engine/Source/Developer/DerivedDataCache/Public/DerivedDataCache.h b/Engine/Source/Developer/DerivedDataCache/Public/DerivedDataCache.h
--- a/Engine/Source/Developer/DerivedDataCache/Public/DerivedDataCache.h
+++ b/Engine/Source/Developer/DerivedDataCache/Public/DerivedDataCache.h
@@ -375,6 +375,17 @@ UE_API ICache& GetCache();
 /** Returns a pointer to the cache. Null if not available or not created. */
 UE_API ICache* TryGetCache();
 
+/**
+ * Creates a cache from a graph name or graph config.
+ *
+ * Must be called from the main thread.
+ * Destroy the cache with operator delete. Consider using TUniquePtr.
+ *
+ * @param GraphNameOrConfig   Name of the graph in [DerivedDataCacheGraphs] or a graph config.
+ * @return A pointer to the new cache, or null if it failed to create.
+ */
+UE_INTERNAL UE_API ICache* CreateCache(FStringView GraphNameOrConfig);
+
 ///////////////////////////////////////////////////////////////////////////////////////////////////////////////////////
 
 } // UE::DerivedData
