DDC: Register cache store maintainers with the cache graph directly

#rnx
#Synced-CL 32500503

diff --git a/Engine/Source/Developer/DerivedDataCache/Private/DerivedDataBackends.cpp b/Engine/Source/Developer/DerivedDataCache/Private/DerivedDataBackends.cpp
--- a/Engine/Source/Developer/DerivedDataCache/Private/DerivedDataBackends.cpp
+++ b/Engine/Source/Developer/DerivedDataCache/Private/DerivedDataBackends.cpp
@@ -3,6 +3,9 @@
 #include "DerivedDataBackendInterface.h"
 
 #include "Algo/AllOf.h"
+#include "Async/SharedLock.h"
+#include "Async/SharedMutex.h"
+#include "Async/UniqueLock.h"
 #include "CoreTypes.h"
 #include "Containers/StringView.h"
 #include "DerivedDataCacheMaintainer.h"
@@ -13,7 +16,6 @@
 #include "DerivedDataCacheUsageStats.h"
 #include "DerivedDataConfig.h"
 #include "DerivedDataRequestOwner.h"
-#include "Features/IModularFeatures.h"
 #include "HAL/CriticalSection.h"
 #include "HAL/FileManager.h"
 #include "HAL/IConsoleManager.h"
@@ -161,11 +163,6 @@ protected:
 		{
 			RootCache = ReplayNode;
 		}
-
-		// Note: This will contain maintainers for any caches that existed prior to this point.
-		//       This is not robust unless later caches are destroyed before earlier caches.
-		//       This restriction is unreasonable and will be solved later.
-		CacheStoreMaintainers = IModularFeatures::Get().GetModularFeatureImplementations<ICacheStoreMaintainer>(FeatureName);
 	}
 
 private:
@@ -668,6 +665,18 @@ private:
 		return Hierarchy->HasAllFlags(Flags);
 	}
 
+	void AddMaintainer(ICacheStoreMaintainer* Maintainer) final
+	{
+		TUniqueLock Lock(CacheStoreMaintainerMutex);
+		CacheStoreMaintainers.Add(Maintainer);
+	}
+
+	void RemoveMaintainer(ICacheStoreMaintainer* Maintainer) final
+	{
+		TUniqueLock Lock(CacheStoreMaintainerMutex);
+		CacheStoreMaintainers.Remove(Maintainer);
+	}
+
 	int32 AddToAsyncTaskCounter(int32 Value) final
 	{
 		check(Hierarchy);
@@ -801,11 +810,13 @@ private:
 
 	bool IsIdle() const final
 	{
+		TSharedLock Lock(CacheStoreMaintainerMutex);
 		return Algo::AllOf(CacheStoreMaintainers, &ICacheStoreMaintainer::IsIdle);
 	}
 
 	void BoostPriority() final
 	{
+		TSharedLock Lock(CacheStoreMaintainerMutex);
 		for (ICacheStoreMaintainer* Maintainer : CacheStoreMaintainers)
 		{
 			Maintainer->BoostPriority();
@@ -842,6 +853,7 @@ private:
 	FString ActiveNodeConfig;
 
 	TArray<ICacheStoreMaintainer*> CacheStoreMaintainers;
+	mutable FSharedMutex CacheStoreMaintainerMutex;
 };
 
 class FDerivedDataBackendStaticGraph final : public FDerivedDataBackendGraph
diff --git a/Engine/Source/Developer/DerivedDataCache/Private/DerivedDataCacheStore.h b/Engine/Source/Developer/DerivedDataCache/Private/DerivedDataCacheStore.h
--- a/Engine/Source/Developer/DerivedDataCache/Private/DerivedDataCacheStore.h
+++ b/Engine/Source/Developer/DerivedDataCache/Private/DerivedDataCacheStore.h
@@ -15,6 +15,7 @@ struct FDerivedDataCacheResourceStat;
 namespace UE::DerivedData { class FCacheRecord; }
 namespace UE::DerivedData { class FCacheStoreRequestTimer; }
 namespace UE::DerivedData { class FValue; }
+namespace UE::DerivedData { class ICacheStoreMaintainer; }
 namespace UE::DerivedData { class ICacheStoreStats; }
 namespace UE::DerivedData { class ILegacyCacheStore; }
 namespace UE::DerivedData { struct FCacheStoreRequestStats; }
@@ -201,6 +202,9 @@ public:
 	/** Returns true if the combined flags of the owned cache stores contain all of these flags. */
 	virtual bool HasAllFlags(ECacheStoreFlags Flags) const = 0;
 
+	virtual void AddMaintainer(ICacheStoreMaintainer* Maintainer) = 0;
+	virtual void RemoveMaintainer(ICacheStoreMaintainer* Maintainer) = 0;
+
 	virtual int32 AddToAsyncTaskCounter(int32 Value) = 0;
 
 	virtual ICacheStoreStats* CreateStats(ILegacyCacheStore* CacheStore, ECacheStoreFlags Flags, FStringView Type, FStringView Name, FStringView Path = {}) = 0;
diff --git a/Engine/Source/Developer/DerivedDataCache/Private/DerivedDataCacheStoreHierarchy.cpp b/Engine/Source/Developer/DerivedDataCache/Private/DerivedDataCacheStoreHierarchy.cpp
--- a/Engine/Source/Developer/DerivedDataCache/Private/DerivedDataCacheStoreHierarchy.cpp
+++ b/Engine/Source/Developer/DerivedDataCache/Private/DerivedDataCacheStoreHierarchy.cpp
@@ -50,6 +50,9 @@ public:
 
 	bool HasAllFlags(ECacheStoreFlags Flags) const final;
 
+	void AddMaintainer(ICacheStoreMaintainer* Maintainer) final;
+	void RemoveMaintainer(ICacheStoreMaintainer* Maintainer) final;
+
 	int32 AddToAsyncTaskCounter(int32 Value) final;
 
 	ICacheStoreStats* CreateStats(ILegacyCacheStore* CacheStore, ECacheStoreFlags Flags, FStringView Type, FStringView Name, FStringView Path) final;
@@ -359,6 +362,20 @@ bool FCacheStoreHierarchy::HasAllFlags(ECacheStoreFlags Flags) const
 	return EnumHasAllFlags(CombinedFlags, Flags);
 }
 
+void FCacheStoreHierarchy::AddMaintainer(ICacheStoreMaintainer* Maintainer)
+{
+	// All stores are currently created with the graph as their owner.
+	// The graph does not delegate this functionality to the hierarchy.
+	checkNoEntry();
+}
+
+void FCacheStoreHierarchy::RemoveMaintainer(ICacheStoreMaintainer* Maintainer)
+{
+	// All stores are currently created with the graph as their owner.
+	// The graph does not delegate this functionality to the hierarchy.
+	checkNoEntry();
+}
+
 int32 FCacheStoreHierarchy::AddToAsyncTaskCounter(int32 Value)
 {
 	return AsyncTaskCounter.fetch_add(Value);
diff --git a/Engine/Source/Developer/DerivedDataCache/Private/FileSystemCacheStore.cpp b/Engine/Source/Developer/DerivedDataCache/Private/FileSystemCacheStore.cpp
--- a/Engine/Source/Developer/DerivedDataCache/Private/FileSystemCacheStore.cpp
+++ b/Engine/Source/Developer/DerivedDataCache/Private/FileSystemCacheStore.cpp
@@ -20,7 +20,6 @@
 #include "DerivedDataChunk.h"
 #include "DerivedDataRequestOwner.h"
 #include "DerivedDataValue.h"
-#include "Features/IModularFeatures.h"
 #include "HAL/Event.h"
 #include "HAL/FileManager.h"
 #include "HAL/Thread.h"
@@ -219,7 +218,10 @@ struct FFileSystemCacheStoreMaintainerParams
 class FFileSystemCacheStoreMaintainer final : public ICacheStoreMaintainer
 {
 public:
-	FFileSystemCacheStoreMaintainer(const FFileSystemCacheStoreMaintainerParams& Params, FStringView CachePath);
+	FFileSystemCacheStoreMaintainer(
+		const FFileSystemCacheStoreMaintainerParams& Params,
+		FStringView CachePath,
+		ICacheStoreOwner& Owner);
 	~FFileSystemCacheStoreMaintainer();
 
 	bool IsIdle() const final { return bIdle; }
@@ -252,6 +254,7 @@ private:
 	struct FRoot;
 	struct FLegacyRoot;
 
+	ICacheStoreOwner& StoreOwner;
 	FFileSystemCacheStoreMaintainerParams Params;
 	/** Path to the root of the cache store. */
 	FString CachePath;
@@ -332,8 +335,10 @@ struct FFileSystemCacheStoreMaintainer::FLegacyRoot
 
 FFileSystemCacheStoreMaintainer::FFileSystemCacheStoreMaintainer(
 	const FFileSystemCacheStoreMaintainerParams& InParams,
-	const FStringView InCachePath)
-	: Params(InParams)
+	const FStringView InCachePath,
+	ICacheStoreOwner& InOwner)
+	: StoreOwner(InOwner)
+	, Params(InParams)
 	, CachePath(InCachePath)
 	, bExitAfterScan(Params.ScanFrequency.GetTotalDays() > MaxScanFrequencyDays)
 	, WaitEvent(EEventMode::AutoReset)
@@ -344,13 +349,13 @@ FFileSystemCacheStoreMaintainer::FFileSystemCacheStoreMaintainer(
 		/*StackSize*/ 32 * 1024,
 		TPri_BelowNormal)
 {
-	IModularFeatures::Get().RegisterModularFeature(FeatureName, this);
+	StoreOwner.AddMaintainer(this);
 }
 
 FFileSystemCacheStoreMaintainer::~FFileSystemCacheStoreMaintainer()
 {
 	bExit = true;
-	IModularFeatures::Get().UnregisterModularFeature(FeatureName, this);
+	StoreOwner.RemoveMaintainer(this);
 	WaitEvent->Trigger();
 	Thread.Join();
 }
@@ -1258,7 +1263,7 @@ FFileSystemCacheStore::FFileSystemCacheStore(
 
 		if (!bLocalDeactivatedForPerformance)
 		{
-			Maintainer = MakeUnique<FFileSystemCacheStoreMaintainer>(*MaintainerParams, CachePath);
+			Maintainer = MakeUnique<FFileSystemCacheStoreMaintainer>(*MaintainerParams, CachePath, StoreOwner);
 
 			if (Params.bClean)
 			{
@@ -2881,7 +2886,7 @@ bool FFileSystemCacheStore::IsDeactivatedForPerformance()
 						// We're no longer deactivated for performance.  If maintenance was deferred, do it now.
 						if (FFileSystemCacheStoreMaintainerParams* MaintainerParams = DeactivationDeferredMaintainerParams.Get())
 						{
-							Maintainer = MakeUnique<FFileSystemCacheStoreMaintainer>(*MaintainerParams, CachePath);
+							Maintainer = MakeUnique<FFileSystemCacheStoreMaintainer>(*MaintainerParams, CachePath, StoreOwner);
 							DeactivationDeferredMaintainerParams.Reset();
 
 							if (bDeactivationDeferredClean)
diff --git a/Engine/Source/Developer/DerivedDataCache/Public/DerivedDataCacheMaintainer.h b/Engine/Source/Developer/DerivedDataCache/Public/DerivedDataCacheMaintainer.h
--- a/Engine/Source/Developer/DerivedDataCache/Public/DerivedDataCacheMaintainer.h
+++ b/Engine/Source/Developer/DerivedDataCache/Public/DerivedDataCacheMaintainer.h
@@ -2,7 +2,6 @@
 
 #pragma once
 
-#include "Features/IModularFeature.h"
 #include "UObject/NameTypes.h"
 
 namespace UE::DerivedData
@@ -17,11 +16,9 @@ namespace UE::DerivedData
  * An example consumer is the DDCCleanup commandlet which uses this to boost priority and wait on
  * the completion of maintenance by cache stores.
  */
-class ICacheStoreMaintainer : public IModularFeature
+class ICacheStoreMaintainer
 {
 public:
-	static inline const FLazyName FeatureName{"CacheStoreMaintainer"};
-
 	/**
 	 * True when maintenance is not active.
 	 *
