Deprecate `FThreadSafeRefCountedObject` as `FRefCountedObject` is now thread-safe by default.

#Synced-CL 48282531

diff --git a/Engine/Plugins/Experimental/ChaosRigidAsset/Source/ChaosRigidAssetNodes/Public/CopyOnWriteArray.h b/Engine/Plugins/Experimental/ChaosRigidAsset/Source/ChaosRigidAssetNodes/Public/CopyOnWriteArray.h
--- a/Engine/Plugins/Experimental/ChaosRigidAsset/Source/ChaosRigidAssetNodes/Public/CopyOnWriteArray.h
+++ b/Engine/Plugins/Experimental/ChaosRigidAsset/Source/ChaosRigidAssetNodes/Public/CopyOnWriteArray.h
@@ -256,7 +256,7 @@ private:
 		State = New;
 	}
 
-	struct FState : public FThreadSafeRefCountedObject
+	struct FState : public FRefCountedObject
 	{
 		TArray<T> Arr;
 	};
diff --git a/Engine/Plugins/Experimental/UAF/UAF/Source/UAF/Private/UAF/Attributes/AttributeBindingData.cpp b/Engine/Plugins/Experimental/UAF/UAF/Source/UAF/Private/UAF/Attributes/AttributeBindingData.cpp
--- a/Engine/Plugins/Experimental/UAF/UAF/Source/UAF/Private/UAF/Attributes/AttributeBindingData.cpp
+++ b/Engine/Plugins/Experimental/UAF/UAF/Source/UAF/Private/UAF/Attributes/AttributeBindingData.cpp
@@ -435,7 +435,7 @@ namespace UE::UAF
 
 	uint32 FAttributeBindingData::Release() const
 	{
-		const uint32 RefsBeforeRelease = FThreadSafeRefCountedObject::Release();
+		const uint32 RefsBeforeRelease = FRefCountedObject::Release();
 		if (RefsBeforeRelease == NumInternalReferences + 1)
 		{
 			// Our last external reference has been released, release our internal references to free the memory
@@ -448,12 +448,12 @@ namespace UE::UAF
 	void FAttributeBindingData::ResetInternalReferences()
 	{
 		// Add an extra reference to keep this container alive during destruction
-		FThreadSafeRefCountedObject::AddRef();
+		FRefCountedObject::AddRef();
 
 		DefaultAttributeValuesByName.Empty();
 
 		// We are done with destruction, destroy ourselves now
-		FThreadSafeRefCountedObject::Release();
+		FRefCountedObject::Release();
 	}
 
 	FAttributeBindingDataPtr MakeAttributeBindingData(const UAbstractSkeletonSetBinding& Binding)
diff --git a/Engine/Plugins/Experimental/UAF/UAF/Source/UAF/Public/UAF/Attributes/AttributeBindingData.h b/Engine/Plugins/Experimental/UAF/UAF/Source/UAF/Public/UAF/Attributes/AttributeBindingData.h
--- a/Engine/Plugins/Experimental/UAF/UAF/Source/UAF/Public/UAF/Attributes/AttributeBindingData.h
+++ b/Engine/Plugins/Experimental/UAF/UAF/Source/UAF/Public/UAF/Attributes/AttributeBindingData.h
@@ -18,7 +18,7 @@ namespace UE::UAF
 
 	// Our baked abstract skeleton binding data
 	// This class is ref counted and all named/typed sets increment/decrement its ref count
-	class FAttributeBindingData : public FThreadSafeRefCountedObject
+	class FAttributeBindingData : public FRefCountedObject
 	{
 	public:
 		using NamedSetMapType = TSortedMap<FName, FAttributeNamedSet*, FDefaultAllocator, FNameFastLess>;
@@ -93,7 +93,7 @@ namespace UE::UAF
 		// Returns a collection of the default attribute values for the specified named set
 		[[nodiscard]] UE_API FAttributeCollectionPtr GetDefaultAttributeValues(FName SetName) const;
 
-		// FThreadSafeRefCountedObject API
+		// FRefCountedObject API
 		UE_API uint32 Release() const;
 
 	private:
diff --git a/Engine/Plugins/Runtime/GeometryCache/Source/GeometryCacheStreamer/Private/GeometryCacheStreamer.cpp b/Engine/Plugins/Runtime/GeometryCache/Source/GeometryCacheStreamer/Private/GeometryCacheStreamer.cpp
--- a/Engine/Plugins/Runtime/GeometryCache/Source/GeometryCacheStreamer/Private/GeometryCacheStreamer.cpp
+++ b/Engine/Plugins/Runtime/GeometryCache/Source/GeometryCacheStreamer/Private/GeometryCacheStreamer.cpp
@@ -51,7 +51,7 @@ private:
 private:
 	FTSTicker::FDelegateHandle TickHandle;
 
-	class FRefCountedStream : public FThreadSafeRefCountedObject
+	class FRefCountedStream : public FRefCountedObject
 	{
 	public:
 		FRefCountedStream(IGeometryCacheStream* InStream)
diff --git a/Engine/Source/Developer/DerivedDataCache/Private/FileSystemCacheStore.cpp b/Engine/Source/Developer/DerivedDataCache/Private/FileSystemCacheStore.cpp
--- a/Engine/Source/Developer/DerivedDataCache/Private/FileSystemCacheStore.cpp
+++ b/Engine/Source/Developer/DerivedDataCache/Private/FileSystemCacheStore.cpp
@@ -2989,7 +2989,7 @@ void FFileSystemCacheStore::UpdateStatus()
 
 bool FFileSystemCacheStore::RunInitialSpeedTest(const FFileSystemCacheStoreParams& Params, FDerivedDataCacheSpeedStats& OutSpeedStats)
 {
-	struct FSpeedTestState : public FThreadSafeRefCountedObject
+	struct FSpeedTestState : public FRefCountedObject
 	{
 		FDerivedDataCacheSpeedStats SpeedStats;
 		std::atomic<uint32> NumLatencyTestsCompleted = 0;
diff --git a/Engine/Source/Developer/DerivedDataCache/Private/HttpCacheStore.cpp b/Engine/Source/Developer/DerivedDataCache/Private/HttpCacheStore.cpp
--- a/Engine/Source/Developer/DerivedDataCache/Private/HttpCacheStore.cpp
+++ b/Engine/Source/Developer/DerivedDataCache/Private/HttpCacheStore.cpp
@@ -760,7 +760,7 @@ private:
 //----------------------------------------------------------------------------------------------------------
 // FHttpCacheStore::FPutPackageOp
 //----------------------------------------------------------------------------------------------------------
-class FHttpCacheStore::FPutPackageOp final : public FThreadSafeRefCountedObject
+class FHttpCacheStore::FPutPackageOp final : public FRefCountedObject
 {
 public:
 	struct FResponse
@@ -1140,7 +1140,7 @@ void FHttpCacheStore::FPutPackageOp::EndPut(EStatus Status)
 //----------------------------------------------------------------------------------------------------------
 // FHttpCacheStore::FGetRecordOp
 //----------------------------------------------------------------------------------------------------------
-class FHttpCacheStore::FGetRecordOp final : public FThreadSafeRefCountedObject
+class FHttpCacheStore::FGetRecordOp final : public FRefCountedObject
 {
 public:
 	static TRefCountPtr<FGetRecordOp> New(FHttpCacheStore& CacheStore, IRequestOwner& Owner, const FSharedString& Name)
@@ -1768,7 +1768,7 @@ void FHttpCacheStore::FGetRecordOp::RecordStats(EStatus Status)
 //----------------------------------------------------------------------------------------------------------
 // FHttpCacheStore::FGetValueOp
 //----------------------------------------------------------------------------------------------------------
-class FHttpCacheStore::FGetValueOp final : public FThreadSafeRefCountedObject
+class FHttpCacheStore::FGetValueOp final : public FRefCountedObject
 {
 public:
 	struct FResponse
@@ -1985,7 +1985,7 @@ void FHttpCacheStore::FGetValueOp::EndGet(FResponse&& Response)
 //----------------------------------------------------------------------------------------------------------
 // FHttpCacheStore::FExistsBatchOp
 //----------------------------------------------------------------------------------------------------------
-class FHttpCacheStore::FExistsBatchOp final : public FThreadSafeRefCountedObject
+class FHttpCacheStore::FExistsBatchOp final : public FRefCountedObject
 {
 public:
 	static TRefCountPtr<FExistsBatchOp> New(FHttpCacheStore& CacheStore, IRequestOwner& Owner)
diff --git a/Engine/Source/Developer/DerivedDataCache/Private/ZenCacheStore.cpp b/Engine/Source/Developer/DerivedDataCache/Private/ZenCacheStore.cpp
--- a/Engine/Source/Developer/DerivedDataCache/Private/ZenCacheStore.cpp
+++ b/Engine/Source/Developer/DerivedDataCache/Private/ZenCacheStore.cpp
@@ -385,7 +385,7 @@ bool FZenCacheStore::ValidateAndPrepareReceivedValue(
 	return true;
 }
 
-class FZenCacheStore::FPutOp final : public FThreadSafeRefCountedObject
+class FZenCacheStore::FPutOp final : public FRefCountedObject
 {
 public:
 	FPutOp(FZenCacheStore& InCacheStore,
@@ -622,7 +622,7 @@ private:
 	FOnCachePutComplete OnComplete;
 };
 
-class FZenCacheStore::FGetOp final : public FThreadSafeRefCountedObject
+class FZenCacheStore::FGetOp final : public FRefCountedObject
 {
 public:
 	FGetOp(FZenCacheStore& InCacheStore,
@@ -865,7 +865,7 @@ private:
 	FOnCacheGetComplete OnComplete;
 };
 
-class FZenCacheStore::FPutValueOp final : public FThreadSafeRefCountedObject
+class FZenCacheStore::FPutValueOp final : public FRefCountedObject
 {
 public:
 	FPutValueOp(FZenCacheStore& InCacheStore,
@@ -1081,7 +1081,7 @@ private:
 	FOnCachePutValueComplete OnComplete;
 };
 
-class FZenCacheStore::FGetValueOp final : public FThreadSafeRefCountedObject
+class FZenCacheStore::FGetValueOp final : public FRefCountedObject
 {
 public:
 	FGetValueOp(FZenCacheStore& InCacheStore,
@@ -1295,7 +1295,7 @@ private:
 	FOnCacheGetValueComplete OnComplete;
 };
 
-class FZenCacheStore::FGetChunksOp final : public FThreadSafeRefCountedObject
+class FZenCacheStore::FGetChunksOp final : public FRefCountedObject
 {
 public:
 	FGetChunksOp(FZenCacheStore& InCacheStore,
diff --git a/Engine/Source/Editor/UnrealEd/Private/Cooker/CookGenerationHelper.h b/Engine/Source/Editor/UnrealEd/Private/Cooker/CookGenerationHelper.h
--- a/Engine/Source/Editor/UnrealEd/Private/Cooker/CookGenerationHelper.h
+++ b/Engine/Source/Editor/UnrealEd/Private/Cooker/CookGenerationHelper.h
@@ -200,7 +200,7 @@ private:
  * reachability to fully support per-platform customization of the cook graph. See TODO_COOKGENERATIONHELPER for
  * code blocks where we implement the workarounds for this todo.
  */
-struct FGenerationHelper : public FThreadSafeRefCountedObject
+struct FGenerationHelper : public FRefCountedObject
 {
 public:
 	/** Store the provided CookPackageSplitter and prepare the packages to generate. */
diff --git a/Engine/Source/Editor/UnrealEd/Private/Cooker/CookPackagePreloader.h b/Engine/Source/Editor/UnrealEd/Private/Cooker/CookPackagePreloader.h
--- a/Engine/Source/Editor/UnrealEd/Private/Cooker/CookPackagePreloader.h
+++ b/Engine/Source/Editor/UnrealEd/Private/Cooker/CookPackagePreloader.h
@@ -19,7 +19,7 @@ namespace UE::Cook
  * Helper class for the Load state on FPackageData. Also supports loads of other packages that depend on the owner
  * package. Handles preloading or asyncloading the package before we call LoadPackage on it.
  */
-class FPackagePreloader : public FThreadSafeRefCountedObject
+class FPackagePreloader : public FRefCountedObject
 {
 public:
 	FPackagePreloader(FPackageData& InPackageData);
diff --git a/Engine/Source/Editor/UnrealEd/Private/Cooker/CookWorkerServer.h b/Engine/Source/Editor/UnrealEd/Private/Cooker/CookWorkerServer.h
--- a/Engine/Source/Editor/UnrealEd/Private/Cooker/CookWorkerServer.h
+++ b/Engine/Source/Editor/UnrealEd/Private/Cooker/CookWorkerServer.h
@@ -58,7 +58,7 @@ struct FAssignPackageExtraData
 };
 
 /** Class in a Director process that communicates over a Socket with FCookWorkerClient in a CookWorker process. */
-class FCookWorkerServer : public FThreadSafeRefCountedObject
+class FCookWorkerServer : public FRefCountedObject
 {
 public:
 	FCookWorkerServer(FCookDirector& InDirector, int32 InProfileId, FWorkerId InWorkerId);
diff --git a/Engine/Source/Editor/UnrealEd/Private/Cooker/DiffWriterArchive.h b/Engine/Source/Editor/UnrealEd/Private/Cooker/DiffWriterArchive.h
--- a/Engine/Source/Editor/UnrealEd/Private/Cooker/DiffWriterArchive.h
+++ b/Engine/Source/Editor/UnrealEd/Private/Cooker/DiffWriterArchive.h
@@ -133,7 +133,7 @@ class FCallstacks
 {
 public:
 	/** Struct to hold the actual Serialize call callstack and any associated data */
-	struct FCallstackData : public FThreadSafeRefCountedObject
+	struct FCallstackData : public FRefCountedObject
 	{
 		/** Full callstack */
 		TUniquePtr<ANSICHAR[]> Callstack;
@@ -149,10 +149,10 @@ public:
 
 		FCallstackData() = default;
 		FCallstackData(TUniquePtr<ANSICHAR[]>&& InCallstack, uint32 InCppCallstackHash, UObject* InSerializedObject, FProperty* InSerializedProperty);
-		FCallstackData(FCallstackData&&) = delete; // FThreadSafeRefCountedObject doesn't allow it
+		FCallstackData(FCallstackData&&) = delete; // FRefCountedObject doesn't allow it
 		FCallstackData(const FCallstackData&) = delete;
 
-		FCallstackData& operator=(FCallstackData&&) = delete; // FThreadSafeRefCountedObject doesn't allow it
+		FCallstackData& operator=(FCallstackData&&) = delete; // FRefCountedObject doesn't allow it
 		FCallstackData& operator=(const FCallstackData&) = delete;
 
 		/** Converts the callstack and associated data to human readable string */
diff --git a/Engine/Source/Editor/UnrealEd/Private/Cooker/PackageTracker.h b/Engine/Source/Editor/UnrealEd/Private/Cooker/PackageTracker.h
--- a/Engine/Source/Editor/UnrealEd/Private/Cooker/PackageTracker.h
+++ b/Engine/Source/Editor/UnrealEd/Private/Cooker/PackageTracker.h
@@ -176,7 +176,7 @@ private:
 };
 
 /** Container for name -> data for FPackageStreamInstancedPackage held by the PackageTracker. */
-struct FPackageStreamInstancedPackageContainer : public FThreadSafeRefCountedObject
+struct FPackageStreamInstancedPackageContainer : public FRefCountedObject
 {
 	FRWLock Lock;
 	TMap<FName, FPackageStreamInstancedPackage*> Map;
@@ -187,7 +187,7 @@ struct FPackageStreamInstancedPackageContainer : public FThreadSafeRefCountedObj
  * loaded into it. We need to handle the recording of packages imported by instanced package loads in a special manner
  * since their AssetRegistry dependencies will not be found by querying the AssetRegistry for their name.
  */
-struct FPackageStreamInstancedPackage : public FThreadSafeRefCountedObject
+struct FPackageStreamInstancedPackage : public FRefCountedObject
 {
 public:
 	// FPackageStreamInstancedPackages are only createable by FPackageTracker, and they
diff --git a/Engine/Source/Editor/UnrealEd/Private/EditorDomain/BulkDataRegistryImpl.h b/Engine/Source/Editor/UnrealEd/Private/EditorDomain/BulkDataRegistryImpl.h
--- a/Engine/Source/Editor/UnrealEd/Private/EditorDomain/BulkDataRegistryImpl.h
+++ b/Engine/Source/Editor/UnrealEd/Private/EditorDomain/BulkDataRegistryImpl.h
@@ -150,7 +150,7 @@ struct FTempLoadedPayload
 };
 
 /** An Active flag and a lock around it for informing AutoDeleteAsyncTasks that their shared data is no longer available. */
-class FTaskSharedDataLock : public FThreadSafeRefCountedObject
+class FTaskSharedDataLock : public FRefCountedObject
 {
 public:
 	FRWLock ActiveLock;
@@ -188,7 +188,7 @@ struct FUpdatingPayload
 };
 
 /** Data storage for a BulkData that is loading its PayloadId from the cache. */
-class FPendingPayloadId : public FThreadSafeRefCountedObject
+class FPendingPayloadId : public FRefCountedObject
 {
 public:
 	explicit FPendingPayloadId(const FGuid& InBulkDataId);
diff --git a/Engine/Source/Editor/UnrealEd/Private/EditorDomain/EditorDomain.h b/Engine/Source/Editor/UnrealEd/Private/EditorDomain/EditorDomain.h
--- a/Engine/Source/Editor/UnrealEd/Private/EditorDomain/EditorDomain.h
+++ b/Engine/Source/Editor/UnrealEd/Private/EditorDomain/EditorDomain.h
@@ -197,7 +197,7 @@ private:
 	 * Reference-counted struct holding the locks used for multithreaded synchronization.
 	 * Shared with Archives and other helpers that might outlive *this.
 	 */
-	class FLocks : public FThreadSafeRefCountedObject
+	class FLocks : public FRefCountedObject
 	{
 	public:
 		FLocks(FEditorDomain& InOwner);
@@ -208,7 +208,7 @@ private:
 	 * Data about which domain a package comes from. Multiple queries of the same
 	 * package have to align (for e.g. BulkData offsets) so we have to keep track of this.
 	 */
-	struct FPackageSource : public FThreadSafeRefCountedObject
+	struct FPackageSource : public FRefCountedObject
 	{
 		FPackageSource()
 			: bHasSaved(false), bHasLoaded(false), bHasQueriedCatalog(false), bLoadedAfterCatalogLoaded(false)
diff --git a/Engine/Source/Programs/AutoRTFMTests/Private/RefCountPtrTests.cpp b/Engine/Source/Programs/AutoRTFMTests/Private/RefCountPtrTests.cpp
--- a/Engine/Source/Programs/AutoRTFMTests/Private/RefCountPtrTests.cpp
+++ b/Engine/Source/Programs/AutoRTFMTests/Private/RefCountPtrTests.cpp
@@ -19,7 +19,7 @@ namespace
 		int Value;
 	};
 
-	struct ThreadSafeRefCountedInt : public FThreadSafeRefCountedObject
+	struct ThreadSafeRefCountedInt : public FRefCountedObject
 	{
 		ThreadSafeRefCountedInt(int InValue) : Value(InValue) {}
 		int Value;
diff --git a/Engine/Source/Runtime/Apple/MetalRHI/Private/MetalSubmission.h b/Engine/Source/Runtime/Apple/MetalRHI/Private/MetalSubmission.h
--- a/Engine/Source/Runtime/Apple/MetalRHI/Private/MetalSubmission.h
+++ b/Engine/Source/Runtime/Apple/MetalRHI/Private/MetalSubmission.h
@@ -60,7 +60,7 @@ struct FMetalResolvedFence
 // Sync points are one-shot, meaning they represent a single timeline point, and are released after use, via ref-counting.
 // Use FMetalSyncPoint::Create() to make a new sync point and hold a reference to it via a FMetalSyncPointRef object.
 //
-class FMetalSyncPoint final : public FThreadSafeRefCountedObject
+class FMetalSyncPoint final : public FRefCountedObject
 {
 	friend FMetalDynamicRHI;
 	friend FMetalCommandQueue;
diff --git a/Engine/Source/Runtime/Core/Private/Experimental/Misc/ExecutionResource.cpp b/Engine/Source/Runtime/Core/Private/Experimental/Misc/ExecutionResource.cpp
--- a/Engine/Source/Runtime/Core/Private/Experimental/Misc/ExecutionResource.cpp
+++ b/Engine/Source/Runtime/Core/Private/Experimental/Misc/ExecutionResource.cpp
@@ -85,7 +85,7 @@ FExecutionResourceContextScope::~FExecutionResourceContextScope()
 	ExecutionResourceImpl::State.Pop(LowLevelTasks::FTask::GetActiveTask());
 }
 
-class FCompositeExecutionResource : public IExecutionResource, public FThreadSafeRefCountedObject
+class FCompositeExecutionResource : public IExecutionResource, public FRefCountedObject
 {
 public:
 	FCompositeExecutionResource(const TArray<TRefCountPtr<IExecutionResource>>& InExecutionResources)
@@ -95,17 +95,17 @@ public:
 
 	FReturnedRefCountValue AddRef() const override
 	{
-		return FThreadSafeRefCountedObject::AddRef();
+		return FRefCountedObject::AddRef();
 	}
 
 	uint32 Release() const override
 	{
-		return FThreadSafeRefCountedObject::Release();
+		return FRefCountedObject::Release();
 	}
 	
 	uint32 GetRefCount() const override
 	{
-		return FThreadSafeRefCountedObject::GetRefCount();
+		return FRefCountedObject::GetRefCount();
 	}
 
 private:
diff --git a/Engine/Source/Runtime/Core/Public/Async/LocalWorkQueue.h b/Engine/Source/Runtime/Core/Public/Async/LocalWorkQueue.h
--- a/Engine/Source/Runtime/Core/Public/Async/LocalWorkQueue.h
+++ b/Engine/Source/Runtime/Core/Public/Async/LocalWorkQueue.h
@@ -46,7 +46,7 @@ class TLocalWorkQueue
 {
 	UE_NONCOPYABLE(TLocalWorkQueue);
 
-	struct FInternalData : public TConcurrentLinearObject<FInternalData, FTaskGraphBlockAllocationTag>, public FThreadSafeRefCountedObject
+	struct FInternalData : public TConcurrentLinearObject<FInternalData, FTaskGraphBlockAllocationTag>, public FRefCountedObject
 	{
 		FAAArrayQueue<TaskType> TaskQueue;
 		std::atomic_int         ActiveWorkers {0};
diff --git a/Engine/Source/Runtime/Core/Public/Async/ParallelFor.h b/Engine/Source/Runtime/Core/Public/Async/ParallelFor.h
--- a/Engine/Source/Runtime/Core/Public/Async/ParallelFor.h
+++ b/Engine/Source/Runtime/Core/Public/Async/ParallelFor.h
@@ -208,7 +208,7 @@ namespace ParallelForImpl
 		//shared data between tasks
 		struct alignas(PLATFORM_CACHE_LINE_SIZE) FParallelForData 
 			: public TConcurrentLinearObject<FParallelForData, FTaskGraphBlockAllocationTag>
-			, public FThreadSafeRefCountedObject
+			, public FRefCountedObject
 			, private UE::FInheritedContextBase
 		{
 			using UE::FInheritedContextBase::RestoreInheritedContext;
diff --git a/Engine/Source/Runtime/Core/Public/Misc/IQueuedWork.h b/Engine/Source/Runtime/Core/Public/Misc/IQueuedWork.h
--- a/Engine/Source/Runtime/Core/Public/Misc/IQueuedWork.h
+++ b/Engine/Source/Runtime/Core/Public/Misc/IQueuedWork.h
@@ -31,7 +31,7 @@ ENUM_CLASS_FLAGS(EQueuedWorkFlags);
 * This interface can be used to track some data between the individual function invokations of the FQueuedThreadPool
 * Usually it is used to store some internal state to support cancellation without having to look it up from a map.
 */
-class IQueuedWorkInternalData : public FThreadSafeRefCountedObject
+class IQueuedWorkInternalData : public FRefCountedObject
 {
 public:
 	/**
diff --git a/Engine/Source/Runtime/Core/Public/Serialization/PackageWriter.h b/Engine/Source/Runtime/Core/Public/Serialization/PackageWriter.h
--- a/Engine/Source/Runtime/Core/Public/Serialization/PackageWriter.h
+++ b/Engine/Source/Runtime/Core/Public/Serialization/PackageWriter.h
@@ -235,7 +235,7 @@ public:
 ENUM_CLASS_FLAGS(IPackageWriter::EWriteOptions);
 
 /** Struct containing hashes computed during cooked package writing. */
-struct FPackageHashes : FThreadSafeRefCountedObject
+struct FPackageHashes : FRefCountedObject
 {
 	/** Hashes for each chunk saved by the package. */
 	TMap<FIoChunkId, FIoHash> ChunkHashes;
diff --git a/Engine/Source/Runtime/Core/Public/Templates/RefCounting.h b/Engine/Source/Runtime/Core/Public/Templates/RefCounting.h
--- a/Engine/Source/Runtime/Core/Public/Templates/RefCounting.h
+++ b/Engine/Source/Runtime/Core/Public/Templates/RefCounting.h
@@ -58,7 +58,7 @@ private:
 
 /**
  * UE::Private::TTransactionalAtomicRefCount manages a transactionally-safe atomic refcount value.
- * This is used by FRefCountBase, FThreadSafeRefCountedObject and TRefCountingMixin (in thread-safe mode).
+ * This is used by FRefCountedObject and TRefCountingMixin (in thread-safe mode).
  */
 namespace UE::Private
 {
@@ -246,9 +246,7 @@ private:
 // some uses of this as a friend class declaration.
 class FRefCountBase : public FRefCountedObject {};
 
-// TODO (SOL-7350): deprecate and remove this entirely. This is a class because there are
-// some uses of this as a friend class declaration.
-class FThreadSafeRefCountedObject : public FRefCountedObject {};
+class UE_DEPRECATED(5.8, "`FRefCountedObject` is now thread-safe so use that instead.") FThreadSafeRefCountedObject : public FRefCountedObject {};
 
 /**
  * ERefCountingMode is used select between either 'fast' or 'thread safe' ref-counting types.
diff --git a/Engine/Source/Runtime/D3D12RHI/Private/D3D12Descriptors.h b/Engine/Source/Runtime/D3D12RHI/Private/D3D12Descriptors.h
--- a/Engine/Source/Runtime/D3D12RHI/Private/D3D12Descriptors.h
+++ b/Engine/Source/Runtime/D3D12RHI/Private/D3D12Descriptors.h
@@ -91,7 +91,7 @@ namespace UE::D3D12Descriptors
 	void FreeOfflineCopy(FD3D12Device* Device, FD3D12OfflineDescriptor& InDescriptor, ERHIDescriptorHeapType InType);
 }
 
-struct FD3D12DescriptorHeap : public FD3D12DeviceChild, public FThreadSafeRefCountedObject
+struct FD3D12DescriptorHeap : public FD3D12DeviceChild, public FRefCountedObject
 {
 public:
 	FD3D12DescriptorHeap() = delete;
diff --git a/Engine/Source/Runtime/D3D12RHI/Private/D3D12Resources.h b/Engine/Source/Runtime/D3D12RHI/Private/D3D12Resources.h
--- a/Engine/Source/Runtime/D3D12RHI/Private/D3D12Resources.h
+++ b/Engine/Source/Runtime/D3D12RHI/Private/D3D12Resources.h
@@ -68,7 +68,7 @@ public:
 	{}
 };
 
-class FD3D12Heap : public FThreadSafeRefCountedObject, public FD3D12DeviceChild, public FD3D12MultiNodeGPUObject
+class FD3D12Heap : public FRefCountedObject, public FD3D12DeviceChild, public FD3D12MultiNodeGPUObject
 {
 public:
 	FD3D12Heap(FD3D12Device* Parent, FRHIGPUMask VisibleNodes, HeapId InTraceParentHeapId = EMemoryTraceRootHeap::VideoMemory);
@@ -176,7 +176,7 @@ struct FD3D12ResourceDesc : public TZeroedStruct<FD3D12ResourceDesc>, D3D12_RESO
 #endif
 };
 
-class FD3D12Resource : public FThreadSafeRefCountedObject, public FD3D12DeviceChild, public FD3D12MultiNodeGPUObject
+class FD3D12Resource : public FRefCountedObject, public FD3D12DeviceChild, public FD3D12MultiNodeGPUObject
 {
 	friend class FD3D12Buffer;
 private:
diff --git a/Engine/Source/Runtime/D3D12RHI/Private/D3D12Submission.h b/Engine/Source/Runtime/D3D12RHI/Private/D3D12Submission.h
--- a/Engine/Source/Runtime/D3D12RHI/Private/D3D12Submission.h
+++ b/Engine/Source/Runtime/D3D12RHI/Private/D3D12Submission.h
@@ -73,7 +73,7 @@ struct FD3D12ResolvedFence
 // Sync points are one-shot, meaning they represent a single timeline point, and are released after use, via ref-counting.
 // Use FD3D12SyncPoint::Create() to make a new sync point and hold a reference to it via a FD3D12SyncPointRef object.
 //
-class FD3D12SyncPoint final : public FThreadSafeRefCountedObject
+class FD3D12SyncPoint final : public FRefCountedObject
 {
 	friend FD3D12DynamicRHI;
 	friend FD3D12Queue;
diff --git a/Engine/Source/Runtime/Engine/Classes/Components/SceneCaptureComponent.h b/Engine/Source/Runtime/Engine/Classes/Components/SceneCaptureComponent.h
--- a/Engine/Source/Runtime/Engine/Classes/Components/SceneCaptureComponent.h
+++ b/Engine/Source/Runtime/Engine/Classes/Components/SceneCaptureComponent.h
@@ -32,7 +32,7 @@ struct FSceneCaptureViewInfo : public FSceneViewProjectionData
 
 #if WITH_EDITORONLY_DATA
 /** Editor only structure for gathering memory size */
-struct FSceneCaptureMemorySize : public FThreadSafeRefCountedObject
+struct FSceneCaptureMemorySize : public FRefCountedObject
 {
 	uint64 Size = 0;
 };
diff --git a/Engine/Source/Runtime/Engine/Private/SkeletalRenderGPUSkin.h b/Engine/Source/Runtime/Engine/Private/SkeletalRenderGPUSkin.h
--- a/Engine/Source/Runtime/Engine/Private/SkeletalRenderGPUSkin.h
+++ b/Engine/Source/Runtime/Engine/Private/SkeletalRenderGPUSkin.h
@@ -317,7 +317,7 @@ private:
 /**
 * Pooled morph vertex buffers that store the vertex deltas.
 */
-class FMorphVertexBufferPool : public FThreadSafeRefCountedObject
+class FMorphVertexBufferPool : public FRefCountedObject
 {
 public:
 	FMorphVertexBufferPool(FSkeletalMeshRenderData* InSkelMeshRenderData, int32 InLOD, ERHIFeatureLevel::Type InFeatureLevel)
diff --git a/Engine/Source/Runtime/ImageCore/Public/ImageCore.h b/Engine/Source/Runtime/ImageCore/Public/ImageCore.h
--- a/Engine/Source/Runtime/ImageCore/Public/ImageCore.h
+++ b/Engine/Source/Runtime/ImageCore/Public/ImageCore.h
@@ -742,7 +742,7 @@ public:
 
 typedef TRefCountPtr<struct FSharedImage> FSharedImageRef;
 typedef TRefCountPtr<const struct FSharedImage> FSharedImageConstRef;
-struct FSharedImage : public FImage, public FThreadSafeRefCountedObject
+struct FSharedImage : public FImage, public FRefCountedObject
 {
 	FSharedImage() = default;
 	virtual ~FSharedImage() = default;
diff --git a/Engine/Source/Runtime/Landscape/Classes/LandscapeHeightfieldCollisionComponent.h b/Engine/Source/Runtime/Landscape/Classes/LandscapeHeightfieldCollisionComponent.h
--- a/Engine/Source/Runtime/Landscape/Classes/LandscapeHeightfieldCollisionComponent.h
+++ b/Engine/Source/Runtime/Landscape/Classes/LandscapeHeightfieldCollisionComponent.h
@@ -99,7 +99,7 @@ public:
 	UFUNCTION(BlueprintCallable, Category = "Landscape")
 	LANDSCAPE_API ULandscapeComponent* GetRenderComponent() const;
 
-	struct FHeightfieldGeometryRef : public FThreadSafeRefCountedObject
+	struct FHeightfieldGeometryRef : public FRefCountedObject
 	{
 		FGuid Guid;
 
diff --git a/Engine/Source/Runtime/OpenGLDrv/Private/OpenGLProgramBinaryFileCache.cpp b/Engine/Source/Runtime/OpenGLDrv/Private/OpenGLProgramBinaryFileCache.cpp
--- a/Engine/Source/Runtime/OpenGLDrv/Private/OpenGLProgramBinaryFileCache.cpp
+++ b/Engine/Source/Runtime/OpenGLDrv/Private/OpenGLProgramBinaryFileCache.cpp
@@ -165,7 +165,7 @@ namespace UE
 
 // This contains the mapping for a binary program cache file.
 // It also contains a list of programs that the cache contains.
-class FOpenGLProgramBinaryMapping : public FThreadSafeRefCountedObject
+class FOpenGLProgramBinaryMapping : public FRefCountedObject
 {
 public:
 	FOpenGLProgramBinaryMapping(TUniquePtr<IMappedFileHandle> MappedCacheFileIn, uint32 ProgramCountIfKnown) : MappedCacheFile(MoveTemp(MappedCacheFileIn)) { Content.Reserve(ProgramCountIfKnown); }
diff --git a/Engine/Source/Runtime/RenderCore/Public/Shader.h b/Engine/Source/Runtime/RenderCore/Public/Shader.h
--- a/Engine/Source/Runtime/RenderCore/Public/Shader.h
+++ b/Engine/Source/Runtime/RenderCore/Public/Shader.h
@@ -517,7 +517,7 @@ private:
 	std::atomic<int32> NumRefs;
 };
 
-class FShaderMapResourceCode : public FThreadSafeRefCountedObject
+class FShaderMapResourceCode : public FRefCountedObject
 {
 public:
 #if WITH_EDITORONLY_DATA
diff --git a/Engine/Source/Runtime/Renderer/Private/ScenePrivate.h b/Engine/Source/Runtime/Renderer/Private/ScenePrivate.h
--- a/Engine/Source/Runtime/Renderer/Private/ScenePrivate.h
+++ b/Engine/Source/Runtime/Renderer/Private/ScenePrivate.h
@@ -644,7 +644,7 @@ struct FGlobalDistanceFieldStreamingReadback
 	uint32 ReadbackBuffersNumPending = 0;
 };
 
-struct FPersistentGlobalDistanceFieldData : public FThreadSafeRefCountedObject
+struct FPersistentGlobalDistanceFieldData : public FRefCountedObject
 {
 	// Array of ClipmapIndex
 	TArray<int32> DeferredUpdates[GDF_Num];
diff --git a/Engine/Source/Runtime/VulkanRHI/Private/VulkanMemory.h b/Engine/Source/Runtime/VulkanRHI/Private/VulkanMemory.h
--- a/Engine/Source/Runtime/VulkanRHI/Private/VulkanMemory.h
+++ b/Engine/Source/Runtime/VulkanRHI/Private/VulkanMemory.h
@@ -868,7 +868,7 @@ namespace VulkanRHI
 		void ProcessPendingUBFrees(bool bForce);
 	};
 
-	class FStagingBuffer : public FThreadSafeRefCountedObject
+	class FStagingBuffer : public FRefCountedObject
 	{
 	public:
 		FStagingBuffer(FVulkanDevice* InDevice);
diff --git a/Engine/Source/Runtime/VulkanRHI/Private/VulkanResources.h b/Engine/Source/Runtime/VulkanRHI/Private/VulkanResources.h
--- a/Engine/Source/Runtime/VulkanRHI/Private/VulkanResources.h
+++ b/Engine/Source/Runtime/VulkanRHI/Private/VulkanResources.h
@@ -90,7 +90,7 @@ public:
 
 struct FGfxPipelineDesc;
 
-class FVulkanShaderModule : public FThreadSafeRefCountedObject
+class FVulkanShaderModule : public FRefCountedObject
 {
 	static FVulkanDevice* Device;
 	VkShaderModule ActualShaderModule;
diff --git a/Engine/Source/Runtime/VulkanRHI/Private/VulkanSubmission.h b/Engine/Source/Runtime/VulkanRHI/Private/VulkanSubmission.h
--- a/Engine/Source/Runtime/VulkanRHI/Private/VulkanSubmission.h
+++ b/Engine/Source/Runtime/VulkanRHI/Private/VulkanSubmission.h
@@ -50,7 +50,7 @@ using FVulkanSyncPointRef = TRefCountPtr<FVulkanSyncPoint>;
 // Sync points are one-shot, meaning they represent a single timeline point, and are released after use, via ref-counting.
 // Use FVulkanSyncPoint::Create() to make a new sync point and hold a reference to it via a FVulkanSyncPointRef object.
 //
-class FVulkanSyncPoint final : public FThreadSafeRefCountedObject
+class FVulkanSyncPoint final : public FRefCountedObject
 {
 	friend FVulkanDynamicRHI;
 	friend FVulkanQueue;
diff --git a/Engine/Source/Runtime/VulkanRHI/Private/VulkanSynchronization.h b/Engine/Source/Runtime/VulkanRHI/Private/VulkanSynchronization.h
--- a/Engine/Source/Runtime/VulkanRHI/Private/VulkanSynchronization.h
+++ b/Engine/Source/Runtime/VulkanRHI/Private/VulkanSynchronization.h
@@ -123,7 +123,7 @@ ENUM_CLASS_FLAGS(EVulkanSemaphoreFlags);
 
 
 // Internal Vulkan Semaphore class that wrap VkSemaphore management (binary or timeline)
-class FVulkanSemaphore : public FThreadSafeRefCountedObject
+class FVulkanSemaphore : public FRefCountedObject
 {
 public:
 	FVulkanSemaphore(FVulkanDevice& InDevice, EVulkanSemaphoreFlags InFlags=EVulkanSemaphoreFlags::None, uint64 InInitialTimelineValue=0);
