DDC: Validate RawHash in compressed buffer headers to mitigate invalid data being returned by Zen

diff --git a/Engine/Source/Developer/DerivedDataCache/Private/DerivedDataCacheRecord.cpp b/Engine/Source/Developer/DerivedDataCache/Private/DerivedDataCacheRecord.cpp
--- a/Engine/Source/Developer/DerivedDataCache/Private/DerivedDataCacheRecord.cpp
+++ b/Engine/Source/Developer/DerivedDataCache/Private/DerivedDataCacheRecord.cpp
@@ -226,7 +226,8 @@ FOptionalCacheRecord FCacheRecord::Load(const FCbPackage& Attachments, const FCb
 		const FIoHash RawHash = ValueObject[ANSITEXTVIEW("RawHash")].AsHash();
 		if (const FCbAttachment* Attachment = Attachments.FindAttachment(RawHash))
 		{
-			if (const FCompressedBuffer& Compressed = Attachment->AsCompressedBinary())
+			const FCompressedBuffer& Compressed = Attachment->AsCompressedBinary();
+			if (Compressed.GetRawHash() == RawHash)
 			{
 				return FValueWithId(Id, Compressed);
 			}
diff --git a/Engine/Source/Developer/DerivedDataCache/Private/ZenCacheStore.cpp b/Engine/Source/Developer/DerivedDataCache/Private/ZenCacheStore.cpp
--- a/Engine/Source/Developer/DerivedDataCache/Private/ZenCacheStore.cpp
+++ b/Engine/Source/Developer/DerivedDataCache/Private/ZenCacheStore.cpp
@@ -1195,22 +1195,24 @@ public:
 							UE_LOG(LogDerivedDataCache, Verbose, TEXT("%s: Simulated miss for GetValue of '%s' from '%s'"),
 								*CacheStore.GetName(), *WriteToString<96>(RequestWithStats.Request.Key), *RequestWithStats.Request.Name);
 						}
-						else
+						else if (FIoHash RawHash = ResultObj[ANSITEXTVIEW("RawHash")].AsHash(); !RawHash.IsZero())
 						{
-							FCbFieldView RawHashField = ResultObj["RawHash"];
-							FIoHash RawHash = RawHashField.AsHash();
-							const FCbAttachment* Attachment = EnumHasAnyFlags(RequestWithStats.Request.Policy, ECachePolicy::SkipData) ? nullptr : Response.FindAttachment(RawHash);
-							if (Attachment)
+							if (EnumHasAnyFlags(RequestWithStats.Request.Policy, ECachePolicy::SkipData))
 							{
-								Value.Emplace(Attachment->AsCompressedBinary());
+								if (uint64 RawSize = ResultObj[ANSITEXTVIEW("RawSize")].AsUInt64(MAX_uint64); RawSize != MAX_uint64)
+								{
+									Value.Emplace(RawHash, RawSize);
+								}
 							}
 							else
 							{
-								FCbFieldView RawSizeField = ResultObj["RawSize"];
-								uint64 RawSize = RawSizeField.AsUInt64();
-								if (!RawSizeField.HasError() && !RawHashField.HasError())
+								if (const FCbAttachment* Attachment = Response.FindAttachment(RawHash))
 								{
-									Value.Emplace(RawHash, RawSize);
+									const FCompressedBuffer& CompressedData = Attachment->AsCompressedBinary();
+									if (CompressedData.GetRawHash() == RawHash)
+									{
+										Value.Emplace(CompressedData);
+									}
 								}
 							}
 						}
