DDC: Added cache store tests

#Synced-CL 32500503

diff --git a/Engine/Source/Developer/DerivedDataCache/Tests/DerivedDataCacheStoreTest.cpp b/Engine/Source/Developer/DerivedDataCache/Tests/DerivedDataCacheStoreTest.cpp
new file mode 100644
--- /dev/null
+++ b/Engine/Source/Developer/DerivedDataCache/Tests/DerivedDataCacheStoreTest.cpp
@@ -0,0 +1,1006 @@
+// Copyright Epic Games, Inc. All Rights Reserved.
+
+#if WITH_LOW_LEVEL_TESTS
+
+#include "DerivedDataCache.h"
+#include "DerivedDataRequestOwner.h"
+#include "HAL/FileManager.h"
+#include "Misc/Guid.h"
+#include "Misc/Paths.h"
+#include "Misc/StringBuilder.h"
+#include "Serialization/CompactBinary.h"
+#include "Serialization/CompactBinaryWriter.h"
+
+#include "TestHarness.h"
+
+namespace UE::DerivedData::CacheStoreTest
+{
+
+static FCompressedBuffer MakeSequentialIntegerBuffer(const uint32 ValueCount)
+{
+	FUniqueBuffer Buffer = FUniqueBuffer::Alloc(ValueCount * sizeof(uint32));
+	uint32* Values = static_cast<uint32*>(Buffer.GetData());
+	for (uint32 ValueIndex = 0; ValueIndex < ValueCount; ++ValueIndex)
+	{
+		Values[ValueIndex] = ValueIndex;
+	}
+	return FCompressedBuffer::Compress(Buffer.MoveToShared());
+}
+
+static bool CheckSequentialIntegerBufferRange(FMemoryView Data, const uint64 Offset)
+{
+	bool bOk = true;
+	const uint32* Values = static_cast<const uint32*>(Data.GetData());
+	const uint32 ValueBase = uint32(Offset / sizeof(uint32));
+	const uint32 ValueCount = uint32(Data.GetSize() / sizeof(uint32));
+	for (uint32 ValueIndex = 0; ValueIndex < ValueCount; ++ValueIndex)
+	{
+		bOk &= Values[ValueIndex] == ValueIndex + ValueBase;
+	}
+	return bOk;
+}
+
+static bool HasValueWithData(const FCacheRecord& Record, const FValueWithId& Value)
+{
+	const FValueWithId& ValueInRecord = Record.GetValue(Value.GetId());
+	return ValueInRecord == Value &&
+		ValueInRecord.GetData().GetRawHash() == Value.GetRawHash() &&
+		ValueInRecord.GetData().GetRawSize() == Value.GetRawSize();
+}
+
+static bool HasValueWithNoData(const FCacheRecord& Record, const FValueWithId& Value)
+{
+	const FValueWithId& ValueInRecord = Record.GetValue(Value.GetId());
+	return ValueInRecord == Value && !ValueInRecord.HasData();
+}
+
+static FCacheBucket GetBucket()
+{
+	static const FCacheBucket Bucket(ANSITEXTVIEW("Test"));
+	return Bucket;
+}
+
+static FCacheKey CreateKeyFromGuid(const ANSICHAR* GuidStr, int32 Index = 0)
+{
+	const FGuid Guid(GuidStr);
+	FIoHashBuilder HashBuilder;
+	HashBuilder.Update(&Guid, sizeof(FGuid));
+	HashBuilder.Update(&Index, sizeof(int32));
+	return {GetBucket(), HashBuilder.Finalize()};
+}
+
+static FCacheKey GetMissingKey()
+{
+	static const FCacheKey Key = CreateKeyFromGuid("c618cf5a-ac13-42d5-8c1f-6fefbc485b65");
+	return Key;
+}
+
+static FValueId GetMissingId()
+{
+	static const FValueId Id = FValueId::FromName(ANSITEXTVIEW("MissingId"));
+	return Id;
+}
+
+static FCbObject CreateMeta(const ANSICHAR* Value)
+{
+	TCbWriter<64> Writer;
+	Writer.BeginObject();
+	Writer.AddString(ANSITEXTVIEW("Key"), Value);
+	Writer.EndObject();
+	return Writer.Save().AsObject();
+};
+
+static void TestCacheChunks(ICache& Cache, const FCacheKey& Key, const FValueId ValueId, const FValue& Value)
+{
+	const uint64 UserData = 0xfedcba9876543210;
+	const FSharedString Name(TEXTVIEW("Chunks"));
+
+	const FIoHash ValueHash = Value.GetRawHash();
+	const FIoHash DifferentHash = FIoHash::HashBuffer(&ValueHash, sizeof(ValueHash));
+
+	{
+		INFO("Get Value Hash + Size");
+		FRequestOwner BlockingOwner(EPriority::Blocking);
+		Cache.GetChunks({{.Name = Name, .Key = Key, .Id = ValueId, .Policy = ECachePolicy::Default | ECachePolicy::SkipData, .UserData = UserData}}, BlockingOwner, [&](FCacheGetChunkResponse&& Response)
+		{
+			CHECK(Response.Status == EStatus::Ok);
+			CHECK(Response.UserData == UserData);
+			CHECK(Response.Name == Name);
+			CHECK(Response.Key == Key);
+			CHECK(Response.Id == ValueId);
+			CHECK(Response.RawOffset == 0);
+			CHECK(Response.RawSize == Value.GetRawSize());
+			CHECK(Response.RawHash == Value.GetRawHash());
+			CHECK(Response.RawData.IsNull());
+		});
+		BlockingOwner.Wait();
+	}
+
+	{
+		INFO("Get Value");
+		FRequestOwner BlockingOwner(EPriority::Blocking);
+		Cache.GetChunks({{.Name = Name, .Key = Key, .Id = ValueId, .UserData = UserData}}, BlockingOwner, [&](FCacheGetChunkResponse&& Response)
+		{
+			CHECK(Response.Status == EStatus::Ok);
+			CHECK(Response.UserData == UserData);
+			CHECK(Response.Name == Name);
+			CHECK(Response.Key == Key);
+			CHECK(Response.Id == ValueId);
+			CHECK(Response.RawOffset == 0);
+			CHECK(Response.RawSize == Value.GetRawSize());
+			CHECK(Response.RawHash == Value.GetRawHash());
+			CHECK(Response.RawData.GetSize() == Response.RawSize);
+			CHECK(FIoHash::HashBuffer(Response.RawData) == Response.RawHash);
+		});
+		BlockingOwner.Wait();
+	}
+
+	{
+		INFO("Get Value w/ RawHash");
+		FRequestOwner BlockingOwner(EPriority::Blocking);
+		Cache.GetChunks({{.Name = Name, .Key = Key, .Id = ValueId, .RawHash = ValueHash, .UserData = UserData}}, BlockingOwner, [&](FCacheGetChunkResponse&& Response)
+		{
+			CHECK(Response.Status == EStatus::Ok);
+			CHECK(Response.UserData == UserData);
+			CHECK(Response.Name == Name);
+			CHECK(Response.Key == Key);
+			CHECK(Response.Id == ValueId);
+			CHECK(Response.RawOffset == 0);
+			CHECK(Response.RawSize == Value.GetRawSize());
+			CHECK(Response.RawHash == Value.GetRawHash());
+			CHECK(Response.RawData.GetSize() == Response.RawSize);
+			CHECK(FIoHash::HashBuffer(Response.RawData) == Response.RawHash);
+		});
+		BlockingOwner.Wait();
+	}
+
+	{
+		INFO("Get Value w/ Wrong RawHash");
+		FRequestOwner BlockingOwner(EPriority::Blocking);
+		Cache.GetChunks({{.Name = Name, .Key = Key, .Id = ValueId, .RawHash = DifferentHash, .UserData = UserData}}, BlockingOwner, [&](FCacheGetChunkResponse&& Response)
+		{
+			CHECK(Response.Status == EStatus::Error);
+			CHECK(Response.UserData == UserData);
+			CHECK(Response.Name == Name);
+			CHECK(Response.Key == Key);
+			CHECK(Response.Id == ValueId);
+			CHECK(Response.RawOffset == 0);
+			CHECK(Response.RawSize == 0);
+			CHECK(Response.RawHash.IsZero());
+			CHECK(Response.RawData.IsNull());
+		});
+		BlockingOwner.Wait();
+	}
+
+	{
+		INFO("Get Missing Key");
+		FRequestOwner BlockingOwner(EPriority::Blocking);
+		Cache.GetChunks({{.Name = Name, .Key = GetMissingKey(), .Id = ValueId, .UserData = UserData}}, BlockingOwner, [&](FCacheGetChunkResponse&& Response)
+		{
+			CHECK(Response.Status == EStatus::Error);
+			CHECK(Response.UserData == UserData);
+			CHECK(Response.Name == Name);
+			CHECK(Response.Key == GetMissingKey());
+			CHECK(Response.Id == ValueId);
+			CHECK(Response.RawOffset == 0);
+			CHECK(Response.RawSize == 0);
+			CHECK(Response.RawHash.IsZero());
+			CHECK(Response.RawData.IsNull());
+		});
+		BlockingOwner.Wait();
+	}
+
+	if (ValueId.IsValid())
+	{
+		INFO("Get Missing Value");
+		FRequestOwner BlockingOwner(EPriority::Blocking);
+		Cache.GetChunks({{.Name = Name, .Key = Key, .Id = GetMissingId(), .UserData = UserData}}, BlockingOwner, [&](FCacheGetChunkResponse&& Response)
+		{
+			CHECK(Response.Status == EStatus::Error);
+			CHECK(Response.UserData == UserData);
+			CHECK(Response.Name == Name);
+			CHECK(Response.Key == Key);
+			CHECK(Response.Id == GetMissingId());
+			CHECK(Response.RawOffset == 0);
+			CHECK(Response.RawSize == 0);
+			CHECK(Response.RawHash.IsZero());
+			CHECK(Response.RawData.IsNull());
+		});
+		BlockingOwner.Wait();
+	}
+
+	if (Value.GetRawSize() < 2 * 1024 * 1024)
+	{
+		return;
+	}
+
+	{
+		INFO("Get Partial Value Hash + Size");
+		FRequestOwner BlockingOwner(EPriority::Blocking);
+		Cache.GetChunks({{.Name = Name, .Key = Key, .Id = ValueId, .RawOffset = 256 * 1024, .RawSize = 512 * 1024,
+			.Policy = ECachePolicy::Default | ECachePolicy::SkipData, .UserData = UserData}}, BlockingOwner, [&](FCacheGetChunkResponse&& Response)
+		{
+			CHECK(Response.Status == EStatus::Ok);
+			CHECK(Response.UserData == UserData);
+			CHECK(Response.Name == Name);
+			CHECK(Response.Key == Key);
+			CHECK(Response.Id == ValueId);
+			CHECK(Response.RawOffset == 256 * 1024);
+			CHECK(Response.RawSize == 512 * 1024);
+			CHECK(Response.RawHash == Value.GetRawHash());
+			CHECK(Response.RawData.IsNull());
+		});
+		BlockingOwner.Wait();
+	}
+
+	{
+		INFO("Get Partial Value Hash + Size w/ RawHash");
+		FRequestOwner BlockingOwner(EPriority::Blocking);
+		Cache.GetChunks({{.Name = Name, .Key = Key, .Id = ValueId, .RawOffset = 256 * 1024, .RawSize = 512 * 1024,
+			.RawHash = ValueHash, .Policy = ECachePolicy::Default | ECachePolicy::SkipData, .UserData = UserData}},
+			BlockingOwner, [&](FCacheGetChunkResponse&& Response)
+		{
+			CHECK(Response.Status == EStatus::Ok);
+			CHECK(Response.UserData == UserData);
+			CHECK(Response.Name == Name);
+			CHECK(Response.Key == Key);
+			CHECK(Response.Id == ValueId);
+			CHECK(Response.RawOffset == 256 * 1024);
+			CHECK(Response.RawSize == 512 * 1024);
+			CHECK(Response.RawHash == Value.GetRawHash());
+			CHECK(Response.RawData.IsNull());
+		});
+		BlockingOwner.Wait();
+	}
+
+	{
+		INFO("Get Partial Value Hash + Size w/ Wrong RawHash");
+		FRequestOwner BlockingOwner(EPriority::Blocking);
+		Cache.GetChunks({{.Name = Name, .Key = Key, .Id = ValueId, .RawOffset = 256 * 1024, .RawSize = 512 * 1024,
+			.RawHash = DifferentHash, .Policy = ECachePolicy::Default | ECachePolicy::SkipData, .UserData = UserData}},
+			BlockingOwner, [&](FCacheGetChunkResponse&& Response)
+		{
+			CHECK(Response.Status == EStatus::Error);
+			CHECK(Response.UserData == UserData);
+			CHECK(Response.Name == Name);
+			CHECK(Response.Key == Key);
+			CHECK(Response.Id == ValueId);
+			CHECK(Response.RawOffset == 256 * 1024);
+			CHECK(Response.RawSize == 0);
+			CHECK(Response.RawHash.IsZero());
+			CHECK(Response.RawData.IsNull());
+		});
+		BlockingOwner.Wait();
+	}
+
+	{
+		INFO("Get Partial Value");
+		FRequestOwner BlockingOwner(EPriority::Blocking);
+		Cache.GetChunks({{.Name = Name, .Key = Key, .Id = ValueId, .RawOffset = 256 * 1024, .RawSize = 512 * 1024, .UserData = UserData}},
+			BlockingOwner, [&](FCacheGetChunkResponse&& Response)
+		{
+			CHECK(Response.Status == EStatus::Ok);
+			CHECK(Response.UserData == UserData);
+			CHECK(Response.Name == Name);
+			CHECK(Response.Key == Key);
+			CHECK(Response.Id == ValueId);
+			CHECK(Response.RawOffset == 256 * 1024);
+			CHECK(Response.RawSize == 512 * 1024);
+			CHECK(Response.RawHash == Value.GetRawHash());
+			CHECK(Response.RawData.GetSize() == Response.RawSize);
+			CHECK(CheckSequentialIntegerBufferRange(Response.RawData, Response.RawOffset));
+		});
+		BlockingOwner.Wait();
+	}
+
+	{
+		INFO("Get Partial Value w/ RawHash");
+		FRequestOwner BlockingOwner(EPriority::Blocking);
+		Cache.GetChunks({{.Name = Name, .Key = Key, .Id = ValueId, .RawOffset = 256 * 1024, .RawSize = 512 * 1024, .RawHash = ValueHash, .UserData = UserData}},
+			BlockingOwner, [&](FCacheGetChunkResponse&& Response)
+		{
+			CHECK(Response.Status == EStatus::Ok);
+			CHECK(Response.UserData == UserData);
+			CHECK(Response.Name == Name);
+			CHECK(Response.Key == Key);
+			CHECK(Response.Id == ValueId);
+			CHECK(Response.RawOffset == 256 * 1024);
+			CHECK(Response.RawSize == 512 * 1024);
+			CHECK(Response.RawHash == Value.GetRawHash());
+			CHECK(Response.RawData.GetSize() == Response.RawSize);
+			CHECK(CheckSequentialIntegerBufferRange(Response.RawData, Response.RawOffset));
+		});
+		BlockingOwner.Wait();
+	}
+
+	// Disabled until Zen is fixed...
+	//{
+	//	INFO("Get Partial Value Out of Bounds");
+	//	FRequestOwner BlockingOwner(EPriority::Blocking);
+	//	Cache.GetChunks({{.Name = Name, .Key = Key, .Id = ValueId, .RawOffset = 1024 * 1024 * 1024, .RawSize = 512 * 1024, .UserData = UserData}}, BlockingOwner, [&](FCacheGetChunkResponse&& Response)
+	//	{
+	//		CHECK(Response.Status == EStatus::Ok);
+	//		CHECK(Response.UserData == UserData);
+	//		CHECK(Response.Name == Name);
+	//		CHECK(Response.Key == Key);
+	//		CHECK(Response.Id == ValueId);
+	//		CHECK(Response.RawOffset == 1024 * 1024 * 1024);
+	//		CHECK(Response.RawSize == 0);
+	//		CHECK(Response.RawHash == Value.GetRawHash());
+	//		CHECK(Response.RawData);
+	//		CHECK(Response.RawData.GetSize() == 0);
+	//	});
+	//	BlockingOwner.Wait();
+	//}
+
+	{
+		INFO("Get Multiple Partial Values");
+		FRequestOwner BlockingOwner(EPriority::Blocking);
+		Cache.GetChunks(
+			{
+				{.Name = Name, .Key = Key, .Id = ValueId, .RawOffset = (1 * 256 + 128) * 1024, .RawSize = 512 * 1024, .UserData = UserData},
+				{.Name = Name, .Key = Key, .Id = ValueId, .RawOffset = (3 * 256 + 128) * 1024, .RawSize = 512 * 1024, .UserData = UserData + 1},
+			}, BlockingOwner, [&](FCacheGetChunkResponse&& Response)
+		{
+			CHECK(Response.Status == EStatus::Ok);
+			CHECK((Response.UserData & ~uint64(1)) == UserData);
+			CHECK(Response.Name == Name);
+			CHECK(Response.Key == Key);
+			CHECK(Response.Id == ValueId);
+			CHECK(Response.RawOffset == ((1 + 2 * (Response.UserData & 1)) * 256 + 128) * 1024);
+			CHECK(Response.RawSize == 512 * 1024);
+			CHECK(Response.RawHash == Value.GetRawHash());
+			CHECK(Response.RawData.GetSize() == Response.RawSize);
+			CHECK(CheckSequentialIntegerBufferRange(Response.RawData, Response.RawOffset));
+		});
+		BlockingOwner.Wait();
+	}
+}
+
+static void TestCacheValues(ICache& Cache, const uint32 ValueCount = 4 * 1024 * 1024)
+{
+	const FValue Value(MakeSequentialIntegerBuffer(ValueCount));
+
+	const uint64 UserData = 0xfedcba9876543210;
+	const FSharedString Name(TEXTVIEW("Values"));
+	const FCacheKey Key = CreateKeyFromGuid("091714ca-5136-438d-a2c8-0304cc6c54e6", ValueCount);
+
+	{
+		INFO("Put Value");
+		FRequestOwner BlockingOwner(EPriority::Blocking);
+		Cache.PutValue({{Name, Key, Value, ECachePolicy::Default, UserData}}, BlockingOwner, [&](FCachePutValueResponse&& Response)
+		{
+			CHECK(Response.Status == EStatus::Ok);
+			CHECK(Response.UserData == UserData);
+			CHECK(Response.Name == Name);
+			CHECK(Response.Key == Key);
+		});
+		BlockingOwner.Wait();
+	}
+
+	{
+		INFO("Get Value Hash + Size");
+		FRequestOwner BlockingOwner(EPriority::Blocking);
+		Cache.GetValue({{Name, Key, ECachePolicy::Default | ECachePolicy::SkipData, UserData}}, BlockingOwner, [&](FCacheGetValueResponse&& Response)
+		{
+			CHECK(Response.Status == EStatus::Ok);
+			CHECK(Response.UserData == UserData);
+			CHECK(Response.Name == Name);
+			CHECK(Response.Key == Key);
+			CHECK(Response.Value.GetRawHash() == Value.GetRawHash());
+			CHECK(Response.Value.GetRawSize() == Value.GetRawSize());
+			CHECK(!Response.Value.HasData());
+		});
+		BlockingOwner.Wait();
+	}
+
+	{
+		INFO("Get Value");
+		FRequestOwner BlockingOwner(EPriority::Blocking);
+		Cache.GetValue({{Name, Key, ECachePolicy::Default, UserData}}, BlockingOwner, [&](FCacheGetValueResponse&& Response)
+		{
+			CHECK(Response.Status == EStatus::Ok);
+			CHECK(Response.UserData == UserData);
+			CHECK(Response.Name == Name);
+			CHECK(Response.Key == Key);
+			CHECK(Response.Value.GetRawHash() == Value.GetRawHash());
+			CHECK(Response.Value.GetRawSize() == Value.GetRawSize());
+			CHECK(Response.Value.HasData());
+		});
+		BlockingOwner.Wait();
+	}
+
+	{
+		INFO("Get Missing Value");
+		FRequestOwner BlockingOwner(EPriority::Blocking);
+		Cache.GetValue({{Name, GetMissingKey(), ECachePolicy::Default, UserData}}, BlockingOwner, [&](FCacheGetValueResponse&& Response)
+		{
+			CHECK(Response.Status == EStatus::Error);
+			CHECK(Response.UserData == UserData);
+			CHECK(Response.Name == Name);
+			CHECK(Response.Key == GetMissingKey());
+			CHECK(Response.Value == FValue::Null);
+		});
+		BlockingOwner.Wait();
+	}
+
+	INFO("TestCacheValues()");
+	TestCacheChunks(Cache, Key, {}, Value);
+}
+
+static void TestCacheRecords(ICache& Cache)
+{
+	const FValueId ValueId = FValueId::FromName(ANSITEXTVIEW("Value"));
+	const FValueId Attachment1Id = FValueId::FromName(ANSITEXTVIEW("Attachment1"));
+	const FValueId Attachment2Id = FValueId::FromName(ANSITEXTVIEW("Attachment2"));
+	const FValueId Attachment3Id = FValueId::FromName(ANSITEXTVIEW("Attachment3"));
+	const FValueId Attachment4Id = FValueId::FromName(ANSITEXTVIEW("Attachment4"));
+
+	const FValueWithId Value(ValueId, MakeSequentialIntegerBuffer(12 * 1024 * 1024));
+	const FValueWithId Attachment1(Attachment1Id, FValue::Compress(FSharedBuffer::MakeView(MakeMemoryView<uint8>({1}))));
+	const FValueWithId Attachment2(Attachment2Id, FValue::Compress(FSharedBuffer::MakeView(MakeMemoryView<uint8>({2, 3}))));
+	const FValueWithId Attachment3(Attachment3Id, FValue::Compress(FSharedBuffer::MakeView(MakeMemoryView<uint8>({3, 4, 5}))));
+	const FValueWithId Attachment4(Attachment4Id, FValue::Compress(FSharedBuffer::MakeView(MakeMemoryView<uint8>({4, 5, 6, 7}))));
+
+	const FCbObject Meta = CreateMeta("Value");
+
+	const uint64 UserData = 0xfedcba9876543210;
+	const FSharedString Name(TEXTVIEW("Records"));
+	const FCacheKey Key = CreateKeyFromGuid("ff173633-71e9-4960-a683-470c77abaa6a");
+
+	{
+		INFO("Put Record");
+		FCacheRecordBuilder RecordBuilder(Key);
+		RecordBuilder.AddValue(Value);
+		RecordBuilder.AddValue(Attachment1);
+		RecordBuilder.AddValue(Attachment2);
+		RecordBuilder.AddValue(Attachment3);
+		RecordBuilder.AddValue(Attachment4);
+		RecordBuilder.SetMeta(CopyTemp(Meta));
+		FRequestOwner BlockingOwner(EPriority::Blocking);
+		Cache.Put({{Name, RecordBuilder.Build(), ECachePolicy::Default, UserData}}, BlockingOwner, [&](FCachePutResponse&& Response)
+		{
+			CHECK(Response.Status == EStatus::Ok);
+			CHECK(Response.UserData == UserData);
+			CHECK(Response.Name == Name);
+			CHECK(Response.Key == Key);
+		});
+		BlockingOwner.Wait();
+	}
+
+	{
+		INFO("Get Record Hash + Size");
+		FRequestOwner BlockingOwner(EPriority::Blocking);
+		Cache.Get({{Name, Key, ECachePolicy::Default | ECachePolicy::SkipData, UserData}}, BlockingOwner, [&](FCacheGetResponse&& Response)
+		{
+			CHECK(Response.Status == EStatus::Ok);
+			CHECK(Response.UserData == UserData);
+			CHECK(Response.Name == Name);
+			CHECK(Response.Record.GetKey() == Key);
+			CHECK(Response.Record.GetValues().Num() == 5);
+			CHECK(HasValueWithNoData(Response.Record, Value));
+			CHECK(HasValueWithNoData(Response.Record, Attachment1));
+			CHECK(HasValueWithNoData(Response.Record, Attachment2));
+			CHECK(HasValueWithNoData(Response.Record, Attachment3));
+			CHECK(HasValueWithNoData(Response.Record, Attachment4));
+			CHECK(Response.Record.GetMeta().Equals(Meta));
+		});
+		BlockingOwner.Wait();
+	}
+
+	{
+		INFO("Get Record");
+		FRequestOwner BlockingOwner(EPriority::Blocking);
+		Cache.Get({{Name, Key, ECachePolicy::Default | ECachePolicy::SkipMeta, UserData}}, BlockingOwner, [&](FCacheGetResponse&& Response)
+		{
+			CHECK(Response.Status == EStatus::Ok);
+			CHECK(Response.UserData == UserData);
+			CHECK(Response.Name == Name);
+			CHECK(Response.Record.GetKey() == Key);
+			CHECK(Response.Record.GetValues().Num() == 5);
+			CHECK(HasValueWithData(Response.Record, Value));
+			CHECK(HasValueWithData(Response.Record, Attachment1));
+			CHECK(HasValueWithData(Response.Record, Attachment2));
+			CHECK(HasValueWithData(Response.Record, Attachment3));
+			CHECK(HasValueWithData(Response.Record, Attachment4));
+			CHECK_FALSE(Response.Record.GetMeta());
+		});
+		BlockingOwner.Wait();
+	}
+
+	{
+		INFO("Get Record w/ Skipped/Ignored Values");
+		FCacheRecordPolicyBuilder PolicyBuilder(ECachePolicy::Default);
+		PolicyBuilder.AddValuePolicy(ValueId, ECachePolicy::Default | ECachePolicy::SkipData);
+		PolicyBuilder.AddValuePolicy(Attachment1Id, ECachePolicy::None);
+		PolicyBuilder.AddValuePolicy(Attachment2Id, ECachePolicy::None);
+		PolicyBuilder.AddValuePolicy(Attachment3Id, ECachePolicy::None);
+		FRequestOwner BlockingOwner(EPriority::Blocking);
+		Cache.Get({{Name, Key, PolicyBuilder.Build(), UserData}}, BlockingOwner, [&](FCacheGetResponse&& Response)
+		{
+			CHECK(Response.Status == EStatus::Ok);
+			CHECK(Response.UserData == UserData);
+			CHECK(Response.Name == Name);
+			CHECK(Response.Record.GetKey() == Key);
+			CHECK(Response.Record.GetValues().Num() == 5);
+			CHECK(HasValueWithNoData(Response.Record, Value));
+			CHECK(HasValueWithNoData(Response.Record, Attachment1));
+			CHECK(HasValueWithNoData(Response.Record, Attachment2));
+			CHECK(HasValueWithNoData(Response.Record, Attachment3));
+			CHECK(HasValueWithData(Response.Record, Attachment4));
+			CHECK(Response.Record.GetMeta().Equals(Meta));
+		});
+		BlockingOwner.Wait();
+	}
+
+	{
+		INFO("Get Missing Record");
+		FRequestOwner BlockingOwner(EPriority::Blocking);
+		Cache.Get({{Name, GetMissingKey(), ECachePolicy::Default, UserData}}, BlockingOwner, [&](FCacheGetResponse&& Response)
+		{
+			CHECK(Response.Status == EStatus::Error);
+			CHECK(Response.UserData == UserData);
+			CHECK(Response.Name == Name);
+			CHECK(Response.Record.GetKey() == GetMissingKey());
+			CHECK(Response.Record.GetValues().IsEmpty());
+			CHECK_FALSE(Response.Record.GetMeta());
+		});
+		BlockingOwner.Wait();
+	}
+	
+	INFO("TestCacheRecords()");
+	TestCacheChunks(Cache, Key, ValueId, Value);
+}
+
+static void TestPartialCacheRecords(ICache& Cache)
+{
+	const FValueId ValueId = FValueId::FromName(ANSITEXTVIEW("Value"));
+	const FValueId Attachment1Id = FValueId::FromName(ANSITEXTVIEW("Attachment1"));
+	const FValueId Attachment2Id = FValueId::FromName(ANSITEXTVIEW("Attachment2"));
+	const FValueId Attachment3Id = FValueId::FromName(ANSITEXTVIEW("Attachment3"));
+	const FValueId Attachment4Id = FValueId::FromName(ANSITEXTVIEW("Attachment4"));
+
+	const FValueWithId Value(ValueId, MakeSequentialIntegerBuffer(16 * 1024 * 1024));
+	const FValueWithId Attachment1(Attachment1Id, FValue::Compress(FSharedBuffer::MakeView(MakeMemoryView<uint8>({5}))));
+	const FValueWithId Attachment2(Attachment2Id, FValue::Compress(FSharedBuffer::MakeView(MakeMemoryView<uint8>({6, 5}))));
+	const FValueWithId Attachment3(Attachment3Id, FValue::Compress(FSharedBuffer::MakeView(MakeMemoryView<uint8>({7, 6, 5}))));
+	const FValueWithId Attachment4(Attachment4Id, FValue::Compress(FSharedBuffer::MakeView(MakeMemoryView<uint8>({8, 7, 6, 5}))));
+
+	const FCbObject Meta = CreateMeta("Value");
+
+	const uint64 UserData = 0xfedcba9876543210;
+	const FSharedString Name(TEXTVIEW("PartialRecords"));
+	const FCacheKey Key = CreateKeyFromGuid("5d077783-8ecc-400a-bb65-168345d4f7a7");
+
+	{
+		INFO("Put Record w/ Value+Attachment1+Attachment2");
+
+		FCacheRecordBuilder RecordBuilder(Key);
+		RecordBuilder.AddValue(Value);
+		RecordBuilder.AddValue(Attachment1);
+		RecordBuilder.AddValue(Attachment2);
+		RecordBuilder.AddValue(Attachment3.RemoveData());
+		RecordBuilder.AddValue(Attachment4.RemoveData());
+		RecordBuilder.SetMeta(CopyTemp(Meta));
+		FRequestOwner BlockingOwner(EPriority::Blocking);
+		Cache.Put({{Name, RecordBuilder.Build(), ECachePolicy::Default | ECachePolicy::PartialRecord, UserData}}, BlockingOwner, [&](FCachePutResponse&& Response)
+		{
+			CHECK(Response.Status == EStatus::Ok);
+			CHECK(Response.UserData == UserData);
+			CHECK(Response.Name == Name);
+			CHECK(Response.Key == Key);
+		});
+		BlockingOwner.Wait();
+	}
+
+	{
+		INFO("Put Record w/ Attachment2+Attachment3");
+		FCacheRecordBuilder RecordBuilder(Key);
+		RecordBuilder.AddValue(Value.RemoveData());
+		RecordBuilder.AddValue(Attachment1.RemoveData());
+		RecordBuilder.AddValue(Attachment2);
+		RecordBuilder.AddValue(Attachment3);
+		RecordBuilder.AddValue(Attachment4.RemoveData());
+		// TODO: Allow Meta to differ or be unset here without blocking the subsequent partial put.
+		//RecordBuilder.SetMeta(CreateMeta("NewValue"));
+		RecordBuilder.SetMeta(CreateMeta("Value"));
+		FRequestOwner BlockingOwner(EPriority::Blocking);
+		Cache.Put({{Name, RecordBuilder.Build(), ECachePolicy::Default | ECachePolicy::PartialRecord, UserData}}, BlockingOwner, [&](FCachePutResponse&& Response)
+		{
+			CHECK(Response.Status == EStatus::Ok);
+			CHECK(Response.UserData == UserData);
+			CHECK(Response.Name == Name);
+			CHECK(Response.Key == Key);
+		});
+		BlockingOwner.Wait();
+	}
+
+	{
+		INFO("Get Record Hash + Size w/o PartialRecord");
+		FRequestOwner BlockingOwner(EPriority::Blocking);
+		Cache.Get({{Name, Key, ECachePolicy::Default | ECachePolicy::SkipData, UserData}}, BlockingOwner, [&](FCacheGetResponse&& Response)
+		{
+			CHECK(Response.Status == EStatus::Error);
+			CHECK(Response.UserData == UserData);
+			CHECK(Response.Name == Name);
+			CHECK(Response.Record.GetKey() == Key);
+			CHECK(Response.Record.GetValues().IsEmpty());
+			CHECK_FALSE(Response.Record.GetMeta());
+		});
+		BlockingOwner.Wait();
+	}
+
+	{
+		INFO("Get Record w/o PartialRecord");
+		FRequestOwner BlockingOwner(EPriority::Blocking);
+		Cache.Get({{Name, Key, ECachePolicy::Default | ECachePolicy::SkipMeta, UserData}}, BlockingOwner, [&](FCacheGetResponse&& Response)
+		{
+			CHECK(Response.Status == EStatus::Error);
+			CHECK(Response.UserData == UserData);
+			CHECK(Response.Name == Name);
+			CHECK(Response.Record.GetKey() == Key);
+			CHECK(Response.Record.GetValues().IsEmpty());
+			CHECK_FALSE(Response.Record.GetMeta());
+		});
+		BlockingOwner.Wait();
+	}
+
+	{
+		INFO("Get Record Hash + Size");
+		FRequestOwner BlockingOwner(EPriority::Blocking);
+		Cache.Get({{Name, Key, ECachePolicy::Default | ECachePolicy::SkipData | ECachePolicy::PartialRecord, UserData}}, BlockingOwner, [&](FCacheGetResponse&& Response)
+		{
+			CHECK(Response.Status == EStatus::Error);
+			CHECK(Response.UserData == UserData);
+			CHECK(Response.Name == Name);
+			CHECK(Response.Record.GetKey() == Key);
+			CHECK(Response.Record.GetValues().Num() == 5);
+			CHECK(HasValueWithNoData(Response.Record, Value));
+			CHECK(HasValueWithNoData(Response.Record, Attachment1));
+			CHECK(HasValueWithNoData(Response.Record, Attachment2));
+			CHECK(HasValueWithNoData(Response.Record, Attachment3));
+			CHECK(HasValueWithNoData(Response.Record, Attachment4));
+			CHECK(Response.Record.GetMeta().Equals(Meta));
+		});
+		BlockingOwner.Wait();
+	}
+
+	{
+		INFO("Get Record");
+		FRequestOwner BlockingOwner(EPriority::Blocking);
+		Cache.Get({{Name, Key, ECachePolicy::Default | ECachePolicy::SkipMeta | ECachePolicy::PartialRecord, UserData}}, BlockingOwner, [&](FCacheGetResponse&& Response)
+		{
+			CHECK(Response.Status == EStatus::Error);
+			CHECK(Response.UserData == UserData);
+			CHECK(Response.Name == Name);
+			CHECK(Response.Record.GetKey() == Key);
+			CHECK(Response.Record.GetValues().Num() == 5);
+			CHECK(HasValueWithData(Response.Record, Value));
+			CHECK(HasValueWithData(Response.Record, Attachment1));
+			CHECK(HasValueWithData(Response.Record, Attachment2));
+			CHECK(HasValueWithData(Response.Record, Attachment3));
+			CHECK(HasValueWithNoData(Response.Record, Attachment4));
+			CHECK_FALSE(Response.Record.GetMeta());
+		});
+		BlockingOwner.Wait();
+	}
+
+	{
+		INFO("Get Record w/ Skipped Value + Ignored Missing Value");
+		FCacheRecordPolicyBuilder PolicyBuilder(ECachePolicy::Default);
+		PolicyBuilder.AddValuePolicy(Attachment3Id, ECachePolicy::Default | ECachePolicy::SkipData);
+		PolicyBuilder.AddValuePolicy(Attachment4Id, ECachePolicy::None);
+		FRequestOwner BlockingOwner(EPriority::Blocking);
+		Cache.Get({{Name, Key, PolicyBuilder.Build(), UserData}}, BlockingOwner, [&](FCacheGetResponse&& Response)
+		{
+			CHECK(Response.Status == EStatus::Ok);
+			CHECK(Response.UserData == UserData);
+			CHECK(Response.Name == Name);
+			CHECK(Response.Record.GetKey() == Key);
+			CHECK(HasValueWithData(Response.Record, Value));
+			CHECK(HasValueWithData(Response.Record, Attachment1));
+			CHECK(HasValueWithData(Response.Record, Attachment2));
+			CHECK(HasValueWithNoData(Response.Record, Attachment3));
+			CHECK(HasValueWithNoData(Response.Record, Attachment4));
+			CHECK(Response.Record.GetMeta().Equals(Meta));
+		});
+		BlockingOwner.Wait();
+	}
+
+	{
+		INFO("Get Record w/ Default-Ignored Values + Required Values");
+		FCacheRecordPolicyBuilder PolicyBuilder(ECachePolicy::None);
+		PolicyBuilder.AddValuePolicy(Attachment1Id, ECachePolicy::Default);
+		PolicyBuilder.AddValuePolicy(Attachment2Id, ECachePolicy::Default);
+		FRequestOwner BlockingOwner(EPriority::Blocking);
+		Cache.Get({{Name, Key, PolicyBuilder.Build(), UserData}}, BlockingOwner, [&](FCacheGetResponse&& Response)
+		{
+			CHECK(Response.Status == EStatus::Ok);
+			CHECK(Response.UserData == UserData);
+			CHECK(Response.Name == Name);
+			CHECK(Response.Record.GetKey() == Key);
+			CHECK(HasValueWithNoData(Response.Record, Value));
+			CHECK(HasValueWithData(Response.Record, Attachment1));
+			CHECK(HasValueWithData(Response.Record, Attachment2));
+			CHECK(HasValueWithNoData(Response.Record, Attachment3));
+			CHECK(HasValueWithNoData(Response.Record, Attachment4));
+			CHECK(Response.Record.GetMeta().Equals(Meta));
+		});
+		BlockingOwner.Wait();
+	}
+
+	{
+		INFO("Get Record w/ Default-Ignored Values + Required Missing Value");
+		FCacheRecordPolicyBuilder PolicyBuilder(ECachePolicy::None);
+		PolicyBuilder.AddValuePolicy(Attachment4Id, ECachePolicy::Default);
+		FRequestOwner BlockingOwner(EPriority::Blocking);
+		Cache.Get({{Name, Key, PolicyBuilder.Build(), UserData}}, BlockingOwner, [&](FCacheGetResponse&& Response)
+		{
+			CHECK(Response.Status == EStatus::Error);
+			CHECK(Response.UserData == UserData);
+			CHECK(Response.Name == Name);
+			CHECK(Response.Record.GetKey() == Key);
+			CHECK(Response.Record.GetValues().IsEmpty());
+			CHECK_FALSE(Response.Record.GetMeta());
+		});
+		BlockingOwner.Wait();
+	}
+
+	INFO("TestPartialCacheRecords()");
+	TestCacheChunks(Cache, Key, ValueId, Value);
+}
+
+static void TestCacheDeterminism(ICache& Cache)
+{
+	const FValue Value(MakeSequentialIntegerBuffer(8 * 1024 * 1024));
+	const FValue OtherValue(MakeSequentialIntegerBuffer(8 * 1024 * 1024 + 1024));
+
+	const FValueId ValueId = FValueId::FromName(ANSITEXTVIEW("Value"));
+	const FValueId AttachmentId = FValueId::FromName(ANSITEXTVIEW("Attachment"));
+
+	const FCbObject Meta = CreateMeta("Value");
+
+	const FCacheKey RecordKey = CreateKeyFromGuid("8660fd50-a31c-4ea1-ace8-eb25b21b6db6");
+
+	// Put Original Record w/ Overwrite
+	const auto PutOriginalRecordWithOverwrite = [&]
+	{
+		FCacheRecordBuilder RecordBuilder(RecordKey);
+		RecordBuilder.AddValue(ValueId, Value);
+		RecordBuilder.AddValue(AttachmentId, Value);
+		FCacheRecordPolicyBuilder PolicyBuilder(ECachePolicy::Store);
+		PolicyBuilder.AddValuePolicy(ValueId, ECachePolicy::Store | ECachePolicy::NonDeterministic);
+		RecordBuilder.SetMeta(CopyTemp(Meta));
+		FRequestOwner BlockingOwner(EPriority::Blocking);
+		Cache.Put({{{TEXTVIEW("OriginalRecord")}, RecordBuilder.Build(), PolicyBuilder.Build()}}, BlockingOwner, [](FCachePutResponse&& Response)
+		{
+			CHECK(Response.Status == EStatus::Ok);
+		});
+		BlockingOwner.Wait();
+	};
+	PutOriginalRecordWithOverwrite();
+
+	// Put Original Record w/ New Meta
+	{
+		FCacheRecordBuilder RecordBuilder(RecordKey);
+		RecordBuilder.AddValue(ValueId, Value);
+		RecordBuilder.AddValue(AttachmentId, Value);
+		RecordBuilder.SetMeta(CreateMeta("ValueCopy"));
+		FRequestOwner BlockingOwner(EPriority::Blocking);
+		Cache.Put({{{TEXTVIEW("OriginalRecordCopy")}, RecordBuilder.Build()}}, BlockingOwner, [](FCachePutResponse&& Response)
+		{
+			CHECK(Response.Status == EStatus::Ok);
+		});
+		BlockingOwner.Wait();
+	}
+
+	// Put Different Record
+	{
+		FCacheRecordBuilder RecordBuilder(RecordKey);
+		RecordBuilder.AddValue(ValueId, OtherValue);
+		RecordBuilder.AddValue(AttachmentId, Value);
+		FRequestOwner BlockingOwner(EPriority::Blocking);
+		Cache.Put({{{TEXTVIEW("DifferentRecord")}, RecordBuilder.Build()}}, BlockingOwner, [](FCachePutResponse&& Response)
+		{
+			CHECK(Response.Status == EStatus::Ok);
+		});
+		BlockingOwner.Wait();
+	}
+
+	// Put Different Record w/ NonDeterministic Flag
+	{
+		FCacheRecordBuilder RecordBuilder(RecordKey);
+		RecordBuilder.AddValue(ValueId, OtherValue);
+		RecordBuilder.AddValue(AttachmentId, Value);
+		FCacheRecordPolicyBuilder PolicyBuilder;
+		PolicyBuilder.AddValuePolicy(ValueId, ECachePolicy::Default | ECachePolicy::NonDeterministic);
+		FRequestOwner BlockingOwner(EPriority::Blocking);
+		Cache.Put({{{TEXTVIEW("DifferentRecordIgnore")}, RecordBuilder.Build(), PolicyBuilder.Build()}}, BlockingOwner, [](FCachePutResponse&& Response)
+		{
+			CHECK(Response.Status == EStatus::Ok);
+		});
+		BlockingOwner.Wait();
+	}
+
+	// Get Original Record
+	{
+		FRequestOwner BlockingOwner(EPriority::Blocking);
+		Cache.Get({{{TEXTVIEW("OriginalRecord")}, RecordKey}}, BlockingOwner, [&](FCacheGetResponse&& Response)
+		{
+			CHECK(Response.Status == EStatus::Ok);
+			CHECK(HasValueWithData(Response.Record, FValueWithId(ValueId, Value)));
+			CHECK(Response.Record.GetMeta().Equals(Meta));
+		});
+		BlockingOwner.Wait();
+	}
+
+	// Put Different Record w/ Overwrite
+	{
+		FCacheRecordBuilder RecordBuilder(RecordKey);
+		RecordBuilder.AddValue(ValueId, OtherValue);
+		RecordBuilder.AddValue(AttachmentId, Value);
+		FRequestOwner BlockingOwner(EPriority::Blocking);
+		Cache.Put({{{TEXTVIEW("DifferentRecordOverwrite")}, RecordBuilder.Build(), ECachePolicy::Store}}, BlockingOwner, [](FCachePutResponse&& Response)
+		{
+			CHECK(Response.Status == EStatus::Ok);
+		});
+		BlockingOwner.Wait();
+	}
+
+	// Get Different Record
+	{
+		FRequestOwner BlockingOwner(EPriority::Blocking);
+		Cache.Get({{{TEXTVIEW("DifferentRecord")}, RecordKey}}, BlockingOwner, [&](FCacheGetResponse&& Response)
+		{
+			CHECK(Response.Status == EStatus::Ok);
+			CHECK(HasValueWithData(Response.Record, FValueWithId(ValueId, OtherValue)));
+			CHECK_FALSE(Response.Record.GetMeta());
+		});
+		BlockingOwner.Wait();
+	}
+
+	// Put Original Record w/ Overwrite
+	PutOriginalRecordWithOverwrite();
+
+	const FCacheKey ValueKey = CreateKeyFromGuid("926ae1e6-8ef8-4efa-8f15-83d6c0468909");
+
+	// Put Original Value w/ Overwrite
+	const auto PutOriginalValueWithOverwrite = [&]
+	{
+		FRequestOwner BlockingOwner(EPriority::Blocking);
+		Cache.PutValue({{{TEXTVIEW("OriginalValue")}, ValueKey, Value, ECachePolicy::Store | ECachePolicy::NonDeterministic}}, BlockingOwner, [](FCachePutValueResponse&& Response)
+		{
+			CHECK(Response.Status == EStatus::Ok);
+		});
+		BlockingOwner.Wait();
+	};
+	PutOriginalValueWithOverwrite();
+
+	// Put Original Value
+	{
+		FRequestOwner BlockingOwner(EPriority::Blocking);
+		Cache.PutValue({{{TEXTVIEW("OriginalValueCopy")}, ValueKey, Value}}, BlockingOwner, [](FCachePutValueResponse&& Response)
+		{
+			CHECK(Response.Status == EStatus::Ok);
+		});
+		BlockingOwner.Wait();
+	}
+
+	// Put Different Value
+	{
+		FRequestOwner BlockingOwner(EPriority::Blocking);
+		Cache.PutValue({{{TEXTVIEW("DifferentValue")}, ValueKey, OtherValue}}, BlockingOwner, [](FCachePutValueResponse&& Response)
+		{
+			CHECK(Response.Status == EStatus::Ok);
+		});
+		BlockingOwner.Wait();
+	}
+
+	// Put Different Value w/ NonDeterministic Flag
+	{
+		FRequestOwner BlockingOwner(EPriority::Blocking);
+		Cache.PutValue({{{TEXTVIEW("DifferentValueIgnore")}, ValueKey, OtherValue, ECachePolicy::Default | ECachePolicy::NonDeterministic}}, BlockingOwner, [](FCachePutValueResponse&& Response)
+		{
+			CHECK(Response.Status == EStatus::Ok);
+		});
+		BlockingOwner.Wait();
+	}
+
+	// Get Original Value
+	{
+		FRequestOwner BlockingOwner(EPriority::Blocking);
+		Cache.GetValue({{{TEXTVIEW("OriginalValue")}, ValueKey}}, BlockingOwner, [&](FCacheGetValueResponse&& Response)
+		{
+			CHECK(Response.Status == EStatus::Ok);
+			CHECK(Response.Value.GetRawHash() == Value.GetRawHash());
+		});
+		BlockingOwner.Wait();
+	}
+
+	// Put Different Value w/ Overwrite
+	{
+		FRequestOwner BlockingOwner(EPriority::Blocking);
+		Cache.PutValue({{{TEXTVIEW("DifferentValueOverwrite")}, ValueKey, OtherValue, ECachePolicy::Store}}, BlockingOwner, [](FCachePutValueResponse&& Response)
+		{
+			CHECK(Response.Status == EStatus::Ok);
+		});
+		BlockingOwner.Wait();
+	}
+
+	// Get Different Value
+	{
+		FRequestOwner BlockingOwner(EPriority::Blocking);
+		Cache.GetValue({{{TEXTVIEW("DifferentValue")}, ValueKey}}, BlockingOwner, [&OtherValue](FCacheGetValueResponse&& Response)
+		{
+			CHECK(Response.Status == EStatus::Ok);
+			CHECK(Response.Value.GetRawHash() == OtherValue.GetRawHash());
+		});
+		BlockingOwner.Wait();
+	}
+
+	// Put Original Value w/ Overwrite
+	PutOriginalValueWithOverwrite();
+}
+
+} // UE::DerivedData::CacheStoreTest
+
+namespace UE::DerivedData
+{
+
+TEST_CASE("DerivedData::CacheStore", "[DerivedData]")
+{
+	SECTION("Memory")
+	{
+		FStringView Config = TEXTVIEW("(Memory=(Type=Memory))");
+		TUniquePtr<ICache> Cache(CreateCache(Config));
+		CHECKED_IF(Cache)
+		{
+			CacheStoreTest::TestCacheValues(*Cache);
+			CacheStoreTest::TestCacheRecords(*Cache);
+			CacheStoreTest::TestPartialCacheRecords(*Cache);
+			CacheStoreTest::TestCacheDeterminism(*Cache);
+		}
+	}
+
+	SECTION("FileSystem")
+	{
+		FString CachePath = FPaths::EngineSavedDir() / TEXT("TestDerivedDataCache");
+
+		TStringBuilder<256> Config;
+		Config.Append(TEXT("(Local=(Type=FileSystem,Path=\""));
+		Config.Append(CachePath);
+		Config.Append(TEXT("\",Flush=true,DeleteUnused=false))"));
+
+		TUniquePtr<ICache> Cache(CreateCache(Config));
+		CHECKED_IF(Cache)
+		{
+			// Test different value sizes because value storage is tied to size.
+			CacheStoreTest::TestCacheValues(*Cache, 8);
+			CacheStoreTest::TestCacheValues(*Cache, 8 * 1024);
+			CacheStoreTest::TestCacheValues(*Cache, 8 * 1024 * 1024);
+
+			CacheStoreTest::TestCacheRecords(*Cache);
+			CacheStoreTest::TestPartialCacheRecords(*Cache);
+			CacheStoreTest::TestCacheDeterminism(*Cache);
+		}
+
+		IFileManager::Get().DeleteDirectory(*CachePath, /*bRequireExists*/ false, /*bTree*/ true);
+	}
+
+	SECTION("Zen")
+	{
+		FStringView Config = TEXTVIEW("(Local=(Type=Zen,Namespace=Test,Sandbox=Test,Flush=true))");
+		TUniquePtr<ICache> Cache(CreateCache(Config));
+		CHECKED_IF(Cache)
+		{
+			CacheStoreTest::TestCacheValues(*Cache);
+			CacheStoreTest::TestCacheRecords(*Cache);
+			CacheStoreTest::TestPartialCacheRecords(*Cache);
+			CacheStoreTest::TestCacheDeterminism(*Cache);
+		}
+	}
+}
+
+} // UE::DerivedData
+
+#endif // WITH_LOW_LEVEL_TESTS
