DDC: Fixed out-of-bounds chunk reads in the memory cache store

#Synced-CL 32500503

diff --git a/Engine/Source/Developer/DerivedDataCache/Private/MemoryCacheStore.cpp b/Engine/Source/Developer/DerivedDataCache/Private/MemoryCacheStore.cpp
--- a/Engine/Source/Developer/DerivedDataCache/Private/MemoryCacheStore.cpp
+++ b/Engine/Source/Developer/DerivedDataCache/Private/MemoryCacheStore.cpp
@@ -731,14 +731,14 @@ void FMemoryCacheStore::GetChunks(
 			}
 		}
 
-		if (bProcessHit && Request.RawOffset <= Value.GetRawSize() &&
-			(Request.RawHash.IsZero() || Request.RawHash == Value.GetRawHash()))
+		if (bProcessHit && (Request.RawHash.IsZero() || Request.RawHash == Value.GetRawHash()))
 		{
-			const uint64 RawSize = FMath::Min(Value.GetRawSize() - Request.RawOffset, Request.RawSize);
+			const uint64 RawOffset = FMath::Min(Value.GetRawSize(), Request.RawOffset);
+			const uint64 RawSize = FMath::Min(Value.GetRawSize() - RawOffset, Request.RawSize);
 			FSharedBuffer Buffer;
 			if (Value.HasData() && !bExistsOnly)
 			{
-				Buffer = Reader.Decompress(Request.RawOffset, RawSize);
+				Buffer = Reader.Decompress(RawOffset, RawSize);
 				RequestStats.LogicalReadSize += Buffer.GetSize();
 			}
 			const EStatus Status = bExistsOnly || Buffer ? EStatus::Ok : EStatus::Error;
