DDC: Added Cloud DDC to the cache store tests

#rnx
#Synced-CL 32500503

diff --git a/Engine/Source/Developer/DerivedDataCache/Tests/DerivedDataCacheStoreTest.cpp b/Engine/Source/Developer/DerivedDataCache/Tests/DerivedDataCacheStoreTest.cpp
--- a/Engine/Source/Developer/DerivedDataCache/Tests/DerivedDataCacheStoreTest.cpp
+++ b/Engine/Source/Developer/DerivedDataCache/Tests/DerivedDataCacheStoreTest.cpp
@@ -4,7 +4,9 @@
 
 #include "DerivedDataCache.h"
 #include "DerivedDataRequestOwner.h"
+#include "DerivedDataPrivate.h"
 #include "HAL/FileManager.h"
+#include "Logging/StructuredLog.h"
 #include "Misc/Guid.h"
 #include "Misc/Paths.h"
 #include "Misc/StringBuilder.h"
@@ -999,6 +1001,20 @@ TEST_CASE("DerivedData::CacheStore", "[DerivedData]")
 			CacheStoreTest::TestCacheDeterminism(*Cache);
 		}
 	}
+
+	SECTION("Cloud")
+	{
+		FStringView Config = TEXTVIEW("(Cloud)");
+		if (TUniquePtr<ICache> Cache{CreateCache(Config)})
+		{
+			CacheStoreTest::TestCacheValues(*Cache);
+			CacheStoreTest::TestCacheRecords(*Cache);
+		}
+		else
+		{
+			UE_LOGFMT(LogDerivedDataCache, Display, "Skipping test of Cloud cache store because it could not be created.");
+		}
+	}
 }
 
 } // UE::DerivedData
