DDC: Stop logging -1 cache size from the memory cache store

#rnx
#Synced-CL 32500503

diff --git a/Engine/Source/Developer/DerivedDataCache/Private/MemoryCacheStore.cpp b/Engine/Source/Developer/DerivedDataCache/Private/MemoryCacheStore.cpp
--- a/Engine/Source/Developer/DerivedDataCache/Private/MemoryCacheStore.cpp
+++ b/Engine/Source/Developer/DerivedDataCache/Private/MemoryCacheStore.cpp
@@ -138,7 +138,7 @@ FMemoryCacheStore::FMemoryCacheStore(
 	FParse::Value(InConfig, TEXT("MaxCacheSize="), SignedMaxCacheSize);
 	const int64 MaxSupportedCacheSize = 2048; // 2 GB
 	MaxCacheSize = FMath::Min(SignedMaxCacheSize, MaxSupportedCacheSize);
-	UE_LOG(LogDerivedDataCache, Display, TEXT("%s: Max Cache Size: %" INT64_FMT " MB"), *Name, SignedMaxCacheSize);
+	UE_CLOG(SignedMaxCacheSize >= 0, LogDerivedDataCache, Display, TEXT("%s: Max Cache Size: %" INT64_FMT " MB"), *Name, SignedMaxCacheSize);
 	MaxCacheSize = SignedMaxCacheSize < 0 ? 0 : SignedMaxCacheSize * 1024 * 1024;
 
 	if (StoreOwner)
