Allow ZenCacheStore to use the same param parsing utilities in place for HttpCacheStore to let ZenCacheStore have equivalent config/env/commandline overriding functionality for Host, ReadOnly, and WriteOnly parameters.

#Synced-CL 48523383

diff --git a/Engine/Source/Developer/DerivedDataCache/Private/HttpCacheStore.cpp b/Engine/Source/Developer/DerivedDataCache/Private/HttpCacheStore.cpp
--- a/Engine/Source/Developer/DerivedDataCache/Private/HttpCacheStore.cpp
+++ b/Engine/Source/Developer/DerivedDataCache/Private/HttpCacheStore.cpp
@@ -51,6 +51,7 @@
 #include "Serialization/JsonReader.h"
 #include "Serialization/JsonSerializer.h"
 #include "String/Find.h"
+#include "String/LexFromString.h"
 
 #if PLATFORM_MICROSOFT
 #include "Microsoft/AllowMicrosoftPlatformTypes.h"
@@ -2363,9 +2364,9 @@ FHttpCacheStore::FHttpCacheStore(const FHttpCacheStoreParams& Params, ICacheStor
 
 		if (StoreOwner)
 		{
-			const ECacheStoreFlags Flags = ECacheStoreFlags::Remote |
-											(Params.bReadOnly ? ECacheStoreFlags::None : ECacheStoreFlags::Store) |
-											(Params.bWriteOnly ? ECacheStoreFlags::None : ECacheStoreFlags::Query);
+			ECacheStoreFlags Flags = ECacheStoreFlags::Remote;
+			Flags |= Params.bWriteOnly ? ECacheStoreFlags::None : ECacheStoreFlags::Query;
+			Flags |= Params.bReadOnly ? ECacheStoreFlags::None : ECacheStoreFlags::Store;
 			TStringBuilder<256> Path(InPlace, ResolvedHost, TEXTVIEW(" ("), Namespace, TEXTVIEW(")"));
 			StoreOwner->Add(this, Flags);
 			StoreStats = StoreOwner->CreateStats(this, Flags, TEXTVIEW("Unreal Cloud DDC"), Params.Name, Path);
@@ -3309,51 +3310,29 @@ static void RegisterInheritedCommandLineArg(const FStringView ArgName)
 	FCommandLine::RegisterArgument(ArgName, ECommandLineArgumentFlags::EditorContext | ECommandLineArgumentFlags::CommandletContext | ECommandLineArgumentFlags::Inherit);
 }
 
-template <typename T>
-static bool ParseConfigValue(const TCHAR* Stream, const TCHAR* Match, T& OutValue)
-{
-	return FParse::Value(Stream, Match, OutValue);
-}
-
-template <>
-bool ParseConfigValue<bool>(const TCHAR* Stream, const TCHAR* Match, bool& OutValue)
-{
-	return FParse::Bool(Stream, Match, OutValue);
-}
-
-static bool ConvertStringToParamValue(const FString& InString, FString& OutString)
-{
-	OutString = InString;
-	return true;
-}
-
-static bool ConvertStringToParamValue(const FString& InString, bool& OutBool)
-{
-	OutBool = FCString::ToBool(*InString);
-	return true;
-}
-
-template <typename T>
-static void ParseConfigParamWithOverrides(const TCHAR* NodeName, const TCHAR* Config, const TCHAR* ParamName, T& OutValue, bool bIsSecret)
+FString ParseConfigParamWithOverrides(const TCHAR* NodeName, const TCHAR* Config, const TCHAR* ParamName, FStringView ExistingValue, bool bIsSecret)
 {
+	FString OutValue(ExistingValue);
 	FString OverrideName;
-	ParseConfigValue(Config, *WriteToString<32>(ParamName, TEXT("=")), OutValue);
-	if (ParseConfigValue(Config, *WriteToString<32>(TEXT("Env"), ParamName, TEXT("Override=")), OverrideName))
+	FParse::Value(Config, *WriteToString<32>(ParamName, TEXT("=")), OutValue);
+	if (FParse::Value(Config, *WriteToString<32>(TEXT("Env"), ParamName, TEXT("Override=")), OverrideName))
 	{
 		FString ParamEnv = FPlatformMisc::GetEnvironmentVariable(*OverrideName);
-		if (!ParamEnv.IsEmpty() && ConvertStringToParamValue(ParamEnv, OutValue))
+		if (!ParamEnv.IsEmpty())
 		{
+			OutValue = ParamEnv;
 			UE_LOG(LogDerivedDataCache, Log, TEXT("%s: Found environment override for %s %s=%s"), NodeName, ParamName, *OverrideName, bIsSecret ? TEXT("{SECRET}") : *ParamEnv);
 		}
 	}
-	if (ParseConfigValue(Config, *WriteToString<32>(TEXT("CommandLine"), ParamName, TEXT("Override=")), OverrideName))
+	if (FParse::Value(Config, *WriteToString<32>(TEXT("CommandLine"), ParamName, TEXT("Override=")), OverrideName))
 	{
-		if (ParseConfigValue(FCommandLine::Get(), *WriteToString<32>(OverrideName, TEXT("=")), OutValue))
+		if (FParse::Value(FCommandLine::Get(), *WriteToString<32>(OverrideName, TEXT("=")), OutValue))
 		{
 			RegisterInheritedCommandLineArg(OverrideName);
-			UE_LOG(LogDerivedDataCache, Log, TEXT("%s: Found command line override for %s %s=%s"), NodeName, ParamName, *OverrideName, bIsSecret ? TEXT("{SECRET}") : *WriteToString<32>(OutValue));
+			UE_LOG(LogDerivedDataCache, Log, TEXT("%s: Found command line override for %s %s=%s"), NodeName, ParamName, *OverrideName, bIsSecret ? TEXT("{SECRET}") : *OutValue);
 		}
 	}
+	return OutValue;
 }
 
 void FHttpCacheStoreParams::Parse(const TCHAR* NodeName, const TCHAR* Config)
@@ -3383,7 +3362,7 @@ void FHttpCacheStoreParams::Parse(const TCHAR* NodeName, const TCHAR* Config)
 	FString OverrideName;
 
 	// Host Params
-	ParseConfigParamWithOverrides(NodeName, Config, TEXT("Host"), Host, /*bIsSecret*/false);
+	Host = ParseConfigParamWithOverrides(NodeName, Config, TEXT("Host"), Host, /*bIsSecret*/false);
 
 	FParse::Value(Config, TEXT("DiscoveryHost="), DiscoveryHost);
 
@@ -3392,10 +3371,10 @@ void FHttpCacheStoreParams::Parse(const TCHAR* NodeName, const TCHAR* Config)
 	FParse::Bool(Config, TEXT("ResolveHostCanonicalName="), bResolveHostCanonicalName);
 
 	// Http version Params
-	ParseConfigParamWithOverrides(NodeName, Config, TEXT("HttpVersion"), HttpVersion, /*bIsSecret*/false);
+	HttpVersion = ParseConfigParamWithOverrides(NodeName, Config, TEXT("HttpVersion"), HttpVersion, /*bIsSecret*/false);
 
 	// Unix Socket Params
-	ParseConfigParamWithOverrides(NodeName, Config, TEXT("UnixSocketPath"), UnixSocketPath, /*bIsSecret*/false);
+	UnixSocketPath = ParseConfigParamWithOverrides(NodeName, Config, TEXT("UnixSocketPath"), UnixSocketPath, /*bIsSecret*/false);
 
 	// Namespace Params
 
@@ -3406,11 +3385,11 @@ void FHttpCacheStoreParams::Parse(const TCHAR* NodeName, const TCHAR* Config)
 	FParse::Value(Config, TEXT("StructuredNamespace="), Namespace);
 
 	// OAuth Params
-	ParseConfigParamWithOverrides(NodeName, Config, TEXT("OAuthProvider"), OAuthProvider, /*bIsSecret*/false);
+	OAuthProvider = ParseConfigParamWithOverrides(NodeName, Config, TEXT("OAuthProvider"), OAuthProvider, /*bIsSecret*/false);
 
 	FParse::Value(Config, TEXT("OAuthClientId="), OAuthClientId);
 
-	ParseConfigParamWithOverrides(NodeName, Config, TEXT("OAuthSecret"), OAuthSecret, /*bIsSecret*/true);
+	OAuthSecret = ParseConfigParamWithOverrides(NodeName, Config, TEXT("OAuthSecret"), OAuthSecret, /*bIsSecret*/true);
 
 	// If the secret is a file path, read the secret from the file.
 	if (OAuthSecret.StartsWith(TEXT("file://")))
@@ -3427,11 +3406,11 @@ void FHttpCacheStoreParams::Parse(const TCHAR* NodeName, const TCHAR* Config)
 	FParse::Value(Config, TEXT("OAuthScope="), OAuthScope);
 
 	// OAuth Provider Identifier
-	ParseConfigParamWithOverrides(NodeName, Config, TEXT("OAuthProviderIdentifier"), OAuthProviderIdentifier, /*bIsSecret*/false);
+	OAuthProviderIdentifier = ParseConfigParamWithOverrides(NodeName, Config, TEXT("OAuthProviderIdentifier"), OAuthProviderIdentifier, /*bIsSecret*/false);
 
 	// OAuth Access
 	// TODO-BEGIN: Replace this with:
-	// ParseConfigParamWithOverrides(NodeName, Config, TEXT("OAuthAccessToken"), OAuthAccessToken, /*bIsSecret*/true);
+	// OAuthAccessToken = ParseConfigParamWithOverrides(NodeName, Config, TEXT("OAuthAccessToken"), OAuthAccessToken, /*bIsSecret*/true);
 	if (!FParse::Value(Config, TEXT("OAuthAccessToken="), OAuthAccessToken))
 	{
 		if (FParse::Value(Config, TEXT("OAuthAccess="), OAuthAccessToken))
@@ -3460,9 +3439,9 @@ void FHttpCacheStoreParams::Parse(const TCHAR* NodeName, const TCHAR* Config)
 			UE_LOG(LogDerivedDataCache, Log, TEXT("%s: Found OAuth access token in %s."), NodeName, *OverrideName);
 		}
 	}
-	if (ParseConfigValue(Config, TEXT("CommandLineOAuthAccessTokenOverride="), OverrideName))
+	if (FParse::Value(Config, TEXT("CommandLineOAuthAccessTokenOverride="), OverrideName))
 	{
-		if (ParseConfigValue(FCommandLine::Get(), *WriteToString<32>(OverrideName, TEXT("=")), OAuthAccessToken))
+		if (FParse::Value(FCommandLine::Get(), *WriteToString<32>(OverrideName, TEXT("=")), OAuthAccessToken))
 		{
 			RegisterInheritedCommandLineArg(OverrideName);
 			UE_LOG(LogDerivedDataCache, Log, TEXT("%s: Found command line override for OAuthAccess %s={SECRET}"), NodeName, *OverrideName);
@@ -3479,8 +3458,8 @@ void FHttpCacheStoreParams::Parse(const TCHAR* NodeName, const TCHAR* Config)
 	FParse::Value(Config, TEXT("OAuthPinnedPublicKeys="), OAuthPinnedPublicKeys);
 
 	// Cache Params
-	ParseConfigParamWithOverrides(NodeName, Config, TEXT("ReadOnly"), bReadOnly, /*bIsSecret*/false);
-	ParseConfigParamWithOverrides(NodeName, Config, TEXT("WriteOnly"), bWriteOnly, /*bIsSecret*/false);
+	LexFromString(bReadOnly, ParseConfigParamWithOverrides(NodeName, Config, TEXT("ReadOnly"), ::LexToString(bReadOnly), /*bIsSecret*/false));
+	LexFromString(bWriteOnly, ParseConfigParamWithOverrides(NodeName, Config, TEXT("WriteOnly"), ::LexToString(bWriteOnly), /*bIsSecret*/false));
 
 	if (bReadOnly && bWriteOnly)
 	{
diff --git a/Engine/Source/Developer/DerivedDataCache/Private/ZenCacheStore.cpp b/Engine/Source/Developer/DerivedDataCache/Private/ZenCacheStore.cpp
--- a/Engine/Source/Developer/DerivedDataCache/Private/ZenCacheStore.cpp
+++ b/Engine/Source/Developer/DerivedDataCache/Private/ZenCacheStore.cpp
@@ -34,6 +34,7 @@
 #include "Serialization/CompactBinaryPackage.h"
 #include "Serialization/CompactBinaryWriter.h"
 #include "Serialization/LargeMemoryWriter.h"
+#include "String/LexFromString.h"
 #include "Templates/Function.h"
 #include "ZenAsyncCbPackageReceiver.h"
 #include "ZenBackendUtils.h"
@@ -58,6 +59,8 @@ static TAutoConsoleVariable<bool> CVarZenCacheLimitSandboxProcessLifetime(
 
 namespace UE::DerivedData
 {
+// Defined in HttpCacheStore.cpp
+FString ParseConfigParamWithOverrides(const TCHAR* NodeName, const TCHAR* Config, const TCHAR* ParamName, FStringView ExistingValue, bool bIsSecret);
 
 template<typename T>
 void ForEachBatch(const int32 BatchSize, const int32 TotalCount, T&& Fn)
@@ -93,6 +96,7 @@ struct FZenCacheStoreParams
 	TOptional<bool> bRemote;
 	bool bFlush = false;
 	bool bReadOnly = false;
+	bool bWriteOnly = false;
 	bool bBypassProxy = true;
 
 	void Parse(const TCHAR* NodeName, const TCHAR* Config);
@@ -1871,7 +1875,8 @@ void FZenCacheStore::Initialize(const FZenCacheStoreParams& Params)
 
 		DeactivateAtMs = Params.DeactivateAtMs;
 
-		ECacheStoreFlags Flags = ECacheStoreFlags::Query;
+		ECacheStoreFlags Flags = ECacheStoreFlags::None;
+		Flags |= Params.bWriteOnly ? ECacheStoreFlags::None : ECacheStoreFlags::Query;
 		Flags |= Params.bReadOnly ? ECacheStoreFlags::None : ECacheStoreFlags::Store;
 		Flags |= bLocal ? ECacheStoreFlags::Local : ECacheStoreFlags::None;
 		Flags |= bRemote ? ECacheStoreFlags::Remote : ECacheStoreFlags::None;
@@ -2219,10 +2224,6 @@ void FZenCacheStore::GetChunks(
 
 void FZenCacheStoreParams::Parse(const TCHAR* NodeName, const TCHAR* Config)
 {
-	auto RegisterInheritedCommandlineArg = [](const FStringView ArgName)
-	{
-		FCommandLine::RegisterArgument(ArgName, ECommandLineArgumentFlags::EditorContext | ECommandLineArgumentFlags::CommandletContext | ECommandLineArgumentFlags::Inherit);
-	};
 	Name = NodeName;
 
 	if (FString ServerId; FParse::Value(Config, TEXT("ServerID="), ServerId))
@@ -2239,32 +2240,18 @@ void FZenCacheStoreParams::Parse(const TCHAR* NodeName, const TCHAR* Config)
 		}
 	}
 
-	FParse::Value(Config, TEXT("Host="), Host);
-
-	FString OverrideName;
-	if (FParse::Value(Config, TEXT("EnvHostOverride="), OverrideName))
-	{
-		FString HostEnv = FPlatformMisc::GetEnvironmentVariable(*OverrideName);
-		if (!HostEnv.IsEmpty())
-		{
-			Host = HostEnv;
-			UE_LOG(LogDerivedDataCache, Log, TEXT("%s: Found environment override for Host %s=%s"), NodeName, *OverrideName, *Host);
-		}
-	}
-
-	if (FParse::Value(Config, TEXT("CommandLineHostOverride="), OverrideName))
-	{
-		if (FParse::Value(FCommandLine::Get(), *(OverrideName + TEXT("=")), Host))
-		{
-			RegisterInheritedCommandlineArg(OverrideName);
-			UE_LOG(LogDerivedDataCache, Log, TEXT("%s: Found command line override for Host %s=%s"), NodeName, *OverrideName, *Host);
-		}
-	}
+	Host = ParseConfigParamWithOverrides(NodeName, Config, TEXT("Host"), Host, /*bIsSecret*/false);
 
 	FParse::Value(Config, TEXT("Namespace="), Namespace);
 	FParse::Value(Config, TEXT("StructuredNamespace="), Namespace);
 
-	FParse::Bool(Config, TEXT("ReadOnly="), bReadOnly);
+	LexFromString(bReadOnly, ParseConfigParamWithOverrides(NodeName, Config, TEXT("ReadOnly"), ::LexToString(bReadOnly), /*bIsSecret*/false));
+	LexFromString(bWriteOnly, ParseConfigParamWithOverrides(NodeName, Config, TEXT("WriteOnly"), ::LexToString(bWriteOnly), /*bIsSecret*/false));
+
+	if (bReadOnly && bWriteOnly)
+	{
+		UE_LOG(LogDerivedDataCache, Display, TEXT("%s: Both ReadOnly and WriteOnly specified.  This cache store will be inactive."), NodeName);
+	}
 
 	// Sandbox and flush configuration for use in Cold/Warm type use cases
 	FParse::Value(Config, TEXT("Sandbox="), Sandbox);
