Add a TotalTimeoutMs parameter to FHttpClientParams in DevHttp module.  Also changed ConnectTimeout to ConnectTimeoutMs so the units are clear and consistent.  ZenCacheStore's readiness check now limits to 5 seconds total timeout on non-buildmachines and 120 seconds total timeout on buildmachines.

#Synced-CL 48523383

diff --git a/Engine/Source/Developer/DerivedDataCache/Private/HttpCacheStore.cpp b/Engine/Source/Developer/DerivedDataCache/Private/HttpCacheStore.cpp
--- a/Engine/Source/Developer/DerivedDataCache/Private/HttpCacheStore.cpp
+++ b/Engine/Source/Developer/DerivedDataCache/Private/HttpCacheStore.cpp
@@ -2440,7 +2440,7 @@ FHttpClientParams FHttpCacheStore::GetDefaultClientParams() const
 {
 	FHttpClientParams ClientParams;
 	ClientParams.DnsCacheTimeout = 15;
-	ClientParams.ConnectTimeout = 3 * 1000;
+	ClientParams.ConnectTimeoutMs = 3 * 1000;
 	ClientParams.LowSpeedLimit = 1024;
 	ClientParams.LowSpeedTime = 10;
 	ClientParams.TlsLevel = EHttpTlsLevel::All;
diff --git a/Engine/Source/Developer/DerivedDataCache/Private/ZenCacheStore.cpp b/Engine/Source/Developer/DerivedDataCache/Private/ZenCacheStore.cpp
--- a/Engine/Source/Developer/DerivedDataCache/Private/ZenCacheStore.cpp
+++ b/Engine/Source/Developer/DerivedDataCache/Private/ZenCacheStore.cpp
@@ -1804,13 +1804,14 @@ void FZenCacheStore::Initialize(const FZenCacheStoreParams& Params)
 
 	bool bReady = false;
 	{
-		// Issue a synchronous health/ready request with a limited idle time
+		// Issue a synchronous health/ready request with a limited total time
 		FHttpClientParams ReadinessClientParams;
 		ReadinessClientParams.Version = EHttpVersion::V2;
 		ReadinessClientParams.MaxRequests = 1;
 		ReadinessClientParams.MinRequests = 1;
-		ReadinessClientParams.LowSpeedLimit = 1;
-		ReadinessClientParams.LowSpeedTime = 5; // 5 second idle time limit for the initial readiness check
+		// 5 second total time limit for the initial readiness check request in non-buildmachine use cases
+		// 120 second (2 min) total time limit for buildmachine use cases
+		ReadinessClientParams.TotalTimeoutMs = GIsBuildMachine ? 120000 : 5000;
 		ReadinessClientParams.bBypassProxy = Params.bBypassProxy;
 		THttpUniquePtr<IHttpClient> ReadinessClient = ConnectionPool->CreateClient(ReadinessClientParams);
 		THttpUniquePtr<IHttpRequest> ReadinessRequest = ReadinessClient->TryCreateRequest({});
diff --git a/Engine/Source/Developer/DevHttp/Internal/Http/HttpClient.h b/Engine/Source/Developer/DevHttp/Internal/Http/HttpClient.h
--- a/Engine/Source/Developer/DevHttp/Internal/Http/HttpClient.h
+++ b/Engine/Source/Developer/DevHttp/Internal/Http/HttpClient.h
@@ -186,7 +186,10 @@ struct FHttpClientParams final
 	uint32 DnsCacheTimeout = 60;
 
 	/** Connection timeout in milliseconds. Use 0 for the default timeout. */
-	uint32 ConnectTimeout = 0;
+	uint32 ConnectTimeoutMs = 0;
+
+	/** Total timeout in milliseconds. Use 0 for no timeout. */
+	uint32 TotalTimeoutMs = 0;
 
 	/** Average transfer speed in bytes/s below which a request will abort. Use 0 to disable. */
 	uint32 LowSpeedLimit = 0;
diff --git a/Engine/Source/Developer/DevHttp/Private/CurlHttpClient.cpp b/Engine/Source/Developer/DevHttp/Private/CurlHttpClient.cpp
--- a/Engine/Source/Developer/DevHttp/Private/CurlHttpClient.cpp
+++ b/Engine/Source/Developer/DevHttp/Private/CurlHttpClient.cpp
@@ -654,7 +654,8 @@ void FCurlHttpClient::SetDefaultOptions(CURL* Curl, FCurlHttpHeaders& Headers) c
 	curl_easy_setopt(Curl, CURLOPT_HTTP_VERSION, ConvertVersion(Params.Version));
 
 	curl_easy_setopt(Curl, CURLOPT_DNS_CACHE_TIMEOUT, (long)FMath::Min(Params.DnsCacheTimeout, (uint32)MAX_int32));
-	curl_easy_setopt(Curl, CURLOPT_CONNECTTIMEOUT_MS, long(Params.ConnectTimeout));
+	curl_easy_setopt(Curl, CURLOPT_CONNECTTIMEOUT_MS, long(Params.ConnectTimeoutMs));
+	curl_easy_setopt(Curl, CURLOPT_TIMEOUT_MS, long(Params.TotalTimeoutMs));
 	curl_easy_setopt(Curl, CURLOPT_LOW_SPEED_LIMIT, long(Params.LowSpeedLimit));
 	curl_easy_setopt(Curl, CURLOPT_LOW_SPEED_TIME, long(Params.LowSpeedTime));
 	curl_easy_setopt(Curl, CURLOPT_USE_SSL, ConvertTlsLevel(Params.TlsLevel));
diff --git a/Engine/Source/Developer/DevHttp/Private/HttpHostBuilder.cpp b/Engine/Source/Developer/DevHttp/Private/HttpHostBuilder.cpp
--- a/Engine/Source/Developer/DevHttp/Private/HttpHostBuilder.cpp
+++ b/Engine/Source/Developer/DevHttp/Private/HttpHostBuilder.cpp
@@ -113,7 +113,7 @@ void FHttpHostBuilder::AddFromEndpoint(const FAnsiStringView HostUrl, const FAns
 	const THttpUniquePtr<IHttpConnectionPool> LocalConnectionPool = IHttpManager::Get().CreateConnectionPool(ConnectionPoolParams);
 	FHttpClientParams DefaultClientParams;
 	// we want to keep this timeout fairly low in case the host is not reachable
-	DefaultClientParams.ConnectTimeout = 5000;
+	DefaultClientParams.ConnectTimeoutMs = 5000;
 	THttpUniquePtr<IHttpClient> LocalClient = LocalConnectionPool->CreateClient(DefaultClientParams);
 
 	FHttpRequestParams RequestParams;
