Rework stats tracking in zen cache store client to avoid data races from tracking stats per batch instead of per request.  Also rework the minimum performance criteria for disconnecting/reconnecting to a zen cache to have disconnections happen based on MINIMUM latency over a 30 second timespan (if we receive more than 20 data points in that timespan), instead of AVERAGE latency.  Reconnection is also based solely on the single probe request latency, not the AVERAGE latency.

#Synced-CL 48523383

diff --git a/Engine/Source/Developer/DerivedDataCache/Private/ZenCacheStore.cpp b/Engine/Source/Developer/DerivedDataCache/Private/ZenCacheStore.cpp
--- a/Engine/Source/Developer/DerivedDataCache/Private/ZenCacheStore.cpp
+++ b/Engine/Source/Developer/DerivedDataCache/Private/ZenCacheStore.cpp
@@ -102,6 +102,34 @@ struct FZenCacheStoreParams
 	void Parse(const TCHAR* NodeName, const TCHAR* Config);
 };
 
+class FZenPerformanceEvaluationTracker
+{
+public:
+	void AddLatency(FMonotonicTimeSpan Latency)
+	{
+		if (Latency.IsInfinity() || Latency.IsZero())
+		{
+			return;
+		}
+		TUniqueLock Lock(Mutex);
+		++DataPointCount;
+		MinimumLatency = FMath::Min(MinimumLatency, Latency);
+	}
+
+	void MeasureAndReset(uint64& OutDataPointCount, FMonotonicTimeSpan& OutMinimumLatency)
+	{
+		TUniqueLock Lock(Mutex);
+		OutDataPointCount = DataPointCount;
+		OutMinimumLatency = MinimumLatency;
+		DataPointCount = 0;
+		MinimumLatency = FMonotonicTimeSpan::Infinity();
+	}
+private:
+	FMutex Mutex;
+	uint64 DataPointCount = 0;
+	FMonotonicTimeSpan MinimumLatency = FMonotonicTimeSpan::Infinity();
+};
+
 /**
  * Backend for a HTTP based caching service (Zen)
  */
@@ -262,6 +290,7 @@ private:
 	float DeactivateAtMs = -1.0f;
 	ECacheStoreFlags OperationalFlags;
 	FRequestOwner PerformanceEvaluationRequestOwner;
+	FZenPerformanceEvaluationTracker PerformanceEvaluationTracker;
 };
 
 template <typename RequestType>
@@ -286,6 +315,7 @@ struct FZenCacheStore::TRequestWithStats
 		{
 			Outer.StoreStats->AddRequest(Stats);
 		}
+		Outer.PerformanceEvaluationTracker.AddLatency(Stats.Latency);
 	}
 };
 
@@ -409,24 +439,12 @@ public:
 
 	virtual ~FPutOp()
 	{
-		FMonotonicTimeSpan AverageMainThreadTime = FMonotonicTimeSpan::FromSeconds(Requests[0].Stats.MainThreadTime.ToSeconds() / Requests.Num());
-		FMonotonicTimeSpan AverageOtherThreadTime = FMonotonicTimeSpan::FromSeconds(Requests[0].Stats.OtherThreadTime.ToSeconds() / Requests.Num());
-		FMonotonicTimeSpan AverageLatency = FMonotonicTimeSpan::FromSeconds(Requests[0].Stats.Latency.ToSeconds() / Requests.Num());
-
-		for (TRequestWithStats<FCachePutRequest>& Request : Requests)
-		{
-			Request.Stats.MainThreadTime = AverageMainThreadTime;
-			Request.Stats.OtherThreadTime = AverageOtherThreadTime;
-			Request.Stats.Latency = AverageLatency;
-		}
 		CacheStore.ConditionalEvaluatePerformance();
 		CacheStore.ConditionalUpdateStorageSize();
 	}
 
 	void IssueRequests()
 	{
-		FRequestTimer RequestTimer(Requests[0].Stats);
-
 		FRequestBarrier Barrier(Owner);
 		for (TArrayView<const TRequestWithStats<FCachePutRequest>> Batch : Batches)
 		{
@@ -446,6 +464,7 @@ public:
 					BatchWriter.BeginArray(ANSITEXTVIEW("Requests"));
 					for (const TRequestWithStats<FCachePutRequest>& RequestWithStats : Batch)
 					{
+						FRequestTimer RequestTimer(RequestWithStats.Stats);
 						const FCachePutRequest& Request = RequestWithStats.Request;
 						const FCacheRecord& Record = Request.Record;
 
@@ -471,10 +490,6 @@ public:
 			{
 				const bool bCanceled = !HttpResponse || HttpResponse->GetErrorCode() == EHttpErrorCode::Canceled;
 
-				FRequestTimer RequestTimer(Requests[0].Stats);
-				// Latency can't be measured for Put operations because it is intertwined with upload time.
-				Requests[0].Stats.Latency = FMonotonicTimeSpan::Infinity();
-
 				int32 RequestIndex = 0;
 				if (!bCanceled && HttpResponse->GetErrorCode() == EHttpErrorCode::None && HttpResponse->GetStatusCode() >= 200 && HttpResponse->GetStatusCode() <= 299)
 				{
@@ -491,6 +506,9 @@ public:
 
 						const TRequestWithStats<FCachePutRequest>& RequestWithStats = Batch[RequestIndex++];
 
+						// Latency can't be measured for Put operations because it is intertwined with upload time.
+						RequestWithStats.Stats.Latency = FMonotonicTimeSpan::Infinity();
+
 						const FCacheKey& Key = RequestWithStats.Request.Record.GetKey();
 						bool bPutSucceeded = ResponseField.AsBool();
 
@@ -554,6 +572,7 @@ private:
 
 	void OnHit(const TRequestWithStats<FCachePutRequest>& RequestWithStats, FCbObject&& DetailsObject)
 	{
+		FRequestTimer RequestTimer(RequestWithStats.Stats);
 		if (DetailsObject)
 		{
 			FCbObject ExistingRecordObject(DetailsObject[ANSITEXTVIEW("Record")].AsObject());
@@ -586,6 +605,8 @@ private:
 			RequestWithStats.Stats.AddLogicalWrite(Value);
 			RequestWithStats.Stats.PhysicalWriteSize += Value.GetData().GetCompressedSize();
 		}
+
+		RequestTimer.Stop();
 		RequestWithStats.EndRequest(CacheStore, EStatus::Ok);
 
 		TRACE_COUNTER_ADD(ZenDDC_BytesSent, int64(RequestWithStats.Stats.PhysicalWriteSize));
@@ -595,6 +616,7 @@ private:
 
 	void OnMiss(const TRequestWithStats<FCachePutRequest>& RequestWithStats, FCbObject&& DetailsObject)
 	{
+		FRequestTimer RequestTimer(RequestWithStats.Stats);
 		FUtf8StringView DetailsMessage = DetailsObject[ANSITEXTVIEW("Message")].AsString();
 		if (DetailsObject && !DetailsMessage.IsEmpty())
 		{
@@ -606,14 +628,19 @@ private:
 			UE_LOG(LogDerivedDataCache, Verbose, TEXT("%s: Cache put failed for '%s' from '%s'"),
 				*CacheStore.GetName(), *WriteToString<96>(RequestWithStats.Request.Record.GetKey()), *RequestWithStats.Request.Name);
 		}
+
+		RequestTimer.Stop();
 		RequestWithStats.EndRequest(CacheStore, EStatus::Error);
 		OnComplete(RequestWithStats.Request.MakeResponse(EStatus::Error));
 	}
 
 	void OnCanceled(const TRequestWithStats<FCachePutRequest>& RequestWithStats)
 	{
+		FRequestTimer RequestTimer(RequestWithStats.Stats);
 		UE_LOG(LogDerivedDataCache, Verbose, TEXT("%s: Cache put failed with canceled request for '%s' from '%s'"),
 			*CacheStore.GetName(), *WriteToString<96>(RequestWithStats.Request.Record.GetKey()), *RequestWithStats.Request.Name);
+
+		RequestTimer.Stop();
 		RequestWithStats.EndRequest(CacheStore, EStatus::Canceled);
 		OnComplete(RequestWithStats.Request.MakeResponse(EStatus::Canceled));
 	}
@@ -646,24 +673,13 @@ public:
 	virtual ~FGetOp()
 	{
 		TRACE_COUNTER_SUBTRACT(ZenDDC_CacheRecordRequestCountInFlight, int64(Requests.Num()));
-		FMonotonicTimeSpan AverageMainThreadTime = FMonotonicTimeSpan::FromSeconds(Requests[0].Stats.MainThreadTime.ToSeconds() / Requests.Num());
-		FMonotonicTimeSpan AverageOtherThreadTime = FMonotonicTimeSpan::FromSeconds(Requests[0].Stats.OtherThreadTime.ToSeconds() / Requests.Num());
-		FMonotonicTimeSpan AverageLatency = FMonotonicTimeSpan::FromSeconds(Requests[0].Stats.Latency.ToSeconds() / Requests.Num());
 
-		for (TRequestWithStats<FCacheGetRequest>& Request : Requests)
-		{
-			Request.Stats.MainThreadTime = AverageMainThreadTime;
-			Request.Stats.OtherThreadTime = AverageOtherThreadTime;
-			Request.Stats.Latency = AverageLatency;
-		}
 		CacheStore.ConditionalEvaluatePerformance();
 		CacheStore.ConditionalUpdateStorageSize();
 	}
 
 	void IssueRequests()
 	{
-		FRequestTimer RequestTimer(Requests[0].Stats);
-
 		FRequestBarrier Barrier(Owner);
 		ForEachBatch(CacheStore.CacheRecordBatchSize, Requests.Num(),
 			[this](int32 BatchFirst, int32 BatchLast)
@@ -690,6 +706,7 @@ public:
 					BatchRequest.BeginArray(ANSITEXTVIEW("Requests"));
 					for (const TRequestWithStats<FCacheGetRequest>& RequestWithStats : Batch)
 					{
+						FRequestTimer RequestTimer(RequestWithStats.Stats);
 						BatchRequest.BeginObject();
 						{
 							BatchRequest << ANSITEXTVIEW("Key") << RequestWithStats.Request.Key;
@@ -712,13 +729,6 @@ public:
 			{
 				const bool bCanceled = !HttpResponse || HttpResponse->GetErrorCode() == EHttpErrorCode::Canceled;
 
-				FRequestTimer RequestTimer(Requests[0].Stats);
-				if (HttpResponse)
-				{
-					const FHttpResponseStats& ResponseStats = HttpResponse->GetStats();
-					Requests[0].Stats.Latency = FMonotonicTimeSpan::FromSeconds(ResponseStats.GetLatency());
-				}
-
 				int32 RequestIndex = 0;
 				if (!bCanceled && HttpResponse->GetErrorCode() == EHttpErrorCode::None && HttpResponse->GetStatusCode() >= 200 && HttpResponse->GetStatusCode() <= 299)
 				{
@@ -735,6 +745,7 @@ public:
 						}
 
 						const TRequestWithStats<FCacheGetRequest>& RequestWithStats = Batch[RequestIndex++];
+						RequestWithStats.Stats.Latency = FMonotonicTimeSpan::FromSeconds(HttpResponse->GetStats().GetLatency());
 
 						const FCacheKey& Key = RequestWithStats.Request.Key;
 						FOptionalCacheRecord Record;
@@ -785,6 +796,7 @@ public:
 private:
 	void OnHit(const TRequestWithStats<FCacheGetRequest>& RequestWithStats, FCacheRecord&& Record, bool bIncomplete)
 	{
+		FRequestTimer RequestTimer(RequestWithStats.Stats);
 		FCacheRecordBuilder RecordBuilder(RequestWithStats.Request.Key);
 
 		const ECachePolicy RecordPolicy = RequestWithStats.Request.Policy.GetRecordPolicy();
@@ -829,6 +841,7 @@ private:
 			else
 			{
 				// End immediately with error and no record
+				RequestTimer.Stop();
 				RequestWithStats.EndRequest(CacheStore, EStatus::Error);
 				OnComplete(RequestWithStats.Request.MakeResponse(EStatus::Error));
 				return;
@@ -840,6 +853,8 @@ private:
 			UE_LOG(LogDerivedDataCache, Verbose, TEXT("%s: Cache hit for '%s' from '%s'"),
 				*CacheStore.GetName(), *WriteToString<96>(RequestWithStats.Request.Key), *RequestWithStats.Request.Name);
 		}
+
+		RequestTimer.Stop();
 		RequestWithStats.EndRequest(CacheStore, Status);
 
 		TRACE_COUNTER_INCREMENT(ZenDDC_GetHit);
@@ -849,16 +864,22 @@ private:
 
 	void OnMiss(const TRequestWithStats<FCacheGetRequest>& RequestWithStats)
 	{
+		FRequestTimer RequestTimer(RequestWithStats.Stats);
 		UE_LOG(LogDerivedDataCache, Verbose, TEXT("%s: Cache miss for '%s' from '%s'"),
 			*CacheStore.GetName(), *WriteToString<96>(RequestWithStats.Request.Key), *RequestWithStats.Request.Name);
+
+		RequestTimer.Stop();
 		RequestWithStats.EndRequest(CacheStore, EStatus::Error);
 		OnComplete(RequestWithStats.Request.MakeResponse(EStatus::Error));
 	}
 
 	void OnCanceled(const TRequestWithStats<FCacheGetRequest>& RequestWithStats)
 	{
+		FRequestTimer RequestTimer(RequestWithStats.Stats);
 		UE_LOG(LogDerivedDataCache, Verbose, TEXT("%s: Cache miss with canceled request for '%s' from '%s'"),
 			*CacheStore.GetName(), *WriteToString<96>(RequestWithStats.Request.Key), *RequestWithStats.Request.Name);
+
+		RequestTimer.Stop();
 		RequestWithStats.EndRequest(CacheStore, EStatus::Canceled);
 		OnComplete(RequestWithStats.Request.MakeResponse(EStatus::Canceled));
 	}
@@ -889,24 +910,12 @@ public:
 
 	virtual ~FPutValueOp()
 	{
-		FMonotonicTimeSpan AverageMainThreadTime = FMonotonicTimeSpan::FromSeconds(Requests[0].Stats.MainThreadTime.ToSeconds() / Requests.Num());
-		FMonotonicTimeSpan AverageOtherThreadTime = FMonotonicTimeSpan::FromSeconds(Requests[0].Stats.OtherThreadTime.ToSeconds() / Requests.Num());
-		FMonotonicTimeSpan AverageLatency = FMonotonicTimeSpan::FromSeconds(Requests[0].Stats.Latency.ToSeconds() / Requests.Num());
-
-		for (TRequestWithStats<FCachePutValueRequest>& Request : Requests)
-		{
-			Request.Stats.MainThreadTime = AverageMainThreadTime;
-			Request.Stats.OtherThreadTime = AverageOtherThreadTime;
-			Request.Stats.Latency = AverageLatency;
-		}
 		CacheStore.ConditionalEvaluatePerformance();
 		CacheStore.ConditionalUpdateStorageSize();
 	}
 
 	void IssueRequests()
 	{
-		FRequestTimer RequestTimer(Requests[0].Stats);
-
 		FRequestBarrier Barrier(Owner);
 		for (TArrayView<const TRequestWithStats<FCachePutValueRequest>> Batch : Batches)
 		{
@@ -926,6 +935,7 @@ public:
 					BatchWriter.BeginArray("Requests");
 					for (const TRequestWithStats<FCachePutValueRequest>& RequestWithStats : Batch)
 					{
+						FRequestTimer RequestTimer(RequestWithStats.Stats);
 						BatchWriter.BeginObject();
 						{
 							BatchWriter << ANSITEXTVIEW("Key") << RequestWithStats.Request.Key;
@@ -953,10 +963,6 @@ public:
 			{
 				const bool bCanceled = !HttpResponse || HttpResponse->GetErrorCode() == EHttpErrorCode::Canceled;
 
-				FRequestTimer RequestTimer(Requests[0].Stats);
-				// Latency can't be measured for Put operations because it is intertwined with upload time.
-				Requests[0].Stats.Latency = FMonotonicTimeSpan::Infinity();
-
 				int32 RequestIndex = 0;
 				if (!bCanceled && HttpResponse->GetErrorCode() == EHttpErrorCode::None && HttpResponse->GetStatusCode() >= 200 && HttpResponse->GetStatusCode() <= 299)
 				{
@@ -973,6 +979,9 @@ public:
 
 						const TRequestWithStats<FCachePutValueRequest>& RequestWithStats = Batch[RequestIndex++];
 
+						// Latency can't be measured for Put operations because it is intertwined with upload time.
+						RequestWithStats.Stats.Latency = FMonotonicTimeSpan::Infinity();
+
 						bool bPutSucceeded = ResponseField.AsBool();
 						if (CacheStore.DebugOptions.ShouldSimulatePutMiss(RequestWithStats.Request.Key))
 						{
@@ -1029,6 +1038,7 @@ private:
 
 	void OnHit(const TRequestWithStats<FCachePutValueRequest>& RequestWithStats, FCbObjectView DetailsObjectView)
 	{
+		FRequestTimer RequestTimer(RequestWithStats.Stats);
 		if (DetailsObjectView && !EnumHasAnyFlags(RequestWithStats.Request.Policy, ECachePolicy::NonDeterministic))
 		{
 			const FValue& Value = RequestWithStats.Request.Value;
@@ -1045,6 +1055,8 @@ private:
 
 		RequestWithStats.Stats.AddLogicalWrite(RequestWithStats.Request.Value);
 		RequestWithStats.Stats.PhysicalWriteSize += RequestWithStats.Request.Value.GetData().GetCompressedSize();
+
+		RequestTimer.Stop();
 		RequestWithStats.EndRequest(CacheStore, EStatus::Ok);
 
 		TRACE_COUNTER_INCREMENT(ZenDDC_PutHit);
@@ -1054,6 +1066,7 @@ private:
 
 	void OnMiss(const TRequestWithStats<FCachePutValueRequest>& RequestWithStats, FCbObjectView DetailsObjectView)
 	{
+		FRequestTimer RequestTimer(RequestWithStats.Stats);
 		FUtf8StringView DetailsMessage = DetailsObjectView[ANSITEXTVIEW("Message")].AsString();
 		if (DetailsObjectView && !DetailsMessage.IsEmpty())
 		{
@@ -1065,14 +1078,19 @@ private:
 			UE_LOG(LogDerivedDataCache, Verbose, TEXT("%s: Cache put failed for '%s' from '%s'"),
 				*CacheStore.GetName(), *WriteToString<96>(RequestWithStats.Request.Key), *RequestWithStats.Request.Name);
 		}
+
+		RequestTimer.Stop();
 		RequestWithStats.EndRequest(CacheStore, EStatus::Error);
 		OnComplete(RequestWithStats.Request.MakeResponse(EStatus::Error));
 	}
 
 	void OnCanceled(const TRequestWithStats<FCachePutValueRequest>& RequestWithStats)
 	{
+		FRequestTimer RequestTimer(RequestWithStats.Stats);
 		UE_LOG(LogDerivedDataCache, Verbose, TEXT("%s: Cache put failed with canceled request for '%s' from '%s'"),
 			*CacheStore.GetName(), *WriteToString<96>(RequestWithStats.Request.Key), *RequestWithStats.Request.Name);
+
+		RequestTimer.Stop();
 		RequestWithStats.EndRequest(CacheStore, EStatus::Canceled);
 		OnComplete(RequestWithStats.Request.MakeResponse(EStatus::Canceled));
 	}
@@ -1105,24 +1123,13 @@ public:
 	virtual ~FGetValueOp()
 	{
 		TRACE_COUNTER_SUBTRACT(ZenDDC_CacheRecordRequestCountInFlight, int64(Requests.Num()));
-		FMonotonicTimeSpan AverageMainThreadTime = FMonotonicTimeSpan::FromSeconds(Requests[0].Stats.MainThreadTime.ToSeconds() / Requests.Num());
-		FMonotonicTimeSpan AverageOtherThreadTime = FMonotonicTimeSpan::FromSeconds(Requests[0].Stats.OtherThreadTime.ToSeconds() / Requests.Num());
-		FMonotonicTimeSpan AverageLatency = FMonotonicTimeSpan::FromSeconds(Requests[0].Stats.Latency.ToSeconds() / Requests.Num());
 
-		for (TRequestWithStats<FCacheGetValueRequest>& Request : Requests)
-		{
-			Request.Stats.MainThreadTime = AverageMainThreadTime;
-			Request.Stats.OtherThreadTime = AverageOtherThreadTime;
-			Request.Stats.Latency = AverageLatency;
-		}
 		CacheStore.ConditionalEvaluatePerformance();
 		CacheStore.ConditionalUpdateStorageSize();
 	}
 
 	void IssueRequests()
 	{
-		FRequestTimer RequestTimer(Requests[0].Stats);
-
 		FRequestBarrier Barrier(Owner);
 		ForEachBatch(CacheStore.CacheRecordBatchSize, Requests.Num(),
 			[this](int32 BatchFirst, int32 BatchLast)
@@ -1149,6 +1156,7 @@ public:
 					BatchRequest.BeginArray("Requests");
 					for (const TRequestWithStats<FCacheGetValueRequest>& RequestWithStats : Batch)
 					{
+						FRequestTimer RequestTimer(RequestWithStats.Stats);
 						BatchRequest.BeginObject();
 						{
 							BatchRequest << ANSITEXTVIEW("Key") << RequestWithStats.Request.Key;
@@ -1170,13 +1178,6 @@ public:
 			{
 				const bool bCanceled = !HttpResponse || HttpResponse->GetErrorCode() == EHttpErrorCode::Canceled;
 
-				FRequestTimer RequestTimer(Requests[0].Stats);
-				if (HttpResponse)
-				{
-					const FHttpResponseStats& ResponseStats = HttpResponse->GetStats();
-					Requests[0].Stats.Latency = FMonotonicTimeSpan::FromSeconds(ResponseStats.GetLatency());
-				}
-
 				int32 RequestIndex = 0;
 				if (!bCanceled && HttpResponse->GetErrorCode() == EHttpErrorCode::None && HttpResponse->GetStatusCode() >= 200 && HttpResponse->GetStatusCode() <= 299)
 				{
@@ -1191,6 +1192,7 @@ public:
 						}
 
 						const TRequestWithStats<FCacheGetValueRequest>& RequestWithStats = Batch[RequestIndex++];
+						RequestWithStats.Stats.Latency = FMonotonicTimeSpan::FromSeconds(HttpResponse->GetStats().GetLatency());
 
 						FCbObjectView ResultObj = ResultField.AsObjectView();
 						TOptional<FValue> Value;
@@ -1256,6 +1258,7 @@ public:
 private:
 	void OnHit(const TRequestWithStats<FCacheGetValueRequest>& RequestWithStats, FValue&& Value)
 	{
+		FRequestTimer RequestTimer(RequestWithStats.Stats);
 		if (!CacheStore.ValidateAndPrepareReceivedValue(RequestWithStats.Request.Name,
 			RequestWithStats.Request.Key,
 			{},
@@ -1265,6 +1268,7 @@ private:
 			RequestWithStats.Stats))
 		{
 			// End immediately with error and no value
+			RequestTimer.Stop();
 			RequestWithStats.EndRequest(CacheStore, EStatus::Error);
 			OnComplete(RequestWithStats.Request.MakeResponse(EStatus::Error));
 			return;
@@ -1272,6 +1276,8 @@ private:
 
 		UE_LOG(LogDerivedDataCache, Verbose, TEXT("%s: Cache hit for '%s' from '%s'"),
 			*CacheStore.GetName(), *WriteToString<96>(RequestWithStats.Request.Key), *RequestWithStats.Request.Name);
+
+		RequestTimer.Stop();
 		RequestWithStats.EndRequest(CacheStore, EStatus::Ok);
 
 		TRACE_COUNTER_INCREMENT(ZenDDC_GetHit);
@@ -1281,16 +1287,22 @@ private:
 
 	void OnMiss(const TRequestWithStats<FCacheGetValueRequest>& RequestWithStats)
 	{
+		FRequestTimer RequestTimer(RequestWithStats.Stats);
 		UE_LOG(LogDerivedDataCache, Verbose, TEXT("%s: Cache miss for '%s' from '%s'"),
 			*CacheStore.GetName(), *WriteToString<96>(RequestWithStats.Request.Key), *RequestWithStats.Request.Name);
+
+		RequestTimer.Stop();
 		RequestWithStats.EndRequest(CacheStore, EStatus::Error);
 		OnComplete(RequestWithStats.Request.MakeResponse(EStatus::Error));
 	};
 
 	void OnCanceled(const TRequestWithStats<FCacheGetValueRequest>& RequestWithStats)
 	{
+		FRequestTimer RequestTimer(RequestWithStats.Stats);
 		UE_LOG(LogDerivedDataCache, Verbose, TEXT("%s: Cache miss with canceled request for '%s' from '%s'"),
 			*CacheStore.GetName(), *WriteToString<96>(RequestWithStats.Request.Key), *RequestWithStats.Request.Name);
+
+		RequestTimer.Stop();
 		RequestWithStats.EndRequest(CacheStore, EStatus::Canceled);
 		OnComplete(RequestWithStats.Request.MakeResponse(EStatus::Canceled));
 	};
@@ -1323,24 +1335,13 @@ public:
 	virtual ~FGetChunksOp()
 	{
 		TRACE_COUNTER_SUBTRACT(ZenDDC_ChunkRequestCountInFlight, int64(Requests.Num()));
-		FMonotonicTimeSpan AverageMainThreadTime = FMonotonicTimeSpan::FromSeconds(Requests[0].Stats.MainThreadTime.ToSeconds() / Requests.Num());
-		FMonotonicTimeSpan AverageOtherThreadTime = FMonotonicTimeSpan::FromSeconds(Requests[0].Stats.OtherThreadTime.ToSeconds() / Requests.Num());
-		FMonotonicTimeSpan AverageLatency = FMonotonicTimeSpan::FromSeconds(Requests[0].Stats.Latency.ToSeconds() / Requests.Num());
 
-		for (TRequestWithStats<FCacheGetChunkRequest>& Request : Requests)
-		{
-			Request.Stats.MainThreadTime = AverageMainThreadTime;
-			Request.Stats.OtherThreadTime = AverageOtherThreadTime;
-			Request.Stats.Latency = AverageLatency;
-		}
 		CacheStore.ConditionalEvaluatePerformance();
 		CacheStore.ConditionalUpdateStorageSize();
 	}
 
 	void IssueRequests()
 	{
-		FRequestTimer RequestTimer(Requests[0].Stats);
-
 		FRequestBarrier Barrier(Owner);
 		ForEachBatch(CacheStore.CacheChunksBatchSize, Requests.Num(),
 			[this](int32 BatchFirst, int32 BatchLast)
@@ -1369,6 +1370,7 @@ public:
 					BatchRequest.BeginArray(ANSITEXTVIEW("ChunkRequests"));
 					for (const TRequestWithStats<FCacheGetChunkRequest>& RequestWithStats : Batch)
 					{
+						FRequestTimer RequestTimer(RequestWithStats.Stats);
 						const FCacheGetChunkRequest& Request = RequestWithStats.Request;
 						BatchRequest.BeginObject();
 						{
@@ -1408,13 +1410,6 @@ public:
 			{
 				const bool bCanceled = !HttpResponse || HttpResponse->GetErrorCode() == EHttpErrorCode::Canceled;
 
-				FRequestTimer RequestTimer(Requests[0].Stats);
-				if (HttpResponse)
-				{
-					const FHttpResponseStats& ResponseStats = HttpResponse->GetStats();
-					Requests[0].Stats.Latency = FMonotonicTimeSpan::FromSeconds(ResponseStats.GetLatency());
-				}
-
 				int32 RequestIndex = 0;
 				if (!bCanceled && HttpResponse->GetErrorCode() == EHttpErrorCode::None && HttpResponse->GetStatusCode() >= 200 && HttpResponse->GetStatusCode() <= 299)
 				{
@@ -1428,6 +1423,7 @@ public:
 							continue;
 						}
 						const TRequestWithStats<FCacheGetChunkRequest>& RequestWithStats = Batch[RequestIndex++];
+						RequestWithStats.Stats.Latency = FMonotonicTimeSpan::FromSeconds(HttpResponse->GetStats().GetLatency());
 
 						FIoHash RawHash;
 						bool Succeeded = false;
@@ -1516,6 +1512,7 @@ public:
 private:
 	void OnHit(const TRequestWithStats<FCacheGetChunkRequest>& RequestWithStats, FIoHash&& RawHash, uint64 RawSize, FSharedBuffer&& RequestedBytes)
 	{
+		FRequestTimer RequestTimer(RequestWithStats.Stats);
 		const FCacheGetChunkRequest& Request = RequestWithStats.Request;
 		UE_LOG(LogDerivedDataCache, Verbose, TEXT("%s: Cache hit for '%s' from '%s'"),
 			*CacheStore.GetName(), *WriteToString<96>(Request.Key, '/', Request.Id), *Request.Name);
@@ -1523,6 +1520,8 @@ private:
 		// This is a rough estimate of physical read size until Zen communicates stats with each response.
 		RequestWithStats.Stats.LogicalReadSize += RawSize;
 		RequestWithStats.Stats.PhysicalReadSize += RequestedBytes.GetSize();
+
+		RequestTimer.Stop();
 		RequestWithStats.EndRequest(CacheStore, EStatus::Ok);
 
 		TRACE_COUNTER_INCREMENT(ZenDDC_GetHit);
@@ -1533,18 +1532,24 @@ private:
 
 	void OnMiss(const TRequestWithStats<FCacheGetChunkRequest>& RequestWithStats)
 	{
+		FRequestTimer RequestTimer(RequestWithStats.Stats);
 		const FCacheGetChunkRequest& Request = RequestWithStats.Request;
 		UE_LOG(LogDerivedDataCache, Verbose, TEXT("%s: Cache miss with missing value '%s' for '%s' from '%s'"),
 			*CacheStore.GetName(), *WriteToString<16>(Request.Id), *WriteToString<96>(Request.Key), *Request.Name);
+
+		RequestTimer.Stop();
 		RequestWithStats.EndRequest(CacheStore, EStatus::Error);
 		OnComplete(Request.MakeResponse(EStatus::Error));
 	};
 
 	void OnCanceled(const TRequestWithStats<FCacheGetChunkRequest>& RequestWithStats)
 	{
+		FRequestTimer RequestTimer(RequestWithStats.Stats);
 		const FCacheGetChunkRequest& Request = RequestWithStats.Request;
 		UE_LOG(LogDerivedDataCache, Verbose, TEXT("%s: Cache miss with canceled request for '%s' from '%s'"),
 			*CacheStore.GetName(), *WriteToString<96>(Request.Key, '/', Request.Id), *Request.Name);
+
+		RequestTimer.Stop();
 		RequestWithStats.EndRequest(CacheStore, EStatus::Canceled);
 		OnComplete(Request.MakeResponse(EStatus::Canceled));
 	};
@@ -2030,8 +2035,11 @@ void FZenCacheStore::ActivatePerformanceEvaluationThread()
 						}
 
 						const double LatencySec = HttpResponse->GetStats().GetLatency();
-						StoreStats->AddLatency(StartTime, FMonotonicTimePoint::Now(), FMonotonicTimeSpan::FromSeconds(LatencySec));
-						if (!bTryEvaluatePerformance || (StoreStats->GetAverageLatency() * 1000 <= DeactivateAtMs))
+						if (StoreStats)
+						{
+							StoreStats->AddLatency(StartTime, FMonotonicTimePoint::Now(), FMonotonicTimeSpan::FromSeconds(LatencySec));
+						}
+						if (!bTryEvaluatePerformance || (LatencySec * 1000 <= DeactivateAtMs))
 						{
 							if (PerformanceEvaluationThread.IsSet())
 							{
@@ -2128,7 +2136,10 @@ void FZenCacheStore::ConditionalEvaluatePerformance()
 		PerformanceEvaluationThreadShutdownEvent.Reset();
 	}
 
-	if (StoreStats->GetAverageLatency() * 1000 > DeactivateAtMs)
+	uint64 DataPointCount = 0;
+	FMonotonicTimeSpan MinimumLatency;
+	PerformanceEvaluationTracker.MeasureAndReset(DataPointCount, MinimumLatency);
+	if (DataPointCount > 20 && MinimumLatency.ToMilliseconds() > DeactivateAtMs)
 	{
 		if (!bDeactivatedForPerformance.load(std::memory_order_relaxed))
 		{
