DDC: Extended compact binary serialization of types used by cache request/response types

#Synced-CL 32500503

diff --git a/Engine/Source/Developer/DerivedDataCache/Internal/DerivedDataCacheMethod.h b/Engine/Source/Developer/DerivedDataCache/Internal/DerivedDataCacheMethod.h
--- a/Engine/Source/Developer/DerivedDataCache/Internal/DerivedDataCacheMethod.h
+++ b/Engine/Source/Developer/DerivedDataCache/Internal/DerivedDataCacheMethod.h
@@ -7,9 +7,6 @@
 
 #define UE_API DERIVEDDATACACHE_API
 
-class FCbFieldView;
-class FCbWriter;
-
 namespace UE::DerivedData
 {
 
@@ -29,9 +26,6 @@ UE_API FUtf8StringBuilderBase& operator<<(FUtf8StringBuilderBase& Builder, ECach
 UE_API bool TryLexFromString(ECacheMethod& OutMethod, FUtf8StringView String);
 UE_API bool TryLexFromString(ECacheMethod& OutMethod, FWideStringView String);
 
-UE_API FCbWriter& operator<<(FCbWriter& Writer, ECacheMethod Method);
-UE_API bool LoadFromCompactBinary(FCbFieldView Field, ECacheMethod& OutMethod);
-
 class FCacheMethodFilter
 {
 public:
diff --git a/Engine/Source/Developer/DerivedDataCache/Private/DerivedDataCache.cpp b/Engine/Source/Developer/DerivedDataCache/Private/DerivedDataCache.cpp
--- a/Engine/Source/Developer/DerivedDataCache/Private/DerivedDataCache.cpp
+++ b/Engine/Source/Developer/DerivedDataCache/Private/DerivedDataCache.cpp
@@ -27,9 +27,6 @@
 #include "Misc/CoreMisc.h"
 #include "Misc/ScopeLock.h"
 #include "ProfilingDebugging/CookStats.h"
-#include "Serialization/CompactBinary.h"
-#include "Serialization/CompactBinarySerialization.h"
-#include "Serialization/CompactBinaryWriter.h"
 #include "Stats/Stats.h"
 #include "Stats/StatsMisc.h"
 #include <atomic>
@@ -255,131 +252,6 @@ FCacheGetChunkResponse FCacheGetChunkRequest::MakeResponse(const EStatus Status)
 	return {Name, Key, Id, RawOffset, 0, {}, {}, UserData, Status};
 }
 
-FCbWriter& operator<<(FCbWriter& Writer, const FCacheGetRequest& Request)
-{
-	Writer.BeginObject();
-	Private::SaveToCompactBinary(Writer, Request);
-	Writer.EndObject();
-	return Writer;
-}
-
-void Private::SaveToCompactBinary(FCbWriter& Writer, const FCacheGetRequest& Request)
-{
-	if (!Request.Name.IsEmpty())
-	{
-		Writer << ANSITEXTVIEW("Name") << Request.Name;
-	}
-	Writer << ANSITEXTVIEW("Key") << Request.Key;
-	if (!Request.Policy.IsDefault())
-	{
-		Writer << ANSITEXTVIEW("Policy") << Request.Policy;
-	}
-	if (Request.UserData != 0)
-	{
-		Writer << ANSITEXTVIEW("UserData") << Request.UserData;
-	}
-}
-
-bool LoadFromCompactBinary(FCbFieldView Field, FCacheGetRequest& Request)
-{
-	bool bOk = Field.IsObject();
-	LoadFromCompactBinary(Field[ANSITEXTVIEW("Name")], Request.Name);
-	bOk &= LoadFromCompactBinary(Field[ANSITEXTVIEW("Key")], Request.Key);
-	LoadFromCompactBinary(Field[ANSITEXTVIEW("Policy")], Request.Policy);
-	LoadFromCompactBinary(Field[ANSITEXTVIEW("UserData")], Request.UserData);
-	return bOk;
-}
-
-FCbWriter& operator<<(FCbWriter& Writer, const FCacheGetValueRequest& Request)
-{
-	Writer.BeginObject();
-	Private::SaveToCompactBinary(Writer, Request);
-	Writer.EndObject();
-	return Writer;
-}
-
-void Private::SaveToCompactBinary(FCbWriter& Writer, const FCacheGetValueRequest& Request)
-{
-	if (!Request.Name.IsEmpty())
-	{
-		Writer << ANSITEXTVIEW("Name") << MakeStringView(Request.Name);
-	}
-	Writer << ANSITEXTVIEW("Key") << Request.Key;
-	if (Request.Policy != ECachePolicy::Default)
-	{
-		Writer << ANSITEXTVIEW("Policy") << Request.Policy;
-	}
-	if (Request.UserData != 0)
-	{
-		Writer << ANSITEXTVIEW("UserData") << Request.UserData;
-	}
-}
-
-bool LoadFromCompactBinary(FCbFieldView Field, FCacheGetValueRequest& Request)
-{
-	bool bOk = Field.IsObject();
-	LoadFromCompactBinary(Field[ANSITEXTVIEW("Name")], Request.Name);
-	bOk &= LoadFromCompactBinary(Field[ANSITEXTVIEW("Key")], Request.Key);
-	LoadFromCompactBinary(Field[ANSITEXTVIEW("Policy")], Request.Policy);
-	LoadFromCompactBinary(Field[ANSITEXTVIEW("UserData")], Request.UserData);
-	return bOk;
-}
-
-FCbWriter& operator<<(FCbWriter& Writer, const FCacheGetChunkRequest& Request)
-{
-	Writer.BeginObject();
-	Private::SaveToCompactBinary(Writer, Request);
-	Writer.EndObject();
-	return Writer;
-}
-
-void Private::SaveToCompactBinary(FCbWriter& Writer, const FCacheGetChunkRequest& Request)
-{
-	if (!Request.Name.IsEmpty())
-	{
-		Writer << ANSITEXTVIEW("Name") << MakeStringView(Request.Name);
-	}
-	Writer << ANSITEXTVIEW("Key") << Request.Key;
-	if (Request.Id.IsValid())
-	{
-		Writer << ANSITEXTVIEW("Id") << Request.Id;
-	}
-	if (Request.RawOffset != 0)
-	{
-		Writer << ANSITEXTVIEW("RawOffset") << Request.RawOffset;
-	}
-	if (Request.RawSize != MAX_uint64)
-	{
-		Writer << ANSITEXTVIEW("RawSize") << Request.RawSize;
-	}
-	if (!Request.RawHash.IsZero())
-	{
-		Writer << ANSITEXTVIEW("RawHash") << Request.RawHash;
-	}
-	if (Request.Policy != ECachePolicy::Default)
-	{
-		Writer << ANSITEXTVIEW("Policy") << Request.Policy;
-	}
-	if (Request.UserData)
-	{
-		Writer << ANSITEXTVIEW("UserData") << Request.UserData;
-	}
-}
-
-bool LoadFromCompactBinary(FCbFieldView Field, FCacheGetChunkRequest& OutRequest)
-{
-	bool bOk = Field.IsObject();
-	LoadFromCompactBinary(Field[ANSITEXTVIEW("Name")], OutRequest.Name);
-	bOk &= LoadFromCompactBinary(Field[ANSITEXTVIEW("Key")], OutRequest.Key);
-	LoadFromCompactBinary(Field[ANSITEXTVIEW("Id")], OutRequest.Id);
-	LoadFromCompactBinary(Field[ANSITEXTVIEW("RawOffset")], OutRequest.RawOffset, 0);
-	LoadFromCompactBinary(Field[ANSITEXTVIEW("RawSize")], OutRequest.RawSize, MAX_uint64);
-	LoadFromCompactBinary(Field[ANSITEXTVIEW("RawHash")], OutRequest.RawHash);
-	LoadFromCompactBinary(Field[ANSITEXTVIEW("Policy")], OutRequest.Policy);
-	LoadFromCompactBinary(Field[ANSITEXTVIEW("UserData")], OutRequest.UserData);
-	return bOk;
-}
-
 } // UE::DerivedData
 
 namespace UE::DerivedData::Private
@@ -389,7 +261,7 @@ FQueuedThreadPool* GCacheThreadPool;
 
 void LaunchTaskInCacheThreadPool(IRequestOwner& Owner, TUniqueFunction<void ()>&& TaskBody)
 {
-	LaunchTaskInThreadPool(Owner, GCacheThreadPool, MoveTemp(TaskBody));
+	LaunchTaskInThreadPool(0, TEXT("LaunchTaskInCacheThreadPool"), Owner, GCacheThreadPool, MoveTemp(TaskBody));
 }
 
 class FLegacyFetchOrBuildTask
diff --git a/Engine/Source/Developer/DerivedDataCache/Private/DerivedDataCacheKey.cpp b/Engine/Source/Developer/DerivedDataCache/Private/DerivedDataCacheKey.cpp
--- a/Engine/Source/Developer/DerivedDataCache/Private/DerivedDataCacheKey.cpp
+++ b/Engine/Source/Developer/DerivedDataCache/Private/DerivedDataCacheKey.cpp
@@ -14,8 +14,6 @@
 #include "Math/UnrealMathUtility.h"
 #include "Misc/AsciiSet.h"
 #include "Misc/StringBuilder.h"
-#include "Serialization/CompactBinary.h"
-#include "Serialization/CompactBinarySerialization.h"
 #include "Serialization/CompactBinaryWriter.h"
 #include "String/Find.h"
 
@@ -202,45 +200,11 @@ void FCacheBucket::ToDisplayName(FStringBuilderBase& OutDisplayName) const
 	Private::GetCacheBuckets().GetDisplayName(*this, OutDisplayName);
 }
 
-FCbWriter& operator<<(FCbWriter& Writer, const FCacheBucket Bucket)
-{
-	Writer.AddString(Bucket.ToString());
-	return Writer;
-}
-
-bool LoadFromCompactBinary(FCbFieldView Field, FCacheBucket& OutBucket)
-{
-	if (const FUtf8StringView Bucket = Field.AsString(); !Field.HasError() && FCacheBucket::IsValidName(Bucket))
-	{
-		OutBucket = FCacheBucket(Bucket);
-		return true;
-	}
-	OutBucket.Reset();
-	return false;
-}
-
-FCbWriter& operator<<(FCbWriter& Writer, const FCacheKey& Key)
-{
-	Writer.BeginObject();
-	Writer << ANSITEXTVIEW("Bucket") << Key.Bucket;
-	Writer << ANSITEXTVIEW("Hash") << Key.Hash;
-	Writer.EndObject();
-	return Writer;
-}
-
 void SerializeForLog(FCbWriter& Writer, const FCacheKey& Key)
 {
 	Writer.AddString(WriteToString<96>(Key));
 }
 
-bool LoadFromCompactBinary(const FCbFieldView Field, FCacheKey& OutKey)
-{
-	bool bOk = Field.IsObject();
-	bOk &= LoadFromCompactBinary(Field[ANSITEXTVIEW("Bucket")], OutKey.Bucket);
-	bOk &= LoadFromCompactBinary(Field[ANSITEXTVIEW("Hash")], OutKey.Hash);
-	return bOk;
-}
-
 FCacheKey ConvertLegacyCacheKey(const FStringView Key)
 {
 	FTCHARToUTF8 Utf8Key(Key);
diff --git a/Engine/Source/Developer/DerivedDataCache/Private/DerivedDataCacheMethod.cpp b/Engine/Source/Developer/DerivedDataCache/Private/DerivedDataCacheMethod.cpp
--- a/Engine/Source/Developer/DerivedDataCache/Private/DerivedDataCacheMethod.cpp
+++ b/Engine/Source/Developer/DerivedDataCache/Private/DerivedDataCacheMethod.cpp
@@ -5,8 +5,6 @@
 #include "Containers/StringConv.h"
 #include "Containers/StringView.h"
 #include "Misc/StringBuilder.h"
-#include "Serialization/CompactBinary.h"
-#include "Serialization/CompactBinaryWriter.h"
 #include "String/ParseTokens.h"
 
 namespace UE::DerivedData::Private
@@ -69,22 +67,6 @@ FUtf8StringBuilderBase& operator<<(FUtf8StringBuilderBase& Builder, ECacheMethod
 bool TryLexFromString(ECacheMethod& OutMethod, FUtf8StringView String) { return Private::MethodFromString(OutMethod, String); }
 bool TryLexFromString(ECacheMethod& OutMethod, FWideStringView String) { return Private::MethodFromString(OutMethod, String); }
 
-FCbWriter& operator<<(FCbWriter& Writer, const ECacheMethod Method)
-{
-	Writer.AddString(WriteToUtf8String<16>(Method));
-	return Writer;
-}
-
-bool LoadFromCompactBinary(FCbFieldView Field, ECacheMethod& OutMethod)
-{
-	if (TryLexFromString(OutMethod, Field.AsString()))
-	{
-		return true;
-	}
-	OutMethod = {};
-	return false;
-}
-
 FCacheMethodFilter FCacheMethodFilter::Parse(const FStringView MethodNames)
 {
 	FCacheMethodFilter MethodFilter;
diff --git a/Engine/Source/Developer/DerivedDataCache/Private/DerivedDataCachePolicy.cpp b/Engine/Source/Developer/DerivedDataCache/Private/DerivedDataCachePolicy.cpp
--- a/Engine/Source/Developer/DerivedDataCache/Private/DerivedDataCachePolicy.cpp
+++ b/Engine/Source/Developer/DerivedDataCache/Private/DerivedDataCachePolicy.cpp
@@ -7,8 +7,6 @@
 #include "Containers/StringConv.h"
 #include "Containers/StringView.h"
 #include "Misc/StringBuilder.h"
-#include "Serialization/CompactBinary.h"
-#include "Serialization/CompactBinaryWriter.h"
 #include "String/ParseTokens.h"
 #include "Templates/Function.h"
 
@@ -79,6 +77,8 @@ static TStringBuilderBase<CharType>& CachePolicyToString(TStringBuilderBase<Char
 template <typename CharType>
 static bool CachePolicyFromString(ECachePolicy& OutPolicy, const TStringView<CharType> String)
 {
+	OutPolicy = ECachePolicy::None;
+
 	if (String.IsEmpty())
 	{
 		return false;
@@ -141,39 +141,6 @@ ECachePolicy ParseCachePolicy(FAnsiStringView Text) { return Private::ParseCache
 ECachePolicy ParseCachePolicy(FWideStringView Text) { return Private::ParseCachePolicy(Text); }
 ECachePolicy ParseCachePolicy(FUtf8StringView Text) { return Private::ParseCachePolicy(Text); }
 
-FCbWriter& operator<<(FCbWriter& Writer, const ECachePolicy Policy)
-{
-	Writer.AddString(WriteToUtf8String<64>(Policy));
-	return Writer;
-}
-
-bool LoadFromCompactBinary(FCbFieldView Field, ECachePolicy& OutPolicy, const ECachePolicy Default)
-{
-	if (TryLexFromString(OutPolicy, Field.AsString()))
-	{
-		return true;
-	}
-	OutPolicy = Default;
-	return false;
-}
-
-FCbWriter& operator<<(FCbWriter& Writer, const FCacheValuePolicy& Policy)
-{
-	Writer.BeginObject();
-	Writer << ANSITEXTVIEW("Id") << Policy.Id;
-	Writer << ANSITEXTVIEW("Policy") << Policy.Policy;
-	Writer.EndObject();
-	return Writer;
-}
-
-bool LoadFromCompactBinary(FCbFieldView Field, FCacheValuePolicy& OutPolicy)
-{
-	bool bOk = Field.IsObject();
-	bOk &= LoadFromCompactBinary(Field[ANSITEXTVIEW("Id")], OutPolicy.Id);
-	bOk &= LoadFromCompactBinary(Field[ANSITEXTVIEW("Policy")], OutPolicy.Policy);
-	return bOk;
-}
-
 class Private::FCacheRecordPolicyShared final : public Private::ICacheRecordPolicyShared
 {
 public:
@@ -267,48 +234,4 @@ FCacheRecordPolicy FCacheRecordPolicyBuilder::Build()
 	return Policy;
 }
 
-FCbWriter& operator<<(FCbWriter& Writer, const FCacheRecordPolicy& Policy)
-{
-	Writer.BeginObject();
-	Writer << ANSITEXTVIEW("BasePolicy") << Policy.GetBasePolicy();
-	if (!Policy.IsUniform())
-	{
-		Writer.BeginArray(ANSITEXTVIEW("ValuePolicies"));
-		for (const FCacheValuePolicy& Value : Policy.GetValuePolicies())
-		{
-			Writer << Value;
-		}
-		Writer.EndArray();
-	}
-	Writer.EndObject();
-	return Writer;
-}
-
-bool LoadFromCompactBinary(FCbFieldView Field, FCacheRecordPolicy& OutPolicy)
-{
-	ECachePolicy BasePolicy;
-	if (!LoadFromCompactBinary(Field[ANSITEXTVIEW("BasePolicy")], BasePolicy))
-	{
-		OutPolicy = {};
-		return false;
-	}
-	FCacheRecordPolicyBuilder Builder(BasePolicy);
-
-	for (FCbFieldView Value : Field[ANSITEXTVIEW("ValuePolicies")])
-	{
-		FCacheValuePolicy ValuePolicy;
-		if (!LoadFromCompactBinary(Value, ValuePolicy) ||
-			ValuePolicy.Id.IsNull() ||
-			EnumHasAnyFlags(ValuePolicy.Policy, ~FCacheValuePolicy::PolicyMask))
-		{
-			OutPolicy = {};
-			return false;
-		}
-		Builder.AddValuePolicy(ValuePolicy);
-	}
-
-	OutPolicy = Builder.Build();
-	return true;
-}
-
 } // UE::DerivedData
diff --git a/Engine/Source/Developer/DerivedDataCache/Private/DerivedDataCachePrivate.h b/Engine/Source/Developer/DerivedDataCache/Private/DerivedDataCachePrivate.h
--- a/Engine/Source/Developer/DerivedDataCache/Private/DerivedDataCachePrivate.h
+++ b/Engine/Source/Developer/DerivedDataCache/Private/DerivedDataCachePrivate.h
@@ -4,26 +4,16 @@
 
 #include "CoreTypes.h"
 
-class FCbWriter;
-class FDerivedDataCacheInterface;
-
 template <typename FuncType> class TUniqueFunction;
 
 namespace UE::DerivedData { class FCacheRecord; }
-namespace UE::DerivedData { class ICache; }
 namespace UE::DerivedData { class IRequestOwner; }
-namespace UE::DerivedData { struct FCacheGetChunkRequest; }
-namespace UE::DerivedData { struct FCacheGetRequest; }
-namespace UE::DerivedData { struct FCacheGetValueRequest; }
 
 namespace UE::DerivedData::Private
 {
 
 // Implemented in DerivedDataCache.cpp
 void LaunchTaskInCacheThreadPool(IRequestOwner& Owner, TUniqueFunction<void ()>&& TaskBody);
-void SaveToCompactBinary(FCbWriter& Writer, const FCacheGetRequest& Request);
-void SaveToCompactBinary(FCbWriter& Writer, const FCacheGetValueRequest& Request);
-void SaveToCompactBinary(FCbWriter& Writer, const FCacheGetChunkRequest& Request);
 
 // Implemented in DerivedDataCacheRecord.cpp
 uint64 GetCacheRecordCompressedSize(const FCacheRecord& Record);
diff --git a/Engine/Source/Developer/DerivedDataCache/Private/DerivedDataCacheRecord.cpp b/Engine/Source/Developer/DerivedDataCache/Private/DerivedDataCacheRecord.cpp
--- a/Engine/Source/Developer/DerivedDataCache/Private/DerivedDataCacheRecord.cpp
+++ b/Engine/Source/Developer/DerivedDataCache/Private/DerivedDataCacheRecord.cpp
@@ -7,6 +7,7 @@
 #include "Algo/IsSorted.h"
 #include "DerivedDataCacheKey.h"
 #include "DerivedDataCachePrivate.h"
+#include "DerivedDataSerialization.h"
 #include "DerivedDataValue.h"
 #include "Misc/ScopeExit.h"
 #include "Misc/StringBuilder.h"
@@ -169,91 +170,25 @@ FCbPackage FCacheRecord::Save() const
 void FCacheRecord::Save(FCbPackage& Attachments, FCbWriter& Writer) const
 {
 	Writer.BeginObject();
-	Writer << ANSITEXTVIEW("Key") << GetKey();
-
-	if (const FCbObject& Meta = GetMeta())
-	{
-		Writer.AddObject(ANSITEXTVIEW("Meta"), Meta);
-	}
-	TConstArrayView<FValueWithId> Values = GetValues();
-	if (!Values.IsEmpty())
+	SaveToCompactBinary(Writer, *this, EMetaSaveFormat::Object, [&Attachments](FCbAttachment&& Attachment)
 	{
-		Writer.BeginArray(ANSITEXTVIEW("Values"));
-		for (const FValueWithId& Value : Values)
-		{
-			if (Value.HasData())
-			{
-				Attachments.AddAttachment(FCbAttachment(Value.GetData()));
-			}
-			Writer.BeginObject();
-			Writer.AddObjectId(ANSITEXTVIEW("Id"), Value.GetId());
-			Writer.AddBinaryAttachment(ANSITEXTVIEW("RawHash"), Value.GetRawHash());
-			Writer.AddInteger(ANSITEXTVIEW("RawSize"), Value.GetRawSize());
-			Writer.EndObject();
-		}
-		Writer.EndArray();
-	}
+		Attachments.AddAttachment(MoveTemp(Attachment));
+	});
 	Writer.EndObject();
 }
 
 FOptionalCacheRecord FCacheRecord::Load(const FCbPackage& Attachments, const FCbObject& Object)
 {
-	const FCbObjectView ObjectView = Object;
-
-	// Check for the previous format of cache record. Remove this check in 5.1.
-	if (ObjectView[ANSITEXTVIEW("Value")] || ObjectView[ANSITEXTVIEW("Attachments")])
-	{
-		return FOptionalCacheRecord();
-	}
-
-	FCacheKey Key;
-	if (!LoadFromCompactBinary(ObjectView[ANSITEXTVIEW("Key")], Key))
-	{
-		return FOptionalCacheRecord();
-	}
-
-	FCacheRecordBuilder Builder(Key);
-
-	Builder.SetMeta(Object[ANSITEXTVIEW("Meta")].AsObject());
-
-	auto LoadValue = [&Attachments](const FCbObjectView& ValueObject)
-	{
-		const FValueId Id = ValueObject[ANSITEXTVIEW("Id")].AsObjectId();
-		if (Id.IsNull())
-		{
-			return FValueWithId();
-		}
-		const FIoHash RawHash = ValueObject[ANSITEXTVIEW("RawHash")].AsHash();
-		if (const FCbAttachment* Attachment = Attachments.FindAttachment(RawHash))
-		{
-			const FCompressedBuffer& Compressed = Attachment->AsCompressedBinary();
-			if (Compressed.GetRawHash() == RawHash)
-			{
-				return FValueWithId(Id, Compressed);
-			}
-		}
-		const uint64 RawSize = ValueObject[ANSITEXTVIEW("RawSize")].AsUInt64(MAX_uint64);
-		if (!RawHash.IsZero() && RawSize != MAX_uint64)
-		{
-			return FValueWithId(Id, RawHash, RawSize);
-		}
-		else
-		{
-			return FValueWithId();
-		}
-	};
-
-	for (FCbFieldView ValueField : ObjectView[ANSITEXTVIEW("Values")])
+	FOptionalCacheRecord OptionalRecord;
+	LoadFromCompactBinary(Object.AsFieldView(), OptionalRecord, [&Attachments](const FIoHash& Hash)
 	{
-		FValueWithId Value = LoadValue(ValueField.AsObjectView());
-		if (!Value)
+		if (const FCbAttachment* Attachment = Attachments.FindAttachment(Hash))
 		{
-			return FOptionalCacheRecord();
+			return *Attachment;
 		}
-		Builder.AddValue(Value);
-	}
-
-	return Builder.Build();
+		return FCbAttachment();
+	});
+	return OptionalRecord;
 }
 
 FOptionalCacheRecord FCacheRecord::Load(const FCbPackage& Package)
diff --git a/Engine/Source/Developer/DerivedDataCache/Private/DerivedDataCacheReplay.cpp b/Engine/Source/Developer/DerivedDataCache/Private/DerivedDataCacheReplay.cpp
--- a/Engine/Source/Developer/DerivedDataCache/Private/DerivedDataCacheReplay.cpp
+++ b/Engine/Source/Developer/DerivedDataCache/Private/DerivedDataCacheReplay.cpp
@@ -11,6 +11,7 @@
 #include "DerivedDataCachePrivate.h"
 #include "DerivedDataLegacyCacheStore.h"
 #include "DerivedDataRequestOwner.h"
+#include "DerivedDataSerialization.h"
 #include "DerivedDataValue.h"
 #include "HAL/CriticalSection.h"
 #include "HAL/FileManager.h"
@@ -98,10 +99,6 @@ public:
 		EPriority Priority);
 
 private:
-	static void SaveValuesToCompactBinary(FCbWriter& Writer, const FCacheGetResponse& Response);
-	static void SaveValuesToCompactBinary(FCbWriter& Writer, const FCacheGetValueResponse& Response);
-	static void SaveValuesToCompactBinary(FCbWriter& Writer, const FCacheGetChunkResponse& Response);
-
 	void WriteBinaryToArchive(const FCompositeBuffer& RawBinary);
 	void WriteToArchive(FCbWriter& Writer);
 	void FlushToArchive();
@@ -266,10 +263,7 @@ void FCacheStoreReplay::SerializeRequests(
 		++Index;
 		if (KeyFilter.IsMatch(Request.Key))
 		{
-			Writer.BeginObject();
-			Private::SaveToCompactBinary(Writer, Request);
-			SaveValuesToCompactBinary(Writer, Responses[Index]);
-			Writer.EndObject();
+			SaveToCompactBinary(Writer, Request, &Responses[Index]);
 		}
 	}
 	Writer.EndArray();
@@ -285,45 +279,6 @@ void FCacheStoreReplay::SerializeRequests(
 	WriteToArchive(Writer);
 }
 
-void FCacheStoreReplay::SaveValuesToCompactBinary(FCbWriter& Writer, const FCacheGetResponse& Response)
-{
-	if (Response.Status == EStatus::Ok)
-	{
-		Writer.BeginArray(ANSITEXTVIEW("Values"));
-		for (const FValueWithId& Value : Response.Record.GetValues())
-		{
-			Writer.BeginObject();
-			Writer << ANSITEXTVIEW("Id") << Value.GetId();
-			Writer << ANSITEXTVIEW("RawHash") << Value.GetRawHash();
-			Writer << ANSITEXTVIEW("RawSize") << Value.GetRawSize();
-			Writer.EndObject();
-		}
-		Writer.EndArray();
-	}
-}
-
-void FCacheStoreReplay::SaveValuesToCompactBinary(FCbWriter& Writer, const FCacheGetValueResponse& Response)
-{
-	if (Response.Status == EStatus::Ok)
-	{
-		Writer.BeginObject(ANSITEXTVIEW("Value"));
-		Writer << ANSITEXTVIEW("RawHash") << Response.Value.GetRawHash();
-		Writer << ANSITEXTVIEW("RawSize") << Response.Value.GetRawSize();
-		Writer.EndObject();
-	}
-}
-
-void FCacheStoreReplay::SaveValuesToCompactBinary(FCbWriter& Writer, const FCacheGetChunkResponse& Response)
-{
-	if (Response.Status == EStatus::Ok)
-	{
-		Writer.BeginObject(ANSITEXTVIEW("Value"));
-		Writer << ANSITEXTVIEW("RawHash") << Response.RawHash;
-		Writer << ANSITEXTVIEW("RawSize") << Response.RawSize;
-		Writer.EndObject();
-	}
-}
-
 void FCacheStoreReplay::WriteBinaryToArchive(const FCompositeBuffer& RawBinary)
 {
 	const FValue CompressedBinary = FValue::Compress(RawBinary, RawBlock.GetSize());
diff --git a/Engine/Source/Developer/DerivedDataCache/Private/DerivedDataRequestTypes.cpp b/Engine/Source/Developer/DerivedDataCache/Private/DerivedDataRequestTypes.cpp
--- a/Engine/Source/Developer/DerivedDataCache/Private/DerivedDataRequestTypes.cpp
+++ b/Engine/Source/Developer/DerivedDataCache/Private/DerivedDataRequestTypes.cpp
@@ -6,8 +6,6 @@
 #include "Containers/StringView.h"
 #include "Misc/QueuedThreadPool.h"
 #include "Misc/StringBuilder.h"
-#include "Serialization/CompactBinary.h"
-#include "Serialization/CompactBinaryWriter.h"
 
 namespace UE::DerivedData::Private
 {
@@ -109,22 +107,6 @@ FUtf8StringBuilderBase& operator<<(FUtf8StringBuilderBase& Builder, EPriority Pr
 bool TryLexFromString(EPriority& OutPriority, FUtf8StringView String) { return Private::PriorityFromString(OutPriority, String); }
 bool TryLexFromString(EPriority& OutPriority, FWideStringView String) { return Private::PriorityFromString(OutPriority, String); }
 
-FCbWriter& operator<<(FCbWriter& Writer, const EPriority Priority)
-{
-	Writer.AddString(WriteToUtf8String<16>(Priority));
-	return Writer;
-}
-
-bool LoadFromCompactBinary(FCbFieldView Field, EPriority& OutPriority, const EPriority Default)
-{
-	if (TryLexFromString(OutPriority, Field.AsString()))
-	{
-		return true;
-	}
-	OutPriority = Default;
-	return false;
-}
-
 EQueuedWorkPriority ConvertToQueuedWorkPriority(EPriority Priority)
 {
 	switch (Priority)
@@ -160,20 +142,4 @@ FUtf8StringBuilderBase& operator<<(FUtf8StringBuilderBase& Builder, EStatus Stat
 bool TryLexFromString(EStatus& OutStatus, FUtf8StringView String) { return Private::StatusFromString(OutStatus, String); }
 bool TryLexFromString(EStatus& OutStatus, FWideStringView String) { return Private::StatusFromString(OutStatus, String); }
 
-FCbWriter& operator<<(FCbWriter& Writer, const EStatus Status)
-{
-	Writer.AddString(WriteToUtf8String<16>(Status));
-	return Writer;
-}
-
-bool LoadFromCompactBinary(FCbFieldView Field, EStatus& OutStatus, const EStatus Default)
-{
-	if (TryLexFromString(OutStatus, Field.AsString()))
-	{
-		return true;
-	}
-	OutStatus = Default;
-	return false;
-}
-
 } // UE::DerivedData
diff --git a/Engine/Source/Developer/DerivedDataCache/Private/DerivedDataSerialization.cpp b/Engine/Source/Developer/DerivedDataCache/Private/DerivedDataSerialization.cpp
new file mode 100644
--- /dev/null
+++ b/Engine/Source/Developer/DerivedDataCache/Private/DerivedDataSerialization.cpp
@@ -0,0 +1,684 @@
+// Copyright Epic Games, Inc. All Rights Reserved.
+
+#include "DerivedDataSerialization.h"
+
+#include "DerivedDataCache.h"
+#include "DerivedDataCacheKey.h"
+#include "DerivedDataCacheMethod.h"
+#include "DerivedDataCachePolicy.h"
+#include "DerivedDataCacheRecord.h"
+#include "DerivedDataValue.h"
+#include "DerivedDataValueId.h"
+#include "Misc/Optional.h"
+#include "Misc/StringBuilder.h"
+#include "Serialization/CompactBinary.h"
+#include "Serialization/CompactBinaryPackage.h"
+#include "Serialization/CompactBinarySerialization.h"
+#include "Serialization/CompactBinaryWriter.h"
+
+namespace UE::DerivedData
+{
+
+FCbWriter& operator<<(FCbWriter& Writer, const FCacheGetRequest& Request)
+{
+	SaveToCompactBinary(Writer, Request);
+	return Writer;
+}
+
+FCbWriter& operator<<(FCbWriter& Writer, const FCacheGetValueRequest& Request)
+{
+	SaveToCompactBinary(Writer, Request);
+	return Writer;
+}
+
+FCbWriter& operator<<(FCbWriter& Writer, const FCacheGetChunkRequest& Request)
+{
+	SaveToCompactBinary(Writer, Request);
+	return Writer;
+}
+
+///////////////////////////////////////////////////////////////////////////////////////////////////////////////////////
+
+void SaveMetaToCompactBinary(FCbWriter& Writer, const FCbObject& Meta, EMetaSaveFormat MetaFormat, FSaveAttachmentFn SaveAttachment)
+{
+	if (Meta)
+	{
+		if (MetaFormat == EMetaSaveFormat::Object)
+		{
+			Writer.AddObject(ANSITEXTVIEW("Meta"), Meta);
+		}
+		else
+		{
+			FCbAttachment MetaAttachment(Meta);
+			Writer.AddObjectAttachment(ANSITEXTVIEW("MetaHash"), MetaAttachment.GetHash());
+			Writer.AddInteger(ANSITEXTVIEW("MetaSize"), MetaAttachment.AsObject().GetSize());
+			if (SaveAttachment)
+			{
+				SaveAttachment(MoveTemp(MetaAttachment));
+			}
+		}
+	}
+}
+
+bool LoadMetaFromCompactBinary(FCbFieldView Field, FCbObject& OutMeta, FLoadAttachmentFn LoadAttachment)
+{
+	if (const FCbObjectView MetaView = Field[ANSITEXTVIEW("Meta")].AsObjectView())
+	{
+		OutMeta = FCbObject::Clone(MetaView);
+		return true;
+	}
+	else if (LoadAttachment)
+	{
+		const FIoHash MetaHash = Field[ANSITEXTVIEW("MetaHash")].AsHash();
+		const uint64 MetaSize = Field[ANSITEXTVIEW("MetaSize")].AsUInt64(MAX_uint64);
+		if (!MetaHash.IsZero() && MetaSize != MAX_uint64)
+		{
+			FCbObject Meta = LoadAttachment(MetaHash).AsObject();
+			if (Meta.GetSize() == MetaSize)
+			{
+				OutMeta = MoveTemp(Meta);
+				return true;
+			}
+		}
+	}
+	return false;
+}
+
+///////////////////////////////////////////////////////////////////////////////////////////////////////////////////////
+
+void SaveToCompactBinary(FCbWriter& Writer, const FCachePutRequest& Request, FSaveAttachmentFn SaveAttachment)
+{
+	Writer.BeginObject();
+	if (!Request.Name.IsEmpty())
+	{
+		Writer << ANSITEXTVIEW("Name") << Request.Name;
+	}
+	SaveToCompactBinary(Writer, Request.Record, EMetaSaveFormat::Attachment, SaveAttachment);
+	if (!Request.Policy.IsDefault())
+	{
+		Writer << ANSITEXTVIEW("Policy") << Request.Policy;
+	}
+	if (Request.UserData != 0)
+	{
+		Writer << ANSITEXTVIEW("UserData") << Request.UserData;
+	}
+	Writer.EndObject();
+}
+
+bool LoadFromCompactBinary(FCbFieldView Field, TOptional<FCachePutRequest>& OutRequest, FLoadAttachmentFn LoadAttachment)
+{
+	FOptionalCacheRecord Record;
+	if (!LoadFromCompactBinary(Field, Record, LoadAttachment))
+	{
+		OutRequest.Reset();
+		return false;
+	}
+	OutRequest = FCachePutRequest{.Record = MoveTemp(Record).Get()};
+	bool bOk = Field.IsObject();
+	LoadFromCompactBinary(Field[ANSITEXTVIEW("Name")], OutRequest->Name);
+	LoadFromCompactBinary(Field[ANSITEXTVIEW("Policy")], OutRequest->Policy);
+	LoadFromCompactBinary(Field[ANSITEXTVIEW("UserData")], OutRequest->UserData, 0);
+	return bOk;
+}
+
+void SaveToCompactBinary(FCbWriter& Writer, const FCacheGetRequest& Request, const FCacheGetResponse* Response, FSaveAttachmentFn SaveAttachment)
+{
+	check(!Response || Request.Key == Response->Record.GetKey());
+	check(!Response || Request.UserData == Response->UserData);
+	Writer.BeginObject();
+	if (!Request.Name.IsEmpty())
+	{
+		Writer << ANSITEXTVIEW("Name") << Request.Name;
+	}
+	Writer << ANSITEXTVIEW("Key") << Request.Key;
+	if (Response)
+	{
+		SaveMetaToCompactBinary(Writer, Response->Record.GetMeta(), EMetaSaveFormat::Attachment, SaveAttachment);
+		if (const TConstArrayView<FValueWithId> Values = Response->Record.GetValues(); !Values.IsEmpty())
+		{
+			Writer.BeginArray(ANSITEXTVIEW("Values"));
+			for (const FValueWithId& Value : Values)
+			{
+				SaveToCompactBinary(Writer, Value, SaveAttachment);
+			}
+			Writer.EndArray();
+		}
+	}
+	if (!Request.Policy.IsDefault())
+	{
+		Writer << ANSITEXTVIEW("Policy") << Request.Policy;
+	}
+	if (Request.UserData != 0)
+	{
+		Writer << ANSITEXTVIEW("UserData") << Request.UserData;
+	}
+	if (Response && Response->Status != EStatus::Ok)
+	{
+		Writer << ANSITEXTVIEW("Status") << Response->Status;
+	}
+	Writer.EndObject();
+}
+
+bool LoadFromCompactBinary(FCbFieldView Field, FCacheGetRequest& OutRequest)
+{
+	bool bOk = Field.IsObject();
+	LoadFromCompactBinary(Field[ANSITEXTVIEW("Name")], OutRequest.Name);
+	bOk &= LoadFromCompactBinary(Field[ANSITEXTVIEW("Key")], OutRequest.Key);
+	LoadFromCompactBinary(Field[ANSITEXTVIEW("Policy")], OutRequest.Policy);
+	LoadFromCompactBinary(Field[ANSITEXTVIEW("UserData")], OutRequest.UserData, 0);
+	return bOk;
+}
+
+bool LoadFromCompactBinary(FCbFieldView Field, TOptional<FCacheGetResponse>& OutResponse, FLoadAttachmentFn LoadAttachment)
+{
+	FOptionalCacheRecord Record;
+	if (!LoadFromCompactBinary(Field, Record, LoadAttachment))
+	{
+		return false;
+	}
+	OutResponse = FCacheGetResponse{.Record = MoveTemp(Record).Get()};
+	bool bOk = Field.IsObject();
+	LoadFromCompactBinary(Field[ANSITEXTVIEW("Name")], OutResponse->Name);
+	LoadFromCompactBinary(Field[ANSITEXTVIEW("UserData")], OutResponse->UserData, 0);
+	LoadFromCompactBinary(Field[ANSITEXTVIEW("Status")], OutResponse->Status, EStatus::Ok);
+	return bOk;
+}
+
+///////////////////////////////////////////////////////////////////////////////////////////////////////////////////////
+
+void SaveToCompactBinary(FCbWriter& Writer, const FCachePutValueRequest& Request, FSaveAttachmentFn SaveAttachment)
+{
+	Writer.BeginObject();
+	if (!Request.Name.IsEmpty())
+	{
+		Writer << ANSITEXTVIEW("Name") << Request.Name;
+	}
+	Writer << ANSITEXTVIEW("Key") << Request.Key;
+	if (Request.Value != FValue::Null)
+	{
+		Writer.SetName(ANSITEXTVIEW("Value"));
+		SaveToCompactBinary(Writer, Request.Value, SaveAttachment);
+	}
+	if (Request.Policy != ECachePolicy::Default)
+	{
+		Writer << ANSITEXTVIEW("Policy") << Request.Policy;
+	}
+	if (Request.UserData != 0)
+	{
+		Writer << ANSITEXTVIEW("UserData") << Request.UserData;
+	}
+	Writer.EndObject();
+}
+
+bool LoadFromCompactBinary(FCbFieldView Field, FCachePutValueRequest& OutRequest, FLoadAttachmentFn LoadAttachment)
+{
+	bool bOk = Field.IsObject();
+	LoadFromCompactBinary(Field[ANSITEXTVIEW("Name")], OutRequest.Name);
+	bOk &= LoadFromCompactBinary(Field[ANSITEXTVIEW("Key")], OutRequest.Key);
+	LoadFromCompactBinary(Field[ANSITEXTVIEW("Value")], OutRequest.Value, LoadAttachment);
+	LoadFromCompactBinary(Field[ANSITEXTVIEW("Policy")], OutRequest.Policy, ECachePolicy::Default);
+	LoadFromCompactBinary(Field[ANSITEXTVIEW("UserData")], OutRequest.UserData, 0);
+	return bOk;
+}
+
+void SaveToCompactBinary(FCbWriter& Writer, const FCacheGetValueRequest& Request, const FCacheGetValueResponse* Response, FSaveAttachmentFn SaveAttachment)
+{
+	check(!Response || Request.Key == Response->Key);
+	check(!Response || Request.UserData == Response->UserData);
+	Writer.BeginObject();
+	if (!Request.Name.IsEmpty())
+	{
+		Writer << ANSITEXTVIEW("Name") << Request.Name;
+	}
+	Writer << ANSITEXTVIEW("Key") << Request.Key;
+	if (Response && Response->Value != FValue::Null)
+	{
+		Writer.SetName(ANSITEXTVIEW("Value"));
+		SaveToCompactBinary(Writer, Response->Value, SaveAttachment);
+	}
+	if (Request.Policy != ECachePolicy::Default)
+	{
+		Writer << ANSITEXTVIEW("Policy") << Request.Policy;
+	}
+	if (Request.UserData != 0)
+	{
+		Writer << ANSITEXTVIEW("UserData") << Request.UserData;
+	}
+	if (Response && Response->Status != EStatus::Ok)
+	{
+		Writer << ANSITEXTVIEW("Status") << Response->Status;
+	}
+	Writer.EndObject();
+}
+
+bool LoadFromCompactBinary(FCbFieldView Field, FCacheGetValueRequest& OutRequest)
+{
+	bool bOk = Field.IsObject();
+	LoadFromCompactBinary(Field[ANSITEXTVIEW("Name")], OutRequest.Name);
+	bOk &= LoadFromCompactBinary(Field[ANSITEXTVIEW("Key")], OutRequest.Key);
+	LoadFromCompactBinary(Field[ANSITEXTVIEW("Policy")], OutRequest.Policy, ECachePolicy::Default);
+	LoadFromCompactBinary(Field[ANSITEXTVIEW("UserData")], OutRequest.UserData, 0);
+	return bOk;
+}
+
+bool LoadFromCompactBinary(FCbFieldView Field, FCacheGetValueResponse& OutResponse, FLoadAttachmentFn LoadAttachment)
+{
+	bool bOk = Field.IsObject();
+	LoadFromCompactBinary(Field[ANSITEXTVIEW("Name")], OutResponse.Name);
+	bOk &= LoadFromCompactBinary(Field[ANSITEXTVIEW("Key")], OutResponse.Key);
+	LoadFromCompactBinary(Field[ANSITEXTVIEW("Value")], OutResponse.Value, LoadAttachment);
+	LoadFromCompactBinary(Field[ANSITEXTVIEW("UserData")], OutResponse.UserData, 0);
+	LoadFromCompactBinary(Field[ANSITEXTVIEW("Status")], OutResponse.Status, EStatus::Ok);
+	return bOk;
+}
+
+///////////////////////////////////////////////////////////////////////////////////////////////////////////////////////
+
+void SaveToCompactBinary(FCbWriter& Writer, const FCacheGetChunkRequest& Request, const FCacheGetChunkResponse* Response, FSaveAttachmentFn SaveAttachment)
+{
+	check(!Response || Request.Key == Response->Key);
+	check(!Response || Request.UserData == Response->UserData);
+	Writer.BeginObject();
+	if (!Request.Name.IsEmpty())
+	{
+		Writer << ANSITEXTVIEW("Name") << Request.Name;
+	}
+	Writer << ANSITEXTVIEW("Key") << Request.Key;
+	if (Request.Id.IsValid())
+	{
+		Writer << ANSITEXTVIEW("Id") << Request.Id;
+	}
+	if (Request.RawOffset != 0)
+	{
+		Writer << ANSITEXTVIEW("RawOffset") << Request.RawOffset;
+	}
+	if (Request.RawSize != MAX_uint64)
+	{
+		Writer << ANSITEXTVIEW("RawSize") << Request.RawSize;
+	}
+	if (!Request.RawHash.IsZero())
+	{
+		Writer << ANSITEXTVIEW("RawHash") << Request.RawHash;
+	}
+	if (Response && !Response->RawHash.IsZero())
+	{
+		Writer << ANSITEXTVIEW("Value") << FValue(Response->RawHash, Response->RawSize);
+		if (SaveAttachment && Response->RawData)
+		{
+			const FIoHash FragmentHash = FIoHash::HashBuffer(Response->RawData);
+			FCbAttachment Attachment(Response->RawData, FragmentHash);
+			if (FragmentHash != Response->RawHash)
+			{
+				Writer.AddAttachment(ANSITEXTVIEW("FragmentHash"), Attachment);
+			}
+			SaveAttachment(MoveTemp(Attachment));
+		}
+	}
+	if (Request.Policy != ECachePolicy::Default)
+	{
+		Writer << ANSITEXTVIEW("Policy") << Request.Policy;
+	}
+	if (Request.UserData != 0)
+	{
+		Writer << ANSITEXTVIEW("UserData") << Request.UserData;
+	}
+	if (Response && Response->Status != EStatus::Ok)
+	{
+		Writer << ANSITEXTVIEW("Status") << Response->Status;
+	}
+	Writer.EndObject();
+}
+
+bool LoadFromCompactBinary(FCbFieldView Field, FCacheGetChunkRequest& OutRequest)
+{
+	bool bOk = Field.IsObject();
+	LoadFromCompactBinary(Field[ANSITEXTVIEW("Name")], OutRequest.Name);
+	bOk &= LoadFromCompactBinary(Field[ANSITEXTVIEW("Key")], OutRequest.Key);
+	LoadFromCompactBinary(Field[ANSITEXTVIEW("Id")], OutRequest.Id);
+	LoadFromCompactBinary(Field[ANSITEXTVIEW("RawOffset")], OutRequest.RawOffset, 0);
+	LoadFromCompactBinary(Field[ANSITEXTVIEW("RawSize")], OutRequest.RawSize, MAX_uint64);
+	LoadFromCompactBinary(Field[ANSITEXTVIEW("RawHash")], OutRequest.RawHash);
+	LoadFromCompactBinary(Field[ANSITEXTVIEW("Policy")], OutRequest.Policy, ECachePolicy::Default);
+	LoadFromCompactBinary(Field[ANSITEXTVIEW("UserData")], OutRequest.UserData, 0);
+	return bOk;
+}
+
+bool LoadFromCompactBinary(FCbFieldView Field, FCacheGetChunkResponse& OutResponse, FLoadAttachmentFn LoadAttachment)
+{
+	FValue Value;
+	bool bOk = Field.IsObject();
+	LoadFromCompactBinary(Field[ANSITEXTVIEW("Name")], OutResponse.Name);
+	bOk &= LoadFromCompactBinary(Field[ANSITEXTVIEW("Key")], OutResponse.Key);
+	LoadFromCompactBinary(Field[ANSITEXTVIEW("Id")], OutResponse.Id);
+	LoadFromCompactBinary(Field[ANSITEXTVIEW("RawOffset")], OutResponse.RawOffset, 0);
+	LoadFromCompactBinary(Field[ANSITEXTVIEW("UserData")], OutResponse.UserData, 0);
+	LoadFromCompactBinary(Field[ANSITEXTVIEW("Value")], Value);
+	LoadFromCompactBinary(Field[ANSITEXTVIEW("Status")], OutResponse.Status, EStatus::Ok);
+	OutResponse.RawHash = Value.GetRawHash();
+	OutResponse.RawSize = Value.GetRawSize();
+	if (LoadAttachment && !OutResponse.RawHash.IsZero())
+	{
+		const FIoHash RawHash = Field[ANSITEXTVIEW("FragmentHash")].AsHash(OutResponse.RawHash);
+		const FCbAttachment Attachment = LoadAttachment(RawHash);
+		if (FSharedBuffer Data = Attachment.AsBinary())
+		{
+			OutResponse.RawData = MoveTemp(Data);
+		}
+		else if (const FCompressedBuffer& CompressedData = Attachment.AsCompressedBinary())
+		{
+			if (CompressedData.GetRawHash() == OutResponse.RawHash && CompressedData.GetRawSize() == OutResponse.RawSize)
+			{
+				OutResponse.RawData = CompressedData.Decompress();
+			}
+		}
+	}
+	return bOk;
+}
+
+///////////////////////////////////////////////////////////////////////////////////////////////////////////////////////
+
+FCbWriter& operator<<(FCbWriter& Writer, const FCacheBucket Bucket)
+{
+	Writer.AddString(Bucket.ToString());
+	return Writer;
+}
+
+bool LoadFromCompactBinary(FCbFieldView Field, FCacheBucket& OutBucket)
+{
+	if (const FUtf8StringView Bucket = Field.AsString(); !Field.HasError() && FCacheBucket::IsValidName(Bucket))
+	{
+		OutBucket = FCacheBucket(Bucket);
+		return true;
+	}
+	OutBucket.Reset();
+	return false;
+}
+
+FCbWriter& operator<<(FCbWriter& Writer, const FCacheKey& Key)
+{
+	Writer.BeginObject();
+	Writer << ANSITEXTVIEW("Bucket") << Key.Bucket;
+	Writer << ANSITEXTVIEW("Hash") << Key.Hash;
+	Writer.EndObject();
+	return Writer;
+}
+
+bool LoadFromCompactBinary(const FCbFieldView Field, FCacheKey& OutKey)
+{
+	bool bOk = Field.IsObject();
+	bOk &= LoadFromCompactBinary(Field[ANSITEXTVIEW("Bucket")], OutKey.Bucket);
+	bOk &= LoadFromCompactBinary(Field[ANSITEXTVIEW("Hash")], OutKey.Hash);
+	return bOk;
+}
+
+///////////////////////////////////////////////////////////////////////////////////////////////////////////////////////
+
+FCbWriter& operator<<(FCbWriter& Writer, const ECacheMethod Method)
+{
+	Writer.AddString(WriteToUtf8String<16>(Method));
+	return Writer;
+}
+
+bool LoadFromCompactBinary(FCbFieldView Field, ECacheMethod& OutMethod)
+{
+	if (TryLexFromString(OutMethod, Field.AsString()))
+	{
+		return true;
+	}
+	OutMethod = {};
+	return false;
+}
+
+///////////////////////////////////////////////////////////////////////////////////////////////////////////////////////
+
+FCbWriter& operator<<(FCbWriter& Writer, const ECachePolicy Policy)
+{
+	Writer.AddString(WriteToUtf8String<64>(Policy));
+	return Writer;
+}
+
+bool LoadFromCompactBinary(FCbFieldView Field, ECachePolicy& OutPolicy, const ECachePolicy Default)
+{
+	if (TryLexFromString(OutPolicy, Field.AsString()))
+	{
+		return true;
+	}
+	OutPolicy = Default;
+	return false;
+}
+
+FCbWriter& operator<<(FCbWriter& Writer, const FCacheValuePolicy& Policy)
+{
+	Writer.BeginObject();
+	Writer << ANSITEXTVIEW("Id") << Policy.Id;
+	Writer << ANSITEXTVIEW("Policy") << Policy.Policy;
+	Writer.EndObject();
+	return Writer;
+}
+
+bool LoadFromCompactBinary(FCbFieldView Field, FCacheValuePolicy& OutPolicy)
+{
+	bool bOk = Field.IsObject();
+	bOk &= LoadFromCompactBinary(Field[ANSITEXTVIEW("Id")], OutPolicy.Id);
+	bOk &= LoadFromCompactBinary(Field[ANSITEXTVIEW("Policy")], OutPolicy.Policy, ECachePolicy::Default);
+	return bOk;
+}
+
+FCbWriter& operator<<(FCbWriter& Writer, const FCacheRecordPolicy& Policy)
+{
+	Writer.BeginObject();
+	Writer << ANSITEXTVIEW("BasePolicy") << Policy.GetBasePolicy();
+	if (!Policy.IsUniform())
+	{
+		Writer.BeginArray(ANSITEXTVIEW("ValuePolicies"));
+		for (const FCacheValuePolicy& Value : Policy.GetValuePolicies())
+		{
+			Writer << Value;
+		}
+		Writer.EndArray();
+	}
+	Writer.EndObject();
+	return Writer;
+}
+
+bool LoadFromCompactBinary(FCbFieldView Field, FCacheRecordPolicy& OutPolicy)
+{
+	ECachePolicy BasePolicy;
+	if (!LoadFromCompactBinary(Field[ANSITEXTVIEW("BasePolicy")], BasePolicy, ECachePolicy::Default))
+	{
+		OutPolicy = {};
+		return false;
+	}
+	FCacheRecordPolicyBuilder Builder(BasePolicy);
+
+	for (FCbFieldView Value : Field[ANSITEXTVIEW("ValuePolicies")])
+	{
+		FCacheValuePolicy ValuePolicy;
+		if (!LoadFromCompactBinary(Value, ValuePolicy) ||
+			ValuePolicy.Id.IsNull() ||
+			EnumHasAnyFlags(ValuePolicy.Policy, ~FCacheValuePolicy::PolicyMask))
+		{
+			OutPolicy = {};
+			return false;
+		}
+		Builder.AddValuePolicy(ValuePolicy);
+	}
+
+	OutPolicy = Builder.Build();
+	return true;
+}
+
+///////////////////////////////////////////////////////////////////////////////////////////////////////////////////////
+
+void SaveToCompactBinary(FCbWriter& Writer, const FCacheRecord& Record, EMetaSaveFormat MetaFormat, FSaveAttachmentFn SaveAttachment)
+{
+	Writer << ANSITEXTVIEW("Key") << Record.GetKey();
+
+	SaveMetaToCompactBinary(Writer, Record.GetMeta(), MetaFormat, SaveAttachment);
+
+	if (const TConstArrayView<FValueWithId> Values = Record.GetValues(); !Values.IsEmpty())
+	{
+		Writer.BeginArray(ANSITEXTVIEW("Values"));
+		for (const FValueWithId& Value : Values)
+		{
+			SaveToCompactBinary(Writer, Value, SaveAttachment);
+		}
+		Writer.EndArray();
+	}
+}
+
+bool LoadFromCompactBinary(FCbFieldView Field, FOptionalCacheRecord& OutRecord, FLoadAttachmentFn LoadAttachment)
+{
+	FCacheKey Key;
+	if (!LoadFromCompactBinary(Field[ANSITEXTVIEW("Key")], Key))
+	{
+		return false;
+	}
+
+	FCacheRecordBuilder Builder(Key);
+
+	FCbObject Meta;
+	if (LoadMetaFromCompactBinary(Field, Meta, LoadAttachment))
+	{
+		Builder.SetMeta(MoveTemp(Meta));
+	}
+
+	for (FCbFieldView ValueField : Field[ANSITEXTVIEW("Values")])
+	{
+		FValueWithId Value;
+		if (!LoadFromCompactBinary(ValueField, Value, LoadAttachment))
+		{
+			return false;
+		}
+		Builder.AddValue(Value);
+	}
+
+	OutRecord = Builder.Build();
+	return true;
+}
+
+///////////////////////////////////////////////////////////////////////////////////////////////////////////////////////
+
+FCbWriter& operator<<(FCbWriter& Writer, const EPriority Priority)
+{
+	Writer.AddString(WriteToUtf8String<16>(Priority));
+	return Writer;
+}
+
+bool LoadFromCompactBinary(FCbFieldView Field, EPriority& OutPriority, const EPriority Default)
+{
+	if (TryLexFromString(OutPriority, Field.AsString()))
+	{
+		return true;
+	}
+	OutPriority = Default;
+	return false;
+}
+
+FCbWriter& operator<<(FCbWriter& Writer, const EStatus Status)
+{
+	Writer.AddString(WriteToUtf8String<16>(Status));
+	return Writer;
+}
+
+bool LoadFromCompactBinary(FCbFieldView Field, EStatus& OutStatus, const EStatus Default)
+{
+	if (TryLexFromString(OutStatus, Field.AsString()))
+	{
+		return true;
+	}
+	OutStatus = Default;
+	return false;
+}
+
+///////////////////////////////////////////////////////////////////////////////////////////////////////////////////////
+
+FCbWriter& operator<<(FCbWriter& Writer, const FValue& Value)
+{
+	SaveToCompactBinary(Writer, Value);
+	return Writer;
+}
+
+void SaveToCompactBinary(FCbWriter& Writer, const FValue& Value, FSaveAttachmentFn SaveAttachment)
+{
+	if (SaveAttachment && Value.GetData())
+	{
+		SaveAttachment(FCbAttachment(Value.GetData()));
+	}
+	Writer.BeginObject();
+	Writer.AddBinaryAttachment(ANSITEXTVIEW("RawHash"), Value.GetRawHash());
+	Writer.AddInteger(ANSITEXTVIEW("RawSize"), Value.GetRawSize());
+	Writer.EndObject();
+}
+
+bool LoadFromCompactBinary(FCbFieldView Field, FValue& OutValue, FLoadAttachmentFn LoadAttachment)
+{
+	FIoHash RawHash;
+	uint64 RawSize = 0;
+	bool bOk = Field.IsObject();
+	bOk &= LoadFromCompactBinary(Field[ANSITEXTVIEW("RawHash")], RawHash);
+	bOk &= LoadFromCompactBinary(Field[ANSITEXTVIEW("RawSize")], RawSize);
+	if (LoadAttachment)
+	{
+		const FCbAttachment Attachment = LoadAttachment(RawHash);
+		const FCompressedBuffer& Compressed = Attachment.AsCompressedBinary();
+		if (Compressed.GetRawHash() == RawHash && Compressed.GetRawSize() == RawSize)
+		{
+			OutValue = FValue(Compressed);
+			return bOk;
+		}
+	}
+	OutValue = FValue(RawHash, RawSize);
+	return bOk;
+}
+
+FCbWriter& operator<<(FCbWriter& Writer, const FValueWithId& Value)
+{
+	SaveToCompactBinary(Writer, Value);
+	return Writer;
+}
+
+void SaveToCompactBinary(FCbWriter& Writer, const FValueWithId& Value, FSaveAttachmentFn SaveAttachment)
+{
+	if (SaveAttachment && Value.GetData())
+	{
+		SaveAttachment(FCbAttachment(Value.GetData()));
+	}
+	Writer.BeginObject();
+	Writer.AddObjectId(ANSITEXTVIEW("Id"), Value.GetId());
+	Writer.AddBinaryAttachment(ANSITEXTVIEW("RawHash"), Value.GetRawHash());
+	Writer.AddInteger(ANSITEXTVIEW("RawSize"), Value.GetRawSize());
+	Writer.EndObject();
+}
+
+bool LoadFromCompactBinary(FCbFieldView Field, FValueWithId& OutValue, FLoadAttachmentFn LoadAttachment)
+{
+	FValueId Id;
+	FValue Value;
+	bool bOk = Field.IsObject();
+	bOk &= LoadFromCompactBinary(Field[ANSITEXTVIEW("Id")], Id);
+	bOk &= LoadFromCompactBinary(Field, Value, LoadAttachment);
+	OutValue = FValueWithId(Id, MoveTemp(Value));
+	return bOk;
+}
+
+///////////////////////////////////////////////////////////////////////////////////////////////////////////////////////
+
+FCbWriter& operator<<(FCbWriter& Writer, const FValueId& Id)
+{
+	Writer.AddObjectId(Id);
+	return Writer;
+}
+
+bool LoadFromCompactBinary(FCbFieldView Field, FValueId& OutId)
+{
+	if (const FCbObjectId ObjectId = Field.AsObjectId(); !Field.HasError())
+	{
+		OutId = ObjectId;
+		return true;
+	}
+	OutId.Reset();
+	return false;
+}
+
+} // UE::DerivedData
diff --git a/Engine/Source/Developer/DerivedDataCache/Private/DerivedDataValue.cpp b/Engine/Source/Developer/DerivedDataCache/Private/DerivedDataValue.cpp
--- a/Engine/Source/Developer/DerivedDataCache/Private/DerivedDataValue.cpp
+++ b/Engine/Source/Developer/DerivedDataCache/Private/DerivedDataValue.cpp
@@ -7,8 +7,6 @@
 #include "Hash/xxhash.h"
 #include "Misc/ByteSwap.h"
 #include "Serialization/CompactBinary.h"
-#include "Serialization/CompactBinarySerialization.h"
-#include "Serialization/CompactBinaryWriter.h"
 #include "Templates/UnrealTemplate.h"
 
 namespace UE::DerivedData
@@ -74,21 +72,4 @@ FValueId FValueId::MakeIndexed(int32 Index) const
 	return FValueId(IndexedBytes);
 }
 
-FCbWriter& operator<<(FCbWriter& Writer, const FValueId& Id)
-{
-	Writer.AddObjectId(Id);
-	return Writer;
-}
-
-bool LoadFromCompactBinary(FCbFieldView Field, FValueId& OutId)
-{
-	if (const FCbObjectId ObjectId = Field.AsObjectId(); !Field.HasError())
-	{
-		OutId = ObjectId;
-		return true;
-	}
-	OutId.Reset();
-	return false;
-}
-
 } // UE::DerivedData
diff --git a/Engine/Source/Developer/DerivedDataCache/Private/ZenCacheStore.cpp b/Engine/Source/Developer/DerivedDataCache/Private/ZenCacheStore.cpp
--- a/Engine/Source/Developer/DerivedDataCache/Private/ZenCacheStore.cpp
+++ b/Engine/Source/Developer/DerivedDataCache/Private/ZenCacheStore.cpp
@@ -18,6 +18,7 @@
 #include "DerivedDataHttpRequestQueue.h"
 #include "DerivedDataRequest.h"
 #include "DerivedDataRequestOwner.h"
+#include "DerivedDataSerialization.h"
 #include "Experimental/ZenStatistics.h"
 #include "HAL/FileManager.h"
 #include "HAL/IConsoleManager.h"
@@ -1475,6 +1476,12 @@ public:
 											Succeeded = true;
 										}
 									}
+									else if (!FragmentHashView.IsHash())
+									{
+										// An out-of-bounds access will respond with no buffer.
+										RequestedBytes = FUniqueBuffer::Alloc(0).MoveToShared();
+										Succeeded = true;
+									}
 								}
 							}
 							else
diff --git a/Engine/Source/Developer/DerivedDataCache/Public/DerivedDataCache.h b/Engine/Source/Developer/DerivedDataCache/Public/DerivedDataCache.h
--- a/Engine/Source/Developer/DerivedDataCache/Public/DerivedDataCache.h
+++ b/Engine/Source/Developer/DerivedDataCache/Public/DerivedDataCache.h
@@ -19,9 +19,6 @@
 
 #define UE_API DERIVEDDATACACHE_API
 
-class FCbFieldView;
-class FCbWriter;
-
 namespace UE::DerivedData { class ICacheStoreMaintainer; }
 namespace UE::DerivedData { class IRequestOwner; }
 namespace UE::DerivedData { struct FCacheGetChunkRequest; }
@@ -359,16 +356,6 @@ struct FCacheGetChunkResponse
 
 ///////////////////////////////////////////////////////////////////////////////////////////////////////////////////////
 
-UE_API FCbWriter& operator<<(FCbWriter& Writer, const FCacheGetRequest& Request);
-UE_API FCbWriter& operator<<(FCbWriter& Writer, const FCacheGetValueRequest& Request);
-UE_API FCbWriter& operator<<(FCbWriter& Writer, const FCacheGetChunkRequest& Request);
-
-UE_API bool LoadFromCompactBinary(FCbFieldView Field, FCacheGetRequest& OutRequest);
-UE_API bool LoadFromCompactBinary(FCbFieldView Field, FCacheGetValueRequest& OutRequest);
-UE_API bool LoadFromCompactBinary(FCbFieldView Field, FCacheGetChunkRequest& OutRequest);
-
-///////////////////////////////////////////////////////////////////////////////////////////////////////////////////////
-
 /** Returns a reference to the cache. Asserts if not available. */
 UE_API ICache& GetCache();
 
@@ -391,3 +378,7 @@ UE_INTERNAL UE_API ICache* CreateCache(FStringView GraphNameOrConfig);
 } // UE::DerivedData
 
 #undef UE_API
+
+#if UE_ENABLE_INCLUDE_ORDER_DEPRECATED_IN_5_8
+#include "DerivedDataSerialization.h"
+#endif
diff --git a/Engine/Source/Developer/DerivedDataCache/Public/DerivedDataCacheKey.h b/Engine/Source/Developer/DerivedDataCache/Public/DerivedDataCacheKey.h
--- a/Engine/Source/Developer/DerivedDataCache/Public/DerivedDataCacheKey.h
+++ b/Engine/Source/Developer/DerivedDataCache/Public/DerivedDataCacheKey.h
@@ -12,7 +12,6 @@
 
 #define UE_API DERIVEDDATACACHE_API
 
-class FCbFieldView;
 class FCbWriter;
 
 namespace UE::DerivedData
@@ -136,9 +135,6 @@ inline TStringBuilderBase<CharType>& operator<<(TStringBuilderBase<CharType>& Bu
 	return Builder << Bucket.ToString();
 }
 
-UE_API FCbWriter& operator<<(FCbWriter& Writer, FCacheBucket Bucket);
-UE_API bool LoadFromCompactBinary(FCbFieldView Field, FCacheBucket& OutBucket);
-
 ///////////////////////////////////////////////////////////////////////////////////////////////////////////////////////
 
 template <typename CharType>
@@ -149,9 +145,6 @@ inline TStringBuilderBase<CharType>& operator<<(TStringBuilderBase<CharType>& Bu
 
 UE_API void SerializeForLog(FCbWriter& Writer, const FCacheKey& Key);
 
-UE_API FCbWriter& operator<<(FCbWriter& Writer, const FCacheKey& Key);
-UE_API bool LoadFromCompactBinary(FCbFieldView Field, FCacheKey& OutKey);
-
 ///////////////////////////////////////////////////////////////////////////////////////////////////////////////////////
 
 /**
@@ -175,3 +168,7 @@ UE_API FCacheKey BuildLegacyCacheKey(FStringView PluginName, FStringView Version
 } // UE::DerivedData
 
 #undef UE_API
+
+#if UE_ENABLE_INCLUDE_ORDER_DEPRECATED_IN_5_8
+#include "DerivedDataSerialization.h"
+#endif
diff --git a/Engine/Source/Developer/DerivedDataCache/Public/DerivedDataCachePolicy.h b/Engine/Source/Developer/DerivedDataCache/Public/DerivedDataCachePolicy.h
--- a/Engine/Source/Developer/DerivedDataCache/Public/DerivedDataCachePolicy.h
+++ b/Engine/Source/Developer/DerivedDataCache/Public/DerivedDataCachePolicy.h
@@ -13,9 +13,6 @@
 
 #define UE_API DERIVEDDATACACHE_API
 
-class FCbFieldView;
-class FCbObjectView;
-class FCbWriter;
 template <typename FuncType> class TFunctionRef;
 
 namespace UE::DerivedData::Private { class ICacheRecordPolicyShared; }
@@ -131,9 +128,6 @@ UE_API FUtf8StringBuilderBase& operator<<(FUtf8StringBuilderBase& Builder, ECach
 UE_API bool TryLexFromString(ECachePolicy& OutPolicy, FUtf8StringView String);
 UE_API bool TryLexFromString(ECachePolicy& OutPolicy, FWideStringView String);
 
-UE_API FCbWriter& operator<<(FCbWriter& Writer, ECachePolicy Policy);
-UE_API bool LoadFromCompactBinary(FCbFieldView Field, ECachePolicy& OutPolicy, ECachePolicy Default = ECachePolicy::Default);
-
 /** A value ID and the cache policy to use for that value. */
 struct FCacheValuePolicy
 {
@@ -144,9 +138,6 @@ struct FCacheValuePolicy
 	static constexpr ECachePolicy PolicyMask = ECachePolicy::Default | ECachePolicy::SkipData | ECachePolicy::NonDeterministic;
 };
 
-UE_API FCbWriter& operator<<(FCbWriter& Writer, const FCacheValuePolicy& Policy);
-UE_API bool LoadFromCompactBinary(FCbFieldView Field, FCacheValuePolicy& OutPolicy);
-
 /** Interface for the private implementation of the cache record policy. */
 class Private::ICacheRecordPolicyShared
 {
@@ -267,9 +258,10 @@ public:
 	inline void Reset() { *this = FOptionalCacheRecordPolicy(); }
 };
 
-UE_API FCbWriter& operator<<(FCbWriter& Writer, const FCacheRecordPolicy& Policy);
-UE_API bool LoadFromCompactBinary(FCbFieldView Field, FCacheRecordPolicy& OutPolicy);
-
 } // UE::DerivedData
 
 #undef UE_API
+
+#if UE_ENABLE_INCLUDE_ORDER_DEPRECATED_IN_5_8
+#include "DerivedDataSerialization.h"
+#endif
diff --git a/Engine/Source/Developer/DerivedDataCache/Public/DerivedDataCacheRecord.h b/Engine/Source/Developer/DerivedDataCache/Public/DerivedDataCacheRecord.h
--- a/Engine/Source/Developer/DerivedDataCache/Public/DerivedDataCacheRecord.h
+++ b/Engine/Source/Developer/DerivedDataCache/Public/DerivedDataCacheRecord.h
@@ -4,7 +4,6 @@
 
 #include "Containers/ContainersFwd.h"
 #include "DerivedDataCacheKey.h"
-#include "DerivedDataValueId.h"
 #include "HAL/Platform.h"
 #include "Memory/MemoryFwd.h"
 #include "Misc/ScopeExit.h"
@@ -24,6 +23,7 @@ namespace UE::DerivedData { class FOptionalCacheRecord; }
 namespace UE::DerivedData { class FValue; }
 namespace UE::DerivedData { class FValueWithId; }
 namespace UE::DerivedData { class IRequestOwner; }
+namespace UE::DerivedData { struct FValueId; }
 namespace UE::DerivedData { using FOnCacheRecordComplete = TUniqueFunction<void (FCacheRecord&& Record)>; }
 
 ///////////////////////////////////////////////////////////////////////////////////////////////////////////////////////
@@ -86,7 +86,7 @@ public:
 
 	/** Save the cache record to a compact binary package. */
 	UE_API FCbPackage Save() const;
-	/** Append the cache record to an existing package and writer for e.g. a batch of records. */
+	/** Append the cache record to an existing package and writer, as an object, for e.g. a batch of records. */
 	UE_API void Save(FCbPackage& Attachments, FCbWriter& Writer) const;
 
 	/** Load a cache record from a compact binary package. Null on error. */
@@ -236,3 +236,7 @@ struct FCacheRecordKeyFuncs
 } // UE::DerivedData
 
 #undef UE_API
+
+#if UE_ENABLE_INCLUDE_ORDER_DEPRECATED_IN_5_8
+#include "DerivedDataValueId.h"
+#endif
diff --git a/Engine/Source/Developer/DerivedDataCache/Public/DerivedDataRequestTypes.h b/Engine/Source/Developer/DerivedDataCache/Public/DerivedDataRequestTypes.h
--- a/Engine/Source/Developer/DerivedDataCache/Public/DerivedDataRequestTypes.h
+++ b/Engine/Source/Developer/DerivedDataCache/Public/DerivedDataRequestTypes.h
@@ -7,8 +7,6 @@
 
 #define UE_API DERIVEDDATACACHE_API
 
-class FCbFieldView;
-class FCbWriter;
 enum class EQueuedWorkPriority : uint8;
 
 namespace UE::DerivedData
@@ -59,9 +57,6 @@ UE_API FUtf8StringBuilderBase& operator<<(FUtf8StringBuilderBase& Builder, EPrio
 UE_API bool TryLexFromString(EPriority& OutPriority, FUtf8StringView String);
 UE_API bool TryLexFromString(EPriority& OutPriority, FWideStringView String);
 
-UE_API FCbWriter& operator<<(FCbWriter& Writer, EPriority Priority);
-UE_API bool LoadFromCompactBinary(FCbFieldView Field, EPriority& OutPriority, EPriority Default = EPriority::Normal);
-
 /** Converts to a queued work priority. Asserts on invalid priority values. */
 UE_API EQueuedWorkPriority ConvertToQueuedWorkPriority(EPriority Priority);
 /** Converts from a queued work priority. Asserts on invalid priority values. */
@@ -87,9 +82,10 @@ UE_API FUtf8StringBuilderBase& operator<<(FUtf8StringBuilderBase& Builder, EStat
 UE_API bool TryLexFromString(EStatus& OutStatus, FUtf8StringView String);
 UE_API bool TryLexFromString(EStatus& OutStatus, FWideStringView String);
 
-UE_API FCbWriter& operator<<(FCbWriter& Writer, EStatus Status);
-UE_API bool LoadFromCompactBinary(FCbFieldView Field, EStatus& OutStatus, EStatus Default = EStatus::Ok);
-
 } // UE::DerivedData
 
 #undef UE_API
+
+#if UE_ENABLE_INCLUDE_ORDER_DEPRECATED_IN_5_8
+#include "DerivedDataSerialization.h"
+#endif
diff --git a/Engine/Source/Developer/DerivedDataCache/Public/DerivedDataSerialization.h b/Engine/Source/Developer/DerivedDataCache/Public/DerivedDataSerialization.h
new file mode 100644
--- /dev/null
+++ b/Engine/Source/Developer/DerivedDataCache/Public/DerivedDataSerialization.h
@@ -0,0 +1,159 @@
+// Copyright Epic Games, Inc. All Rights Reserved.
+
+#pragma once
+
+#include "CoreTypes.h"
+#include "Misc/OptionalFwd.h"
+#include "Templates/FunctionWithContext.h"
+
+#define UE_API DERIVEDDATACACHE_API
+
+class FCbAttachment;
+class FCbFieldView;
+class FCbObject;
+class FCbWriter;
+struct FIoHash;
+
+namespace UE::DerivedData { class FCacheBucket; }
+namespace UE::DerivedData { class FCacheRecord; }
+namespace UE::DerivedData { class FCacheRecordPolicy; }
+namespace UE::DerivedData { class FOptionalCacheRecord; }
+namespace UE::DerivedData { class FValue; }
+namespace UE::DerivedData { class FValueWithId; }
+namespace UE::DerivedData { struct FCacheGetChunkRequest; }
+namespace UE::DerivedData { struct FCacheGetChunkResponse; }
+namespace UE::DerivedData { struct FCacheGetRequest; }
+namespace UE::DerivedData { struct FCacheGetResponse; }
+namespace UE::DerivedData { struct FCacheGetValueRequest; }
+namespace UE::DerivedData { struct FCacheGetValueResponse; }
+namespace UE::DerivedData { struct FCacheKey; }
+namespace UE::DerivedData { struct FCachePutRequest; }
+namespace UE::DerivedData { struct FCachePutValueRequest; }
+namespace UE::DerivedData { struct FCacheValuePolicy; }
+namespace UE::DerivedData { struct FValueId; }
+namespace UE::DerivedData { enum class ECacheMethod : uint8; }
+namespace UE::DerivedData { enum class ECachePolicy : uint32; }
+namespace UE::DerivedData { enum class EPriority : uint8; }
+namespace UE::DerivedData { enum class EStatus : uint8; }
+
+namespace UE::DerivedData
+{
+
+using FSaveAttachmentFn = TFunctionWithContext<void (FCbAttachment&& Attachment)>;
+using FLoadAttachmentFn = TFunctionWithContext<FCbAttachment (const FIoHash& Hash)>;
+
+/** Format in which to save metadata. */
+enum class EMetaSaveFormat
+{
+	/** Save metadata as an object field. */
+	Object,
+	/** Save metadata as an attachment referenced by its hash and size. */
+	Attachment,
+};
+
+/** Save metadata to a stream of named fields. Caller must call BeginObject/EndObject to save this into an object. */
+UE_API void SaveMetaToCompactBinary(FCbWriter& Writer, const FCbObject& Meta, EMetaSaveFormat MetaFormat, FSaveAttachmentFn SaveAttachment = nullptr);
+UE_API bool LoadMetaFromCompactBinary(FCbFieldView Field, FCbObject& OutMeta, FLoadAttachmentFn LoadAttachment = nullptr);
+
+// Types from DerivedDataCache.h
+
+UE_DEPRECATED(5.8, "Call SaveToCompactBinary(Writer, Request).")
+UE_API FCbWriter& operator<<(FCbWriter& Writer, const FCacheGetRequest& Request);
+UE_DEPRECATED(5.8, "Call SaveToCompactBinary(Writer, Request).")
+UE_API FCbWriter& operator<<(FCbWriter& Writer, const FCacheGetValueRequest& Request);
+UE_DEPRECATED(5.8, "Call SaveToCompactBinary(Writer, Request).")
+UE_API FCbWriter& operator<<(FCbWriter& Writer, const FCacheGetChunkRequest& Request);
+
+UE_API void SaveToCompactBinary(FCbWriter& Writer, const FCachePutRequest& Request, FSaveAttachmentFn SaveAttachment = nullptr);
+UE_API bool LoadFromCompactBinary(FCbFieldView Field, TOptional<FCachePutRequest>& OutRequest, FLoadAttachmentFn LoadAttachment = nullptr);
+
+UE_API void SaveToCompactBinary(FCbWriter& Writer, const FCacheGetRequest& Request, const FCacheGetResponse* Response = nullptr, FSaveAttachmentFn SaveAttachment = nullptr);
+UE_API bool LoadFromCompactBinary(FCbFieldView Field, FCacheGetRequest& OutRequest);
+UE_API bool LoadFromCompactBinary(FCbFieldView Field, TOptional<FCacheGetResponse>& OutResponse, FLoadAttachmentFn LoadAttachment = nullptr);
+
+UE_API void SaveToCompactBinary(FCbWriter& Writer, const FCachePutValueRequest& Request, FSaveAttachmentFn SaveAttachment = nullptr);
+UE_API bool LoadFromCompactBinary(FCbFieldView Field, FCachePutValueRequest& OutRequest, FLoadAttachmentFn LoadAttachment = nullptr);
+
+UE_API void SaveToCompactBinary(FCbWriter& Writer, const FCacheGetValueRequest& Request, const FCacheGetValueResponse* Response = nullptr, FSaveAttachmentFn SaveAttachment = nullptr);
+UE_API bool LoadFromCompactBinary(FCbFieldView Field, FCacheGetValueRequest& OutRequest);
+UE_API bool LoadFromCompactBinary(FCbFieldView Field, FCacheGetValueResponse& OutResponse, FLoadAttachmentFn LoadAttachment = nullptr);
+
+UE_API void SaveToCompactBinary(FCbWriter& Writer, const FCacheGetChunkRequest& Request, const FCacheGetChunkResponse* Response = nullptr, FSaveAttachmentFn SaveAttachment = nullptr);
+UE_API bool LoadFromCompactBinary(FCbFieldView Field, FCacheGetChunkRequest& OutRequest);
+UE_API bool LoadFromCompactBinary(FCbFieldView Field, FCacheGetChunkResponse& OutResponse, FLoadAttachmentFn LoadAttachment = nullptr);
+
+// Types from DerivedDataCacheKey.h
+
+UE_API FCbWriter& operator<<(FCbWriter& Writer, ECachePolicy Policy);
+UE_API bool LoadFromCompactBinary(FCbFieldView Field, ECachePolicy& OutPolicy, ECachePolicy Default);
+UE_DEPRECATED(5.8, "Call the overload with an explicit default of ECachePolicy::Default.")
+UE_API bool LoadFromCompactBinary(FCbFieldView Field, ECachePolicy& OutPolicy);
+
+UE_API FCbWriter& operator<<(FCbWriter& Writer, FCacheBucket Bucket);
+UE_API bool LoadFromCompactBinary(FCbFieldView Field, FCacheBucket& OutBucket);
+
+UE_API FCbWriter& operator<<(FCbWriter& Writer, const FCacheKey& Key);
+UE_API bool LoadFromCompactBinary(FCbFieldView Field, FCacheKey& OutKey);
+
+// Types from DerivedDataCacheMethod.h
+
+UE_API FCbWriter& operator<<(FCbWriter& Writer, ECacheMethod Method);
+UE_API bool LoadFromCompactBinary(FCbFieldView Field, ECacheMethod& OutMethod);
+
+// Types from DerivedDataCachePolicy.h
+
+UE_API FCbWriter& operator<<(FCbWriter& Writer, const FCacheValuePolicy& Policy);
+UE_API bool LoadFromCompactBinary(FCbFieldView Field, FCacheValuePolicy& OutPolicy);
+
+UE_API FCbWriter& operator<<(FCbWriter& Writer, const FCacheRecordPolicy& Policy);
+UE_API bool LoadFromCompactBinary(FCbFieldView Field, FCacheRecordPolicy& OutPolicy);
+
+// Types from DerivedDataCacheRecord.h
+
+/**
+ * Save a cache record as a stream of named fields. Attachments optionally saved to a callback.
+ *
+ * It is the responsibility of the caller to call BeginObject/EndObject to save this into an object.
+ * When SaveAttachment is provided, it is called for each attachment referenced by the saved record.
+ */
+UE_API void SaveToCompactBinary(FCbWriter& Writer, const FCacheRecord& Record, EMetaSaveFormat MetaFormat, FSaveAttachmentFn SaveAttachment = nullptr);
+
+/**
+ * Load a cache record from an object. Attachments optionally loaded from a callback.
+ *
+ * When LoadAttachment is provided, it is called for each attachment referenced by the loaded object.
+ * An attachment that cannot be loaded will be omitted from the returned record without failing the load.
+ */
+UE_API bool LoadFromCompactBinary(FCbFieldView Field, FOptionalCacheRecord& OutRecord, FLoadAttachmentFn LoadAttachment = nullptr);
+
+// Types from DerivedDataRequestTypes.h
+
+UE_API FCbWriter& operator<<(FCbWriter& Writer, EPriority Priority);
+UE_API bool LoadFromCompactBinary(FCbFieldView Field, EPriority& OutPriority, EPriority Default);
+UE_DEPRECATED(5.8, "Call the overload with an explicit default of EPriority::Normal.")
+UE_API bool LoadFromCompactBinary(FCbFieldView Field, EPriority& OutPriority);
+
+UE_API FCbWriter& operator<<(FCbWriter& Writer, EStatus Status);
+UE_API bool LoadFromCompactBinary(FCbFieldView Field, EStatus& OutStatus, EStatus Default);
+UE_DEPRECATED(5.8, "Call the overload with an explicit default of EStatus::Ok.")
+UE_API bool LoadFromCompactBinary(FCbFieldView Field, EStatus& OutStatus);
+
+// Types from DerivedDataValue.h
+
+UE_API FCbWriter& operator<<(FCbWriter& Writer, const FValue& Value);
+UE_API FCbWriter& operator<<(FCbWriter& Writer, const FValueWithId& Value);
+
+UE_API void SaveToCompactBinary(FCbWriter& Writer, const FValue& Value, FSaveAttachmentFn SaveAttachment = nullptr);
+UE_API void SaveToCompactBinary(FCbWriter& Writer, const FValueWithId& Value, FSaveAttachmentFn SaveAttachment = nullptr);
+
+UE_API bool LoadFromCompactBinary(FCbFieldView Field, FValue& OutValue, FLoadAttachmentFn LoadAttachment = nullptr);
+UE_API bool LoadFromCompactBinary(FCbFieldView Field, FValueWithId& OutValue, FLoadAttachmentFn LoadAttachment = nullptr);
+
+// Types from DerivedDataValueId.h
+
+UE_API FCbWriter& operator<<(FCbWriter& Writer, const FValueId& Id);
+UE_API bool LoadFromCompactBinary(FCbFieldView Field, FValueId& OutId);
+
+} // UE::DerivedData
+
+#undef UE_API
diff --git a/Engine/Source/Developer/DerivedDataCache/Public/DerivedDataValueId.h b/Engine/Source/Developer/DerivedDataCache/Public/DerivedDataValueId.h
--- a/Engine/Source/Developer/DerivedDataCache/Public/DerivedDataValueId.h
+++ b/Engine/Source/Developer/DerivedDataCache/Public/DerivedDataValueId.h
@@ -9,9 +9,7 @@
 
 #define UE_API DERIVEDDATACACHE_API
 
-class FCbFieldView;
 class FCbObjectId;
-class FCbWriter;
 struct FIoHash;
 
 namespace UE::DerivedData
@@ -109,9 +107,10 @@ inline bool FValueId::IsNull() const
 	return *this == FValueId();
 }
 
-UE_API FCbWriter& operator<<(FCbWriter& Writer, const FValueId& Id);
-UE_API bool LoadFromCompactBinary(FCbFieldView Field, FValueId& OutId);
-
 } // UE::DerivedData
 
 #undef UE_API
+
+#if UE_ENABLE_INCLUDE_ORDER_DEPRECATED_IN_5_8
+#include "DerivedDataSerialization.h"
+#endif
diff --git a/Engine/Source/Developer/DerivedDataCache/Tests/DerivedDataSerializationTest.cpp b/Engine/Source/Developer/DerivedDataCache/Tests/DerivedDataSerializationTest.cpp
new file mode 100644
--- /dev/null
+++ b/Engine/Source/Developer/DerivedDataCache/Tests/DerivedDataSerializationTest.cpp
@@ -0,0 +1,327 @@
+// Copyright Epic Games, Inc. All Rights Reserved.
+
+#if WITH_LOW_LEVEL_TESTS
+
+#include "DerivedDataSerialization.h"
+
+#include "DerivedDataCache.h"
+#include "Misc/Optional.h"
+#include "Serialization/CompactBinaryPackage.h"
+#include "Serialization/CompactBinaryWriter.h"
+#include "TestHarness.h"
+
+namespace UE::DerivedData
+{
+
+TEST_CASE("DerivedData::Serialization::CompactBinary", "[DerivedData]")
+{
+	const FSharedString Name(TEXT("Name"));
+	const FCacheKey Key{FCacheBucket(ANSITEXTVIEW("Bucket")), FIoHash::HashBuffer(MakeMemoryView<uint32>({0xf9332932, 0x68c04d2e, 0xaea27221, 0xb190704a}))};
+	const FCbObject Meta = []
+	{
+		TCbWriter<64> Writer;
+		Writer.BeginObject();
+		Writer.AddString(ANSITEXTVIEW("Key"), ANSITEXTVIEW("Value"));
+		Writer.EndObject();
+		return Writer.Save().AsObject();
+	}();
+	const FIoHash MetaHash = Meta.GetHash();
+	const FValueId ValueId = FValueId::FromName(TEXT("Value"));
+	const FValue Value = FValue::Compress(FSharedBuffer::MakeView(MakeMemoryView<uint8>({0, 1, 2, 3, 4, 5, 6, 7})));
+	const FCacheRecord Record = [&]
+	{
+		FCacheRecordBuilder Builder(Key);
+		Builder.SetMeta(CopyTemp(Meta));
+		Builder.AddValue(ValueId, Value);
+		return Builder.Build();
+	}();
+	const ECachePolicy ValuePolicy = ECachePolicy::Default | ECachePolicy::SkipData;
+	const FCacheRecordPolicy RecordPolicy = [&]
+	{
+		FCacheRecordPolicyBuilder Builder(ECachePolicy::Local | ECachePolicy::SkipMeta | ECachePolicy::KeepAlive);
+		Builder.AddValuePolicy(ValueId, ValuePolicy);
+		return Builder.Build();
+	}();
+	const uint64 UserData = 0xfedcba9876543210;
+	const FSharedBuffer EmptyBuffer = FUniqueBuffer::Alloc(0).MoveToShared();
+	const FIoHash EmptyBufferHash = FIoHash::HashBuffer(EmptyBuffer);
+
+	const auto LoadAttachment = [&](const FIoHash& Hash) -> FCbAttachment
+	{
+		if (Hash == Value.GetRawHash())
+		{
+			return FCbAttachment(Value.GetData());
+		}
+		if (Hash == MetaHash)
+		{
+			return FCbAttachment(Meta);
+		}
+		if (Hash == EmptyBufferHash)
+		{
+			return FCbAttachment(EmptyBuffer);
+		}
+		return FCbAttachment();
+	};
+
+	const auto SaveAttachment = [&](FCbAttachment&& Attachment)
+	{
+		const FIoHash Hash = Attachment.GetHash();
+		CHECK((Hash == Value.GetRawHash() || Hash == MetaHash || Hash == EmptyBufferHash));
+	};
+
+	SECTION("CachePutRequest")
+	{
+		const FCachePutRequest Request{Name, Record, RecordPolicy, UserData};
+
+		const FCbField Field = [&]
+		{
+			TCbWriter<512> Writer;
+			SaveToCompactBinary(Writer, Request, SaveAttachment);
+			return Writer.Save();
+		}();
+
+		{
+			TOptional<FCachePutRequest> LoadedRequest;
+			CHECK(LoadFromCompactBinary(Field, LoadedRequest));
+			CHECKED_IF(LoadedRequest)
+			{
+				CHECK(LoadedRequest->Name == Request.Name);
+				CHECK(LoadedRequest->Record.GetKey() == Request.Record.GetKey());
+				CHECK_FALSE(LoadedRequest->Record.GetMeta());
+				CHECK(LoadedRequest->Record.GetValue(ValueId) == Value);
+				CHECK_FALSE(LoadedRequest->Record.GetValue(ValueId).HasData());
+				CHECK(LoadedRequest->UserData == Request.UserData);
+			}
+		}
+
+		{
+			TOptional<FCachePutRequest> LoadedRequest;
+			CHECK(LoadFromCompactBinary(Field, LoadedRequest, LoadAttachment));
+			CHECKED_IF(LoadedRequest)
+			{
+				CHECK(LoadedRequest->Name == Request.Name);
+				CHECK(LoadedRequest->Record.GetKey() == Request.Record.GetKey());
+				CHECK(LoadedRequest->Record.GetMeta().Equals(Meta));
+				CHECK(LoadedRequest->Record.GetValue(ValueId) == Value);
+				CHECK(LoadedRequest->Record.GetValue(ValueId).HasData());
+				CHECK(LoadedRequest->UserData == Request.UserData);
+			}
+		}
+	}
+
+	SECTION("CacheGetRequest")
+	{
+		const FCacheGetRequest Request{Name, Key, RecordPolicy, UserData};
+
+		const FCbField Field = [&]
+		{
+			TCbWriter<512> Writer;
+			SaveToCompactBinary(Writer, Request);
+			return Writer.Save();
+		}();
+
+		FCacheGetRequest LoadedRequest;
+		CHECK(LoadFromCompactBinary(Field, LoadedRequest));
+		CHECK(LoadedRequest.Name == Request.Name);
+		CHECK(LoadedRequest.Key == Request.Key);
+		CHECK(LoadedRequest.Policy.GetBasePolicy() == Request.Policy.GetBasePolicy());
+		CHECK(LoadedRequest.Policy.GetRecordPolicy() == Request.Policy.GetRecordPolicy());
+		CHECK(LoadedRequest.Policy.GetValuePolicy(ValueId) == Request.Policy.GetValuePolicy(ValueId));
+		CHECK(LoadedRequest.UserData == Request.UserData);
+	}
+
+	SECTION("CacheGetResponse")
+	{
+		const FCacheGetRequest Request{Name, Key, RecordPolicy, UserData};
+		const FCacheGetResponse Response{Name, Record, UserData, EStatus::Ok};
+
+		const FCbField Field = [&]
+		{
+			TCbWriter<512> Writer;
+			SaveToCompactBinary(Writer, Request, &Response, SaveAttachment);
+			return Writer.Save();
+		}();
+
+		{
+			TOptional<FCacheGetResponse> LoadedResponse;
+			CHECK(LoadFromCompactBinary(Field, LoadedResponse));
+			CHECKED_IF(LoadedResponse)
+			{
+				CHECK(LoadedResponse->Name == Response.Name);
+				CHECK(LoadedResponse->Record.GetKey() == Response.Record.GetKey());
+				CHECK_FALSE(LoadedResponse->Record.GetMeta());
+				CHECK(LoadedResponse->Record.GetValue(ValueId) == Value);
+				CHECK_FALSE(LoadedResponse->Record.GetValue(ValueId).HasData());
+				CHECK(LoadedResponse->UserData == Response.UserData);
+				CHECK(LoadedResponse->Status == Response.Status);
+			}
+		}
+
+		{
+			TOptional<FCacheGetResponse> LoadedResponse;
+			CHECK(LoadFromCompactBinary(Field, LoadedResponse, LoadAttachment));
+			CHECKED_IF(LoadedResponse)
+			{
+				CHECK(LoadedResponse->Name == Response.Name);
+				CHECK(LoadedResponse->Record.GetKey() == Response.Record.GetKey());
+				CHECK(LoadedResponse->Record.GetMeta().Equals(Meta));
+				CHECK(LoadedResponse->Record.GetValue(ValueId) == Value);
+				CHECK(LoadedResponse->Record.GetValue(ValueId).HasData());
+				CHECK(LoadedResponse->UserData == Response.UserData);
+				CHECK(LoadedResponse->Status == Response.Status);
+			}
+		}
+	}
+
+	SECTION("CachePutValueRequest")
+	{
+		const FCachePutValueRequest Request{Name, Key, Value, ValuePolicy, UserData};
+
+		const FCbField Field = [&]
+		{
+			TCbWriter<512> Writer;
+			SaveToCompactBinary(Writer, Request, SaveAttachment);
+			return Writer.Save();
+		}();
+
+		{
+			FCachePutValueRequest LoadedRequest;
+			CHECK(LoadFromCompactBinary(Field, LoadedRequest));
+			CHECK(LoadedRequest.Name == Request.Name);
+			CHECK(LoadedRequest.Key == Request.Key);
+			CHECK(LoadedRequest.Value == Request.Value);
+			CHECK_FALSE(LoadedRequest.Value.HasData());
+			CHECK(LoadedRequest.UserData == Request.UserData);
+		}
+
+		{
+			FCachePutValueRequest LoadedRequest;
+			CHECK(LoadFromCompactBinary(Field, LoadedRequest, LoadAttachment));
+			CHECK(LoadedRequest.Name == Request.Name);
+			CHECK(LoadedRequest.Key == Request.Key);
+			CHECK(LoadedRequest.Value == Request.Value);
+			CHECK(LoadedRequest.Value.HasData());
+			CHECK(LoadedRequest.UserData == Request.UserData);
+		}
+	}
+
+	SECTION("CacheGetValueRequest")
+	{
+		const FCacheGetValueRequest Request{Name, Key, ValuePolicy, UserData};
+
+		const FCbField Field = [&]
+		{
+			TCbWriter<512> Writer;
+			SaveToCompactBinary(Writer, Request);
+			return Writer.Save();
+		}();
+
+		FCacheGetValueRequest LoadedRequest;
+		CHECK(LoadFromCompactBinary(Field, LoadedRequest));
+		CHECK(LoadedRequest.Name == Request.Name);
+		CHECK(LoadedRequest.Key == Request.Key);
+		CHECK(LoadedRequest.Policy == Request.Policy);
+		CHECK(LoadedRequest.UserData == Request.UserData);
+	}
+
+	SECTION("CacheGetValueResponse")
+	{
+		const FCacheGetValueRequest Request{Name, Key, ValuePolicy, UserData};
+		const FCacheGetValueResponse Response{Name, Key, Value, UserData, EStatus::Ok};
+
+		const FCbField Field = [&]
+		{
+			TCbWriter<512> Writer;
+			SaveToCompactBinary(Writer, Request, &Response, SaveAttachment);
+			return Writer.Save();
+		}();
+
+		{
+			FCacheGetValueResponse LoadedResponse;
+			CHECK(LoadFromCompactBinary(Field, LoadedResponse));
+			CHECK(LoadedResponse.Name == Response.Name);
+			CHECK(LoadedResponse.Key == Response.Key);
+			CHECK(LoadedResponse.Value == Response.Value);
+			CHECK_FALSE(LoadedResponse.Value.HasData());
+			CHECK(LoadedResponse.UserData == Response.UserData);
+			CHECK(LoadedResponse.Status == Response.Status);
+		}
+
+		{
+			FCacheGetValueResponse LoadedResponse;
+			CHECK(LoadFromCompactBinary(Field, LoadedResponse, LoadAttachment));
+			CHECK(LoadedResponse.Name == Name);
+			CHECK(LoadedResponse.Key == Key);
+			CHECK(LoadedResponse.Value == Value);
+			CHECK(LoadedResponse.Value.HasData());
+			CHECK(LoadedResponse.UserData == UserData);
+			CHECK(LoadedResponse.Status == Response.Status);
+		}
+	}
+
+	SECTION("CacheGetChunkRequest")
+	{
+		const FCacheGetChunkRequest Request{Name, Key, ValueId, 1024, 512, Value.GetRawHash(), ValuePolicy, UserData};
+
+		const FCbField Field = [&]
+		{
+			TCbWriter<512> Writer;
+			SaveToCompactBinary(Writer, Request);
+			return Writer.Save();
+		}();
+
+		FCacheGetChunkRequest LoadedRequest;
+		CHECK(LoadFromCompactBinary(Field, LoadedRequest));
+		CHECK(LoadedRequest.Name == Request.Name);
+		CHECK(LoadedRequest.Key == Request.Key);
+		CHECK(LoadedRequest.Id == Request.Id);
+		CHECK(LoadedRequest.RawOffset == Request.RawOffset);
+		CHECK(LoadedRequest.RawSize == Request.RawSize);
+		CHECK(LoadedRequest.RawHash == Request.RawHash);
+		CHECK(LoadedRequest.Policy == Request.Policy);
+		CHECK(LoadedRequest.UserData == Request.UserData);
+	}
+
+	SECTION("CacheGetChunkResponse")
+	{
+		const FCacheGetChunkRequest Request{Name, Key, ValueId, 1024, 512, Value.GetRawHash(), ValuePolicy, UserData};
+		const FCacheGetChunkResponse Response{Name, Key, ValueId, 1024, 256, Value.GetRawHash(), EmptyBuffer, UserData, EStatus::Ok};
+
+		const FCbField Field = [&]
+		{
+			TCbWriter<512> Writer;
+			SaveToCompactBinary(Writer, Request, &Response, SaveAttachment);
+			return Writer.Save();
+		}();
+
+		FCacheGetChunkResponse LoadedResponse;
+		CHECK(LoadFromCompactBinary(Field, LoadedResponse, LoadAttachment));
+		CHECK(LoadedResponse.Name == Response.Name);
+		CHECK(LoadedResponse.Key == Response.Key);
+		CHECK(LoadedResponse.Id == Response.Id);
+		CHECK(LoadedResponse.RawOffset == Response.RawOffset);
+		CHECK(LoadedResponse.RawSize == Response.RawSize);
+		CHECK(LoadedResponse.RawHash == Response.RawHash);
+		CHECK(LoadedResponse.RawData);
+		CHECK(LoadedResponse.RawData.GetSize() == 0);
+		CHECK(LoadedResponse.UserData == Response.UserData);
+		CHECK(LoadedResponse.Status == Response.Status);
+	}
+
+	SECTION("CacheRecord")
+	{
+		const FCbPackage Package = Record.Save();
+		const FOptionalCacheRecord LoadedRecord = FCacheRecord::Load(Package);
+		CHECKED_IF(LoadedRecord)
+		{
+			CHECK(LoadedRecord.Get().GetKey() == Key);
+			CHECK(LoadedRecord.Get().GetMeta().Equals(Meta));
+			CHECK(LoadedRecord.Get().GetValue(ValueId) == Value);
+			CHECK(LoadedRecord.Get().GetValue(ValueId).HasData());
+		}
+	}
+}
+
+} // UE::DerivedData
+
+#endif // WITH_LOW_LEVEL_TESTS
