DDC: Remove deprecated code and deprecate an overload of LaunchTaskInThreadPool

diff --git a/Engine/Source/Developer/DerivedDataCache/Private/DerivedDataThreadPoolTask.cpp b/Engine/Source/Developer/DerivedDataCache/Private/DerivedDataThreadPoolTask.cpp
--- a/Engine/Source/Developer/DerivedDataCache/Private/DerivedDataThreadPoolTask.cpp
+++ b/Engine/Source/Developer/DerivedDataCache/Private/DerivedDataThreadPoolTask.cpp
@@ -149,21 +149,12 @@ void LaunchTaskInThreadPool(
 	}
 }
 
-void LaunchTaskInThreadPool(
-	uint64 MemoryEstimate,
-	IRequestOwner& Owner,
-	FQueuedThreadPool* ThreadPool,
-	TUniqueFunction<void ()>&& TaskBody)
-{
-	LaunchTaskInThreadPool(MemoryEstimate,TEXT("LaunchTaskInThreadPool"), Owner, ThreadPool, MoveTemp(TaskBody));
-}
-
 void LaunchTaskInThreadPool(
 	IRequestOwner& Owner,
 	FQueuedThreadPool* ThreadPool,
 	TUniqueFunction<void ()>&& TaskBody)
 {
-	LaunchTaskInThreadPool(0,TEXT("LaunchTaskInThreadPool"), Owner, ThreadPool, MoveTemp(TaskBody));
+	LaunchTaskInThreadPool(0, TEXT("LaunchTaskInThreadPool"), Owner, ThreadPool, MoveTemp(TaskBody));
 }
 
 } // UE::DerivedData
diff --git a/Engine/Source/Developer/DerivedDataCache/Public/DerivedDataSharedString.h b/Engine/Source/Developer/DerivedDataCache/Public/DerivedDataSharedString.h
deleted file mode 100644
diff --git a/Engine/Source/Developer/DerivedDataCache/Public/DerivedDataSharedStringFwd.h b/Engine/Source/Developer/DerivedDataCache/Public/DerivedDataSharedStringFwd.h
deleted file mode 100644
diff --git a/Engine/Source/Developer/DerivedDataCache/Public/DerivedDataThreadPoolTask.h b/Engine/Source/Developer/DerivedDataCache/Public/DerivedDataThreadPoolTask.h
--- a/Engine/Source/Developer/DerivedDataCache/Public/DerivedDataThreadPoolTask.h
+++ b/Engine/Source/Developer/DerivedDataCache/Public/DerivedDataThreadPoolTask.h
@@ -27,26 +27,20 @@ namespace UE::DerivedData
  * TaskBody is responsible for checking for cancellation by checking Owner.IsCanceled().
  *
  * @param MemoryEstimate   The estimated peak memory required to execute TaskBody.
- * @param DebugName		   Name of the job, for debugging/logging purposes only
+ * @param DebugName        The name of the job, for debugging/logging purposes only.
  * @param Owner            The owner of the request that will execute TaskBody.
  * @param ThreadPool       The thread pool within which to execute the task, if not retracted.
  * @param TaskBody         The task body that will be unconditionally executed exactly once.
  */
 UE_API void LaunchTaskInThreadPool(
 	uint64 MemoryEstimate,
-	const TCHAR * DebugName,
-	IRequestOwner& Owner,
-	FQueuedThreadPool* ThreadPool,
-	TUniqueFunction<void ()>&& TaskBody);
-	
-UE_DEPRECATED(5.3, "Use LaunchTaskInThreadPool with DebugName")
-UE_API void LaunchTaskInThreadPool(
-	uint64 MemoryEstimate,
+	const TCHAR* DebugName,
 	IRequestOwner& Owner,
 	FQueuedThreadPool* ThreadPool,
 	TUniqueFunction<void ()>&& TaskBody);
 
 // with no MemoryEstimate arg, assumes 0 (none) not -1 (unknown/max)
+UE_DEPRECATED(5.8, "Use LaunchTaskInThreadPool with DebugName and MemoryEstimate = 0.")
 UE_API void LaunchTaskInThreadPool(
 	IRequestOwner& Owner,
 	FQueuedThreadPool* ThreadPool,
diff --git a/Engine/Source/Runtime/Engine/Private/TextureDerivedDataTask.cpp b/Engine/Source/Runtime/Engine/Private/TextureDerivedDataTask.cpp
--- a/Engine/Source/Runtime/Engine/Private/TextureDerivedDataTask.cpp
+++ b/Engine/Source/Runtime/Engine/Private/TextureDerivedDataTask.cpp
@@ -2526,7 +2526,8 @@ static void LaunchBuildWithFallback(
 		OutBuildResults.PlatformData.ResultMetadata = *InInitialBuild.ResultMetadata;
 	}
 
-	LaunchTaskInThreadPool(InRequestOwner, FTextureCompilingManager::Get().GetThreadPool(),
+	// TODO: This has no memory estimate!
+	LaunchTaskInThreadPool(0, TEXT("TextureBuild"), InRequestOwner, FTextureCompilingManager::Get().GetThreadPool(),
 		[
 			FallbackBuild = MoveTemp(InFallbackBuild),
 			RequestOwner = &InRequestOwner,
