DDC: Added a ICacheReplayWriter interface and archive-based implementation

#Synced-CL 32500503

diff --git a/Engine/Source/Developer/DerivedDataCache/Internal/DerivedDataCacheReplay.h b/Engine/Source/Developer/DerivedDataCache/Internal/DerivedDataCacheReplay.h
--- a/Engine/Source/Developer/DerivedDataCache/Internal/DerivedDataCacheReplay.h
+++ b/Engine/Source/Developer/DerivedDataCache/Internal/DerivedDataCacheReplay.h
@@ -2,25 +2,64 @@
 
 #pragma once
 
+#include "Async/Mutex.h"
 #include "CoreTypes.h"
+#include "DerivedDataCache.h"
 #include "Templates/PimplPtr.h"
 
 #define UE_API DERIVEDDATACACHE_API
 
 class FArchive;
 class FCbObjectView;
+class FCbWriter;
 
 namespace UE::DerivedData { class FCacheKeyFilter; }
 namespace UE::DerivedData { class FCacheMethodFilter; }
 namespace UE::DerivedData { class ICacheStoreOwner; }
 namespace UE::DerivedData { class ILegacyCacheStore; }
-namespace UE::DerivedData { enum class ECacheMethod : uint8; }
 namespace UE::DerivedData { enum class ECachePolicy : uint32; }
 namespace UE::DerivedData { enum class EPriority : uint8; }
 
 namespace UE::DerivedData
 {
 
+class ICacheReplayWriter
+{
+public:
+	UE_API void Write(TConstArrayView<FCachePutRequest> Requests, EPriority Priority);
+	UE_API void Write(TConstArrayView<FCacheGetRequest> Requests, TConstArrayView<FCacheGetResponse> Responses, EPriority Priority);
+	UE_API void Write(TConstArrayView<FCachePutValueRequest> Requests, EPriority Priority);
+	UE_API void Write(TConstArrayView<FCacheGetValueRequest> Requests, TConstArrayView<FCacheGetValueResponse> Responses, EPriority Priority);
+	UE_API void Write(TConstArrayView<FCacheGetChunkRequest> Requests, TConstArrayView<FCacheGetChunkResponse> Responses, EPriority Priority);
+
+	ICacheReplayWriter() = default;
+	virtual ~ICacheReplayWriter() = default;
+
+	// Do not call directly unless forwarding from another ICacheReplayWriter.
+	virtual void SerializeBatch(const FCbWriter& Batch) = 0;
+};
+
+class FCacheReplayArchiveWriter final : public ICacheReplayWriter
+{
+public:
+	UE_API explicit FCacheReplayArchiveWriter(FArchive& Ar, uint64 CompressionBlockSize = 256 * 1024);
+	UE_API ~FCacheReplayArchiveWriter() final;
+
+	FCacheReplayArchiveWriter(const FCacheReplayArchiveWriter&) = delete;
+	FCacheReplayArchiveWriter& operator=(const FCacheReplayArchiveWriter&) = delete;
+
+	void SerializeBatch(const FCbWriter& Batch) final;
+
+private:
+	void WriteCompressedBinaryToArchive(const FSharedBuffer& RawBinary);
+	void FlushCompressedBinaryToArchive();
+
+	FArchive& Ar;
+	FUniqueBuffer RawBlock;
+	FMutableMemoryView RawBlockTail;
+	FMutex Mutex;
+};
+
 class FCacheReplayReader
 {
 public:
diff --git a/Engine/Source/Developer/DerivedDataCache/Private/DerivedDataCacheReplay.cpp b/Engine/Source/Developer/DerivedDataCache/Private/DerivedDataCacheReplay.cpp
--- a/Engine/Source/Developer/DerivedDataCache/Private/DerivedDataCacheReplay.cpp
+++ b/Engine/Source/Developer/DerivedDataCache/Private/DerivedDataCacheReplay.cpp
@@ -3,6 +3,8 @@
 #include "DerivedDataCacheReplay.h"
 
 #include "Algo/AnyOf.h"
+#include "Algo/Copy.h"
+#include "Async/UniqueLock.h"
 #include "Compression/CompressedBuffer.h"
 #include "Containers/StringView.h"
 #include "CoreGlobals.h"
@@ -13,14 +15,14 @@
 #include "DerivedDataRequestOwner.h"
 #include "DerivedDataSerialization.h"
 #include "DerivedDataValue.h"
-#include "HAL/CriticalSection.h"
 #include "HAL/FileManager.h"
 #include "HAL/PlatformTime.h"
+#include "Logging/StructuredLog.h"
 #include "Misc/CommandLine.h"
 #include "Misc/CoreDelegates.h"
 #include "Misc/Optional.h"
 #include "Misc/Parse.h"
-#include "Misc/ScopeLock.h"
+#include "Misc/Zip.h"
 #include "ProfilingDebugging/CpuProfilerTrace.h"
 #include "Serialization/CompactBinary.h"
 #include "Serialization/CompactBinarySerialization.h"
@@ -38,6 +40,34 @@ static constexpr uint64 GCacheReplayCompressionBlockSize = 256 * 1024;
 
 ///////////////////////////////////////////////////////////////////////////////////////////////////////////////////////
 
+class FCacheStoreReplayWriter final : public ICacheReplayWriter
+{
+public:
+	explicit FCacheStoreReplayWriter(const TCHAR* Path, uint64 CompressionBlockSize)
+		: Ar(IFileManager::Get().CreateFileWriter(Path, FILEWRITE_NoFail))
+		, Writer(*Ar, CompressionBlockSize)
+	{
+	}
+
+private:
+	void SerializeBatch(const FCbWriter& Batch) final
+	{
+		if (UE_LOG_ACTIVE(LogDerivedDataCache, Verbose))
+		{
+			TStringBuilder<1024> Json;
+			CompactBinaryToCompactJson(Batch.Save().AsObject(), Json);
+			UE_LOG(LogDerivedDataCache, Verbose, TEXT("Replay: %s"), *Json);
+		}
+
+		Writer.SerializeBatch(Batch);
+	}
+
+	TUniquePtr<FArchive> Ar;
+	FCacheReplayArchiveWriter Writer;
+};
+
+///////////////////////////////////////////////////////////////////////////////////////////////////////////////////////
+
 class FCacheStoreReplay final : public ILegacyCacheStore
 {
 public:
@@ -107,10 +137,7 @@ private:
 	FCacheKeyFilter KeyFilter;
 	FCacheMethodFilter MethodFilter;
 	FString ReplayPath;
-	TUniquePtr<FArchive> ReplayAr;
-	FUniqueBuffer RawBlock;
-	FMutableMemoryView RawBlockTail;
-	FCriticalSection Lock;
+	TUniquePtr<FCacheStoreReplayWriter> ReplayWriter;
 	TOptional<FCacheReplayReader> ReplayReader;
 
 #if WITH_EDITOR
@@ -137,7 +164,8 @@ struct TCacheStoreReplayBatch
 	using FResponse = TCacheResponseFor<FRequest>;
 	using FOnComplete = TCacheOnCompleteFor<FRequest>;
 
-	FCacheStoreReplay* Replay;
+	ICacheReplayWriter& ReplayWriter;
+	const FCacheKeyFilter& KeyFilter;
 	const TArray<FRequest> Requests;
 	TArray<FResponse> Responses;
 	FOnComplete OnComplete;
@@ -145,11 +173,14 @@ struct TCacheStoreReplayBatch
 	EPriority Priority;
 
 	explicit TCacheStoreReplayBatch(
-		FCacheStoreReplay* InReplay,
+		ICacheStore& InCacheStore,
+		ICacheReplayWriter& InReplayWriter,
+		const FCacheKeyFilter& InKeyFilter,
 		TConstArrayView<FRequest> InRequests,
 		IRequestOwner& Owner,
 		FOnComplete&& InOnComplete)
-		: Replay(InReplay)
+		: ReplayWriter(InReplayWriter)
+		, KeyFilter(InKeyFilter)
 		, Requests(InRequests)
 		, OnComplete(MoveTemp(InOnComplete))
 		, ActiveRequests(Requests.Num())
@@ -164,7 +195,7 @@ struct TCacheStoreReplayBatch
 			Request.UserData = Index++;
 		}
 
-		(Replay->GetInnerCache()->*CacheFunctionFor<FRequest>)(ModifiedRequests, Owner, [this](FResponse&& Response)
+		(InCacheStore.*CacheFunctionFor<FRequest>)(ModifiedRequests, Owner, [this](FResponse&& Response)
 		{
 			const int32 Index = int32(Response.UserData);
 			Response.UserData = Requests[Index].UserData;
@@ -174,7 +205,17 @@ struct TCacheStoreReplayBatch
 
 			if (ActiveRequests.fetch_sub(1, std::memory_order_acq_rel) == 1)
 			{
-				Replay->SerializeRequests<FRequest>(Requests, Responses, CacheMethodFor<FRequest>, Priority);
+				TArray<FRequest, TInlineAllocator<16>> ReplayRequests;
+				TArray<FResponse, TInlineAllocator<16>> ReplayResponses;
+				for (const auto& Pair : Zip(Requests, Responses))
+				{
+					if (KeyFilter.IsMatch(Pair.template Get<0>().Key))
+					{
+						ReplayRequests.Add(Pair.template Get<0>());
+						ReplayResponses.Add(Pair.template Get<1>());
+					}
+				}
+				ReplayWriter.Write(ReplayRequests, ReplayResponses, Priority);
 				delete this;
 			}
 		});
@@ -207,19 +248,13 @@ FCacheStoreReplay::FCacheStoreReplay(
 			FCoreDelegates::OnMultiprocessWorkerDetached.AddRaw(this, &FCacheStoreReplay::WorkerDetached);
 		}
 	#endif // WITH_EDITOR
-		ReplayAr.Reset(IFileManager::Get().CreateFileWriter(*ReplayPath, FILEWRITE_NoFail));
-	}
-	if (CompressionBlockSize)
-	{
-		RawBlock = FUniqueBuffer::Alloc(CompressionBlockSize);
-		RawBlockTail = RawBlock;
+		ReplayWriter = MakeUnique<FCacheStoreReplayWriter>(*ReplayPath, CompressionBlockSize);
+		UE_LOGFMT(LogDerivedDataCache, Display, "Replay: Saving cache replay to '{ReplayPath}'", ReplayPath);
 	}
-	UE_CLOG(ReplayAr, LogDerivedDataCache, Display, TEXT("Replay: Saving cache replay to '%s'"), *ReplayPath);
 }
 
 FCacheStoreReplay::~FCacheStoreReplay()
 {
-	FlushToArchive();
 	delete InnerCache;
 	
 #if WITH_EDITOR
@@ -232,113 +267,17 @@ FCacheStoreReplay::~FCacheStoreReplay()
 #endif // WITH_EDITOR
 }
 
-template <typename RequestType>
-void FCacheStoreReplay::SerializeRequests(
-	const TConstArrayView<RequestType> Requests,
-	const TConstArrayView<TCacheResponseFor<RequestType>> Responses,
-	const ECacheMethod Method,
-	const EPriority Priority)
-{
-	if (!ReplayAr)
-	{
-		return;
-	}
-
-	TRACE_CPUPROFILER_EVENT_SCOPE(ReplayDDC_Serialize);
-
-	const auto IsKeyMatch = [this](const RequestType& Request) { return KeyFilter.IsMatch(Request.Key); };
-	if (!MethodFilter.IsMatch(Method) || !Algo::AnyOf(Requests, IsKeyMatch))
-	{
-		return;
-	}
-
-	TCbWriter<512> Writer;
-	Writer.BeginObject();
-	Writer << ANSITEXTVIEW("Method") << Method;
-	Writer << ANSITEXTVIEW("Priority") << Priority;
-	Writer.BeginArray(ANSITEXTVIEW("Requests"));
-	int32 Index = -1;
-	for (const RequestType& Request : Requests)
-	{
-		++Index;
-		if (KeyFilter.IsMatch(Request.Key))
-		{
-			SaveToCompactBinary(Writer, Request, &Responses[Index]);
-		}
-	}
-	Writer.EndArray();
-	Writer.EndObject();
-
-	if (UE_LOG_ACTIVE(LogDerivedDataCache, Verbose))
-	{
-		TStringBuilder<1024> Batch;
-		CompactBinaryToCompactJson(Writer.Save().AsObject(), Batch);
-		UE_LOG(LogDerivedDataCache, Verbose, TEXT("Replay: %s"), *Batch);
-	}
-
-	WriteToArchive(Writer);
-}
-
-void FCacheStoreReplay::WriteBinaryToArchive(const FCompositeBuffer& RawBinary)
-{
-	const FValue CompressedBinary = FValue::Compress(RawBinary, RawBlock.GetSize());
-	TCbWriter<64> BinaryWriter;
-	BinaryWriter.AddBinary(CompressedBinary.GetData().GetCompressed());
-	BinaryWriter.Save(*ReplayAr);
-}
-
-void FCacheStoreReplay::WriteToArchive(FCbWriter& Writer)
-{
-	FScopeLock ScopeLock(&Lock);
-	if (RawBlock)
-	{
-		const uint64 SaveSize = Writer.GetSaveSize();
-		if (RawBlockTail.GetSize() < SaveSize)
-		{
-			FlushToArchive();
-		}
-		if (RawBlockTail.GetSize() < SaveSize)
-		{
-			WriteBinaryToArchive(Writer.Save().AsObject().GetBuffer());
-		}
-		else
-		{
-			Writer.Save(RawBlockTail.Left(SaveSize));
-			RawBlockTail += SaveSize;
-		}
-	}
-	else
-	{
-		Writer.Save(*ReplayAr);
-	}
-}
-
-void FCacheStoreReplay::FlushToArchive()
-{
-	const FSharedBuffer RawBlockHead = FSharedBuffer::MakeView(RawBlock.GetView().LeftChop(RawBlockTail.GetSize()));
-	if (RawBlockHead.GetSize() > 0)
-	{
-		WriteBinaryToArchive(FCompositeBuffer(RawBlockHead));
-		RawBlockTail = RawBlock;
-	}
-}
-
 void FCacheStoreReplay::Put(
 	const TConstArrayView<FCachePutRequest> Requests,
 	IRequestOwner& Owner,
 	FOnCachePutComplete&& OnComplete)
 {
-	// Serialize Put() as a Get() except for the Put method.
-	if (ReplayAr)
+	const auto IsKeyMatch = [this](const FCachePutRequest& Request) { return KeyFilter.IsMatch(Request.Record.GetKey()); };
+	if (ReplayWriter && MethodFilter.IsMatch(ECacheMethod::Put) && Algo::AnyOf(Requests, IsKeyMatch))
 	{
-		TArray<FCacheGetRequest, TInlineAllocator<16>> GetRequests;
-		TArray<FCacheGetResponse, TInlineAllocator<16>> GetResponses;
-		for (const FCachePutRequest& Request : Requests)
-		{
-			GetRequests.Add({Request.Name, Request.Record.GetKey(), Request.Policy, Request.UserData});
-			GetResponses.Add({Request.Name, Request.Record, Request.UserData, EStatus::Ok});
-		}
-		SerializeRequests<FCacheGetRequest>(GetRequests, GetResponses, ECacheMethod::Put, Owner.GetPriority());
+		TArray<FCachePutRequest, TInlineAllocator<16>> ReplayRequests;
+		Algo::CopyIf(Requests, ReplayRequests, IsKeyMatch);
+		ReplayWriter->Write(ReplayRequests, Owner.GetPriority());
 	}
 
 	InnerCache->Put(Requests, Owner, MoveTemp(OnComplete));
@@ -349,7 +288,15 @@ void FCacheStoreReplay::Get(
 	IRequestOwner& Owner,
 	FOnCacheGetComplete&& OnComplete)
 {
-	new TCacheStoreReplayBatch(this, Requests, Owner, MoveTemp(OnComplete));
+	const auto IsKeyMatch = [this](const FCacheGetRequest& Request) { return KeyFilter.IsMatch(Request.Key); };
+	if (ReplayWriter && MethodFilter.IsMatch(ECacheMethod::Get) && Algo::AnyOf(Requests, IsKeyMatch))
+	{
+		new TCacheStoreReplayBatch(*InnerCache, *ReplayWriter, KeyFilter, Requests, Owner, MoveTemp(OnComplete));
+	}
+	else
+	{
+		InnerCache->Get(Requests, Owner, MoveTemp(OnComplete));
+	}
 }
 
 void FCacheStoreReplay::PutValue(
@@ -357,17 +304,12 @@ void FCacheStoreReplay::PutValue(
 	IRequestOwner& Owner,
 	FOnCachePutValueComplete&& OnComplete)
 {
-	// Serialize PutValue() as a GetValue() except for the PutValue method.
-	if (ReplayAr)
+	const auto IsKeyMatch = [this](const FCachePutValueRequest& Request) { return KeyFilter.IsMatch(Request.Key); };
+	if (ReplayWriter && MethodFilter.IsMatch(ECacheMethod::PutValue) && Algo::AnyOf(Requests, IsKeyMatch))
 	{
-		TArray<FCacheGetValueRequest, TInlineAllocator<16>> GetRequests;
-		TArray<FCacheGetValueResponse, TInlineAllocator<16>> GetResponses;
-		for (const FCachePutValueRequest& Request : Requests)
-		{
-			GetRequests.Add({Request.Name, Request.Key, Request.Policy, Request.UserData});
-			GetResponses.Add({Request.Name, Request.Key, Request.Value, Request.UserData, EStatus::Ok});
-		}
-		SerializeRequests<FCacheGetValueRequest>(GetRequests, GetResponses, ECacheMethod::PutValue, Owner.GetPriority());
+		TArray<FCachePutValueRequest, TInlineAllocator<16>> ReplayRequests;
+		Algo::CopyIf(Requests, ReplayRequests, IsKeyMatch);
+		ReplayWriter->Write(ReplayRequests, Owner.GetPriority());
 	}
 
 	InnerCache->PutValue(Requests, Owner, MoveTemp(OnComplete));
@@ -378,7 +320,15 @@ void FCacheStoreReplay::GetValue(
 	IRequestOwner& Owner,
 	FOnCacheGetValueComplete&& OnComplete)
 {
-	new TCacheStoreReplayBatch(this, Requests, Owner, MoveTemp(OnComplete));
+	const auto IsKeyMatch = [this](const FCacheGetValueRequest& Request) { return KeyFilter.IsMatch(Request.Key); };
+	if (ReplayWriter && MethodFilter.IsMatch(ECacheMethod::GetValue) && Algo::AnyOf(Requests, IsKeyMatch))
+	{
+		new TCacheStoreReplayBatch(*InnerCache, *ReplayWriter, KeyFilter, Requests, Owner, MoveTemp(OnComplete));
+	}
+	else
+	{
+		InnerCache->GetValue(Requests, Owner, MoveTemp(OnComplete));
+	}
 }
 
 void FCacheStoreReplay::GetChunks(
@@ -386,7 +336,15 @@ void FCacheStoreReplay::GetChunks(
 	IRequestOwner& Owner,
 	FOnCacheGetChunkComplete&& OnComplete)
 {
-	new TCacheStoreReplayBatch(this, Requests, Owner, MoveTemp(OnComplete));
+	const auto IsKeyMatch = [this](const FCacheGetChunkRequest& Request) { return KeyFilter.IsMatch(Request.Key); };
+	if (ReplayWriter && MethodFilter.IsMatch(ECacheMethod::GetChunks) && Algo::AnyOf(Requests, IsKeyMatch))
+	{
+		new TCacheStoreReplayBatch(*InnerCache, *ReplayWriter, KeyFilter, Requests, Owner, MoveTemp(OnComplete));
+	}
+	else
+	{
+		InnerCache->GetChunks(Requests, Owner, MoveTemp(OnComplete));
+	}
 }
 
 #if WITH_EDITOR
@@ -409,6 +367,8 @@ void FCacheStoreReplay::MergeWorkerReplays()
 	{
 		return;
 	}
+	ReplayWriter.Reset();
+	TUniquePtr<FArchive> ReplayAr(IFileManager::Get().CreateFileWriter(*ReplayPath, FILEWRITE_NoFail | FILEWRITE_Append));
 	for (const TPair<int32, bool> Worker : WorkerIdToState)
 	{
 		const int32 WorkerId = Worker.Key;
@@ -444,6 +404,131 @@ void FCacheStoreReplay::MergeWorkerReplays()
 
 ///////////////////////////////////////////////////////////////////////////////////////////////////////////////////////
 
+template <typename RequestType>
+static void SaveBatchToCompactBinary(
+	FCbWriter& Writer,
+	const TConstArrayView<RequestType> Requests,
+	const TConstArrayView<TCacheResponseFor<RequestType>> Responses,
+	const ECacheMethod Method,
+	const EPriority Priority)
+{
+	Writer.BeginObject();
+	Writer << ANSITEXTVIEW("Method") << Method;
+	Writer << ANSITEXTVIEW("Priority") << Priority;
+	Writer.BeginArray(ANSITEXTVIEW("Requests"));
+	int32 Index = -1;
+	for (const RequestType& Request : Requests)
+	{
+		++Index;
+		SaveToCompactBinary(Writer, Request, Responses.IsValidIndex(Index) ? &Responses[Index] : nullptr);
+	}
+	Writer.EndArray();
+	Writer.EndObject();
+}
+
+void ICacheReplayWriter::Write(TConstArrayView<FCachePutRequest> Requests, EPriority Priority)
+{
+	TCbWriter<512> Writer;
+	SaveBatchToCompactBinary(Writer, Requests, {}, ECacheMethod::Put, Priority);
+	SerializeBatch(Writer);
+}
+
+void ICacheReplayWriter::Write(TConstArrayView<FCacheGetRequest> Requests, TConstArrayView<FCacheGetResponse> Responses, EPriority Priority)
+{
+	TCbWriter<512> Writer;
+	SaveBatchToCompactBinary(Writer, Requests, Responses, ECacheMethod::Get, Priority);
+	SerializeBatch(Writer);
+}
+
+void ICacheReplayWriter::Write(TConstArrayView<FCachePutValueRequest> Requests, EPriority Priority)
+{
+	TCbWriter<512> Writer;
+	SaveBatchToCompactBinary(Writer, Requests, {}, ECacheMethod::PutValue, Priority);
+	SerializeBatch(Writer);
+}
+
+void ICacheReplayWriter::Write(TConstArrayView<FCacheGetValueRequest> Requests, TConstArrayView<FCacheGetValueResponse> Responses, EPriority Priority)
+{
+	TCbWriter<512> Writer;
+	SaveBatchToCompactBinary(Writer, Requests, Responses, ECacheMethod::GetValue, Priority);
+	SerializeBatch(Writer);
+}
+
+void ICacheReplayWriter::Write(TConstArrayView<FCacheGetChunkRequest> Requests, TConstArrayView<FCacheGetChunkResponse> Responses, EPriority Priority)
+{
+	TCbWriter<512> Writer;
+	SaveBatchToCompactBinary(Writer, Requests, Responses, ECacheMethod::GetChunks, Priority);
+	SerializeBatch(Writer);
+}
+
+///////////////////////////////////////////////////////////////////////////////////////////////////////////////////////
+
+FCacheReplayArchiveWriter::FCacheReplayArchiveWriter(FArchive& InAr, uint64 InCompressionBlockSize)
+	: Ar(InAr)
+{
+	if (InCompressionBlockSize)
+	{
+		RawBlock = FUniqueBuffer::Alloc(InCompressionBlockSize);
+		RawBlockTail = RawBlock;
+	}
+}
+
+FCacheReplayArchiveWriter::~FCacheReplayArchiveWriter()
+{
+	FlushCompressedBinaryToArchive();
+}
+
+void FCacheReplayArchiveWriter::SerializeBatch(const FCbWriter& Batch)
+{
+	TUniqueLock Lock(Mutex);
+
+	if (!RawBlock)
+	{
+		// Write Batch directly to the archive when not using compression.
+		Batch.Save(Ar);
+		return;
+	}
+
+	const uint64 SaveSize = Batch.GetSaveSize();
+	if (RawBlockTail.GetSize() < SaveSize)
+	{
+		// Flush buffered batches because this batch exceeds the remaining buffer size.
+		FlushCompressedBinaryToArchive();
+	}
+
+	if (RawBlockTail.GetSize() < SaveSize)
+	{
+		// Write Batch into its own compressed block because it exceeds the block size.
+		WriteCompressedBinaryToArchive(Batch.Save().AsObject().GetBuffer().ToShared());
+	}
+	else
+	{
+		// Write Batch into the buffer to be compressed with a block of batches.
+		Batch.Save(RawBlockTail.Left(SaveSize));
+		RawBlockTail += SaveSize;
+	}
+}
+
+void FCacheReplayArchiveWriter::WriteCompressedBinaryToArchive(const FSharedBuffer& RawBinary)
+{
+	const FValue CompressedBinary = FValue::Compress(RawBinary, RawBlock.GetSize());
+	TCbWriter<64> BinaryWriter;
+	BinaryWriter.AddBinary(CompressedBinary.GetData().GetCompressed());
+	BinaryWriter.Save(Ar);
+}
+
+void FCacheReplayArchiveWriter::FlushCompressedBinaryToArchive()
+{
+	const FSharedBuffer RawBlockHead = FSharedBuffer::MakeView(RawBlock.GetView().LeftChop(RawBlockTail.GetSize()));
+	if (RawBlockHead.GetSize() > 0)
+	{
+		WriteCompressedBinaryToArchive(RawBlockHead);
+		RawBlockTail = RawBlock;
+	}
+}
+
+///////////////////////////////////////////////////////////////////////////////////////////////////////////////////////
+
 class FCacheReplayReader::FState
 {
 public:
diff --git a/Engine/Source/Developer/DerivedDataCache/Private/DerivedDataSerialization.cpp b/Engine/Source/Developer/DerivedDataCache/Private/DerivedDataSerialization.cpp
--- a/Engine/Source/Developer/DerivedDataCache/Private/DerivedDataSerialization.cpp
+++ b/Engine/Source/Developer/DerivedDataCache/Private/DerivedDataSerialization.cpp
@@ -86,8 +86,10 @@ bool LoadMetaFromCompactBinary(FCbFieldView Field, FCbObject& OutMeta, FLoadAtta
 
 ///////////////////////////////////////////////////////////////////////////////////////////////////////////////////////
 
-void SaveToCompactBinary(FCbWriter& Writer, const FCachePutRequest& Request, FSaveAttachmentFn SaveAttachment)
+void SaveToCompactBinary(FCbWriter& Writer, const FCachePutRequest& Request, const FCachePutResponse* Response, FSaveAttachmentFn SaveAttachment)
 {
+	check(!Response || Request.Record.GetKey() == Response->Key);
+	check(!Response || Request.UserData == Response->UserData);
 	Writer.BeginObject();
 	if (!Request.Name.IsEmpty())
 	{
@@ -102,6 +104,10 @@ void SaveToCompactBinary(FCbWriter& Writer, const FCachePutRequest& Request, FSa
 	{
 		Writer << ANSITEXTVIEW("UserData") << Request.UserData;
 	}
+	if (Response && Response->Status != EStatus::Ok)
+	{
+		Writer << ANSITEXTVIEW("Status") << Response->Status;
+	}
 	Writer.EndObject();
 }
 
@@ -121,6 +127,16 @@ bool LoadFromCompactBinary(FCbFieldView Field, TOptional<FCachePutRequest>& OutR
 	return bOk;
 }
 
+bool LoadFromCompactBinary(FCbFieldView Field, FCachePutResponse& OutResponse)
+{
+	bool bOk = Field.IsObject();
+	LoadFromCompactBinary(Field[ANSITEXTVIEW("Name")], OutResponse.Name);
+	bOk &= LoadFromCompactBinary(Field[ANSITEXTVIEW("Key")], OutResponse.Key);
+	LoadFromCompactBinary(Field[ANSITEXTVIEW("UserData")], OutResponse.UserData, 0);
+	LoadFromCompactBinary(Field[ANSITEXTVIEW("Status")], OutResponse.Status, EStatus::Ok);
+	return bOk;
+}
+
 void SaveToCompactBinary(FCbWriter& Writer, const FCacheGetRequest& Request, const FCacheGetResponse* Response, FSaveAttachmentFn SaveAttachment)
 {
 	check(!Response || Request.Key == Response->Record.GetKey());
@@ -186,8 +202,10 @@ bool LoadFromCompactBinary(FCbFieldView Field, TOptional<FCacheGetResponse>& Out
 
 ///////////////////////////////////////////////////////////////////////////////////////////////////////////////////////
 
-void SaveToCompactBinary(FCbWriter& Writer, const FCachePutValueRequest& Request, FSaveAttachmentFn SaveAttachment)
+void SaveToCompactBinary(FCbWriter& Writer, const FCachePutValueRequest& Request, const FCachePutValueResponse* Response, FSaveAttachmentFn SaveAttachment)
 {
+	check(!Response || Request.Key == Response->Key);
+	check(!Response || Request.UserData == Response->UserData);
 	Writer.BeginObject();
 	if (!Request.Name.IsEmpty())
 	{
@@ -207,6 +225,10 @@ void SaveToCompactBinary(FCbWriter& Writer, const FCachePutValueRequest& Request
 	{
 		Writer << ANSITEXTVIEW("UserData") << Request.UserData;
 	}
+	if (Response && Response->Status != EStatus::Ok)
+	{
+		Writer << ANSITEXTVIEW("Status") << Response->Status;
+	}
 	Writer.EndObject();
 }
 
@@ -221,6 +243,16 @@ bool LoadFromCompactBinary(FCbFieldView Field, FCachePutValueRequest& OutRequest
 	return bOk;
 }
 
+bool LoadFromCompactBinary(FCbFieldView Field, FCachePutValueResponse& OutResponse)
+{
+	bool bOk = Field.IsObject();
+	LoadFromCompactBinary(Field[ANSITEXTVIEW("Name")], OutResponse.Name);
+	bOk &= LoadFromCompactBinary(Field[ANSITEXTVIEW("Key")], OutResponse.Key);
+	LoadFromCompactBinary(Field[ANSITEXTVIEW("UserData")], OutResponse.UserData, 0);
+	LoadFromCompactBinary(Field[ANSITEXTVIEW("Status")], OutResponse.Status, EStatus::Ok);
+	return bOk;
+}
+
 void SaveToCompactBinary(FCbWriter& Writer, const FCacheGetValueRequest& Request, const FCacheGetValueResponse* Response, FSaveAttachmentFn SaveAttachment)
 {
 	check(!Response || Request.Key == Response->Key);
diff --git a/Engine/Source/Developer/DerivedDataCache/Public/DerivedDataSerialization.h b/Engine/Source/Developer/DerivedDataCache/Public/DerivedDataSerialization.h
--- a/Engine/Source/Developer/DerivedDataCache/Public/DerivedDataSerialization.h
+++ b/Engine/Source/Developer/DerivedDataCache/Public/DerivedDataSerialization.h
@@ -28,7 +28,9 @@ namespace UE::DerivedData { struct FCacheGetValueRequest; }
 namespace UE::DerivedData { struct FCacheGetValueResponse; }
 namespace UE::DerivedData { struct FCacheKey; }
 namespace UE::DerivedData { struct FCachePutRequest; }
+namespace UE::DerivedData { struct FCachePutResponse; }
 namespace UE::DerivedData { struct FCachePutValueRequest; }
+namespace UE::DerivedData { struct FCachePutValueResponse; }
 namespace UE::DerivedData { struct FCacheValuePolicy; }
 namespace UE::DerivedData { struct FValueId; }
 namespace UE::DerivedData { enum class ECacheMethod : uint8; }
@@ -64,15 +66,17 @@ UE_API FCbWriter& operator<<(FCbWriter& Writer, const FCacheGetValueRequest& Req
 UE_DEPRECATED(5.8, "Call SaveToCompactBinary(Writer, Request).")
 UE_API FCbWriter& operator<<(FCbWriter& Writer, const FCacheGetChunkRequest& Request);
 
-UE_API void SaveToCompactBinary(FCbWriter& Writer, const FCachePutRequest& Request, FSaveAttachmentFn SaveAttachment = nullptr);
+UE_API void SaveToCompactBinary(FCbWriter& Writer, const FCachePutRequest& Request, const FCachePutResponse* Response = nullptr, FSaveAttachmentFn SaveAttachment = nullptr);
 UE_API bool LoadFromCompactBinary(FCbFieldView Field, TOptional<FCachePutRequest>& OutRequest, FLoadAttachmentFn LoadAttachment = nullptr);
+UE_API bool LoadFromCompactBinary(FCbFieldView Field, FCachePutResponse& OutResponse);
 
 UE_API void SaveToCompactBinary(FCbWriter& Writer, const FCacheGetRequest& Request, const FCacheGetResponse* Response = nullptr, FSaveAttachmentFn SaveAttachment = nullptr);
 UE_API bool LoadFromCompactBinary(FCbFieldView Field, FCacheGetRequest& OutRequest);
 UE_API bool LoadFromCompactBinary(FCbFieldView Field, TOptional<FCacheGetResponse>& OutResponse, FLoadAttachmentFn LoadAttachment = nullptr);
 
-UE_API void SaveToCompactBinary(FCbWriter& Writer, const FCachePutValueRequest& Request, FSaveAttachmentFn SaveAttachment = nullptr);
+UE_API void SaveToCompactBinary(FCbWriter& Writer, const FCachePutValueRequest& Request, const FCachePutValueResponse* Response = nullptr, FSaveAttachmentFn SaveAttachment = nullptr);
 UE_API bool LoadFromCompactBinary(FCbFieldView Field, FCachePutValueRequest& OutRequest, FLoadAttachmentFn LoadAttachment = nullptr);
+UE_API bool LoadFromCompactBinary(FCbFieldView Field, FCachePutValueResponse& OutResponse);
 
 UE_API void SaveToCompactBinary(FCbWriter& Writer, const FCacheGetValueRequest& Request, const FCacheGetValueResponse* Response = nullptr, FSaveAttachmentFn SaveAttachment = nullptr);
 UE_API bool LoadFromCompactBinary(FCbFieldView Field, FCacheGetValueRequest& OutRequest);
diff --git a/Engine/Source/Developer/DerivedDataCache/Tests/DerivedDataSerializationTest.cpp b/Engine/Source/Developer/DerivedDataCache/Tests/DerivedDataSerializationTest.cpp
--- a/Engine/Source/Developer/DerivedDataCache/Tests/DerivedDataSerializationTest.cpp
+++ b/Engine/Source/Developer/DerivedDataCache/Tests/DerivedDataSerializationTest.cpp
@@ -76,7 +76,7 @@ TEST_CASE("DerivedData::Serialization::CompactBinary", "[DerivedData]")
 		const FCbField Field = [&]
 		{
 			TCbWriter<512> Writer;
-			SaveToCompactBinary(Writer, Request, SaveAttachment);
+			SaveToCompactBinary(Writer, Request, nullptr, SaveAttachment);
 			return Writer.Save();
 		}();
 
@@ -109,6 +109,26 @@ TEST_CASE("DerivedData::Serialization::CompactBinary", "[DerivedData]")
 		}
 	}
 
+	SECTION("CachePutResponse")
+	{
+		const FCachePutRequest Request{Name, Record, RecordPolicy, UserData};
+		const FCachePutResponse Response{Name, Key, UserData, EStatus::Canceled};
+
+		const FCbField Field = [&]
+		{
+			TCbWriter<512> Writer;
+			SaveToCompactBinary(Writer, Request, &Response);
+			return Writer.Save();
+		}();
+
+		FCachePutResponse LoadedResponse;
+		CHECK(LoadFromCompactBinary(Field, LoadedResponse));
+		CHECK(LoadedResponse.Name == Response.Name);
+		CHECK(LoadedResponse.Key == Response.Key);
+		CHECK(LoadedResponse.UserData == Response.UserData);
+		CHECK(LoadedResponse.Status == Response.Status);
+	}
+
 	SECTION("CacheGetRequest")
 	{
 		const FCacheGetRequest Request{Name, Key, RecordPolicy, UserData};
@@ -133,7 +153,7 @@ TEST_CASE("DerivedData::Serialization::CompactBinary", "[DerivedData]")
 	SECTION("CacheGetResponse")
 	{
 		const FCacheGetRequest Request{Name, Key, RecordPolicy, UserData};
-		const FCacheGetResponse Response{Name, Record, UserData, EStatus::Ok};
+		const FCacheGetResponse Response{Name, Record, UserData, EStatus::Canceled};
 
 		const FCbField Field = [&]
 		{
@@ -180,7 +200,7 @@ TEST_CASE("DerivedData::Serialization::CompactBinary", "[DerivedData]")
 		const FCbField Field = [&]
 		{
 			TCbWriter<512> Writer;
-			SaveToCompactBinary(Writer, Request, SaveAttachment);
+			SaveToCompactBinary(Writer, Request, nullptr, SaveAttachment);
 			return Writer.Save();
 		}();
 
@@ -205,6 +225,26 @@ TEST_CASE("DerivedData::Serialization::CompactBinary", "[DerivedData]")
 		}
 	}
 
+	SECTION("CachePutValueResponse")
+	{
+		const FCachePutValueRequest Request{Name, Key, Value, ValuePolicy, UserData};
+		const FCachePutValueResponse Response{Name, Key, UserData, EStatus::Canceled};
+
+		const FCbField Field = [&]
+		{
+			TCbWriter<512> Writer;
+			SaveToCompactBinary(Writer, Request, &Response);
+			return Writer.Save();
+		}();
+
+		FCachePutValueResponse LoadedResponse;
+		CHECK(LoadFromCompactBinary(Field, LoadedResponse));
+		CHECK(LoadedResponse.Name == Response.Name);
+		CHECK(LoadedResponse.Key == Response.Key);
+		CHECK(LoadedResponse.UserData == Response.UserData);
+		CHECK(LoadedResponse.Status == Response.Status);
+	}
+
 	SECTION("CacheGetValueRequest")
 	{
 		const FCacheGetValueRequest Request{Name, Key, ValuePolicy, UserData};
@@ -227,7 +267,7 @@ TEST_CASE("DerivedData::Serialization::CompactBinary", "[DerivedData]")
 	SECTION("CacheGetValueResponse")
 	{
 		const FCacheGetValueRequest Request{Name, Key, ValuePolicy, UserData};
-		const FCacheGetValueResponse Response{Name, Key, Value, UserData, EStatus::Ok};
+		const FCacheGetValueResponse Response{Name, Key, Value, UserData, EStatus::Canceled};
 
 		const FCbField Field = [&]
 		{
@@ -285,7 +325,7 @@ TEST_CASE("DerivedData::Serialization::CompactBinary", "[DerivedData]")
 	SECTION("CacheGetChunkResponse")
 	{
 		const FCacheGetChunkRequest Request{Name, Key, ValueId, 1024, 512, Value.GetRawHash(), ValuePolicy, UserData};
-		const FCacheGetChunkResponse Response{Name, Key, ValueId, 1024, 256, Value.GetRawHash(), EmptyBuffer, UserData, EStatus::Ok};
+		const FCacheGetChunkResponse Response{Name, Key, ValueId, 1024, 256, Value.GetRawHash(), EmptyBuffer, UserData, EStatus::Canceled};
 
 		const FCbField Field = [&]
 		{
