DDC: Added a ICacheReplayReader interface

#Synced-CL 32500503

diff --git a/Engine/Source/Developer/DerivedDataCache/Internal/DerivedDataCacheReplay.h b/Engine/Source/Developer/DerivedDataCache/Internal/DerivedDataCacheReplay.h
--- a/Engine/Source/Developer/DerivedDataCache/Internal/DerivedDataCacheReplay.h
+++ b/Engine/Source/Developer/DerivedDataCache/Internal/DerivedDataCacheReplay.h
@@ -5,7 +5,9 @@
 #include "Async/Mutex.h"
 #include "CoreTypes.h"
 #include "DerivedDataCache.h"
+#include "Misc/NotNull.h"
 #include "Templates/PimplPtr.h"
+#include "Templates/UniquePtr.h"
 
 #define UE_API DERIVEDDATACACHE_API
 
@@ -26,38 +28,87 @@ namespace UE::DerivedData
 class ICacheReplayWriter
 {
 public:
-	UE_API void Write(TConstArrayView<FCachePutRequest> Requests, EPriority Priority);
+	UE_API void Write(TConstArrayView<FCachePutRequest> Requests, TConstArrayView<FCachePutResponse> Responses, EPriority Priority);
 	UE_API void Write(TConstArrayView<FCacheGetRequest> Requests, TConstArrayView<FCacheGetResponse> Responses, EPriority Priority);
-	UE_API void Write(TConstArrayView<FCachePutValueRequest> Requests, EPriority Priority);
+	UE_API void Write(TConstArrayView<FCachePutValueRequest> Requests, TConstArrayView<FCachePutValueResponse> Responses, EPriority Priority);
 	UE_API void Write(TConstArrayView<FCacheGetValueRequest> Requests, TConstArrayView<FCacheGetValueResponse> Responses, EPriority Priority);
 	UE_API void Write(TConstArrayView<FCacheGetChunkRequest> Requests, TConstArrayView<FCacheGetChunkResponse> Responses, EPriority Priority);
 
-	ICacheReplayWriter() = default;
 	virtual ~ICacheReplayWriter() = default;
 
 	// Do not call directly unless forwarding from another ICacheReplayWriter.
 	virtual void SerializeBatch(const FCbWriter& Batch) = 0;
 };
 
-class FCacheReplayArchiveWriter final : public ICacheReplayWriter
+class ICacheReplayReader
 {
 public:
-	UE_API explicit FCacheReplayArchiveWriter(FArchive& Ar, uint64 CompressionBlockSize = 256 * 1024);
-	UE_API ~FCacheReplayArchiveWriter() final;
+	virtual void Read(TConstArrayView<FCachePutRequest> Requests, TConstArrayView<FCachePutResponse> Responses, EPriority Priority) = 0;
+	virtual void Read(TConstArrayView<FCacheGetRequest> Requests, TConstArrayView<FCacheGetResponse> Responses, EPriority Priority) = 0;
+	virtual void Read(TConstArrayView<FCachePutValueRequest> Requests, TConstArrayView<FCachePutValueResponse> Responses, EPriority Priority) = 0;
+	virtual void Read(TConstArrayView<FCacheGetValueRequest> Requests, TConstArrayView<FCacheGetValueResponse> Responses, EPriority Priority) = 0;
+	virtual void Read(TConstArrayView<FCacheGetChunkRequest> Requests, TConstArrayView<FCacheGetChunkResponse> Responses, EPriority Priority) = 0;
 
-	FCacheReplayArchiveWriter(const FCacheReplayArchiveWriter&) = delete;
-	FCacheReplayArchiveWriter& operator=(const FCacheReplayArchiveWriter&) = delete;
+	virtual ~ICacheReplayReader() = default;
 
-	void SerializeBatch(const FCbWriter& Batch) final;
+	// Do not call directly unless forwarding from another ICacheReplayReader.
+	virtual bool SerializeBatch(FCbObjectView Batch);
+};
 
-private:
-	void WriteCompressedBinaryToArchive(const FSharedBuffer& RawBinary);
-	void FlushCompressedBinaryToArchive();
+/**
+ * Create a cache replay writer that serializes into an archive.
+ *
+ * @param Ar   Archive to serialize to, and delete when the writer is deleted.
+ * @param CompressionBlockSize   Power-of-two block size or zero to disable compression.
+ */
+UE_API TUniquePtr<ICacheReplayWriter> CreateReplayArchiveWriter(TNotNull<TUniquePtr<FArchive>> Ar, uint64 CompressionBlockSize = 256 * 1024);
+
+/**
+ * Reads a replay from an archive that was created by CreateReplayArchiveWriter.
+ *
+ * @return True if the replay was serialized successfully, otherwise false.
+ */
+UE_API bool ReadReplayFromArchive(ICacheReplayReader& Reader, FArchive& Ar);
+UE_API bool ReadReplayFromFile(ICacheReplayReader& Reader, const TCHAR* Path);
+
+/**
+ * A replay reader that forwards every batch to a writer.
+ */
+class FCacheReplayForwardingReader final : public ICacheReplayReader
+{
+public:
+	inline explicit FCacheReplayForwardingReader(ICacheReplayWriter& TargetWriter)
+		: Writer(TargetWriter)
+	{
+	}
+
+	void Read(TConstArrayView<FCachePutRequest> Requests, TConstArrayView<FCachePutResponse> Responses, EPriority Priority) final
+	{
+		Writer.Write(Requests, Responses, Priority);
+	}
+
+	void Read(TConstArrayView<FCacheGetRequest> Requests, TConstArrayView<FCacheGetResponse> Responses, EPriority Priority) final
+	{
+		Writer.Write(Requests, Responses, Priority);
+	}
+
+	void Read(TConstArrayView<FCachePutValueRequest> Requests, TConstArrayView<FCachePutValueResponse> Responses, EPriority Priority) final
+	{
+		Writer.Write(Requests, Responses, Priority);
+	}
+
+	void Read(TConstArrayView<FCacheGetValueRequest> Requests, TConstArrayView<FCacheGetValueResponse> Responses, EPriority Priority) final
+	{
+		Writer.Write(Requests, Responses, Priority);
+	}
+
+	void Read(TConstArrayView<FCacheGetChunkRequest> Requests, TConstArrayView<FCacheGetChunkResponse> Responses, EPriority Priority) final
+	{
+		Writer.Write(Requests, Responses, Priority);
+	}
 
-	FArchive& Ar;
-	FUniqueBuffer RawBlock;
-	FMutableMemoryView RawBlockTail;
-	FMutex Mutex;
+private:
+	ICacheReplayWriter& Writer;
 };
 
 class FCacheReplayReader
diff --git a/Engine/Source/Developer/DerivedDataCache/Private/DerivedDataCacheReplay.cpp b/Engine/Source/Developer/DerivedDataCache/Private/DerivedDataCacheReplay.cpp
--- a/Engine/Source/Developer/DerivedDataCache/Private/DerivedDataCacheReplay.cpp
+++ b/Engine/Source/Developer/DerivedDataCache/Private/DerivedDataCacheReplay.cpp
@@ -3,7 +3,6 @@
 #include "DerivedDataCacheReplay.h"
 
 #include "Algo/AnyOf.h"
-#include "Algo/Copy.h"
 #include "Async/UniqueLock.h"
 #include "Compression/CompressedBuffer.h"
 #include "Containers/StringView.h"
@@ -44,9 +43,9 @@ class FCacheStoreReplayWriter final : public ICacheReplayWriter
 {
 public:
 	explicit FCacheStoreReplayWriter(const TCHAR* Path, uint64 CompressionBlockSize)
-		: Ar(IFileManager::Get().CreateFileWriter(Path, FILEWRITE_NoFail))
-		, Writer(*Ar, CompressionBlockSize)
 	{
+		TUniquePtr<FArchive> Ar(IFileManager::Get().CreateFileWriter(Path, FILEWRITE_NoFail));
+		Writer = CreateReplayArchiveWriter(MoveTemp(Ar), CompressionBlockSize);
 	}
 
 private:
@@ -59,11 +58,10 @@ private:
 			UE_LOG(LogDerivedDataCache, Verbose, TEXT("Replay: %s"), *Json);
 		}
 
-		Writer.SerializeBatch(Batch);
+		Writer->SerializeBatch(Batch);
 	}
 
-	TUniquePtr<FArchive> Ar;
-	FCacheReplayArchiveWriter Writer;
+	TUniquePtr<ICacheReplayWriter> Writer;
 };
 
 ///////////////////////////////////////////////////////////////////////////////////////////////////////////////////////
@@ -209,7 +207,7 @@ struct TCacheStoreReplayBatch
 				TArray<FResponse, TInlineAllocator<16>> ReplayResponses;
 				for (const auto& Pair : Zip(Requests, Responses))
 				{
-					if (KeyFilter.IsMatch(Pair.template Get<0>().Key))
+					if (KeyFilter.IsMatch(GetCacheKey(Pair.template Get<0>())))
 					{
 						ReplayRequests.Add(Pair.template Get<0>());
 						ReplayResponses.Add(Pair.template Get<1>());
@@ -275,12 +273,12 @@ void FCacheStoreReplay::Put(
 	const auto IsKeyMatch = [this](const FCachePutRequest& Request) { return KeyFilter.IsMatch(Request.Record.GetKey()); };
 	if (ReplayWriter && MethodFilter.IsMatch(ECacheMethod::Put) && Algo::AnyOf(Requests, IsKeyMatch))
 	{
-		TArray<FCachePutRequest, TInlineAllocator<16>> ReplayRequests;
-		Algo::CopyIf(Requests, ReplayRequests, IsKeyMatch);
-		ReplayWriter->Write(ReplayRequests, Owner.GetPriority());
+		new TCacheStoreReplayBatch(*InnerCache, *ReplayWriter, KeyFilter, Requests, Owner, MoveTemp(OnComplete));
+	}
+	else
+	{
+		InnerCache->Put(Requests, Owner, MoveTemp(OnComplete));
 	}
-
-	InnerCache->Put(Requests, Owner, MoveTemp(OnComplete));
 }
 
 void FCacheStoreReplay::Get(
@@ -307,12 +305,12 @@ void FCacheStoreReplay::PutValue(
 	const auto IsKeyMatch = [this](const FCachePutValueRequest& Request) { return KeyFilter.IsMatch(Request.Key); };
 	if (ReplayWriter && MethodFilter.IsMatch(ECacheMethod::PutValue) && Algo::AnyOf(Requests, IsKeyMatch))
 	{
-		TArray<FCachePutValueRequest, TInlineAllocator<16>> ReplayRequests;
-		Algo::CopyIf(Requests, ReplayRequests, IsKeyMatch);
-		ReplayWriter->Write(ReplayRequests, Owner.GetPriority());
+		new TCacheStoreReplayBatch(*InnerCache, *ReplayWriter, KeyFilter, Requests, Owner, MoveTemp(OnComplete));
+	}
+	else
+	{
+		InnerCache->PutValue(Requests, Owner, MoveTemp(OnComplete));
 	}
-
-	InnerCache->PutValue(Requests, Owner, MoveTemp(OnComplete));
 }
 
 void FCacheStoreReplay::GetValue(
@@ -367,8 +365,7 @@ void FCacheStoreReplay::MergeWorkerReplays()
 	{
 		return;
 	}
-	ReplayWriter.Reset();
-	TUniquePtr<FArchive> ReplayAr(IFileManager::Get().CreateFileWriter(*ReplayPath, FILEWRITE_NoFail | FILEWRITE_Append));
+	FCacheReplayForwardingReader ForwardReplay(*ReplayWriter);
 	for (const TPair<int32, bool> Worker : WorkerIdToState)
 	{
 		const int32 WorkerId = Worker.Key;
@@ -380,20 +377,8 @@ void FCacheStoreReplay::MergeWorkerReplays()
 				TEXT("Replay: Skipped replay file '%s' because its worker has not detached."), *WorkerReplayPath);
 			continue;
 		}
-		if (TUniquePtr<FArchive> WorkerAr{IFileManager::Get().CreateFileReader(*WorkerReplayPath)})
+		if (ReadReplayFromFile(ForwardReplay, *WorkerReplayPath))
 		{
-			int64 RemainingSize = WorkerAr->TotalSize();
-			constexpr int64 MaxBufferSize = 64 * 1024 * 1024;
-			const int64 BufferSize = FMath::Min(RemainingSize, MaxBufferSize);
-			TUniquePtr<uint8[]> Buffer(new uint8[BufferSize]);
-			while (RemainingSize > 0)
-			{
-				const int64 BlockSize = FMath::Min(RemainingSize, BufferSize);
-				WorkerAr->Serialize(Buffer.Get(), BlockSize);
-				ReplayAr->Serialize(Buffer.Get(), BlockSize);
-				RemainingSize -= BlockSize;
-			}
-			WorkerAr = nullptr;
 			IFileManager::Get().Delete(*WorkerReplayPath);
 			UE_LOG(LogDerivedDataCache, Display,
 				TEXT("Replay: Merged replay file '%s' from a worker process."), *WorkerReplayPath);
@@ -412,24 +397,28 @@ static void SaveBatchToCompactBinary(
 	const ECacheMethod Method,
 	const EPriority Priority)
 {
+	check(Requests.Num() == Responses.Num());
 	Writer.BeginObject();
 	Writer << ANSITEXTVIEW("Method") << Method;
-	Writer << ANSITEXTVIEW("Priority") << Priority;
+	if (Priority != EPriority::Normal)
+	{
+		Writer << ANSITEXTVIEW("Priority") << Priority;
+	}
 	Writer.BeginArray(ANSITEXTVIEW("Requests"));
 	int32 Index = -1;
 	for (const RequestType& Request : Requests)
 	{
 		++Index;
-		SaveToCompactBinary(Writer, Request, Responses.IsValidIndex(Index) ? &Responses[Index] : nullptr);
+		SaveToCompactBinary(Writer, Request, &Responses.GetData()[Index]);
 	}
 	Writer.EndArray();
 	Writer.EndObject();
 }
 
-void ICacheReplayWriter::Write(TConstArrayView<FCachePutRequest> Requests, EPriority Priority)
+void ICacheReplayWriter::Write(TConstArrayView<FCachePutRequest> Requests, TConstArrayView<FCachePutResponse> Responses, EPriority Priority)
 {
 	TCbWriter<512> Writer;
-	SaveBatchToCompactBinary(Writer, Requests, {}, ECacheMethod::Put, Priority);
+	SaveBatchToCompactBinary(Writer, Requests, Responses, ECacheMethod::Put, Priority);
 	SerializeBatch(Writer);
 }
 
@@ -440,10 +429,10 @@ void ICacheReplayWriter::Write(TConstArrayView<FCacheGetRequest> Requests, TCons
 	SerializeBatch(Writer);
 }
 
-void ICacheReplayWriter::Write(TConstArrayView<FCachePutValueRequest> Requests, EPriority Priority)
+void ICacheReplayWriter::Write(TConstArrayView<FCachePutValueRequest> Requests, TConstArrayView<FCachePutValueResponse> Responses, EPriority Priority)
 {
 	TCbWriter<512> Writer;
-	SaveBatchToCompactBinary(Writer, Requests, {}, ECacheMethod::PutValue, Priority);
+	SaveBatchToCompactBinary(Writer, Requests, Responses, ECacheMethod::PutValue, Priority);
 	SerializeBatch(Writer);
 }
 
@@ -463,8 +452,103 @@ void ICacheReplayWriter::Write(TConstArrayView<FCacheGetChunkRequest> Requests,
 
 ///////////////////////////////////////////////////////////////////////////////////////////////////////////////////////
 
-FCacheReplayArchiveWriter::FCacheReplayArchiveWriter(FArchive& InAr, uint64 InCompressionBlockSize)
-	: Ar(InAr)
+template <typename CacheType, typename ArrayType>
+static bool LoadCacheTypeFromCompactBinary(const FCbFieldView Field, ArrayType& OutArray)
+{
+	if constexpr (std::is_default_constructible_v<CacheType>)
+	{
+		CacheType& Value = OutArray.AddDefaulted_GetRef();
+		return LoadFromCompactBinary(Field, Value);
+	}
+	else
+	{
+		TOptional<CacheType> Value;
+		if (LoadFromCompactBinary(Field, Value))
+		{
+			OutArray.Emplace(MoveTemp(Value.GetValue()));
+			return true;
+		}
+		return false;
+	}
+}
+
+template <typename RequestType>
+static bool LoadBatchFromCompactBinary(ICacheReplayReader& Reader, const FCbObjectView Batch)
+{
+	EPriority Priority;
+	LoadFromCompactBinary(Batch[ANSITEXTVIEW("Priority")], Priority, EPriority::Normal);
+
+	TArray<RequestType, TInlineAllocator<16>> Requests;
+	TArray<TCacheResponseFor<RequestType>, TInlineAllocator<16>> Responses;
+	{
+		const FCbArrayView Array = Batch[ANSITEXTVIEW("Requests")].AsArrayView();
+		const int32 BatchSize = IntCastChecked<int32>(Array.Num());
+		Requests.Reserve(BatchSize);
+		Responses.Reserve(BatchSize);
+		for (FCbFieldView Field : Array)
+		{
+			if (!LoadCacheTypeFromCompactBinary<RequestType>(Field, Requests) ||
+				!LoadCacheTypeFromCompactBinary<TCacheResponseFor<RequestType>>(Field, Responses))
+			{
+				return false;
+			}
+		}
+	}
+
+	Reader.Read(Requests, Responses, Priority);
+	return true;
+}
+
+bool ICacheReplayReader::SerializeBatch(const FCbObjectView Batch)
+{
+	ECacheMethod Method{};
+	if (!LoadFromCompactBinary(Batch[ANSITEXTVIEW("Method")], Method))
+	{
+		return false;
+	}
+
+	switch (Method)
+	{
+	case ECacheMethod::Put:
+		return LoadBatchFromCompactBinary<FCachePutRequest>(*this, Batch);
+	case ECacheMethod::Get:
+		return LoadBatchFromCompactBinary<FCacheGetRequest>(*this, Batch);
+	case ECacheMethod::PutValue:
+		return LoadBatchFromCompactBinary<FCachePutValueRequest>(*this, Batch);
+	case ECacheMethod::GetValue:
+		return LoadBatchFromCompactBinary<FCacheGetValueRequest>(*this, Batch);
+	case ECacheMethod::GetChunks:
+		return LoadBatchFromCompactBinary<FCacheGetChunkRequest>(*this, Batch);
+	default:
+		return false;
+	}
+}
+
+///////////////////////////////////////////////////////////////////////////////////////////////////////////////////////
+
+class FCacheReplayArchiveWriter final : public ICacheReplayWriter
+{
+public:
+	explicit FCacheReplayArchiveWriter(TNotNull<TUniquePtr<FArchive>> Ar, uint64 CompressionBlockSize = 256 * 1024);
+	~FCacheReplayArchiveWriter() final;
+
+	FCacheReplayArchiveWriter(const FCacheReplayArchiveWriter&) = delete;
+	FCacheReplayArchiveWriter& operator=(const FCacheReplayArchiveWriter&) = delete;
+
+	void SerializeBatch(const FCbWriter& Batch) final;
+
+private:
+	void WriteCompressedBinaryToArchive(const FSharedBuffer& RawBinary);
+	void FlushCompressedBinaryToArchive();
+
+	TUniquePtr<FArchive> Ar;
+	FUniqueBuffer RawBlock;
+	FMutableMemoryView RawBlockTail;
+	FMutex Mutex;
+};
+
+FCacheReplayArchiveWriter::FCacheReplayArchiveWriter(TNotNull<TUniquePtr<FArchive>> InAr, uint64 InCompressionBlockSize)
+	: Ar(MoveTemp(InAr))
 {
 	if (InCompressionBlockSize)
 	{
@@ -485,7 +569,7 @@ void FCacheReplayArchiveWriter::SerializeBatch(const FCbWriter& Batch)
 	if (!RawBlock)
 	{
 		// Write Batch directly to the archive when not using compression.
-		Batch.Save(Ar);
+		Batch.Save(*Ar);
 		return;
 	}
 
@@ -514,7 +598,7 @@ void FCacheReplayArchiveWriter::WriteCompressedBinaryToArchive(const FSharedBuff
 	const FValue CompressedBinary = FValue::Compress(RawBinary, RawBlock.GetSize());
 	TCbWriter<64> BinaryWriter;
 	BinaryWriter.AddBinary(CompressedBinary.GetData().GetCompressed());
-	BinaryWriter.Save(Ar);
+	BinaryWriter.Save(*Ar);
 }
 
 void FCacheReplayArchiveWriter::FlushCompressedBinaryToArchive()
@@ -527,12 +611,162 @@ void FCacheReplayArchiveWriter::FlushCompressedBinaryToArchive()
 	}
 }
 
+TUniquePtr<ICacheReplayWriter> CreateReplayArchiveWriter(TNotNull<TUniquePtr<FArchive>> Ar, uint64 CompressionBlockSize)
+{
+	return MakeUnique<FCacheReplayArchiveWriter>(MoveTemp(Ar), CompressionBlockSize);
+}
+
 ///////////////////////////////////////////////////////////////////////////////////////////////////////////////////////
 
-class FCacheReplayReader::FState
+static bool ReadReplayFromArchive(ICacheReplayReader& Reader, FArchive& Ar, FUniqueBuffer& Scratch)
+{
+	const auto Alloc = [&Scratch](const uint64 Size) -> FUniqueBuffer
+	{
+		if (Size <= Scratch.GetSize())
+		{
+			return FUniqueBuffer::MakeView(Scratch.GetView().Left(Size));
+		}
+		return FUniqueBuffer::Alloc(Size);
+	};
+
+	// A scratch buffer for decompressed blocks of fields.
+	FUniqueBuffer BlockScratch;
+
+	bool bOk = true;
+	for (int64 Offset = Ar.Tell(); Offset < Ar.TotalSize(); Offset = Ar.Tell())
+	{
+		FCbField Field = LoadCompactBinary(Ar, Alloc);
+		if (!Field)
+		{
+			UE_LOG(LogDerivedDataCache, Error,
+				TEXT("Replay: Failed to load compact binary at offset %" INT64_FMT ". "
+					 "Archive is at offset %" INT64_FMT " and has total size %" INT64_FMT "."),
+				Offset, Ar.Tell(), Ar.TotalSize());
+			return false;
+		}
+
+		// A binary field is used to store a compressed buffer containing a sequence of compact binary objects.
+		if (Field.IsBinary())
+		{
+			FCompressedBuffer CompressedBuffer = FCompressedBuffer::FromCompressed(Field.AsBinary());
+			if (!CompressedBuffer)
+			{
+				UE_LOG(LogDerivedDataCache, Warning,
+					TEXT("Replay: Failed to load compressed buffer from binary field at offset %" INT64_FMT ". "
+						 "Archive is at offset %" INT64_FMT " and has total size %" INT64_FMT "."),
+					Offset, Ar.Tell(), Ar.TotalSize());
+				bOk = false;
+				continue;
+			}
+
+			const uint64 RawBlockSize = CompressedBuffer.GetRawSize();
+			if (BlockScratch.GetSize() < RawBlockSize)
+			{
+				// Allocate the max of the raw size and the compression block size to maximize reuse of the buffer.
+				ECompressedBufferCompressor Compressor;
+				ECompressedBufferCompressionLevel CompressionLevel;
+				uint64 CompressionBlockSize = 0;
+				(void)CompressedBuffer.TryGetCompressParameters(Compressor, CompressionLevel, CompressionBlockSize);
+				BlockScratch = FUniqueBuffer::Alloc(FMath::Max(RawBlockSize, CompressionBlockSize));
+			}
+
+			const FMutableMemoryView RawBlockView = BlockScratch.GetView().Left(RawBlockSize);
+			if (!CompressedBuffer.TryDecompressTo(RawBlockView))
+			{
+				UE_LOG(LogDerivedDataCache, Warning,
+					TEXT("Replay: Failed to decompress compressed buffer from binary field at offset %" INT64_FMT ". "
+						 "Archive is at offset %" INT64_FMT " and has total size %" INT64_FMT "."),
+					Offset, Ar.Tell(), Ar.TotalSize());
+				bOk = false;
+				continue;
+			}
+
+			// DO NOT ACCESS Field PAST THIS POINT. The recursive call overwrites it in Scratch.
+			FMemoryReaderView InnerAr(RawBlockView);
+			if (!ReadReplayFromArchive(Reader, InnerAr, Scratch))
+			{
+				bOk = false;
+				continue;
+			}
+		}
+
+		// An object field is used to store one batch of cache requests.
+		if (Field.IsObject() && !Reader.SerializeBatch(Field.AsObject()))
+		{
+			UE_LOG(LogDerivedDataCache, Warning,
+				TEXT("Replay: Failed to load cache request from object field at offset %" INT64_FMT ". "
+					 "Archive is at offset %" INT64_FMT " and has total size %" INT64_FMT "."),
+				Offset, Ar.Tell(), Ar.TotalSize());
+			bOk = false;
+			continue;
+		}
+	}
+	return bOk;
+}
+
+bool ReadReplayFromArchive(ICacheReplayReader& Reader, FArchive& Ar)
+{
+	// A scratch buffer for the compact binary fields.
+	FUniqueBuffer Scratch = FUniqueBuffer::Alloc(256 * 1024);
+	return ReadReplayFromArchive(Reader, Ar, Scratch);
+}
+
+bool ReadReplayFromFile(ICacheReplayReader& Reader, const TCHAR* Path)
+{
+	if (TUniquePtr<FArchive> ReplayAr{IFileManager::Get().CreateFileReader(Path, FILEREAD_Silent)})
+	{
+		return ReadReplayFromArchive(Reader, *ReplayAr);
+	}
+	return false;
+}
+
+///////////////////////////////////////////////////////////////////////////////////////////////////////////////////////
+
+class FCacheReplayReader::FState final : public ICacheReplayReader
 {
 public:
-	~FState();
+	~FState() final;
+
+	void Read(TConstArrayView<FCachePutRequest> Requests, TConstArrayView<FCachePutResponse> Responses, EPriority Priority) final
+	{
+		// Replay puts as gets.
+		TArray<FCacheGetRequest, TInlineAllocator<16>> GetRequests;
+		GetRequests.Reserve(Requests.Num());
+		for (const FCachePutRequest& Request : Requests)
+		{
+			GetRequests.Add({Request.Name, Request.Record.GetKey(), Request.Policy, Request.UserData});
+		}
+		DispatchRequests(MakeConstArrayView(GetRequests), ECacheMethod::Put, Priority);
+	}
+
+	void Read(TConstArrayView<FCacheGetRequest> Requests, TConstArrayView<FCacheGetResponse> Responses, EPriority Priority) final
+	{
+		DispatchRequests(Requests, ECacheMethod::Get, Priority);
+	}
+
+	void Read(TConstArrayView<FCachePutValueRequest> Requests, TConstArrayView<FCachePutValueResponse> Responses, EPriority Priority) final
+	{
+		// Replay puts as gets.
+		TArray<FCacheGetValueRequest, TInlineAllocator<16>> GetRequests;
+		GetRequests.Reserve(Requests.Num());
+		for (const FCachePutValueRequest& Request : Requests)
+		{
+			GetRequests.Add({Request.Name, Request.Key, Request.Policy, Request.UserData});
+		}
+		DispatchRequests(MakeConstArrayView(GetRequests), ECacheMethod::PutValue, Priority);
+	}
+
+	void Read(TConstArrayView<FCacheGetValueRequest> Requests, TConstArrayView<FCacheGetValueResponse> Responses, EPriority Priority) final
+	{
+		DispatchRequests(Requests, ECacheMethod::GetValue, Priority);
+	}
+
+	void Read(TConstArrayView<FCacheGetChunkRequest> Requests, TConstArrayView<FCacheGetChunkResponse> Responses, EPriority Priority) final
+	{
+		DispatchRequests(Requests, ECacheMethod::GetChunks, Priority);
+	}
+
+	bool SerializeBatch(FCbObjectView Batch) final;
 
 	void WaitForAsyncReads();
 
@@ -565,7 +799,7 @@ public:
 
 private:
 	template <typename RequestType>
-	bool DispatchRequests(FCbObjectView Object, ECacheMethod Method);
+	bool DispatchRequests(TConstArrayView<RequestType> Requests, ECacheMethod Method, EPriority Priority);
 
 	void ApplyPolicyTransform(FCacheGetRequest& Request);
 	void ApplyPolicyTransform(FCacheGetValueRequest& Request);
@@ -629,64 +863,63 @@ void FCacheReplayReader::FState::WaitForAsyncReads()
 	}
 }
 
+bool FCacheReplayReader::FState::SerializeBatch(FCbObjectView Batch)
+{
+	if (UE_LOG_ACTIVE(LogDerivedDataCache, Verbose))
+	{
+		ECacheMethod Method;
+		if (LoadFromCompactBinary(Batch[ANSITEXTVIEW("Method")], Method) &&
+			MethodFilter.IsMatch(Method) &&
+			Algo::AnyOf(Batch[ANSITEXTVIEW("Requests")], [this](FCbFieldView Request)
+			{ 
+				FCacheKey Key;
+				return LoadFromCompactBinary(Request[ANSITEXTVIEW("Key")], Key) && KeyFilter.IsMatch(Key);
+			}))
+		{
+			TStringBuilder<1024> Text;
+			CompactBinaryToCompactJson(Batch, Text);
+			UE_LOG(LogDerivedDataCache, Verbose, TEXT("Replay: %s"), *Text);
+		}
+	}
+
+	return ICacheReplayReader::SerializeBatch(Batch);
+}
+
 template <typename RequestType>
-bool FCacheReplayReader::FState::DispatchRequests(const FCbObjectView Object, const ECacheMethod Method)
+bool FCacheReplayReader::FState::DispatchRequests(const TConstArrayView<RequestType> Requests, const ECacheMethod Method, const EPriority Priority)
 {
 	if (!MethodFilter.IsMatch(Method))
 	{
 		return true;
 	}
 
-	EPriority Priority = EPriority::Normal;
-	LoadFromCompactBinary(Object[ANSITEXTVIEW("Priority")], Priority, Priority);
-
-	TArray<RequestType, TInlineAllocator<16>> Requests;
+	TArray<RequestType, TInlineAllocator<16>> FilteredRequests;
+	for (const RequestType& Request : Requests)
 	{
-		TRACE_CPUPROFILER_EVENT_SCOPE(ReplayDDC_Serialize);
-		const FCbArrayView Array = Object[ANSITEXTVIEW("Requests")].AsArrayView();
-		Requests.Reserve(IntCastChecked<int32>(Array.Num()));
-		for (FCbFieldView Field : Array)
+		if (KeyFilter.IsMatch(Request.Key))
 		{
-			RequestType& Request = Requests.AddDefaulted_GetRef();
-			if (!LoadFromCompactBinary(Field, Request))
-			{
-				return false;
-			}
-			if (KeyFilter.IsMatch(Request.Key))
-			{
-				ApplyPolicyTransform(Request);
-			}
-			else
-			{
-				Requests.Pop(EAllowShrinking::No);
-			}
+			RequestType& FilteredRequest = FilteredRequests.Emplace_GetRef(Request);
+			ApplyPolicyTransform(FilteredRequest);
 		}
 	}
 
-	if (!Requests.IsEmpty())
+	if (!FilteredRequests.IsEmpty())
 	{
-		if (UE_LOG_ACTIVE(LogDerivedDataCache, Verbose))
-		{
-			TStringBuilder<1024> Batch;
-			CompactBinaryToCompactJson(Object, Batch);
-			UE_LOG(LogDerivedDataCache, Verbose, TEXT("Replay: %s"), *Batch);
-		}
-
 		TRACE_CPUPROFILER_EVENT_SCOPE(ReplayDDC_Dispatch);
 
-		DispatchCount += Requests.Num();
+		DispatchCount += FilteredRequests.Num();
 
 		FRequestOwner& Owner = Owners[uint8(Priority)];
 		if (Owner.GetPriority() < EPriority::Blocking)
 		{
 			// Owners with non-blocking priority can execute the request directly because it will be async.
 			FRequestBarrier Barrier(Owner);
-			Invoke(CacheFunctionFor<RequestType>, TargetCache, Requests, Owner, [](TCacheResponseFor<RequestType>&&){});
+			Invoke(CacheFunctionFor<RequestType>, TargetCache, FilteredRequests, Owner, [](TCacheResponseFor<RequestType>&&){});
 		}
 		else
 		{
 			// Owners with blocking priority launch a task to execute the blocking request to allow concurrent replay.
-			BlockingTasks.Add(Tasks::Launch(TEXT("CacheReplayTask"), [this, Requests = MoveTemp(Requests)]
+			BlockingTasks.Add(Tasks::Launch(TEXT("CacheReplayTask"), [this, Requests = MoveTemp(FilteredRequests)]
 			{
 				FRequestOwner BlockingOwner(EPriority::Blocking);
 				Invoke(CacheFunctionFor<RequestType>, TargetCache, Requests, BlockingOwner, [](TCacheResponseFor<RequestType>&&){});
@@ -733,123 +966,20 @@ void FCacheReplayReader::FState::ReadFromFileAsync(const TCHAR* const ReplayPath
 
 bool FCacheReplayReader::FState::ReadFromFile(const TCHAR* const ReplayPath, const uint64 ScratchSize)
 {
-	TUniquePtr<FArchive> ReplayAr(IFileManager::Get().CreateFileReader(ReplayPath, FILEREAD_Silent));
-	if (!ReplayAr)
-	{
-		UE_LOG(LogDerivedDataCache, Warning, TEXT("Replay: File '%s' failed to open."), ReplayPath);
-		return false;
-	}
-
 	FDispatchScope Dispatch(*this);
-	UE_LOG(LogDerivedDataCache, Display, TEXT("Replay: Loading cache replay from '%s'"), ReplayPath);
-	return ReadFromArchive(*ReplayAr, FMath::Max(ScratchSize, GCacheReplayCompressionBlockSize));
+	return ReadReplayFromFile(*this, ReplayPath);
 }
 
 bool FCacheReplayReader::FState::ReadFromArchive(FArchive& ReplayAr, const uint64 ScratchSize)
 {
 	FDispatchScope Dispatch(*this);
-
-	// A scratch buffer for the compact binary fields.
-	FUniqueBuffer Scratch = FUniqueBuffer::Alloc(ScratchSize);
-	const auto Alloc = [&Scratch](const uint64 Size) -> FUniqueBuffer
-	{
-		if (Size <= Scratch.GetSize())
-		{
-			return FUniqueBuffer::MakeView(Scratch.GetView().Left(Size));
-		}
-		return FUniqueBuffer::Alloc(Size);
-	};
-
-	// A scratch buffer for decompressed blocks of fields.
-	FUniqueBuffer BlockScratch;
-
-	bool bOk = true;
-	for (int64 Offset = ReplayAr.Tell(); Offset < ReplayAr.TotalSize(); Offset = ReplayAr.Tell())
-	{
-		FCbField Field = LoadCompactBinary(ReplayAr, Alloc);
-		if (!Field)
-		{
-			UE_LOG(LogDerivedDataCache, Error,
-				TEXT("Replay: Failed to load compact binary at offset %" INT64_FMT ". "
-					 "Archive is at offset %" INT64_FMT " and has total size %" INT64_FMT "."),
-				Offset, ReplayAr.Tell(), ReplayAr.TotalSize());
-			return false;
-		}
-
-		// A binary field is used to store a compressed buffer containing a sequence of compact binary objects.
-		if (Field.IsBinary())
-		{
-			FCompressedBuffer CompressedBuffer = FCompressedBuffer::FromCompressed(Field.AsBinary());
-			if (!CompressedBuffer)
-			{
-				UE_LOG(LogDerivedDataCache, Warning,
-					TEXT("Replay: Failed to load compressed buffer from binary field at offset %" INT64_FMT ". "
-						 "Archive is at offset %" INT64_FMT " and has total size %" INT64_FMT "."),
-					Offset, ReplayAr.Tell(), ReplayAr.TotalSize());
-				bOk = false;
-				continue;
-			}
-
-			const uint64 RawBlockSize = CompressedBuffer.GetRawSize();
-			if (BlockScratch.GetSize() < RawBlockSize)
-			{
-				BlockScratch = FUniqueBuffer::Alloc(RawBlockSize);
-			}
-
-			const FMutableMemoryView RawBlockView = BlockScratch.GetView().Left(RawBlockSize);
-			if (!CompressedBuffer.TryDecompressTo(RawBlockView))
-			{
-				UE_LOG(LogDerivedDataCache, Warning,
-					TEXT("Replay: Failed to decompress compressed buffer from binary field at offset %" INT64_FMT ". "
-						 "Archive is at offset %" INT64_FMT " and has total size %" INT64_FMT "."),
-					Offset, ReplayAr.Tell(), ReplayAr.TotalSize());
-				bOk = false;
-				continue;
-			}
-
-			FMemoryReaderView InnerAr(RawBlockView);
-			if (!ReadFromArchive(InnerAr, DefaultScratchSize))
-			{
-				bOk = false;
-				continue;
-			}
-		}
-
-		// An object field is used to store one batch of cache requests.
-		if (Field.IsObject() && !ReadFromObject(Field.AsObject()))
-		{
-			UE_LOG(LogDerivedDataCache, Warning,
-				TEXT("Replay: Failed to load cache request from object field at offset %" INT64_FMT ". "
-						"Archive is at offset %" INT64_FMT " and has total size %" INT64_FMT "."),
-				Offset, ReplayAr.Tell(), ReplayAr.TotalSize());
-			bOk = false;
-			continue;
-		}
-	}
-	return bOk;
+	return ReadReplayFromArchive(*this, ReplayAr);
 }
 
 bool FCacheReplayReader::FState::ReadFromObject(const FCbObjectView Object)
 {
-	ECacheMethod Method{};
-	if (!LoadFromCompactBinary(Object[ANSITEXTVIEW("Method")], Method))
-	{
-		return false;
-	}
-
-	switch (Method)
-	{
-	case ECacheMethod::Put:
-	case ECacheMethod::Get:
-		return DispatchRequests<FCacheGetRequest>(Object, Method);
-	case ECacheMethod::PutValue:
-	case ECacheMethod::GetValue:
-		return DispatchRequests<FCacheGetValueRequest>(Object, Method);
-	case ECacheMethod::GetChunks:
-		return DispatchRequests<FCacheGetChunkRequest>(Object, Method);
-	}
-
-	return false;
+	FDispatchScope Dispatch(*this);
+	return SerializeBatch(Object);
 }
 
 FCacheReplayReader::FCacheReplayReader(ILegacyCacheStore* const TargetCache)
@@ -1029,6 +1159,7 @@ ILegacyCacheStore* TryCreateCacheStoreReplay(ILegacyCacheStore* InnerCache, ICac
 			FString ReplayLoadPath;
 			if (FParse::Value(*Token, TEXT("-DDC-ReplayLoad="), ReplayLoadPath))
 			{
+				UE_LOG(LogDerivedDataCache, Display, TEXT("Replay: Loading cache replay from '%s'"), *ReplayLoadPath);
 				Reader.ReadFromFileAsync(*ReplayLoadPath, Owner);
 			}
 		}
diff --git a/Engine/Source/Developer/DerivedDataCache/Private/DerivedDataCacheStore.h b/Engine/Source/Developer/DerivedDataCache/Private/DerivedDataCacheStore.h
--- a/Engine/Source/Developer/DerivedDataCache/Private/DerivedDataCacheStore.h
+++ b/Engine/Source/Developer/DerivedDataCache/Private/DerivedDataCacheStore.h
@@ -350,6 +350,7 @@ struct TCacheRequestTraits<FCachePutRequest>
 	using FOnComplete = FOnCachePutComplete;
 	constexpr static ECacheMethod Method = ECacheMethod::Put;
 	constexpr static void (ICacheStore::*Function)(TConstArrayView<FRequest>, IRequestOwner&, FOnComplete&&) = &ICacheStore::Put;
+	static const FCacheKey& GetKey(const FCachePutRequest& Request) { return Request.Record.GetKey(); }
 };
 
 template <>
@@ -360,6 +361,7 @@ struct TCacheRequestTraits<FCacheGetRequest>
 	using FOnComplete = FOnCacheGetComplete;
 	constexpr static ECacheMethod Method = ECacheMethod::Get;
 	constexpr static void (ICacheStore::*Function)(TConstArrayView<FRequest>, IRequestOwner&, FOnComplete&&) = &ICacheStore::Get;
+	static const FCacheKey& GetKey(const FCacheGetRequest& Request) { return Request.Key; }
 };
 
 template <>
@@ -370,6 +372,7 @@ struct TCacheRequestTraits<FCachePutValueRequest>
 	using FOnComplete = FOnCachePutValueComplete;
 	constexpr static ECacheMethod Method = ECacheMethod::PutValue;
 	constexpr static void (ICacheStore::*Function)(TConstArrayView<FRequest>, IRequestOwner&, FOnComplete&&) = &ICacheStore::PutValue;
+	static const FCacheKey& GetKey(const FCachePutValueRequest& Request) { return Request.Key; }
 };
 
 template <>
@@ -380,6 +383,7 @@ struct TCacheRequestTraits<FCacheGetValueRequest>
 	using FOnComplete = FOnCacheGetValueComplete;
 	constexpr static ECacheMethod Method = ECacheMethod::GetValue;
 	constexpr static void (ICacheStore::*Function)(TConstArrayView<FRequest>, IRequestOwner&, FOnComplete&&) = &ICacheStore::GetValue;
+	static const FCacheKey& GetKey(const FCacheGetValueRequest& Request) { return Request.Key; }
 };
 
 template <>
@@ -390,6 +394,7 @@ struct TCacheRequestTraits<FCacheGetChunkRequest>
 	using FOnComplete = FOnCacheGetChunkComplete;
 	constexpr static ECacheMethod Method = ECacheMethod::GetChunks;
 	constexpr static void (ICacheStore::*Function)(TConstArrayView<FRequest>, IRequestOwner&, FOnComplete&&) = &ICacheStore::GetChunks;
+	static const FCacheKey& GetKey(const FCacheGetChunkRequest& Request) { return Request.Key; }
 };
 
 } // Private
@@ -406,4 +411,10 @@ template <typename RequestType> inline constexpr auto CacheMethodFor = Private::
 /** Alias for the cache function corresponding to the request type. */
 template <typename RequestType> inline constexpr auto CacheFunctionFor = Private::TCacheRequestTraits<RequestType>::Function;
 
+template <typename RequestType>
+inline const FCacheKey& GetCacheKey(const RequestType& Request)
+{
+	return Private::TCacheRequestTraits<RequestType>::GetKey(Request);
+}
+
 } // UE::DerivedData
