DDC: Fixed inconsistent whitespace and an incorrect comment

#rnx
#Synced-CL 48512380

diff --git a/Engine/Source/Developer/DerivedDataCache/Private/DerivedDataThreadPoolTask.cpp b/Engine/Source/Developer/DerivedDataCache/Private/DerivedDataThreadPoolTask.cpp
--- a/Engine/Source/Developer/DerivedDataCache/Private/DerivedDataThreadPoolTask.cpp
+++ b/Engine/Source/Developer/DerivedDataCache/Private/DerivedDataThreadPoolTask.cpp
@@ -19,7 +19,7 @@ class FThreadPoolTaskRequest final : public FRequestBase, private FInheritedCont
 public:
 	FThreadPoolTaskRequest(
 		uint64 MemoryEstimate,
-		const TCHAR * DebugName,
+		const TCHAR* DebugName,
 		IRequestOwner& Owner,
 		FQueuedThreadPool& ThreadPool,
 		TUniqueFunction<void ()>&& TaskBody);
@@ -45,7 +45,7 @@ private:
 		return int64(MemoryEstimate);
 	}
 	
-	const TCHAR * GetDebugName() const final
+	const TCHAR* GetDebugName() const final
 	{
 		return DebugName;
 	}
@@ -57,7 +57,7 @@ private:
 
 private:
 	uint64 MemoryEstimate;
-	const TCHAR * DebugName = nullptr;
+	const TCHAR* DebugName = nullptr;
 	IRequestOwner& Owner;
 	FQueuedThreadPool& ThreadPool;
 	TUniqueFunction<void ()> TaskBody;
@@ -67,7 +67,7 @@ private:
 
 FThreadPoolTaskRequest::FThreadPoolTaskRequest(
 	const uint64 InMemoryEstimate,
-	const TCHAR * InDebugName,
+	const TCHAR* InDebugName,
 	IRequestOwner& InOwner,
 	FQueuedThreadPool& InThreadPool,
 	TUniqueFunction<void ()>&& InTaskBody)
@@ -79,7 +79,7 @@ FThreadPoolTaskRequest::FThreadPoolTaskRequest(
 {
 	check(MemoryEstimate <= MAX_int64);
 	CaptureInheritedContext();
-	AddRef(); // Released in Execute() or Cancel()
+	AddRef(); // Released in Execute() or Wait()
 	Owner.Begin(this);
 	ThreadPool.AddQueuedWork(this, ConvertToQueuedWorkPriority(Owner.GetPriority()));
 }
@@ -133,7 +133,7 @@ void FThreadPoolTaskRequest::Wait()
 
 void LaunchTaskInThreadPool(
 	uint64 MemoryEstimate,
-	const TCHAR * DebugName,
+	const TCHAR* DebugName,
 	IRequestOwner& Owner,
 	FQueuedThreadPool* ThreadPool,
 	TUniqueFunction<void ()>&& TaskBody)
