DDC: Fixed Zen client support for out-of-bounds chunk reads

Requires zenserver 5.7.13.

- Clamp RawSize to avoid decompression errors.
- Validate that RawHash matches in the compressed buffer header when GetChunks receives a whole value.

#Synced-CL 48512380

diff --git a/Engine/Source/Developer/DerivedDataCache/Private/ZenCacheStore.cpp b/Engine/Source/Developer/DerivedDataCache/Private/ZenCacheStore.cpp
--- a/Engine/Source/Developer/DerivedDataCache/Private/ZenCacheStore.cpp
+++ b/Engine/Source/Developer/DerivedDataCache/Private/ZenCacheStore.cpp
@@ -1426,10 +1426,9 @@ public:
 						const TRequestWithStats<FCacheGetChunkRequest>& RequestWithStats = Batch[RequestIndex++];
 						RequestWithStats.Stats.Latency = FMonotonicTimeSpan::FromSeconds(HttpResponse->GetStats().GetLatency());
 
+						bool bSucceeded = false;
 						FIoHash RawHash;
-						bool Succeeded = false;
 						uint64 RawSize = 0;
-						FCbObjectView ResultObject = ResultView.AsObjectView();
 						FSharedBuffer RequestedBytes;
 						const FCacheGetChunkRequest& Request = RequestWithStats.Request;
 						if (CacheStore.DebugOptions.ShouldSimulateGetMiss(Request.Key))
@@ -1437,60 +1436,41 @@ public:
 							UE_LOG(LogDerivedDataCache, Verbose, TEXT("%s: Simulated miss for get of '%s' from '%s'"),
 								*CacheStore.GetName(), *WriteToString<96>(Request.Key, '/', Request.Id), *Request.Name);
 						}
-						else
+						else if (RawHash = ResultView[ANSITEXTVIEW("RawHash")].AsHash(); !RawHash.IsZero())
 						{
-							FCbFieldView HashView = ResultObject[ANSITEXTVIEW("RawHash")];
-							RawHash = HashView.AsHash();
-							if (!HashView.HasError())
+							if (EnumHasAnyFlags(RequestWithStats.Request.Policy, ECachePolicy::SkipData))
 							{
-								FIoHash AttachmentHash = RawHash;
-								const FCbAttachment* Attachment = nullptr;
-								if (EnumHasAnyFlags(RequestWithStats.Request.Policy, ECachePolicy::SkipData))
+								if (uint64 TotalSize = ResultView[ANSITEXTVIEW("RawSize")].AsUInt64(MAX_uint64); TotalSize != MAX_uint64)
 								{
-									FCbFieldView RawSizeField = ResultObject[ANSITEXTVIEW("RawSize")];
-									uint64 TotalSize = RawSizeField.AsUInt64();
-									Succeeded = !RawSizeField.HasError();
-									if (Succeeded)
-									{
-										RawSize = FMath::Min(Request.RawSize, TotalSize - FMath::Min(Request.RawOffset, TotalSize));
-									}
+									RawSize = FMath::Min(Request.RawSize, TotalSize - FMath::Min(Request.RawOffset, TotalSize));
+									bSucceeded = true;
 								}
-								else
+							}
+							else
+							{
+								const FIoHash FragmentHash = ResultView[ANSITEXTVIEW("FragmentHash")].AsHash();
+								const uint64 FragmentOffset = ResultView[ANSITEXTVIEW("FragmentOffset")].AsUInt64(0);
+								const bool bFragment = !FragmentHash.IsZero();
+								if (const FCbAttachment* Attachment = Response.FindAttachment(bFragment ? FragmentHash : RawHash))
 								{
-									FCbFieldView FragmentOffsetField = ResultObject[ANSITEXTVIEW("FragmentOffset")];
-									uint64 FragmentOffset = FragmentOffsetField.AsUInt64();
-
-									FCbFieldView FragmentHashView = ResultObject[ANSITEXTVIEW("FragmentHash")];
-									if (FragmentHashView.IsHash())
-									{
-										FIoHash FragmentHash = FragmentHashView.AsHash();
-										AttachmentHash = FragmentHash;
-									}
-
-									if (Attachment = Response.FindAttachment(AttachmentHash); Attachment != nullptr)
-									{
-										if (FCompressedBuffer CompressedBuffer = Attachment->AsCompressedBinary(); CompressedBuffer)
-										{
-											RequestedBytes = FCompressedBufferReader(CompressedBuffer).Decompress(Request.RawOffset - FragmentOffset, Request.RawSize);
-											RawSize = RequestedBytes.GetSize();
-											Succeeded = true;
-										}
-									}
-									else if (!FragmentHashView.IsHash())
+									const FCompressedBuffer& CompressedData = Attachment->AsCompressedBinary();
+									if (CompressedData && (bFragment || CompressedData.GetRawHash() == RawHash))
 									{
-										// An out-of-bounds access will respond with no buffer.
-										RequestedBytes = FUniqueBuffer::Alloc(0).MoveToShared();
-										Succeeded = true;
+										const uint64 OffsetInFragment = Request.RawOffset - FragmentOffset;
+										RawSize = FMath::Min(CompressedData.GetRawSize() - OffsetInFragment, Request.RawSize);
+										RequestedBytes = FCompressedBufferReader(CompressedData).Decompress(OffsetInFragment, RawSize);
+										bSucceeded = RequestedBytes.GetSize() == RawSize;
 									}
 								}
-							}
-							else
-							{
-								FCbFieldView HashView2 = ResultObject[ANSITEXTVIEW("RawHash")];
-								RawHash = HashView2.AsHash();
+								else if (!bFragment)
+								{
+									// An out-of-bounds access will respond with no buffer.
+									RequestedBytes = FUniqueBuffer::Alloc(0).MoveToShared();
+									bSucceeded = true;
+								}
 							}
 						}
-						Succeeded ? OnHit(RequestWithStats, MoveTemp(RawHash), RawSize, MoveTemp(RequestedBytes)) : OnMiss(RequestWithStats);
+						bSucceeded ? OnHit(RequestWithStats, MoveTemp(RawHash), RawSize, MoveTemp(RequestedBytes)) : OnMiss(RequestWithStats);
 					}
 				}
 				else if (!bCanceled && (HttpResponse->GetStatusCode() != 404))
diff --git a/Engine/Source/Developer/DerivedDataCache/Tests/DerivedDataCacheStoreTest.cpp b/Engine/Source/Developer/DerivedDataCache/Tests/DerivedDataCacheStoreTest.cpp
--- a/Engine/Source/Developer/DerivedDataCache/Tests/DerivedDataCacheStoreTest.cpp
+++ b/Engine/Source/Developer/DerivedDataCache/Tests/DerivedDataCacheStoreTest.cpp
@@ -315,25 +315,84 @@ static void TestCacheChunks(ICache& Cache, const FCacheKey& Key, const FValueId
 		BlockingOwner.Wait();
 	}
 
-	// Disabled until Zen is fixed...
-	//{
-	//	INFO("Get Partial Value Out of Bounds");
-	//	FRequestOwner BlockingOwner(EPriority::Blocking);
-	//	Cache.GetChunks({{.Name = Name, .Key = Key, .Id = ValueId, .RawOffset = 1024 * 1024 * 1024, .RawSize = 512 * 1024, .UserData = UserData}}, BlockingOwner, [&](FCacheGetChunkResponse&& Response)
-	//	{
-	//		CHECK(Response.Status == EStatus::Ok);
-	//		CHECK(Response.UserData == UserData);
-	//		CHECK(Response.Name == Name);
-	//		CHECK(Response.Key == Key);
-	//		CHECK(Response.Id == ValueId);
-	//		CHECK(Response.RawOffset == 1024 * 1024 * 1024);
-	//		CHECK(Response.RawSize == 0);
-	//		CHECK(Response.RawHash == Value.GetRawHash());
-	//		CHECK(Response.RawData);
-	//		CHECK(Response.RawData.GetSize() == 0);
-	//	});
-	//	BlockingOwner.Wait();
-	//}
+	{
+		INFO("Get Partial Value Partially Out of Bounds");
+		FRequestOwner BlockingOwner(EPriority::Blocking);
+		Cache.GetChunks({{.Name = Name, .Key = Key, .Id = ValueId, .RawOffset = Value.GetRawSize() - 16 * 1024, .RawSize = 32 * 1024, .UserData = UserData}},
+			BlockingOwner, [&](FCacheGetChunkResponse&& Response)
+		{
+			CHECK(Response.Status == EStatus::Ok);
+			CHECK(Response.UserData == UserData);
+			CHECK(Response.Name == Name);
+			CHECK(Response.Key == Key);
+			CHECK(Response.Id == ValueId);
+			CHECK(Response.RawOffset == Value.GetRawSize() - 16 * 1024);
+			CHECK(Response.RawSize == 16 * 1024);
+			CHECK(Response.RawHash == Value.GetRawHash());
+			CHECK(Response.RawData.GetSize() == Response.RawSize);
+			CHECK(CheckSequentialIntegerBufferRange(Response.RawData, Response.RawOffset));
+		});
+		BlockingOwner.Wait();
+	}
+
+	{
+		INFO("Get Partial Value Hash + Size Partially Out of Bounds");
+		FRequestOwner BlockingOwner(EPriority::Blocking);
+		Cache.GetChunks({{.Name = Name, .Key = Key, .Id = ValueId, .RawOffset = Value.GetRawSize() - 16 * 1024, .RawSize = 32 * 1024,
+			.Policy = ECachePolicy::Default | ECachePolicy::SkipData, .UserData = UserData}},
+			BlockingOwner, [&](FCacheGetChunkResponse&& Response)
+		{
+			CHECK(Response.Status == EStatus::Ok);
+			CHECK(Response.UserData == UserData);
+			CHECK(Response.Name == Name);
+			CHECK(Response.Key == Key);
+			CHECK(Response.Id == ValueId);
+			CHECK(Response.RawOffset == Value.GetRawSize() - 16 * 1024);
+			CHECK(Response.RawSize == 16 * 1024);
+			CHECK(Response.RawHash == Value.GetRawHash());
+			CHECK(Response.RawData.IsNull());
+		});
+		BlockingOwner.Wait();
+	}
+
+	{
+		INFO("Get Partial Value Out of Bounds");
+		FRequestOwner BlockingOwner(EPriority::Blocking);
+		Cache.GetChunks({{.Name = Name, .Key = Key, .Id = ValueId, .RawOffset = 1024 * 1024 * 1024, .RawSize = 512 * 1024, .UserData = UserData}}, BlockingOwner, [&](FCacheGetChunkResponse&& Response)
+		{
+			CHECK(Response.Status == EStatus::Ok);
+			CHECK(Response.UserData == UserData);
+			CHECK(Response.Name == Name);
+			CHECK(Response.Key == Key);
+			CHECK(Response.Id == ValueId);
+			CHECK(Response.RawOffset == 1024 * 1024 * 1024);
+			CHECK(Response.RawSize == 0);
+			CHECK(Response.RawHash == Value.GetRawHash());
+			CHECK(Response.RawData);
+			CHECK(Response.RawData.GetSize() == 0);
+		});
+		BlockingOwner.Wait();
+	}
+
+	{
+		INFO("Get Partial Value Hash + Size Out of Bounds");
+		FRequestOwner BlockingOwner(EPriority::Blocking);
+		Cache.GetChunks({{.Name = Name, .Key = Key, .Id = ValueId, .RawOffset = 1024 * 1024 * 1024, .RawSize = 512 * 1024,
+			.Policy = ECachePolicy::Default | ECachePolicy::SkipData, .UserData = UserData}},
+			BlockingOwner, [&](FCacheGetChunkResponse&& Response)
+		{
+			CHECK(Response.Status == EStatus::Ok);
+			CHECK(Response.UserData == UserData);
+			CHECK(Response.Name == Name);
+			CHECK(Response.Key == Key);
+			CHECK(Response.Id == ValueId);
+			CHECK(Response.RawOffset == 1024 * 1024 * 1024);
+			CHECK(Response.RawSize == 0);
+			CHECK(Response.RawHash == Value.GetRawHash());
+			CHECK(Response.RawData.IsNull());
+		});
+		BlockingOwner.Wait();
+	}
 
 	{
 		INFO("Get Multiple Partial Values");
