Disable the CacheStoreTest as it conflicts with the possible presence of read-only and write-only flags on the cache store.

#Synced-CL 48523383

diff --git a/Engine/Source/Developer/DerivedDataCache/Private/Tests/CacheStoreTest.cpp b/Engine/Source/Developer/DerivedDataCache/Private/Tests/CacheStoreTest.cpp
--- a/Engine/Source/Developer/DerivedDataCache/Private/Tests/CacheStoreTest.cpp
+++ b/Engine/Source/Developer/DerivedDataCache/Private/Tests/CacheStoreTest.cpp
@@ -16,7 +16,7 @@
 #include "Misc/StringBuilder.h"
 
 IMPLEMENT_SIMPLE_AUTOMATION_TEST(FCacheStoreTest, "System.DerivedDataCache.CacheStore",
-	EAutomationTestFlags::EditorContext | EAutomationTestFlags::EngineFilter);
+	EAutomationTestFlags::EditorContext | EAutomationTestFlags::EngineFilter | EAutomationTestFlags::Disabled);
 
 bool FCacheStoreTest::RunTest(const FString& Parameters)
 {
