DDC: Added support for natively-async cache stores to skip creation of an async wrapper

Use this on FTestCacheStore to avoid it being called from multiple threads when the hierarchy schedules multiple async puts to it.

#Synced-CL 48955278

diff --git a/Engine/Source/Developer/DerivedDataCache/Private/DerivedDataCacheStore.h b/Engine/Source/Developer/DerivedDataCache/Private/DerivedDataCacheStore.h
--- a/Engine/Source/Developer/DerivedDataCache/Private/DerivedDataCacheStore.h
+++ b/Engine/Source/Developer/DerivedDataCache/Private/DerivedDataCacheStore.h
@@ -188,6 +188,9 @@ enum class ECacheStoreFlags : uint32
 	StopGetStore    = 1 << 5,
 	/** A record or value contained by this cache store will not store to later cache stores. */
 	StopStore       = StopPutStore | StopGetStore,
+
+	/** Async operation is natively provided by this cache store for requests of non-blocking priority. */
+	NativeAsync     = 1 << 6,
 };
 
 ENUM_CLASS_FLAGS(ECacheStoreFlags);
diff --git a/Engine/Source/Developer/DerivedDataCache/Private/DerivedDataCacheStoreHierarchy.cpp b/Engine/Source/Developer/DerivedDataCache/Private/DerivedDataCacheStoreHierarchy.cpp
--- a/Engine/Source/Developer/DerivedDataCache/Private/DerivedDataCacheStoreHierarchy.cpp
+++ b/Engine/Source/Developer/DerivedDataCache/Private/DerivedDataCacheStoreHierarchy.cpp
@@ -149,7 +149,7 @@ private:
 		ILegacyCacheStore* Cache{};
 		ECacheStoreFlags CacheFlags{};
 		ECacheStoreNodeFlags NodeFlags{};
-		TUniquePtr<ILegacyCacheStore> AsyncCache;
+		ILegacyCacheStore* AsyncCache{};
 		TArray<FCacheStoreStats*> CacheStats;
 	};
 
@@ -289,7 +289,10 @@ FCacheStoreHierarchy::~FCacheStoreHierarchy()
 		const int32 NodeIndex = Nodes.Num() - 1;
 		FCacheStoreNode& Node = Nodes[NodeIndex];
 
-		Node.AsyncCache.Reset();
+		if (Node.AsyncCache != Node.Cache)
+		{
+			delete Node.AsyncCache;
+		}
 		delete Node.Cache;
 
 		if (Nodes.Num() <= NodeIndex)
@@ -315,8 +318,9 @@ void FCacheStoreHierarchy::Add(ILegacyCacheStore* CacheStore, ECacheStoreFlags F
 	TUniqueLock Lock(NodesMutex);
 	checkf(!Algo::FindBy(Nodes, CacheStore, &FCacheStoreNode::Cache),
 		TEXT("Attempting to add a cache store that was previously registered to the hierarchy."));
-	TUniquePtr<ILegacyCacheStore> AsyncCacheStore(CreateCacheStoreAsync(this, CacheStore, /*MemoryCache*/ nullptr, /*bDeleteInnerCache*/ false));
-	Nodes.Add({CacheStore, Flags, {}, MoveTemp(AsyncCacheStore)});
+	ILegacyCacheStore* AsyncCacheStore = EnumHasAnyFlags(Flags, ECacheStoreFlags::NativeAsync) ? CacheStore
+		: CreateCacheStoreAsync(this, CacheStore, /*MemoryCache*/ nullptr, /*bDeleteInnerCache*/ false);
+	Nodes.Add({CacheStore, Flags, {}, AsyncCacheStore});
 	UpdateNodeFlags();
 }
 
@@ -344,6 +348,10 @@ void FCacheStoreHierarchy::RemoveNotSafe(ILegacyCacheStore* CacheStore)
 			"Leaked stats for {Type} cache store '{Name}' with path '{Path}'.",
 			Stats->Type, Stats->Name, Stats->Path);
 	}
+	if (Node->AsyncCache != Node->Cache)
+	{
+		delete Node->AsyncCache;
+	}
 	Nodes.RemoveAt(NodeIndex);
 	UpdateNodeFlags();
 }
diff --git a/Engine/Source/Developer/DerivedDataCache/Tests/TestCacheStore.cpp b/Engine/Source/Developer/DerivedDataCache/Tests/TestCacheStore.cpp
--- a/Engine/Source/Developer/DerivedDataCache/Tests/TestCacheStore.cpp
+++ b/Engine/Source/Developer/DerivedDataCache/Tests/TestCacheStore.cpp
@@ -139,7 +139,7 @@ private:
 
 FTestCacheStore::FTestCacheStore(const ECacheStoreFlags Flags, const ETestCacheStoreFlags InTestFlags)
 	: CacheStorePolicy(ECachePolicy::None)
-	, CacheStoreFlags(Flags)
+	, CacheStoreFlags(Flags | ECacheStoreFlags::NativeAsync)
 	, TestFlags(InTestFlags)
 {
 	CacheStorePolicy |= EnumHasAllFlags(Flags, ECacheStoreFlags::Query | ECacheStoreFlags::Local) ? ECachePolicy::QueryLocal : ECachePolicy::None;
