DDC: Moved config functions from Private to Internal

#rnx
#Synced-CL 48512380

diff --git a/Engine/Source/Developer/DerivedDataCache/Internal/DerivedDataConfig.h b/Engine/Source/Developer/DerivedDataCache/Internal/DerivedDataConfig.h
new file mode 100644
--- /dev/null
+++ b/Engine/Source/Developer/DerivedDataCache/Internal/DerivedDataConfig.h
@@ -0,0 +1,98 @@
+// Copyright Epic Games, Inc. All Rights Reserved.
+
+#pragma once
+
+#include "CoreTypes.h"
+#include "Containers/StringFwd.h"
+#include "Containers/StringView.h"
+#include "Templates/FunctionWithContext.h"
+
+namespace UE::DerivedData::Config
+{
+
+/**
+ * Iterates fields of a config string in the form (Name1=Value1, Name2=Value2, Name3=(X=1, Y=2, Z=3), ValueOnly).
+ *
+ * Names must contain only [A-Za-z0-9_].
+ * Strings with spaces or quotes must be quoted and escaped: (Description="A \"real\" description...")
+ * Values that are a name and config are explicitly supported: (Node=Point(X=1, Y=2, Z=3), ValueOnly(Name="Test"))
+ */
+class FConfigIterator
+{
+public:
+	explicit FConfigIterator(FStringView Config);
+
+	FConfigIterator& operator++();
+
+	inline bool HasError() const { return bError; }
+	inline explicit operator bool() const { return !bAtEnd; }
+	inline const FStringView& GetName() const { return Name; }
+	inline const FStringView& GetValue() const { return Value; }
+
+private:
+	FStringView Tail;
+	FStringView Name;
+	FStringView Value;
+	bool bAtEnd = false;
+	bool bError = false;
+	bool bLastField = false;
+};
+
+/**
+ * Returns the possibly-quoted string, with quotes removed and escapes processed.
+ *
+ * @param String       View of a string that may be surrounded by quotes and may contain escaped characters.
+ * @param OutScratch   Builder to use as scratch space when processing escape sequences.
+ * @return A view into String if there are no escapes, and a view into OutScratch otherwise.
+ */
+[[nodiscard]] FStringView ParseQuotedString(FStringView String, FStringBuilderBase& OutScratch);
+
+/** Finds the first field in Config matching Name case-insensitively. */
+[[nodiscard]] bool FindFirstFieldValue(FStringView Config, FStringView Name, FStringView& OutValue);
+
+/** Visits every field in Config matching Name case-insensitively. */
+[[nodiscard]] bool VisitFieldValues(FStringView Config, FStringView Name, TFunctionWithContext<void (FStringView)> OnValue);
+
+/** Splits "Name(Config)" to ("Name", "(Config)"), "Name" to ("Name", ""), "(Config)" to ("", "(Config)"). */
+[[nodiscard]] bool TrySplitNameAndConfig(FStringView Value, FStringView& OutName, FStringView& OutConfig);
+
+/** Appends "Name=Value", or "Name" if Value is empty, with extra whitespace removed from any valid nested objects. */
+void AppendConfig(FStringBuilderBase& OutConfig, FStringView Name, FStringView Value);
+
+} // UE::DerivedData::Config
+
+namespace UE::DerivedData
+{
+
+/**
+ * Finds the configuration for a cache graph.
+ *
+ * Written in the normalized format "(StoreName1, StoreName2, StoreName3)".
+ * Store remapping and inline configuration is stripped from the output.
+ * Example: "(StoreName1=OtherName(Key=Value))" outputs "(StoreName1)".
+ *
+ * @param GraphNameOrConfig   Name of the cache graph, which is compared case-insensitively, or a graph configuration.
+ * @param OutConfig           Appended with the configuration for the cache graph if it is found.
+ * @return True if the configuration was found, otherwise false.
+ */
+bool TryFindCacheGraphConfig(FStringView GraphNameOrConfig, FStringBuilderBase& OutConfig);
+
+/**
+ * Finds the configuration for a cache store, optionally within a specific cache graph.
+ *
+ * The optional cache graph may contain configuration overrides for the cache store.
+ * The configuration is appended to OutConfig from most-derived to least-derived,
+ * such that FParse::Value() will see override values before base values.
+ *
+ * Written in the normalized format "(Type=Value, Key1=Value1, Key2=Value2, Value3)".
+ * Store inheritance is stripped from the output.
+ * Example: "(Base=OtherStore, Key1 = Value1, Value2)" outputs "(Key1=Value1, Value2, KeyFromOtherStore=OtherValue)".
+ *
+ * @param StoreName           Name of the cache store, which is compared case-insensitively.
+ * @param GraphNameOrConfig   Name of the cache graph, which is compared case-insensitively, or a graph configuration.
+ * @param OutConfig           Appended with the configuration for the cache store if it is found.
+ * @return True if the configuration was found, otherwise false.
+ */
+bool TryFindCacheStoreConfig(FStringView StoreName, FStringView GraphNameOrConfig, FStringBuilderBase& OutConfig);
+
+} // UE::DerivedData
diff --git a/Engine/Source/Developer/DerivedDataCache/Private/DerivedDataConfig.h b/Engine/Source/Developer/DerivedDataCache/Private/DerivedDataConfig.h
deleted file mode 100644
