DDC: Added FCacheKeyFilter::RequiresSalt() to query whether the salt is used by the filter

#Synced-CL 48512380

diff --git a/Engine/Source/Developer/DerivedDataCache/Internal/DerivedDataCacheKeyFilter.h b/Engine/Source/Developer/DerivedDataCache/Internal/DerivedDataCacheKeyFilter.h
--- a/Engine/Source/Developer/DerivedDataCache/Internal/DerivedDataCacheKeyFilter.h
+++ b/Engine/Source/Developer/DerivedDataCache/Internal/DerivedDataCacheKeyFilter.h
@@ -41,6 +41,7 @@ public:
 	 */
 	UE_API void SetSalt(uint32 Salt);
 	UE_API uint32 GetSalt() const;
+	UE_API bool RequiresSalt() const;
 
 	/** Returns true if the filter matches the key. The output is the same until the salt changes. */
 	UE_API bool IsMatch(const FCacheKey& Key) const;
diff --git a/Engine/Source/Developer/DerivedDataCache/Private/DerivedDataBackends.cpp b/Engine/Source/Developer/DerivedDataCache/Private/DerivedDataBackends.cpp
--- a/Engine/Source/Developer/DerivedDataCache/Private/DerivedDataBackends.cpp
+++ b/Engine/Source/Developer/DerivedDataCache/Private/DerivedDataBackends.cpp
@@ -1137,8 +1137,8 @@ bool FBackendDebugOptions::ParseFromTokens(FBackendDebugOptions& OutOptions, con
 	}
 	if (OutOptions.SimulateMissFilter)
 	{
-		UE_LOG(LogDerivedDataCache, Display,
-			TEXT("%s: Using salt %s%u to filter cache keys to simulate misses on."),
+		UE_CLOGF(OutOptions.SimulateMissFilter.RequiresSalt(), LogDerivedDataCache, Display,
+			"%ls: Using salt %ls%u to filter cache keys to simulate misses on.",
 			InNodeName, *Prefix, OutOptions.SimulateMissFilter.GetSalt());
 	}
 	Prefix.RemoveAt(PrefixLen, Prefix.Len() - PrefixLen);
diff --git a/Engine/Source/Developer/DerivedDataCache/Private/DerivedDataCacheKeyFilter.cpp b/Engine/Source/Developer/DerivedDataCache/Private/DerivedDataCacheKeyFilter.cpp
--- a/Engine/Source/Developer/DerivedDataCache/Private/DerivedDataCacheKeyFilter.cpp
+++ b/Engine/Source/Developer/DerivedDataCache/Private/DerivedDataCacheKeyFilter.cpp
@@ -2,6 +2,7 @@
 
 #include "DerivedDataCacheKeyFilter.h"
 
+#include "Algo/AnyOf.h"
 #include "Algo/Find.h"
 #include "Containers/SharedString.h"
 #include "DerivedDataCacheKey.h"
@@ -122,6 +123,14 @@ uint32 FCacheKeyFilter::GetSalt() const
 	return State ? State->Salt : 0;
 }
 
+bool FCacheKeyFilter::RequiresSalt() const
+{
+	using namespace UE::DerivedData::Private;
+	using FTypeRate = FCacheKeyFilterState::FTypeRate;
+	const auto HasFractionalRate = [](uint32 Rate) { return Rate > 0 && Rate < MAX_uint32; };
+	return State && (HasFractionalRate(State->DefaultRate) || Algo::AnyOf(State->Types, Projection(&FTypeRate::Rate, HasFractionalRate)));
+}
+
 bool FCacheKeyFilter::IsMatch(const FCacheKey& Key) const
 {
 	using namespace UE::DerivedData::Private;
diff --git a/Engine/Source/Developer/DerivedDataCache/Private/DerivedDataCacheReplay.cpp b/Engine/Source/Developer/DerivedDataCache/Private/DerivedDataCacheReplay.cpp
--- a/Engine/Source/Developer/DerivedDataCache/Private/DerivedDataCacheReplay.cpp
+++ b/Engine/Source/Developer/DerivedDataCache/Private/DerivedDataCacheReplay.cpp
@@ -1080,8 +1080,8 @@ static FCacheKeyFilter ParseReplayKeyFilter(const TCHAR* const CommandLine)
 			}
 		}
 
-		UE_LOG(LogDerivedDataCache, Display,
-			TEXT("Replay: Using salt -DDC-ReplaySalt=%u to filter cache keys to replay."), KeyFilter.GetSalt());
+		UE_CLOGF(KeyFilter.RequiresSalt(), LogDerivedDataCache, Display,
+			"Replay: Using salt -DDC-ReplaySalt=%u to filter cache keys to replay.", KeyFilter.GetSalt());
 	}
 
 	return KeyFilter;
diff --git a/Engine/Source/Developer/DerivedDataCache/Private/DerivedDataCacheStoreVerify.cpp b/Engine/Source/Developer/DerivedDataCache/Private/DerivedDataCacheStoreVerify.cpp
--- a/Engine/Source/Developer/DerivedDataCache/Private/DerivedDataCacheStoreVerify.cpp
+++ b/Engine/Source/Developer/DerivedDataCache/Private/DerivedDataCacheStoreVerify.cpp
@@ -62,8 +62,8 @@ public:
 
 		if (Filter)
 		{
-			UE_LOG(LogDerivedDataCache, Display,
-				TEXT("Verify: Using salt -DDC-VerifySalt=%u to filter cache keys to verify."), Filter.GetSalt());
+			UE_CLOGF(Filter.RequiresSalt(), LogDerivedDataCache, Display,
+				"Verify: Using salt -DDC-VerifySalt=%u to filter cache keys to verify.", Filter.GetSalt());
 		}
 
 		bPutOnError = bPutOnError || FParse::Param(CommandLine, TEXT("DDC-VerifyFix"));
