DDC: Moved cache request traits from Private to Internal

#Synced-CL 48512380

diff --git a/Engine/Source/Developer/DerivedDataCache/Internal/DerivedDataCacheRequestTraits.h b/Engine/Source/Developer/DerivedDataCache/Internal/DerivedDataCacheRequestTraits.h
new file mode 100644
--- /dev/null
+++ b/Engine/Source/Developer/DerivedDataCache/Internal/DerivedDataCacheRequestTraits.h
@@ -0,0 +1,92 @@
+// Copyright Epic Games, Inc. All Rights Reserved.
+
+#pragma once
+
+#include "DerivedDataCache.h"
+#include "DerivedDataCacheMethod.h"
+
+namespace UE::DerivedData::Private
+{
+
+template <typename RequestType>
+struct TCacheRequestTraits;
+
+template <>
+struct TCacheRequestTraits<FCachePutRequest>
+{
+	using FRequest = FCachePutRequest;
+	using FResponse = FCachePutResponse;
+	using FOnComplete = FOnCachePutComplete;
+	constexpr static ECacheMethod Method = ECacheMethod::Put;
+	constexpr static void (ICache::*Function)(TConstArrayView<FRequest>, IRequestOwner&, FOnComplete&&) = &ICache::Put;
+	static const FCacheKey& GetKey(const FCachePutRequest& Request) { return Request.Record.GetKey(); }
+};
+
+template <>
+struct TCacheRequestTraits<FCacheGetRequest>
+{
+	using FRequest = FCacheGetRequest;
+	using FResponse = FCacheGetResponse;
+	using FOnComplete = FOnCacheGetComplete;
+	constexpr static ECacheMethod Method = ECacheMethod::Get;
+	constexpr static void (ICache::*Function)(TConstArrayView<FRequest>, IRequestOwner&, FOnComplete&&) = &ICache::Get;
+	static const FCacheKey& GetKey(const FCacheGetRequest& Request) { return Request.Key; }
+};
+
+template <>
+struct TCacheRequestTraits<FCachePutValueRequest>
+{
+	using FRequest = FCachePutValueRequest;
+	using FResponse = FCachePutValueResponse;
+	using FOnComplete = FOnCachePutValueComplete;
+	constexpr static ECacheMethod Method = ECacheMethod::PutValue;
+	constexpr static void (ICache::*Function)(TConstArrayView<FRequest>, IRequestOwner&, FOnComplete&&) = &ICache::PutValue;
+	static const FCacheKey& GetKey(const FCachePutValueRequest& Request) { return Request.Key; }
+};
+
+template <>
+struct TCacheRequestTraits<FCacheGetValueRequest>
+{
+	using FRequest = FCacheGetValueRequest;
+	using FResponse = FCacheGetValueResponse;
+	using FOnComplete = FOnCacheGetValueComplete;
+	constexpr static ECacheMethod Method = ECacheMethod::GetValue;
+	constexpr static void (ICache::*Function)(TConstArrayView<FRequest>, IRequestOwner&, FOnComplete&&) = &ICache::GetValue;
+	static const FCacheKey& GetKey(const FCacheGetValueRequest& Request) { return Request.Key; }
+};
+
+template <>
+struct TCacheRequestTraits<FCacheGetChunkRequest>
+{
+	using FRequest = FCacheGetChunkRequest;
+	using FResponse = FCacheGetChunkResponse;
+	using FOnComplete = FOnCacheGetChunkComplete;
+	constexpr static ECacheMethod Method = ECacheMethod::GetChunks;
+	constexpr static void (ICache::*Function)(TConstArrayView<FRequest>, IRequestOwner&, FOnComplete&&) = &ICache::GetChunks;
+	static const FCacheKey& GetKey(const FCacheGetChunkRequest& Request) { return Request.Key; }
+};
+
+} // UE::DerivedData::Private
+
+namespace UE::DerivedData
+{
+
+/** Alias for the response type corresponding to the request type. */
+template <typename RequestType> using TCacheResponseFor = typename Private::TCacheRequestTraits<RequestType>::FResponse;
+
+/** Alias for the completion callback type corresponding to the request type. */
+template <typename RequestType> using TCacheOnCompleteFor = typename Private::TCacheRequestTraits<RequestType>::FOnComplete;
+
+/** Alias for the cache method corresponding to the request type. */
+template <typename RequestType> inline constexpr auto CacheMethodFor = Private::TCacheRequestTraits<RequestType>::Method;
+
+/** Alias for the cache function corresponding to the request type. */
+template <typename RequestType> inline constexpr auto CacheFunctionFor = Private::TCacheRequestTraits<RequestType>::Function;
+
+template <typename RequestType>
+inline const FCacheKey& GetCacheKey(const RequestType& Request)
+{
+	return Private::TCacheRequestTraits<RequestType>::GetKey(Request);
+}
+
+} // UE::DerivedData
diff --git a/Engine/Source/Developer/DerivedDataCache/Private/DerivedDataCacheReplay.cpp b/Engine/Source/Developer/DerivedDataCache/Private/DerivedDataCacheReplay.cpp
--- a/Engine/Source/Developer/DerivedDataCache/Private/DerivedDataCacheReplay.cpp
+++ b/Engine/Source/Developer/DerivedDataCache/Private/DerivedDataCacheReplay.cpp
@@ -193,7 +193,7 @@ struct TCacheStoreReplayBatch
 			Request.UserData = Index++;
 		}
 
-		(InCacheStore.*CacheFunctionFor<FRequest>)(ModifiedRequests, Owner, [this](FResponse&& Response)
+		(InCacheStore.*CacheStoreFunctionFor<FRequest>)(ModifiedRequests, Owner, [this](FResponse&& Response)
 		{
 			const int32 Index = int32(Response.UserData);
 			Response.UserData = Requests[Index].UserData;
@@ -914,7 +914,7 @@ bool FCacheReplayReader::FState::DispatchRequests(const TConstArrayView<RequestT
 		{
 			// Owners with non-blocking priority can execute the request directly because it will be async.
 			FRequestBarrier Barrier(Owner);
-			Invoke(CacheFunctionFor<RequestType>, TargetCache, FilteredRequests, Owner, [](TCacheResponseFor<RequestType>&&){});
+			Invoke(CacheStoreFunctionFor<RequestType>, TargetCache, FilteredRequests, Owner, [](TCacheResponseFor<RequestType>&&){});
 		}
 		else
 		{
@@ -922,7 +922,7 @@ bool FCacheReplayReader::FState::DispatchRequests(const TConstArrayView<RequestT
 			BlockingTasks.Add(Tasks::Launch(TEXT("CacheReplayTask"), [this, Requests = MoveTemp(FilteredRequests)]
 			{
 				FRequestOwner BlockingOwner(EPriority::Blocking);
-				Invoke(CacheFunctionFor<RequestType>, TargetCache, Requests, BlockingOwner, [](TCacheResponseFor<RequestType>&&){});
+				Invoke(CacheStoreFunctionFor<RequestType>, TargetCache, Requests, BlockingOwner, [](TCacheResponseFor<RequestType>&&){});
 				BlockingOwner.Wait();
 			}));
 		}
diff --git a/Engine/Source/Developer/DerivedDataCache/Private/DerivedDataCacheStore.h b/Engine/Source/Developer/DerivedDataCache/Private/DerivedDataCacheStore.h
--- a/Engine/Source/Developer/DerivedDataCache/Private/DerivedDataCacheStore.h
+++ b/Engine/Source/Developer/DerivedDataCache/Private/DerivedDataCacheStore.h
@@ -5,7 +5,7 @@
 #include "Async/Mutex.h"
 #include "Containers/ContainersFwd.h"
 #include "DerivedDataCache.h"
-#include "DerivedDataCacheMethod.h"
+#include "DerivedDataCacheRequestTraits.h"
 #include "Misc/MonotonicTime.h"
 #include "Templates/DontCopy.h"
 
@@ -343,81 +343,41 @@ namespace Private
 {
 
 template <typename RequestType>
-struct TCacheRequestTraits;
+struct TCacheStoreFunctionFor;
 
 template <>
-struct TCacheRequestTraits<FCachePutRequest>
+struct TCacheStoreFunctionFor<FCachePutRequest>
 {
-	using FRequest = FCachePutRequest;
-	using FResponse = FCachePutResponse;
-	using FOnComplete = FOnCachePutComplete;
-	constexpr static ECacheMethod Method = ECacheMethod::Put;
-	constexpr static void (ICacheStore::*Function)(TConstArrayView<FRequest>, IRequestOwner&, FOnComplete&&) = &ICacheStore::Put;
-	static const FCacheKey& GetKey(const FCachePutRequest& Request) { return Request.Record.GetKey(); }
+	constexpr static void (ICacheStore::*Function)(TConstArrayView<FCachePutRequest>, IRequestOwner&, TCacheOnCompleteFor<FCachePutRequest>&&) = &ICacheStore::Put;
 };
 
 template <>
-struct TCacheRequestTraits<FCacheGetRequest>
+struct TCacheStoreFunctionFor<FCacheGetRequest>
 {
-	using FRequest = FCacheGetRequest;
-	using FResponse = FCacheGetResponse;
-	using FOnComplete = FOnCacheGetComplete;
-	constexpr static ECacheMethod Method = ECacheMethod::Get;
-	constexpr static void (ICacheStore::*Function)(TConstArrayView<FRequest>, IRequestOwner&, FOnComplete&&) = &ICacheStore::Get;
-	static const FCacheKey& GetKey(const FCacheGetRequest& Request) { return Request.Key; }
+	constexpr static void (ICacheStore::*Function)(TConstArrayView<FCacheGetRequest>, IRequestOwner&, TCacheOnCompleteFor<FCacheGetRequest>&&) = &ICacheStore::Get;
 };
 
 template <>
-struct TCacheRequestTraits<FCachePutValueRequest>
+struct TCacheStoreFunctionFor<FCachePutValueRequest>
 {
-	using FRequest = FCachePutValueRequest;
-	using FResponse = FCachePutValueResponse;
-	using FOnComplete = FOnCachePutValueComplete;
-	constexpr static ECacheMethod Method = ECacheMethod::PutValue;
-	constexpr static void (ICacheStore::*Function)(TConstArrayView<FRequest>, IRequestOwner&, FOnComplete&&) = &ICacheStore::PutValue;
-	static const FCacheKey& GetKey(const FCachePutValueRequest& Request) { return Request.Key; }
+	constexpr static void (ICacheStore::*Function)(TConstArrayView<FCachePutValueRequest>, IRequestOwner&, TCacheOnCompleteFor<FCachePutValueRequest>&&) = &ICacheStore::PutValue;
 };
 
 template <>
-struct TCacheRequestTraits<FCacheGetValueRequest>
+struct TCacheStoreFunctionFor<FCacheGetValueRequest>
 {
-	using FRequest = FCacheGetValueRequest;
-	using FResponse = FCacheGetValueResponse;
-	using FOnComplete = FOnCacheGetValueComplete;
-	constexpr static ECacheMethod Method = ECacheMethod::GetValue;
-	constexpr static void (ICacheStore::*Function)(TConstArrayView<FRequest>, IRequestOwner&, FOnComplete&&) = &ICacheStore::GetValue;
-	static const FCacheKey& GetKey(const FCacheGetValueRequest& Request) { return Request.Key; }
+	constexpr static void (ICacheStore::*Function)(TConstArrayView<FCacheGetValueRequest>, IRequestOwner&, TCacheOnCompleteFor<FCacheGetValueRequest>&&) = &ICacheStore::GetValue;
 };
 
 template <>
-struct TCacheRequestTraits<FCacheGetChunkRequest>
+struct TCacheStoreFunctionFor<FCacheGetChunkRequest>
 {
-	using FRequest = FCacheGetChunkRequest;
-	using FResponse = FCacheGetChunkResponse;
-	using FOnComplete = FOnCacheGetChunkComplete;
-	constexpr static ECacheMethod Method = ECacheMethod::GetChunks;
-	constexpr static void (ICacheStore::*Function)(TConstArrayView<FRequest>, IRequestOwner&, FOnComplete&&) = &ICacheStore::GetChunks;
-	static const FCacheKey& GetKey(const FCacheGetChunkRequest& Request) { return Request.Key; }
+	constexpr static void (ICacheStore::*Function)(TConstArrayView<FCacheGetChunkRequest>, IRequestOwner&, TCacheOnCompleteFor<FCacheGetChunkRequest>&&) = &ICacheStore::GetChunks;
 };
 
 } // Private
 
-/** Alias for the response type corresponding to the request type. */
-template <typename RequestType> using TCacheResponseFor = typename Private::TCacheRequestTraits<RequestType>::FResponse;
-
-/** Alias for the completion callback type corresponding to the request type. */
-template <typename RequestType> using TCacheOnCompleteFor = typename Private::TCacheRequestTraits<RequestType>::FOnComplete;
-
-/** Alias for the cache method corresponding to the request type. */
-template <typename RequestType> inline constexpr auto CacheMethodFor = Private::TCacheRequestTraits<RequestType>::Method;
-
-/** Alias for the cache function corresponding to the request type. */
-template <typename RequestType> inline constexpr auto CacheFunctionFor = Private::TCacheRequestTraits<RequestType>::Function;
-
-template <typename RequestType>
-inline const FCacheKey& GetCacheKey(const RequestType& Request)
-{
-	return Private::TCacheRequestTraits<RequestType>::GetKey(Request);
-}
+/** Alias for the cache store function corresponding to the request type. */
+template <typename RequestType> inline constexpr auto CacheStoreFunctionFor = Private::TCacheStoreFunctionFor<RequestType>::Function;
 
 } // UE::DerivedData
