DDC: Added TryLexFromString for FCacheKey and FCacheBucket

#Synced-CL 48512380

diff --git a/Engine/Source/Developer/DerivedDataCache/Private/DerivedDataCacheKey.cpp b/Engine/Source/Developer/DerivedDataCache/Private/DerivedDataCacheKey.cpp
--- a/Engine/Source/Developer/DerivedDataCache/Private/DerivedDataCacheKey.cpp
+++ b/Engine/Source/Developer/DerivedDataCache/Private/DerivedDataCacheKey.cpp
@@ -16,6 +16,7 @@
 #include "Misc/StringBuilder.h"
 #include "Serialization/CompactBinaryWriter.h"
 #include "String/Find.h"
+#include "String/Split.h"
 
 namespace UE::DerivedData::Private
 {
@@ -200,6 +201,24 @@ void FCacheBucket::ToDisplayName(FStringBuilderBase& OutDisplayName) const
 	Private::GetCacheBuckets().GetDisplayName(*this, OutDisplayName);
 }
 
+template <typename CharType>
+static inline bool TryLexKeyFromString(FCacheKey& OutKey, TStringView<CharType> String)
+{
+	TStringView<CharType> Bucket, Hash;
+	return String::SplitFirstChar(String, '/', Bucket, Hash) &&
+		TryLexFromString(OutKey.Bucket, Bucket) && TryLexFromString(OutKey.Hash, Hash);
+}
+
+bool TryLexFromString(FCacheKey& OutKey, FAnsiStringView String)
+{
+	return TryLexKeyFromString(OutKey, String);
+}
+
+bool TryLexFromString(FCacheKey& OutKey, FWideStringView String)
+{
+	return TryLexKeyFromString(OutKey, String);
+}
+
 void SerializeForLog(FCbWriter& Writer, const FCacheKey& Key)
 {
 	Writer.AddString(WriteToString<96>(Key));
diff --git a/Engine/Source/Developer/DerivedDataCache/Public/DerivedDataCacheKey.h b/Engine/Source/Developer/DerivedDataCache/Public/DerivedDataCacheKey.h
--- a/Engine/Source/Developer/DerivedDataCache/Public/DerivedDataCacheKey.h
+++ b/Engine/Source/Developer/DerivedDataCache/Public/DerivedDataCacheKey.h
@@ -135,6 +135,17 @@ inline TStringBuilderBase<CharType>& operator<<(TStringBuilderBase<CharType>& Bu
 	return Builder << Bucket.ToString();
 }
 
+template <typename CharType>
+inline bool TryLexFromString(FCacheBucket& OutBucket, TStringView<CharType> String)
+{
+	if (FCacheBucket::IsValidName(String))
+	{
+		OutBucket = FCacheBucket(String);
+		return true;
+	}
+	return false;
+}
+
 ///////////////////////////////////////////////////////////////////////////////////////////////////////////////////////
 
 template <typename CharType>
@@ -143,6 +154,9 @@ inline TStringBuilderBase<CharType>& operator<<(TStringBuilderBase<CharType>& Bu
 	return Builder << Key.Bucket << CharType('/') << Key.Hash;
 }
 
+UE_API bool TryLexFromString(FCacheKey& OutKey, FAnsiStringView String);
+UE_API bool TryLexFromString(FCacheKey& OutKey, FWideStringView String);
+
 UE_API void SerializeForLog(FCbWriter& Writer, const FCacheKey& Key);
 
 ///////////////////////////////////////////////////////////////////////////////////////////////////////////////////////
diff --git a/Engine/Source/Runtime/Core/Public/Hash/Blake3.h b/Engine/Source/Runtime/Core/Public/Hash/Blake3.h
--- a/Engine/Source/Runtime/Core/Public/Hash/Blake3.h
+++ b/Engine/Source/Runtime/Core/Public/Hash/Blake3.h
@@ -173,6 +173,18 @@ inline void LexFromString(FBlake3Hash& OutHash, const TCHAR* Buffer)
 	OutHash = FBlake3Hash(Buffer);
 }
 
+/** Construct a hash from a 64-character hex string. */
+template <typename CharType>
+inline bool TryLexFromString(FBlake3Hash& OutHash, TStringView<CharType> String)
+{
+	if (String.Len() == sizeof(FBlake3Hash) * 2)
+	{
+		OutHash = FBlake3Hash(String);
+		return true;
+	}
+	return false;
+}
+
 /** Convert a hash to a 64-character hex string. */
 [[nodiscard]] CORE_API FString LexToString(const FBlake3Hash& Hash);
 
diff --git a/Engine/Source/Runtime/Core/Public/IO/IoHash.h b/Engine/Source/Runtime/Core/Public/IO/IoHash.h
--- a/Engine/Source/Runtime/Core/Public/IO/IoHash.h
+++ b/Engine/Source/Runtime/Core/Public/IO/IoHash.h
@@ -167,6 +167,18 @@ inline void LexFromString(FIoHash& OutHash, const TCHAR* Buffer)
 	OutHash = FIoHash(Buffer);
 }
 
+/** Construct a hash from a 40-character hex string. */
+template <typename CharType>
+inline bool TryLexFromString(FIoHash& OutHash, TStringView<CharType> String)
+{
+	if (String.Len() == sizeof(FIoHash) * 2)
+	{
+		OutHash = FIoHash(String);
+		return true;
+	}
+	return false;
+}
+
 /** Convert a hash to a 40-character hex string. */
 [[nodiscard]] CORE_API FString LexToString(const FIoHash& Hash);
 
