DDC: Optimized loading from compact binary

The net reduction in replay loading time is over 50%.
Loading requests and responses together and serializing fields in one pass reduced replay load time by ~36%.
Validating compact binary hashes instead of walking the object graph reduced replay load time by ~26%.
Loading compact binary as views instead of ref-counted objects reduced replay load time by ~6%.

#Synced-CL 49572071

diff --git a/Engine/Source/Developer/DerivedDataCache/Private/DerivedDataCacheRecord.cpp b/Engine/Source/Developer/DerivedDataCache/Private/DerivedDataCacheRecord.cpp
--- a/Engine/Source/Developer/DerivedDataCache/Private/DerivedDataCacheRecord.cpp
+++ b/Engine/Source/Developer/DerivedDataCache/Private/DerivedDataCacheRecord.cpp
@@ -46,6 +46,7 @@ class FCacheRecordInternal final : public ICacheRecordInternal
 public:
 	FCacheRecordInternal() = default;
 	explicit FCacheRecordInternal(FCacheRecordBuilderInternal&& RecordBuilder);
+	explicit FCacheRecordInternal(const FCacheKey& Key, FCbObject&& Meta, TArray<FValueWithId>&& Values);
 
 	~FCacheRecordInternal() final = default;
 
@@ -83,6 +84,13 @@ FCacheRecordInternal::FCacheRecordInternal(FCacheRecordBuilderInternal&& RecordB
 {
 }
 
+FCacheRecordInternal::FCacheRecordInternal(const FCacheKey& InKey, FCbObject&& InMeta, TArray<FValueWithId>&& InValues)
+	: Key(InKey)
+	, Meta(MoveTemp(InMeta))
+	, Values(MoveTemp(InValues))
+{
+}
+
 const FCacheKey& FCacheRecordInternal::GetKey() const
 {
 	return Key;
@@ -196,6 +204,11 @@ FOptionalCacheRecord FCacheRecord::Load(const FCbPackage& Package)
 	return Load(Package, Package.GetObject());
 }
 
+FCacheRecord FCacheRecord::CreateByMove(const FCacheKey& Key, FCbObject&& Meta, TArray<FValueWithId>&& Values)
+{
+	return FCacheRecord(new Private::FCacheRecordInternal(Key, MoveTemp(Meta), MoveTemp(Values)));
+}
+
 FCacheRecordBuilder::FCacheRecordBuilder(const FCacheKey& Key)
 	: RecordBuilder(new Private::FCacheRecordBuilderInternal(Key))
 {
diff --git a/Engine/Source/Developer/DerivedDataCache/Private/DerivedDataCacheReplay.cpp b/Engine/Source/Developer/DerivedDataCache/Private/DerivedDataCacheReplay.cpp
--- a/Engine/Source/Developer/DerivedDataCache/Private/DerivedDataCacheReplay.cpp
+++ b/Engine/Source/Developer/DerivedDataCache/Private/DerivedDataCacheReplay.cpp
@@ -25,6 +25,7 @@
 #include "ProfilingDebugging/CpuProfilerTrace.h"
 #include "Serialization/CompactBinary.h"
 #include "Serialization/CompactBinarySerialization.h"
+#include "Serialization/CompactBinaryValidation.h"
 #include "Serialization/CompactBinaryWriter.h"
 #include "Serialization/MemoryReader.h"
 #include "String/Find.h"
@@ -452,24 +453,41 @@ void ICacheReplayWriter::Write(TConstArrayView<FCacheGetChunkRequest> Requests,
 
 ///////////////////////////////////////////////////////////////////////////////////////////////////////////////////////
 
-template <typename CacheType, typename ArrayType>
-static bool LoadCacheTypeFromCompactBinary(const FCbFieldView Field, ArrayType& OutArray)
+template <typename RequestType>
+static bool LoadBatchArraysFromCompactBinary(
+	const FCbFieldView Field,
+	TArray<RequestType, TInlineAllocator<16>>& OutRequests,
+	TArray<TCacheResponseFor<RequestType>, TInlineAllocator<16>>& OutResponses)
 {
-	if constexpr (std::is_default_constructible_v<CacheType>)
+	using ResponseType = TCacheResponseFor<RequestType>;
+	if constexpr (!std::is_default_constructible_v<RequestType>)
 	{
-		CacheType& Value = OutArray.AddDefaulted_GetRef();
-		return LoadFromCompactBinary(Field, Value);
+		TOptional<RequestType> Request;
+		ResponseType& Response = OutResponses.AddDefaulted_GetRef();
+		if (LoadFromCompactBinary(Field, Request, &Response))
+		{
+			OutRequests.Emplace(MoveTemp(Request.GetValue()));
+			return true;
+		}
+		return false;
 	}
-	else
+	else if constexpr (!std::is_default_constructible_v<ResponseType>)
 	{
-		TOptional<CacheType> Value;
-		if (LoadFromCompactBinary(Field, Value))
+		RequestType& Request = OutRequests.AddDefaulted_GetRef();
+		TOptional<ResponseType> Response;
+		if (LoadFromCompactBinary(Field, Request, &Response))
 		{
-			OutArray.Emplace(MoveTemp(Value.GetValue()));
+			OutResponses.Emplace(MoveTemp(Response.GetValue()));
 			return true;
 		}
 		return false;
 	}
+	else
+	{
+		RequestType& Request = OutRequests.AddDefaulted_GetRef();
+		ResponseType& Response = OutResponses.AddDefaulted_GetRef();
+		return LoadFromCompactBinary(Field, Request, &Response);
+	}
 }
 
 template <typename RequestType>
@@ -487,8 +505,7 @@ static bool LoadBatchFromCompactBinary(ICacheReplayReader& Reader, const FCbObje
 		Responses.Reserve(BatchSize);
 		for (FCbFieldView Field : Array)
 		{
-			if (!LoadCacheTypeFromCompactBinary<RequestType>(Field, Requests) ||
-				!LoadCacheTypeFromCompactBinary<TCacheResponseFor<RequestType>>(Field, Responses))
+			if (!LoadBatchArraysFromCompactBinary<RequestType>(Field, Requests, Responses))
 			{
 				return false;
 			}
@@ -618,15 +635,15 @@ TUniquePtr<ICacheReplayWriter> CreateReplayArchiveWriter(TNotNull<TUniquePtr<FAr
 
 ///////////////////////////////////////////////////////////////////////////////////////////////////////////////////////
 
-static bool ReadReplayFromArchive(ICacheReplayReader& Reader, FArchive& Ar, FUniqueBuffer& Scratch)
+static bool ReadReplayFromArchive(ICacheReplayReader& Reader, FArchive& Ar, FUniqueBuffer& Scratch, ECbValidateMode ValidateMode)
 {
-	const auto Alloc = [&Scratch](const uint64 Size) -> FUniqueBuffer
+	const auto Alloc = [&Scratch](const uint64 Size) -> FMutableMemoryView
 	{
-		if (Size <= Scratch.GetSize())
+		if (UNLIKELY(Scratch.GetSize() < Size))
 		{
-			return FUniqueBuffer::MakeView(Scratch.GetView().Left(Size));
+			Scratch = FUniqueBuffer::Alloc(FMath::RoundUpToPowerOfTwo64(Size));
 		}
-		return FUniqueBuffer::Alloc(Size);
+		return Scratch.GetView().Left(Size);
 	};
 
 	// A scratch buffer for decompressed blocks of fields.
@@ -635,7 +652,7 @@ static bool ReadReplayFromArchive(ICacheReplayReader& Reader, FArchive& Ar, FUni
 	bool bOk = true;
 	for (int64 Offset = Ar.Tell(); Offset < Ar.TotalSize(); Offset = Ar.Tell())
 	{
-		FCbField Field = LoadCompactBinary(Ar, Alloc);
+		FCbFieldView Field = LoadAndValidateCompactBinaryView(Ar, Alloc, ValidateMode);
 		if (!Field)
 		{
 			UE_LOG(LogDerivedDataCache, Error,
@@ -648,7 +665,7 @@ static bool ReadReplayFromArchive(ICacheReplayReader& Reader, FArchive& Ar, FUni
 		// A binary field is used to store a compressed buffer containing a sequence of compact binary objects.
 		if (Field.IsBinary())
 		{
-			FCompressedBuffer CompressedBuffer = FCompressedBuffer::FromCompressed(Field.AsBinary());
+			FCompressedBuffer CompressedBuffer = FCompressedBuffer::FromCompressed(FSharedBuffer::MakeView(Field.AsBinaryView()));
 			if (!CompressedBuffer)
 			{
 				UE_LOG(LogDerivedDataCache, Warning,
@@ -681,9 +698,21 @@ static bool ReadReplayFromArchive(ICacheReplayReader& Reader, FArchive& Ar, FUni
 				continue;
 			}
 
+			const FIoHash RawBlockHash = FIoHash::HashBuffer(RawBlockView);
+			if (RawBlockHash != CompressedBuffer.GetRawHash())
+			{
+				UE_LOGFMT(LogDerivedDataCache, Warning,
+					"Replay: Failed to decompress compressed buffer from binary field at offset {FieldOffset}. "
+					"Decompressed buffer has raw hash {ActualHash} when {ExpectedHash} was expected. "
+					"Archive is at offset {ArchiveOffset} and has total size {ArchiveSize}.",
+					Offset, RawBlockHash, CompressedBuffer.GetRawHash(), Ar.Tell(), Ar.TotalSize());
+				bOk = false;
+				continue;
+			}
+
 			// DO NOT ACCESS Field PAST THIS POINT. The recursive call overwrites it in Scratch.
 			FMemoryReaderView InnerAr(RawBlockView);
-			if (!ReadReplayFromArchive(Reader, InnerAr, Scratch))
+			if (!ReadReplayFromArchive(Reader, InnerAr, Scratch, ECbValidateMode::None))
 			{
 				bOk = false;
 				continue;
@@ -691,7 +720,7 @@ static bool ReadReplayFromArchive(ICacheReplayReader& Reader, FArchive& Ar, FUni
 		}
 
 		// An object field is used to store one batch of cache requests.
-		if (Field.IsObject() && !Reader.SerializeBatch(Field.AsObject()))
+		if (Field.IsObject() && !Reader.SerializeBatch(Field.AsObjectView()))
 		{
 			UE_LOG(LogDerivedDataCache, Warning,
 				TEXT("Replay: Failed to load cache request from object field at offset %" INT64_FMT ". "
@@ -708,7 +737,7 @@ bool ReadReplayFromArchive(ICacheReplayReader& Reader, FArchive& Ar)
 {
 	// A scratch buffer for the compact binary fields.
 	FUniqueBuffer Scratch = FUniqueBuffer::Alloc(256 * 1024);
-	return ReadReplayFromArchive(Reader, Ar, Scratch);
+	return ReadReplayFromArchive(Reader, Ar, Scratch, ECbValidateMode::Default);
 }
 
 bool ReadReplayFromFile(ICacheReplayReader& Reader, const TCHAR* Path)
diff --git a/Engine/Source/Developer/DerivedDataCache/Private/DerivedDataSerialization.cpp b/Engine/Source/Developer/DerivedDataCache/Private/DerivedDataSerialization.cpp
--- a/Engine/Source/Developer/DerivedDataCache/Private/DerivedDataSerialization.cpp
+++ b/Engine/Source/Developer/DerivedDataCache/Private/DerivedDataSerialization.cpp
@@ -2,6 +2,8 @@
 
 #include "DerivedDataSerialization.h"
 
+#include "Algo/IsSorted.h"
+#include "Algo/Sort.h"
 #include "DerivedDataCache.h"
 #include "DerivedDataCacheKey.h"
 #include "DerivedDataCacheMethod.h"
@@ -9,11 +11,13 @@
 #include "DerivedDataCacheRecord.h"
 #include "DerivedDataValue.h"
 #include "DerivedDataValueId.h"
+#include "Misc/NotNull.h"
 #include "Misc/Optional.h"
 #include "Misc/StringBuilder.h"
 #include "Serialization/CompactBinary.h"
 #include "Serialization/CompactBinaryPackage.h"
 #include "Serialization/CompactBinarySerialization.h"
+#include "Serialization/CompactBinaryValue.h"
 #include "Serialization/CompactBinaryWriter.h"
 
 namespace UE::DerivedData
@@ -60,25 +64,41 @@ void SaveMetaToCompactBinary(FCbWriter& Writer, const FCbObject& Meta, EMetaSave
 	}
 }
 
+inline bool LoadMetaFromObject(FCbValue Value, FCbObject& OutMeta)
+{
+	if (FCbFieldType::IsObject(Value.GetType()))
+	{
+		OutMeta = FCbObject::Clone(Value.AsObjectView());
+		return true;
+	}
+	return false;
+}
+
+inline bool LoadMetaFromHashAndSize(FCbValue Hash, FCbValue Size, FCbObject& OutMeta, TNotNull<FLoadAttachmentFn> LoadAttachment)
+{
+	if (FCbFieldType::IsHash(Hash.GetType()) && Size.GetType() == ECbFieldType::IntegerPositive)
+	{
+		FCbObject Meta = LoadAttachment(Hash.AsHash()).AsObject();
+		if (Meta.GetSize() == Size.AsIntegerPositive())
+		{
+			OutMeta = MoveTemp(Meta);
+			return true;
+		}
+	}
+	return false;
+}
+
 bool LoadMetaFromCompactBinary(FCbFieldView Field, FCbObject& OutMeta, FLoadAttachmentFn LoadAttachment)
 {
-	if (const FCbObjectView MetaView = Field[ANSITEXTVIEW("Meta")].AsObjectView())
+	if (LoadMetaFromObject(Field[ANSITEXTVIEW("Meta")], OutMeta))
 	{
-		OutMeta = FCbObject::Clone(MetaView);
 		return true;
 	}
 	else if (LoadAttachment)
 	{
-		const FIoHash MetaHash = Field[ANSITEXTVIEW("MetaHash")].AsHash();
-		const uint64 MetaSize = Field[ANSITEXTVIEW("MetaSize")].AsUInt64(MAX_uint64);
-		if (!MetaHash.IsZero() && MetaSize != MAX_uint64)
+		if (LoadMetaFromHashAndSize(Field[ANSITEXTVIEW("MetaHash")], Field[ANSITEXTVIEW("MetaSize")], OutMeta, LoadAttachment))
 		{
-			FCbObject Meta = LoadAttachment(MetaHash).AsObject();
-			if (Meta.GetSize() == MetaSize)
-			{
-				OutMeta = MoveTemp(Meta);
-				return true;
-			}
+			return true;
 		}
 	}
 	return false;
@@ -95,7 +115,17 @@ void SaveToCompactBinary(FCbWriter& Writer, const FCachePutRequest& Request, con
 	{
 		Writer << ANSITEXTVIEW("Name") << Request.Name;
 	}
-	SaveToCompactBinary(Writer, Request.Record, EMetaSaveFormat::Attachment, SaveAttachment);
+	Writer << ANSITEXTVIEW("Key") << Request.Record.GetKey();
+	SaveMetaToCompactBinary(Writer, Request.Record.GetMeta(), EMetaSaveFormat::Attachment, SaveAttachment);
+	if (const TConstArrayView<FValueWithId> Values = Request.Record.GetValues(); !Values.IsEmpty())
+	{
+		Writer.BeginArray(ANSITEXTVIEW("Values"));
+		for (const FValueWithId& Value : Values)
+		{
+			SaveToCompactBinary(Writer, Value, SaveAttachment);
+		}
+		Writer.EndArray();
+	}
 	if (!Request.Policy.IsDefault())
 	{
 		Writer << ANSITEXTVIEW("Policy") << Request.Policy;
@@ -111,29 +141,132 @@ void SaveToCompactBinary(FCbWriter& Writer, const FCachePutRequest& Request, con
 	Writer.EndObject();
 }
 
-bool LoadFromCompactBinary(FCbFieldView Field, TOptional<FCachePutRequest>& OutRequest, FLoadAttachmentFn LoadAttachment)
+bool LoadFromCompactBinary(FCbFieldView Field, TOptional<FCachePutRequest>& OutRequest, FCachePutResponse* OutResponse, FLoadAttachmentFn LoadAttachment)
 {
-	FOptionalCacheRecord Record;
-	if (!LoadFromCompactBinary(Field, Record, LoadAttachment))
+	bool bOk = Field.IsObject();
+
+	FSharedString Name;
+	FCacheKey Key;
+	FCbValue MetaObjectValue;
+	FCbValue MetaHashValue;
+	FCbValue MetaSizeValue;
+	FCbArrayView ValuesArray;
+	FCacheRecordPolicy Policy;
+	uint64 UserData = 0;
+	EStatus Status = EStatus::Ok;
+
+	constexpr static uint32 FoundName     = 1 << 0;
+	constexpr static uint32 FoundKey      = 1 << 1;
+	constexpr static uint32 FoundMeta     = 1 << 2;
+	constexpr static uint32 FoundMetaHash = 1 << 3;
+	constexpr static uint32 FoundMetaSize = 1 << 4;
+	constexpr static uint32 FoundValues   = 1 << 5;
+	constexpr static uint32 FoundPolicy   = 1 << 6;
+	constexpr static uint32 FoundUserData = 1 << 7;
+	constexpr static uint32 FoundStatus   = 1 << 8;
+
+	uint32 FoundMask = 0;
+	for (FCbFieldViewIterator It = Field.CreateViewIterator(); It;)
+	{
+		const uint32 LastFoundMask = FoundMask;
+		if (!(FoundMask & FoundName) && It.GetName() == ANSITEXTVIEW("Name"))
+		{
+			LoadFromCompactBinary(It, Name);
+			FoundMask |= FoundName;
+			++It;
+		}
+		if (!(FoundMask & FoundKey) && It.GetName() == ANSITEXTVIEW("Key"))
+		{
+			bOk &= LoadFromCompactBinary(It, Key);
+			FoundMask |= FoundKey;
+			++It;
+		}
+		if (!(FoundMask & FoundMeta) && It.GetName() == ANSITEXTVIEW("Meta"))
+		{
+			MetaObjectValue = It;
+			FoundMask |= FoundMeta;
+			++It;
+		}
+		else
+		{
+			if (!(FoundMask & FoundMetaHash) && It.GetName() == ANSITEXTVIEW("MetaHash"))
+			{
+				MetaHashValue = It;
+				FoundMask |= FoundMetaHash;
+				++It;
+			}
+			if (!(FoundMask & FoundMetaSize) && It.GetName() == ANSITEXTVIEW("MetaSize"))
+			{
+				MetaSizeValue = It;
+				FoundMask |= FoundMetaSize;
+				++It;
+			}
+		}
+		if (!(FoundMask & FoundValues) && It.GetName() == ANSITEXTVIEW("Values"))
+		{
+			ValuesArray = It.AsArrayView();
+			FoundMask |= FoundValues;
+			++It;
+		}
+		if (!(FoundMask & FoundPolicy) && It.GetName() == ANSITEXTVIEW("Policy"))
+		{
+			LoadFromCompactBinary(It, Policy);
+			FoundMask |= FoundPolicy;
+			++It;
+		}
+		if (!(FoundMask & FoundUserData) && It.GetName() == ANSITEXTVIEW("UserData"))
+		{
+			LoadFromCompactBinary(It, UserData, 0);
+			FoundMask |= FoundUserData;
+			++It;
+		}
+		if (!(FoundMask & FoundStatus) && It.GetName() == ANSITEXTVIEW("Status"))
+		{
+			LoadFromCompactBinary(It, Status, EStatus::Ok);
+			FoundMask |= FoundStatus;
+			++It;
+		}
+		if (LastFoundMask == FoundMask)
+		{
+			++It;
+		}
+	}
+
+	FCbObject Meta;
+	if (MetaObjectValue)
+	{
+		LoadMetaFromObject(MetaObjectValue, Meta);
+	}
+	else if (LoadAttachment && MetaHashValue && MetaSizeValue)
+	{
+		LoadMetaFromHashAndSize(MetaHashValue, MetaSizeValue, Meta, LoadAttachment);
+	}
+
+	TArray<FValueWithId> Values;
+	Values.Reserve(int32(ValuesArray.Num()));
+	for (FCbFieldView ValueField : ValuesArray)
+	{
+		FValueWithId& Value = Values.AddDefaulted_GetRef();
+		bOk &= LoadFromCompactBinary(ValueField, Value, LoadAttachment);
+	}
+	if (UNLIKELY(!Algo::IsSorted(Values)))
+	{
+		Algo::Sort(Values);
+	}
+
+	if (bOk)
+	{
+		if (OutResponse)
+		{
+			*OutResponse = {Name, Key, UserData, Status};
+		}
+		OutRequest = {Name, FCacheRecord::CreateByMove(Key, MoveTemp(Meta), MoveTemp(Values)), Policy, UserData};
+	}
+	else
 	{
 		OutRequest.Reset();
-		return false;
 	}
-	OutRequest = FCachePutRequest{.Record = MoveTemp(Record).Get()};
-	bool bOk = Field.IsObject();
-	LoadFromCompactBinary(Field[ANSITEXTVIEW("Name")], OutRequest->Name);
-	LoadFromCompactBinary(Field[ANSITEXTVIEW("Policy")], OutRequest->Policy);
-	LoadFromCompactBinary(Field[ANSITEXTVIEW("UserData")], OutRequest->UserData, 0);
-	return bOk;
-}
 
-bool LoadFromCompactBinary(FCbFieldView Field, FCachePutResponse& OutResponse)
-{
-	bool bOk = Field.IsObject();
-	LoadFromCompactBinary(Field[ANSITEXTVIEW("Name")], OutResponse.Name);
-	bOk &= LoadFromCompactBinary(Field[ANSITEXTVIEW("Key")], OutResponse.Key);
-	LoadFromCompactBinary(Field[ANSITEXTVIEW("UserData")], OutResponse.UserData, 0);
-	LoadFromCompactBinary(Field[ANSITEXTVIEW("Status")], OutResponse.Status, EStatus::Ok);
 	return bOk;
 }
 
@@ -175,28 +308,127 @@ void SaveToCompactBinary(FCbWriter& Writer, const FCacheGetRequest& Request, con
 	Writer.EndObject();
 }
 
-bool LoadFromCompactBinary(FCbFieldView Field, FCacheGetRequest& OutRequest)
+bool LoadFromCompactBinary(FCbFieldView Field, FCacheGetRequest& OutRequest, TOptional<FCacheGetResponse>* OutResponse, FLoadAttachmentFn LoadAttachment)
 {
 	bool bOk = Field.IsObject();
-	LoadFromCompactBinary(Field[ANSITEXTVIEW("Name")], OutRequest.Name);
-	bOk &= LoadFromCompactBinary(Field[ANSITEXTVIEW("Key")], OutRequest.Key);
-	LoadFromCompactBinary(Field[ANSITEXTVIEW("Policy")], OutRequest.Policy);
-	LoadFromCompactBinary(Field[ANSITEXTVIEW("UserData")], OutRequest.UserData, 0);
-	return bOk;
-}
 
-bool LoadFromCompactBinary(FCbFieldView Field, TOptional<FCacheGetResponse>& OutResponse, FLoadAttachmentFn LoadAttachment)
-{
-	FOptionalCacheRecord Record;
-	if (!LoadFromCompactBinary(Field, Record, LoadAttachment))
+	FCbValue MetaObjectValue;
+	FCbValue MetaHashValue;
+	FCbValue MetaSizeValue;
+	FCbArrayView ValuesArray;
+	EStatus Status = EStatus::Ok;
+
+	constexpr static uint32 FoundName     = 1 << 0;
+	constexpr static uint32 FoundKey      = 1 << 1;
+	constexpr static uint32 FoundMeta     = 1 << 2;
+	constexpr static uint32 FoundMetaHash = 1 << 3;
+	constexpr static uint32 FoundMetaSize = 1 << 4;
+	constexpr static uint32 FoundValues   = 1 << 5;
+	constexpr static uint32 FoundPolicy   = 1 << 6;
+	constexpr static uint32 FoundUserData = 1 << 7;
+	constexpr static uint32 FoundStatus   = 1 << 8;
+
+	uint32 FoundMask = 0;
+	for (FCbFieldViewIterator It = Field.CreateViewIterator(); It;)
 	{
-		return false;
+		const uint32 LastFoundMask = FoundMask;
+		if (!(FoundMask & FoundName) && It.GetName() == ANSITEXTVIEW("Name"))
+		{
+			LoadFromCompactBinary(It, OutRequest.Name);
+			FoundMask |= FoundName;
+			++It;
+		}
+		if (!(FoundMask & FoundKey) && It.GetName() == ANSITEXTVIEW("Key"))
+		{
+			bOk &= LoadFromCompactBinary(It, OutRequest.Key);
+			FoundMask |= FoundKey;
+			++It;
+		}
+		if (!(FoundMask & FoundMeta) && It.GetName() == ANSITEXTVIEW("Meta"))
+		{
+			MetaObjectValue = It;
+			FoundMask |= FoundMeta;
+			++It;
+		}
+		else
+		{
+			if (!(FoundMask & FoundMetaHash) && It.GetName() == ANSITEXTVIEW("MetaHash"))
+			{
+				MetaHashValue = It;
+				FoundMask |= FoundMetaHash;
+				++It;
+			}
+			if (!(FoundMask & FoundMetaSize) && It.GetName() == ANSITEXTVIEW("MetaSize"))
+			{
+				MetaSizeValue = It;
+				FoundMask |= FoundMetaSize;
+				++It;
+			}
+		}
+		if (!(FoundMask & FoundValues) && It.GetName() == ANSITEXTVIEW("Values"))
+		{
+			ValuesArray = It.AsArrayView();
+			FoundMask |= FoundValues;
+			++It;
+		}
+		if (!(FoundMask & FoundPolicy) && It.GetName() == ANSITEXTVIEW("Policy"))
+		{
+			LoadFromCompactBinary(It, OutRequest.Policy);
+			FoundMask |= FoundPolicy;
+			++It;
+		}
+		if (!(FoundMask & FoundUserData) && It.GetName() == ANSITEXTVIEW("UserData"))
+		{
+			LoadFromCompactBinary(It, OutRequest.UserData, 0);
+			FoundMask |= FoundUserData;
+			++It;
+		}
+		if (!(FoundMask & FoundStatus) && It.GetName() == ANSITEXTVIEW("Status"))
+		{
+			LoadFromCompactBinary(It, Status, EStatus::Ok);
+			FoundMask |= FoundStatus;
+			++It;
+		}
+		if (LastFoundMask == FoundMask)
+		{
+			++It;
+		}
 	}
-	OutResponse = FCacheGetResponse{.Record = MoveTemp(Record).Get()};
-	bool bOk = Field.IsObject();
-	LoadFromCompactBinary(Field[ANSITEXTVIEW("Name")], OutResponse->Name);
-	LoadFromCompactBinary(Field[ANSITEXTVIEW("UserData")], OutResponse->UserData, 0);
-	LoadFromCompactBinary(Field[ANSITEXTVIEW("Status")], OutResponse->Status, EStatus::Ok);
+
+	if (OutResponse)
+	{
+		FCbObject Meta;
+		if (MetaObjectValue)
+		{
+			LoadMetaFromObject(MetaObjectValue, Meta);
+		}
+		else if (LoadAttachment && MetaHashValue && MetaSizeValue)
+		{
+			LoadMetaFromHashAndSize(MetaHashValue, MetaSizeValue, Meta, LoadAttachment);
+		}
+
+		TArray<FValueWithId> Values;
+		Values.Reserve(int32(ValuesArray.Num()));
+		for (FCbFieldView ValueField : ValuesArray)
+		{
+			FValueWithId& Value = Values.AddDefaulted_GetRef();
+			bOk &= LoadFromCompactBinary(ValueField, Value, LoadAttachment);
+		}
+		if (UNLIKELY(!Algo::IsSorted(Values)))
+		{
+			Algo::Sort(Values);
+		}
+
+		if (bOk)
+		{
+			*OutResponse = {OutRequest.Name, FCacheRecord::CreateByMove(OutRequest.Key, MoveTemp(Meta), MoveTemp(Values)), OutRequest.UserData, Status};
+		}
+		else
+		{
+			OutResponse->Reset();
+		}
+	}
+
 	return bOk;
 }
 
@@ -232,24 +464,70 @@ void SaveToCompactBinary(FCbWriter& Writer, const FCachePutValueRequest& Request
 	Writer.EndObject();
 }
 
-bool LoadFromCompactBinary(FCbFieldView Field, FCachePutValueRequest& OutRequest, FLoadAttachmentFn LoadAttachment)
+bool LoadFromCompactBinary(FCbFieldView Field, FCachePutValueRequest& OutRequest, FCachePutValueResponse* OutResponse, FLoadAttachmentFn LoadAttachment)
 {
 	bool bOk = Field.IsObject();
-	LoadFromCompactBinary(Field[ANSITEXTVIEW("Name")], OutRequest.Name);
-	bOk &= LoadFromCompactBinary(Field[ANSITEXTVIEW("Key")], OutRequest.Key);
-	LoadFromCompactBinary(Field[ANSITEXTVIEW("Value")], OutRequest.Value, LoadAttachment);
-	LoadFromCompactBinary(Field[ANSITEXTVIEW("Policy")], OutRequest.Policy, ECachePolicy::Default);
-	LoadFromCompactBinary(Field[ANSITEXTVIEW("UserData")], OutRequest.UserData, 0);
-	return bOk;
-}
 
-bool LoadFromCompactBinary(FCbFieldView Field, FCachePutValueResponse& OutResponse)
-{
-	bool bOk = Field.IsObject();
-	LoadFromCompactBinary(Field[ANSITEXTVIEW("Name")], OutResponse.Name);
-	bOk &= LoadFromCompactBinary(Field[ANSITEXTVIEW("Key")], OutResponse.Key);
-	LoadFromCompactBinary(Field[ANSITEXTVIEW("UserData")], OutResponse.UserData, 0);
-	LoadFromCompactBinary(Field[ANSITEXTVIEW("Status")], OutResponse.Status, EStatus::Ok);
+	EStatus Status = EStatus::Ok;
+
+	constexpr static uint32 FoundName     = 1 << 0;
+	constexpr static uint32 FoundKey      = 1 << 1;
+	constexpr static uint32 FoundValue    = 1 << 2;
+	constexpr static uint32 FoundPolicy   = 1 << 3;
+	constexpr static uint32 FoundUserData = 1 << 4;
+	constexpr static uint32 FoundStatus   = 1 << 5;
+
+	uint32 FoundMask = 0;
+	for (FCbFieldViewIterator It = Field.CreateViewIterator(); It;)
+	{
+		const uint32 LastFoundMask = FoundMask;
+		if (!(FoundMask & FoundName) && It.GetName() == ANSITEXTVIEW("Name"))
+		{
+			LoadFromCompactBinary(It, OutRequest.Name);
+			FoundMask |= FoundName;
+			++It;
+		}
+		if (!(FoundMask & FoundKey) && It.GetName() == ANSITEXTVIEW("Key"))
+		{
+			bOk &= LoadFromCompactBinary(It, OutRequest.Key);
+			FoundMask |= FoundKey;
+			++It;
+		}
+		if (!(FoundMask & FoundValue) && It.GetName() == ANSITEXTVIEW("Value"))
+		{
+			LoadFromCompactBinary(It, OutRequest.Value, LoadAttachment);
+			FoundMask |= FoundValue;
+			++It;
+		}
+		if (!(FoundMask & FoundPolicy) && It.GetName() == ANSITEXTVIEW("Policy"))
+		{
+			LoadFromCompactBinary(It, OutRequest.Policy, ECachePolicy::Default);
+			FoundMask |= FoundPolicy;
+			++It;
+		}
+		if (!(FoundMask & FoundUserData) && It.GetName() == ANSITEXTVIEW("UserData"))
+		{
+			LoadFromCompactBinary(It, OutRequest.UserData, 0);
+			FoundMask |= FoundUserData;
+			++It;
+		}
+		if (!(FoundMask & FoundStatus) && It.GetName() == ANSITEXTVIEW("Status"))
+		{
+			LoadFromCompactBinary(It, Status, EStatus::Ok);
+			FoundMask |= FoundStatus;
+			++It;
+		}
+		if (LastFoundMask == FoundMask)
+		{
+			++It;
+		}
+	}
+
+	if (OutResponse)
+	{
+		*OutResponse = {OutRequest.Name, OutRequest.Key, OutRequest.UserData, Status};
+	}
+
 	return bOk;
 }
 
@@ -283,24 +561,71 @@ void SaveToCompactBinary(FCbWriter& Writer, const FCacheGetValueRequest& Request
 	Writer.EndObject();
 }
 
-bool LoadFromCompactBinary(FCbFieldView Field, FCacheGetValueRequest& OutRequest)
+bool LoadFromCompactBinary(FCbFieldView Field, FCacheGetValueRequest& OutRequest, FCacheGetValueResponse* OutResponse, FLoadAttachmentFn LoadAttachment)
 {
 	bool bOk = Field.IsObject();
-	LoadFromCompactBinary(Field[ANSITEXTVIEW("Name")], OutRequest.Name);
-	bOk &= LoadFromCompactBinary(Field[ANSITEXTVIEW("Key")], OutRequest.Key);
-	LoadFromCompactBinary(Field[ANSITEXTVIEW("Policy")], OutRequest.Policy, ECachePolicy::Default);
-	LoadFromCompactBinary(Field[ANSITEXTVIEW("UserData")], OutRequest.UserData, 0);
-	return bOk;
-}
 
-bool LoadFromCompactBinary(FCbFieldView Field, FCacheGetValueResponse& OutResponse, FLoadAttachmentFn LoadAttachment)
-{
-	bool bOk = Field.IsObject();
-	LoadFromCompactBinary(Field[ANSITEXTVIEW("Name")], OutResponse.Name);
-	bOk &= LoadFromCompactBinary(Field[ANSITEXTVIEW("Key")], OutResponse.Key);
-	LoadFromCompactBinary(Field[ANSITEXTVIEW("Value")], OutResponse.Value, LoadAttachment);
-	LoadFromCompactBinary(Field[ANSITEXTVIEW("UserData")], OutResponse.UserData, 0);
-	LoadFromCompactBinary(Field[ANSITEXTVIEW("Status")], OutResponse.Status, EStatus::Ok);
+	FValue Value;
+	EStatus Status = EStatus::Ok;
+
+	constexpr static uint32 FoundName     = 1 << 0;
+	constexpr static uint32 FoundKey      = 1 << 1;
+	constexpr static uint32 FoundValue    = 1 << 2;
+	constexpr static uint32 FoundPolicy   = 1 << 3;
+	constexpr static uint32 FoundUserData = 1 << 4;
+	constexpr static uint32 FoundStatus   = 1 << 5;
+
+	uint32 FoundMask = 0;
+	for (FCbFieldViewIterator It = Field.CreateViewIterator(); It;)
+	{
+		const uint32 LastFoundMask = FoundMask;
+		if (!(FoundMask & FoundName) && It.GetName() == ANSITEXTVIEW("Name"))
+		{
+			LoadFromCompactBinary(It, OutRequest.Name);
+			FoundMask |= FoundName;
+			++It;
+		}
+		if (!(FoundMask & FoundKey) && It.GetName() == ANSITEXTVIEW("Key"))
+		{
+			bOk &= LoadFromCompactBinary(It, OutRequest.Key);
+			FoundMask |= FoundKey;
+			++It;
+		}
+		if (!(FoundMask & FoundValue) && It.GetName() == ANSITEXTVIEW("Value"))
+		{
+			LoadFromCompactBinary(It, Value, LoadAttachment);
+			FoundMask |= FoundValue;
+			++It;
+		}
+		if (!(FoundMask & FoundPolicy) && It.GetName() == ANSITEXTVIEW("Policy"))
+		{
+			LoadFromCompactBinary(It, OutRequest.Policy, ECachePolicy::Default);
+			FoundMask |= FoundPolicy;
+			++It;
+		}
+		if (!(FoundMask & FoundUserData) && It.GetName() == ANSITEXTVIEW("UserData"))
+		{
+			LoadFromCompactBinary(It, OutRequest.UserData, 0);
+			FoundMask |= FoundUserData;
+			++It;
+		}
+		if (!(FoundMask & FoundStatus) && It.GetName() == ANSITEXTVIEW("Status"))
+		{
+			LoadFromCompactBinary(It, Status, EStatus::Ok);
+			FoundMask |= FoundStatus;
+			++It;
+		}
+		if (LastFoundMask == FoundMask)
+		{
+			++It;
+		}
+	}
+
+	if (OutResponse)
+	{
+		*OutResponse = {OutRequest.Name, OutRequest.Key, Value, OutRequest.UserData, Status};
+	}
+
 	return bOk;
 }
 
@@ -361,49 +686,138 @@ void SaveToCompactBinary(FCbWriter& Writer, const FCacheGetChunkRequest& Request
 	Writer.EndObject();
 }
 
-bool LoadFromCompactBinary(FCbFieldView Field, FCacheGetChunkRequest& OutRequest)
+bool LoadFromCompactBinary(FCbFieldView Field, FCacheGetChunkRequest& OutRequest, FCacheGetChunkResponse* OutResponse, FLoadAttachmentFn LoadAttachment)
 {
 	bool bOk = Field.IsObject();
-	LoadFromCompactBinary(Field[ANSITEXTVIEW("Name")], OutRequest.Name);
-	bOk &= LoadFromCompactBinary(Field[ANSITEXTVIEW("Key")], OutRequest.Key);
-	LoadFromCompactBinary(Field[ANSITEXTVIEW("Id")], OutRequest.Id);
-	LoadFromCompactBinary(Field[ANSITEXTVIEW("RawOffset")], OutRequest.RawOffset, 0);
-	LoadFromCompactBinary(Field[ANSITEXTVIEW("RawSize")], OutRequest.RawSize, MAX_uint64);
-	LoadFromCompactBinary(Field[ANSITEXTVIEW("RawHash")], OutRequest.RawHash);
-	LoadFromCompactBinary(Field[ANSITEXTVIEW("Policy")], OutRequest.Policy, ECachePolicy::Default);
-	LoadFromCompactBinary(Field[ANSITEXTVIEW("UserData")], OutRequest.UserData, 0);
-	return bOk;
-}
 
-bool LoadFromCompactBinary(FCbFieldView Field, FCacheGetChunkResponse& OutResponse, FLoadAttachmentFn LoadAttachment)
-{
 	FValue Value;
-	bool bOk = Field.IsObject();
-	LoadFromCompactBinary(Field[ANSITEXTVIEW("Name")], OutResponse.Name);
-	bOk &= LoadFromCompactBinary(Field[ANSITEXTVIEW("Key")], OutResponse.Key);
-	LoadFromCompactBinary(Field[ANSITEXTVIEW("Id")], OutResponse.Id);
-	LoadFromCompactBinary(Field[ANSITEXTVIEW("RawOffset")], OutResponse.RawOffset, 0);
-	LoadFromCompactBinary(Field[ANSITEXTVIEW("UserData")], OutResponse.UserData, 0);
-	LoadFromCompactBinary(Field[ANSITEXTVIEW("Value")], Value);
-	LoadFromCompactBinary(Field[ANSITEXTVIEW("Status")], OutResponse.Status, EStatus::Ok);
-	OutResponse.RawHash = Value.GetRawHash();
-	OutResponse.RawSize = Value.GetRawSize();
-	if (LoadAttachment && !OutResponse.RawHash.IsZero())
-	{
-		const FIoHash RawHash = Field[ANSITEXTVIEW("FragmentHash")].AsHash(OutResponse.RawHash);
-		const FCbAttachment Attachment = LoadAttachment(RawHash);
-		if (FSharedBuffer Data = Attachment.AsBinary())
+	FIoHash FragmentHash;
+	EStatus Status = EStatus::Ok;
+
+	constexpr static uint32 FoundName         = 1 << 0;
+	constexpr static uint32 FoundKey          = 1 << 1;
+	constexpr static uint32 FoundId           = 1 << 2;
+	constexpr static uint32 FoundRawOffset    = 1 << 3;
+	constexpr static uint32 FoundRawSize      = 1 << 4;
+	constexpr static uint32 FoundRawHash      = 1 << 5;
+	constexpr static uint32 FoundValue        = 1 << 6;
+	constexpr static uint32 FoundFragmentHash = 1 << 7;
+	constexpr static uint32 FoundPolicy       = 1 << 8;
+	constexpr static uint32 FoundUserData     = 1 << 9;
+	constexpr static uint32 FoundStatus       = 1 << 10;
+
+	uint32 FoundMask = 0;
+	for (FCbFieldViewIterator It = Field.CreateViewIterator(); It;)
+	{
+		const uint32 LastFoundMask = FoundMask;
+		if (!(FoundMask & FoundName) && It.GetName() == ANSITEXTVIEW("Name"))
+		{
+			LoadFromCompactBinary(It, OutRequest.Name);
+			FoundMask |= FoundName;
+			++It;
+		}
+		if (!(FoundMask & FoundKey) && It.GetName() == ANSITEXTVIEW("Key"))
+		{
+			bOk &= LoadFromCompactBinary(It, OutRequest.Key);
+			FoundMask |= FoundKey;
+			++It;
+		}
+		if (!(FoundMask & FoundId) && It.GetName() == ANSITEXTVIEW("Id"))
 		{
-			OutResponse.RawData = MoveTemp(Data);
+			bOk &= LoadFromCompactBinary(It, OutRequest.Id);
+			FoundMask |= FoundId;
+			++It;
 		}
-		else if (const FCompressedBuffer& CompressedData = Attachment.AsCompressedBinary())
+		if (!(FoundMask & FoundRawOffset) && It.GetName() == ANSITEXTVIEW("RawOffset"))
 		{
-			if (CompressedData.GetRawHash() == OutResponse.RawHash && CompressedData.GetRawSize() == OutResponse.RawSize)
+			bOk &= LoadFromCompactBinary(It, OutRequest.RawOffset, 0);
+			FoundMask |= FoundRawOffset;
+			++It;
+		}
+		if (!(FoundMask & FoundRawSize) && It.GetName() == ANSITEXTVIEW("RawSize"))
+		{
+			bOk &= LoadFromCompactBinary(It, OutRequest.RawSize, MAX_uint64);
+			FoundMask |= FoundRawSize;
+			++It;
+		}
+		if (!(FoundMask & FoundRawHash) && It.GetName() == ANSITEXTVIEW("RawHash"))
+		{
+			bOk &= LoadFromCompactBinary(It, OutRequest.RawHash);
+			FoundMask |= FoundRawHash;
+			++It;
+		}
+		if (!(FoundMask & FoundValue) && It.GetName() == ANSITEXTVIEW("Value"))
+		{
+			if (OutResponse)
+			{
+				LoadFromCompactBinary(It, Value);
+			}
+			FoundMask |= FoundValue;
+			++It;
+		}
+		if (!(FoundMask & FoundFragmentHash) && It.GetName() == ANSITEXTVIEW("FragmentHash"))
+		{
+			if (OutResponse)
 			{
-				OutResponse.RawData = CompressedData.Decompress();
+				LoadFromCompactBinary(It, FragmentHash);
+			}
+			FoundMask |= FoundFragmentHash;
+			++It;
+		}
+		if (!(FoundMask & FoundPolicy) && It.GetName() == ANSITEXTVIEW("Policy"))
+		{
+			LoadFromCompactBinary(It, OutRequest.Policy, ECachePolicy::Default);
+			FoundMask |= FoundPolicy;
+			++It;
+		}
+		if (!(FoundMask & FoundUserData) && It.GetName() == ANSITEXTVIEW("UserData"))
+		{
+			LoadFromCompactBinary(It, OutRequest.UserData, 0);
+			FoundMask |= FoundUserData;
+			++It;
+		}
+		if (!(FoundMask & FoundStatus) && It.GetName() == ANSITEXTVIEW("Status"))
+		{
+			if (OutResponse)
+			{
+				LoadFromCompactBinary(It, Status, EStatus::Ok);
+			}
+			FoundMask |= FoundStatus;
+			++It;
+		}
+		if (LastFoundMask == FoundMask)
+		{
+			++It;
+		}
+	}
+
+	if (OutResponse)
+	{
+		OutResponse->Name = OutRequest.Name;
+		OutResponse->Key = OutRequest.Key;
+		OutResponse->Id = OutRequest.Id;
+		OutResponse->RawOffset = OutRequest.RawOffset;
+		OutResponse->UserData = OutRequest.UserData;
+		OutResponse->Status = Status;
+		OutResponse->RawHash = Value.GetRawHash();
+		OutResponse->RawSize = Value.GetRawSize();
+		if (LoadAttachment && !OutResponse->RawHash.IsZero())
+		{
+			const FCbAttachment Attachment = LoadAttachment(!FragmentHash.IsZero() ? FragmentHash : OutResponse->RawHash);
+			if (FSharedBuffer Data = Attachment.AsBinary())
+			{
+				OutResponse->RawData = MoveTemp(Data);
+			}
+			else if (const FCompressedBuffer& CompressedData = Attachment.AsCompressedBinary())
+			{
+				if (CompressedData.GetRawHash() == OutResponse->RawHash && CompressedData.GetRawSize() == OutResponse->RawSize)
+				{
+					OutResponse->RawData = CompressedData.Decompress();
+				}
 			}
 		}
 	}
+
 	return bOk;
 }
 
@@ -438,8 +852,32 @@ FCbWriter& operator<<(FCbWriter& Writer, const FCacheKey& Key)
 bool LoadFromCompactBinary(const FCbFieldView Field, FCacheKey& OutKey)
 {
 	bool bOk = Field.IsObject();
-	bOk &= LoadFromCompactBinary(Field[ANSITEXTVIEW("Bucket")], OutKey.Bucket);
-	bOk &= LoadFromCompactBinary(Field[ANSITEXTVIEW("Hash")], OutKey.Hash);
+
+	constexpr static uint32 FoundBucket = 1 << 0;
+	constexpr static uint32 FoundHash   = 1 << 1;
+
+	uint32 FoundMask = 0;
+	for (FCbFieldViewIterator It = Field.CreateViewIterator(); It;)
+	{
+		const uint32 LastFoundMask = FoundMask;
+		if (!(FoundMask & FoundBucket) && It.GetName() == ANSITEXTVIEW("Bucket"))
+		{
+			bOk &= LoadFromCompactBinary(It, OutKey.Bucket);
+			FoundMask |= FoundBucket;
+			++It;
+		}
+		if (!(FoundMask & FoundHash) && It.GetName() == ANSITEXTVIEW("Hash"))
+		{
+			bOk &= LoadFromCompactBinary(It, OutKey.Hash);
+			FoundMask |= FoundHash;
+			++It;
+		}
+		if (LastFoundMask == FoundMask)
+		{
+			++It;
+		}
+	}
+
 	return bOk;
 }
 
@@ -491,8 +929,32 @@ FCbWriter& operator<<(FCbWriter& Writer, const FCacheValuePolicy& Policy)
 bool LoadFromCompactBinary(FCbFieldView Field, FCacheValuePolicy& OutPolicy)
 {
 	bool bOk = Field.IsObject();
-	bOk &= LoadFromCompactBinary(Field[ANSITEXTVIEW("Id")], OutPolicy.Id);
-	bOk &= LoadFromCompactBinary(Field[ANSITEXTVIEW("Policy")], OutPolicy.Policy, ECachePolicy::Default);
+
+	constexpr static uint32 FoundId     = 1 << 0;
+	constexpr static uint32 FoundPolicy = 1 << 1;
+
+	uint32 FoundMask = 0;
+	for (FCbFieldViewIterator It = Field.CreateViewIterator(); It;)
+	{
+		const uint32 LastFoundMask = FoundMask;
+		if (!(FoundMask & FoundId) && It.GetName() == ANSITEXTVIEW("Id"))
+		{
+			bOk &= LoadFromCompactBinary(It, OutPolicy.Id);
+			FoundMask |= FoundId;
+			++It;
+		}
+		if (!(FoundMask & FoundPolicy) && It.GetName() == ANSITEXTVIEW("Policy"))
+		{
+			bOk &= LoadFromCompactBinary(It, OutPolicy.Policy, ECachePolicy::Default);
+			FoundMask |= FoundPolicy;
+			++It;
+		}
+		if (LastFoundMask == FoundMask)
+		{
+			++It;
+		}
+	}
+
 	return bOk;
 }
 
@@ -562,30 +1024,44 @@ void SaveToCompactBinary(FCbWriter& Writer, const FCacheRecord& Record, EMetaSav
 bool LoadFromCompactBinary(FCbFieldView Field, FOptionalCacheRecord& OutRecord, FLoadAttachmentFn LoadAttachment)
 {
 	FCacheKey Key;
+	FCbObject Meta;
+	TArray<FValueWithId> Values;
+
 	if (!LoadFromCompactBinary(Field[ANSITEXTVIEW("Key")], Key))
 	{
 		return false;
 	}
 
-	FCacheRecordBuilder Builder(Key);
-
-	FCbObject Meta;
-	if (LoadMetaFromCompactBinary(Field, Meta, LoadAttachment))
+	if (FCbValue MetaValue = Field[ANSITEXTVIEW("Meta")])
 	{
-		Builder.SetMeta(MoveTemp(Meta));
+		LoadMetaFromObject(MetaValue, Meta);
+	}
+	else if (LoadAttachment)
+	{
+		FCbValue MetaHashValue = Field[ANSITEXTVIEW("MetaHash")];
+		FCbValue MetaSizeValue = Field[ANSITEXTVIEW("MetaSize")];
+		if (MetaHashValue && MetaSizeValue)
+		{
+			LoadMetaFromHashAndSize(MetaHashValue, MetaSizeValue, Meta, LoadAttachment);
+		}
 	}
 
-	for (FCbFieldView ValueField : Field[ANSITEXTVIEW("Values")])
+	FCbArrayView ValuesArray = Field[ANSITEXTVIEW("Values")].AsArrayView();
+	Values.Reserve(int32(ValuesArray.Num()));
+	for (FCbFieldView ValueField : ValuesArray)
 	{
-		FValueWithId Value;
+		FValueWithId& Value = Values.AddDefaulted_GetRef();
 		if (!LoadFromCompactBinary(ValueField, Value, LoadAttachment))
 		{
 			return false;
 		}
-		Builder.AddValue(Value);
+	}
+	if (UNLIKELY(!Algo::IsSorted(Values)))
+	{
+		Algo::Sort(Values);
 	}
 
-	OutRecord = Builder.Build();
+	OutRecord = FCacheRecord::CreateByMove(Key, MoveTemp(Meta), MoveTemp(Values));
 	return true;
 }
 
@@ -645,12 +1121,37 @@ void SaveToCompactBinary(FCbWriter& Writer, const FValue& Value, FSaveAttachment
 
 bool LoadFromCompactBinary(FCbFieldView Field, FValue& OutValue, FLoadAttachmentFn LoadAttachment)
 {
+	bool bOk = Field.IsObject();
+
 	FIoHash RawHash;
 	uint64 RawSize = 0;
-	bool bOk = Field.IsObject();
-	bOk &= LoadFromCompactBinary(Field[ANSITEXTVIEW("RawHash")], RawHash);
-	bOk &= LoadFromCompactBinary(Field[ANSITEXTVIEW("RawSize")], RawSize);
-	if (LoadAttachment)
+
+	constexpr static uint32 FoundRawHash = 1 << 0;
+	constexpr static uint32 FoundRawSize = 1 << 1;
+
+	uint32 FoundMask = 0;
+	for (FCbFieldViewIterator It = Field.CreateViewIterator(); It;)
+	{
+		const uint32 LastFoundMask = FoundMask;
+		if (!(FoundMask & FoundRawHash) && It.GetName() == ANSITEXTVIEW("RawHash"))
+		{
+			bOk &= LoadFromCompactBinary(It, RawHash);
+			FoundMask |= FoundRawHash;
+			++It;
+		}
+		if (!(FoundMask & FoundRawSize) && It.GetName() == ANSITEXTVIEW("RawSize"))
+		{
+			bOk &= LoadFromCompactBinary(It, RawSize);
+			FoundMask |= FoundRawSize;
+			++It;
+		}
+		if (LastFoundMask == FoundMask)
+		{
+			++It;
+		}
+	}
+
+	if (bOk && LoadAttachment)
 	{
 		const FCbAttachment Attachment = LoadAttachment(RawHash);
 		const FCompressedBuffer& Compressed = Attachment.AsCompressedBinary();
@@ -660,6 +1161,7 @@ bool LoadFromCompactBinary(FCbFieldView Field, FValue& OutValue, FLoadAttachment
 			return bOk;
 		}
 	}
+
 	OutValue = FValue(RawHash, RawSize);
 	return bOk;
 }
@@ -685,12 +1187,62 @@ void SaveToCompactBinary(FCbWriter& Writer, const FValueWithId& Value, FSaveAtta
 
 bool LoadFromCompactBinary(FCbFieldView Field, FValueWithId& OutValue, FLoadAttachmentFn LoadAttachment)
 {
-	FValueId Id;
-	FValue Value;
 	bool bOk = Field.IsObject();
-	bOk &= LoadFromCompactBinary(Field[ANSITEXTVIEW("Id")], Id);
-	bOk &= LoadFromCompactBinary(Field, Value, LoadAttachment);
-	OutValue = FValueWithId(Id, MoveTemp(Value));
+
+	FValueId Id;
+	FIoHash RawHash;
+	uint64 RawSize = 0;
+
+	constexpr static uint32 FoundId      = 1 << 0;
+	constexpr static uint32 FoundRawHash = 1 << 1;
+	constexpr static uint32 FoundRawSize = 1 << 2;
+
+	uint32 FoundMask = 0;
+	for (FCbFieldViewIterator It = Field.CreateViewIterator(); It;)
+	{
+		const uint32 LastFoundMask = FoundMask;
+		if (!(FoundMask & FoundId) && It.GetName() == ANSITEXTVIEW("Id"))
+		{
+			bOk &= LoadFromCompactBinary(It, Id);
+			FoundMask |= FoundId;
+			++It;
+		}
+		if (!(FoundMask & FoundRawHash) && It.GetName() == ANSITEXTVIEW("RawHash"))
+		{
+			bOk &= LoadFromCompactBinary(It, RawHash);
+			FoundMask |= FoundRawHash;
+			++It;
+		}
+		if (!(FoundMask & FoundRawSize) && It.GetName() == ANSITEXTVIEW("RawSize"))
+		{
+			bOk &= LoadFromCompactBinary(It, RawSize);
+			FoundMask |= FoundRawSize;
+			++It;
+		}
+		if (LastFoundMask == FoundMask)
+		{
+			++It;
+		}
+	}
+
+	if (Id.IsNull())
+	{
+		// An invalid Id asserts in FValueWithId().
+		return false;
+	}
+
+	if (bOk && LoadAttachment)
+	{
+		const FCbAttachment Attachment = LoadAttachment(RawHash);
+		const FCompressedBuffer& Compressed = Attachment.AsCompressedBinary();
+		if (Compressed.GetRawHash() == RawHash && Compressed.GetRawSize() == RawSize)
+		{
+			OutValue = FValueWithId(Id, Compressed);
+			return bOk;
+		}
+	}
+
+	OutValue = FValueWithId(Id, RawHash, RawSize);
 	return bOk;
 }
 
diff --git a/Engine/Source/Developer/DerivedDataCache/Public/DerivedDataCacheRecord.h b/Engine/Source/Developer/DerivedDataCache/Public/DerivedDataCacheRecord.h
--- a/Engine/Source/Developer/DerivedDataCache/Public/DerivedDataCacheRecord.h
+++ b/Engine/Source/Developer/DerivedDataCache/Public/DerivedDataCacheRecord.h
@@ -93,6 +93,9 @@ public:
 	UE_API static FOptionalCacheRecord Load(const FCbPackage& Package);
 	UE_API static FOptionalCacheRecord Load(const FCbPackage& Attachments, const FCbObject& Object);
 
+	/** Internal. Create a record from a key, optional metadata, and sorted values. Moves metadata and values. */
+	UE_INTERNAL static FCacheRecord CreateByMove(const FCacheKey& Key, FCbObject&& Meta, TArray<FValueWithId>&& Values);
+
 private:
 	friend class FOptionalCacheRecord;
 	friend FCacheRecord Private::CreateCacheRecord(Private::ICacheRecordInternal* Record);
diff --git a/Engine/Source/Developer/DerivedDataCache/Public/DerivedDataSerialization.h b/Engine/Source/Developer/DerivedDataCache/Public/DerivedDataSerialization.h
--- a/Engine/Source/Developer/DerivedDataCache/Public/DerivedDataSerialization.h
+++ b/Engine/Source/Developer/DerivedDataCache/Public/DerivedDataSerialization.h
@@ -67,24 +67,19 @@ UE_DEPRECATED(5.8, "Call SaveToCompactBinary(Writer, Request).")
 UE_API FCbWriter& operator<<(FCbWriter& Writer, const FCacheGetChunkRequest& Request);
 
 UE_API void SaveToCompactBinary(FCbWriter& Writer, const FCachePutRequest& Request, const FCachePutResponse* Response = nullptr, FSaveAttachmentFn SaveAttachment = nullptr);
-UE_API bool LoadFromCompactBinary(FCbFieldView Field, TOptional<FCachePutRequest>& OutRequest, FLoadAttachmentFn LoadAttachment = nullptr);
-UE_API bool LoadFromCompactBinary(FCbFieldView Field, FCachePutResponse& OutResponse);
+UE_API bool LoadFromCompactBinary(FCbFieldView Field, TOptional<FCachePutRequest>& OutRequest, FCachePutResponse* OutResponse = nullptr, FLoadAttachmentFn LoadAttachment = nullptr);
 
 UE_API void SaveToCompactBinary(FCbWriter& Writer, const FCacheGetRequest& Request, const FCacheGetResponse* Response = nullptr, FSaveAttachmentFn SaveAttachment = nullptr);
-UE_API bool LoadFromCompactBinary(FCbFieldView Field, FCacheGetRequest& OutRequest);
-UE_API bool LoadFromCompactBinary(FCbFieldView Field, TOptional<FCacheGetResponse>& OutResponse, FLoadAttachmentFn LoadAttachment = nullptr);
+UE_API bool LoadFromCompactBinary(FCbFieldView Field, FCacheGetRequest& OutRequest, TOptional<FCacheGetResponse>* OutResponse = nullptr, FLoadAttachmentFn LoadAttachment = nullptr);
 
 UE_API void SaveToCompactBinary(FCbWriter& Writer, const FCachePutValueRequest& Request, const FCachePutValueResponse* Response = nullptr, FSaveAttachmentFn SaveAttachment = nullptr);
-UE_API bool LoadFromCompactBinary(FCbFieldView Field, FCachePutValueRequest& OutRequest, FLoadAttachmentFn LoadAttachment = nullptr);
-UE_API bool LoadFromCompactBinary(FCbFieldView Field, FCachePutValueResponse& OutResponse);
+UE_API bool LoadFromCompactBinary(FCbFieldView Field, FCachePutValueRequest& OutRequest, FCachePutValueResponse* OutResponse = nullptr, FLoadAttachmentFn LoadAttachment = nullptr);
 
 UE_API void SaveToCompactBinary(FCbWriter& Writer, const FCacheGetValueRequest& Request, const FCacheGetValueResponse* Response = nullptr, FSaveAttachmentFn SaveAttachment = nullptr);
-UE_API bool LoadFromCompactBinary(FCbFieldView Field, FCacheGetValueRequest& OutRequest);
-UE_API bool LoadFromCompactBinary(FCbFieldView Field, FCacheGetValueResponse& OutResponse, FLoadAttachmentFn LoadAttachment = nullptr);
+UE_API bool LoadFromCompactBinary(FCbFieldView Field, FCacheGetValueRequest& OutRequest, FCacheGetValueResponse* OutResponse = nullptr, FLoadAttachmentFn LoadAttachment = nullptr);
 
 UE_API void SaveToCompactBinary(FCbWriter& Writer, const FCacheGetChunkRequest& Request, const FCacheGetChunkResponse* Response = nullptr, FSaveAttachmentFn SaveAttachment = nullptr);
-UE_API bool LoadFromCompactBinary(FCbFieldView Field, FCacheGetChunkRequest& OutRequest);
-UE_API bool LoadFromCompactBinary(FCbFieldView Field, FCacheGetChunkResponse& OutResponse, FLoadAttachmentFn LoadAttachment = nullptr);
+UE_API bool LoadFromCompactBinary(FCbFieldView Field, FCacheGetChunkRequest& OutRequest, FCacheGetChunkResponse* OutResponse = nullptr, FLoadAttachmentFn LoadAttachment = nullptr);
 
 // Types from DerivedDataCacheKey.h
 
diff --git a/Engine/Source/Developer/DerivedDataCache/Tests/DerivedDataSerializationTest.cpp b/Engine/Source/Developer/DerivedDataCache/Tests/DerivedDataSerializationTest.cpp
--- a/Engine/Source/Developer/DerivedDataCache/Tests/DerivedDataSerializationTest.cpp
+++ b/Engine/Source/Developer/DerivedDataCache/Tests/DerivedDataSerializationTest.cpp
@@ -6,10 +6,13 @@
 
 #include "DerivedDataCache.h"
 #include "Misc/Optional.h"
+#include "Misc/ReverseIterate.h"
 #include "Serialization/CompactBinaryPackage.h"
 #include "Serialization/CompactBinaryWriter.h"
 #include "TestHarness.h"
 
+#include <catch2/generators/catch_generators.hpp>
+
 namespace UE::DerivedData
 {
 
@@ -69,16 +72,36 @@ TEST_CASE("DerivedData::Serialization::CompactBinary", "[DerivedData]")
 		CHECK((Hash == Value.GetRawHash() || Hash == MetaHash || Hash == EmptyBufferHash));
 	};
 
+	const auto ReverseObjectFields = [&](const FCbField& ObjectField) -> FCbField
+	{
+		REQUIRE(ObjectField.IsObject());
+		TArray<FCbFieldView> Fields;
+		for (FCbFieldView Field : ObjectField)
+		{
+			Fields.Add(Field);
+		}
+		FCbWriter Writer;
+		Writer.BeginObject();
+		for (FCbFieldView Field : ReverseIterate(Fields))
+		{
+			Writer.AddField(Field.GetName(), Field);
+		}
+		Writer.EndObject();
+		return Writer.Save();
+	};
+
 	SECTION("CachePutRequest")
 	{
 		const FCachePutRequest Request{Name, Record, RecordPolicy, UserData};
 
-		const FCbField Field = [&]
+		const FCbField OriginalField = [&]
 		{
 			TCbWriter<512> Writer;
 			SaveToCompactBinary(Writer, Request, nullptr, SaveAttachment);
 			return Writer.Save();
 		}();
+		const FCbField ReversedField = ReverseObjectFields(OriginalField);
+		const FCbField Field = GENERATE_COPY(OriginalField, ReversedField);
 
 		{
 			TOptional<FCachePutRequest> LoadedRequest;
@@ -96,7 +119,7 @@ TEST_CASE("DerivedData::Serialization::CompactBinary", "[DerivedData]")
 
 		{
 			TOptional<FCachePutRequest> LoadedRequest;
-			CHECK(LoadFromCompactBinary(Field, LoadedRequest, LoadAttachment));
+			CHECK(LoadFromCompactBinary(Field, LoadedRequest, nullptr, LoadAttachment));
 			CHECKED_IF(LoadedRequest)
 			{
 				CHECK(LoadedRequest->Name == Request.Name);
@@ -114,15 +137,18 @@ TEST_CASE("DerivedData::Serialization::CompactBinary", "[DerivedData]")
 		const FCachePutRequest Request{Name, Record, RecordPolicy, UserData};
 		const FCachePutResponse Response{Name, Key, UserData, EStatus::Canceled};
 
-		const FCbField Field = [&]
+		const FCbField OriginalField = [&]
 		{
 			TCbWriter<512> Writer;
 			SaveToCompactBinary(Writer, Request, &Response);
 			return Writer.Save();
 		}();
+		const FCbField ReversedField = ReverseObjectFields(OriginalField);
+		const FCbField Field = GENERATE_COPY(OriginalField, ReversedField);
 
+		TOptional<FCachePutRequest> LoadedRequest;
 		FCachePutResponse LoadedResponse;
-		CHECK(LoadFromCompactBinary(Field, LoadedResponse));
+		CHECK(LoadFromCompactBinary(Field, LoadedRequest, &LoadedResponse));
 		CHECK(LoadedResponse.Name == Response.Name);
 		CHECK(LoadedResponse.Key == Response.Key);
 		CHECK(LoadedResponse.UserData == Response.UserData);
@@ -133,12 +159,14 @@ TEST_CASE("DerivedData::Serialization::CompactBinary", "[DerivedData]")
 	{
 		const FCacheGetRequest Request{Name, Key, RecordPolicy, UserData};
 
-		const FCbField Field = [&]
+		const FCbField OriginalField = [&]
 		{
 			TCbWriter<512> Writer;
 			SaveToCompactBinary(Writer, Request);
 			return Writer.Save();
 		}();
+		const FCbField ReversedField = ReverseObjectFields(OriginalField);
+		const FCbField Field = GENERATE_COPY(OriginalField, ReversedField);
 
 		FCacheGetRequest LoadedRequest;
 		CHECK(LoadFromCompactBinary(Field, LoadedRequest));
@@ -155,16 +183,19 @@ TEST_CASE("DerivedData::Serialization::CompactBinary", "[DerivedData]")
 		const FCacheGetRequest Request{Name, Key, RecordPolicy, UserData};
 		const FCacheGetResponse Response{Name, Record, UserData, EStatus::Canceled};
 
-		const FCbField Field = [&]
+		const FCbField OriginalField = [&]
 		{
 			TCbWriter<512> Writer;
 			SaveToCompactBinary(Writer, Request, &Response, SaveAttachment);
 			return Writer.Save();
 		}();
+		const FCbField ReversedField = ReverseObjectFields(OriginalField);
+		const FCbField Field = GENERATE_COPY(OriginalField, ReversedField);
 
 		{
+			FCacheGetRequest LoadedRequest;
 			TOptional<FCacheGetResponse> LoadedResponse;
-			CHECK(LoadFromCompactBinary(Field, LoadedResponse));
+			CHECK(LoadFromCompactBinary(Field, LoadedRequest, &LoadedResponse));
 			CHECKED_IF(LoadedResponse)
 			{
 				CHECK(LoadedResponse->Name == Response.Name);
@@ -178,8 +209,9 @@ TEST_CASE("DerivedData::Serialization::CompactBinary", "[DerivedData]")
 		}
 
 		{
+			FCacheGetRequest LoadedRequest;
 			TOptional<FCacheGetResponse> LoadedResponse;
-			CHECK(LoadFromCompactBinary(Field, LoadedResponse, LoadAttachment));
+			CHECK(LoadFromCompactBinary(Field, LoadedRequest, &LoadedResponse, LoadAttachment));
 			CHECKED_IF(LoadedResponse)
 			{
 				CHECK(LoadedResponse->Name == Response.Name);
@@ -197,13 +229,14 @@ TEST_CASE("DerivedData::Serialization::CompactBinary", "[DerivedData]")
 	{
 		const FCachePutValueRequest Request{Name, Key, Value, ValuePolicy, UserData};
 
-		const FCbField Field = [&]
+		const FCbField OriginalField = [&]
 		{
 			TCbWriter<512> Writer;
 			SaveToCompactBinary(Writer, Request, nullptr, SaveAttachment);
 			return Writer.Save();
 		}();
-
+		const FCbField ReversedField = ReverseObjectFields(OriginalField);
+		const FCbField Field = GENERATE_COPY(OriginalField, ReversedField);
 		{
 			FCachePutValueRequest LoadedRequest;
 			CHECK(LoadFromCompactBinary(Field, LoadedRequest));
@@ -216,7 +249,7 @@ TEST_CASE("DerivedData::Serialization::CompactBinary", "[DerivedData]")
 
 		{
 			FCachePutValueRequest LoadedRequest;
-			CHECK(LoadFromCompactBinary(Field, LoadedRequest, LoadAttachment));
+			CHECK(LoadFromCompactBinary(Field, LoadedRequest, nullptr, LoadAttachment));
 			CHECK(LoadedRequest.Name == Request.Name);
 			CHECK(LoadedRequest.Key == Request.Key);
 			CHECK(LoadedRequest.Value == Request.Value);
@@ -230,15 +263,18 @@ TEST_CASE("DerivedData::Serialization::CompactBinary", "[DerivedData]")
 		const FCachePutValueRequest Request{Name, Key, Value, ValuePolicy, UserData};
 		const FCachePutValueResponse Response{Name, Key, UserData, EStatus::Canceled};
 
-		const FCbField Field = [&]
+		const FCbField OriginalField = [&]
 		{
 			TCbWriter<512> Writer;
 			SaveToCompactBinary(Writer, Request, &Response);
 			return Writer.Save();
 		}();
+		const FCbField ReversedField = ReverseObjectFields(OriginalField);
+		const FCbField Field = GENERATE_COPY(OriginalField, ReversedField);
 
+		FCachePutValueRequest LoadedRequest;
 		FCachePutValueResponse LoadedResponse;
-		CHECK(LoadFromCompactBinary(Field, LoadedResponse));
+		CHECK(LoadFromCompactBinary(Field, LoadedRequest, &LoadedResponse));
 		CHECK(LoadedResponse.Name == Response.Name);
 		CHECK(LoadedResponse.Key == Response.Key);
 		CHECK(LoadedResponse.UserData == Response.UserData);
@@ -249,12 +285,14 @@ TEST_CASE("DerivedData::Serialization::CompactBinary", "[DerivedData]")
 	{
 		const FCacheGetValueRequest Request{Name, Key, ValuePolicy, UserData};
 
-		const FCbField Field = [&]
+		const FCbField OriginalField = [&]
 		{
 			TCbWriter<512> Writer;
 			SaveToCompactBinary(Writer, Request);
 			return Writer.Save();
 		}();
+		const FCbField ReversedField = ReverseObjectFields(OriginalField);
+		const FCbField Field = GENERATE_COPY(OriginalField, ReversedField);
 
 		FCacheGetValueRequest LoadedRequest;
 		CHECK(LoadFromCompactBinary(Field, LoadedRequest));
@@ -269,16 +307,19 @@ TEST_CASE("DerivedData::Serialization::CompactBinary", "[DerivedData]")
 		const FCacheGetValueRequest Request{Name, Key, ValuePolicy, UserData};
 		const FCacheGetValueResponse Response{Name, Key, Value, UserData, EStatus::Canceled};
 
-		const FCbField Field = [&]
+		const FCbField OriginalField = [&]
 		{
 			TCbWriter<512> Writer;
 			SaveToCompactBinary(Writer, Request, &Response, SaveAttachment);
 			return Writer.Save();
 		}();
+		const FCbField ReversedField = ReverseObjectFields(OriginalField);
+		const FCbField Field = GENERATE_COPY(OriginalField, ReversedField);
 
 		{
+			FCacheGetValueRequest LoadedRequest;
 			FCacheGetValueResponse LoadedResponse;
-			CHECK(LoadFromCompactBinary(Field, LoadedResponse));
+			CHECK(LoadFromCompactBinary(Field, LoadedRequest, &LoadedResponse));
 			CHECK(LoadedResponse.Name == Response.Name);
 			CHECK(LoadedResponse.Key == Response.Key);
 			CHECK(LoadedResponse.Value == Response.Value);
@@ -288,8 +329,9 @@ TEST_CASE("DerivedData::Serialization::CompactBinary", "[DerivedData]")
 		}
 
 		{
+			FCacheGetValueRequest LoadedRequest;
 			FCacheGetValueResponse LoadedResponse;
-			CHECK(LoadFromCompactBinary(Field, LoadedResponse, LoadAttachment));
+			CHECK(LoadFromCompactBinary(Field, LoadedRequest, &LoadedResponse, LoadAttachment));
 			CHECK(LoadedResponse.Name == Name);
 			CHECK(LoadedResponse.Key == Key);
 			CHECK(LoadedResponse.Value == Value);
@@ -303,12 +345,14 @@ TEST_CASE("DerivedData::Serialization::CompactBinary", "[DerivedData]")
 	{
 		const FCacheGetChunkRequest Request{Name, Key, ValueId, 1024, 512, Value.GetRawHash(), ValuePolicy, UserData};
 
-		const FCbField Field = [&]
+		const FCbField OriginalField = [&]
 		{
 			TCbWriter<512> Writer;
 			SaveToCompactBinary(Writer, Request);
 			return Writer.Save();
 		}();
+		const FCbField ReversedField = ReverseObjectFields(OriginalField);
+		const FCbField Field = GENERATE_COPY(OriginalField, ReversedField);
 
 		FCacheGetChunkRequest LoadedRequest;
 		CHECK(LoadFromCompactBinary(Field, LoadedRequest));
@@ -327,15 +371,18 @@ TEST_CASE("DerivedData::Serialization::CompactBinary", "[DerivedData]")
 		const FCacheGetChunkRequest Request{Name, Key, ValueId, 1024, 512, Value.GetRawHash(), ValuePolicy, UserData};
 		const FCacheGetChunkResponse Response{Name, Key, ValueId, 1024, 256, Value.GetRawHash(), EmptyBuffer, UserData, EStatus::Canceled};
 
-		const FCbField Field = [&]
+		const FCbField OriginalField = [&]
 		{
 			TCbWriter<512> Writer;
 			SaveToCompactBinary(Writer, Request, &Response, SaveAttachment);
 			return Writer.Save();
 		}();
+		const FCbField ReversedField = ReverseObjectFields(OriginalField);
+		const FCbField Field = GENERATE_COPY(OriginalField, ReversedField);
 
+		FCacheGetChunkRequest LoadedRequest;
 		FCacheGetChunkResponse LoadedResponse;
-		CHECK(LoadFromCompactBinary(Field, LoadedResponse, LoadAttachment));
+		CHECK(LoadFromCompactBinary(Field, LoadedRequest, &LoadedResponse, LoadAttachment));
 		CHECK(LoadedResponse.Name == Response.Name);
 		CHECK(LoadedResponse.Key == Response.Key);
 		CHECK(LoadedResponse.Id == Response.Id);
diff --git a/Engine/Source/Runtime/Core/Private/Serialization/CompactBinarySerialization.cpp b/Engine/Source/Runtime/Core/Private/Serialization/CompactBinarySerialization.cpp
--- a/Engine/Source/Runtime/Core/Private/Serialization/CompactBinarySerialization.cpp
+++ b/Engine/Source/Runtime/Core/Private/Serialization/CompactBinarySerialization.cpp
@@ -11,6 +11,7 @@
 #include "CoreGlobals.h"
 #include "HAL/PlatformString.h"
 #include "Logging/LogMacros.h"
+#include "Logging/StructuredLog.h"
 #include "Memory/MemoryView.h"
 #include "Misc/AsciiSet.h"
 #include "Misc/AssertionMacros.h"
@@ -144,9 +145,9 @@ template void AppendQuotedJsonString<WIDECHAR>(FWideStringBuilderBase& Builder,
 namespace UE::Serialization::Private
 {
 
-void LogFieldTooLargeForArrayWarning(uint64 FieldLength)
+void LogFieldTooLargeForArrayWarning(uint64 ArraySize)
 {
-	UE_LOG(LogSerialization, Warning, TEXT("Attempting to load a FCbFieldView that has too many entries (%" UINT64_FMT ") to fit in a TArray."), FieldLength);
+	UE_LOGFMT(LogSerialization, Warning, "Attempting to load a FCbFieldView that has too many elements ({ArraySize}) to fit in TArray.", ArraySize);
 }
 
 } // namespace UE::Serialization::Private
@@ -303,6 +304,11 @@ bool TryMeasureCompactBinary(FMemoryView View, ECbFieldType& OutType, uint64& Ou
 ///////////////////////////////////////////////////////////////////////////////////////////////////////////////////////
 
 FCbField LoadCompactBinary(FArchive& Ar, FCbBufferAllocator Allocator)
+{
+	return LoadAndValidateCompactBinary(Ar, Allocator, ECbValidateMode::Default);
+}
+
+FCbField LoadAndValidateCompactBinary(FArchive& Ar, FCbBufferAllocator Allocator, ECbValidateMode ValidateMode)
 {
 	TArray<uint8, TInlineAllocator<64>> HeaderBytes;
 	ECbFieldType FieldType;
@@ -323,14 +329,14 @@ FCbField LoadCompactBinary(FArchive& Ar, FCbBufferAllocator Allocator)
 				checkf(Buffer.GetSize() == FieldSize, TEXT("Allocator returned a buffer of size %" UINT64_FMT " bytes "
 					"when %" UINT64_FMT " bytes were requested."), Buffer.GetSize(), FieldSize);
 
-				FMutableMemoryView View = Buffer.GetView().CopyFrom(MakeMemoryView(HeaderBytes));
-				if (!View.IsEmpty())
+				FMutableMemoryView TailView = Buffer.GetView().CopyFrom(MakeMemoryView(HeaderBytes));
+				if (!TailView.IsEmpty())
 				{
 					// Read the remainder of the field.
-					Ar.Serialize(View.GetData(), static_cast<int64>(View.GetSize()));
+					Ar.Serialize(TailView.GetData(), static_cast<int64>(TailView.GetSize()));
 				}
 
-				if (!Ar.IsError() && ValidateCompactBinary(Buffer, ECbValidateMode::Default) == ECbValidateError::None)
+				if (!Ar.IsError() && (ValidateMode == ECbValidateMode::None || ValidateCompactBinary(Buffer, ValidateMode) == ECbValidateError::None))
 				{
 					return FCbField(Buffer.MoveToShared());
 				}
@@ -343,6 +349,52 @@ FCbField LoadCompactBinary(FArchive& Ar, FCbBufferAllocator Allocator)
 	return FCbField();
 }
 
+FCbFieldView LoadCompactBinaryView(FArchive& Ar, FCbMemoryAllocator Allocator)
+{
+	return LoadAndValidateCompactBinaryView(Ar, Allocator, ECbValidateMode::Default);
+}
+
+FCbFieldView LoadAndValidateCompactBinaryView(FArchive& Ar, FCbMemoryAllocator Allocator, ECbValidateMode ValidateMode)
+{
+	TArray<uint8, TInlineAllocator<64>> HeaderBytes;
+	ECbFieldType FieldType;
+	uint64 FieldSize = 1;
+
+	for (const int64 StartPos = Ar.Tell(); !Ar.IsError() && FieldSize > 0;)
+	{
+		// Read in small increments until the total field size is known, to avoid reading too far.
+		const int32 ReadSize = int32(FieldSize - HeaderBytes.Num());
+		const int32 ReadOffset = HeaderBytes.AddUninitialized(ReadSize);
+		Ar.Serialize(HeaderBytes.GetData() + ReadOffset, ReadSize);
+
+		if (!Ar.IsError() && TryMeasureCompactBinary(MakeMemoryView(HeaderBytes), FieldType, FieldSize))
+		{
+			if (FieldSize <= uint64(Ar.TotalSize() - StartPos))
+			{
+				FMutableMemoryView View = Allocator(FieldSize);
+				checkf(View.GetSize() == FieldSize, TEXT("Allocator returned a buffer of size %" UINT64_FMT " bytes "
+					"when %" UINT64_FMT " bytes were requested."), View.GetSize(), FieldSize);
+
+				FMutableMemoryView TailView = View.CopyFrom(MakeMemoryView(HeaderBytes));
+				if (!TailView.IsEmpty())
+				{
+					// Read the remainder of the field.
+					Ar.Serialize(TailView.GetData(), static_cast<int64>(TailView.GetSize()));
+				}
+
+				if (!Ar.IsError() && (ValidateMode == ECbValidateMode::None || ValidateCompactBinary(View, ValidateMode) == ECbValidateError::None))
+				{
+					return FCbFieldView(View.GetData());
+				}
+			}
+			break;
+		}
+	}
+
+	Ar.SetError();
+	return FCbFieldView();
+}
+
 ///////////////////////////////////////////////////////////////////////////////////////////////////////////////////////
 
 void SaveCompactBinary(FArchive& Ar, const FCbFieldView& Field)
diff --git a/Engine/Source/Runtime/Core/Private/Serialization/CompactBinaryWriter.cpp b/Engine/Source/Runtime/Core/Private/Serialization/CompactBinaryWriter.cpp
--- a/Engine/Source/Runtime/Core/Private/Serialization/CompactBinaryWriter.cpp
+++ b/Engine/Source/Runtime/Core/Private/Serialization/CompactBinaryWriter.cpp
@@ -377,11 +377,23 @@ void FCbWriter::AddBinary(const void* const Value, const uint64 Size)
 	EndField(ECbFieldType::Binary);
 }
 
+void FCbWriter::AddBinary(FUtf8StringView Name, const FSharedBuffer& Value)
+{
+	SetName(Name);
+	AddBinary(Value);
+}
+
 void FCbWriter::AddBinary(const FSharedBuffer& Buffer)
 {
 	AddBinary(Buffer.GetData(), Buffer.GetSize());
 }
 
+void FCbWriter::AddBinary(FUtf8StringView Name, const FCompositeBuffer& Value)
+{
+	SetName(Name);
+	AddBinary(Value);
+}
+
 void FCbWriter::AddBinary(const FCompositeBuffer& Buffer)
 {
 	AddBinary(Buffer.ToShared());
diff --git a/Engine/Source/Runtime/Core/Public/Serialization/CompactBinary.h b/Engine/Source/Runtime/Core/Public/Serialization/CompactBinary.h
--- a/Engine/Source/Runtime/Core/Public/Serialization/CompactBinary.h
+++ b/Engine/Source/Runtime/Core/Public/Serialization/CompactBinary.h
@@ -7,20 +7,17 @@
 #include "CoreTypes.h"
 #include "HAL/UnrealMemory.h"
 #include "IO/IoHash.h"
-#include "Memory/CompositeBuffer.h"
 #include "Memory/MemoryFwd.h"
 #include "Memory/MemoryView.h"
 #include "Memory/SharedBuffer.h"
 #include "Misc/AssertionMacros.h"
 #include "Misc/EnumClassFlags.h"
+#include "Misc/NotNull.h"
 #include "String/BytesToHex.h"
-#include "Templates/Function.h"
-#include "Templates/RemoveReference.h"
+#include "Templates/FunctionWithContext.h"
 #include "Templates/UnrealTemplate.h"
 #include <type_traits>
 
-template <typename CharType> class TStringBuilderBase;
-
 /**
  * This file declares a compact binary data format that is compatible with JSON and only slightly
  * more expressive. The format is designed to achieve fairly small encoded sizes while also being
@@ -80,14 +77,18 @@ class FCbFieldViewIterator;
 class FCbObjectId;
 class FCbObjectView;
 class FCbValue;
-class FIoHashBuilder;
 struct FDateTime;
 struct FGuid;
 struct FTimespan;
-template <typename FuncType> class TFunctionRef;
 
 /** A reference to a function that is used to allocate buffers for compact binary data. */
-using FCbBufferAllocator = TFunctionRef<FUniqueBuffer (uint64 Size)>;
+using FCbBufferAllocator = TNotNull<UE::TFunctionWithContext<FUniqueBuffer (uint64 Size)>>;
+
+/**
+ * A reference to a function that is used to allocate memory for compact binary data.
+ * The provider of the allocator must manage ownership of the memory for which a view is returned.
+ */
+using FCbMemoryAllocator = TNotNull<UE::TFunctionWithContext<FMutableMemoryView (uint64 Size)>>;
 
 ///////////////////////////////////////////////////////////////////////////////////////////////////////////////////////
 
@@ -320,7 +321,7 @@ public:
 ///////////////////////////////////////////////////////////////////////////////////////////////////////////////////////
 
 /** A reference to a function that is used to visit fields. */
-using FCbFieldVisitor = TFunctionRef<void (FCbFieldView)>;
+using FCbFieldVisitor = TNotNull<UE::TFunctionWithContext<void (FCbFieldView)>>;
 
 ///////////////////////////////////////////////////////////////////////////////////////////////////////////////////////
 
@@ -1463,3 +1464,9 @@ inline FSharedBuffer FCbField::AsBinary(const FSharedBuffer& Default) &&
 }
 
 ///////////////////////////////////////////////////////////////////////////////////////////////////////////////////////
+
+#if UE_ENABLE_INCLUDE_ORDER_DEPRECATED_IN_5_8
+#include "Memory/CompositeBuffer.h"
+#include "Templates/Function.h"
+#include "Templates/RemoveReference.h"
+#endif
diff --git a/Engine/Source/Runtime/Core/Public/Serialization/CompactBinarySerialization.h b/Engine/Source/Runtime/Core/Public/Serialization/CompactBinarySerialization.h
--- a/Engine/Source/Runtime/Core/Public/Serialization/CompactBinarySerialization.h
+++ b/Engine/Source/Runtime/Core/Public/Serialization/CompactBinarySerialization.h
@@ -16,6 +16,7 @@ class FArchive;
 class FName;
 struct FDateTime;
 struct FGuid;
+enum class ECbValidateMode : uint32;
 
 ///////////////////////////////////////////////////////////////////////////////////////////////////////////////////////
 
@@ -23,7 +24,7 @@ namespace UE::Serialization::Private
 {
 
 /** Utility for logging problems with FCbField. For internal use only */
-CORE_API void LogFieldTooLargeForArrayWarning(uint64 FieldLength);
+CORE_API void LogFieldTooLargeForArrayWarning(uint64 ArraySize);
 
 } // namespace UE::Serialization::Private
 
@@ -74,13 +75,31 @@ CORE_API bool TryMeasureCompactBinary(
  *
  * The field may be an array or an object, which the caller can convert to by using AsArray or
  * AsObject as appropriate. The buffer allocator is called to provide the buffer for the field
- * to load into once its size has been determined.
+ * to load into once its size has been determined. Validated with ECbValidateMode::Default.
  *
  * @param Ar Archive to read the field from. An error state is set on failure.
  * @param Allocator Allocator for the buffer that the field is loaded into.
  * @return A field with a reference to the allocated buffer, or a default field on failure.
  */
 CORE_API FCbField LoadCompactBinary(FArchive& Ar, FCbBufferAllocator Allocator = FUniqueBuffer::Alloc);
+/** Load a compact binary field view from an archive. Allocator must manage lifetime. */
+CORE_API FCbFieldView LoadCompactBinaryView(FArchive& Ar, FCbMemoryAllocator Allocator);
+
+/**
+ * Load a compact binary field from an archive.
+ *
+ * The field may be an array or an object, which the caller can convert to by using AsArray or
+ * AsObject as appropriate. The buffer allocator is called to provide the buffer for the field
+ * to load into once its size has been determined.
+ *
+ * @param Ar Archive to read the field from. An error state is set on failure.
+ * @param Allocator Allocator for the buffer that the field is loaded into. May be FUniqueBuffer::Alloc.
+ * @param ValidateMode Validation mode in which to validate the loaded field. May be None.
+ * @return A field with a reference to the allocated buffer, or a default field on failure.
+ */
+CORE_API FCbField LoadAndValidateCompactBinary(FArchive& Ar, FCbBufferAllocator Allocator, ECbValidateMode ValidateMode);
+/** Load a compact binary field view from an archive. Allocator must manage lifetime. */
+CORE_API FCbFieldView LoadAndValidateCompactBinaryView(FArchive& Ar, FCbMemoryAllocator Allocator, ECbValidateMode ValidateMode);
 
 /** Save a compact binary value to an archive. */
 CORE_API void SaveCompactBinary(FArchive& Ar, const FCbFieldView& Field);
@@ -198,20 +217,20 @@ inline bool LoadFromCompactBinary(FCbFieldView Field, FCbObjectId& OutValue, con
 template <typename T, typename Allocator>
 inline bool LoadFromCompactBinary(FCbFieldView Field, TArray<T, Allocator>& OutValue)
 {
-	const uint64 Length = Field.AsArrayView().Num();
-	if (Length <= MAX_int32)
+	const uint64 Size = Field.AsArrayView().Num();
+	if (Size <= MAX_int32)
 	{
-		OutValue.Reset((int32)Length);
+		OutValue.Reset((int32)Size);
 		bool bOk = !Field.HasError();
 		for (const FCbFieldView& ElementField : Field)
 		{
-			bOk = LoadFromCompactBinary(ElementField, OutValue.Emplace_GetRef()) & bOk;
+			bOk &= LoadFromCompactBinary(ElementField, OutValue.Emplace_GetRef());
 		}
 		return bOk;
 	}
 	else
 	{
-		UE::Serialization::Private::LogFieldTooLargeForArrayWarning(Length);
+		UE::Serialization::Private::LogFieldTooLargeForArrayWarning(Size);
 		return false;
 	}
 }
diff --git a/Engine/Source/Runtime/Core/Public/Serialization/CompactBinaryValue.h b/Engine/Source/Runtime/Core/Public/Serialization/CompactBinaryValue.h
--- a/Engine/Source/Runtime/Core/Public/Serialization/CompactBinaryValue.h
+++ b/Engine/Source/Runtime/Core/Public/Serialization/CompactBinaryValue.h
@@ -22,6 +22,8 @@
 class FCbValue
 {
 public:
+	constexpr FCbValue() = default;
+	FCbValue(const FCbFieldView& View);
 	FCbValue(ECbFieldType Type, const void* Value);
 
 	FCbObjectView AsObjectView() const;
@@ -64,11 +66,18 @@ public:
 	FORCEINLINE ECbFieldType GetType() const { return Type; }
 	FORCEINLINE const void* GetData() const { return Data; }
 
+	FORCEINLINE explicit operator bool() const { return Type != ECbFieldType::None; }
+
 private:
-	const void* Data;
-	ECbFieldType Type;
+	const void* Data = nullptr;
+	ECbFieldType Type = ECbFieldType::None;
 };
 
+FORCEINLINE FCbValue::FCbValue(const FCbFieldView& View)
+	: FCbValue(View.GetValue())
+{
+}
+
 FORCEINLINE FCbFieldView::FCbFieldView(const FCbValue& InValue)
 	: TypeWithFlags(InValue.GetType())
 	, Value(InValue.GetData())
diff --git a/Engine/Source/Runtime/Core/Public/Serialization/CompactBinaryWriter.h b/Engine/Source/Runtime/Core/Public/Serialization/CompactBinaryWriter.h
--- a/Engine/Source/Runtime/Core/Public/Serialization/CompactBinaryWriter.h
+++ b/Engine/Source/Runtime/Core/Public/Serialization/CompactBinaryWriter.h
@@ -14,9 +14,6 @@
 
 #include <type_traits>
 
-class FCompositeBuffer;
-class FSharedBuffer;
-
 ///////////////////////////////////////////////////////////////////////////////////////////////////////////////////////
 
 class FArchive;
@@ -24,7 +21,6 @@ class FCbAttachment;
 class FName;
 struct FDateTime;
 struct FGuid;
-struct FIoHash;
 struct FTimespan;
 
 /**
@@ -164,9 +160,9 @@ public:
 	inline void AddBinary(FUtf8StringView Name, FMemoryView Value) { SetName(Name); AddBinary(Value); }
 	inline void AddBinary(FMemoryView Value) { AddBinary(Value.GetData(), Value.GetSize()); }
 	/** Write a binary field by copying the buffer. Holds a reference if owned. */
-	inline void AddBinary(FUtf8StringView Name, const FSharedBuffer& Value) { SetName(Name); AddBinary(Value); }
+	CORE_API void AddBinary(FUtf8StringView Name, const FSharedBuffer& Value);
 	CORE_API void AddBinary(const FSharedBuffer& Value);
-	inline void AddBinary(FUtf8StringView Name, const FCompositeBuffer& Value) { SetName(Name); AddBinary(Value); }
+	CORE_API void AddBinary(FUtf8StringView Name, const FCompositeBuffer& Value);
 	CORE_API void AddBinary(const FCompositeBuffer& Value);
 
 	/** Write a string field by copying the UTF-8 value. */
diff --git a/Engine/Source/Runtime/Core/Tests/Serialization/CompactBinarySerializationTest.cpp b/Engine/Source/Runtime/Core/Tests/Serialization/CompactBinarySerializationTest.cpp
--- a/Engine/Source/Runtime/Core/Tests/Serialization/CompactBinarySerializationTest.cpp
+++ b/Engine/Source/Runtime/Core/Tests/Serialization/CompactBinarySerializationTest.cpp
@@ -146,9 +146,21 @@ TEST_CASE_NAMED(FCbSaveTest, "System::Core::Serialization::SaveCompactBinary", "
 			if (TestEqual(FString::Printf(TEXT("SaveCompactBinary(%s)->Size"), Test), uint64(WriteAr.Num()), ExpectedData.GetSize()) &&
 				TestTrue(FString::Printf(TEXT("SaveCompactBinary(%s)->EqualBytes"), Test), ExpectedData.EqualBytes(MakeMemoryView(WriteAr))))
 			{
-				FMemoryReader ReadAr(WriteAr);
-				FCbField Field = LoadCompactBinary(ReadAr);
-				CHECK_MESSAGE(FString::Printf(TEXT("LoadCompactBinary(%s)->EqualBytes"), Test), ExpectedData.EqualBytes(Field.GetOuterBuffer().GetView()));
+				{
+					FMemoryReader ReadAr(WriteAr);
+					FCbField Field = LoadCompactBinary(ReadAr);
+					CHECK_MESSAGE(FString::Printf(TEXT("LoadCompactBinary(%s)->EqualBytes"), Test), ExpectedData.EqualBytes(Field.GetOuterBuffer().GetView()));
+				}
+				{
+					TArray<uint8, TInlineAllocator<256>> FieldBuffer;
+					FMemoryReader ReadAr(WriteAr);
+					FCbFieldView FieldView = LoadCompactBinaryView(ReadAr, [&FieldBuffer](uint64 Size)
+					{
+						FieldBuffer.SetNumUninitialized(IntCastChecked<int32>(Size));
+						return MakeMemoryView(FieldBuffer);
+					});
+					CHECK_MESSAGE(FString::Printf(TEXT("LoadCompactBinaryView(%s)->EqualBytes"), Test), ExpectedData.EqualBytes(MakeMemoryView(FieldBuffer)));
+				}
 			}
 		}
 		{
