Improve batch existence check efficiency for Cloud DDC client.  Trust the server to enforce the requirement that if a ref exists, the blobs it references also exist and not do round trip requests to check on every blobs existence.  Also allow batch existence checks (for Get, GetValues, GetChunks) to be done in a single http request.  Previously that special case optimization was only implemented for GetValues.

Test cases have been added for batched existence checks.

#Synced-CL 49552467

diff --git a/Engine/Source/Developer/DerivedDataCache/Private/HttpCacheStore.cpp b/Engine/Source/Developer/DerivedDataCache/Private/HttpCacheStore.cpp
--- a/Engine/Source/Developer/DerivedDataCache/Private/HttpCacheStore.cpp
+++ b/Engine/Source/Developer/DerivedDataCache/Private/HttpCacheStore.cpp
@@ -21,6 +21,7 @@
 #include "DerivedDataHttpRequestQueue.h"
 #include "DerivedDataRequest.h"
 #include "DerivedDataRequestOwner.h"
+#include "DerivedDataSerialization.h"
 #include "DerivedDataValue.h"
 #include "DesktopPlatformModule.h"
 #include "Dom/JsonObject.h"
@@ -85,25 +86,6 @@ TRACE_DECLARE_ATOMIC_INT_COUNTER(HttpDDC_PutHit, TEXT("HttpDDC Put Hit"));
 TRACE_DECLARE_ATOMIC_INT_COUNTER(HttpDDC_BytesReceived, TEXT("HttpDDC Bytes Received"));
 TRACE_DECLARE_ATOMIC_INT_COUNTER(HttpDDC_BytesSent, TEXT("HttpDDC Bytes Sent"));
 
-static bool IsValueDataReady(FValue& Value, const ECachePolicy Policy)
-{
-	if (!EnumHasAnyFlags(Policy, ECachePolicy::Query))
-	{
-		Value = Value.RemoveData();
-		return true;
-	}
-
-	if (Value.HasData())
-	{
-		if (EnumHasAnyFlags(Policy, ECachePolicy::SkipData))
-		{
-			Value = Value.RemoveData();
-		}
-		return true;
-	}
-	return false;
-};
-
 static FAnsiStringView GetDomainFromUri(const FAnsiStringView Uri)
 {
 	FAnsiStringView Domain = Uri;
@@ -409,7 +391,10 @@ private:
 	class FPutPackageOp;
 	class FGetRecordOp;
 	class FGetValueOp;
-	class FExistsBatchOp;
+	template <typename RequestType> class TExistsBatchOp;
+	using FExistsRecordBatchOp = TExistsBatchOp<FCacheGetRequest>;
+	using FExistsValueBatchOp = TExistsBatchOp<FCacheGetValueRequest>;
+	using FExistsChunkBatchOp = TExistsBatchOp<FCacheGetChunkRequest>;
 };
 
 //----------------------------------------------------------------------------------------------------------
@@ -1405,18 +1390,17 @@ void FHttpCacheStore::FGetRecordOp::BeginGetValues(const FCacheRecord& Record, c
 	OnRecordComplete = MoveTemp(OnComplete);
 
 	TArray<FValueWithId> RequiredGets;
-	TArray<FValueWithId> RequiredHeads;
 
 	for (const FValueWithId& Value : Record.GetValues())
 	{
 		const ECachePolicy ValuePolicy = Policy.GetValuePolicy(Value.GetId());
-		if (EnumHasAnyFlags(ValuePolicy, ECachePolicy::QueryRemote))
+		if (EnumHasAnyFlags(ValuePolicy, ECachePolicy::QueryRemote) && !EnumHasAnyFlags(ValuePolicy, ECachePolicy::SkipData))
 		{
-			(EnumHasAnyFlags(ValuePolicy, ECachePolicy::SkipData) ? RequiredHeads : RequiredGets).Emplace(Value);
+			RequiredGets.Emplace(Value);
 		}
 	}
 
-	PrepareForPendingValues(RequiredGets.Num() + RequiredHeads.Num());
+	PrepareForPendingValues(RequiredGets.Num());
 
 	RequestTimer.Stop();
 
@@ -1433,14 +1417,6 @@ void FHttpCacheStore::FGetRecordOp::BeginGetValues(const FCacheRecord& Record, c
 			Self->EndGetValues(Policy, Response.Status);
 		}
 	});
-
-	GetValuesExist(RequiredHeads, [Self = TRefCountPtr(this), Policy](FValueResponse&& Response)
-	{
-		if (Self->FinishPendingValueExists(Response.Status))
-		{
-			Self->EndGetValues(Policy, Response.Status);
-		}
-	});
 }
 
 void FHttpCacheStore::FGetRecordOp::EndGetValues(const FCacheRecordPolicy& Policy, EStatus Status)
@@ -1874,8 +1850,6 @@ void FHttpCacheStore::FGetValueOp::EndGetRef(FHttpOperation& Operation)
 
 	Operation.GetStats(RequestStats);
 
-	const bool bSkipData = EnumHasAnyFlags(Policy, ECachePolicy::SkipData);
-
 	if (const EHttpErrorCode ErrorCode = Operation.GetErrorCode(); ErrorCode != EHttpErrorCode::None)
 	{
 		if (ErrorCode != EHttpErrorCode::Canceled)
@@ -1896,7 +1870,7 @@ void FHttpCacheStore::FGetValueOp::EndGetRef(FHttpOperation& Operation)
 
 	FSharedBuffer Body = Operation.GetBody();
 
-	if (bSkipData)
+	if (EnumHasAnyFlags(Policy, ECachePolicy::SkipData))
 	{
 		if (FAnsiStringView ReceivedHashStr = Operation.GetHeader("X-Jupiter-IoHash"); !ReceivedHashStr.IsEmpty())
 		{
@@ -1984,47 +1958,62 @@ void FHttpCacheStore::FGetValueOp::EndGet(FResponse&& Response)
 }
 
 //----------------------------------------------------------------------------------------------------------
-// FHttpCacheStore::FExistsBatchOp
+// FHttpCacheStore::TExistsBatchOp
 //----------------------------------------------------------------------------------------------------------
-class FHttpCacheStore::FExistsBatchOp final : public FRefCountedObject
+template <typename RequestType>
+class FHttpCacheStore::TExistsBatchOp final : public FRefCountedObject
 {
 public:
-	static TRefCountPtr<FExistsBatchOp> New(FHttpCacheStore& CacheStore, IRequestOwner& Owner)
+	using FOnComplete = TCacheOnCompleteFor<RequestType>;
+
+	static TRefCountPtr<TExistsBatchOp> New(FHttpCacheStore& CacheStore, IRequestOwner& Owner)
 	{
-		return new FExistsBatchOp(CacheStore, Owner);
+		return new TExistsBatchOp(CacheStore, Owner);
 	}
 
-	void Exists(TConstArrayView<FCacheGetValueRequest> Requests, FOnCacheGetValueComplete&& OnComplete);
+	void Exists(TConstArrayView<RequestType> Requests, FOnComplete&& OnComplete);
 
 private:
-	FExistsBatchOp(FHttpCacheStore& CacheStore, IRequestOwner& Owner);
+	TExistsBatchOp(FHttpCacheStore& CacheStore, IRequestOwner& Owner);
 
 	void BeginExists(TUniquePtr<FHttpOperation>&& Operation, FCbFieldIterator&& Body);
 	void EndExists(FHttpOperation& Operation);
 
-	void EndRequest(const FCacheGetValueRequest& Request, const FValue& Value, EStatus Status);
+	void EndRequest(const RequestType& Request, FCbObjectView ResponseView, EStatus Status);
 
 	FHttpCacheStore& CacheStore;
 	IRequestOwner& Owner;
-	TArray<FCacheGetValueRequest> Requests;
-	FOnCacheGetValueComplete OnComplete;
+	TArray<RequestType> Requests;
+	FOnComplete OnComplete;
 	FRequestStats RequestStats;
 };
 
-FHttpCacheStore::FExistsBatchOp::FExistsBatchOp(FHttpCacheStore& InCacheStore, IRequestOwner& InOwner)
+template <typename RequestType>
+FHttpCacheStore::TExistsBatchOp<RequestType>::TExistsBatchOp(FHttpCacheStore& InCacheStore, IRequestOwner& InOwner)
 	: CacheStore(InCacheStore)
 	, Owner(InOwner)
 {
 }
 
-void FHttpCacheStore::FExistsBatchOp::Exists(TConstArrayView<FCacheGetValueRequest> InRequests, FOnCacheGetValueComplete&& InOnComplete)
+static bool PolicyHasAnyFlags(ECachePolicy Policy, ECachePolicy Contains)
+{
+	return EnumHasAnyFlags(Policy, Contains);
+}
+
+static bool PolicyHasAnyFlags(const FCacheRecordPolicy& Policy, ECachePolicy Contains)
+{
+	return EnumHasAnyFlags(Policy.GetRecordPolicy(), Contains);
+}
+
+template <typename RequestType>
+void FHttpCacheStore::TExistsBatchOp<RequestType>::Exists(TConstArrayView<RequestType> InRequests, FOnComplete&& InOnComplete)
 {
 	FRequestTimer RequestTimer(RequestStats);
 
 	OnComplete = MoveTemp(InOnComplete);
 
 	Requests.Empty(InRequests.Num());
-	for (const FCacheGetValueRequest& Request : InRequests)
+	for (const RequestType& Request : InRequests)
 	{
 		if (!CacheStore.IsUsable())
 		{
@@ -2035,7 +2024,7 @@ void FHttpCacheStore::FExistsBatchOp::Exists(TConstArrayView<FCacheGetValueReque
 			continue;
 		}
 
-		if (!EnumHasAnyFlags(Request.Policy, ECachePolicy::QueryRemote))
+		if (!PolicyHasAnyFlags(Request.Policy, ECachePolicy::QueryRemote))
 		{
 			UE_LOG(LogDerivedDataCache, VeryVerbose, TEXT("%s: Skipped exists check of %s from '%s' due to cache policy"),
 				*CacheStore.NodeName, *WriteToString<96>(Request.Key), *Request.Name);
@@ -2063,7 +2052,7 @@ void FHttpCacheStore::FExistsBatchOp::Exists(TConstArrayView<FCacheGetValueReque
 	BodyWriter.BeginObject();
 	BodyWriter.BeginArray(ANSITEXTVIEW("ops"));
 	uint32 OpIndex = 0;
-	for (const FCacheGetValueRequest& Request : Requests)
+	for (const RequestType& Request : Requests)
 	{
 		BodyWriter.BeginObject();
 		BodyWriter.AddInteger(ANSITEXTVIEW("opId"), OpIndex);
@@ -2088,11 +2077,12 @@ void FHttpCacheStore::FExistsBatchOp::Exists(TConstArrayView<FCacheGetValueReque
 	});
 }
 
-void FHttpCacheStore::FExistsBatchOp::BeginExists(TUniquePtr<FHttpOperation>&& Operation, FCbFieldIterator&& Body)
+template <typename RequestType>
+void FHttpCacheStore::TExistsBatchOp<RequestType>::BeginExists(TUniquePtr<FHttpOperation>&& Operation, FCbFieldIterator&& Body)
 {
 	if (UNLIKELY(!Operation))
 	{
-		for (const FCacheGetValueRequest& Request : Requests)
+		for (const RequestType& Request : Requests)
 		{
 			UE_LOG(LogDerivedDataCache, Verbose, TEXT("%s: Cache miss with canceled request for %s from '%s'"),
 				*CacheStore.NodeName, *WriteToString<96>(Request.Key), *Request.Name);
@@ -2119,7 +2109,8 @@ void FHttpCacheStore::FExistsBatchOp::BeginExists(TUniquePtr<FHttpOperation>&& O
 	});
 }
 
-void FHttpCacheStore::FExistsBatchOp::EndExists(FHttpOperation& Operation)
+template <typename RequestType>
+void FHttpCacheStore::TExistsBatchOp<RequestType>::EndExists(FHttpOperation& Operation)
 {
 	TRACE_CPUPROFILER_EVENT_SCOPE(HttpDDC_ExistsBatch_EndExists);
 	ON_SCOPE_EXIT
@@ -2138,13 +2129,12 @@ void FHttpCacheStore::FExistsBatchOp::EndExists(FHttpOperation& Operation)
 	RequestStats.MainThreadTime = FMonotonicTimeSpan::FromSeconds(RequestStats.MainThreadTime.ToSeconds() / Requests.Num());
 	RequestStats.OtherThreadTime = FMonotonicTimeSpan::FromSeconds(RequestStats.OtherThreadTime.ToSeconds() / Requests.Num());
 	RequestStats.EndTime = FMonotonicTimePoint::Now();
-	RequestStats.Type = ERequestType::Value;
 	RequestStats.Op = ERequestOp::Get;
 
 	if (const EHttpErrorCode ErrorCode = Operation.GetErrorCode(); ErrorCode != EHttpErrorCode::None)
 	{
 		RequestTimer.Stop();
-		for (const FCacheGetValueRequest& Request : Requests)
+		for (const RequestType& Request : Requests)
 		{
 			if (ErrorCode != EHttpErrorCode::Canceled)
 			{
@@ -2161,7 +2151,7 @@ void FHttpCacheStore::FExistsBatchOp::EndExists(FHttpOperation& Operation)
 	if (OverallStatusCode < 200 || OverallStatusCode > 204)
 	{
 		RequestTimer.Stop();
-		for (const FCacheGetValueRequest& Request : Requests)
+		for (const RequestType& Request : Requests)
 		{
 			UE_LOG(LogDerivedDataCache, Verbose, TEXT("%s: Cache miss with failed non-success status code for %s from '%s'"),
 				*CacheStore.NodeName, *WriteToString<96>(Request.Key), *Request.Name);
@@ -2175,7 +2165,7 @@ void FHttpCacheStore::FExistsBatchOp::EndExists(FHttpOperation& Operation)
 	if (ValidateCompactBinary(ResponseView, ECbValidateMode::Default) != ECbValidateError::None)
 	{
 		RequestTimer.Stop();
-		for (const FCacheGetValueRequest& Request : Requests)
+		for (const RequestType& Request : Requests)
 		{
 			UE_LOG(LogDerivedDataCache, Log, TEXT("%s: Cache miss with corrupt response for %s from '%s'."),
 				*CacheStore.NodeName, *WriteToString<96>(Request.Key), *Request.Name);
@@ -2195,7 +2185,7 @@ void FHttpCacheStore::FExistsBatchOp::EndExists(FHttpOperation& Operation)
 		UE_LOG(LogDerivedDataCache, Log,
 			TEXT("%s: Cache exists returned unexpected quantity of results (expected %d, got %" UINT64_FMT ")."),
 			*CacheStore.NodeName, Requests.Num(), Results.Num());
-		for (const FCacheGetValueRequest& Request : Requests)
+		for (const RequestType& Request : Requests)
 		{
 			UE_LOG(LogDerivedDataCache, Display, TEXT("%s: Cache miss with invalid response for %s from '%s'"),
 				*CacheStore.NodeName, *WriteToString<96>(Request.Key), *Request.Name);
@@ -2210,7 +2200,7 @@ void FHttpCacheStore::FExistsBatchOp::EndExists(FHttpOperation& Operation)
 		const FCbObjectView ResultObject = ResultField.AsObjectView();
 		const uint32 OpId = ResultObject[ANSITEXTVIEW("opId")].AsUInt32();
 		const int32 StatusCode = ResultObject[ANSITEXTVIEW("statusCode")].AsInt32();
-		const FCbObjectView Value = ResultObject[ANSITEXTVIEW("response")].AsObjectView();
+		const FCbObjectView ResponseItemView = ResultObject[ANSITEXTVIEW("response")].AsObjectView();
 
 		if (OpId >= (uint32)Requests.Num())
 		{
@@ -2219,7 +2209,7 @@ void FHttpCacheStore::FExistsBatchOp::EndExists(FHttpOperation& Operation)
 			continue;
 		}
 
-		const FCacheGetValueRequest& Request = Requests[int32(OpId)];
+		const RequestType& Request = Requests[int32(OpId)];
 		RequestStats.Bucket = Request.Key.Bucket;
 
 		if (StatusCode < 200 || StatusCode > 204)
@@ -2230,22 +2220,56 @@ void FHttpCacheStore::FExistsBatchOp::EndExists(FHttpOperation& Operation)
 			continue;
 		}
 
-		const FIoHash RawHash = Value[ANSITEXTVIEW("RawHash")].AsHash();
-		const uint64 RawSize = Value[ANSITEXTVIEW("RawSize")].AsUInt64(MAX_uint64);
+		EndRequest(Request, ResponseItemView, EStatus::Ok);
+	}
+}
+
+template <>
+void FHttpCacheStore::FExistsRecordBatchOp::EndRequest(const FCacheGetRequest& Request, const FCbObjectView ResponseView, EStatus Status)
+{
+	FOptionalCacheRecord Record;
+	if (Status == EStatus::Ok && ResponseView && !LoadFromCompactBinary(ResponseView.AsFieldView(), Record))
+	{
+		UE_LOG(LogDerivedDataCache, Log, TEXT("%s: Cache miss with record load failure for %s from '%s'"),
+			*CacheStore.NodeName, *WriteToString<96>(Request.Key), *Request.Name);
+		Status = EStatus::Error;
+	}
+
+	RequestStats.Type = ERequestType::Record;
+	RequestStats.EndTime = FMonotonicTimePoint::Now();
+	RequestStats.Status = Status;
+	if (Record.IsNull())
+	{
+		Record = FCacheRecordBuilder(Request.Key).Build();
+	}
+	OnComplete({Request.Name, MoveTemp(Record).Get(), Request.UserData, Status});
+	if (CacheStore.StoreStats)
+	{
+		CacheStore.StoreStats->AddRequest(RequestStats);
+	}
+}
+
+template <>
+void FHttpCacheStore::FExistsValueBatchOp::EndRequest(const FCacheGetValueRequest& Request, const FCbObjectView ResponseView, EStatus Status)
+{
+	FValue Value;
+	if (Status == EStatus::Ok && ResponseView)
+	{
+		const FIoHash RawHash = ResponseView[ANSITEXTVIEW("RawHash")].AsHash();
+		const uint64 RawSize = ResponseView[ANSITEXTVIEW("RawSize")].AsUInt64(MAX_uint64);
 		if (RawHash.IsZero() || RawSize == MAX_uint64)
 		{
 			UE_LOG(LogDerivedDataCache, Display, TEXT("%s: Cache miss with invalid value for %s from '%s'"),
 				*CacheStore.NodeName, *WriteToString<96>(Request.Key), *Request.Name);
-			EndRequest(Request, {}, EStatus::Error);
-			continue;
+			Status = EStatus::Error;
+		}
+		else
+		{
+			Value = FValue(RawHash, RawSize);
 		}
-
-		EndRequest(Request, FValue(RawHash, RawSize), EStatus::Ok);
 	}
-}
 
-void FHttpCacheStore::FExistsBatchOp::EndRequest(const FCacheGetValueRequest& Request, const FValue& Value, EStatus Status)
-{
+	RequestStats.Type = ERequestType::Value;
 	RequestStats.EndTime = FMonotonicTimePoint::Now();
 	RequestStats.Status = Status;
 	OnComplete({Request.Name, Request.Key, Value, Request.UserData, Status});
@@ -2255,6 +2279,69 @@ void FHttpCacheStore::FExistsBatchOp::EndRequest(const FCacheGetValueRequest& Re
 	}
 }
 
+template <>
+void FHttpCacheStore::FExistsChunkBatchOp::EndRequest(const FCacheGetChunkRequest& Request, const FCbObjectView ResponseView, EStatus Status)
+{
+	if (Request.Id.IsValid())
+	{
+		FValueWithId Value;
+		FOptionalCacheRecord Record;
+		if (Status == EStatus::Ok && ResponseView && !LoadFromCompactBinary(ResponseView.AsFieldView(), Record))
+		{
+			UE_LOG(LogDerivedDataCache, Log, TEXT("%s: Cache miss with record load failure for %s from '%s'"),
+				*CacheStore.NodeName, *WriteToString<96>(Request.Key), *Request.Name);
+			Status = EStatus::Error;
+		}
+		else if (Record.IsValid())
+		{
+			Value = Record.Get().GetValue(Request.Id);
+			if (!Value || Value.GetRawHash().IsZero())
+			{
+				UE_LOG(LogDerivedDataCache, Display, TEXT("%s: Cache miss with missing/invalid value for %s from '%s'"),
+					*CacheStore.NodeName, *WriteToString<96>(Request.Key), *Request.Name);
+				Status = EStatus::Error;
+			}
+		}
+
+		RequestStats.Type = ERequestType::Record;
+		RequestStats.EndTime = FMonotonicTimePoint::Now();
+		RequestStats.Status = Status;
+		OnComplete({Request.Name, Request.Key, Request.Id, Request.RawOffset, Value.GetRawSize(), Value.GetRawHash(), {}, Request.UserData, Status});
+		if (CacheStore.StoreStats)
+		{
+			CacheStore.StoreStats->AddRequest(RequestStats);
+		}
+	}
+	else
+	{
+		FValue Value;
+		if (Status == EStatus::Ok && ResponseView)
+		{
+			const FIoHash RawHash = ResponseView[ANSITEXTVIEW("RawHash")].AsHash();
+			const uint64 RawSize = ResponseView[ANSITEXTVIEW("RawSize")].AsUInt64(MAX_uint64);
+			if (RawHash.IsZero() || RawSize == MAX_uint64)
+			{
+				UE_LOG(LogDerivedDataCache, Display, TEXT("%s: Cache miss with invalid value for %s from '%s'"),
+					*CacheStore.NodeName, *WriteToString<96>(Request.Key), *Request.Name);
+				Status = EStatus::Error;
+			}
+			else
+			{
+				Value = FValue(RawHash, RawSize);
+			}
+		}
+
+		RequestStats.Type = ERequestType::Value;
+		RequestStats.EndTime = FMonotonicTimePoint::Now();
+		RequestStats.Status = Status;
+		OnComplete({Request.Name, Request.Key, Request.Id, Request.RawOffset, Value.GetRawSize(), Value.GetRawHash(), {}, Request.UserData, Status});
+		if (CacheStore.StoreStats)
+		{
+			CacheStore.StoreStats->AddRequest(RequestStats);
+		}
+	}
+}
+
 ///////////////////////////////////////////////////////////////////////////////////////////////////////////////////////
 
 FHttpCacheStore::FHttpCacheStore(const FHttpCacheStoreParams& Params, ICacheStoreOwner* Owner)
@@ -2937,53 +3024,6 @@ void FHttpCacheStore::FinishChunkRequest(
 	}
 }
 
-static void AppendGetAndHeadOpsForChunkRequestGroupItem(
-	const FCacheGetChunkRequest& Request,
-	const FValueWithId& ValueWithId,
-	TArray<FValueWithId>& RequiredGets,
-	TArray<TArray<FCacheGetChunkRequest>>& RequiredGetRequests,
-	TArray<FValueWithId>& RequiredHeads,
-	TArray<TArray<FCacheGetChunkRequest>>& RequiredHeadRequests)
-{
-	const bool bAlreadyRequiredGet = !RequiredGets.IsEmpty() && RequiredGets.Last() == ValueWithId;
-	const bool bAlreadyRequiredHead = !RequiredHeads.IsEmpty() && RequiredHeads.Last() == ValueWithId;
-	if (EnumHasAnyFlags(Request.Policy, ECachePolicy::SkipData))
-	{
-		if (!bAlreadyRequiredHead && !bAlreadyRequiredGet)
-		{
-			RequiredHeads.Emplace(ValueWithId);
-			RequiredHeadRequests.AddDefaulted();
-		}
-		if (bAlreadyRequiredGet)
-		{
-			RequiredGetRequests.Last().Add(Request);
-		}
-		else
-		{
-			RequiredHeadRequests.Last().Add(Request);
-		}
-	}
-	else
-	{
-		if (!bAlreadyRequiredGet)
-		{
-			RequiredGets.Emplace(ValueWithId);
-			if (bAlreadyRequiredHead)
-			{
-				//Steal existing head contents first
-				RequiredGetRequests.Emplace(MoveTemp(RequiredHeadRequests.Last()));
-				RequiredHeads.SetNum(RequiredHeads.Num() - 1, EAllowShrinking::No);
-			}
-			else
-			{
-				RequiredGetRequests.AddDefaulted();
-			}
-		}
-
-		RequiredGetRequests.Last().Add(Request);
-	}
-}
-
 void FHttpCacheStore::GetChunkGroupAsync(
 	IRequestOwner& Owner,
 	const FCacheGetChunkRequest* StartRequest,
@@ -3029,27 +3069,30 @@ void FHttpCacheStore::GetChunkGroupAsync(
 
 				TArray<FValueWithId> RequiredGets;
 				TArray<TArray<FCacheGetChunkRequest>> RequiredGetRequests;
-				TArray<FValueWithId> RequiredHeads;
-				TArray<TArray<FCacheGetChunkRequest>> RequiredHeadRequests;
 				FCompressedBufferReader NullReader;
 				for (const FCacheGetChunkRequest& Request : RequestGroup)
 				{
-					const FValueWithId& ValueWithId = Response.Record.GetValue(Request.Id);
-					bool bHasValue = ValueWithId.IsValid();
-					FValue Value = ValueWithId;
-
-					if (!bHasValue || IsValueDataReady(Value, Request.Policy))
+					const FValueWithId& Value = Response.Record.GetValue(Request.Id);
+					if (!Value || !EnumHasAnyFlags(Request.Policy, ECachePolicy::Query) || EnumHasAnyFlags(Request.Policy, ECachePolicy::SkipData))
 					{
-						const EStatus Status = bHasValue ? Response.Status : EStatus::Error;
+						const EStatus Status = Value ? Response.Status : EStatus::Error;
 						FinishChunkRequest(Request, Status, Value, NullReader, SharedOnComplete);
 					}
 					else
 					{
-						AppendGetAndHeadOpsForChunkRequestGroupItem(Request, ValueWithId, RequiredGets,RequiredGetRequests, RequiredHeads, RequiredHeadRequests);
+						const bool bAlreadyRequiredGet = !RequiredGets.IsEmpty() && RequiredGets.Last() == Value;
+
+						if (!bAlreadyRequiredGet)
+						{
+							RequiredGets.Emplace(Value);
+							RequiredGetRequests.AddDefaulted();
+						}
+
+						RequiredGetRequests.Last().Add(Request);
 					}
 				}
 
-				int32 PendingValues = RequiredGets.Num() + RequiredHeads.Num();
+				int32 PendingValues = RequiredGets.Num();
 				Op->PrepareForPendingValues(PendingValues);
 
 				RequestTimer.Stop();
@@ -3082,29 +3125,6 @@ void FHttpCacheStore::GetChunkGroupAsync(
 						FinishChunkRequest(ChunkRequest, Response.Status, Response.Value, ValueReader, SharedOnComplete);
 					}
 				});
-
-				Op->GetValuesExist(RequiredHeads, [this, RecordStats, Op = TRefCountPtr(Op), ChunkRequestsForValues = MoveTemp(RequiredHeadRequests), SharedOnComplete](FGetRecordOp::FValueResponse&& Response)
-				{
-					int FoundRequestsIndex = Algo::BinarySearchBy(ChunkRequestsForValues, Response.Value.GetId(), [](const TArray<FCacheGetChunkRequest>& ChunkRequests)
-					{
-						check(!ChunkRequests.IsEmpty());
-						return ChunkRequests[0].Id;
-					});
-
-					check(FoundRequestsIndex != INDEX_NONE);
-					const TArray<FCacheGetChunkRequest>& ChunkRequests = ChunkRequestsForValues[FoundRequestsIndex];
-
-					if (Op->FinishPendingValueExists(Response.Status))
-					{
-						RecordStats(*Op, ChunkRequests[0].Key.Bucket, Op->GetFailedValues() > 0 ? EStatus::Error : EStatus::Ok);
-					}
-
-					FCompressedBufferReader NullReader;
-					for (const FCacheGetChunkRequest& ChunkRequest : ChunkRequests)
-					{
-						FinishChunkRequest(ChunkRequest, Response.Status, Response.Value, NullReader, SharedOnComplete);
-					}
-				});
 			}
 			else
 			{
@@ -3169,20 +3189,43 @@ void FHttpCacheStore::Get(
 	TRACE_CPUPROFILER_EVENT_SCOPE(HttpDDC_Get);
 	TRACE_COUNTER_ADD(HttpDDC_Get, Requests.Num());
 
-	FRequestBarrier Barrier(Owner);
-	TSharedRef<FOnCacheGetComplete> SharedOnComplete = MakeShared<FOnCacheGetComplete>(MoveTemp(OnComplete));
-	for (const FCacheGetRequest& Request : Requests)
+	// NOTE: Both the 'if' and the 'else' branch handle SkipData - see how FHttpCacheStore::FGetRecordOp::BeginGetValues chooses to send
+	//		 value queries depending on whether the request has SkipData - we just choose to use the batch existence operation so we avoid
+	//		 multiple round trips when there are multiple requests and all of them are existence queries (on the record policy).
+	const auto HasSkipData = [](const FCacheGetRequest& Request) { return EnumHasAnyFlags(Request.Policy.GetRecordPolicy(), ECachePolicy::SkipData); };
+	if ((Requests.Num() > 1) && Algo::AllOf(Requests, HasSkipData))
 	{
-		GetCacheRecordAsync(Owner, Request.Name, Request.Key, Request.Policy, Request.UserData,
-			[SharedOnComplete](FCacheGetResponse&& Response)
+		TRefCountPtr<FExistsRecordBatchOp> Op = FExistsRecordBatchOp::New(*this, Owner);
+		Op->Exists(Requests, [this, OnComplete = MoveTemp(OnComplete)](FCacheGetResponse&& Response)
 		{
 			if (Response.Status == EStatus::Ok)
 			{
 				TRACE_COUNTER_INCREMENT(HttpDDC_GetHit);
+				UE_LOG(LogDerivedDataCache, Verbose, TEXT("%s: Cache hit for %s from '%s'"),
+					*NodeName, *WriteToString<96>(Response.Record.GetKey()), *Response.Name);
 			}
-			SharedOnComplete.Get()(MoveTemp(Response));
+			OnComplete(MoveTemp(Response));
 		});
 	}
+	else
+	{
+		FRequestBarrier Barrier(Owner);
+		TSharedRef<FOnCacheGetComplete> SharedOnComplete = MakeShared<FOnCacheGetComplete>(MoveTemp(OnComplete));
+		for (const FCacheGetRequest& Request : Requests)
+		{
+			GetCacheRecordAsync(Owner, Request.Name, Request.Key, Request.Policy, Request.UserData,
+				[this, SharedOnComplete](FCacheGetResponse&& Response)
+			{
+				if (Response.Status == EStatus::Ok)
+				{
+					TRACE_COUNTER_INCREMENT(HttpDDC_GetHit);
+					UE_LOG(LogDerivedDataCache, Verbose, TEXT("%s: Cache hit for %s from '%s'"),
+						*NodeName, *WriteToString<96>(Response.Record.GetKey()), *Response.Name);
+				}
+				SharedOnComplete.Get()(MoveTemp(Response));
+			});
+		}
+	}
 }
 
 void FHttpCacheStore::PutValue(
@@ -3218,11 +3261,12 @@ void FHttpCacheStore::GetValue(
 
 	// NOTE: Both the 'if' and the 'else' branch handle SkipData - see how FHttpCacheStore::FGetValueOp::BeginGetRef chooses to send
 	//		 or not send x-jupiter-inline as an accept type depending on whether the request has SkipData - we just choose to use
-	//		 the batch existence operation so we avoid multiple round trips when there is more than one existence check being requested.
+	//		 the batch existence operation so we avoid multiple round trips when there are multiple requests and all of them are
+	//		 existence queries.
 	const auto HasSkipData = [](const FCacheGetValueRequest& Request) { return EnumHasAnyFlags(Request.Policy, ECachePolicy::SkipData); };
 	if ((Requests.Num() > 1) && Algo::AllOf(Requests, HasSkipData))
 	{
-		TRefCountPtr<FExistsBatchOp> Op = FExistsBatchOp::New(*this, Owner);
+		TRefCountPtr<FExistsValueBatchOp> Op = FExistsValueBatchOp::New(*this, Owner);
 		Op->Exists(Requests, [this, OnComplete = MoveTemp(OnComplete)](FCacheGetValueResponse&& Response)
 		{
 			if (Response.Status == EStatus::Ok)
@@ -3243,7 +3287,10 @@ void FHttpCacheStore::GetValue(
 			GetCacheValueAsync(Owner, Request.Name, Request.Key, Request.Policy, ERequestOp::Get, Request.UserData,
 			[this, Policy = Request.Policy, SharedOnComplete](FCacheGetValueResponse&& Response)
 			{
-				if (Response.Status == EStatus::Ok && !IsValueDataReady(Response.Value, Policy) && !EnumHasAnyFlags(Policy, ECachePolicy::SkipData))
+				if (Response.Status == EStatus::Ok &&
+					EnumHasAnyFlags(Policy, ECachePolicy::Query) &&
+					!EnumHasAnyFlags(Policy, ECachePolicy::SkipData) &&
+					!Response.Value.HasData())
 				{
 					Response.Status = EStatus::Error;
 					// With inline fetching, expect we will always have a value we can use.
@@ -3278,31 +3325,51 @@ void FHttpCacheStore::GetChunks(
 		return;
 	}
 
-	// TODO: This is inefficient because Jupiter doesn't allow us to get only part of a compressed blob, so we have to
-	//		 get the whole thing and then decompress only the portion we need.  Furthermore, because there is no propagation
-	//		 between cache stores during chunk requests, the fetched result won't end up in the local store.
-	//		 These efficiency issues will be addressed by changes to the Hierarchy that translate chunk requests that
-	//		 are missing in local/fast stores and have to be retrieved from slow stores into record requests instead.  That
-	//		 will make this code path unused/uncommon as Jupiter will most always be a slow store with a local/fast store in front of it.
-	//		 Regardless, to adhere to the functional contract, this implementation must exist.
-	TArray<FCacheGetChunkRequest, TInlineAllocator<16>> SortedRequests(Requests);
-	SortedRequests.StableSort(TChunkLess());
+	// NOTE: Both the 'if' and the 'else' branch handle SkipData - see how FHttpCacheStore::FGetRecordOp::BeginGetValues chooses to send
+	//		 value queries depending on whether the request has SkipData - we just choose to use the batch existence operation so we avoid
+	//		 multiple round trips when there are multiple requests of uniform policy and all of them are existence queries.
+	const auto HasSkipData = [](const FCacheGetChunkRequest& Request) { return EnumHasAnyFlags(Request.Policy, ECachePolicy::SkipData); };
+	if ((Requests.Num() > 1) && Algo::AllOf(Requests, HasSkipData))
+	{
+		TRefCountPtr<FExistsChunkBatchOp> Op = FExistsChunkBatchOp::New(*this, Owner);
+		Op->Exists(Requests, [this, OnComplete = MoveTemp(OnComplete)](FCacheGetChunkResponse&& Response)
+		{
+			if (Response.Status == EStatus::Ok)
+			{
+				TRACE_COUNTER_INCREMENT(HttpDDC_GetHit);
+				UE_LOG(LogDerivedDataCache, Verbose, TEXT("%s: Cache hit for %s from '%s'"),
+					*NodeName, *WriteToString<96>(Response.Key), *Response.Name);
+			}
+			OnComplete(MoveTemp(Response));
+		});
+	}
+	else
+	{
+		// TODO: This is inefficient because Jupiter doesn't allow us to get only part of a compressed blob, so we have to
+		//		 get the whole thing and then decompress only the portion we need.  Furthermore, because there is no propagation
+		//		 between cache stores during chunk requests, the fetched result won't end up in the local store.
+		//		 These efficiency issues will be addressed by changes to the Hierarchy that translate chunk requests that
+		//		 are missing in local/fast stores and have to be retrieved from slow stores into record requests instead.  That
+		//		 will make this code path unused/uncommon as Jupiter will most always be a slow store with a local/fast store in front of it.
+		//		 Regardless, to adhere to the functional contract, this implementation must exist.
+		TArray<FCacheGetChunkRequest, TInlineAllocator<16>> SortedRequests(Requests);
+		SortedRequests.StableSort(TChunkLess());
 
-	const FCacheGetChunkRequest* PendingGroupStartRequest = &SortedRequests[0];
+		const FCacheGetChunkRequest* PendingGroupStartRequest = &SortedRequests[0];
 
-	FRequestBarrier Barrier(Owner);
-	TSharedRef<FOnCacheGetChunkComplete> SharedOnComplete = MakeShared<FOnCacheGetChunkComplete>(MoveTemp(OnComplete));
-	for (const FCacheGetChunkRequest& Request : SortedRequests)
-	{
-		const bool bExistsOnly = EnumHasAnyFlags(Request.Policy, ECachePolicy::SkipData);
-		const bool bMatchesExistingGroup = PendingGroupStartRequest != nullptr && PendingGroupStartRequest->Key == Request.Key && PendingGroupStartRequest->Id.IsValid() == Request.Id.IsValid();
-		if (!bMatchesExistingGroup)
+		FRequestBarrier Barrier(Owner);
+		TSharedRef<FOnCacheGetChunkComplete> SharedOnComplete = MakeShared<FOnCacheGetChunkComplete>(MoveTemp(OnComplete));
+		for (const FCacheGetChunkRequest& Request : SortedRequests)
 		{
-			GetChunkGroupAsync(Owner, PendingGroupStartRequest, &Request, SharedOnComplete);
-			PendingGroupStartRequest = &Request;
+			const bool bMatchesExistingGroup = PendingGroupStartRequest != nullptr && PendingGroupStartRequest->Key == Request.Key && PendingGroupStartRequest->Id.IsValid() == Request.Id.IsValid();
+			if (!bMatchesExistingGroup)
+			{
+				GetChunkGroupAsync(Owner, PendingGroupStartRequest, &Request, SharedOnComplete);
+				PendingGroupStartRequest = &Request;
+			}
 		}
+		GetChunkGroupAsync(Owner, PendingGroupStartRequest, SortedRequests.GetData() + SortedRequests.Num(), SharedOnComplete);
 	}
-	GetChunkGroupAsync(Owner, PendingGroupStartRequest, SortedRequests.GetData() + SortedRequests.Num(), SharedOnComplete);
 }
 
 static void RegisterInheritedCommandLineArg(const FStringView ArgName)
diff --git a/Engine/Source/Developer/DerivedDataCache/Tests/DerivedDataCacheStoreTest.cpp b/Engine/Source/Developer/DerivedDataCache/Tests/DerivedDataCacheStoreTest.cpp
--- a/Engine/Source/Developer/DerivedDataCache/Tests/DerivedDataCacheStoreTest.cpp
+++ b/Engine/Source/Developer/DerivedDataCache/Tests/DerivedDataCacheStoreTest.cpp
@@ -118,6 +118,30 @@ static void TestCacheChunks(ICache& Cache, const FCacheKey& Key, const FValueId
 		BlockingOwner.Wait();
 	}
 
+	{
+		INFO("Get Value Hash + Size Redundant Batch");
+		FRequestOwner BlockingOwner(EPriority::Blocking);
+		std::atomic<uint64> UserDataAccumulator = 0;
+		Cache.GetChunks(
+			{
+				{.Name = Name, .Key = Key, .Id = ValueId, .Policy = ECachePolicy::Default | ECachePolicy::SkipData, .UserData = UserData},
+				{.Name = Name, .Key = Key, .Id = ValueId, .Policy = ECachePolicy::Default | ECachePolicy::SkipData, .UserData = UserData + 1} 
+			}, BlockingOwner, [&](FCacheGetChunkResponse&& Response)
+		{
+			UserDataAccumulator.fetch_xor(Response.UserData, std::memory_order_relaxed);
+			CHECK(Response.Status == EStatus::Ok);
+			CHECK(Response.Name == Name);
+			CHECK(Response.Key == Key);
+			CHECK(Response.Id == ValueId);
+			CHECK(Response.RawOffset == 0);
+			CHECK(Response.RawSize == Value.GetRawSize());
+			CHECK(Response.RawHash == Value.GetRawHash());
+			CHECK(Response.RawData.IsNull());
+		});
+		BlockingOwner.Wait();
+		CHECK(UserDataAccumulator == 1);
+	}
+
 	{
 		INFO("Get Value");
 		FRequestOwner BlockingOwner(EPriority::Blocking);
@@ -455,6 +479,28 @@ static void TestCacheValues(ICache& Cache, const uint32 ValueCount = 4 * 1024 *
 		BlockingOwner.Wait();
 	}
 
+	{
+		INFO("Get Value Hash + Size Redundant Batch");
+		FRequestOwner BlockingOwner(EPriority::Blocking);
+		std::atomic<uint64> UserDataAccumulator = 0;
+		Cache.GetValue(
+			{
+				{Name, Key, ECachePolicy::Default | ECachePolicy::SkipData, UserData},
+				{Name, Key, ECachePolicy::Default | ECachePolicy::SkipData, UserData + 1}
+			}, BlockingOwner, [&](FCacheGetValueResponse&& Response)
+		{
+			UserDataAccumulator.fetch_xor(Response.UserData, std::memory_order_relaxed);
+			CHECK(Response.Status == EStatus::Ok);
+			CHECK(Response.Name == Name);
+			CHECK(Response.Key == Key);
+			CHECK(Response.Value.GetRawHash() == Value.GetRawHash());
+			CHECK(Response.Value.GetRawSize() == Value.GetRawSize());
+			CHECK(!Response.Value.HasData());
+		});
+		BlockingOwner.Wait();
+		CHECK(UserDataAccumulator == 1);
+	}
+
 	{
 		INFO("Get Value");
 		FRequestOwner BlockingOwner(EPriority::Blocking);
@@ -549,6 +595,32 @@ static void TestCacheRecords(ICache& Cache)
 		BlockingOwner.Wait();
 	}
 
+	{
+		INFO("Get Record Hash + Size Redundant Batch");
+		FRequestOwner BlockingOwner(EPriority::Blocking);
+		std::atomic<uint64> UserDataAccumulator = 0;
+		Cache.Get(
+			{
+				{Name, Key, ECachePolicy::Default | ECachePolicy::SkipData, UserData},
+				{Name, Key, ECachePolicy::Default | ECachePolicy::SkipData, UserData + 1}
+			}, BlockingOwner, [&](FCacheGetResponse&& Response)
+		{
+			UserDataAccumulator.fetch_xor(Response.UserData, std::memory_order_relaxed);
+			CHECK(Response.Status == EStatus::Ok);
+			CHECK(Response.Name == Name);
+			CHECK(Response.Record.GetKey() == Key);
+			CHECK(Response.Record.GetValues().Num() == 5);
+			CHECK(HasValueWithNoData(Response.Record, Value));
+			CHECK(HasValueWithNoData(Response.Record, Attachment1));
+			CHECK(HasValueWithNoData(Response.Record, Attachment2));
+			CHECK(HasValueWithNoData(Response.Record, Attachment3));
+			CHECK(HasValueWithNoData(Response.Record, Attachment4));
+			CHECK(Response.Record.GetMeta().Equals(Meta));
+		});
+		BlockingOwner.Wait();
+		CHECK(UserDataAccumulator == 1);
+	}
+
 	{
 		INFO("Get Record");
 		FRequestOwner BlockingOwner(EPriority::Blocking);
