DDC: Move the DDC I/O thread pool into the module

Move the thread pool from the singleton cache to the module and ensure that it exists before creating any cache graph.

#Synced-CL 49572071

diff --git a/Engine/Source/Developer/DerivedDataCache/Private/DerivedDataCache.cpp b/Engine/Source/Developer/DerivedDataCache/Private/DerivedDataCache.cpp
--- a/Engine/Source/Developer/DerivedDataCache/Private/DerivedDataCache.cpp
+++ b/Engine/Source/Developer/DerivedDataCache/Private/DerivedDataCache.cpp
@@ -257,13 +257,6 @@ FCacheGetChunkResponse FCacheGetChunkRequest::MakeResponse(const EStatus Status)
 namespace UE::DerivedData::Private
 {
 
-FQueuedThreadPool* GCacheThreadPool;
-
-void LaunchTaskInCacheThreadPool(IRequestOwner& Owner, TUniqueFunction<void ()>&& TaskBody)
-{
-	LaunchTaskInThreadPool(0, TEXT("LaunchTaskInCacheThreadPool"), Owner, GCacheThreadPool, MoveTemp(TaskBody));
-}
-
 class FLegacyFetchOrBuildTask
 {
 public:
@@ -544,18 +537,6 @@ public:
 	FDerivedDataCache()
 		: CurrentHandle(19248) // we will skip some potential handles to catch errors
 	{
-		if (FPlatformProcess::SupportsMultithreading())
-		{
-			GCacheThreadPool = FQueuedThreadPool::Allocate();
-			const int32 ThreadCount = FPlatformMisc::NumberOfIOWorkerThreadsToSpawn();
-#if WITH_EDITOR
-			// Use normal priority to avoid preempting GT/RT/RHI and other more important threads with CPU processing (i.e. compression) happening on the IO Threads in editor.
-			verify(GCacheThreadPool->Create(ThreadCount, 96 * 1024, TPri_Normal, TEXT("DDC IO ThreadPool")));
-#else
-			verify(GCacheThreadPool->Create(ThreadCount, 96 * 1024, TPri_AboveNormal, TEXT("DDC IO ThreadPool")));
-#endif
-		}
-
 		Backend = FDerivedDataBackend::CreateStatic();
 
 		GVerifyDDC = FParse::Param(FCommandLine::Get(), TEXT("VerifyDDC"));
diff --git a/Engine/Source/Developer/DerivedDataCache/Private/DerivedDataCacheModule.cpp b/Engine/Source/Developer/DerivedDataCache/Private/DerivedDataCacheModule.cpp
--- a/Engine/Source/Developer/DerivedDataCache/Private/DerivedDataCacheModule.cpp
+++ b/Engine/Source/Developer/DerivedDataCache/Private/DerivedDataCacheModule.cpp
@@ -2,6 +2,8 @@
 
 #include "DerivedDataCacheModule.h"
 
+#include "Async/RecursiveMutex.h"
+#include "Async/UniqueLock.h"
 #include "CoreGlobals.h"
 #include "DerivedDataBuild.h"
 #include "DerivedDataBuildPrivate.h"
@@ -9,13 +11,17 @@
 #include "DerivedDataCacheInterface.h"
 #include "DerivedDataCachePrivate.h"
 #include "DerivedDataPrivate.h"
-#include "HAL/CriticalSection.h"
-#include "Misc/ScopeLock.h"
+#include "DerivedDataThreadPoolTask.h"
+#include "HAL/PlatformMisc.h"
+#include "HAL/PlatformProcess.h"
+#include "Misc/QueuedThreadPool.h"
 #include "Modules/ModuleManager.h"
 
 namespace UE::DerivedData::Private
 {
 
+FQueuedThreadPool* GCacheThreadPool;
+
 // Implemented in DerivedDataCache.cpp
 void CreateCache(ICache*& OutCache, FDerivedDataCacheInterface*& OutLegacyCache);
 
@@ -41,9 +47,26 @@ public:
 		return &GDerivedDataLegacyCache;
 	}
 
+	void CreateThreadPoolOnce()
+	{
+		TUniqueLock Lock(CreateMutex);
+		if (!GCacheThreadPool && FPlatformProcess::SupportsMultithreading())
+		{
+			GCacheThreadPool = FQueuedThreadPool::Allocate();
+			const int32 ThreadCount = FPlatformMisc::NumberOfIOWorkerThreadsToSpawn();
+		#if WITH_EDITOR
+			// Use normal priority in the editor to avoid preempting GT/RT/RHI and other more important threads
+			// with CPU processing (i.e. compression) that happens on these IO Threads in the editor.
+			verify(GCacheThreadPool->Create(ThreadCount, 96 * 1024, TPri_Normal, TEXT("DDC IO ThreadPool")));
+		#else
+			verify(GCacheThreadPool->Create(ThreadCount, 96 * 1024, TPri_AboveNormal, TEXT("DDC IO ThreadPool")));
+		#endif
+		}
+	}
+
 	void CreateCacheOnce()
 	{
-		FScopeLock Lock(&CreateLock);
+		TUniqueLock Lock(CreateMutex);
 		if (!GDerivedDataCache)
 		{
 			CreateCache(GDerivedDataCache, GDerivedDataLegacyCache);
@@ -54,7 +77,7 @@ public:
 
 	void CreateBuildOnce()
 	{
-		FScopeLock Lock(&CreateLock);
+		TUniqueLock Lock(CreateMutex);
 		if (!GDerivedDataBuild)
 		{
 			GDerivedDataBuild = CreateBuild(GDerivedDataCache);
@@ -79,7 +102,7 @@ public:
 	}
 
 private:
-	FCriticalSection CreateLock;
+	FRecursiveMutex CreateMutex;
 };
 
 static FDerivedDataCacheModule* GetModule()
@@ -94,6 +117,11 @@ static FDerivedDataCacheModule* GetModule()
 	return nullptr;
 }
 
+void LaunchTaskInCacheThreadPool(IRequestOwner& Owner, TUniqueFunction<void ()>&& TaskBody)
+{
+	LaunchTaskInThreadPool(0, TEXT("LaunchTaskInCacheThreadPool"), Owner, GCacheThreadPool, MoveTemp(TaskBody));
+}
+
 } // UE::DerivedData::Private
 
 namespace UE::DerivedData
@@ -112,6 +140,7 @@ ICache& GetCache()
 	check(IsInGameThread());
 	if (Private::FDerivedDataCacheModule* Module = Private::GetModule())
 	{
+		Module->CreateThreadPoolOnce();
 		Module->CreateCacheOnce();
 	}
 	ICache* Cache = Private::GDerivedDataCache;
@@ -127,8 +156,9 @@ ICache* TryGetCache()
 ICache* CreateCache(FStringView GraphNameOrConfig)
 {
 	check(IsInGameThread());
-	if (Private::GetModule())
+	if (Private::FDerivedDataCacheModule* Module = Private::GetModule())
 	{
+		Module->CreateThreadPoolOnce();
 		return Private::TryCreateCache(GraphNameOrConfig);
 	}
 	return nullptr;
diff --git a/Engine/Source/Developer/DerivedDataCache/Private/DerivedDataCachePrivate.h b/Engine/Source/Developer/DerivedDataCache/Private/DerivedDataCachePrivate.h
--- a/Engine/Source/Developer/DerivedDataCache/Private/DerivedDataCachePrivate.h
+++ b/Engine/Source/Developer/DerivedDataCache/Private/DerivedDataCachePrivate.h
@@ -12,7 +12,7 @@ namespace UE::DerivedData { class IRequestOwner; }
 namespace UE::DerivedData::Private
 {
 
-// Implemented in DerivedDataCache.cpp
+// Implemented in DerivedDataCacheModule.cpp
 void LaunchTaskInCacheThreadPool(IRequestOwner& Owner, TUniqueFunction<void ()>&& TaskBody);
 
 // Implemented in DerivedDataCacheRecord.cpp
