DDC: Added a concurrent chunked sparse array for internal use by DDC and DerivedDataTool

#Synced-CL 49572071

diff --git a/Engine/Source/Developer/DerivedDataCache/Internal/Containers/ConcurrentChunkedSparseArray.h b/Engine/Source/Developer/DerivedDataCache/Internal/Containers/ConcurrentChunkedSparseArray.h
new file mode 100644
--- /dev/null
+++ b/Engine/Source/Developer/DerivedDataCache/Internal/Containers/ConcurrentChunkedSparseArray.h
@@ -0,0 +1,315 @@
+// Copyright Epic Games, Inc. All Rights Reserved.
+
+#pragma once
+
+#include "CoreTypes.h"
+#include "HAL/UnrealMemory.h"
+#include "Math/UnrealMathUtility.h"
+#include "Templates/TypeCompatibleBytes.h"
+
+#include <atomic>
+
+namespace UE
+{
+
+/**
+ * A concurrent chunked sparse array that maintains every index at a consistent address once allocated.
+ *
+ * It is safe to add, remove, and access elements from multiple threads concurrently.
+ * Callers are responsible for ensuring multiple threads are not mutating the same element concurrently.
+ * Iteration is not provided because it is not safe in the presence of concurrent add and remove operations.
+ */
+template <typename ElementType, uint32 ElementsPerChunk>
+class TConcurrentChunkedSparseArray
+{
+	static_assert(sizeof(ElementType) >= sizeof(uint32), "ElementType is too small to support the element free list.");
+	static_assert(FMath::IsPowerOfTwo(ElementsPerChunk), "ElementsPerChunk must be a power of two to support efficient index calculations.");
+	// ChunkArray is approximately 16 bytes per chunk, which makes overhead per element is 1 byte at 16 elements per chunk
+	// and 2/4/8/16 bytes per element at 8/4/2/1 elements per chunk. Require at least 16 to avoid massive overhead.
+	static_assert(ElementsPerChunk >= 16, "ElementsPerChunk is too low for ChunkArray to be small relative to the element data.");
+
+	// Not implemented but can be if necessary.
+	TConcurrentChunkedSparseArray(const TConcurrentChunkedSparseArray&) = delete;
+	TConcurrentChunkedSparseArray& operator=(const TConcurrentChunkedSparseArray&) = delete;
+
+public:
+	constexpr TConcurrentChunkedSparseArray() = default;
+
+	~TConcurrentChunkedSparseArray()
+	{
+		FChunkArray* Array = ChunkArray.load(std::memory_order_acquire);
+		if (Array == &EmptyChunkArray)
+		{
+			return;
+		}
+
+		// Walk the free list to measure its size.
+		uint32 FreeElementCount = 0;
+		for (uint32 Index = uint32(FreeElementIndexAndCounter.load(std::memory_order_relaxed) >> 32); Index != MAX_uint32;)
+		{
+			++FreeElementCount;
+			const FChunk* Chunk = Array->Chunks[Index / ElementsPerChunk].load(std::memory_order_acquire);
+			Index = reinterpret_cast<const uint32&>(Chunk->Elements[Index % ElementsPerChunk]);
+		}
+
+		// Walk the free list to copy it into an array.
+		TArray<uint32> FreeElements;
+		FreeElements.Reserve(FreeElementCount);
+		for (uint32 Index = uint32(FreeElementIndexAndCounter.load(std::memory_order_relaxed) >> 32); Index != MAX_uint32;)
+		{
+			FreeElements.Add(Index);
+			const FChunk* Chunk = Array->Chunks[Index / ElementsPerChunk].load(std::memory_order_acquire);
+			Index = reinterpret_cast<const uint32&>(Chunk->Elements[Index % ElementsPerChunk]);
+		}
+		FreeElements.Sort();
+
+		// Destroy elements that are not in the free list.
+		uint32 LastElementIndex = NextElementIndex.load(std::memory_order_relaxed);
+		typename TArray<uint32>::TConstIterator FreeIt = FreeElements.CreateConstIterator();
+		for (uint32 ChunkIndex = 0, ChunkCount = Array->ChunkCount; ChunkIndex < ChunkCount; ++ChunkIndex)
+		{
+			const uint32 ElementsInChunk = FMath::Min(LastElementIndex, ElementsPerChunk);
+			if (FChunk* Chunk = Array->Chunks[ChunkIndex].load(std::memory_order_acquire))
+			{
+				for (uint32 IndexInChunk = 0; IndexInChunk < ElementsInChunk; ++IndexInChunk)
+				{
+					if (FreeIt && *FreeIt == ChunkIndex * ElementsPerChunk + IndexInChunk)
+					{
+						++FreeIt;
+						continue;
+					}
+					Chunk->Elements[IndexInChunk].DestroyUnchecked();
+				}
+				delete Chunk;
+			}
+			LastElementIndex -= ElementsInChunk;
+		}
+
+		// Destroy the chain of chunks arrays.
+		while (Array != &EmptyChunkArray)
+		{
+			FChunkArray* PrevArray = Array->Prev;
+			FMemory::Free(Array);
+			Array = PrevArray;
+		}
+	}
+
+	[[nodiscard]] inline ElementType& operator[](const uint32 ElementIndex) UE_LIFETIMEBOUND
+	{
+		FChunkArray* Array = ChunkArray.load(std::memory_order_acquire);
+		const uint32 ChunkIndex = ElementIndex / ElementsPerChunk;
+		checkfSlow(Array && ChunkIndex < Array->ChunkCount && ElementIndex < NextElementIndex.load(std::memory_order_relaxed),
+			TEXT("Array index out of bounds: %u into an array of size %u."), ElementIndex, NextElementIndex.load(std::memory_order_relaxed));
+		const uint32 IndexInChunk = ElementIndex % ElementsPerChunk;
+		FChunk* Chunk = Array->Chunks[ChunkIndex].load(std::memory_order_acquire);
+		checkSlow(Chunk);
+		return Chunk->Elements[IndexInChunk].GetUnchecked();
+	}
+
+	[[nodiscard]] UE_REWRITE const ElementType& operator[](const uint32 ElementIndex) const UE_LIFETIMEBOUND
+	{
+		return (*const_cast<TConcurrentChunkedSparseArray*>(this))[ElementIndex];
+	}
+
+	template <typename... ArgTypes>
+	[[nodiscard]] inline uint32 Emplace(ArgTypes&&... Args)
+	{
+		FChunk* Chunk;
+		const uint32 ElementIndex = PrivateAddUninitialized(Chunk);
+		const uint32 IndexInChunk = ElementIndex % ElementsPerChunk;
+		Chunk->Elements[IndexInChunk].EmplaceUnchecked((ArgTypes&&)Args...);
+		return ElementIndex;
+	}
+
+	[[nodiscard]] UE_REWRITE uint32 AddUninitialized()
+	{
+		FChunk* Chunk;
+		return PrivateAddUninitialized(Chunk);
+	}
+
+	[[nodiscard]] inline uint32 AddUninitialized(const uint32 ElementCount)
+	{
+		if (ElementCount == 0)
+		{
+			return NextElementIndex.load(std::memory_order_relaxed);
+		}
+		if (ElementCount == 1)
+		{
+			return AddUninitialized();
+		}
+
+		const uint32 FirstElementIndex = NextElementIndex.fetch_add(ElementCount, std::memory_order_relaxed);
+		const uint32 LastElementIndex = FirstElementIndex + ElementCount - 1;
+		FChunkArray* Array = GetOrGrowChunkArray(LastElementIndex);
+		const uint32 FirstChunkIndex = FirstElementIndex / ElementsPerChunk;
+		const uint32 LastChunkIndex = LastElementIndex / ElementsPerChunk;
+		for (uint32 ChunkIndex = FirstChunkIndex; ChunkIndex <= LastChunkIndex; ++ChunkIndex)
+		{
+			GetOrCreateChunk(Array, ChunkIndex);
+		}
+		return FirstElementIndex;
+	}
+
+	inline void RemoveAt(uint32 ElementIndex)
+	{
+		FChunkArray* Array = ChunkArray.load(std::memory_order_acquire);
+		const uint32 ChunkIndex = ElementIndex / ElementsPerChunk;
+		checkfSlow(Array && ChunkIndex < Array->ChunkCount && ElementIndex < NextElementIndex.load(std::memory_order_relaxed),
+			TEXT("Array index out of bounds: %u into an array of size %u."), ElementIndex, NextElementIndex.load(std::memory_order_relaxed));
+		const uint32 IndexInChunk = ElementIndex % ElementsPerChunk;
+		FChunk* Chunk = Array->Chunks[ChunkIndex].load(std::memory_order_acquire);
+		checkSlow(Chunk);
+		TTypeCompatibleBytes<ElementType>& Element = Chunk->Elements[IndexInChunk];
+		Element.DestroyUnchecked();
+
+		uint32& NextFreeElementIndex = reinterpret_cast<uint32&>(Element);
+		uint64 IndexAndCounter = FreeElementIndexAndCounter.load(std::memory_order_relaxed);
+		for (;;)
+		{
+			NextFreeElementIndex = uint32(IndexAndCounter >> 32);
+			const uint32 Counter = uint32(IndexAndCounter) + 1;
+			const uint64 NextIndexAndCounter = (uint64(ElementIndex) << 32) | Counter;
+			if (FreeElementIndexAndCounter.compare_exchange_strong(IndexAndCounter, NextIndexAndCounter, std::memory_order_release, std::memory_order_relaxed))
+			{
+				return;
+			}
+		}
+	}
+
+private:
+	struct FChunk
+	{
+		TTypeCompatibleBytes<ElementType> Elements[ElementsPerChunk];
+	};
+
+	struct FChunkArray
+	{
+		uint32 ChunkCount = 0;
+		uint32 Padding = 0;
+		// Pointer to the previous chunk array, to be deleted upon destruction of the container.
+		FChunkArray* Prev = nullptr;
+		// Array of ChunkCount chunks that are allocated on demand.
+		std::atomic<FChunk*> Chunks[0];
+	};
+
+	[[nodiscard]] inline uint32 PrivateAddUninitialized(FChunk*& OutChunk)
+	{
+		const uint32 ElementIndex = AcquireElementIndex();
+		FChunkArray* Array = GetOrGrowChunkArray(ElementIndex);
+		const uint32 ChunkIndex = ElementIndex / ElementsPerChunk;
+		OutChunk = GetOrCreateChunk(Array, ChunkIndex);
+		return ElementIndex;
+	}
+
+	FORCEINLINE uint32 AcquireElementIndex()
+	{
+		// Acquire an index from the free list or from the end of the array.
+		uint64 IndexAndCounter = FreeElementIndexAndCounter.load(std::memory_order_acquire);
+		for (;;)
+		{
+			uint32 ElementIndex = uint32(IndexAndCounter >> 32);
+
+			if (LIKELY(ElementIndex == MAX_uint32))
+			{
+				// Free list is empty. Return the next index at the end of the array.
+				return NextElementIndex.fetch_add(1, std::memory_order_relaxed);
+			}
+
+			// Attempt to pop the index from the free list.
+			const uint32 Counter = uint32(IndexAndCounter) + 1;
+			const FChunkArray* Array = ChunkArray.load(std::memory_order_acquire);
+			const FChunk* Chunk = Array->Chunks[ElementIndex / ElementsPerChunk].load(std::memory_order_acquire);
+			const uint32 NextFreeElementIndex = reinterpret_cast<const uint32&>(Chunk->Elements[ElementIndex % ElementsPerChunk]);
+			const uint64 NextIndexAndCounter = (uint64(NextFreeElementIndex) << 32) | Counter;
+			if (FreeElementIndexAndCounter.compare_exchange_strong(IndexAndCounter, NextIndexAndCounter, std::memory_order_acquire))
+			{
+				return ElementIndex;
+			}
+		}
+	}
+
+	FORCEINLINE FChunkArray* GetOrGrowChunkArray(const uint32 ElementIndex)
+	{
+		FChunkArray* Array = ChunkArray.load(std::memory_order_acquire);
+		if (ElementIndex < Array->ChunkCount * ElementsPerChunk)
+		{
+			return Array;
+		}
+		return GrowChunkArray(ElementIndex, Array);
+	}
+
+	FORCENOINLINE FChunkArray* GrowChunkArray(const uint32 ElementIndex, FChunkArray* Array)
+	{
+		do
+		{
+			// Grow to the next power of two past the index.
+			// Allocate a minimum of 16 chunks to avoid the overhead of frequent reallocations at small sizes.
+			const uint32 RequiredChunkCount = FMath::DivideAndRoundUp(ElementIndex + 1, ElementsPerChunk);
+			const uint32 NextChunkCount = FMath::Max(FMath::RoundUpToPowerOfTwo(RequiredChunkCount), 16u);
+			const uint32 NextArraySize = sizeof(FChunkArray) + sizeof(std::atomic<FChunk*>) * NextChunkCount;
+			FChunkArray* NextArray = new(FMemory::MallocZeroed(NextArraySize, alignof(FChunkArray))) FChunkArray;
+			NextArray->ChunkCount = NextChunkCount;
+			NextArray->Prev = Array;
+			for (uint32 ChunkIndex = 0; ChunkIndex < Array->ChunkCount; ++ChunkIndex)
+			{
+				// Force creation of every chunk in the existing array before the new array becomes visible
+				// to avoid the possibility of two threads creating the same chunk in two different arrays.
+				NextArray->Chunks[ChunkIndex].store(GetOrCreateChunk(Array, ChunkIndex), std::memory_order_relaxed);
+			}
+			if (ChunkArray.compare_exchange_strong(Array, NextArray, std::memory_order_release, std::memory_order_acquire))
+			{
+				return NextArray;
+			}
+			FMemory::Free(NextArray);
+		}
+		while (ElementIndex >= Array->ChunkCount * ElementsPerChunk);
+		return Array;
+	}
+
+	FORCEINLINE static FChunk* GetOrCreateChunk(FChunkArray* Array, uint32 ChunkIndex)
+	{
+		if (FChunk* Chunk = Array->Chunks[ChunkIndex].load(std::memory_order_acquire))
+		{
+			return Chunk;
+		}
+		return CreateChunk(Array, ChunkIndex);
+	}
+
+	FORCENOINLINE static FChunk* CreateChunk(FChunkArray* Array, uint32 ChunkIndex)
+	{
+		FChunk* ExistingChunk = nullptr;
+		FChunk* NewChunk = new FChunk;
+		if (Array->Chunks[ChunkIndex].compare_exchange_strong(ExistingChunk, NewChunk, std::memory_order_release, std::memory_order_acquire))
+		{
+			return NewChunk;
+		}
+
+		// Try to assign the chunk later in the array to minimize future conflicts and avoid wasting time reallocating.
+		for (uint32 NextIndex = ChunkIndex + 1, ChunkCount = Array->ChunkCount; NextIndex < ChunkCount; ++NextIndex)
+		{
+			// Use relaxed on failure because this does not read the existing chunk returned in the failure case.
+			FChunk* Chunk = nullptr;
+			if (Array->Chunks[NextIndex].compare_exchange_strong(Chunk, NewChunk, std::memory_order_release, std::memory_order_relaxed))
+			{
+				return ExistingChunk;
+			}
+		}
+
+		delete NewChunk;
+		return ExistingChunk;
+	}
+
+	// A shared empty chunk array that avoids requiring null checks to handle empty arrays.
+	inline static FChunkArray EmptyChunkArray;
+
+	// Pointer to the current chunk array.
+	std::atomic<FChunkArray*> ChunkArray = &EmptyChunkArray;
+	// Index of the first free element (or MAX_uint32) in the high 32 bits.
+	// Counter to avoid the ABA problem in the low 32 bits.
+	std::atomic<uint64> FreeElementIndexAndCounter = uint64(MAX_uint32) << 32;
+	// Index of the next element allocated at the end of the array.
+	// Used when there are no free elements or when allocating multiple elements at once.
+	std::atomic<uint32> NextElementIndex = 0;
+};
+
+} // UE
