DDC: Added a linear-probing hash table for internal use by DDC and DerivedDataTool

diff --git a/Engine/Source/Developer/DerivedDataCache/Internal/Containers/ProbingHashTable.h b/Engine/Source/Developer/DerivedDataCache/Internal/Containers/ProbingHashTable.h
new file mode 100644
--- /dev/null
+++ b/Engine/Source/Developer/DerivedDataCache/Internal/Containers/ProbingHashTable.h
@@ -0,0 +1,388 @@
+// Copyright Epic Games, Inc. All Rights Reserved.
+
+#pragma once
+
+#include "CoreTypes.h"
+#include "Math/UnrealMathUtility.h"
+#include "Misc/AssertionMacros.h"
+#include "Misc/ScopeExit.h"
+
+namespace UE
+{
+
+/**
+ * A linear-probing hash table that can be used to index another data structure.
+ *
+ * Example:
+ *	uint32 Index;
+ *	uint32 Hash = GetTypeHash(ID);
+ *	for (FProbingHashTable::FSlotId SlotId = HashTable.FindFirst(Hash, Index); SlotId; SlotId = HashTable.FindNext(SlotId, Hash, Index))
+ *	{
+ *		if (Array[Index].ID == ID)
+ *		{
+ *			return Array[Index];
+ *		}
+ *	}
+ */
+class FProbingHashTable
+{
+	inline constexpr static uint32 MaxLoadFactorNum = 7;
+	inline constexpr static uint32 MaxLoadFactorDen = 8;
+	static_assert(MaxLoadFactorNum < MaxLoadFactorDen);
+
+	inline constexpr static uint32 SlotsPerBlock = 7;
+
+	struct FSlot
+	{
+		uint32 Hash = 0;
+		uint32 Index = 0;
+	};
+
+	struct alignas(64) FBlock
+	{
+		// Byte 0:0 is set when probing must consider the next block if a match is not found in this block.
+		// Byte 0:1-7 are a mask of which slots are occupied in this block.
+		// Bytes 1-7 are the least significant byte of the hash of the element in the corresponding slot.
+		uint64 MaskAndHashes = 0;
+		FSlot Slots[SlotsPerBlock];
+
+		inline constexpr static uint64 ProbeMask = 0x01;
+		inline constexpr static uint64 SlotMask = 0xfe;
+		inline constexpr static uint64 HashMask = 0xffff'ffff'ffff'ff00;
+	};
+
+	static_assert(sizeof(FBlock) == 64);
+
+	// Not implemented but can be if necessary.
+	FProbingHashTable(const FProbingHashTable&) = delete;
+	FProbingHashTable& operator=(const FProbingHashTable&) = delete;
+
+public:
+	constexpr FProbingHashTable() = default;
+
+	inline FProbingHashTable(FProbingHashTable&& Other)
+		: Blocks(Other.Blocks)
+		, BlockCountMinusOne(Other.BlockCountMinusOne)
+		, FreeSlotCount(Other.FreeSlotCount)
+	{
+		Other.Blocks = nullptr;
+		Other.BlockCountMinusOne = MAX_uint32;
+		Other.FreeSlotCount = 0;
+	}
+
+	inline FProbingHashTable& operator=(FProbingHashTable&& Other)
+	{
+		delete[] Blocks;
+
+		Blocks = Other.Blocks;
+		BlockCountMinusOne = Other.BlockCountMinusOne;
+		FreeSlotCount = Other.FreeSlotCount;
+
+		Other.Blocks = nullptr;
+		Other.BlockCountMinusOne = MAX_uint32;
+		Other.FreeSlotCount = 0;
+
+		return *this;
+	}
+
+	inline ~FProbingHashTable()
+	{
+		delete[] Blocks;
+	}
+
+	/** Remove everything from the hash table and release its memory. */
+	UE_REWRITE void Empty()
+	{
+		*this = {};
+	}
+
+	/** A reference to a slot within the hash table. */
+	struct FSlotId
+	{
+		[[nodiscard]] inline explicit operator bool() const
+		{
+			return SlotIndexPlusOne != 0;
+		}
+
+	private:
+		[[nodiscard]] inline static FSlotId Create(uint32 InBlockIndex, uint32 InSlotIndexPlusOne)
+		{
+			FSlotId Id;
+			Id.BlockIndex = InBlockIndex;
+			Id.SlotIndexPlusOne = InSlotIndexPlusOne;
+			return Id;
+		}
+
+		uint32 BlockIndex : 29 = 0;
+		uint32 SlotIndexPlusOne : 3 = 0;
+		static_assert(SlotsPerBlock <= 7);
+
+		friend FProbingHashTable;
+	};
+
+	/**
+	 * Add the pair of Hash and Index. Does not check for duplicates.
+	 *
+	 * @return The ID of the slot containing the pair. Valid until the next call to Add.
+	 */
+	FSlotId Add(uint32 Hash, uint32 Index)
+	{
+		if (UNLIKELY(FreeSlotCount == 0))
+		{
+			ExpandCapacity(1);
+		}
+
+		--FreeSlotCount;
+		for (uint32 BlockIndex = Hash & BlockCountMinusOne;; BlockIndex = (BlockIndex + 1) & BlockCountMinusOne)
+		{
+			FBlock& Block = Blocks[BlockIndex];
+			if (const uint64 AvailableSlots = ~Block.MaskAndHashes & FBlock::SlotMask)
+			{
+				const uint32 SlotIndexPlusOne = FMath::CountTrailingZeros((uint32)AvailableSlots);
+				Block.MaskAndHashes |= (uint64(1) << SlotIndexPlusOne) | (uint64(uint8(Hash)) << (SlotIndexPlusOne * 8));
+				Block.Slots[SlotIndexPlusOne - 1] = {Hash, Index};
+				return FSlotId::Create(BlockIndex, SlotIndexPlusOne);
+			}
+			Block.MaskAndHashes |= FBlock::ProbeMask;
+		}
+	}
+
+	/**
+	 * Remove the pair of hash and index in the slot referenced by ID.
+	 *
+	 * The ID remains valid for subsequent calls to FindNext.
+	 */
+	inline void Remove(FSlotId Id)
+	{
+		checkSlow(Id.BlockIndex <= BlockCountMinusOne);
+		checkSlow(Id.SlotIndexPlusOne >= 1 && Id.SlotIndexPlusOne <= SlotsPerBlock);
+		checkSlow(Blocks && (Blocks[Id.BlockIndex].MaskAndHashes & (uint64(1) << Id.SlotIndexPlusOne)));
+		FBlock& Block = Blocks[Id.BlockIndex];
+		Block.MaskAndHashes &= ~((uint64(1) << Id.SlotIndexPlusOne) | (uint64(0xff) << (Id.SlotIndexPlusOne * 8)));
+		++FreeSlotCount;
+	}
+
+	/**
+	 * Find the first slot that matches Hash and assign its index to OutIndex.
+	 *
+	 * @return The ID of the matching slot, or an invalid ID. Valid until the next call to Add.
+	 */
+	[[nodiscard]] inline FSlotId FindFirst(uint32 Hash, uint32& OutIndex) const
+	{
+		if (UNLIKELY(BlockCountMinusOne == MAX_uint32))
+		{
+			return {};
+		}
+		const uint32 BlockIndex = Hash & BlockCountMinusOne;
+		return FindNext(FSlotId::Create(BlockIndex, 0), Hash, OutIndex);
+	}
+
+	/**
+	 * Find the next slot after Id that matches Hash and assign its index to OutIndex.
+	 *
+	 * @return The ID of the matching slot, or an invalid ID. Valid until the next call to Add.
+	 */
+	[[nodiscard]] FSlotId FindNext(FSlotId Id, uint32 Hash, uint32& OutIndex) const
+	{
+		checkSlow(Id.BlockIndex <= BlockCountMinusOne);
+
+		// Splat the least significant byte of Hash across a uint64.
+		uint64 HashLsb = uint8(Hash);
+		HashLsb |= HashLsb << 32;
+		HashLsb |= HashLsb << 16;
+		HashLsb |= HashLsb << 8;
+
+		uint32 BlockIndex = Id.BlockIndex;
+		uint32 SlotIndexPlusOne = Id.SlotIndexPlusOne;
+		uint32 SlotMask = (uint32(0xff) << SlotIndexPlusOne) & 0xff;
+
+		// Advance to the next slot.
+		++SlotIndexPlusOne;
+		SlotMask &= SlotMask - 1;
+
+		for (;;)
+		{
+			const FBlock& Block = Blocks[BlockIndex];
+			if (Block.MaskAndHashes & SlotMask)
+			{
+				// Compare by XOR which leaves a 0 byte for each matching hash.
+				uint64 MatchingSlotMask = Block.MaskAndHashes ^ HashLsb;
+				// Propagate 1 bits in each byte into the least significant bit.
+				MatchingSlotMask |= MatchingSlotMask >> 4;
+				MatchingSlotMask |= MatchingSlotMask >> 2;
+				MatchingSlotMask |= MatchingSlotMask >> 1;
+				// Matching bytes have LSB=0. Negate the bits and mask out everything else.
+				MatchingSlotMask = ~MatchingSlotMask & 0x0101010101010101;
+				// Move the bits for matching slots from one byte each into one bit each in the least significant byte.
+				MatchingSlotMask = (MatchingSlotMask * 0x0102040810204080) >> 56;
+				// Mask out slots based on the starting slot index.
+				MatchingSlotMask &= uint64(0xff) << SlotIndexPlusOne;
+				// Mask out unoccupied slots.
+				MatchingSlotMask &= Block.MaskAndHashes;
+
+				// Compare the entire hash for any slots that matched the LSB of Hash.
+				for (; MatchingSlotMask; MatchingSlotMask &= MatchingSlotMask - 1)
+				{
+					SlotIndexPlusOne = FMath::CountTrailingZeros((uint32)MatchingSlotMask);
+					const FSlot& Slot = Block.Slots[SlotIndexPlusOne - 1];
+					if (Slot.Hash == Hash)
+					{
+						OutIndex = Slot.Index;
+						return FSlotId::Create(BlockIndex, SlotIndexPlusOne);
+					}
+				}
+			}
+
+			// Advance to the next block only if the probe flag has been set.
+			if (Block.MaskAndHashes & FBlock::ProbeMask)
+			{
+				BlockIndex = (BlockIndex + 1) & BlockCountMinusOne;
+				if (BlockIndex != Id.BlockIndex)
+				{
+					SlotIndexPlusOne = 1;
+					SlotMask = FBlock::SlotMask;
+					continue;
+				}
+			}
+
+			return {};
+		}
+	}
+
+	/**
+	 * Access the index in the slot referenced by ID.
+	 */
+	[[nodiscard]] inline uint32& operator[](FSlotId Id)
+	{
+		checkSlow(Id.BlockIndex <= BlockCountMinusOne);
+		checkSlow(Id.SlotIndexPlusOne >= 1 && Id.SlotIndexPlusOne <= SlotsPerBlock);
+		checkSlow(Blocks && (Blocks[Id.BlockIndex].MaskAndHashes & (uint64(1) << Id.SlotIndexPlusOne)));
+		return Blocks[Id.BlockIndex].Slots[Id.SlotIndexPlusOne - 1].Index;
+	}
+
+	/**
+	 * Resize if needed to ensure that Capacity additional elements can be added without rehashing.
+	 *
+	 * Like Reserve() on other containers but takes *additional* capacity since neither Num() or Max()
+	 * are provided by this container.
+	 */
+	inline void EnsureAdditionalCapacity(uint32 Capacity)
+	{
+		if (Capacity > FreeSlotCount)
+		{
+			ExpandCapacity(Capacity - FreeSlotCount);
+		}
+	}
+
+	class FConstIterator
+	{
+	public:
+		constexpr FConstIterator() = default;
+
+		inline explicit FConstIterator(const FProbingHashTable& HashTable)
+			: Block(HashTable.Blocks)
+			, BlockCount(HashTable.BlockCountMinusOne + 1)
+			, SlotMask(BlockCount ? (Block->MaskAndHashes & FBlock::SlotMask) : 0)
+		{
+			++*this;
+		}
+
+		inline FConstIterator& operator++()
+		{
+			if (BlockCount)
+			{
+				while (SlotMask == 0)
+				{
+					if (--BlockCount == 0)
+					{
+						Block = nullptr;
+						SlotIndex = 0;
+						return *this;
+					}
+					++Block;
+					SlotMask = Block->MaskAndHashes & FBlock::SlotMask;
+				}
+				SlotIndex = (uint8)FMath::CountTrailingZeros(SlotMask) - 1;
+				SlotMask &= SlotMask - 1;
+			}
+			return *this;
+		}
+
+		[[nodiscard]] inline explicit operator bool() const
+		{
+			return BlockCount > 0;
+		}
+
+		[[nodiscard]] inline bool operator==(const FConstIterator& Other) const
+		{
+			return Block == Other.Block && SlotIndex == Other.SlotIndex;
+		}
+
+		[[nodiscard]] UE_REWRITE uint32 operator*() const
+		{
+			return GetIndex();
+		}
+
+		[[nodiscard]] inline uint32 GetHash() const
+		{
+			return Block->Slots[SlotIndex].Hash;
+		}
+
+		[[nodiscard]] inline uint32 GetIndex() const
+		{
+			return Block->Slots[SlotIndex].Index;
+		}
+
+	private:
+		const FBlock* Block = nullptr;
+		uint32 BlockCount = 0;
+		uint8 SlotIndex = 0;
+		uint8 SlotMask = 0;
+	};
+
+	[[nodiscard]] UE_REWRITE FConstIterator CreateConstIterator() const
+	{
+		return FConstIterator(*this);
+	}
+
+private:
+	/** Expand capacity by at least SlotCount free slots, ignoring existing free slots. */
+	FORCENOINLINE void ExpandCapacity(const uint32 SlotCount)
+	{
+		FConstIterator It(*this);
+		FBlock* OldBlocks = Blocks;
+		ON_SCOPE_EXIT { delete[] OldBlocks; };
+
+		const uint32 MinCapacity = CalculateCapacity(BlockCountMinusOne + 1) + SlotCount;
+		const uint32 MinBlockCount = FMath::DivideAndRoundUp(MinCapacity + MinCapacity * (MaxLoadFactorDen - MaxLoadFactorNum) / MaxLoadFactorNum, SlotsPerBlock);
+		// Allocate at least 16 blocks and always round up to the next power of two to optimize calculations.
+		const uint32 BlockCount = FMath::Max(16u, FMath::RoundUpToPowerOfTwo(MinBlockCount));
+		BlockCountMinusOne = BlockCount - 1;
+		Blocks = new FBlock[BlockCount];
+
+		FreeSlotCount = CalculateCapacity(BlockCount);
+
+		// Copy everything from the old blocks to the new blocks.
+		for (; It; ++It)
+		{
+			Add(It.GetHash(), It.GetIndex());
+		}
+	}
+
+	/** Calculate the effective capacity from BlockCount, accounting for the max load factor. */
+	constexpr static uint32 CalculateCapacity(uint32 BlockCount)
+	{
+		constexpr uint32_t Multiplier = SlotsPerBlock * MaxLoadFactorNum / MaxLoadFactorDen;
+		constexpr uint32_t Remainder = SlotsPerBlock * MaxLoadFactorNum - Multiplier * MaxLoadFactorDen;
+		return BlockCount * Multiplier + BlockCount * Remainder / MaxLoadFactorDen;
+	}
+
+	[[nodiscard]] UE_REWRITE friend FConstIterator begin(const FProbingHashTable& HashTable) { return FConstIterator(HashTable); }
+	[[nodiscard]] constexpr friend FConstIterator end(const FProbingHashTable& HashTable) { return FConstIterator(); }
+
+	FBlock* Blocks = nullptr;
+	uint32 BlockCountMinusOne = MAX_uint32;
+	uint32 FreeSlotCount = 0;
+};
+
+} // UE
