HttpCacheStore: Always do existence checks using the existence operation as the Get path can be unreliable in the specific situation where the server has a ref, but the blobs referenced from it are missing/absent.  This leads to outcomes where existence checks return true, but you go to get the data and it can't be found.

Found and fixed edge cases in the results surfaced by the unit tests.  Also cleaned up some blob existence check code that was no longer called.

#Synced-CL 49552467

diff --git a/Engine/Source/Developer/DerivedDataCache/Private/HttpCacheStore.cpp b/Engine/Source/Developer/DerivedDataCache/Private/HttpCacheStore.cpp
--- a/Engine/Source/Developer/DerivedDataCache/Private/HttpCacheStore.cpp
+++ b/Engine/Source/Developer/DerivedDataCache/Private/HttpCacheStore.cpp
@@ -1152,7 +1152,6 @@ public:
 	using FOnValueComplete = TUniqueFunction<void (FValueResponse&& Response)>;
 
 	void GetValues(TConstArrayView<FValueWithId> Values, FOnValueComplete&& OnComplete);
-	void GetValuesExist(TConstArrayView<FValueWithId> Values, FOnValueComplete&& OnComplete);
 
 	const FRequestStats& ReadStats() const { return RequestStats; }
 	FRequestStats& EditStats() { return RequestStats; }
@@ -1174,9 +1173,6 @@ private:
 	void BeginGetValue(TUniquePtr<FHttpOperation>&& Operation, const FValueWithId& Value, const TSharedRef<FOnValueComplete>& OnComplete);
 	void EndGetValue(FHttpOperation& Operation, const FValueWithId& Value, const FOnValueComplete& OnComplete);
 
-	void BeginGetValuesExist(TUniquePtr<FHttpOperation>&& Operation, TArray<FValueWithId>&& Values, FOnValueComplete&& OnComplete);
-	void EndGetValuesExist(FHttpOperation* Operation, TArray<FValueWithId>&& Values, FOnValueComplete&& OnComplete);
-
 	FHttpCacheStore& CacheStore;
 	IRequestOwner& Owner;
 	FSharedString Name;
@@ -1603,135 +1599,6 @@ void FHttpCacheStore::FGetRecordOp::EndGetValue(FHttpOperation& Operation, const
 	}
 }
 
-void FHttpCacheStore::FGetRecordOp::GetValuesExist(TConstArrayView<FValueWithId> Values, FOnValueComplete&& OnComplete)
-{
-	TArray<FValueWithId> QueryValues;
-	for (const FValueWithId& Value : Values)
-	{
-		if (Value.HasData())
-		{
-			OnComplete({Value, EStatus::Ok});
-			continue;
-		}
-		QueryValues.Emplace(Value);
-	}
-
-	if (QueryValues.IsEmpty())
-	{
-		return;
-	}
-
-	FRequestTimer RequestTimer(RequestStats);
-	RequestTimer.Stop();
-
-	FRequestBarrier Barrier(Owner);
-	CacheStore.WaitForHttpOperationAsync(Owner, EOperationCategory::Get, [Self = TRefCountPtr(this), Values = MoveTemp(QueryValues), OnComplete = MoveTemp(OnComplete)](TUniquePtr<FHttpOperation>&& Operation) mutable
-	{
-		Self->BeginGetValuesExist(MoveTemp(Operation), MoveTemp(Values), MoveTemp(OnComplete));
-	});
-}
-
-void FHttpCacheStore::FGetRecordOp::BeginGetValuesExist(TUniquePtr<FHttpOperation>&& Operation, TArray<FValueWithId>&& Values, FOnValueComplete&& OnComplete)
-{
-	if (UNLIKELY(!Operation))
-	{
-		EndGetValuesExist(nullptr, MoveTemp(Values), MoveTemp(OnComplete));
-		return;
-	}
-
-	FRequestTimer RequestTimer(RequestStats);
-
-	TAnsiStringBuilder<256> Uri;
-	Uri << CacheStore.EffectiveDomain << ANSITEXTVIEW("/api/v1/compressed-blobs/") << CacheStore.Namespace << ANSITEXTVIEW("/exists?");
-	for (const FValueWithId& Value : Values)
-	{
-		Uri << ANSITEXTVIEW("id=") << Value.GetRawHash() << '&';
-	}
-	Uri.RemoveSuffix(1);
-
-	FHttpOperation& LocalOperation = *Operation;
-	LocalOperation.SetUri(Uri);
-	LocalOperation.SetUnixSocketPath(StringCast<ANSICHAR>(*CacheStore.UnixSocketPath));
-	LocalOperation.SetMethod(EHttpMethod::Post);
-	LocalOperation.SetContentType(EHttpMediaType::FormUrlEncoded);
-	LocalOperation.AddAcceptType(EHttpMediaType::Json);
-
-	RequestTimer.Stop();
-	LocalOperation.SendAsync(Owner, [Self = TRefCountPtr(this), Operation = MoveTemp(Operation), Values = MoveTemp(Values), OnComplete = MoveTemp(OnComplete)]() mutable
-	{
-		Self->EndGetValuesExist(Operation.Get(), MoveTemp(Values), MoveTemp(OnComplete));
-	});
-}
-
-void FHttpCacheStore::FGetRecordOp::EndGetValuesExist(FHttpOperation* Operation, TArray<FValueWithId>&& Values, FOnValueComplete&& OnComplete)
-{
-	TRACE_CPUPROFILER_EVENT_SCOPE(HttpDDC_DataProbablyExistsBatch_OnHttpRequestComplete);
-
-	FRequestTimer RequestTimer(RequestStats);
-
-	if (Operation)
-	{
-		Operation->GetStats(RequestStats);
-	}
-
-	const TCHAR* DefaultMessage = TEXT("Cache exists miss for");
-	EStatus DefaultStatus = EStatus::Error;
-
-	if (!Operation || Operation->GetErrorCode() == EHttpErrorCode::Canceled)
-	{
-		DefaultMessage = TEXT("Cache exists miss with canceled request for");
-		DefaultStatus = EStatus::Canceled;
-	}
-	else if (Operation->GetErrorCode() != EHttpErrorCode::None)
-	{
-		DefaultMessage = TEXT("Cache exists miss with error or retry exhaustion during request for");
-		DefaultStatus = EStatus::Error;
-	}
-	else if (const int32 StatusCode = Operation->GetStatusCode(); StatusCode < 200 || StatusCode > 204)
-	{
-		DefaultMessage = TEXT("Cache exists miss with failed response for");
-	}
-	else if (TSharedPtr<FJsonObject> ResponseObject = Operation->GetBodyAsJson(); !ResponseObject)
-	{
-		DefaultMessage = TEXT("Cache exists miss with invalid response for");
-	}
-	else if (TArray<FString> NeedsArrayStrings; ResponseObject->TryGetStringArrayField(TEXT("needs"), NeedsArrayStrings))
-	{
-		DefaultMessage = TEXT("Cache exists hit for");
-		DefaultStatus = EStatus::Ok;
-
-		for (const FString& NeedsString : NeedsArrayStrings)
-		{
-			const FIoHash NeedHash(NeedsString);
-			for (auto It = Values.CreateIterator(); It; ++It)
-			{
-				const FValueWithId& Value = *It;
-				if (Value.GetRawHash() == NeedHash)
-				{
-					UE_LOG(LogDerivedDataCache, Verbose,
-						TEXT("%s: Cache exists miss with missing value %s with hash %s for %s from '%s'"),
-						*CacheStore.NodeName, *WriteToString<32>(Value.GetId()),
-						*WriteToString<48>(Value.GetRawHash()), *WriteToString<96>(Key), *Name);
-					OnComplete({Value, EStatus::Error});
-					It.RemoveCurrentSwap();
-					break;
-				}
-			}
-		}
-	}
-
-	RequestTimer.Stop();
-
-	for (const FValueWithId& Value : Values)
-	{
-		UE_LOG(LogDerivedDataCache, Verbose,
-			TEXT("%s: %s value %s with hash %s for %s from '%s'"),
-			*CacheStore.NodeName, DefaultMessage, *WriteToString<32>(Value.GetId()),
-			*WriteToString<48>(Value.GetRawHash()), *WriteToString<96>(Key), *Name);
-		OnComplete({Value, DefaultStatus});
-	}
-}
-
 void FHttpCacheStore::FGetRecordOp::RecordStats(EStatus Status)
 {
 	RequestStats.EndTime = FMonotonicTimePoint::Now();
@@ -2297,16 +2164,26 @@ void FHttpCacheStore::FExistsChunkBatchOp::EndRequest(const FCacheGetChunkReques
 			Value = Record.Get().GetValue(Request.Id);
 			if (!Value || Value.GetRawHash().IsZero())
 			{
-				UE_LOG(LogDerivedDataCache, Display, TEXT("%s: Cache miss with missing/invalid value for %s from '%s'"),
-					*CacheStore.NodeName, *WriteToString<96>(Request.Key), *Request.Name);
+				UE_LOG(LogDerivedDataCache, Display, TEXT("%s: Cache miss with missing/invalid value %s for %s from '%s'"),
+					*CacheStore.NodeName, *WriteToString<24>(Request.Id), *WriteToString<96>(Request.Key), *Request.Name);
+				Status = EStatus::Error;
+			}
+			else if (!Request.RawHash.IsZero() && Request.RawHash != Value.GetRawHash())
+			{
+				UE_LOG(LogDerivedDataCache, Display, TEXT("%s: Cache miss with mismatched value %s received hash %s when expected hash %s for %s from '%s'"),
+					*CacheStore.NodeName, *WriteToString<24>(Request.Id), *WriteToString<40>(Value.GetRawHash()), *WriteToString<40>(Request.RawHash), *WriteToString<96>(Request.Key), *Request.Name);
+				Value.Reset();
 				Status = EStatus::Error;
 			}
 		}
 
+		const uint64 LogicalRawOffset = FMath::Min(Value.GetRawSize(), Request.RawOffset);
+		const uint64 ResponseRawSize = FMath::Min(Value.GetRawSize() - LogicalRawOffset, Request.RawSize);
+
 		RequestStats.Type = ERequestType::Record;
 		RequestStats.EndTime = FMonotonicTimePoint::Now();
 		RequestStats.Status = Status;
-		OnComplete({Request.Name, Request.Key, Request.Id, Request.RawOffset, Value.GetRawSize(), Value.GetRawHash(), {}, Request.UserData, Status});
+		OnComplete({Request.Name, Request.Key, Request.Id, Request.RawOffset, ResponseRawSize, Value.GetRawHash(), {}, Request.UserData, Status});
 		if (CacheStore.StoreStats)
 		{
 			CacheStore.StoreStats->AddRequest(RequestStats);
@@ -2321,8 +2198,14 @@ void FHttpCacheStore::FExistsChunkBatchOp::EndRequest(const FCacheGetChunkReques
 			const uint64 RawSize = ResponseView[ANSITEXTVIEW("RawSize")].AsUInt64(MAX_uint64);
 			if (RawHash.IsZero() || RawSize == MAX_uint64)
 			{
-				UE_LOG(LogDerivedDataCache, Display, TEXT("%s: Cache miss with invalid value for %s from '%s'"),
-					*CacheStore.NodeName, *WriteToString<96>(Request.Key), *Request.Name);
+				UE_LOG(LogDerivedDataCache, Display, TEXT("%s: Cache miss with invalid value %s for %s from '%s'"),
+					*CacheStore.NodeName, *WriteToString<24>(Request.Id), *WriteToString<96>(Request.Key), *Request.Name);
+				Status = EStatus::Error;
+			}
+			else if (!Request.RawHash.IsZero() && Request.RawHash != RawHash)
+			{
+				UE_LOG(LogDerivedDataCache, Display, TEXT("%s: Cache miss with mismatched value received hash %s when expected hash %s for %s from '%s'"),
+					*CacheStore.NodeName, *WriteToString<40>(RawHash), *WriteToString<40>(Request.RawHash), *WriteToString<96>(Request.Key), *Request.Name);
 				Status = EStatus::Error;
 			}
 			else
@@ -2330,11 +2213,14 @@ void FHttpCacheStore::FExistsChunkBatchOp::EndRequest(const FCacheGetChunkReques
 				Value = FValue(RawHash, RawSize);
 			}
 		}
+		
+		const uint64 LogicalRawOffset = FMath::Min(Value.GetRawSize(), Request.RawOffset);
+		const uint64 ResponseRawSize = FMath::Min(Value.GetRawSize() - LogicalRawOffset, Request.RawSize);
 
 		RequestStats.Type = ERequestType::Value;
 		RequestStats.EndTime = FMonotonicTimePoint::Now();
 		RequestStats.Status = Status;
-		OnComplete({Request.Name, Request.Key, Request.Id, Request.RawOffset, Value.GetRawSize(), Value.GetRawHash(), {}, Request.UserData, Status});
+		OnComplete({Request.Name, Request.Key, Request.Id, Request.RawOffset, ResponseRawSize, Value.GetRawHash(), {}, Request.UserData, Status});
 		if (CacheStore.StoreStats)
 		{
 			CacheStore.StoreStats->AddRequest(RequestStats);
@@ -3189,14 +3075,23 @@ void FHttpCacheStore::Get(
 	TRACE_CPUPROFILER_EVENT_SCOPE(HttpDDC_Get);
 	TRACE_COUNTER_ADD(HttpDDC_Get, Requests.Num());
 
-	// NOTE: Both the 'if' and the 'else' branch handle SkipData - see how FHttpCacheStore::FGetRecordOp::BeginGetValues chooses to send
-	//		 value queries depending on whether the request has SkipData - we just choose to use the batch existence operation so we avoid
-	//		 multiple round trips when there are multiple requests and all of them are existence queries (on the record policy).
 	const auto HasSkipData = [](const FCacheGetRequest& Request) { return EnumHasAnyFlags(Request.Policy.GetRecordPolicy(), ECachePolicy::SkipData); };
-	if ((Requests.Num() > 1) && Algo::AllOf(Requests, HasSkipData))
+
+	TArray<FCacheGetRequest, TInlineAllocator<16>> ExistRequests;
+	for (const FCacheGetRequest& Request : Requests)
+	{
+		if (HasSkipData(Request))
+		{
+			ExistRequests.Emplace(Request);
+		}
+	}
+
+	FRequestBarrier Barrier(Owner);
+	TSharedRef<FOnCacheGetComplete> SharedOnComplete = MakeShared<FOnCacheGetComplete>(MoveTemp(OnComplete));
+	if (!ExistRequests.IsEmpty())
 	{
 		TRefCountPtr<FExistsRecordBatchOp> Op = FExistsRecordBatchOp::New(*this, Owner);
-		Op->Exists(Requests, [this, OnComplete = MoveTemp(OnComplete)](FCacheGetResponse&& Response)
+		Op->Exists(ExistRequests, [this, SharedOnComplete](FCacheGetResponse&& Response)
 		{
 			if (Response.Status == EStatus::Ok)
 			{
@@ -3204,14 +3099,13 @@ void FHttpCacheStore::Get(
 				UE_LOG(LogDerivedDataCache, Verbose, TEXT("%s: Cache hit for %s from '%s'"),
 					*NodeName, *WriteToString<96>(Response.Record.GetKey()), *Response.Name);
 			}
-			OnComplete(MoveTemp(Response));
+			SharedOnComplete.Get()(MoveTemp(Response));
 		});
 	}
-	else
+
+	for (const FCacheGetRequest& Request : Requests)
 	{
-		FRequestBarrier Barrier(Owner);
-		TSharedRef<FOnCacheGetComplete> SharedOnComplete = MakeShared<FOnCacheGetComplete>(MoveTemp(OnComplete));
-		for (const FCacheGetRequest& Request : Requests)
+		if (!HasSkipData(Request))
 		{
 			GetCacheRecordAsync(Owner, Request.Name, Request.Key, Request.Policy, Request.UserData,
 				[this, SharedOnComplete](FCacheGetResponse&& Response)
@@ -3259,15 +3153,23 @@ void FHttpCacheStore::GetValue(
 	TRACE_CPUPROFILER_EVENT_SCOPE(HttpDDC_GetValue);
 	TRACE_COUNTER_ADD(HttpDDC_Get, Requests.Num());
 
-	// NOTE: Both the 'if' and the 'else' branch handle SkipData - see how FHttpCacheStore::FGetValueOp::BeginGetRef chooses to send
-	//		 or not send x-jupiter-inline as an accept type depending on whether the request has SkipData - we just choose to use
-	//		 the batch existence operation so we avoid multiple round trips when there are multiple requests and all of them are
-	//		 existence queries.
 	const auto HasSkipData = [](const FCacheGetValueRequest& Request) { return EnumHasAnyFlags(Request.Policy, ECachePolicy::SkipData); };
-	if ((Requests.Num() > 1) && Algo::AllOf(Requests, HasSkipData))
+
+	TArray<FCacheGetValueRequest, TInlineAllocator<16>> ExistRequests;
+	for (const FCacheGetValueRequest& Request : Requests)
+	{
+		if (HasSkipData(Request))
+		{
+			ExistRequests.Emplace(Request);
+		}
+	}
+
+	FRequestBarrier Barrier(Owner);
+	TSharedRef<FOnCacheGetValueComplete> SharedOnComplete = MakeShared<FOnCacheGetValueComplete>(MoveTemp(OnComplete));
+	if (!ExistRequests.IsEmpty())
 	{
 		TRefCountPtr<FExistsValueBatchOp> Op = FExistsValueBatchOp::New(*this, Owner);
-		Op->Exists(Requests, [this, OnComplete = MoveTemp(OnComplete)](FCacheGetValueResponse&& Response)
+		Op->Exists(ExistRequests, [this, SharedOnComplete](FCacheGetValueResponse&& Response)
 		{
 			if (Response.Status == EStatus::Ok)
 			{
@@ -3275,14 +3177,13 @@ void FHttpCacheStore::GetValue(
 				UE_LOG(LogDerivedDataCache, Verbose, TEXT("%s: Cache hit for %s from '%s'"),
 					*NodeName, *WriteToString<96>(Response.Key), *Response.Name);
 			}
-			OnComplete(MoveTemp(Response));
+			SharedOnComplete.Get()(MoveTemp(Response));
 		});
 	}
-	else
+
+	for (const FCacheGetValueRequest& Request : Requests)
 	{
-		FRequestBarrier Barrier(Owner);
-		TSharedRef<FOnCacheGetValueComplete> SharedOnComplete = MakeShared<FOnCacheGetValueComplete>(MoveTemp(OnComplete));
-		for (const FCacheGetValueRequest& Request : Requests)
+		if (!HasSkipData(Request))
 		{
 			GetCacheValueAsync(Owner, Request.Name, Request.Key, Request.Policy, ERequestOp::Get, Request.UserData,
 			[this, Policy = Request.Policy, SharedOnComplete](FCacheGetValueResponse&& Response)
@@ -3325,14 +3226,23 @@ void FHttpCacheStore::GetChunks(
 		return;
 	}
 
-	// NOTE: Both the 'if' and the 'else' branch handle SkipData - see how FHttpCacheStore::FGetRecordOp::BeginGetValues chooses to send
-	//		 value queries depending on whether the request has SkipData - we just choose to use the batch existence operation so we avoid
-	//		 multiple round trips when there are multiple requests of uniform policy and all of them are existence queries.
 	const auto HasSkipData = [](const FCacheGetChunkRequest& Request) { return EnumHasAnyFlags(Request.Policy, ECachePolicy::SkipData); };
-	if ((Requests.Num() > 1) && Algo::AllOf(Requests, HasSkipData))
+
+	TArray<FCacheGetChunkRequest, TInlineAllocator<16>> ExistRequests;
+	for (const FCacheGetChunkRequest& Request : Requests)
+	{
+		if (HasSkipData(Request))
+		{
+			ExistRequests.Emplace(Request);
+		}
+	}
+
+	FRequestBarrier Barrier(Owner);
+	TSharedRef<FOnCacheGetChunkComplete> SharedOnComplete = MakeShared<FOnCacheGetChunkComplete>(MoveTemp(OnComplete));
+	if (!ExistRequests.IsEmpty())
 	{
 		TRefCountPtr<FExistsChunkBatchOp> Op = FExistsChunkBatchOp::New(*this, Owner);
-		Op->Exists(Requests, [this, OnComplete = MoveTemp(OnComplete)](FCacheGetChunkResponse&& Response)
+		Op->Exists(ExistRequests, [this, SharedOnComplete](FCacheGetChunkResponse&& Response)
 		{
 			if (Response.Status == EStatus::Ok)
 			{
@@ -3340,10 +3250,12 @@ void FHttpCacheStore::GetChunks(
 				UE_LOG(LogDerivedDataCache, Verbose, TEXT("%s: Cache hit for %s from '%s'"),
 					*NodeName, *WriteToString<96>(Response.Key), *Response.Name);
 			}
-			OnComplete(MoveTemp(Response));
+			SharedOnComplete.Get()(MoveTemp(Response));
 		});
 	}
-	else
+
+	// Only do the following if there are any non-existence requests
+	if (ExistRequests.Num() != Requests.Num())
 	{
 		// TODO: This is inefficient because Jupiter doesn't allow us to get only part of a compressed blob, so we have to
 		//		 get the whole thing and then decompress only the portion we need.  Furthermore, because there is no propagation
@@ -3352,13 +3264,19 @@ void FHttpCacheStore::GetChunks(
 		//		 are missing in local/fast stores and have to be retrieved from slow stores into record requests instead.  That
 		//		 will make this code path unused/uncommon as Jupiter will most always be a slow store with a local/fast store in front of it.
 		//		 Regardless, to adhere to the functional contract, this implementation must exist.
-		TArray<FCacheGetChunkRequest, TInlineAllocator<16>> SortedRequests(Requests);
+		TArray<FCacheGetChunkRequest, TInlineAllocator<16>> SortedRequests;
+		SortedRequests.Reserve(Requests.Num() - ExistRequests.Num());
+		for (const FCacheGetChunkRequest& Request : Requests)
+		{
+			if (!HasSkipData(Request))
+			{
+				SortedRequests.Emplace(Request);
+			}
+		}
 		SortedRequests.StableSort(TChunkLess());
 
 		const FCacheGetChunkRequest* PendingGroupStartRequest = &SortedRequests[0];
 
-		FRequestBarrier Barrier(Owner);
-		TSharedRef<FOnCacheGetChunkComplete> SharedOnComplete = MakeShared<FOnCacheGetChunkComplete>(MoveTemp(OnComplete));
 		for (const FCacheGetChunkRequest& Request : SortedRequests)
 		{
 			const bool bMatchesExistingGroup = PendingGroupStartRequest != nullptr && PendingGroupStartRequest->Key == Request.Key && PendingGroupStartRequest->Id.IsValid() == Request.Id.IsValid();
