DDC: Added a constructor for FCacheRecord from FCacheKey

#Synced-CL 50002741

diff --git a/Engine/Source/Developer/DerivedDataCache/Private/DerivedDataCache.cpp b/Engine/Source/Developer/DerivedDataCache/Private/DerivedDataCache.cpp
--- a/Engine/Source/Developer/DerivedDataCache/Private/DerivedDataCache.cpp
+++ b/Engine/Source/Developer/DerivedDataCache/Private/DerivedDataCache.cpp
@@ -234,7 +234,7 @@ FCachePutResponse FCachePutRequest::MakeResponse(const EStatus Status) const
 
 FCacheGetResponse FCacheGetRequest::MakeResponse(const EStatus Status) const
 {
-	return {Name, FCacheRecordBuilder(Key).Build(), UserData, Status};
+	return {Name, FCacheRecord(Key), UserData, Status};
 }
 
 FCachePutValueResponse FCachePutValueRequest::MakeResponse(const EStatus Status) const
diff --git a/Engine/Source/Developer/DerivedDataCache/Private/DerivedDataCacheRecord.cpp b/Engine/Source/Developer/DerivedDataCache/Private/DerivedDataCacheRecord.cpp
--- a/Engine/Source/Developer/DerivedDataCache/Private/DerivedDataCacheRecord.cpp
+++ b/Engine/Source/Developer/DerivedDataCache/Private/DerivedDataCacheRecord.cpp
@@ -166,6 +166,11 @@ FCacheRecord CreateCacheRecord(ICacheRecordInternal* Record)
 namespace UE::DerivedData
 {
 
+FCacheRecord::FCacheRecord(const FCacheKey& Key)
+	: Record(new Private::FCacheRecordInternal(Key, {}, {}))
+{
+}
+
 FCbPackage FCacheRecord::Save() const
 {
 	FCbPackage Package;
diff --git a/Engine/Source/Developer/DerivedDataCache/Private/FileSystemCacheStore.cpp b/Engine/Source/Developer/DerivedDataCache/Private/FileSystemCacheStore.cpp
--- a/Engine/Source/Developer/DerivedDataCache/Private/FileSystemCacheStore.cpp
+++ b/Engine/Source/Developer/DerivedDataCache/Private/FileSystemCacheStore.cpp
@@ -1672,7 +1672,7 @@ void FFileSystemCacheStore::Get(
 			}
 			else
 			{
-				Record = FCacheRecordBuilder(Request.Key).Build();
+				Record = FCacheRecord(Request.Key);
 			}
 		}
 		RequestStats.AddLogicalRead(Record.Get());
@@ -2123,36 +2123,30 @@ FOptionalCacheRecord FFileSystemCacheStore::GetCacheRecord(
 
 	OutStatus = EStatus::Ok;
 
-	FCacheRecordBuilder RecordBuilder(Key);
+	FCbObject Meta;
+	TArray<FValueWithId> Values(Record.Get().GetValues());
 
 	const ECachePolicy RecordPolicy = Policy.GetRecordPolicy();
 	if (!EnumHasAnyFlags(RecordPolicy, ECachePolicy::SkipMeta))
 	{
-		RecordBuilder.SetMeta(FCbObject(Record.Get().GetMeta()));
+		Meta = Record.Get().GetMeta();
 	}
 
-	for (const FValueWithId& Value : Record.Get().GetValues())
+	for (FValueWithId& Value : Values)
 	{
 		const FValueId& Id = Value.GetId();
 		const ECachePolicy ValuePolicy = Policy.GetValuePolicy(Id);
-		FValue Content;
-		if (GetCacheContent(Name, Key, Id, Value, ValuePolicy, Content, Stats))
-		{
-			RecordBuilder.AddValue(Id, MoveTemp(Content));
-		}
-		else if (EnumHasAnyFlags(RecordPolicy, ECachePolicy::PartialRecord))
-		{
-			OutStatus = EStatus::Error;
-			RecordBuilder.AddValue(Value);
-		}
-		else
+		if (!GetCacheContent(Name, Key, Id, Value, ValuePolicy, Value, Stats))
 		{
 			OutStatus = EStatus::Error;
-			return FOptionalCacheRecord();
+			if (!EnumHasAnyFlags(RecordPolicy, ECachePolicy::PartialRecord))
+			{
+				return {};
+			}
 		}
 	}
 
-	return RecordBuilder.Build();
+	return FCacheRecord::CreateByMove(Key, MoveTemp(Meta), MoveTemp(Values));
 }
 
 bool FFileSystemCacheStore::PutCacheValue(
diff --git a/Engine/Source/Developer/DerivedDataCache/Private/HttpCacheStore.cpp b/Engine/Source/Developer/DerivedDataCache/Private/HttpCacheStore.cpp
--- a/Engine/Source/Developer/DerivedDataCache/Private/HttpCacheStore.cpp
+++ b/Engine/Source/Developer/DerivedDataCache/Private/HttpCacheStore.cpp
@@ -1206,7 +1206,7 @@ void FHttpCacheStore::FGetRecordOp::GetRecordOnly(const FCacheKey& InKey, const
 		UE_LOG(LogDerivedDataCache, VeryVerbose,
 			TEXT("%s: Skipped get of %s from '%s' because this cache store is not available"),
 			*CacheStore.NodeName, *WriteToString<96>(Key), *Name);
-		return InOnComplete({FCacheRecordBuilder(Key).Build(), EStatus::Error});
+		return InOnComplete({FCacheRecord(Key), EStatus::Error});
 	}
 
 	// Skip the request if querying the cache is disabled.
@@ -1214,14 +1214,14 @@ void FHttpCacheStore::FGetRecordOp::GetRecordOnly(const FCacheKey& InKey, const
 	{
 		UE_LOG(LogDerivedDataCache, VeryVerbose, TEXT("%s: Skipped get of %s from '%s' due to cache policy"),
 			*CacheStore.NodeName, *WriteToString<96>(Key), *Name);
-		return InOnComplete({FCacheRecordBuilder(Key).Build(), EStatus::Error});
+		return InOnComplete({FCacheRecord(Key), EStatus::Error});
 	}
 
 	if (CacheStore.DebugOptions.ShouldSimulateGetMiss(Key))
 	{
 		UE_LOG(LogDerivedDataCache, Verbose, TEXT("%s: Simulated miss for get of %s from '%s'"),
 			*CacheStore.NodeName, *WriteToString<96>(Key), *Name);
-		return InOnComplete({FCacheRecordBuilder(Key).Build(), EStatus::Error});
+		return InOnComplete({FCacheRecord(Key), EStatus::Error});
 	}
 
 	OnRecordComplete = MoveTemp(InOnComplete);
@@ -1267,7 +1267,7 @@ void FHttpCacheStore::FGetRecordOp::EndGetRef(TUniquePtr<FHttpOperation> Operati
 		Operation.Reset();
 		if (Record.IsNull())
 		{
-			Record = FCacheRecordBuilder(Key).Build();
+			Record = FCacheRecord(Key);
 		}
 		RequestTimer.Stop();
 		// Ensuring that the OnRecordComplete method is destroyed by the time we exit this method by moving it to a local scope variable
@@ -2107,7 +2107,7 @@ void FHttpCacheStore::FExistsRecordBatchOp::EndRequest(const FCacheGetRequest& R
 	RequestStats.Status = Status;
 	if (Record.IsNull())
 	{
-		Record = FCacheRecordBuilder(Request.Key).Build();
+		Record = FCacheRecord(Request.Key);
 	}
 	OnComplete({Request.Name, MoveTemp(Record).Get(), Request.UserData, Status});
 	if (CacheStore.StoreStats)
diff --git a/Engine/Source/Developer/DerivedDataCache/Private/PakFileCacheStore.cpp b/Engine/Source/Developer/DerivedDataCache/Private/PakFileCacheStore.cpp
--- a/Engine/Source/Developer/DerivedDataCache/Private/PakFileCacheStore.cpp
+++ b/Engine/Source/Developer/DerivedDataCache/Private/PakFileCacheStore.cpp
@@ -701,7 +701,7 @@ void FPakFileCacheStore::Get(
 			}
 			else
 			{
-				Record = FCacheRecordBuilder(Request.Key).Build();
+				Record = FCacheRecord(Request.Key);
 			}
 		}
 		RequestStats.AddLogicalRead(Record.Get());
diff --git a/Engine/Source/Developer/DerivedDataCache/Private/S3CacheStore.cpp b/Engine/Source/Developer/DerivedDataCache/Private/S3CacheStore.cpp
--- a/Engine/Source/Developer/DerivedDataCache/Private/S3CacheStore.cpp
+++ b/Engine/Source/Developer/DerivedDataCache/Private/S3CacheStore.cpp
@@ -1715,7 +1715,7 @@ void FS3CacheStore::Get(
 			}
 			else
 			{
-				Record = FCacheRecordBuilder(Request.Key).Build();
+				Record = FCacheRecord(Request.Key);
 			}
 		}
 		RequestStats.AddLogicalRead(Record.Get());
diff --git a/Engine/Source/Developer/DerivedDataCache/Public/DerivedDataCacheRecord.h b/Engine/Source/Developer/DerivedDataCache/Public/DerivedDataCacheRecord.h
--- a/Engine/Source/Developer/DerivedDataCache/Public/DerivedDataCacheRecord.h
+++ b/Engine/Source/Developer/DerivedDataCache/Public/DerivedDataCacheRecord.h
@@ -72,6 +72,9 @@ namespace UE::DerivedData
 class FCacheRecord
 {
 public:
+	/** Construct a cache record for the key with no metadata or values. */
+	UE_API explicit FCacheRecord(const FCacheKey& Key);
+
 	/** Returns the key that identifies this record in the cache. */
 	inline const FCacheKey& GetKey() const { return Record->GetKey(); }
 
