DDC: Added traits for cache response types

#Synced-CL 50002741

diff --git a/Engine/Source/Developer/DerivedDataCache/Internal/DerivedDataCacheRequestTraits.h b/Engine/Source/Developer/DerivedDataCache/Internal/DerivedDataCacheRequestTraits.h
deleted file mode 100644
diff --git a/Engine/Source/Developer/DerivedDataCache/Internal/DerivedDataCacheTraits.h b/Engine/Source/Developer/DerivedDataCache/Internal/DerivedDataCacheTraits.h
new file mode 100644
--- /dev/null
+++ b/Engine/Source/Developer/DerivedDataCache/Internal/DerivedDataCacheTraits.h
@@ -0,0 +1,150 @@
+// Copyright Epic Games, Inc. All Rights Reserved.
+
+#pragma once
+
+#include "DerivedDataCache.h"
+#include "DerivedDataCacheMethod.h"
+
+namespace UE::DerivedData::Private
+{
+
+template <typename RequestOrResponseType>
+struct TCacheRequestResponseTraits;
+
+template <>
+struct TCacheRequestResponseTraits<FCachePutRequest>
+{
+	using FRequest = FCachePutRequest;
+	using FResponse = FCachePutResponse;
+	using FOnComplete = FOnCachePutComplete;
+	constexpr static ECacheMethod Method = ECacheMethod::Put;
+	constexpr static void (ICache::*Function)(TConstArrayView<FRequest>, IRequestOwner&, FOnComplete&&) = &ICache::Put;
+	static const FCacheKey& GetKey(const FCachePutRequest& Request) { return Request.Record.GetKey(); }
+};
+
+template <>
+struct TCacheRequestResponseTraits<FCacheGetRequest>
+{
+	using FRequest = FCacheGetRequest;
+	using FResponse = FCacheGetResponse;
+	using FOnComplete = FOnCacheGetComplete;
+	constexpr static ECacheMethod Method = ECacheMethod::Get;
+	constexpr static void (ICache::*Function)(TConstArrayView<FRequest>, IRequestOwner&, FOnComplete&&) = &ICache::Get;
+	static const FCacheKey& GetKey(const FCacheGetRequest& Request) { return Request.Key; }
+};
+
+template <>
+struct TCacheRequestResponseTraits<FCachePutValueRequest>
+{
+	using FRequest = FCachePutValueRequest;
+	using FResponse = FCachePutValueResponse;
+	using FOnComplete = FOnCachePutValueComplete;
+	constexpr static ECacheMethod Method = ECacheMethod::PutValue;
+	constexpr static void (ICache::*Function)(TConstArrayView<FRequest>, IRequestOwner&, FOnComplete&&) = &ICache::PutValue;
+	static const FCacheKey& GetKey(const FCachePutValueRequest& Request) { return Request.Key; }
+};
+
+template <>
+struct TCacheRequestResponseTraits<FCacheGetValueRequest>
+{
+	using FRequest = FCacheGetValueRequest;
+	using FResponse = FCacheGetValueResponse;
+	using FOnComplete = FOnCacheGetValueComplete;
+	constexpr static ECacheMethod Method = ECacheMethod::GetValue;
+	constexpr static void (ICache::*Function)(TConstArrayView<FRequest>, IRequestOwner&, FOnComplete&&) = &ICache::GetValue;
+	static const FCacheKey& GetKey(const FCacheGetValueRequest& Request) { return Request.Key; }
+};
+
+template <>
+struct TCacheRequestResponseTraits<FCacheGetChunkRequest>
+{
+	using FRequest = FCacheGetChunkRequest;
+	using FResponse = FCacheGetChunkResponse;
+	using FOnComplete = FOnCacheGetChunkComplete;
+	constexpr static ECacheMethod Method = ECacheMethod::GetChunks;
+	constexpr static void (ICache::*Function)(TConstArrayView<FRequest>, IRequestOwner&, FOnComplete&&) = &ICache::GetChunks;
+	static const FCacheKey& GetKey(const FCacheGetChunkRequest& Request) { return Request.Key; }
+};
+
+template <>
+struct TCacheRequestResponseTraits<FCachePutResponse>
+{
+	using FRequest = FCachePutRequest;
+	using FResponse = FCachePutResponse;
+	using FOnComplete = FOnCachePutComplete;
+	constexpr static ECacheMethod Method = ECacheMethod::Put;
+	constexpr static void (ICache::*Function)(TConstArrayView<FRequest>, IRequestOwner&, FOnComplete&&) = &ICache::Put;
+	static const FCacheKey& GetKey(const FCachePutResponse& Response) { return Response.Key; }
+};
+
+template <>
+struct TCacheRequestResponseTraits<FCacheGetResponse>
+{
+	using FRequest = FCacheGetRequest;
+	using FResponse = FCacheGetResponse;
+	using FOnComplete = FOnCacheGetComplete;
+	constexpr static ECacheMethod Method = ECacheMethod::Get;
+	constexpr static void (ICache::*Function)(TConstArrayView<FRequest>, IRequestOwner&, FOnComplete&&) = &ICache::Get;
+	static const FCacheKey& GetKey(const FCacheGetResponse& Response) { return Response.Record.GetKey(); }
+};
+
+template <>
+struct TCacheRequestResponseTraits<FCachePutValueResponse>
+{
+	using FRequest = FCachePutValueRequest;
+	using FResponse = FCachePutValueResponse;
+	using FOnComplete = FOnCachePutValueComplete;
+	constexpr static ECacheMethod Method = ECacheMethod::PutValue;
+	constexpr static void (ICache::*Function)(TConstArrayView<FRequest>, IRequestOwner&, FOnComplete&&) = &ICache::PutValue;
+	static const FCacheKey& GetKey(const FCachePutValueResponse& Response) { return Response.Key; }
+};
+
+template <>
+struct TCacheRequestResponseTraits<FCacheGetValueResponse>
+{
+	using FRequest = FCacheGetValueRequest;
+	using FResponse = FCacheGetValueResponse;
+	using FOnComplete = FOnCacheGetValueComplete;
+	constexpr static ECacheMethod Method = ECacheMethod::GetValue;
+	constexpr static void (ICache::*Function)(TConstArrayView<FRequest>, IRequestOwner&, FOnComplete&&) = &ICache::GetValue;
+	static const FCacheKey& GetKey(const FCacheGetValueResponse& Response) { return Response.Key; }
+};
+
+template <>
+struct TCacheRequestResponseTraits<FCacheGetChunkResponse>
+{
+	using FRequest = FCacheGetChunkRequest;
+	using FResponse = FCacheGetChunkResponse;
+	using FOnComplete = FOnCacheGetChunkComplete;
+	constexpr static ECacheMethod Method = ECacheMethod::GetChunks;
+	constexpr static void (ICache::*Function)(TConstArrayView<FRequest>, IRequestOwner&, FOnComplete&&) = &ICache::GetChunks;
+	static const FCacheKey& GetKey(const FCacheGetChunkResponse& Response) { return Response.Key; }
+};
+
+} // UE::DerivedData::Private
+
+namespace UE::DerivedData
+{
+
+/** Alias for the request type corresponding to the request/response type. */
+template <typename RequestOrResponseType> using TCacheRequestFor = typename Private::TCacheRequestResponseTraits<RequestOrResponseType>::FRequest;
+
+/** Alias for the response type corresponding to the request/response type. */
+template <typename RequestOrResponseType> using TCacheResponseFor = typename Private::TCacheRequestResponseTraits<RequestOrResponseType>::FResponse;
+
+/** Alias for the completion callback type corresponding to the request/response type. */
+template <typename RequestOrResponseType> using TCacheOnCompleteFor = typename Private::TCacheRequestResponseTraits<RequestOrResponseType>::FOnComplete;
+
+/** Alias for the cache method corresponding to the request/response type. */
+template <typename RequestOrResponseType> inline constexpr auto CacheMethodFor = Private::TCacheRequestResponseTraits<RequestOrResponseType>::Method;
+
+/** Alias for the cache function corresponding to the request/response type. */
+template <typename RequestOrResponseType> inline constexpr auto CacheFunctionFor = Private::TCacheRequestResponseTraits<RequestOrResponseType>::Function;
+
+template <typename RequestOrResponseType>
+inline auto GetCacheKey(const RequestOrResponseType& Request) -> decltype(Private::TCacheRequestResponseTraits<RequestOrResponseType>::GetKey(Request))
+{
+	return Private::TCacheRequestResponseTraits<RequestOrResponseType>::GetKey(Request);
+}
+
+} // UE::DerivedData
diff --git a/Engine/Source/Developer/DerivedDataCache/Private/DerivedDataCacheStore.h b/Engine/Source/Developer/DerivedDataCache/Private/DerivedDataCacheStore.h
--- a/Engine/Source/Developer/DerivedDataCache/Private/DerivedDataCacheStore.h
+++ b/Engine/Source/Developer/DerivedDataCache/Private/DerivedDataCacheStore.h
@@ -5,7 +5,7 @@
 #include "Async/Mutex.h"
 #include "Containers/ContainersFwd.h"
 #include "DerivedDataCache.h"
-#include "DerivedDataCacheRequestTraits.h"
+#include "DerivedDataCacheTraits.h"
 #include "Misc/MonotonicTime.h"
 #include "Templates/DontCopy.h"
 
diff --git a/Engine/Source/Developer/DerivedDataCache/Private/DerivedDataCacheStoreHierarchy.cpp b/Engine/Source/Developer/DerivedDataCache/Private/DerivedDataCacheStoreHierarchy.cpp
--- a/Engine/Source/Developer/DerivedDataCache/Private/DerivedDataCacheStoreHierarchy.cpp
+++ b/Engine/Source/Developer/DerivedDataCache/Private/DerivedDataCacheStoreHierarchy.cpp
@@ -891,7 +891,7 @@ void FCacheStoreHierarchy::TPutBatch<Params>::CompletePutRequest(FPutResponse&&
 template <typename Params>
 void FCacheStoreHierarchy::TPutBatch<Params>::FinishRequest(FPutResponse&& Response, const FPutRequest& Request)
 {
-	Hierarchy.RecordStats(Response.Key.Bucket, GetBatchType(), ERequestOp::Put, Response.Status, 0, MeasureLogicalSize(Request));
+	Hierarchy.RecordStats(GetCacheKey(Response).Bucket, GetBatchType(), ERequestOp::Put, Response.Status, 0, MeasureLogicalSize(Request));
 	OnComplete(MoveTemp(Response));
 }
 
diff --git a/Engine/Source/Programs/DerivedDataTool/Private/DerivedDataCopyCommand.cpp b/Engine/Source/Programs/DerivedDataTool/Private/DerivedDataCopyCommand.cpp
--- a/Engine/Source/Programs/DerivedDataTool/Private/DerivedDataCopyCommand.cpp
+++ b/Engine/Source/Programs/DerivedDataTool/Private/DerivedDataCopyCommand.cpp
@@ -17,7 +17,7 @@
 #include "DerivedDataCache.h"
 #include "DerivedDataCacheKeyFilter.h"
 #include "DerivedDataCacheReplay.h"
-#include "DerivedDataCacheRequestTraits.h"
+#include "DerivedDataCacheTraits.h"
 #include "DerivedDataConfig.h"
 #include "DerivedDataRequestOwner.h"
 #include "DerivedDataSerialization.h"
@@ -811,21 +811,11 @@ uint32 FCopyContext::Find(uint32 HeadValueIndex, const FValue& SearchValue) cons
 
 ///////////////////////////////////////////////////////////////////////////////////////////////////////////////////////
 
-inline static const FCacheKey& GetKeyFromResponse(const FCacheGetResponse& Response)
-{
-	return Response.Record.GetKey();
-}
-
 inline static TConstArrayView<FValueWithId> GetValuesFromResponse(const FCacheGetResponse& Response)
 {
 	return Response.Record.GetValues();
 }
 
-inline static const FCacheKey& GetKeyFromResponse(const FCacheGetValueResponse& Response)
-{
-	return Response.Key;
-}
-
 inline static const FValue& GetValuesFromResponse(const FCacheGetValueResponse& Response)
 {
 	return Response.Value;
@@ -911,7 +901,7 @@ public:
 		for (const FCachePutRequest& Request : Requests)
 		{
 			const FCachePutResponse& Response = Responses[Index++];
-			if (Response.Status == EStatus::Ok && Filter && Filter.IsMatch(Response.Key))
+			if (Response.Status == EStatus::Ok && Filter && Filter.IsMatch(Request.Record.GetKey()))
 			{
 				AddRecord(Request.Name, Request.Record.GetKey(), Request.Record.GetValues());
 			}
@@ -1341,13 +1331,13 @@ private:
 			{
 				Request.Flags.TargetMissingMask.fetch_or(uint32(1) << TargetIndex, std::memory_order_relaxed);
 				UE_LOGFMT(LogDerivedDataTool, Verbose, "Value in target '{Target}' missing for {Key} from '{Name}'.",
-					Context.TargetNames[TargetIndex], GetKeyFromResponse(Response), Response.Name);
+					Context.TargetNames[TargetIndex], GetCacheKey(Response), Response.Name);
 			}
 			else if (Context.Find(Index, GetValuesFromResponse(Response)) == MAX_uint32)
 			{
 				Request.Flags.TargetDiffersMask.fetch_or(uint32(1) << TargetIndex, std::memory_order_relaxed);
 				UE_LOGFMT(LogDerivedDataTool, Display, "Value in target '{Target}' differs for {Key} from '{Name}'.",
-					Context.TargetNames[TargetIndex], GetKeyFromResponse(Response), Response.Name);
+					Context.TargetNames[TargetIndex], GetCacheKey(Response), Response.Name);
 			}
 
 			Batch.Complete();
@@ -1456,19 +1446,19 @@ private:
 			{
 				Request.Flags.SourceMask.fetch_or(FCopyFlags::SourceMissingFlag, std::memory_order_relaxed);
 				UE_LOGFMT(LogDerivedDataTool, Error, "Value in source missing for {Key} from '{Name}'.",
-					GetKeyFromResponse(Response), Response.Name);
+					GetCacheKey(Response), Response.Name);
 			}
 			else if (!HasDataInResponse(Response))
 			{
 				Request.Flags.SourceMask.fetch_or(FCopyFlags::SourceMissingFlag, std::memory_order_relaxed);
 				UE_LOGFMT(LogDerivedDataTool, Error, "Value in source missing data for {Key} from '{Name}'.",
-					GetKeyFromResponse(Response), Response.Name);
+					GetCacheKey(Response), Response.Name);
 			}
 			else if (Context.Find(Index, GetValuesFromResponse(Response)) == MAX_uint32)
 			{
 				Request.Flags.SourceMask.fetch_or(FCopyFlags::SourceDiffersFlag, std::memory_order_relaxed);
 				UE_LOGFMT(LogDerivedDataTool, Error, "Value in source differs for {Key} from '{Name}'.",
-					GetKeyFromResponse(Response), Response.Name);
+					GetCacheKey(Response), Response.Name);
 			}
 			else
 			{
@@ -1598,7 +1588,7 @@ private:
 			if (Response.Status != EStatus::Ok && !Context.Params.bDebugSkipPut)
 			{
 				UE_LOGFMT(LogDerivedDataTool, Warning, "Failed to store to '{Target}' for {Key} from '{Name}'.",
-					Context.TargetNames[TargetIndex], Response.Key, Response.Name);
+					Context.TargetNames[TargetIndex], GetCacheKey(Response), Response.Name);
 			}
 			else
 			{
