DDC: Added ApplySkipPolicy to apply ECachePolicy::Skip to values and records

#Synced-CL 50002741

diff --git a/Engine/Source/Developer/DerivedDataCache/Private/DerivedDataCachePolicy.cpp b/Engine/Source/Developer/DerivedDataCache/Private/DerivedDataCachePolicy.cpp
--- a/Engine/Source/Developer/DerivedDataCache/Private/DerivedDataCachePolicy.cpp
+++ b/Engine/Source/Developer/DerivedDataCache/Private/DerivedDataCachePolicy.cpp
@@ -3,10 +3,14 @@
 #include "DerivedDataCachePolicy.h"
 
 #include "Algo/Accumulate.h"
+#include "Algo/AnyOf.h"
 #include "Algo/BinarySearch.h"
 #include "Containers/StringConv.h"
 #include "Containers/StringView.h"
+#include "DerivedDataCacheRecord.h"
+#include "DerivedDataValue.h"
 #include "Misc/StringBuilder.h"
+#include "Serialization/CompactBinary.h"
 #include "String/ParseTokens.h"
 #include "Templates/Function.h"
 
@@ -234,4 +238,51 @@ FCacheRecordPolicy FCacheRecordPolicyBuilder::Build()
 	return Policy;
 }
 
+FValue ApplySkipPolicy(const FValue& Value, ECachePolicy Policy)
+{
+	return EnumHasAnyFlags(Policy, ECachePolicy::SkipData) ? Value.RemoveData() : Value;
+}
+
+FValueWithId ApplySkipPolicy(const FValueWithId& Value, ECachePolicy Policy)
+{
+	return EnumHasAnyFlags(Policy, ECachePolicy::SkipData) ? Value.RemoveData() : Value;
+}
+
+FCacheRecord ApplySkipPolicy(const FCacheRecord& Record, const FCacheRecordPolicy& Policy)
+{
+	const ECachePolicy RecordPolicy = Policy.GetRecordPolicy();
+	const bool bMightSkipMeta = EnumHasAnyFlags(RecordPolicy, ECachePolicy::SkipMeta);
+	const bool bMightSkipData = EnumHasAnyFlags(RecordPolicy, ECachePolicy::SkipData) || !Policy.IsUniform();
+
+	const FCbObject& Meta = Record.GetMeta();
+	const TConstArrayView<FValueWithId> Values = Record.GetValues();
+
+	const bool bSkippedMeta = bMightSkipMeta && Meta;
+	if (bSkippedMeta || (bMightSkipData && Algo::AnyOf(Values, &FValueWithId::HasData)))
+	{
+		FCbObject NewMeta;
+		if (!bMightSkipMeta && Meta)
+		{
+			NewMeta = Meta;
+		}
+
+		bool bSkippedData = false;
+		TArray<FValueWithId> NewValues;
+		NewValues.Reserve(Values.Num());
+		for (const FValueWithId& Value : Values)
+		{
+			const bool bSkipData = EnumHasAnyFlags(Policy.GetValuePolicy(Value.GetId()), ECachePolicy::SkipData);
+			NewValues.Add(bSkipData ? Value.RemoveData() : Value);
+			bSkippedData |= bSkipData & Value.HasData();
+		}
+
+		if ((bSkippedMeta | bSkippedData))
+		{
+			return FCacheRecord::CreateByMove(Record.GetKey(), MoveTemp(NewMeta), MoveTemp(NewValues));
+		}
+	}
+
+	return Record;
+}
+
 } // UE::DerivedData
diff --git a/Engine/Source/Developer/DerivedDataCache/Private/DerivedDataCacheStoreHierarchy.cpp b/Engine/Source/Developer/DerivedDataCache/Private/DerivedDataCacheStoreHierarchy.cpp
--- a/Engine/Source/Developer/DerivedDataCache/Private/DerivedDataCacheStoreHierarchy.cpp
+++ b/Engine/Source/Developer/DerivedDataCache/Private/DerivedDataCacheStoreHierarchy.cpp
@@ -1280,35 +1280,12 @@ FCacheGetResponse FCacheStoreHierarchy::FCacheRecordBatchParams::FilterResponseB
 	const FCacheGetResponse& Response,
 	const FCacheGetRequest& Request)
 {
-	const ECachePolicy RecordPolicy = Request.Policy.GetRecordPolicy();
-	if (Response.Status != EStatus::Ok && !EnumHasAnyFlags(RecordPolicy, ECachePolicy::PartialRecord))
-	{
-		FCacheRecordBuilder Builder(Response.Record.GetKey());
-		return {Response.Name, Builder.Build(), Response.UserData, Response.Status};
-	}
-	const bool bMightSkipData = EnumHasAnyFlags(RecordPolicy, ECachePolicy::SkipData) || !Request.Policy.IsUniform();
-	if ((bMightSkipData && Algo::AnyOf(Response.Record.GetValues(), &FValue::HasData)) ||
-		(EnumHasAnyFlags(RecordPolicy, ECachePolicy::SkipMeta) && Response.Record.GetMeta()))
-	{
-		FCacheRecordBuilder Builder(Response.Record.GetKey());
-		if (!EnumHasAnyFlags(RecordPolicy, ECachePolicy::SkipMeta))
-		{
-			Builder.SetMeta(CopyTemp(Response.Record.GetMeta()));
-		}
-		for (const FValueWithId& Value : Response.Record.GetValues())
-		{
-			if (EnumHasAnyFlags(Request.Policy.GetValuePolicy(Value.GetId()), ECachePolicy::SkipData))
-			{
-				Builder.AddValue(Value.GetId(), Value.RemoveData());
-			}
-			else
-			{
-				Builder.AddValue(Value);
-			}
-		}
-		return {Response.Name, Builder.Build(), Response.UserData, Response.Status};
-	}
-	return Response;
+	const bool bCreateEmptyRecord =
+		Response.Status != EStatus::Ok &&
+		!EnumHasAnyFlags(Request.Policy.GetRecordPolicy(), ECachePolicy::PartialRecord) &&
+		(Response.Record.GetMeta() || !Response.Record.GetValues().IsEmpty());
+	FCacheRecord Record(bCreateEmptyRecord ? FCacheRecord(Response.Record.GetKey()) : ApplySkipPolicy(Response.Record, Request.Policy));
+	return {Response.Name, MoveTemp(Record), Response.UserData, Response.Status};
 }
 
 template <>
diff --git a/Engine/Source/Developer/DerivedDataCache/Private/FileSystemCacheStore.cpp b/Engine/Source/Developer/DerivedDataCache/Private/FileSystemCacheStore.cpp
--- a/Engine/Source/Developer/DerivedDataCache/Private/FileSystemCacheStore.cpp
+++ b/Engine/Source/Developer/DerivedDataCache/Private/FileSystemCacheStore.cpp
@@ -2454,7 +2454,7 @@ bool FFileSystemCacheStore::GetCacheContent(
 
 	if (Value.HasData())
 	{
-		OutValue = EnumHasAnyFlags(Policy, ECachePolicy::SkipData) ? Value.RemoveData() : Value;
+		OutValue = ApplySkipPolicy(Value, Policy);
 		return true;
 	}
 
diff --git a/Engine/Source/Developer/DerivedDataCache/Private/PakFileCacheStore.cpp b/Engine/Source/Developer/DerivedDataCache/Private/PakFileCacheStore.cpp
--- a/Engine/Source/Developer/DerivedDataCache/Private/PakFileCacheStore.cpp
+++ b/Engine/Source/Developer/DerivedDataCache/Private/PakFileCacheStore.cpp
@@ -1268,7 +1268,7 @@ bool FPakFileCacheStore::GetCacheContent(
 
 	if (Value.HasData())
 	{
-		OutValue = EnumHasAnyFlags(Policy, ECachePolicy::SkipData) ? Value.RemoveData() : Value;
+		OutValue = ApplySkipPolicy(Value, Policy);
 		return true;
 	}
 
diff --git a/Engine/Source/Developer/DerivedDataCache/Private/S3CacheStore.cpp b/Engine/Source/Developer/DerivedDataCache/Private/S3CacheStore.cpp
--- a/Engine/Source/Developer/DerivedDataCache/Private/S3CacheStore.cpp
+++ b/Engine/Source/Developer/DerivedDataCache/Private/S3CacheStore.cpp
@@ -1426,7 +1426,7 @@ bool FS3CacheStore::GetCacheContent(
 
 	if (Value.HasData())
 	{
-		OutValue = EnumHasAnyFlags(Policy, ECachePolicy::SkipData) ? Value.RemoveData() : Value;
+		OutValue = ApplySkipPolicy(Value, Policy);
 		return true;
 	}
 
diff --git a/Engine/Source/Developer/DerivedDataCache/Public/DerivedDataCachePolicy.h b/Engine/Source/Developer/DerivedDataCache/Public/DerivedDataCachePolicy.h
--- a/Engine/Source/Developer/DerivedDataCache/Public/DerivedDataCachePolicy.h
+++ b/Engine/Source/Developer/DerivedDataCache/Public/DerivedDataCachePolicy.h
@@ -15,8 +15,11 @@
 
 template <typename FuncType> class TFunctionRef;
 
-namespace UE::DerivedData::Private { class ICacheRecordPolicyShared; }
+namespace UE::DerivedData { class FCacheRecord; }
 namespace UE::DerivedData { class FOptionalCacheRecordPolicy; }
+namespace UE::DerivedData { class FValue; }
+namespace UE::DerivedData { class FValueWithId; }
+namespace UE::DerivedData::Private { class ICacheRecordPolicyShared; }
 
 namespace UE::DerivedData
 {
@@ -33,7 +36,10 @@ namespace UE::DerivedData
  * Get(Default | SkipData): Check for existence in any cache. Store the response to any caches if missing.
  * Get(Default | PartialRecord): Fetch from any cache, and return a partial record if values are missing data.
  *
- * Put(Default): Store to every cache, and do not overwrite existing valid records or values.
+ * Put(Default | SkipData): Store to every cache, do not overwrite existing valid records or values, and skip fetching
+ *                          values from conflicting existing records for the record returned in the response.
+ * Put(Default): Store to every cache, and do not overwrite existing valid records or values, and fetch values
+ *               from conflicting existing records for the record returned in the response.
  * Put(Store): Store to every cache, and overwrite existing records or values.
  * Put(Local): Store to every local cache, skipping remote caches.
  * Put(Default | PartialRecord): Store to every cache, even if the record has missing data for its values.
@@ -258,6 +264,13 @@ public:
 	inline void Reset() { *this = FOptionalCacheRecordPolicy(); }
 };
 
+/** Apply the policy to the value by removing data if SkipData is present. */
+UE_API FValue ApplySkipPolicy(const FValue& Value, ECachePolicy Policy);
+/** Apply the policy to the value by removing data if SkipData is present. */
+UE_API FValueWithId ApplySkipPolicy(const FValueWithId& Value, ECachePolicy Policy);
+/** Apply the policy to the record by removing data if SkipData is present and metadata if SkipMeta is present. */
+UE_API FCacheRecord ApplySkipPolicy(const FCacheRecord& Record, const FCacheRecordPolicy& Policy);
+
 } // UE::DerivedData
 
 #undef UE_API
