DDC: Fixed minor issues in the FileSystem and Memory cache stores

- Allow records with no values in the Memory cache store.
- Report an error when the memory cache store exceeds its size limit.
- Verify that content exists for differing values even if they are expected to be non-deterministic.

#Synced-CL 50002741

diff --git a/Engine/Source/Developer/DerivedDataCache/Private/FileSystemCacheStore.cpp b/Engine/Source/Developer/DerivedDataCache/Private/FileSystemCacheStore.cpp
--- a/Engine/Source/Developer/DerivedDataCache/Private/FileSystemCacheStore.cpp
+++ b/Engine/Source/Developer/DerivedDataCache/Private/FileSystemCacheStore.cpp
@@ -4,6 +4,7 @@
 #include "Algo/AllOf.h"
 #include "Algo/Compare.h"
 #include "Algo/Find.h"
+#include "Algo/NoneOf.h"
 #include "Algo/StableSort.h"
 #include "Algo/Transform.h"
 #include "Async/Async.h"
@@ -1910,16 +1911,18 @@ bool FFileSystemCacheStore::PutCacheRecord(
 			// Save the new package if the existing package is invalid.
 			const FOptionalCacheRecord ExistingRecord = FCacheRecord::Load(ExistingPackage);
 			bSavePackage |= !ExistingRecord;
-			const auto MakeValueTuple = [&Policy](const FValueWithId& Value) -> TTuple<FValueId, FIoHash>
-			{
-				return !EnumHasAnyFlags(Policy.GetValuePolicy(Value.GetId()), ECachePolicy::NonDeterministic)
-					? MakeTuple(Value.GetId(), Value.GetRawHash())
-					: TTuple<FValueId, FIoHash>{};
-			};
-			if (ExistingRecord && !Algo::CompareBy(ExistingRecord.Get().GetValues(), Record.GetValues(), MakeValueTuple))
+			if (ExistingRecord && !Algo::Compare(ExistingRecord.Get().GetValues(), Record.GetValues()))
 			{
 				// Content differs between the existing record and the new record.
-				UE_LOG(LogDerivedDataCache, Display, TEXT("%s: Cache put found non-deterministic record for %s from '%.*s'"),
+				const bool bReportConflicts = !EnumHasAnyFlags(Policy.GetBasePolicy(), ECachePolicy::NonDeterministic) &&
+					Algo::NoneOf(Policy.GetValuePolicies(), [](const FCacheValuePolicy& Value) { return EnumHasAnyFlags(Value.Policy, ECachePolicy::NonDeterministic); });
+				const auto MakeCompareValue = [&Policy](const FValueWithId& Value) -> FValueWithId
+				{
+					return !EnumHasAnyFlags(Policy.GetValuePolicy(Value.GetId()), ECachePolicy::NonDeterministic)
+						? Value : FValueWithId(Value.GetId(), FIoHash::Zero, 0);
+				};
+				UE_CLOG(bReportConflicts && !Algo::CompareBy(ExistingRecord.Get().GetValues(), Record.GetValues(), MakeCompareValue),
+					LogDerivedDataCache, Display, TEXT("%s: Cache put found non-deterministic record for %s from '%.*s'"),
 					*CachePath, *WriteToString<96>(Key), Name.Len(), Name.GetData());
 				const auto HasValueContent = [this, Name, &Key, &Stats](const FValueWithId& Value) -> bool
 				{
@@ -2211,11 +2214,11 @@ bool FFileSystemCacheStore::PutCacheValue(
 					*CachePath, *WriteToString<96>(Key), Name.Len(), Name.GetData());
 				bSavePackage = true;
 			}
-			else if (!EnumHasAnyFlags(Policy, ECachePolicy::NonDeterministic) &&
-				!(RawHash == Value.GetRawHash() && RawSize == Value.GetRawSize()))
+			else if (!(RawHash == Value.GetRawHash() && RawSize == Value.GetRawSize()))
 			{
 				// Content differs between the existing value and the new value.
-				UE_LOG(LogDerivedDataCache, Display, TEXT("%s: Cache put found non-deterministic value "
+				UE_CLOG(!EnumHasAnyFlags(Policy, ECachePolicy::NonDeterministic),
+					LogDerivedDataCache, Display, TEXT("%s: Cache put found non-deterministic value "
 					"with new hash %s and existing hash %s for %s from '%.*s'"),
 					*CachePath, *WriteToString<48>(Value.GetRawHash()), *WriteToString<48>(RawHash),
 					*WriteToString<96>(Key), Name.Len(), Name.GetData());
diff --git a/Engine/Source/Developer/DerivedDataCache/Private/MemoryCacheStore.cpp b/Engine/Source/Developer/DerivedDataCache/Private/MemoryCacheStore.cpp
--- a/Engine/Source/Developer/DerivedDataCache/Private/MemoryCacheStore.cpp
+++ b/Engine/Source/Developer/DerivedDataCache/Private/MemoryCacheStore.cpp
@@ -230,11 +230,6 @@ void FMemoryCacheStore::Put(
 
 		FRequestTimer RequestTimer(RequestStats);
 
-		if (Algo::NoneOf(Values, &FValue::HasData))
-		{
-			continue;
-		}
-
 		if (bCanBeDisabled && DebugOptions.ShouldSimulatePutMiss(Key))
 		{
 			UE_LOG(LogDerivedDataCache, Verbose, TEXT("%s: Simulated miss for put of %s from '%s'"),
@@ -245,8 +240,16 @@ void FMemoryCacheStore::Put(
 		const bool bReplaceExisting = !EnumHasAnyFlags(Request.Policy.GetRecordPolicy(), ECachePolicy::QueryLocal);
 
 		FWriteScopeLock ScopeLock(SynchronizationObject);
-		FCacheRecordComponents& Components = CacheRecords.FindOrAdd(Key);
-		const bool bHasExisting = !Components.Values.IsEmpty();
+		bool bHasExisting = false;
+		FCacheRecordComponents& Components = [this, &Key, &bHasExisting]() -> FCacheRecordComponents&
+		{
+			if (FCacheRecordComponents* Existing = CacheRecords.Find(Key))
+			{
+				bHasExisting = true;
+				return *Existing;
+			}
+			return CacheRecords.Add(Key);
+		}();
 		Status = EStatus::Ok;
 
 		if (bHasExisting && !bReplaceExisting && Algo::AllOf(Components.Values, &FValue::HasData))
@@ -317,7 +320,7 @@ void FMemoryCacheStore::Put(
 				     "CurrentSize %" UINT64_FMT " KiB / MaxSize: %" UINT64_FMT " KiB"),
 				CurrentCacheSize / 1024, MaxCacheSize / 1024);
 			bMaxSizeExceeded = true;
-			Status = EStatus::Ok;
+			Status = EStatus::Error;
 			continue;
 		}
 
