DDC: Changed & to && to silence static analysis

#rnx

#Synced-CL 50002741

#Synced-CL 50002741

diff --git a/Engine/Source/Developer/DerivedDataCache/Private/DerivedDataCachePolicy.cpp b/Engine/Source/Developer/DerivedDataCache/Private/DerivedDataCachePolicy.cpp
--- a/Engine/Source/Developer/DerivedDataCache/Private/DerivedDataCachePolicy.cpp
+++ b/Engine/Source/Developer/DerivedDataCache/Private/DerivedDataCachePolicy.cpp
@@ -273,7 +273,7 @@ FCacheRecord ApplySkipPolicy(const FCacheRecord& Record, const FCacheRecordPolic
 		{
 			const bool bSkipData = EnumHasAnyFlags(Policy.GetValuePolicy(Value.GetId()), ECachePolicy::SkipData);
 			NewValues.Add(bSkipData ? Value.RemoveData() : Value);
-			bSkippedData |= bSkipData & Value.HasData();
+			bSkippedData |= bSkipData && Value.HasData();
 		}
 
 		if ((bSkippedMeta | bSkippedData))
