DDC: Added record/value to cache put responses

These currently contain a copy of what was passed in the request. They will later contain a copy of what ended up stored in the cache.

#Synced-CL 50002741

diff --git a/Engine/Source/Developer/DerivedDataCache/Internal/DerivedDataCacheTraits.h b/Engine/Source/Developer/DerivedDataCache/Internal/DerivedDataCacheTraits.h
--- a/Engine/Source/Developer/DerivedDataCache/Internal/DerivedDataCacheTraits.h
+++ b/Engine/Source/Developer/DerivedDataCache/Internal/DerivedDataCacheTraits.h
@@ -74,7 +74,7 @@ struct TCacheRequestResponseTraits<FCachePutResponse>
 	using FOnComplete = FOnCachePutComplete;
 	constexpr static ECacheMethod Method = ECacheMethod::Put;
 	constexpr static void (ICache::*Function)(TConstArrayView<FRequest>, IRequestOwner&, FOnComplete&&) = &ICache::Put;
-	static const FCacheKey& GetKey(const FCachePutResponse& Response) { return Response.Key; }
+	static const FCacheKey& GetKey(const FCachePutResponse& Response) { return Response.Record.GetKey(); }
 };
 
 template <>
diff --git a/Engine/Source/Developer/DerivedDataCache/Private/DerivedDataCache.cpp b/Engine/Source/Developer/DerivedDataCache/Private/DerivedDataCache.cpp
--- a/Engine/Source/Developer/DerivedDataCache/Private/DerivedDataCache.cpp
+++ b/Engine/Source/Developer/DerivedDataCache/Private/DerivedDataCache.cpp
@@ -229,7 +229,9 @@ namespace UE::DerivedData
 
 FCachePutResponse FCachePutRequest::MakeResponse(const EStatus Status) const
 {
-	return {Name, Record.GetKey(), UserData, Status};
+	PRAGMA_DISABLE_DEPRECATION_WARNINGS;
+	return {Name, Record.GetKey(), ApplySkipPolicy(Record, Policy), UserData, Status};
+	PRAGMA_ENABLE_DEPRECATION_WARNINGS;
 }
 
 FCacheGetResponse FCacheGetRequest::MakeResponse(const EStatus Status) const
@@ -239,7 +241,9 @@ FCacheGetResponse FCacheGetRequest::MakeResponse(const EStatus Status) const
 
 FCachePutValueResponse FCachePutValueRequest::MakeResponse(const EStatus Status) const
 {
-	return {Name, Key, UserData, Status};
+	PRAGMA_DISABLE_DEPRECATION_WARNINGS;
+	return {Name, Key, ApplySkipPolicy(Value, Policy), UserData, Status};
+	PRAGMA_ENABLE_DEPRECATION_WARNINGS;
 }
 
 FCacheGetValueResponse FCacheGetValueRequest::MakeResponse(const EStatus Status) const
diff --git a/Engine/Source/Developer/DerivedDataCache/Private/DerivedDataCacheReplay.cpp b/Engine/Source/Developer/DerivedDataCache/Private/DerivedDataCacheReplay.cpp
--- a/Engine/Source/Developer/DerivedDataCache/Private/DerivedDataCacheReplay.cpp
+++ b/Engine/Source/Developer/DerivedDataCache/Private/DerivedDataCacheReplay.cpp
@@ -460,7 +460,19 @@ static bool LoadBatchArraysFromCompactBinary(
 	TArray<TCacheResponseFor<RequestType>, TInlineAllocator<16>>& OutResponses)
 {
 	using ResponseType = TCacheResponseFor<RequestType>;
-	if constexpr (!std::is_default_constructible_v<RequestType>)
+	if constexpr (!std::is_default_constructible_v<RequestType> && !std::is_default_constructible_v<ResponseType>)
+	{
+		TOptional<RequestType> Request;
+		TOptional<ResponseType> Response;
+		if (LoadFromCompactBinary(Field, Request, &Response))
+		{
+			OutRequests.Emplace(MoveTemp(Request.GetValue()));
+			OutResponses.Emplace(MoveTemp(Response.GetValue()));
+			return true;
+		}
+		return false;
+	}
+	else if constexpr (!std::is_default_constructible_v<RequestType>)
 	{
 		TOptional<RequestType> Request;
 		ResponseType& Response = OutResponses.AddDefaulted_GetRef();
diff --git a/Engine/Source/Developer/DerivedDataCache/Private/DerivedDataCacheStoreAsync.cpp b/Engine/Source/Developer/DerivedDataCache/Private/DerivedDataCacheStoreAsync.cpp
--- a/Engine/Source/Developer/DerivedDataCache/Private/DerivedDataCacheStoreAsync.cpp
+++ b/Engine/Source/Developer/DerivedDataCache/Private/DerivedDataCacheStoreAsync.cpp
@@ -39,7 +39,7 @@ public:
 			Execute(COOK_STAT(&FDerivedDataCacheUsageStats::TimePut,) Requests, Owner,
 				[this, OnComplete = MoveTemp(OnComplete)](FCachePutResponse&& Response)
 				{
-					MemoryCache->Delete(Response.Key, Response.Name);
+					MemoryCache->Delete(Response.Record.GetKey(), Response.Name);
 					OnComplete(MoveTemp(Response));
 				}, &ICacheStore::Put);
 		}
diff --git a/Engine/Source/Developer/DerivedDataCache/Private/DerivedDataCacheStoreThrottle.cpp b/Engine/Source/Developer/DerivedDataCache/Private/DerivedDataCacheStoreThrottle.cpp
--- a/Engine/Source/Developer/DerivedDataCache/Private/DerivedDataCacheStoreThrottle.cpp
+++ b/Engine/Source/Developer/DerivedDataCache/Private/DerivedDataCacheStoreThrottle.cpp
@@ -54,7 +54,7 @@ public:
 		InnerCache->Put(Requests, Owner,
 			[this, RecordSizes = MoveTemp(RecordSizes), State = EnterThrottlingScope(), OnComplete = MoveTemp(OnComplete)](FCachePutResponse&& Response)
 			{
-				const FRecordSize* Size = Algo::FindBy(RecordSizes, Response.Key, &FRecordSize::Key);
+				const FRecordSize* Size = Algo::FindBy(RecordSizes, Response.Record.GetKey(), &FRecordSize::Key);
 				CloseThrottlingScope(State, FThrottlingState(this, Size ? Size->Size : 0));
 				OnComplete(MoveTemp(Response));
 			});
diff --git a/Engine/Source/Developer/DerivedDataCache/Private/DerivedDataSerialization.cpp b/Engine/Source/Developer/DerivedDataCache/Private/DerivedDataSerialization.cpp
--- a/Engine/Source/Developer/DerivedDataCache/Private/DerivedDataSerialization.cpp
+++ b/Engine/Source/Developer/DerivedDataCache/Private/DerivedDataSerialization.cpp
@@ -2,6 +2,7 @@
 
 #include "DerivedDataSerialization.h"
 
+#include "Algo/Compare.h"
 #include "Algo/IsSorted.h"
 #include "Algo/Sort.h"
 #include "DerivedDataCache.h"
@@ -104,11 +105,78 @@ bool LoadMetaFromCompactBinary(FCbFieldView Field, FCbObject& OutMeta, FLoadAtta
 	return false;
 }
 
+template <CIntegral FoundMaskType, FoundMaskType FoundMeta, FoundMaskType FoundMetaHash, FoundMaskType FoundMetaSize>
+inline void LoadMetaFromCompactBinaryIterator(
+	FoundMaskType& FoundMask,
+	FCbFieldViewIterator& It,
+	FCbValue& MetaObjectValue,
+	FCbValue& MetaHashValue,
+	FCbValue& MetaSizeValue)
+{
+	if (!(FoundMask & FoundMeta) && It.GetName() == ANSITEXTVIEW("Meta"))
+	{
+		MetaObjectValue = It;
+		FoundMask |= FoundMeta;
+		++It;
+	}
+	else
+	{
+		if (!(FoundMask & FoundMetaHash) && It.GetName() == ANSITEXTVIEW("MetaHash"))
+		{
+			MetaHashValue = It;
+			FoundMask |= FoundMetaHash;
+			++It;
+		}
+		if (!(FoundMask & FoundMetaSize) && It.GetName() == ANSITEXTVIEW("MetaSize"))
+		{
+			MetaSizeValue = It;
+			FoundMask |= FoundMetaSize;
+			++It;
+		}
+	}
+}
+
+inline bool LoadMetaFromCompactBinaryValues(
+	FCbValue& MetaObjectValue,
+	FCbValue& MetaHashValue,
+	FCbValue& MetaSizeValue,
+	FCbObject& OutMeta,
+	FLoadAttachmentFn LoadAttachment)
+{
+	if (MetaObjectValue)
+	{
+		return LoadMetaFromObject(MetaObjectValue, OutMeta);
+	}
+	else if (LoadAttachment && MetaHashValue && MetaSizeValue)
+	{
+		return LoadMetaFromHashAndSize(MetaHashValue, MetaSizeValue, OutMeta, LoadAttachment);
+	}
+	return false;
+}
+
+///////////////////////////////////////////////////////////////////////////////////////////////////////////////////////
+
+[[nodiscard]] inline bool LoadValuesFromCompactBinary(FCbArrayView ValuesArray, TArray<FValueWithId>& OutValues, FLoadAttachmentFn LoadAttachment)
+{
+	bool bOk = true;
+	OutValues.Reset(int32(ValuesArray.Num()));
+	for (FCbFieldView ValueField : ValuesArray)
+	{
+		FValueWithId& Value = OutValues.AddDefaulted_GetRef();
+		bOk &= LoadFromCompactBinary(ValueField, Value, LoadAttachment);
+	}
+	if (UNLIKELY(!Algo::IsSorted(OutValues)))
+	{
+		Algo::Sort(OutValues);
+	}
+	return bOk;
+}
+
 ///////////////////////////////////////////////////////////////////////////////////////////////////////////////////////
 
 void SaveToCompactBinary(FCbWriter& Writer, const FCachePutRequest& Request, const FCachePutResponse* Response, FSaveAttachmentFn SaveAttachment)
 {
-	check(!Response || Request.Record.GetKey() == Response->Key);
+	check(!Response || Request.Record.GetKey() == Response->Record.GetKey());
 	check(!Response || Request.UserData == Response->UserData);
 	Writer.BeginObject();
 	if (!Request.Name.IsEmpty())
@@ -138,10 +206,30 @@ void SaveToCompactBinary(FCbWriter& Writer, const FCachePutRequest& Request, con
 	{
 		Writer << ANSITEXTVIEW("Status") << Response->Status;
 	}
+	if (Response)
+	{
+		const bool bMetaConflict = !Response->Record.GetMeta().Equals(Request.Record.GetMeta());
+		const bool bValueConflict = !Algo::Compare(Response->Record.GetValues(), Request.Record.GetValues());
+		if (bMetaConflict || bValueConflict)
+		{
+			Writer.BeginObject(ANSITEXTVIEW("Conflict"));
+			SaveMetaToCompactBinary(Writer, Response->Record.GetMeta(), EMetaSaveFormat::Attachment, SaveAttachment);
+			if (const TConstArrayView<FValueWithId> Values = Response->Record.GetValues(); !Values.IsEmpty())
+			{
+				Writer.BeginArray(ANSITEXTVIEW("Values"));
+				for (const FValueWithId& Value : Values)
+				{
+					SaveToCompactBinary(Writer, Value, SaveAttachment);
+				}
+				Writer.EndArray();
+			}
+			Writer.EndObject();
+		}
+	}
 	Writer.EndObject();
 }
 
-bool LoadFromCompactBinary(FCbFieldView Field, TOptional<FCachePutRequest>& OutRequest, FCachePutResponse* OutResponse, FLoadAttachmentFn LoadAttachment)
+bool LoadFromCompactBinary(FCbFieldView Field, TOptional<FCachePutRequest>& OutRequest, TOptional<FCachePutResponse>* OutResponse, FLoadAttachmentFn LoadAttachment)
 {
 	bool bOk = Field.IsObject();
 
@@ -155,6 +243,11 @@ bool LoadFromCompactBinary(FCbFieldView Field, TOptional<FCachePutRequest>& OutR
 	uint64 UserData = 0;
 	EStatus Status = EStatus::Ok;
 
+	FCbValue ConflictMetaObjectValue;
+	FCbValue ConflictMetaHashValue;
+	FCbValue ConflictMetaSizeValue;
+	FCbArrayView ConflictValuesArray;
+
 	constexpr static uint32 FoundName     = 1 << 0;
 	constexpr static uint32 FoundKey      = 1 << 1;
 	constexpr static uint32 FoundMeta     = 1 << 2;
@@ -164,6 +257,12 @@ bool LoadFromCompactBinary(FCbFieldView Field, TOptional<FCachePutRequest>& OutR
 	constexpr static uint32 FoundPolicy   = 1 << 6;
 	constexpr static uint32 FoundUserData = 1 << 7;
 	constexpr static uint32 FoundStatus   = 1 << 8;
+	constexpr static uint32 FoundConflict = 1 << 9;
+
+	constexpr static uint32 FoundConflictMeta     = 1 << 28;
+	constexpr static uint32 FoundConflictMetaHash = 1 << 29;
+	constexpr static uint32 FoundConflictMetaSize = 1 << 30;
+	constexpr static uint32 FoundConflictValues   = 1 << 31;
 
 	uint32 FoundMask = 0;
 	for (FCbFieldViewIterator It = Field.CreateViewIterator(); It;)
@@ -181,27 +280,7 @@ bool LoadFromCompactBinary(FCbFieldView Field, TOptional<FCachePutRequest>& OutR
 			FoundMask |= FoundKey;
 			++It;
 		}
-		if (!(FoundMask & FoundMeta) && It.GetName() == ANSITEXTVIEW("Meta"))
-		{
-			MetaObjectValue = It;
-			FoundMask |= FoundMeta;
-			++It;
-		}
-		else
-		{
-			if (!(FoundMask & FoundMetaHash) && It.GetName() == ANSITEXTVIEW("MetaHash"))
-			{
-				MetaHashValue = It;
-				FoundMask |= FoundMetaHash;
-				++It;
-			}
-			if (!(FoundMask & FoundMetaSize) && It.GetName() == ANSITEXTVIEW("MetaSize"))
-			{
-				MetaSizeValue = It;
-				FoundMask |= FoundMetaSize;
-				++It;
-			}
-		}
+		LoadMetaFromCompactBinaryIterator<uint32, FoundMeta, FoundMetaHash, FoundMetaSize>(FoundMask, It, MetaObjectValue, MetaHashValue, MetaSizeValue);
 		if (!(FoundMask & FoundValues) && It.GetName() == ANSITEXTVIEW("Values"))
 		{
 			ValuesArray = It.AsArrayView();
@@ -226,6 +305,27 @@ bool LoadFromCompactBinary(FCbFieldView Field, TOptional<FCachePutRequest>& OutR
 			FoundMask |= FoundStatus;
 			++It;
 		}
+		if (!(FoundMask & FoundConflict) && It.GetName() == ANSITEXTVIEW("Conflict"))
+		{
+			for (FCbFieldViewIterator ConflictIt = It.CreateViewIterator(); ConflictIt;)
+			{
+				const uint32 LastFoundConflictMask = FoundMask;
+				LoadMetaFromCompactBinaryIterator<uint32, FoundConflictMeta, FoundConflictMetaHash, FoundConflictMetaSize>(
+					FoundMask, ConflictIt, ConflictMetaObjectValue, ConflictMetaHashValue, ConflictMetaSizeValue);
+				if (!(FoundMask & FoundConflictValues) && ConflictIt.GetName() == ANSITEXTVIEW("Values"))
+				{
+					ConflictValuesArray = ConflictIt.AsArrayView();
+					FoundMask |= FoundConflictValues;
+					++ConflictIt;
+				}
+				if (LastFoundConflictMask == FoundMask)
+				{
+					++ConflictIt;
+				}
+			}
+			FoundMask |= FoundConflict;
+			++It;
+		}
 		if (LastFoundMask == FoundMask)
 		{
 			++It;
@@ -233,37 +333,38 @@ bool LoadFromCompactBinary(FCbFieldView Field, TOptional<FCachePutRequest>& OutR
 	}
 
 	FCbObject Meta;
-	if (MetaObjectValue)
-	{
-		LoadMetaFromObject(MetaObjectValue, Meta);
-	}
-	else if (LoadAttachment && MetaHashValue && MetaSizeValue)
-	{
-		LoadMetaFromHashAndSize(MetaHashValue, MetaSizeValue, Meta, LoadAttachment);
-	}
+	LoadMetaFromCompactBinaryValues(MetaObjectValue, MetaHashValue, MetaSizeValue, Meta, LoadAttachment);
 
 	TArray<FValueWithId> Values;
-	Values.Reserve(int32(ValuesArray.Num()));
-	for (FCbFieldView ValueField : ValuesArray)
-	{
-		FValueWithId& Value = Values.AddDefaulted_GetRef();
-		bOk &= LoadFromCompactBinary(ValueField, Value, LoadAttachment);
-	}
-	if (UNLIKELY(!Algo::IsSorted(Values)))
+	bOk &= LoadValuesFromCompactBinary(ValuesArray, Values, LoadAttachment);
+
+	const bool bConflict = !!(FoundMask & FoundConflict);
+	FCbObject ConflictMeta;
+	TArray<FValueWithId> ConflictValues;
+	if (OutResponse && bConflict)
 	{
-		Algo::Sort(Values);
+		LoadMetaFromCompactBinaryValues(ConflictMetaObjectValue, ConflictMetaHashValue, ConflictMetaSizeValue, ConflictMeta, LoadAttachment);
+		bOk &= LoadValuesFromCompactBinary(ConflictValuesArray, ConflictValues, LoadAttachment);
 	}
 
 	if (bOk)
 	{
+		FCacheRecord Record = FCacheRecord::CreateByMove(Key, MoveTemp(Meta), MoveTemp(Values));
 		if (OutResponse)
 		{
-			*OutResponse = {Name, Key, UserData, Status};
+			FCacheRecord ResponseRecord = bConflict ? FCacheRecord::CreateByMove(Key, MoveTemp(ConflictMeta), MoveTemp(ConflictValues)) : Record;
+			PRAGMA_DISABLE_DEPRECATION_WARNINGS;
+			*OutResponse = {Name, Key, MoveTemp(ResponseRecord), UserData, Status};
+			PRAGMA_ENABLE_DEPRECATION_WARNINGS;
 		}
-		OutRequest = {Name, FCacheRecord::CreateByMove(Key, MoveTemp(Meta), MoveTemp(Values)), Policy, UserData};
+		OutRequest = {Name, MoveTemp(Record), Policy, UserData};
 	}
 	else
 	{
+		if (OutResponse)
+		{
+			OutResponse->Reset();
+		}
 		OutRequest.Reset();
 	}
 
@@ -344,27 +445,7 @@ bool LoadFromCompactBinary(FCbFieldView Field, FCacheGetRequest& OutRequest, TOp
 			FoundMask |= FoundKey;
 			++It;
 		}
-		if (!(FoundMask & FoundMeta) && It.GetName() == ANSITEXTVIEW("Meta"))
-		{
-			MetaObjectValue = It;
-			FoundMask |= FoundMeta;
-			++It;
-		}
-		else
-		{
-			if (!(FoundMask & FoundMetaHash) && It.GetName() == ANSITEXTVIEW("MetaHash"))
-			{
-				MetaHashValue = It;
-				FoundMask |= FoundMetaHash;
-				++It;
-			}
-			if (!(FoundMask & FoundMetaSize) && It.GetName() == ANSITEXTVIEW("MetaSize"))
-			{
-				MetaSizeValue = It;
-				FoundMask |= FoundMetaSize;
-				++It;
-			}
-		}
+		LoadMetaFromCompactBinaryIterator<uint32, FoundMeta, FoundMetaHash, FoundMetaSize>(FoundMask, It, MetaObjectValue, MetaHashValue, MetaSizeValue);
 		if (!(FoundMask & FoundValues) && It.GetName() == ANSITEXTVIEW("Values"))
 		{
 			ValuesArray = It.AsArrayView();
@@ -398,26 +479,10 @@ bool LoadFromCompactBinary(FCbFieldView Field, FCacheGetRequest& OutRequest, TOp
 	if (OutResponse)
 	{
 		FCbObject Meta;
-		if (MetaObjectValue)
-		{
-			LoadMetaFromObject(MetaObjectValue, Meta);
-		}
-		else if (LoadAttachment && MetaHashValue && MetaSizeValue)
-		{
-			LoadMetaFromHashAndSize(MetaHashValue, MetaSizeValue, Meta, LoadAttachment);
-		}
+		LoadMetaFromCompactBinaryValues(MetaObjectValue, MetaHashValue, MetaSizeValue, Meta, LoadAttachment);
 
 		TArray<FValueWithId> Values;
-		Values.Reserve(int32(ValuesArray.Num()));
-		for (FCbFieldView ValueField : ValuesArray)
-		{
-			FValueWithId& Value = Values.AddDefaulted_GetRef();
-			bOk &= LoadFromCompactBinary(ValueField, Value, LoadAttachment);
-		}
-		if (UNLIKELY(!Algo::IsSorted(Values)))
-		{
-			Algo::Sort(Values);
-		}
+		bOk &= LoadValuesFromCompactBinary(ValuesArray, Values, LoadAttachment);
 
 		if (bOk)
 		{
@@ -461,6 +526,16 @@ void SaveToCompactBinary(FCbWriter& Writer, const FCachePutValueRequest& Request
 	{
 		Writer << ANSITEXTVIEW("Status") << Response->Status;
 	}
+	if (Response && Request.Value != Response->Value)
+	{
+		Writer.BeginObject(ANSITEXTVIEW("Conflict"));
+		if (Response->Value != FValue::Null)
+		{
+			Writer.SetName(ANSITEXTVIEW("Value"));
+			SaveToCompactBinary(Writer, Response->Value, SaveAttachment);
+		}
+		Writer.EndObject();
+	}
 	Writer.EndObject();
 }
 
@@ -468,6 +543,7 @@ bool LoadFromCompactBinary(FCbFieldView Field, FCachePutValueRequest& OutRequest
 {
 	bool bOk = Field.IsObject();
 
+	FValue ConflictValue;
 	EStatus Status = EStatus::Ok;
 
 	constexpr static uint32 FoundName     = 1 << 0;
@@ -476,6 +552,7 @@ bool LoadFromCompactBinary(FCbFieldView Field, FCachePutValueRequest& OutRequest
 	constexpr static uint32 FoundPolicy   = 1 << 3;
 	constexpr static uint32 FoundUserData = 1 << 4;
 	constexpr static uint32 FoundStatus   = 1 << 5;
+	constexpr static uint32 FoundConflict = 1 << 6;
 
 	uint32 FoundMask = 0;
 	for (FCbFieldViewIterator It = Field.CreateViewIterator(); It;)
@@ -517,6 +594,18 @@ bool LoadFromCompactBinary(FCbFieldView Field, FCachePutValueRequest& OutRequest
 			FoundMask |= FoundStatus;
 			++It;
 		}
+		if (!(FoundMask & FoundConflict) && It.GetName() == ANSITEXTVIEW("Conflict"))
+		{
+			if (OutResponse)
+			{
+				if (FCbFieldViewIterator ConflictIt = It.CreateViewIterator(); ConflictIt && ConflictIt.GetName() == ANSITEXTVIEW("Value"))
+				{
+					LoadFromCompactBinary(ConflictIt, ConflictValue, LoadAttachment);
+				}
+			}
+			FoundMask |= FoundConflict;
+			++It;
+		}
 		if (LastFoundMask == FoundMask)
 		{
 			++It;
@@ -525,7 +614,8 @@ bool LoadFromCompactBinary(FCbFieldView Field, FCachePutValueRequest& OutRequest
 
 	if (OutResponse)
 	{
-		*OutResponse = {OutRequest.Name, OutRequest.Key, OutRequest.UserData, Status};
+		FValue ResponseValue = (FoundMask & FoundConflict) ? MoveTemp(ConflictValue) : CopyTemp(OutRequest.Value);
+		*OutResponse = {OutRequest.Name, OutRequest.Key, MoveTemp(ResponseValue), OutRequest.UserData, Status};
 	}
 
 	return bOk;
@@ -1046,19 +1136,9 @@ bool LoadFromCompactBinary(FCbFieldView Field, FOptionalCacheRecord& OutRecord,
 		}
 	}
 
-	FCbArrayView ValuesArray = Field[ANSITEXTVIEW("Values")].AsArrayView();
-	Values.Reserve(int32(ValuesArray.Num()));
-	for (FCbFieldView ValueField : ValuesArray)
+	if (!LoadValuesFromCompactBinary(Field[ANSITEXTVIEW("Values")].AsArrayView(), Values, LoadAttachment))
 	{
-		FValueWithId& Value = Values.AddDefaulted_GetRef();
-		if (!LoadFromCompactBinary(ValueField, Value, LoadAttachment))
-		{
-			return false;
-		}
-	}
-	if (UNLIKELY(!Algo::IsSorted(Values)))
-	{
-		Algo::Sort(Values);
+		return false;
 	}
 
 	OutRecord = FCacheRecord::CreateByMove(Key, MoveTemp(Meta), MoveTemp(Values));
diff --git a/Engine/Source/Developer/DerivedDataCache/Public/DerivedDataCache.h b/Engine/Source/Developer/DerivedDataCache/Public/DerivedDataCache.h
--- a/Engine/Source/Developer/DerivedDataCache/Public/DerivedDataCache.h
+++ b/Engine/Source/Developer/DerivedDataCache/Public/DerivedDataCache.h
@@ -144,7 +144,7 @@ struct FCachePutRequest
 	FCacheRecord Record;
 
 	/** Flags to control the behavior of the request. See FCacheRecordPolicy. */
-	FCacheRecordPolicy Policy;
+	FCacheRecordPolicy Policy = ECachePolicy::Default | ECachePolicy::SkipData;
 
 	/** A value that will be returned in the completion callback. */
 	uint64 UserData = 0;
@@ -154,20 +154,31 @@ struct FCachePutRequest
 };
 
 /** Parameters for the completion callback for cache put requests. */
+PRAGMA_DISABLE_DEPRECATION_WARNINGS; // Remove this when the deprecated Key is removed.
 struct FCachePutResponse
 {
 	/** A copy of the name from the request. */
 	FSharedString Name;
 
 	/** A copy of the key from the request. */
+	UE_DEPRECATED(5.8, "Access the key through Record.GetKey().")
 	FCacheKey Key;
 
+	/**
+	 * A copy of the record from the request, or the conflicting record if there is one.
+	 *
+	 * A conflicting record is returned when the cache contains a record with the requested key
+	 * that differs from the record in the request, and the request policy disallows overwrite.
+	 */
+	FCacheRecord Record;
+
 	/** A copy of the value from the request. */
 	uint64 UserData = 0;
 
 	/** The status of the request. */
 	EStatus Status = EStatus::Error;
 };
+PRAGMA_ENABLE_DEPRECATION_WARNINGS; // Remove this when the deprecated Key is removed.
 
 /** Parameters to request to get a cache record. */
 struct FCacheGetRequest
@@ -242,6 +253,14 @@ struct FCachePutValueResponse
 	/** A copy of the key from the request. */
 	FCacheKey Key;
 
+	/**
+	 * A copy of the value from the request, or the conflicting value if there is one.
+	 *
+	 * A conflicting value is returned when the cache contains a value with the requested key
+	 * that differs from the value in the request, and the request policy disallows overwrite.
+	 */
+	FValue Value;
+
 	/** A copy of the value from the request. */
 	uint64 UserData = 0;
 
diff --git a/Engine/Source/Developer/DerivedDataCache/Public/DerivedDataSerialization.h b/Engine/Source/Developer/DerivedDataCache/Public/DerivedDataSerialization.h
--- a/Engine/Source/Developer/DerivedDataCache/Public/DerivedDataSerialization.h
+++ b/Engine/Source/Developer/DerivedDataCache/Public/DerivedDataSerialization.h
@@ -67,7 +67,7 @@ UE_DEPRECATED(5.8, "Call SaveToCompactBinary(Writer, Request).")
 UE_API FCbWriter& operator<<(FCbWriter& Writer, const FCacheGetChunkRequest& Request);
 
 UE_API void SaveToCompactBinary(FCbWriter& Writer, const FCachePutRequest& Request, const FCachePutResponse* Response = nullptr, FSaveAttachmentFn SaveAttachment = nullptr);
-UE_API bool LoadFromCompactBinary(FCbFieldView Field, TOptional<FCachePutRequest>& OutRequest, FCachePutResponse* OutResponse = nullptr, FLoadAttachmentFn LoadAttachment = nullptr);
+UE_API bool LoadFromCompactBinary(FCbFieldView Field, TOptional<FCachePutRequest>& OutRequest, TOptional<FCachePutResponse>* OutResponse = nullptr, FLoadAttachmentFn LoadAttachment = nullptr);
 
 UE_API void SaveToCompactBinary(FCbWriter& Writer, const FCacheGetRequest& Request, const FCacheGetResponse* Response = nullptr, FSaveAttachmentFn SaveAttachment = nullptr);
 UE_API bool LoadFromCompactBinary(FCbFieldView Field, FCacheGetRequest& OutRequest, TOptional<FCacheGetResponse>* OutResponse = nullptr, FLoadAttachmentFn LoadAttachment = nullptr);
diff --git a/Engine/Source/Developer/DerivedDataCache/Tests/DerivedDataCacheStoreTest.cpp b/Engine/Source/Developer/DerivedDataCache/Tests/DerivedDataCacheStoreTest.cpp
--- a/Engine/Source/Developer/DerivedDataCache/Tests/DerivedDataCacheStoreTest.cpp
+++ b/Engine/Source/Developer/DerivedDataCache/Tests/DerivedDataCacheStoreTest.cpp
@@ -570,7 +570,9 @@ static void TestCacheRecords(ICache& Cache)
 			CHECK(Response.Status == EStatus::Ok);
 			CHECK(Response.UserData == UserData);
 			CHECK(Response.Name == Name);
+			PRAGMA_DISABLE_DEPRECATION_WARNINGS;
 			CHECK(Response.Key == Key);
+			PRAGMA_ENABLE_DEPRECATION_WARNINGS;
 		});
 		BlockingOwner.Wait();
 	}
@@ -721,7 +723,9 @@ static void TestPartialCacheRecords(ICache& Cache)
 			CHECK(Response.Status == EStatus::Ok);
 			CHECK(Response.UserData == UserData);
 			CHECK(Response.Name == Name);
+			PRAGMA_DISABLE_DEPRECATION_WARNINGS;
 			CHECK(Response.Key == Key);
+			PRAGMA_ENABLE_DEPRECATION_WARNINGS;
 		});
 		BlockingOwner.Wait();
 	}
@@ -743,7 +747,9 @@ static void TestPartialCacheRecords(ICache& Cache)
 			CHECK(Response.Status == EStatus::Ok);
 			CHECK(Response.UserData == UserData);
 			CHECK(Response.Name == Name);
+			PRAGMA_DISABLE_DEPRECATION_WARNINGS;
 			CHECK(Response.Key == Key);
+			PRAGMA_ENABLE_DEPRECATION_WARNINGS;
 		});
 		BlockingOwner.Wait();
 	}
diff --git a/Engine/Source/Developer/DerivedDataCache/Tests/DerivedDataSerializationTest.cpp b/Engine/Source/Developer/DerivedDataCache/Tests/DerivedDataSerializationTest.cpp
--- a/Engine/Source/Developer/DerivedDataCache/Tests/DerivedDataSerializationTest.cpp
+++ b/Engine/Source/Developer/DerivedDataCache/Tests/DerivedDataSerializationTest.cpp
@@ -16,6 +16,28 @@
 namespace UE::DerivedData
 {
 
+static FCbField ReverseObjectFields(const FCbField& MaybeObject)
+{
+	if (!MaybeObject.IsObject())
+	{
+		return MaybeObject;
+	}
+
+	TArray<FCbField> Fields;
+	for (const FCbField& Field : MaybeObject)
+	{
+		Fields.Add(Field);
+	}
+	FCbWriter Writer;
+	Writer.BeginObject();
+	for (const FCbField& Field : ReverseIterate(Fields))
+	{
+		Writer.AddField(Field.GetName(), ReverseObjectFields(Field));
+	}
+	Writer.EndObject();
+	return Writer.Save();
+}
+
 TEST_CASE("DerivedData::Serialization::CompactBinary", "[DerivedData]")
 {
 	const FSharedString Name(TEXT("Name"));
@@ -25,6 +47,7 @@ TEST_CASE("DerivedData::Serialization::CompactBinary", "[DerivedData]")
 		TCbWriter<64> Writer;
 		Writer.BeginObject();
 		Writer.AddString(ANSITEXTVIEW("Key"), ANSITEXTVIEW("Value"));
+		Writer.AddString(ANSITEXTVIEW("Name"), ANSITEXTVIEW("Test"));
 		Writer.EndObject();
 		return Writer.Save().AsObject();
 	}();
@@ -49,6 +72,17 @@ TEST_CASE("DerivedData::Serialization::CompactBinary", "[DerivedData]")
 	const FSharedBuffer EmptyBuffer = FUniqueBuffer::Alloc(0).MoveToShared();
 	const FIoHash EmptyBufferHash = FIoHash::HashBuffer(EmptyBuffer);
 
+	const FCbObject DiffMeta = ReverseObjectFields(Meta.AsField()).AsObject();
+	const FIoHash DiffMetaHash = DiffMeta.GetHash();
+	const FValue DiffValue = FValue::Compress(FSharedBuffer::MakeView(MakeMemoryView<uint8>({0, 1, 2, 3, 4, 5, 6, 7, 8})));
+	const FCacheRecord DiffRecord = [&]
+	{
+		FCacheRecordBuilder Builder(Key);
+		Builder.SetMeta(CopyTemp(DiffMeta));
+		Builder.AddValue(ValueId, DiffValue);
+		return Builder.Build();
+	}();
+
 	const auto LoadAttachment = [&](const FIoHash& Hash) -> FCbAttachment
 	{
 		if (Hash == Value.GetRawHash())
@@ -63,6 +97,14 @@ TEST_CASE("DerivedData::Serialization::CompactBinary", "[DerivedData]")
 		{
 			return FCbAttachment(EmptyBuffer);
 		}
+		if (Hash == DiffValue.GetRawHash())
+		{
+			return FCbAttachment(DiffValue.GetData());
+		}
+		if (Hash == DiffMetaHash)
+		{
+			return FCbAttachment(DiffMeta);
+		}
 		return FCbAttachment();
 	};
 
@@ -72,24 +114,6 @@ TEST_CASE("DerivedData::Serialization::CompactBinary", "[DerivedData]")
 		CHECK((Hash == Value.GetRawHash() || Hash == MetaHash || Hash == EmptyBufferHash));
 	};
 
-	const auto ReverseObjectFields = [&](const FCbField& ObjectField) -> FCbField
-	{
-		REQUIRE(ObjectField.IsObject());
-		TArray<FCbFieldView> Fields;
-		for (FCbFieldView Field : ObjectField)
-		{
-			Fields.Add(Field);
-		}
-		FCbWriter Writer;
-		Writer.BeginObject();
-		for (FCbFieldView Field : ReverseIterate(Fields))
-		{
-			Writer.AddField(Field.GetName(), Field);
-		}
-		Writer.EndObject();
-		return Writer.Save();
-	};
-
 	SECTION("CachePutRequest")
 	{
 		const FCachePutRequest Request{Name, Record, RecordPolicy, UserData};
@@ -135,7 +159,54 @@ TEST_CASE("DerivedData::Serialization::CompactBinary", "[DerivedData]")
 	SECTION("CachePutResponse")
 	{
 		const FCachePutRequest Request{Name, Record, RecordPolicy, UserData};
-		const FCachePutResponse Response{Name, Key, UserData, EStatus::Canceled};
+		const FCachePutResponse Response{Name, Key, Record, UserData, EStatus::Canceled};
+
+		const FCbField OriginalField = [&]
+		{
+			TCbWriter<512> Writer;
+			SaveToCompactBinary(Writer, Request, &Response);
+			return Writer.Save();
+		}();
+		const FCbField ReversedField = ReverseObjectFields(OriginalField);
+		const FCbField Field = GENERATE_COPY(OriginalField, ReversedField);
+
+		{
+			TOptional<FCachePutRequest> LoadedRequest;
+			TOptional<FCachePutResponse> LoadedResponse;
+			CHECK(LoadFromCompactBinary(Field, LoadedRequest, &LoadedResponse));
+			CHECKED_IF(LoadedResponse)
+			{
+				CHECK(LoadedResponse->Name == Response.Name);
+				CHECK(LoadedResponse->Record.GetKey() == Response.Record.GetKey());
+				CHECK_FALSE(LoadedResponse->Record.GetMeta());
+				CHECK(LoadedResponse->Record.GetValue(ValueId) == Value);
+				CHECK_FALSE(LoadedResponse->Record.GetValue(ValueId).HasData());
+				CHECK(LoadedResponse->UserData == Response.UserData);
+				CHECK(LoadedResponse->Status == Response.Status);
+			}
+		}
+
+		{
+			TOptional<FCachePutRequest> LoadedRequest;
+			TOptional<FCachePutResponse> LoadedResponse;
+			CHECK(LoadFromCompactBinary(Field, LoadedRequest, &LoadedResponse, LoadAttachment));
+			CHECKED_IF(LoadedResponse)
+			{
+				CHECK(LoadedResponse->Name == Response.Name);
+				CHECK(LoadedResponse->Record.GetKey() == Response.Record.GetKey());
+				CHECK(LoadedResponse->Record.GetMeta().Equals(Response.Record.GetMeta()));
+				CHECK(LoadedResponse->Record.GetValue(ValueId) == Value);
+				CHECK(LoadedResponse->Record.GetValue(ValueId).HasData());
+				CHECK(LoadedResponse->UserData == Response.UserData);
+				CHECK(LoadedResponse->Status == Response.Status);
+			}
+		}
+	}
+
+	SECTION("CachePutRequestResponse")
+	{
+		const FCachePutRequest Request{Name, Record, RecordPolicy, UserData};
+		const FCachePutResponse Response{Name, Key, DiffRecord, UserData, EStatus::Canceled};
 
 		const FCbField OriginalField = [&]
 		{
@@ -147,12 +218,27 @@ TEST_CASE("DerivedData::Serialization::CompactBinary", "[DerivedData]")
 		const FCbField Field = GENERATE_COPY(OriginalField, ReversedField);
 
 		TOptional<FCachePutRequest> LoadedRequest;
-		FCachePutResponse LoadedResponse;
-		CHECK(LoadFromCompactBinary(Field, LoadedRequest, &LoadedResponse));
-		CHECK(LoadedResponse.Name == Response.Name);
-		CHECK(LoadedResponse.Key == Response.Key);
-		CHECK(LoadedResponse.UserData == Response.UserData);
-		CHECK(LoadedResponse.Status == Response.Status);
+		TOptional<FCachePutResponse> LoadedResponse;
+		CHECK(LoadFromCompactBinary(Field, LoadedRequest, &LoadedResponse, LoadAttachment));
+		CHECKED_IF(LoadedRequest)
+		{
+			CHECK(LoadedRequest->Name == Request.Name);
+			CHECK(LoadedRequest->Record.GetKey() == Request.Record.GetKey());
+			CHECK(LoadedRequest->Record.GetMeta().Equals(Meta));
+			CHECK(LoadedRequest->Record.GetValue(ValueId) == Value);
+			CHECK(LoadedRequest->Record.GetValue(ValueId).HasData());
+			CHECK(LoadedRequest->UserData == Request.UserData);
+		}
+		CHECKED_IF(LoadedResponse)
+		{
+			CHECK(LoadedResponse->Name == Response.Name);
+			CHECK(LoadedResponse->Record.GetKey() == Response.Record.GetKey());
+			CHECK(LoadedResponse->Record.GetMeta().Equals(Response.Record.GetMeta()));
+			CHECK(LoadedResponse->Record.GetValue(ValueId) == DiffValue);
+			CHECK(LoadedResponse->Record.GetValue(ValueId).HasData());
+			CHECK(LoadedResponse->UserData == Response.UserData);
+			CHECK(LoadedResponse->Status == Response.Status);
+		}
 	}
 
 	SECTION("CacheGetRequest")
@@ -261,7 +347,46 @@ TEST_CASE("DerivedData::Serialization::CompactBinary", "[DerivedData]")
 	SECTION("CachePutValueResponse")
 	{
 		const FCachePutValueRequest Request{Name, Key, Value, ValuePolicy, UserData};
-		const FCachePutValueResponse Response{Name, Key, UserData, EStatus::Canceled};
+		const FCachePutValueResponse Response{Name, Key, Value.RemoveData(), UserData, EStatus::Canceled};
+
+		const FCbField OriginalField = [&]
+		{
+			TCbWriter<512> Writer;
+			SaveToCompactBinary(Writer, Request, &Response);
+			return Writer.Save();
+		}();
+		const FCbField ReversedField = ReverseObjectFields(OriginalField);
+		const FCbField Field = GENERATE_COPY(OriginalField, ReversedField);
+
+		{
+			FCachePutValueRequest LoadedRequest;
+			FCachePutValueResponse LoadedResponse;
+			CHECK(LoadFromCompactBinary(Field, LoadedRequest, &LoadedResponse));
+			CHECK(LoadedResponse.Name == Response.Name);
+			CHECK(LoadedResponse.Key == Response.Key);
+			CHECK(LoadedResponse.Value == Response.Value);
+			CHECK_FALSE(LoadedResponse.Value.HasData());
+			CHECK(LoadedResponse.UserData == Response.UserData);
+			CHECK(LoadedResponse.Status == Response.Status);
+		}
+
+		{
+			FCachePutValueRequest LoadedRequest;
+			FCachePutValueResponse LoadedResponse;
+			CHECK(LoadFromCompactBinary(Field, LoadedRequest, &LoadedResponse, LoadAttachment));
+			CHECK(LoadedResponse.Name == Response.Name);
+			CHECK(LoadedResponse.Key == Response.Key);
+			CHECK(LoadedResponse.Value == Response.Value);
+			CHECK(LoadedResponse.Value.HasData());
+			CHECK(LoadedResponse.UserData == Response.UserData);
+			CHECK(LoadedResponse.Status == Response.Status);
+		}
+	}
+
+	SECTION("CachePutValueRequestResponse")
+	{
+		const FCachePutValueRequest Request{Name, Key, Value, ValuePolicy, UserData};
+		const FCachePutValueResponse Response{Name, Key, DiffValue, UserData, EStatus::Canceled};
 
 		const FCbField OriginalField = [&]
 		{
@@ -274,9 +399,16 @@ TEST_CASE("DerivedData::Serialization::CompactBinary", "[DerivedData]")
 
 		FCachePutValueRequest LoadedRequest;
 		FCachePutValueResponse LoadedResponse;
-		CHECK(LoadFromCompactBinary(Field, LoadedRequest, &LoadedResponse));
+		CHECK(LoadFromCompactBinary(Field, LoadedRequest, &LoadedResponse, LoadAttachment));
+		CHECK(LoadedRequest.Name == Request.Name);
+		CHECK(LoadedRequest.Key == Request.Key);
+		CHECK(LoadedRequest.Value == Request.Value);
+		CHECK(LoadedRequest.Value.HasData());
+		CHECK(LoadedRequest.UserData == Request.UserData);
 		CHECK(LoadedResponse.Name == Response.Name);
 		CHECK(LoadedResponse.Key == Response.Key);
+		CHECK(LoadedResponse.Value == Response.Value);
+		CHECK(LoadedResponse.Value.HasData());
 		CHECK(LoadedResponse.UserData == Response.UserData);
 		CHECK(LoadedResponse.Status == Response.Status);
 	}
diff --git a/Engine/Source/Developer/Virtualization/Private/VirtualizationDDCBackend.cpp b/Engine/Source/Developer/Virtualization/Private/VirtualizationDDCBackend.cpp
--- a/Engine/Source/Developer/Virtualization/Private/VirtualizationDDCBackend.cpp
+++ b/Engine/Source/Developer/Virtualization/Private/VirtualizationDDCBackend.cpp
@@ -109,7 +109,6 @@ bool FDDCBackend::PushData(TArrayView<FPushRequest> Requests, EPushFlags Flags)
 			UE::DerivedData::FCacheRecordBuilder RecordBuilder(Key);
 			RecordBuilder.AddValue(ToDerivedDataValueId(Request.GetIdentifier()), DerivedDataValue);
 
-			UE::DerivedData::FCachePutResponse Result;
 			auto Callback = [&Request, &bWasSuccess](UE::DerivedData::FCachePutResponse&& Response)
 			{
 				if (Response.Status == UE::DerivedData::EStatus::Ok)
diff --git a/Engine/Source/Programs/DerivedDataTool/Private/DerivedDataCopyCommand.cpp b/Engine/Source/Programs/DerivedDataTool/Private/DerivedDataCopyCommand.cpp
--- a/Engine/Source/Programs/DerivedDataTool/Private/DerivedDataCopyCommand.cpp
+++ b/Engine/Source/Programs/DerivedDataTool/Private/DerivedDataCopyCommand.cpp
@@ -901,7 +901,7 @@ public:
 		for (const FCachePutRequest& Request : Requests)
 		{
 			const FCachePutResponse& Response = Responses[Index++];
-			if (Response.Status == EStatus::Ok && Filter && Filter.IsMatch(Request.Record.GetKey()))
+			if (Response.Status == EStatus::Ok && Filter && Filter.IsMatch(Response.Record.GetKey()))
 			{
 				AddRecord(Request.Name, Request.Record.GetKey(), Request.Record.GetValues());
 			}
