DDC: Fixed deprecation warnings in the serialization tests

#rnx
#Synced-CL 50298992

diff --git a/Engine/Source/Developer/DerivedDataCache/Tests/DerivedDataSerializationTest.cpp b/Engine/Source/Developer/DerivedDataCache/Tests/DerivedDataSerializationTest.cpp
--- a/Engine/Source/Developer/DerivedDataCache/Tests/DerivedDataSerializationTest.cpp
+++ b/Engine/Source/Developer/DerivedDataCache/Tests/DerivedDataSerializationTest.cpp
@@ -159,7 +159,9 @@ TEST_CASE("DerivedData::Serialization::CompactBinary", "[DerivedData]")
 	SECTION("CachePutResponse")
 	{
 		const FCachePutRequest Request{Name, Record, RecordPolicy, UserData};
+		PRAGMA_DISABLE_DEPRECATION_WARNINGS;
 		const FCachePutResponse Response{Name, Key, Record, UserData, EStatus::Canceled};
+		PRAGMA_ENABLE_DEPRECATION_WARNINGS;
 
 		const FCbField OriginalField = [&]
 		{
@@ -206,7 +208,9 @@ TEST_CASE("DerivedData::Serialization::CompactBinary", "[DerivedData]")
 	SECTION("CachePutRequestResponse")
 	{
 		const FCachePutRequest Request{Name, Record, RecordPolicy, UserData};
+		PRAGMA_DISABLE_DEPRECATION_WARNINGS;
 		const FCachePutResponse Response{Name, Key, DiffRecord, UserData, EStatus::Canceled};
+		PRAGMA_ENABLE_DEPRECATION_WARNINGS;
 
 		const FCbField OriginalField = [&]
 		{
