Fix missing include

#rnx

#Synced-CL 50235042

diff --git a/Engine/Source/Developer/DerivedDataCache/Internal/Containers/ConcurrentChunkedSparseArray.h b/Engine/Source/Developer/DerivedDataCache/Internal/Containers/ConcurrentChunkedSparseArray.h
--- a/Engine/Source/Developer/DerivedDataCache/Internal/Containers/ConcurrentChunkedSparseArray.h
+++ b/Engine/Source/Developer/DerivedDataCache/Internal/Containers/ConcurrentChunkedSparseArray.h
@@ -3,6 +3,7 @@
 #pragma once
 
 #include "CoreTypes.h"
+#include "Containers/Array.h"
 #include "HAL/UnrealMemory.h"
 #include "Math/UnrealMathUtility.h"
 #include "Templates/TypeCompatibleBytes.h"
