DDC: Exported internal config API

#rnx

#Synced-CL 50253651

diff --git a/Engine/Source/Developer/DerivedDataCache/Internal/DerivedDataConfig.h b/Engine/Source/Developer/DerivedDataCache/Internal/DerivedDataConfig.h
--- a/Engine/Source/Developer/DerivedDataCache/Internal/DerivedDataConfig.h
+++ b/Engine/Source/Developer/DerivedDataCache/Internal/DerivedDataConfig.h
@@ -7,6 +7,8 @@
 #include "Containers/StringView.h"
 #include "Templates/FunctionWithContext.h"
 
+#define UE_API DERIVEDDATACACHE_API
+
 namespace UE::DerivedData::Config
 {
 
@@ -20,9 +22,9 @@ namespace UE::DerivedData::Config
 class FConfigIterator
 {
 public:
-	explicit FConfigIterator(FStringView Config);
+	UE_API explicit FConfigIterator(FStringView Config);
 
-	FConfigIterator& operator++();
+	UE_API FConfigIterator& operator++();
 
 	inline bool HasError() const { return bError; }
 	inline explicit operator bool() const { return !bAtEnd; }
@@ -45,19 +47,19 @@ private:
  * @param OutScratch   Builder to use as scratch space when processing escape sequences.
  * @return A view into String if there are no escapes, and a view into OutScratch otherwise.
  */
-[[nodiscard]] FStringView ParseQuotedString(FStringView String, FStringBuilderBase& OutScratch);
+[[nodiscard]] UE_API FStringView ParseQuotedString(FStringView String, FStringBuilderBase& OutScratch);
 
 /** Finds the first field in Config matching Name case-insensitively. */
-[[nodiscard]] bool FindFirstFieldValue(FStringView Config, FStringView Name, FStringView& OutValue);
+[[nodiscard]] UE_API bool FindFirstFieldValue(FStringView Config, FStringView Name, FStringView& OutValue);
 
 /** Visits every field in Config matching Name case-insensitively. */
-[[nodiscard]] bool VisitFieldValues(FStringView Config, FStringView Name, TFunctionWithContext<void (FStringView)> OnValue);
+[[nodiscard]] UE_API bool VisitFieldValues(FStringView Config, FStringView Name, TFunctionWithContext<void (FStringView)> OnValue);
 
 /** Splits "Name(Config)" to ("Name", "(Config)"), "Name" to ("Name", ""), "(Config)" to ("", "(Config)"). */
-[[nodiscard]] bool TrySplitNameAndConfig(FStringView Value, FStringView& OutName, FStringView& OutConfig);
+[[nodiscard]] UE_API bool TrySplitNameAndConfig(FStringView Value, FStringView& OutName, FStringView& OutConfig);
 
 /** Appends "Name=Value", or "Name" if Value is empty, with extra whitespace removed from any valid nested objects. */
-void AppendConfig(FStringBuilderBase& OutConfig, FStringView Name, FStringView Value);
+UE_API void AppendConfig(FStringBuilderBase& OutConfig, FStringView Name, FStringView Value);
 
 } // UE::DerivedData::Config
 
@@ -75,7 +77,7 @@ namespace UE::DerivedData
  * @param OutConfig           Appended with the configuration for the cache graph if it is found.
  * @return True if the configuration was found, otherwise false.
  */
-bool TryFindCacheGraphConfig(FStringView GraphNameOrConfig, FStringBuilderBase& OutConfig);
+UE_API bool TryFindCacheGraphConfig(FStringView GraphNameOrConfig, FStringBuilderBase& OutConfig);
 
 /**
  * Finds the configuration for a cache store, optionally within a specific cache graph.
@@ -93,6 +95,8 @@ bool TryFindCacheGraphConfig(FStringView GraphNameOrConfig, FStringBuilderBase&
  * @param OutConfig           Appended with the configuration for the cache store if it is found.
  * @return True if the configuration was found, otherwise false.
  */
-bool TryFindCacheStoreConfig(FStringView StoreName, FStringView GraphNameOrConfig, FStringBuilderBase& OutConfig);
+UE_API bool TryFindCacheStoreConfig(FStringView StoreName, FStringView GraphNameOrConfig, FStringBuilderBase& OutConfig);
 
 } // UE::DerivedData
+
+#undef UE_API
