DDC: Optimized FProbingHashTable using the MSB of the hash

This is how it was orginally authored but I changed it to the LSB at the suggestion of BugHawk. Unfortunately, BugHawk missed that this means that large hash tables will have mostly identical single-byte hashes in each bucket, defeating the purpose of optimization.

#Synced-CL 50253651

diff --git a/Engine/Source/Developer/DerivedDataCache/Internal/Containers/ProbingHashTable.h b/Engine/Source/Developer/DerivedDataCache/Internal/Containers/ProbingHashTable.h
--- a/Engine/Source/Developer/DerivedDataCache/Internal/Containers/ProbingHashTable.h
+++ b/Engine/Source/Developer/DerivedDataCache/Internal/Containers/ProbingHashTable.h
@@ -13,6 +13,10 @@ namespace UE
 /**
  * A linear-probing hash table that can be used to index another data structure.
  *
+ * The least significant bits of the hash are used to pick a bucket and the most significant bits
+ * of the hash are used to quickly match slots within a bucket. For the best performance, hashes
+ * must have good distribution across their whole range, not only in the lower bits.
+ *
  * Example:
  *	uint32 Index;
  *	uint32 Hash = GetTypeHash(ID);
@@ -42,7 +46,7 @@ class FProbingHashTable
 	{
 		// Byte 0:0 is set when probing must consider the next block if a match is not found in this block.
 		// Byte 0:1-7 are a mask of which slots are occupied in this block.
-		// Bytes 1-7 are the least significant byte of the hash of the element in the corresponding slot.
+		// Bytes 1-7 are the most significant byte of the hash of the element in the corresponding slot.
 		uint64 MaskAndHashes = 0;
 		FSlot Slots[SlotsPerBlock];
 
@@ -139,7 +143,7 @@ public:
 			if (const uint64 AvailableSlots = ~Block.MaskAndHashes & FBlock::SlotMask)
 			{
 				const uint32 SlotIndexPlusOne = FMath::CountTrailingZeros((uint32)AvailableSlots);
-				Block.MaskAndHashes |= (uint64(1) << SlotIndexPlusOne) | (uint64(uint8(Hash)) << (SlotIndexPlusOne * 8));
+				Block.MaskAndHashes |= (uint64(1) << SlotIndexPlusOne) | (uint64(uint8(Hash >> 24)) << (SlotIndexPlusOne * 8));
 				Block.Slots[SlotIndexPlusOne - 1] = {Hash, Index};
 				return FSlotId::Create(BlockIndex, SlotIndexPlusOne);
 			}
@@ -186,11 +190,11 @@ public:
 	{
 		checkSlow(Id.BlockIndex <= BlockCountMinusOne);
 
-		// Splat the least significant byte of Hash across a uint64.
-		uint64 HashLsb = uint8(Hash);
-		HashLsb |= HashLsb << 32;
-		HashLsb |= HashLsb << 16;
-		HashLsb |= HashLsb << 8;
+		// Splat the most significant byte of Hash across a uint64.
+		uint64 HashMsb = uint8(Hash >> 24);
+		HashMsb |= HashMsb << 32;
+		HashMsb |= HashMsb << 16;
+		HashMsb |= HashMsb << 8;
 
 		uint32 BlockIndex = Id.BlockIndex;
 		uint32 SlotIndexPlusOne = Id.SlotIndexPlusOne;
@@ -206,7 +210,7 @@ public:
 			if (Block.MaskAndHashes & SlotMask)
 			{
 				// Compare by XOR which leaves a 0 byte for each matching hash.
-				uint64 MatchingSlotMask = Block.MaskAndHashes ^ HashLsb;
+				uint64 MatchingSlotMask = Block.MaskAndHashes ^ HashMsb;
 				// Propagate 1 bits in each byte into the least significant bit.
 				MatchingSlotMask |= MatchingSlotMask >> 4;
 				MatchingSlotMask |= MatchingSlotMask >> 2;
@@ -220,7 +224,7 @@ public:
 				// Mask out unoccupied slots.
 				MatchingSlotMask &= Block.MaskAndHashes;
 
-				// Compare the entire hash for any slots that matched the LSB of Hash.
+				// Compare the entire hash for any slots that matched the MSB of Hash.
 				for (; MatchingSlotMask; MatchingSlotMask &= MatchingSlotMask - 1)
 				{
 					SlotIndexPlusOne = FMath::CountTrailingZeros((uint32)MatchingSlotMask);
