DDC: Make cache put responses return the record/value that the cache contains

This can be used to detect conflicts between new and existing data, when not overwriting existing data.

Includes fixes for a put of a record/value with missing data that cannot be recovered from the cache.

#Synced-CL 50253651

diff --git a/Engine/Source/Developer/DerivedDataCache/Private/FileSystemCacheStore.cpp b/Engine/Source/Developer/DerivedDataCache/Private/FileSystemCacheStore.cpp
--- a/Engine/Source/Developer/DerivedDataCache/Private/FileSystemCacheStore.cpp
+++ b/Engine/Source/Developer/DerivedDataCache/Private/FileSystemCacheStore.cpp
@@ -917,7 +917,7 @@ private:
 	bool LegacyDebugOptions(FBackendDebugOptions& Options) final;
 
 private:
-	[[nodiscard]] bool PutCacheRecord(FStringView Name, const FCacheRecord& Record, const FCacheRecordPolicy& Policy, FRequestStats& Stats);
+	[[nodiscard]] FOptionalCacheRecord PutCacheRecord(FStringView Name, const FCacheRecord& Record, const FCacheRecordPolicy& Policy, FRequestStats& Stats);
 
 	[[nodiscard]] FOptionalCacheRecord GetCacheRecordOnly(
 		FStringView Name,
@@ -930,8 +930,14 @@ private:
 		const FCacheRecordPolicy& Policy,
 		EStatus& OutStatus,
 		FRequestStats& Stats);
+	[[nodiscard]] FOptionalCacheRecord GetCacheRecordContent(
+		const FStringView Name,
+		const FCacheRecord& Record,
+		const FCacheRecordPolicy& Policy,
+		EStatus& OutStatus,
+		FRequestStats& Stats);
 
-	[[nodiscard]] bool PutCacheValue(FStringView Name, const FCacheKey& Key, const FValue& Value, ECachePolicy Policy, FRequestStats& Stats);
+	[[nodiscard]] TOptional<FValue> PutCacheValue(FStringView Name, const FCacheKey& Key, const FValue& Value, ECachePolicy Policy, FRequestStats& Stats);
 
 	[[nodiscard]] bool GetCacheValueOnly(FStringView Name, const FCacheKey& Key, ECachePolicy Policy, FValue& OutValue, FRequestStats& Stats);
 	[[nodiscard]] bool GetCacheValue(FStringView Name, const FCacheKey& Key, ECachePolicy Policy, FValue& OutValue, FRequestStats& Stats);
@@ -1617,31 +1623,37 @@ void FFileSystemCacheStore::Put(
 {
 	for (const FCachePutRequest& Request : Requests)
 	{
-		bool bOk;
+		FOptionalCacheRecord Record;
 		FRequestStats RequestStats;
 		RequestStats.Name = Request.Name;
 		RequestStats.Bucket = Request.Record.GetKey().Bucket;
 		RequestStats.Type = ERequestType::Record;
 		RequestStats.Op = ERequestOp::Put;
 		{
-			const FCacheRecord& Record = Request.Record;
 			TRACE_CPUPROFILER_EVENT_SCOPE(FileSystemDDC_Put);
 			TRACE_COUNTER_INCREMENT(FileSystemDDC_Put);
 			FRequestTimer RequestTimer(RequestStats);
 			uint64 WriteSize = 0;
-			bOk = PutCacheRecord(Request.Name, Record, Request.Policy, RequestStats);
-			if (bOk)
+			Record = PutCacheRecord(Request.Name, Request.Record, Request.Policy, RequestStats);
+			if (Record)
 			{
 				UE_LOG(LogDerivedDataCache, Verbose, TEXT("%s: Cache put complete for %s from '%s'"),
-					*CachePath, *WriteToString<96>(Record.GetKey()), *Request.Name);
+					*CachePath, *WriteToString<96>(Record.Get().GetKey()), *Request.Name);
 				TRACE_COUNTER_INCREMENT(FileSystemDDC_PutHit);
 				TRACE_COUNTER_ADD(FileSystemDDC_BytesRead, RequestStats.PhysicalReadSize);
 				TRACE_COUNTER_ADD(FileSystemDDC_BytesWritten, RequestStats.PhysicalWriteSize);
 			}
 		}
-		RequestStats.Status = bOk ? EStatus::Ok : EStatus::Error;
+		const EStatus Status = Record ? EStatus::Ok : EStatus::Error;
+		RequestStats.Status = Status;
 		StoreStats->AddRequest(RequestStats);
-		OnComplete(Request.MakeResponse(bOk ? EStatus::Ok : EStatus::Error));
+		if (!Record)
+		{
+			Record = FCacheRecord(Request.Record.GetKey());
+		}
+		PRAGMA_DISABLE_DEPRECATION_WARNINGS;
+		OnComplete({Request.Name, Record.Get().GetKey(), Record.Get(), Request.UserData, Status});
+		PRAGMA_ENABLE_DEPRECATION_WARNINGS;
 	}
 }
 
@@ -1690,7 +1702,7 @@ void FFileSystemCacheStore::PutValue(
 {
 	for (const FCachePutValueRequest& Request : Requests)
 	{
-		bool bOk;
+		TOptional<FValue> Value;
 		FRequestStats RequestStats;
 		RequestStats.Name = Request.Name;
 		RequestStats.Bucket = Request.Key.Bucket;
@@ -1701,8 +1713,8 @@ void FFileSystemCacheStore::PutValue(
 			TRACE_COUNTER_INCREMENT(FileSystemDDC_Put);
 			FRequestTimer RequestTimer(RequestStats);
 			uint64 WriteSize = 0;
-			bOk = PutCacheValue(Request.Name, Request.Key, Request.Value, Request.Policy, RequestStats);
-			if (bOk)
+			Value = PutCacheValue(Request.Name, Request.Key, Request.Value, Request.Policy, RequestStats);
+			if (Value)
 			{
 				UE_LOG(LogDerivedDataCache, Verbose, TEXT("%s: Cache put complete for %s from '%s'"),
 					*CachePath, *WriteToString<96>(Request.Key), *Request.Name);
@@ -1711,9 +1723,10 @@ void FFileSystemCacheStore::PutValue(
 				TRACE_COUNTER_ADD(FileSystemDDC_BytesWritten, RequestStats.PhysicalWriteSize);
 			}
 		}
-		RequestStats.Status = bOk ? EStatus::Ok : EStatus::Error;
+		const EStatus Status = Value ? EStatus::Ok : EStatus::Error;
+		RequestStats.Status = Status;
 		StoreStats->AddRequest(RequestStats);
-		OnComplete(Request.MakeResponse(bOk ? EStatus::Ok : EStatus::Error));
+		OnComplete({Request.Name, Request.Key, Value.Get({}), Request.UserData, Status});
 	}
 }
 
@@ -1855,7 +1868,7 @@ void FFileSystemCacheStore::GetChunks(
 	}
 }
 
-bool FFileSystemCacheStore::PutCacheRecord(
+FOptionalCacheRecord FFileSystemCacheStore::PutCacheRecord(
 	const FStringView Name,
 	const FCacheRecord& Record,
 	const FCacheRecordPolicy& Policy,
@@ -1871,7 +1884,7 @@ bool FFileSystemCacheStore::PutCacheRecord(
 			Name.Len(),
 			Name.GetData(),
 			bLocalDeactivatedForPerformance ? TEXT("deactivated due to low performance") : TEXT("read-only"));
-		return false;
+		return {};
 	}
 
 	const FCacheKey& Key = Record.GetKey();
@@ -1883,20 +1896,21 @@ bool FFileSystemCacheStore::PutCacheRecord(
 	{
 		UE_LOG(LogDerivedDataCache, VeryVerbose, TEXT("%s: Skipped put of %s from '%.*s' due to cache policy"),
 			*CachePath, *WriteToString<96>(Key), Name.Len(), Name.GetData());
-		return false;
+		return {};
 	}
 
 	if (DebugOptions.ShouldSimulatePutMiss(Key))
 	{
 		UE_LOG(LogDerivedDataCache, Verbose, TEXT("%s: Simulated miss for put of %s from '%.*s'"),
 			*CachePath, *WriteToString<96>(Key), Name.Len(), Name.GetData());
-		return false;
+		return {};
 	}
 
 	TStringBuilder<256> Path;
 	BuildCachePackagePath(Key, Path);
 
 	// Check if there is an existing record package.
+	FOptionalCacheRecord ExistingRecord;
 	FCbPackage ExistingPackage;
 	const ECachePolicy QueryFlag = bRemote ? ECachePolicy::QueryRemote : ECachePolicy::QueryLocal;
 	bool bReplaceExisting = !EnumHasAnyFlags(RecordPolicy, QueryFlag);
@@ -1909,7 +1923,7 @@ bool FFileSystemCacheStore::PutCacheRecord(
 		if (!bSavePackage)
 		{
 			// Save the new package if the existing package is invalid.
-			const FOptionalCacheRecord ExistingRecord = FCacheRecord::Load(ExistingPackage);
+			ExistingRecord = FCacheRecord::Load(ExistingPackage);
 			bSavePackage |= !ExistingRecord;
 			if (ExistingRecord && !Algo::Compare(ExistingRecord.Get().GetValues(), Record.GetValues()))
 			{
@@ -2001,12 +2015,30 @@ bool FFileSystemCacheStore::PutCacheRecord(
 		}
 	}
 
+	// Verify that missing content exists in storage.
+	if (!EnumHasAnyFlags(RecordPolicy, ECachePolicy::PartialRecord))
+	{
+		for (const FValueWithId& Value : Record.GetValues())
+		{
+			if (!Value.HasData() &&
+				!ExistingPackage.FindAttachment(Value.GetRawHash()) &&
+				!GetCacheContentExists(Key, Value.GetRawHash(), Stats))
+			{
+				UE_LOG(LogDerivedDataCache, Verbose,
+					TEXT("%s: Cache put failed due to missing value %s with hash %s for put of %s from '%.*s'"),
+					*CachePath, *WriteToString<32>(Value.GetId()), *WriteToString<48>(Value.GetRawHash()),
+					*WriteToString<96>(Key), Name.Len(), Name.GetData());
+				return {};
+			}
+		}
+	}
+
 	// Save the external content to storage.
 	for (FCompressedBuffer& Content : ExternalContent)
 	{
 		if (!PutCacheContent(Name, Content, Stats))
 		{
-			return false;
+			return {};
 		}
 	}
 
@@ -2016,7 +2048,7 @@ bool FFileSystemCacheStore::PutCacheRecord(
 	{
 		if (!SaveFileWithHash(Path, Name, Stats, WritePackage, bReplaceExisting))
 		{
-			return false;
+			return {};
 		}
 		if (const FCbObject& Meta = Record.GetMeta())
 		{
@@ -2031,7 +2063,10 @@ bool FFileSystemCacheStore::PutCacheRecord(
 		AccessLogWriter->Append(Key, Path);
 	}
 
-	return true;
+	// Load the value content if requested.
+	EStatus Status;
+	const FCacheRecordPolicy GetPolicy = Policy.Transform([](ECachePolicy P) { return P | ECachePolicy::Query; });
+	return GetCacheRecordContent(Name, bSavePackage ? Record : ExistingRecord.Get(), GetPolicy, Status, Stats);
 }
 
 FOptionalCacheRecord FFileSystemCacheStore::GetCacheRecordOnly(
@@ -2118,21 +2153,31 @@ FOptionalCacheRecord FFileSystemCacheStore::GetCacheRecord(
 	FRequestStats& Stats)
 {
 	FOptionalCacheRecord Record = GetCacheRecordOnly(Name, Key, Policy, Stats);
-	if (Record.IsNull())
+	if (Record)
 	{
-		OutStatus = EStatus::Error;
-		return Record;
+		return GetCacheRecordContent(Name, Record.Get(), Policy, OutStatus, Stats);
 	}
+	OutStatus = EStatus::Error;
+	return Record;
+}
 
+FOptionalCacheRecord FFileSystemCacheStore::GetCacheRecordContent(
+	const FStringView Name,
+	const FCacheRecord& Record,
+	const FCacheRecordPolicy& Policy,
+	EStatus& OutStatus,
+	FRequestStats& Stats)
+{
 	OutStatus = EStatus::Ok;
 
+	const FCacheKey& Key = Record.GetKey();
 	FCbObject Meta;
-	TArray<FValueWithId> Values(Record.Get().GetValues());
+	TArray<FValueWithId> Values(Record.GetValues());
 
 	const ECachePolicy RecordPolicy = Policy.GetRecordPolicy();
 	if (!EnumHasAnyFlags(RecordPolicy, ECachePolicy::SkipMeta))
 	{
-		Meta = Record.Get().GetMeta();
+		Meta = Record.GetMeta();
 	}
 
 	for (FValueWithId& Value : Values)
@@ -2152,7 +2197,7 @@ FOptionalCacheRecord FFileSystemCacheStore::GetCacheRecord(
 	return FCacheRecord::CreateByMove(Key, MoveTemp(Meta), MoveTemp(Values));
 }
 
-bool FFileSystemCacheStore::PutCacheValue(
+TOptional<FValue> FFileSystemCacheStore::PutCacheValue(
 	const FStringView Name,
 	const FCacheKey& Key,
 	const FValue& Value,
@@ -2171,7 +2216,7 @@ bool FFileSystemCacheStore::PutCacheValue(
 			bLocalDeactivatedForPerformance ?
 				TEXT("deactivated due to low performance") :
 				TEXT("read-only"));
-		return false;
+		return {};
 	}
 
 	// Skip the request if storing to the cache is disabled.
@@ -2180,17 +2225,18 @@ bool FFileSystemCacheStore::PutCacheValue(
 	{
 		UE_LOG(LogDerivedDataCache, VeryVerbose, TEXT("%s: Skipped put of %s from '%.*s' due to cache policy"),
 			*CachePath, *WriteToString<96>(Key), Name.Len(), Name.GetData());
-		return false;
+		return {};
 	}
 
 	if (DebugOptions.ShouldSimulatePutMiss(Key))
 	{
 		UE_LOG(LogDerivedDataCache, Verbose, TEXT("%s: Simulated miss for put of %s from '%.*s'"),
 			*CachePath, *WriteToString<96>(Key), Name.Len(), Name.GetData());
-		return false;
+		return {};
 	}
 
 	// Check if there is an existing value package.
+	FValue ExistingValue;
 	FCbPackage ExistingPackage;
 	TStringBuilder<256> Path;
 	BuildCachePackagePath(Key, Path);
@@ -2207,6 +2253,7 @@ bool FFileSystemCacheStore::PutCacheValue(
 			const FCbObjectView Object = ExistingPackage.GetObject();
 			const FIoHash RawHash = Object["RawHash"].AsHash();
 			const uint64 RawSize = Object["RawSize"].AsUInt64(MAX_uint64);
+			ExistingValue = FValue(RawHash, RawSize);
 			if (RawHash.IsZero() || RawSize == MAX_uint64)
 			{
 				// Save the new package because the existing package is invalid.
@@ -2261,9 +2308,9 @@ bool FFileSystemCacheStore::PutCacheValue(
 			// Verify that the content exists in storage.
 			if (!GetCacheContentExists(Key, Value.GetRawHash(), Stats))
 			{
-				UE_LOG(LogDerivedDataCache, Verbose, TEXT("%s: Failed due to missing data for put of %s from '%.*s'"),
+				UE_LOG(LogDerivedDataCache, Verbose, TEXT("%s: Cache put failed due to missing data for put of %s from '%.*s'"),
 					*CachePath, *WriteToString<96>(Key), Name.Len(), Name.GetData());
-				return false;
+				return {};
 			}
 		}
 		else if (Value.GetData().GetCompressedSize() <= MaxValueSizeKB * 1024)
@@ -2283,7 +2330,7 @@ bool FFileSystemCacheStore::PutCacheValue(
 	{
 		if (!PutCacheContent(Name, Content, Stats))
 		{
-			return false;
+			return {};
 		}
 	}
 
@@ -2293,7 +2340,7 @@ bool FFileSystemCacheStore::PutCacheValue(
 	{
 		if (!SaveFileWithHash(Path, Name, Stats, WritePackage, bReplaceExisting))
 		{
-			return false;
+			return {};
 		}
 		Stats.LogicalWriteSize += LogicalPackageSize;
 	}
@@ -2303,7 +2350,28 @@ bool FFileSystemCacheStore::PutCacheValue(
 		AccessLogWriter->Append(Key, Path);
 	}
 
-	return true;
+	// Load the value content if requested.
+	if (!bSavePackage)
+	{
+		if (EnumHasAnyFlags(Policy, ECachePolicy::SkipData))
+		{
+			return ExistingValue;
+		}
+		if (const FCbAttachment* const Attachment = ExistingPackage.FindAttachment(ExistingValue.GetRawHash()))
+		{
+			const FCompressedBuffer& Data = Attachment->AsCompressedBinary();
+			if (Data.GetRawHash() == ExistingValue.GetRawHash() && Data.GetRawSize() == ExistingValue.GetRawSize())
+			{
+				return FValue(Data);
+			}
+		}
+	}
+	FValue OutValue;
+	if (GetCacheContent(Name, Key, {}, bSavePackage ? Value : ExistingValue, Policy | ECachePolicy::Query, OutValue, Stats))
+	{
+		return OutValue;
+	}
+	return {};
 }
 
 bool FFileSystemCacheStore::GetCacheValueOnly(
diff --git a/Engine/Source/Developer/DerivedDataCache/Private/HttpCacheStore.cpp b/Engine/Source/Developer/DerivedDataCache/Private/HttpCacheStore.cpp
--- a/Engine/Source/Developer/DerivedDataCache/Private/HttpCacheStore.cpp
+++ b/Engine/Source/Developer/DerivedDataCache/Private/HttpCacheStore.cpp
@@ -6,6 +6,7 @@
 #if WITH_HTTP_DDC_BACKEND
 
 #include "Algo/AllOf.h"
+#include "Algo/Compare.h"
 #include "Algo/NoneOf.h"
 #include "Algo/Transform.h"
 #include "Async/ManualResetEvent.h"
@@ -752,6 +753,7 @@ public:
 	struct FResponse
 	{
 		EStatus Status = EStatus::Error;
+		FCbObject ExistingObject;
 	};
 	using FOnPackageComplete = TUniqueFunction<void (FResponse&& Response)>;
 
@@ -762,6 +764,8 @@ public:
 
 	void Put(const FCacheKey& Key, const FCacheRecordPolicy& Policy, FCbPackage&& Package, FOnPackageComplete&& OnComplete);
 
+	IRequestOwner& GetOwner() const { return Owner; }
+
 	const FRequestStats& ReadStats() const { return RequestStats; }
 	FRequestStats& EditStats() { return RequestStats; }
 
@@ -775,6 +779,8 @@ private:
 	FIoHash ObjectHash;
 	FOnPackageComplete OnPackageComplete;
 
+	FCbObject ExistingObject;
+
 	FRequestStats RequestStats;
 
 	std::atomic<uint32> SuccessfulBlobUploads = 0;
@@ -790,13 +796,13 @@ private:
 
 	FPutPackageOp(FHttpCacheStore& CacheStore, IRequestOwner& Owner, const FSharedString& Name);
 
-	void BeginOperation(bool bFinalize, bool bAllowOverwrite, bool bReportConflicts, FOnCachePutRefComplete&& OnComplete);
+	void BeginOperation(bool bFinalize, bool bAllowOverwrite, FOnCachePutRefComplete&& OnComplete);
 
-	void BeginPutRef(TUniquePtr<FHttpOperation> Operation, bool bFinalize, bool bAllowOverwrite, bool bReportConflicts, FOnCachePutRefComplete&& OnComplete);
-	void EndPutRef(TUniquePtr<FHttpOperation> Operation, bool bFinalize, bool bReportConflicts, FOnCachePutRefComplete&& OnComplete);
+	void BeginPutRef(TUniquePtr<FHttpOperation> Operation, bool bFinalize, bool bAllowOverwrite, FOnCachePutRefComplete&& OnComplete);
+	void EndPutRef(TUniquePtr<FHttpOperation> Operation, bool bFinalize, FOnCachePutRefComplete&& OnComplete);
 
-	void BeginPutBlobs(FCbPackage&& Package, FCachePutRefResponse&& Response, bool bReportConflicts);
-	void EndPutBlob(FHttpOperation* Operation, uint64 LogicalSize, bool bReportConflicts);
+	void BeginPutBlobs(FCbPackage&& Package, FCachePutRefResponse&& Response);
+	void EndPutBlob(FHttpOperation* Operation, uint64 LogicalSize);
 
 	void EndPutRefFinalize(FCachePutRefResponse&& Response);
 
@@ -814,24 +820,22 @@ FHttpCacheStore::FPutPackageOp::FPutPackageOp(FHttpCacheStore& InCacheStore, IRe
 void FHttpCacheStore::FPutPackageOp::Put(const FCacheKey& InKey, const FCacheRecordPolicy& Policy, FCbPackage&& Package, FOnPackageComplete&& OnComplete)
 {
 	const bool bAllowOverwrite = !EnumHasAnyFlags(Policy.GetRecordPolicy(), ECachePolicy::QueryRemote);
-	const bool bReportConflicts = !EnumHasAnyFlags(Policy.GetBasePolicy(), ECachePolicy::NonDeterministic) &&
-		Algo::NoneOf(Policy.GetValuePolicies(), [](const FCacheValuePolicy& Value) { return EnumHasAnyFlags(Value.Policy, ECachePolicy::NonDeterministic); });
 	Key = InKey;
 	Object = Package.GetObject();
 	ObjectHash = Package.GetObjectHash();
 	OnPackageComplete = MoveTemp(OnComplete);
-	BeginOperation(/*bFinalize*/ false, bAllowOverwrite, bReportConflicts, [Self = TRefCountPtr(this), Package = MoveTemp(Package), bReportConflicts](FCachePutRefResponse&& Response) mutable
+	BeginOperation(/*bFinalize*/ false, bAllowOverwrite, [Self = TRefCountPtr(this), Package = MoveTemp(Package)](FCachePutRefResponse&& Response) mutable
 	{
-		Self->BeginPutBlobs(MoveTemp(Package), MoveTemp(Response), bReportConflicts);
+		Self->BeginPutBlobs(MoveTemp(Package), MoveTemp(Response));
 	});
 }
 
-void FHttpCacheStore::FPutPackageOp::BeginOperation(bool bFinalize, bool bAllowOverwrite, bool bReportConflicts, FOnCachePutRefComplete&& OnComplete)
+void FHttpCacheStore::FPutPackageOp::BeginOperation(bool bFinalize, bool bAllowOverwrite, FOnCachePutRefComplete&& OnComplete)
 {
 	CacheStore.WaitForHttpOperationAsync(Owner, bFinalize ? EOperationCategory::PutFinalize : EOperationCategory::PutRef,
-	[Self = TRefCountPtr(this), bFinalize, bAllowOverwrite, bReportConflicts, OnComplete = MoveTemp(OnComplete)](TUniquePtr<FHttpOperation>&& Operation) mutable
+	[Self = TRefCountPtr(this), bFinalize, bAllowOverwrite, OnComplete = MoveTemp(OnComplete)](TUniquePtr<FHttpOperation>&& Operation) mutable
 	{
-		Self->BeginPutRef(MoveTemp(Operation), bFinalize, bAllowOverwrite, bReportConflicts, MoveTemp(OnComplete));
+		Self->BeginPutRef(MoveTemp(Operation), bFinalize, bAllowOverwrite, MoveTemp(OnComplete));
 	});
 }
 
@@ -839,7 +843,6 @@ void FHttpCacheStore::FPutPackageOp::BeginPutRef(
 	TUniquePtr<FHttpOperation> Operation,
 	bool bFinalize,
 	bool bAllowOverwrite,
-	bool bReportConflicts,
 	FOnCachePutRefComplete&& OnComplete)
 {
 	if (UNLIKELY(!Operation))
@@ -891,17 +894,16 @@ void FHttpCacheStore::FPutPackageOp::BeginPutRef(
 	LocalOperation.AddAcceptType(EHttpMediaType::Json);
 
 	RequestTimer.Stop();
-	LocalOperation.SendAsync(Owner, [Self = TRefCountPtr(this), Operation = MoveTemp(Operation), bFinalize, bReportConflicts, OnComplete = MoveTemp(OnComplete)]() mutable
+	LocalOperation.SendAsync(Owner, [Self = TRefCountPtr(this), Operation = MoveTemp(Operation), bFinalize, OnComplete = MoveTemp(OnComplete)]() mutable
 	{
 		Operation->GetStats(Self->RequestStats);
-		Self->EndPutRef(MoveTemp(Operation), bFinalize, bReportConflicts, MoveTemp(OnComplete));
+		Self->EndPutRef(MoveTemp(Operation), bFinalize, MoveTemp(OnComplete));
 	});
 }
 
 void FHttpCacheStore::FPutPackageOp::EndPutRef(
 	TUniquePtr<FHttpOperation> Operation,
 	bool bFinalize,
-	bool bReportConflicts,
 	FOnCachePutRefComplete&& OnComplete)
 {
 	TRACE_CPUPROFILER_EVENT_SCOPE(HttpDDC_PutPackage_EndPutRef);
@@ -919,10 +921,10 @@ void FHttpCacheStore::FPutPackageOp::EndPutRef(
 
 	if (const int32 StatusCode = Operation->GetStatusCode(); StatusCode < 200 || StatusCode > 204)
 	{
-		if (bReportConflicts && (StatusCode == 409))
+		if (StatusCode == 409 && Operation->GetContentType() == EHttpMediaType::CbObject &&
+			ValidateCompactBinary(Operation->GetBody(), ECbValidateMode::Default) == ECbValidateError::None)
 		{
-			UE_LOG(LogDerivedDataCache, Display, TEXT("%s: Cache put found non-deterministic record for %s from '%s'"),
-				*CacheStore.NodeName, *WriteToString<96>(Key), *Name);
+			ExistingObject = FCbObject(Operation->GetBody());
 		}
 		OnComplete({{}, StatusCode == 409 ? EStatus::Ok : EStatus::Error});
 		return;
@@ -967,7 +969,7 @@ void FHttpCacheStore::FPutPackageOp::EndPutRef(
 	OnComplete({NeededBlobHashes, EStatus::Ok});
 }
 
-void FHttpCacheStore::FPutPackageOp::BeginPutBlobs(FCbPackage&& Package, FCachePutRefResponse&& Response, bool bReportConflicts)
+void FHttpCacheStore::FPutPackageOp::BeginPutBlobs(FCbPackage&& Package, FCachePutRefResponse&& Response)
 {
 	if (Response.Status != EStatus::Ok)
 	{
@@ -1039,11 +1041,11 @@ void FHttpCacheStore::FPutPackageOp::BeginPutBlobs(FCbPackage&& Package, FCacheP
 	FRequestBarrier Barrier(Owner);
 	for (const FCompressedBuffer& Blob : Blobs)
 	{
-		CacheStore.WaitForHttpOperationAsync(Owner, EOperationCategory::PutBlobs, [Self = TRefCountPtr(this), Blob, bReportConflicts](TUniquePtr<FHttpOperation>&& Operation)
+		CacheStore.WaitForHttpOperationAsync(Owner, EOperationCategory::PutBlobs, [Self = TRefCountPtr(this), Blob](TUniquePtr<FHttpOperation>&& Operation)
 		{
 			if (UNLIKELY(!Operation))
 			{
-				Self->EndPutBlob(nullptr, 0, bReportConflicts);
+				Self->EndPutBlob(nullptr, 0);
 				return;
 			}
 
@@ -1053,16 +1055,16 @@ void FHttpCacheStore::FPutPackageOp::BeginPutBlobs(FCbPackage&& Package, FCacheP
 			LocalOperation.SetMethod(EHttpMethod::Put);
 			LocalOperation.SetContentType(EHttpMediaType::CompressedBinary);
 			LocalOperation.SetBody(Blob.GetCompressed());
-			LocalOperation.SendAsync(Self->Owner, [Self, Operation = MoveTemp(Operation), LogicalSize = Blob.GetRawSize(), bReportConflicts]
+			LocalOperation.SendAsync(Self->Owner, [Self, Operation = MoveTemp(Operation), LogicalSize = Blob.GetRawSize()]
 			{
 				Operation->GetStats(Self->RequestStats);
-				Self->EndPutBlob(Operation.Get(), LogicalSize, bReportConflicts);
+				Self->EndPutBlob(Operation.Get(), LogicalSize);
 			});
 		});
 	}
 }
 
-void FHttpCacheStore::FPutPackageOp::EndPutBlob(FHttpOperation* Operation, uint64 LogicalSize, bool bReportConflicts)
+void FHttpCacheStore::FPutPackageOp::EndPutBlob(FHttpOperation* Operation, uint64 LogicalSize)
 {
 	if (Operation)
 	{
@@ -1084,7 +1086,7 @@ void FHttpCacheStore::FPutPackageOp::EndPutBlob(FHttpOperation* Operation, uint6
 		}
 		else if (LocalSuccessfulBlobUploads == TotalBlobUploads)
 		{
-			BeginOperation(/*bFinalize*/ true, /*bAllowOverwrite*/ false, bReportConflicts, [Self = TRefCountPtr(this)](FCachePutRefResponse&& Response)
+			BeginOperation(/*bFinalize*/ true, /*bAllowOverwrite*/ false, [Self = TRefCountPtr(this)](FCachePutRefResponse&& Response)
 			{
 				Self->EndPutRefFinalize(MoveTemp(Response));
 			});
@@ -1116,7 +1118,7 @@ void FHttpCacheStore::FPutPackageOp::EndPut(EStatus Status)
 	RequestStats.Status = Status;
 	// Ensuring that the OnPackageComplete method is destroyed by the time we exit this method by moving it to a local scope variable
 	FOnPackageComplete LocalOnComplete = MoveTemp(OnPackageComplete);
-	LocalOnComplete({Status});
+	LocalOnComplete({Status, MoveTemp(ExistingObject)});
 	if (CacheStore.StoreStats)
 	{
 		CacheStore.StoreStats->AddRequest(RequestStats);
@@ -2715,18 +2717,69 @@ void FHttpCacheStore::PutCacheRecordAsync(IRequestOwner& Owner, const FCachePutR
 		Package = Request.Record.Save();
 	}
 
-	Op->Put(Key, Request.Policy, MoveTemp(Package), [Op, Request, OnComplete = MoveTemp(OnComplete)](FPutPackageOp::FResponse&& Response)
+	Op->Put(Key, Request.Policy, MoveTemp(Package), [this, Op, Request, OnComplete = MoveTemp(OnComplete)](FPutPackageOp::FResponse&& OpResponse) mutable
 	{
-		if (Response.Status == EStatus::Ok)
+		if (OpResponse.Status == EStatus::Ok)
 		{
 			if (const FCbObject& Meta = Request.Record.GetMeta())
 			{
 				Op->EditStats().LogicalWriteSize += Meta.GetSize();
 			}
 		}
+
 		TRACE_COUNTER_ADD(HttpDDC_BytesReceived, Op->ReadStats().PhysicalReadSize);
 		TRACE_COUNTER_ADD(HttpDDC_BytesSent, Op->ReadStats().PhysicalWriteSize);
-		OnComplete(Request.MakeResponse(Response.Status));
+
+		FCachePutResponse Response = Request.MakeResponse(OpResponse.Status);
+
+		if (FOptionalCacheRecord ExistingRecord; OpResponse.ExistingObject && LoadFromCompactBinary(OpResponse.ExistingObject.AsFieldView(), ExistingRecord))
+		{
+			// Report puts of non-deterministic records by ignoring known-non-deterministic values.
+			const auto MakeValueTuple = [&Request](const FValueWithId& Value) -> TTuple<FValueId, FIoHash>
+			{
+				return !EnumHasAnyFlags(Request.Policy.GetValuePolicy(Value.GetId()), ECachePolicy::NonDeterministic)
+					? MakeTuple(Value.GetId(), Value.GetRawHash())
+					: TTuple<FValueId, FIoHash>{};
+			};
+			if (!Algo::CompareBy(ExistingRecord.Get().GetValues(), Request.Record.GetValues(), MakeValueTuple))
+			{
+				UE_LOG(LogDerivedDataCache, Display, TEXT("%s: Cache put found non-deterministic record for %s from '%s'"),
+					*NodeName, *WriteToString<96>(Request.Record.GetKey()), *Request.Name);
+			}
+
+			// Compare the whole record to check whether the record from the cache needs to be returned.
+			if (!Algo::Compare(ExistingRecord.Get().GetValues(), Request.Record.GetValues()))
+			{
+				Response.Record = ExistingRecord.Get();
+			}
+
+			// Compare the metadata to check whether the metadata from the cache needs to be returned.
+			if (!ExistingRecord.Get().GetMeta().Equals(Response.Record.GetMeta()))
+			{
+				Response.Record = FCacheRecord::CreateByMove(Response.Record.GetKey(), CopyTemp(ExistingRecord.Get().GetMeta()),
+					TArray<FValueWithId>(Response.Record.GetValues()));
+			}
+		}
+
+		if (EnumHasAnyFlags(Request.Policy.GetRecordPolicy(), ECachePolicy::SkipData) ||
+			Algo::AllOf(Response.Record.GetValues(), [&Request](const FValueWithId& Value)
+			{
+				return Value.HasData() || EnumHasAnyFlags(Request.Policy.GetValuePolicy(Value.GetId()), ECachePolicy::SkipData);
+			}))
+		{
+			return OnComplete(MoveTemp(Response));
+		}
+
+		IRequestOwner& Owner = Op->GetOwner();
+		FRequestBarrier Barrier(Owner);
+		const FCacheRecordPolicy Policy = Request.Policy.Transform([](ECachePolicy P) { return P | ECachePolicy::Query; });
+		GetCacheRecordAsync(Owner, Request.Name, Request.Record.GetKey(), Policy, Request.UserData,
+			[OnComplete = MoveTemp(OnComplete)](FCacheGetResponse&& Response)
+			{
+				PRAGMA_DISABLE_DEPRECATION_WARNINGS;
+				OnComplete({Response.Name, Response.Record.GetKey(), Response.Record, Response.UserData, Response.Status});
+				PRAGMA_ENABLE_DEPRECATION_WARNINGS;
+			});
 	});
 }
 
@@ -2767,20 +2820,37 @@ void FHttpCacheStore::PutCacheValueAsync(IRequestOwner& Owner, const FCachePutVa
 		FRequestTimer RequestTimer(RequestStats);
 
 		FCbWriter Writer;
-		Writer.BeginObject();
-		Writer.AddBinaryAttachment("RawHash", Request.Value.GetRawHash());
-		Writer.AddInteger("RawSize", Request.Value.GetRawSize());
-		Writer.EndObject();
-
+		SaveToCompactBinary(Writer, Request.Value, [&Package](FCbAttachment&& Attachment) { Package.AddAttachment(MoveTemp(Attachment)); });
 		Package.SetObject(Writer.Save().AsObject());
-		Package.AddAttachment(FCbAttachment(Request.Value.GetData()));
 	}
 
-	Op->Put(Request.Key, Request.Policy, MoveTemp(Package), [Op, Request, OnComplete = MoveTemp(OnComplete)](FPutPackageOp::FResponse&& Response)
+	Op->Put(Request.Key, Request.Policy, MoveTemp(Package), [this, Op, Request, OnComplete = MoveTemp(OnComplete)](FPutPackageOp::FResponse&& OpResponse) mutable
 	{
 		TRACE_COUNTER_ADD(HttpDDC_BytesReceived, Op->ReadStats().PhysicalReadSize);
 		TRACE_COUNTER_ADD(HttpDDC_BytesSent, Op->ReadStats().PhysicalWriteSize);
-		OnComplete(Request.MakeResponse(Response.Status));
+
+		FCachePutValueResponse Response = Request.MakeResponse(OpResponse.Status);
+
+		if (FValue ExistingValue; OpResponse.ExistingObject && LoadFromCompactBinary(OpResponse.ExistingObject.AsFieldView(), ExistingValue))
+		{
+			UE_CLOG(!EnumHasAnyFlags(Request.Policy, ECachePolicy::NonDeterministic),
+				LogDerivedDataCache, Display, TEXT("%s: Cache put found non-deterministic value for %s from '%s'"),
+				*NodeName, *WriteToString<96>(Request.Key), *Request.Name);
+			Response.Value = ExistingValue;
+		}
+
+		if (EnumHasAnyFlags(Request.Policy, ECachePolicy::SkipData) || Response.Value.HasData())
+		{
+			return OnComplete(MoveTemp(Response));
+		}
+
+		IRequestOwner& Owner = Op->GetOwner();
+		FRequestBarrier Barrier(Owner);
+		GetCacheValueAsync(Owner, Request.Name, Request.Key, Request.Policy | ECachePolicy::Query, ERequestOp::Put, Request.UserData,
+			[OnComplete = MoveTemp(OnComplete)](FCacheGetValueResponse&& Response)
+			{
+				OnComplete({Response.Name, Response.Key, Response.Value, Response.UserData, Response.Status});
+			});
 	});
 }
 
diff --git a/Engine/Source/Developer/DerivedDataCache/Private/MemoryCacheStore.cpp b/Engine/Source/Developer/DerivedDataCache/Private/MemoryCacheStore.cpp
--- a/Engine/Source/Developer/DerivedDataCache/Private/MemoryCacheStore.cpp
+++ b/Engine/Source/Developer/DerivedDataCache/Private/MemoryCacheStore.cpp
@@ -203,10 +203,13 @@ void FMemoryCacheStore::Put(
 
 	for (const FCachePutRequest& Request : Requests)
 	{
+		FCacheRecord ResponseRecord = Request.Record;
 		EStatus Status = EStatus::Error;
 		ON_SCOPE_EXIT
 		{
-			OnComplete(Request.MakeResponse(Status));
+			PRAGMA_DISABLE_DEPRECATION_WARNINGS;
+			OnComplete({Request.Name, Request.Record.GetKey(), ApplySkipPolicy(ResponseRecord, Request.Policy), Request.UserData, Status});
+			PRAGMA_ENABLE_DEPRECATION_WARNINGS;
 		};
 
 		const FCacheRecord& Record = Request.Record;
@@ -252,6 +255,15 @@ void FMemoryCacheStore::Put(
 		}();
 		Status = EStatus::Ok;
 
+		ON_SCOPE_EXIT
+		{
+			ResponseRecord = FCacheRecord::CreateByMove(Key, CopyTemp(Components.Meta), CopyTemp(Components.Values));
+			if (!bHasExisting && Status == EStatus::Error)
+			{
+				CacheRecords.Remove(Key);
+			}
+		};
+
 		if (bHasExisting && !bReplaceExisting && Algo::AllOf(Components.Values, &FValue::HasData))
 		{
 			continue;
@@ -505,10 +517,11 @@ void FMemoryCacheStore::PutValue(
 
 	for (const FCachePutValueRequest& Request : Requests)
 	{
+		FValue ResponseValue = Request.Value;
 		EStatus Status = EStatus::Error;
 		ON_SCOPE_EXIT
 		{
-			OnComplete(Request.MakeResponse(Status));
+			OnComplete({Request.Name, Request.Key, ApplySkipPolicy(ResponseValue, Request.Policy), Request.UserData, Status});
 		};
 
 		const FCacheKey& Key = Request.Key;
@@ -531,11 +544,6 @@ void FMemoryCacheStore::PutValue(
 
 		FRequestTimer RequestTimer(RequestStats);
 
-		if (!Value.HasData())
-		{
-			continue;
-		}
-
 		if (bCanBeDisabled && DebugOptions.ShouldSimulatePutMiss(Key))
 		{
 			UE_LOG(LogDerivedDataCache, Verbose, TEXT("%s: Simulated miss for put of %s from '%s'"),
@@ -543,7 +551,6 @@ void FMemoryCacheStore::PutValue(
 			continue;
 		}
 
-		const int64 ValueSize = Value.GetData().GetCompressedSize();
 		const bool bReplaceExisting = !EnumHasAnyFlags(Request.Policy, ECachePolicy::QueryLocal);
 
 		FWriteScopeLock ScopeLock(SynchronizationObject);
@@ -551,12 +558,13 @@ void FMemoryCacheStore::PutValue(
 
 		if (ExistingValue && !bReplaceExisting)
 		{
+			ResponseValue = *ExistingValue;
 			Status = EStatus::Ok;
 			continue;
 		}
 
 		const int64 ExistingValueSize = ExistingValue ? ExistingValue->GetData().GetCompressedSize() : 0;
-		const int64 RequiredSize = ValueSize - ExistingValueSize;
+		const int64 RequiredSize = Value.HasData() ? Value.GetData().GetCompressedSize() - ExistingValueSize : 0;
 
 		if (MaxCacheSize > 0 && (CurrentCacheSize + RequiredSize) > MaxCacheSize)
 		{
@@ -569,7 +577,13 @@ void FMemoryCacheStore::PutValue(
 		}
 
 		CurrentCacheSize += RequiredSize;
-		if (ExistingValue)
+		if (!Value.HasData())
+		{
+			UE_LOG(LogDerivedDataCache, Verbose, TEXT("%s: Cache put missing data for put of %s from '%s'"),
+				*Name, *WriteToString<96>(Key), *Request.Name);
+			continue;
+		}
+		else if (ExistingValue)
 		{
 			*ExistingValue = Value;
 		}
diff --git a/Engine/Source/Developer/DerivedDataCache/Private/ZenCacheStore.cpp b/Engine/Source/Developer/DerivedDataCache/Private/ZenCacheStore.cpp
--- a/Engine/Source/Developer/DerivedDataCache/Private/ZenCacheStore.cpp
+++ b/Engine/Source/Developer/DerivedDataCache/Private/ZenCacheStore.cpp
@@ -6,6 +6,7 @@
 
 #if UE_WITH_ZEN
 
+#include "Algo/AnyOf.h"
 #include "Algo/Compare.h"
 #include "Async/ManualResetEvent.h"
 #include "Async/UniqueLock.h"
@@ -25,6 +26,7 @@
 #include "HAL/Thread.h"
 #include "Http/HttpClient.h"
 #include "Http/HttpHostBuilder.h"
+#include "Logging/StructuredLog.h"
 #include "Math/UnrealMathUtility.h"
 #include "Misc/App.h"
 #include "Misc/ConfigCacheIni.h"
@@ -203,6 +205,19 @@
 
 	const Zen::FZenServiceInstance& GetServiceInstance() const { return ZenService.GetInstance(); }
 
+	class FPutOp;
+	class FGetOp;
+	class FPutValueOp;
+	class FGetValueOp;
+	class FGetChunksOp;
+
+	template <typename T, typename... ArgTypes>
+	static TRefCountPtr<T> MakeAsyncOp(ArgTypes&&... Args)
+	{
+		// TODO: This should in-place construct from a pre-allocated memory pool
+		return TRefCountPtr<T>(new T(Forward<ArgTypes>(Args)...));
+	}
+
 private:
 	void Initialize(const FZenCacheStoreParams& Params);
 
@@ -219,13 +234,6 @@
 	void ConditionalUpdateStorageSize();
 	void UpdateStatus();
 
-	template <typename T, typename... ArgTypes>
-	static TRefCountPtr<T> MakeAsyncOp(ArgTypes&&... Args)
-	{
-		// TODO: This should in-place construct from a pre-allocated memory pool
-		return TRefCountPtr<T>(new T(Forward<ArgTypes>(Args)...));
-	}
-
 private:
 	template <typename RequestType>
 	struct TRequestWithStats;
@@ -249,12 +257,6 @@
 
 	void SetRequestUri(IHttpRequest& Request, FAnsiStringView Path);
 
-	class FPutOp;
-	class FGetOp;
-	class FPutValueOp;
-	class FGetValueOp;
-	class FGetChunksOp;
-
 	enum class EHealth
 	{
 		Unknown,
@@ -574,13 +576,30 @@
 	void OnHit(const TRequestWithStats<FCachePutRequest>& RequestWithStats, FCbObject&& DetailsObject)
 	{
 		FRequestTimer RequestTimer(RequestWithStats.Stats);
+
+		if (const FCbObject& Meta = RequestWithStats.Request.Record.GetMeta())
+		{
+			RequestWithStats.Stats.PhysicalWriteSize += Meta.GetSize();
+		}
+		for (const FValueWithId& Value : RequestWithStats.Request.Record.GetValues())
+		{
+			RequestWithStats.Stats.AddLogicalWrite(Value);
+			RequestWithStats.Stats.PhysicalWriteSize += Value.GetData().GetCompressedSize();
+		}
+
+		TRACE_COUNTER_ADD(ZenDDC_BytesSent, int64(RequestWithStats.Stats.PhysicalWriteSize));
+		TRACE_COUNTER_INCREMENT(ZenDDC_PutHit);
+
+		FCachePutResponse Response = RequestWithStats.Request.MakeResponse(EStatus::Ok);
+
 		if (DetailsObject)
 		{
-			FCbObject ExistingRecordObject(DetailsObject[ANSITEXTVIEW("Record")].AsObject());
-			if (ExistingRecordObject)
+			if (FCbObject ExistingRecordObject = DetailsObject[ANSITEXTVIEW("Record")].AsObject())
 			{
 				FCbPackage ExistingRecordPackage(ExistingRecordObject);
-				const FOptionalCacheRecord ExistingRecord = FCacheRecord::Load(ExistingRecordPackage);
+				FOptionalCacheRecord ExistingRecord = FCacheRecord::Load(ExistingRecordPackage);
+
+				// Report puts of non-deterministic records by ignoring known-non-deterministic values.
 				const auto MakeValueTuple = [&RequestWithStats](const FValueWithId& Value) -> TTuple<FValueId, FIoHash>
 				{
 					return !EnumHasAnyFlags(RequestWithStats.Request.Policy.GetValuePolicy(Value.GetId()), ECachePolicy::NonDeterministic)
@@ -592,27 +611,41 @@
 					UE_LOG(LogDerivedDataCache, Display, TEXT("%s: Cache put found non-deterministic record for %s from '%s'"),
 						*CacheStore.GetName(), *WriteToString<96>(RequestWithStats.Request.Record.GetKey()), *RequestWithStats.Request.Name);
 				}
+
+				// Compare the whole record to check whether the record from the cache needs to be returned.
+				if (ExistingRecord && !Algo::Compare(ExistingRecord.Get().GetValues(), RequestWithStats.Request.Record.GetValues()))
+				{
+					Response.Record = ExistingRecord.Get();
+				}
+
+				// Compare the metadata to check whether the metadata from the cache needs to be returned.
+				if (ExistingRecord && !ExistingRecord.Get().GetMeta().Equals(Response.Record.GetMeta()))
+				{
+					Response.Record = FCacheRecord::CreateByMove(Response.Record.GetKey(), CopyTemp(ExistingRecord.Get().GetMeta()),
+						TArray<FValueWithId>(Response.Record.GetValues()));
+				}
 			}
 		}
+
+		// A put may request data for values that were not in the incoming request.
+		// Check for that case and fetch the record from the cache if that is the case.
+		if (!EnumHasAnyFlags(RequestWithStats.Request.Policy.GetRecordPolicy(), ECachePolicy::SkipData) &&
+			Algo::AnyOf(Response.Record.GetValues(), [&RequestWithStats](const FValueWithId& Value)
+			{
+				return !Value.HasData() && !EnumHasAnyFlags(RequestWithStats.Request.Policy.GetValuePolicy(Value.GetId()), ECachePolicy::SkipData);
+			}))
+		{
+			RequestTimer.Stop();
+			ForwardAsGet(RequestWithStats);
+			return;
+		}
+
 		UE_LOG(LogDerivedDataCache, Verbose, TEXT("%s: Cache put complete for %s from '%s'"),
 			*CacheStore.GetName(), *WriteToString<96>(RequestWithStats.Request.Record.GetKey()), *RequestWithStats.Request.Name);
 
-		if (const FCbObject& Meta = RequestWithStats.Request.Record.GetMeta())
-		{
-			RequestWithStats.Stats.PhysicalWriteSize += Meta.GetSize();
-		}
-		for (const FValueWithId& Value : RequestWithStats.Request.Record.GetValues())
-		{
-			RequestWithStats.Stats.AddLogicalWrite(Value);
-			RequestWithStats.Stats.PhysicalWriteSize += Value.GetData().GetCompressedSize();
-		}
-
 		RequestTimer.Stop();
 		RequestWithStats.EndRequest(CacheStore, EStatus::Ok);
-
-		TRACE_COUNTER_ADD(ZenDDC_BytesSent, int64(RequestWithStats.Stats.PhysicalWriteSize));
-		TRACE_COUNTER_INCREMENT(ZenDDC_PutHit);
-		OnComplete(RequestWithStats.Request.MakeResponse(EStatus::Ok));
+		OnComplete(MoveTemp(Response));
 	}
 
 	void OnMiss(const TRequestWithStats<FCachePutRequest>& RequestWithStats, FCbObject&& DetailsObject)
@@ -646,6 +679,8 @@
 		OnComplete(RequestWithStats.Request.MakeResponse(EStatus::Canceled));
 	}
 
+	void ForwardAsGet(const TRequestWithStats<FCachePutRequest>& RequestWithStats);
+
 	FZenCacheStore& CacheStore;
 	IRequestOwner& Owner;
 	TArray<TRequestWithStats<FCachePutRequest>, TInlineAllocator<1>> Requests;
@@ -1040,29 +1075,43 @@
 	void OnHit(const TRequestWithStats<FCachePutValueRequest>& RequestWithStats, FCbObjectView DetailsObjectView)
 	{
 		FRequestTimer RequestTimer(RequestWithStats.Stats);
-		if (DetailsObjectView && !EnumHasAnyFlags(RequestWithStats.Request.Policy, ECachePolicy::NonDeterministic))
+
+		RequestWithStats.Stats.AddLogicalWrite(RequestWithStats.Request.Value);
+		RequestWithStats.Stats.PhysicalWriteSize += RequestWithStats.Request.Value.GetData().GetCompressedSize();
+
+		TRACE_COUNTER_INCREMENT(ZenDDC_PutHit);
+		TRACE_COUNTER_ADD(ZenDDC_BytesSent, RequestWithStats.Stats.PhysicalWriteSize);
+
+		FCachePutValueResponse Response = RequestWithStats.Request.MakeResponse(EStatus::Ok);
+
+		if (DetailsObjectView)
 		{
 			const FValue& Value = RequestWithStats.Request.Value;
 			const FIoHash RawHash = DetailsObjectView["RawHash"].AsHash();
 			const uint64 RawSize = DetailsObjectView["RawSize"].AsUInt64(MAX_uint64);
 
-			UE_LOG(LogDerivedDataCache, Display, TEXT("%s: Cache put found non-deterministic value "
+			UE_CLOG(!EnumHasAnyFlags(RequestWithStats.Request.Policy, ECachePolicy::NonDeterministic),
+				LogDerivedDataCache, Display, TEXT("%s: Cache put found non-deterministic value "
 				"with new hash %s and existing hash %s for %s from '%s'"),
 				*CacheStore.GetName(), *WriteToString<48>(Value.GetRawHash()), *WriteToString<48>(RawHash),
 				*WriteToString<96>(RequestWithStats.Request.Key), *RequestWithStats.Request.Name);
+
+			Response.Value = FValue(RawHash, RawSize);
 		}
+
+		if (!EnumHasAnyFlags(RequestWithStats.Request.Policy, ECachePolicy::SkipData) && !Response.Value.HasData())
+		{
+			RequestTimer.Stop();
+			ForwardAsGet(RequestWithStats);
+			return;
+		}
+
 		UE_LOG(LogDerivedDataCache, Verbose, TEXT("%s: Cache put complete for %s from '%s'"),
 			*CacheStore.GetName(), *WriteToString<96>(RequestWithStats.Request.Key), *RequestWithStats.Request.Name);
 
-		RequestWithStats.Stats.AddLogicalWrite(RequestWithStats.Request.Value);
-		RequestWithStats.Stats.PhysicalWriteSize += RequestWithStats.Request.Value.GetData().GetCompressedSize();
-
 		RequestTimer.Stop();
 		RequestWithStats.EndRequest(CacheStore, EStatus::Ok);
-
-		TRACE_COUNTER_INCREMENT(ZenDDC_PutHit);
-		TRACE_COUNTER_ADD(ZenDDC_BytesSent, RequestWithStats.Stats.PhysicalWriteSize);
-		OnComplete(RequestWithStats.Request.MakeResponse(EStatus::Ok));
+		OnComplete(MoveTemp(Response));
 	}
 
 	void OnMiss(const TRequestWithStats<FCachePutValueRequest>& RequestWithStats, FCbObjectView DetailsObjectView)
@@ -1096,6 +1145,8 @@
 		OnComplete(RequestWithStats.Request.MakeResponse(EStatus::Canceled));
 	}
 
+	void ForwardAsGet(const TRequestWithStats<FCachePutValueRequest>& RequestWithStats);
+
 	FZenCacheStore& CacheStore;
 	IRequestOwner& Owner;
 	TArray<TRequestWithStats<FCachePutValueRequest>, TInlineAllocator<1>> Requests;
@@ -1547,6 +1598,50 @@
 	FOnCacheGetChunkComplete OnComplete;
 };
 
+void FZenCacheStore::FPutOp::ForwardAsGet(const TRequestWithStats<FCachePutRequest>& RequestWithStats)
+{
+	const FCacheRecordPolicy Policy = RequestWithStats.Request.Policy.Transform([](ECachePolicy P) { return P | ECachePolicy::Query; });
+	const FCacheGetRequest Request[]{{RequestWithStats.Request.Name, RequestWithStats.Request.Record.GetKey(), Policy, RequestWithStats.Request.UserData}};
+	TRefCountPtr<FGetOp> GetOp = CacheStore.MakeAsyncOp<FGetOp>(CacheStore, Owner, MakeArrayView(Request),
+		[this, OpRef = TRefCountPtr<FPutOp>(this), &RequestWithStats](FCacheGetResponse&& Response)
+		{
+			// A put of a partial record will lead to Error if the returned record is partial.
+			// Convert Error to Ok if this was a put of a partial record.
+			if (Response.Status == EStatus::Error && EnumHasAnyFlags(RequestWithStats.Request.Policy.GetRecordPolicy(), ECachePolicy::PartialRecord))
+			{
+				Response.Status = EStatus::Ok;
+			}
+
+			UE_LOG(LogDerivedDataCache, Verbose, TEXT("%s: Cache put %s for %s from '%s'"), *CacheStore.GetName(),
+				Response.Status == EStatus::Ok ? TEXT("complete") :
+				Response.Status == EStatus::Canceled ? TEXT("failed with canceled request") : TEXT("failed"),
+				*WriteToString<96>(Response.Record.GetKey()), *Response.Name);
+
+			RequestWithStats.EndRequest(CacheStore, Response.Status);
+			PRAGMA_DISABLE_DEPRECATION_WARNINGS;
+			OnComplete({Response.Name, Response.Record.GetKey(), Response.Record, Response.UserData, Response.Status});
+			PRAGMA_ENABLE_DEPRECATION_WARNINGS;
+		});
+	GetOp->IssueRequests();
+}
+
+void FZenCacheStore::FPutValueOp::ForwardAsGet(const TRequestWithStats<FCachePutValueRequest>& RequestWithStats)
+{
+	const FCacheGetValueRequest Request[]{{RequestWithStats.Request.Name, RequestWithStats.Request.Key, RequestWithStats.Request.Policy | ECachePolicy::Query, RequestWithStats.Request.UserData}};
+	TRefCountPtr<FGetValueOp> GetValueOp = CacheStore.MakeAsyncOp<FGetValueOp>(CacheStore, Owner, MakeArrayView(Request),
+		[this, OpRef = TRefCountPtr<FPutValueOp>(this), &RequestWithStats](FCacheGetValueResponse&& Response)
+		{
+			UE_LOG(LogDerivedDataCache, Verbose, TEXT("%s: Cache put %s for %s from '%s'"), *CacheStore.GetName(),
+				Response.Status == EStatus::Ok ? TEXT("complete") :
+				Response.Status == EStatus::Canceled ? TEXT("failed with canceled request") : TEXT("failed"),
+				*WriteToString<96>(Response.Key), *Response.Name);
+
+			RequestWithStats.EndRequest(CacheStore, Response.Status);
+			OnComplete({Response.Name, Response.Key, Response.Value, Response.UserData, Response.Status});
+		});
+	GetValueOp->IssueRequests();
+}
+
 class FZenCacheStore::FHealthReceiver final : public IHttpReceiver
 {
 public:
diff --git a/Engine/Source/Developer/DerivedDataCache/Public/DerivedDataCache.h b/Engine/Source/Developer/DerivedDataCache/Public/DerivedDataCache.h
--- a/Engine/Source/Developer/DerivedDataCache/Public/DerivedDataCache.h
+++ b/Engine/Source/Developer/DerivedDataCache/Public/DerivedDataCache.h
@@ -235,7 +235,7 @@ struct FCachePutValueRequest
 	FValue Value;
 
 	/** Flags to control the behavior of the request. See ECachePolicy. */
-	ECachePolicy Policy = ECachePolicy::Default;
+	ECachePolicy Policy = ECachePolicy::Default | ECachePolicy::SkipData;
 
 	/** A value that will be returned in the completion callback. */
 	uint64 UserData = 0;
diff --git a/Engine/Source/Developer/DerivedDataCache/Tests/DerivedDataCacheStoreTest.cpp b/Engine/Source/Developer/DerivedDataCache/Tests/DerivedDataCacheStoreTest.cpp
--- a/Engine/Source/Developer/DerivedDataCache/Tests/DerivedDataCacheStoreTest.cpp
+++ b/Engine/Source/Developer/DerivedDataCache/Tests/DerivedDataCacheStoreTest.cpp
@@ -442,7 +442,15 @@ static void TestCacheChunks(ICache& Cache, const FCacheKey& Key, const FValueId
 	}
 }
 
-static void TestCacheValues(ICache& Cache, const uint32 ValueCount = 4 * 1024 * 1024)
+enum class ETestCacheValuesFlags : uint32
+{
+	None = 0,
+	SkipPutMissingData = 1 << 0,
+};
+
+ENUM_CLASS_FLAGS(ETestCacheValuesFlags);
+
+static void TestCacheValues(ICache& Cache, ETestCacheValuesFlags Flags = ETestCacheValuesFlags::None, const uint32 ValueCount = 4 * 1024 * 1024)
 {
 	const FValue Value(MakeSequentialIntegerBuffer(ValueCount));
 
@@ -450,6 +458,22 @@ static void TestCacheValues(ICache& Cache, const uint32 ValueCount = 4 * 1024 *
 	const FSharedString Name(TEXTVIEW("Values"));
 	const FCacheKey Key = CreateKeyFromGuid("091714ca-5136-438d-a2c8-0304cc6c54e6", ValueCount);
 
+	if (!EnumHasAnyFlags(Flags, ETestCacheValuesFlags::SkipPutMissingData))
+	{
+		INFO("Put Value w/o Data");
+		FRequestOwner BlockingOwner(EPriority::Blocking);
+		Cache.PutValue({{Name, Key, Value.RemoveData(), ECachePolicy::Default, UserData}}, BlockingOwner, [&](FCachePutValueResponse&& Response)
+		{
+			CHECK(Response.Status == EStatus::Error);
+			CHECK(Response.UserData == UserData);
+			CHECK(Response.Name == Name);
+			CHECK(Response.Key == Key);
+			CHECK(Response.Value == Value);
+			CHECK_FALSE(Response.Value.HasData());
+		});
+		BlockingOwner.Wait();
+	}
+
 	{
 		INFO("Put Value");
 		FRequestOwner BlockingOwner(EPriority::Blocking);
@@ -459,6 +483,38 @@ static void TestCacheValues(ICache& Cache, const uint32 ValueCount = 4 * 1024 *
 			CHECK(Response.UserData == UserData);
 			CHECK(Response.Name == Name);
 			CHECK(Response.Key == Key);
+			CHECK(Response.Value == Value);
+			CHECK(Response.Value.HasData());
+		});
+		BlockingOwner.Wait();
+	}
+
+	{
+		INFO("Put Value SkipData");
+		FRequestOwner BlockingOwner(EPriority::Blocking);
+		Cache.PutValue({{Name, Key, Value, ECachePolicy::Default | ECachePolicy::SkipData, UserData}}, BlockingOwner, [&](FCachePutValueResponse&& Response)
+		{
+			CHECK(Response.Status == EStatus::Ok);
+			CHECK(Response.UserData == UserData);
+			CHECK(Response.Name == Name);
+			CHECK(Response.Key == Key);
+			CHECK(Response.Value == Value);
+			CHECK_FALSE(Response.Value.HasData());
+		});
+		BlockingOwner.Wait();
+	}
+
+	{
+		INFO("Put Existing Value w/o Data");
+		FRequestOwner BlockingOwner(EPriority::Blocking);
+		Cache.PutValue({{Name, Key, Value.RemoveData(), ECachePolicy::Default, UserData}}, BlockingOwner, [&](FCachePutValueResponse&& Response)
+		{
+			CHECK(Response.Status == EStatus::Ok);
+			CHECK(Response.UserData == UserData);
+			CHECK(Response.Name == Name);
+			CHECK(Response.Key == Key);
+			CHECK(Response.Value == Value);
+			CHECK(Response.Value.HasData());
 		});
 		BlockingOwner.Wait();
 	}
@@ -474,7 +530,7 @@ static void TestCacheValues(ICache& Cache, const uint32 ValueCount = 4 * 1024 *
 			CHECK(Response.Key == Key);
 			CHECK(Response.Value.GetRawHash() == Value.GetRawHash());
 			CHECK(Response.Value.GetRawSize() == Value.GetRawSize());
-			CHECK(!Response.Value.HasData());
+			CHECK_FALSE(Response.Value.HasData());
 		});
 		BlockingOwner.Wait();
 	}
@@ -495,7 +551,7 @@ static void TestCacheValues(ICache& Cache, const uint32 ValueCount = 4 * 1024 *
 			CHECK(Response.Key == Key);
 			CHECK(Response.Value.GetRawHash() == Value.GetRawHash());
 			CHECK(Response.Value.GetRawSize() == Value.GetRawSize());
-			CHECK(!Response.Value.HasData());
+			CHECK_FALSE(Response.Value.HasData());
 		});
 		BlockingOwner.Wait();
 		CHECK(UserDataAccumulator == 1);
@@ -535,7 +591,15 @@ static void TestCacheValues(ICache& Cache, const uint32 ValueCount = 4 * 1024 *
 	TestCacheChunks(Cache, Key, {}, Value);
 }
 
-static void TestCacheRecords(ICache& Cache)
+enum class ETestCacheRecordsFlags : uint32
+{
+	None = 0,
+	SkipPutMissingData = 1 << 0,
+};
+
+ENUM_CLASS_FLAGS(ETestCacheRecordsFlags);
+
+static void TestCacheRecords(ICache& Cache, ETestCacheRecordsFlags Flags = ETestCacheRecordsFlags::None)
 {
 	const FValueId ValueId = FValueId::FromName(ANSITEXTVIEW("Value"));
 	const FValueId Attachment1Id = FValueId::FromName(ANSITEXTVIEW("Attachment1"));
@@ -555,8 +619,8 @@ static void TestCacheRecords(ICache& Cache)
 	const FSharedString Name(TEXTVIEW("Records"));
 	const FCacheKey Key = CreateKeyFromGuid("ff173633-71e9-4960-a683-470c77abaa6a");
 
+	const FCacheRecord Record = [&]
 	{
-		INFO("Put Record");
 		FCacheRecordBuilder RecordBuilder(Key);
 		RecordBuilder.AddValue(Value);
 		RecordBuilder.AddValue(Attachment1);
@@ -564,8 +628,37 @@ static void TestCacheRecords(ICache& Cache)
 		RecordBuilder.AddValue(Attachment3);
 		RecordBuilder.AddValue(Attachment4);
 		RecordBuilder.SetMeta(CopyTemp(Meta));
+		return RecordBuilder.Build();
+	}();
+
+	if (!EnumHasAnyFlags(Flags, ETestCacheRecordsFlags::SkipPutMissingData))
+	{
+		INFO("Put Record w/o Meta+Data");
+		FRequestOwner BlockingOwner(EPriority::Blocking);
+		Cache.Put({{Name, ApplySkipPolicy(Record, ECachePolicy::SkipMeta | ECachePolicy::SkipData), ECachePolicy::Default, UserData}}, BlockingOwner, [&](FCachePutResponse&& Response)
+		{
+			CHECK(Response.Status == EStatus::Error);
+			CHECK(Response.UserData == UserData);
+			CHECK(Response.Name == Name);
+			PRAGMA_DISABLE_DEPRECATION_WARNINGS;
+			CHECK(Response.Key == Key);
+			PRAGMA_ENABLE_DEPRECATION_WARNINGS;
+			CHECK(Response.Record.GetKey() == Key);
+			CHECK_FALSE(Response.Record.GetMeta());
+			CHECK(Response.Record.GetValues().Num() == 5);
+			CHECK(HasValueWithNoData(Response.Record, Value));
+			CHECK(HasValueWithNoData(Response.Record, Attachment1));
+			CHECK(HasValueWithNoData(Response.Record, Attachment2));
+			CHECK(HasValueWithNoData(Response.Record, Attachment3));
+			CHECK(HasValueWithNoData(Response.Record, Attachment4));
+		});
+		BlockingOwner.Wait();
+	}
+
+	{
+		INFO("Put Record");
 		FRequestOwner BlockingOwner(EPriority::Blocking);
-		Cache.Put({{Name, RecordBuilder.Build(), ECachePolicy::Default, UserData}}, BlockingOwner, [&](FCachePutResponse&& Response)
+		Cache.Put({{Name, Record, ECachePolicy::Default, UserData}}, BlockingOwner, [&](FCachePutResponse&& Response)
 		{
 			CHECK(Response.Status == EStatus::Ok);
 			CHECK(Response.UserData == UserData);
@@ -573,6 +666,83 @@ static void TestCacheRecords(ICache& Cache)
 			PRAGMA_DISABLE_DEPRECATION_WARNINGS;
 			CHECK(Response.Key == Key);
 			PRAGMA_ENABLE_DEPRECATION_WARNINGS;
+			CHECK(Response.Record.GetKey() == Key);
+			CHECK(Response.Record.GetMeta().Equals(Meta));
+			CHECK(Response.Record.GetValues().Num() == 5);
+			CHECK(HasValueWithData(Response.Record, Value));
+			CHECK(HasValueWithData(Response.Record, Attachment1));
+			CHECK(HasValueWithData(Response.Record, Attachment2));
+			CHECK(HasValueWithData(Response.Record, Attachment3));
+			CHECK(HasValueWithData(Response.Record, Attachment4));
+		});
+		BlockingOwner.Wait();
+	}
+
+	{
+		INFO("Put Record SkipMeta");
+		FRequestOwner BlockingOwner(EPriority::Blocking);
+		Cache.Put({{Name, Record, ECachePolicy::Default | ECachePolicy::SkipMeta, UserData}}, BlockingOwner, [&](FCachePutResponse&& Response)
+		{
+			CHECK(Response.Status == EStatus::Ok);
+			CHECK(Response.UserData == UserData);
+			CHECK(Response.Name == Name);
+			PRAGMA_DISABLE_DEPRECATION_WARNINGS;
+				CHECK(Response.Key == Key);
+			PRAGMA_ENABLE_DEPRECATION_WARNINGS;
+				CHECK(Response.Record.GetKey() == Key);
+			CHECK_FALSE(Response.Record.GetMeta());
+			CHECK(Response.Record.GetValues().Num() == 5);
+			CHECK(HasValueWithData(Response.Record, Value));
+			CHECK(HasValueWithData(Response.Record, Attachment1));
+			CHECK(HasValueWithData(Response.Record, Attachment2));
+			CHECK(HasValueWithData(Response.Record, Attachment3));
+			CHECK(HasValueWithData(Response.Record, Attachment4));
+		});
+		BlockingOwner.Wait();
+	}
+
+	{
+		INFO("Put Record SkipData");
+		FRequestOwner BlockingOwner(EPriority::Blocking);
+		Cache.Put({{Name, Record, ECachePolicy::Default | ECachePolicy::SkipData, UserData}}, BlockingOwner, [&](FCachePutResponse&& Response)
+		{
+			CHECK(Response.Status == EStatus::Ok);
+			CHECK(Response.UserData == UserData);
+			CHECK(Response.Name == Name);
+			PRAGMA_DISABLE_DEPRECATION_WARNINGS;
+			CHECK(Response.Key == Key);
+			PRAGMA_ENABLE_DEPRECATION_WARNINGS;
+			CHECK(Response.Record.GetKey() == Key);
+			CHECK(Response.Record.GetMeta().Equals(Meta));
+			CHECK(Response.Record.GetValues().Num() == 5);
+			CHECK(HasValueWithNoData(Response.Record, Value));
+			CHECK(HasValueWithNoData(Response.Record, Attachment1));
+			CHECK(HasValueWithNoData(Response.Record, Attachment2));
+			CHECK(HasValueWithNoData(Response.Record, Attachment3));
+			CHECK(HasValueWithNoData(Response.Record, Attachment4));
+		});
+		BlockingOwner.Wait();
+	}
+
+	{
+		INFO("Put Existing Record w/o Meta+Data");
+		FRequestOwner BlockingOwner(EPriority::Blocking);
+		Cache.Put({{Name, ApplySkipPolicy(Record, ECachePolicy::SkipMeta | ECachePolicy::SkipData), ECachePolicy::Default, UserData}}, BlockingOwner, [&](FCachePutResponse&& Response)
+		{
+			CHECK(Response.Status == EStatus::Ok);
+			CHECK(Response.UserData == UserData);
+			CHECK(Response.Name == Name);
+			PRAGMA_DISABLE_DEPRECATION_WARNINGS;
+			CHECK(Response.Key == Key);
+			PRAGMA_ENABLE_DEPRECATION_WARNINGS;
+			CHECK(Response.Record.GetKey() == Key);
+			CHECK(Response.Record.GetMeta().Equals(Meta));
+			CHECK(Response.Record.GetValues().Num() == 5);
+			CHECK(HasValueWithData(Response.Record, Value));
+			CHECK(HasValueWithData(Response.Record, Attachment1));
+			CHECK(HasValueWithData(Response.Record, Attachment2));
+			CHECK(HasValueWithData(Response.Record, Attachment3));
+			CHECK(HasValueWithData(Response.Record, Attachment4));
 		});
 		BlockingOwner.Wait();
 	}
@@ -687,7 +857,15 @@ static void TestCacheRecords(ICache& Cache)
 	TestCacheChunks(Cache, Key, ValueId, Value);
 }
 
-static void TestPartialCacheRecords(ICache& Cache)
+enum class ETestPartialCacheRecordsFlags : uint32
+{
+	None = 0,
+	SkipMetaDiff = 1 << 0,
+};
+
+ENUM_CLASS_FLAGS(ETestPartialCacheRecordsFlags);
+
+static void TestPartialCacheRecords(ICache& Cache, ETestPartialCacheRecordsFlags Flags = ETestPartialCacheRecordsFlags::None)
 {
 	const FValueId ValueId = FValueId::FromName(ANSITEXTVIEW("Value"));
 	const FValueId Attachment1Id = FValueId::FromName(ANSITEXTVIEW("Attachment1"));
@@ -708,8 +886,37 @@ static void TestPartialCacheRecords(ICache& Cache)
 	const FCacheKey Key = CreateKeyFromGuid("5d077783-8ecc-400a-bb65-168345d4f7a7");
 
 	{
-		INFO("Put Record w/ Value+Attachment1+Attachment2");
+		INFO("Put Record w/ No Values");
+		FCacheRecordBuilder RecordBuilder(Key);
+		RecordBuilder.AddValue(Value.RemoveData());
+		RecordBuilder.AddValue(Attachment1.RemoveData());
+		RecordBuilder.AddValue(Attachment2.RemoveData());
+		RecordBuilder.AddValue(Attachment3.RemoveData());
+		RecordBuilder.AddValue(Attachment4.RemoveData());
+		RecordBuilder.SetMeta(CopyTemp(Meta));
+		FRequestOwner BlockingOwner(EPriority::Blocking);
+		Cache.Put({{Name, RecordBuilder.Build(), ECachePolicy::Default | ECachePolicy::PartialRecord, UserData}}, BlockingOwner, [&](FCachePutResponse&& Response)
+		{
+			CHECK(Response.Status == EStatus::Ok);
+			CHECK(Response.UserData == UserData);
+			CHECK(Response.Name == Name);
+			PRAGMA_DISABLE_DEPRECATION_WARNINGS;
+			CHECK(Response.Key == Key);
+			PRAGMA_ENABLE_DEPRECATION_WARNINGS;
+			CHECK(Response.Record.GetKey() == Key);
+			CHECK(Response.Record.GetMeta().Equals(Meta));
+			CHECK(Response.Record.GetValues().Num() == 5);
+			CHECK(HasValueWithNoData(Response.Record, Value));
+			CHECK(HasValueWithNoData(Response.Record, Attachment1));
+			CHECK(HasValueWithNoData(Response.Record, Attachment2));
+			CHECK(HasValueWithNoData(Response.Record, Attachment3));
+			CHECK(HasValueWithNoData(Response.Record, Attachment4));
+		});
+		BlockingOwner.Wait();
+	}
 
+	{
+		INFO("Put Record w/ Value+Attachment1+Attachment2");
 		FCacheRecordBuilder RecordBuilder(Key);
 		RecordBuilder.AddValue(Value);
 		RecordBuilder.AddValue(Attachment1);
@@ -726,21 +933,28 @@ static void TestPartialCacheRecords(ICache& Cache)
 			PRAGMA_DISABLE_DEPRECATION_WARNINGS;
 			CHECK(Response.Key == Key);
 			PRAGMA_ENABLE_DEPRECATION_WARNINGS;
+			CHECK(Response.Record.GetKey() == Key);
+			CHECK(Response.Record.GetMeta().Equals(Meta));
+			CHECK(Response.Record.GetValues().Num() == 5);
+			CHECK(HasValueWithData(Response.Record, Value));
+			CHECK(HasValueWithData(Response.Record, Attachment1));
+			CHECK(HasValueWithData(Response.Record, Attachment2));
+			CHECK(HasValueWithNoData(Response.Record, Attachment3));
+			CHECK(HasValueWithNoData(Response.Record, Attachment4));
 		});
 		BlockingOwner.Wait();
 	}
 
 	{
 		INFO("Put Record w/ Attachment2+Attachment3");
+		const FCbObject NewMeta = EnumHasAnyFlags(Flags, ETestPartialCacheRecordsFlags::SkipMetaDiff) ? Meta : CreateMeta("NewValue");
 		FCacheRecordBuilder RecordBuilder(Key);
 		RecordBuilder.AddValue(Value.RemoveData());
 		RecordBuilder.AddValue(Attachment1.RemoveData());
 		RecordBuilder.AddValue(Attachment2);
 		RecordBuilder.AddValue(Attachment3);
 		RecordBuilder.AddValue(Attachment4.RemoveData());
-		// TODO: Allow Meta to differ or be unset here without blocking the subsequent partial put.
-		//RecordBuilder.SetMeta(CreateMeta("NewValue"));
-		RecordBuilder.SetMeta(CreateMeta("Value"));
+		RecordBuilder.SetMeta(CopyTemp(NewMeta));
 		FRequestOwner BlockingOwner(EPriority::Blocking);
 		Cache.Put({{Name, RecordBuilder.Build(), ECachePolicy::Default | ECachePolicy::PartialRecord, UserData}}, BlockingOwner, [&](FCachePutResponse&& Response)
 		{
@@ -750,6 +964,13 @@ static void TestPartialCacheRecords(ICache& Cache)
 			PRAGMA_DISABLE_DEPRECATION_WARNINGS;
 			CHECK(Response.Key == Key);
 			PRAGMA_ENABLE_DEPRECATION_WARNINGS;
+			CHECK(Response.Record.GetKey() == Key);
+			CHECK(Response.Record.GetMeta().Equals(Meta));
+			CHECK(HasValueWithData(Response.Record, Value));
+			CHECK(HasValueWithData(Response.Record, Attachment1));
+			CHECK(HasValueWithData(Response.Record, Attachment2));
+			CHECK(HasValueWithData(Response.Record, Attachment3));
+			CHECK(HasValueWithNoData(Response.Record, Attachment4));
 		});
 		BlockingOwner.Wait();
 	}
@@ -889,7 +1110,15 @@ static void TestPartialCacheRecords(ICache& Cache)
 	TestCacheChunks(Cache, Key, ValueId, Value);
 }
 
-static void TestCacheDeterminism(ICache& Cache)
+enum class ETestCacheDeterminismFlags : uint32
+{
+	None = 0,
+	SkipOverwrite = 1 << 0,
+};
+
+ENUM_CLASS_FLAGS(ETestCacheDeterminismFlags);
+
+static void TestCacheDeterminism(ICache& Cache, ETestCacheDeterminismFlags Flags = ETestCacheDeterminismFlags::None)
 {
 	const FValue Value(MakeSequentialIntegerBuffer(8 * 1024 * 1024));
 	const FValue OtherValue(MakeSequentialIntegerBuffer(8 * 1024 * 1024 + 1024));
@@ -898,6 +1127,7 @@ static void TestCacheDeterminism(ICache& Cache)
 	const FValueId AttachmentId = FValueId::FromName(ANSITEXTVIEW("Attachment"));
 
 	const FCbObject Meta = CreateMeta("Value");
+	const FCbObject OtherMeta = CreateMeta("OtherValue");
 
 	const FCacheKey RecordKey = CreateKeyFromGuid("8660fd50-a31c-4ea1-ace8-eb25b21b6db6");
 
@@ -907,13 +1137,16 @@ static void TestCacheDeterminism(ICache& Cache)
 		FCacheRecordBuilder RecordBuilder(RecordKey);
 		RecordBuilder.AddValue(ValueId, Value);
 		RecordBuilder.AddValue(AttachmentId, Value);
-		FCacheRecordPolicyBuilder PolicyBuilder(ECachePolicy::Store);
-		PolicyBuilder.AddValuePolicy(ValueId, ECachePolicy::Store | ECachePolicy::NonDeterministic);
+		FCacheRecordPolicyBuilder PolicyBuilder(ECachePolicy::Store | ECachePolicy::SkipData);
+		PolicyBuilder.AddValuePolicy(ValueId, ECachePolicy::Store | ECachePolicy::SkipData | ECachePolicy::NonDeterministic);
 		RecordBuilder.SetMeta(CopyTemp(Meta));
 		FRequestOwner BlockingOwner(EPriority::Blocking);
-		Cache.Put({{{TEXTVIEW("OriginalRecord")}, RecordBuilder.Build(), PolicyBuilder.Build()}}, BlockingOwner, [](FCachePutResponse&& Response)
+		Cache.Put({{{TEXTVIEW("OriginalRecord")}, RecordBuilder.Build(), PolicyBuilder.Build()}}, BlockingOwner, [&](FCachePutResponse&& Response)
 		{
 			CHECK(Response.Status == EStatus::Ok);
+			CHECK(Response.Record.GetMeta().Equals(Meta));
+			CHECK(HasValueWithNoData(Response.Record, FValueWithId(ValueId, Value)));
+			CHECK(HasValueWithNoData(Response.Record, FValueWithId(AttachmentId, Value)));
 		});
 		BlockingOwner.Wait();
 	};
@@ -924,11 +1157,14 @@ static void TestCacheDeterminism(ICache& Cache)
 		FCacheRecordBuilder RecordBuilder(RecordKey);
 		RecordBuilder.AddValue(ValueId, Value);
 		RecordBuilder.AddValue(AttachmentId, Value);
-		RecordBuilder.SetMeta(CreateMeta("ValueCopy"));
+		RecordBuilder.SetMeta(CopyTemp(OtherMeta));
 		FRequestOwner BlockingOwner(EPriority::Blocking);
-		Cache.Put({{{TEXTVIEW("OriginalRecordCopy")}, RecordBuilder.Build()}}, BlockingOwner, [](FCachePutResponse&& Response)
+		Cache.Put({{{TEXTVIEW("OriginalRecordCopy")}, RecordBuilder.Build()}}, BlockingOwner, [&](FCachePutResponse&& Response)
 		{
 			CHECK(Response.Status == EStatus::Ok);
+			CHECK(Response.Record.GetMeta().Equals(Meta));
+			CHECK(HasValueWithNoData(Response.Record, FValueWithId(ValueId, Value)));
+			CHECK(HasValueWithNoData(Response.Record, FValueWithId(AttachmentId, Value)));
 		});
 		BlockingOwner.Wait();
 	}
@@ -939,9 +1175,12 @@ static void TestCacheDeterminism(ICache& Cache)
 		RecordBuilder.AddValue(ValueId, OtherValue);
 		RecordBuilder.AddValue(AttachmentId, Value);
 		FRequestOwner BlockingOwner(EPriority::Blocking);
-		Cache.Put({{{TEXTVIEW("DifferentRecord")}, RecordBuilder.Build()}}, BlockingOwner, [](FCachePutResponse&& Response)
+		Cache.Put({{{TEXTVIEW("DifferentRecord")}, RecordBuilder.Build()}}, BlockingOwner, [&](FCachePutResponse&& Response)
 		{
 			CHECK(Response.Status == EStatus::Ok);
+			CHECK(Response.Record.GetMeta().Equals(Meta));
+			CHECK(HasValueWithNoData(Response.Record, FValueWithId(ValueId, Value)));
+			CHECK(HasValueWithNoData(Response.Record, FValueWithId(AttachmentId, Value)));
 		});
 		BlockingOwner.Wait();
 	}
@@ -954,9 +1193,12 @@ static void TestCacheDeterminism(ICache& Cache)
 		FCacheRecordPolicyBuilder PolicyBuilder;
 		PolicyBuilder.AddValuePolicy(ValueId, ECachePolicy::Default | ECachePolicy::NonDeterministic);
 		FRequestOwner BlockingOwner(EPriority::Blocking);
-		Cache.Put({{{TEXTVIEW("DifferentRecordIgnore")}, RecordBuilder.Build(), PolicyBuilder.Build()}}, BlockingOwner, [](FCachePutResponse&& Response)
+		Cache.Put({{{TEXTVIEW("DifferentRecordIgnore")}, RecordBuilder.Build(), PolicyBuilder.Build()}}, BlockingOwner, [&](FCachePutResponse&& Response)
 		{
 			CHECK(Response.Status == EStatus::Ok);
+			CHECK(Response.Record.GetMeta().Equals(Meta));
+			CHECK(HasValueWithData(Response.Record, FValueWithId(ValueId, Value)));
+			CHECK(HasValueWithData(Response.Record, FValueWithId(AttachmentId, Value)));
 		});
 		BlockingOwner.Wait();
 	}
@@ -974,26 +1216,32 @@ static void TestCacheDeterminism(ICache& Cache)
 	}
 
 	// Put Different Record w/ Overwrite
+	if (!EnumHasAnyFlags(Flags, ETestCacheDeterminismFlags::SkipOverwrite))
 	{
 		FCacheRecordBuilder RecordBuilder(RecordKey);
 		RecordBuilder.AddValue(ValueId, OtherValue);
 		RecordBuilder.AddValue(AttachmentId, Value);
+		RecordBuilder.SetMeta(CopyTemp(OtherMeta));
 		FRequestOwner BlockingOwner(EPriority::Blocking);
-		Cache.Put({{{TEXTVIEW("DifferentRecordOverwrite")}, RecordBuilder.Build(), ECachePolicy::Store}}, BlockingOwner, [](FCachePutResponse&& Response)
+		Cache.Put({{{TEXTVIEW("DifferentRecordOverwrite")}, RecordBuilder.Build(), ECachePolicy::Store}}, BlockingOwner, [&](FCachePutResponse&& Response)
 		{
 			CHECK(Response.Status == EStatus::Ok);
+			CHECK(Response.Record.GetMeta().Equals(OtherMeta));
+			CHECK(HasValueWithData(Response.Record, FValueWithId(ValueId, OtherValue)));
+			CHECK(HasValueWithData(Response.Record, FValueWithId(AttachmentId, Value)));
 		});
 		BlockingOwner.Wait();
 	}
 
 	// Get Different Record
+	if (!EnumHasAnyFlags(Flags, ETestCacheDeterminismFlags::SkipOverwrite))
 	{
 		FRequestOwner BlockingOwner(EPriority::Blocking);
 		Cache.Get({{{TEXTVIEW("DifferentRecord")}, RecordKey}}, BlockingOwner, [&](FCacheGetResponse&& Response)
 		{
 			CHECK(Response.Status == EStatus::Ok);
 			CHECK(HasValueWithData(Response.Record, FValueWithId(ValueId, OtherValue)));
-			CHECK_FALSE(Response.Record.GetMeta());
+			CHECK(Response.Record.GetMeta().Equals(OtherMeta));
 		});
 		BlockingOwner.Wait();
 	}
@@ -1007,9 +1255,11 @@ static void TestCacheDeterminism(ICache& Cache)
 	const auto PutOriginalValueWithOverwrite = [&]
 	{
 		FRequestOwner BlockingOwner(EPriority::Blocking);
-		Cache.PutValue({{{TEXTVIEW("OriginalValue")}, ValueKey, Value, ECachePolicy::Store | ECachePolicy::NonDeterministic}}, BlockingOwner, [](FCachePutValueResponse&& Response)
+		Cache.PutValue({{{TEXTVIEW("OriginalValue")}, ValueKey, Value, ECachePolicy::Store | ECachePolicy::SkipData | ECachePolicy::NonDeterministic}}, BlockingOwner, [&](FCachePutValueResponse&& Response)
 		{
 			CHECK(Response.Status == EStatus::Ok);
+			CHECK(Response.Value == Value);
+			CHECK_FALSE(Response.Value.HasData());
 		});
 		BlockingOwner.Wait();
 	};
@@ -1018,9 +1268,11 @@ static void TestCacheDeterminism(ICache& Cache)
 	// Put Original Value
 	{
 		FRequestOwner BlockingOwner(EPriority::Blocking);
-		Cache.PutValue({{{TEXTVIEW("OriginalValueCopy")}, ValueKey, Value}}, BlockingOwner, [](FCachePutValueResponse&& Response)
+		Cache.PutValue({{{TEXTVIEW("OriginalValueCopy")}, ValueKey, Value}}, BlockingOwner, [&](FCachePutValueResponse&& Response)
 		{
 			CHECK(Response.Status == EStatus::Ok);
+			CHECK(Response.Value == Value);
+			CHECK_FALSE(Response.Value.HasData());
 		});
 		BlockingOwner.Wait();
 	}
@@ -1028,9 +1280,11 @@ static void TestCacheDeterminism(ICache& Cache)
 	// Put Different Value
 	{
 		FRequestOwner BlockingOwner(EPriority::Blocking);
-		Cache.PutValue({{{TEXTVIEW("DifferentValue")}, ValueKey, OtherValue}}, BlockingOwner, [](FCachePutValueResponse&& Response)
+		Cache.PutValue({{{TEXTVIEW("DifferentValue")}, ValueKey, OtherValue}}, BlockingOwner, [&](FCachePutValueResponse&& Response)
 		{
 			CHECK(Response.Status == EStatus::Ok);
+			CHECK(Response.Value == Value);
+			CHECK_FALSE(Response.Value.HasData());
 		});
 		BlockingOwner.Wait();
 	}
@@ -1038,9 +1292,11 @@ static void TestCacheDeterminism(ICache& Cache)
 	// Put Different Value w/ NonDeterministic Flag
 	{
 		FRequestOwner BlockingOwner(EPriority::Blocking);
-		Cache.PutValue({{{TEXTVIEW("DifferentValueIgnore")}, ValueKey, OtherValue, ECachePolicy::Default | ECachePolicy::NonDeterministic}}, BlockingOwner, [](FCachePutValueResponse&& Response)
+		Cache.PutValue({{{TEXTVIEW("DifferentValueIgnore")}, ValueKey, OtherValue, ECachePolicy::Default | ECachePolicy::NonDeterministic}}, BlockingOwner, [&](FCachePutValueResponse&& Response)
 		{
 			CHECK(Response.Status == EStatus::Ok);
+			CHECK(Response.Value == Value);
+			CHECK(Response.Value.HasData());
 		});
 		BlockingOwner.Wait();
 	}
@@ -1057,16 +1313,20 @@ static void TestCacheDeterminism(ICache& Cache)
 	}
 
 	// Put Different Value w/ Overwrite
+	if (!EnumHasAnyFlags(Flags, ETestCacheDeterminismFlags::SkipOverwrite))
 	{
 		FRequestOwner BlockingOwner(EPriority::Blocking);
-		Cache.PutValue({{{TEXTVIEW("DifferentValueOverwrite")}, ValueKey, OtherValue, ECachePolicy::Store}}, BlockingOwner, [](FCachePutValueResponse&& Response)
+		Cache.PutValue({{{TEXTVIEW("DifferentValueOverwrite")}, ValueKey, OtherValue, ECachePolicy::Store}}, BlockingOwner, [&](FCachePutValueResponse&& Response)
 		{
 			CHECK(Response.Status == EStatus::Ok);
+			CHECK(Response.Value == OtherValue);
+			CHECK(Response.Value.HasData());
 		});
 		BlockingOwner.Wait();
 	}
 
 	// Get Different Value
+	if (!EnumHasAnyFlags(Flags, ETestCacheDeterminismFlags::SkipOverwrite))
 	{
 		FRequestOwner BlockingOwner(EPriority::Blocking);
 		Cache.GetValue({{{TEXTVIEW("DifferentValue")}, ValueKey}}, BlockingOwner, [&OtherValue](FCacheGetValueResponse&& Response)
@@ -1088,16 +1348,18 @@ namespace UE::DerivedData
 
 TEST_CASE("DerivedData::CacheStore", "[DerivedData]")
 {
+	using namespace CacheStoreTest;
+
 	SECTION("Memory")
 	{
 		FStringView Config = TEXTVIEW("(Memory=(Type=Memory))");
 		TUniquePtr<ICache> Cache(CreateCache(Config));
 		CHECKED_IF(Cache)
 		{
-			CacheStoreTest::TestCacheValues(*Cache);
-			CacheStoreTest::TestCacheRecords(*Cache);
-			CacheStoreTest::TestPartialCacheRecords(*Cache);
-			CacheStoreTest::TestCacheDeterminism(*Cache);
+			TestCacheValues(*Cache);
+			TestCacheRecords(*Cache);
+			TestPartialCacheRecords(*Cache);
+			TestCacheDeterminism(*Cache);
 		}
 	}
 
@@ -1114,13 +1376,13 @@ TEST_CASE("DerivedData::CacheStore", "[DerivedData]")
 		CHECKED_IF(Cache)
 		{
 			// Test different value sizes because value storage is tied to size.
-			CacheStoreTest::TestCacheValues(*Cache, 8);
-			CacheStoreTest::TestCacheValues(*Cache, 8 * 1024);
-			CacheStoreTest::TestCacheValues(*Cache, 8 * 1024 * 1024);
+			TestCacheValues(*Cache, ETestCacheValuesFlags::None, 8);
+			TestCacheValues(*Cache, ETestCacheValuesFlags::None, 8 * 1024);
+			TestCacheValues(*Cache, ETestCacheValuesFlags::None, 8 * 1024 * 1024);
 
-			CacheStoreTest::TestCacheRecords(*Cache);
-			CacheStoreTest::TestPartialCacheRecords(*Cache);
-			CacheStoreTest::TestCacheDeterminism(*Cache);
+			TestCacheRecords(*Cache);
+			TestPartialCacheRecords(*Cache);
+			TestCacheDeterminism(*Cache);
 		}
 
 		IFileManager::Get().DeleteDirectory(*CachePath, /*bRequireExists*/ false, /*bTree*/ true);
@@ -1132,10 +1394,11 @@ TEST_CASE("DerivedData::CacheStore", "[DerivedData]")
 		TUniquePtr<ICache> Cache(CreateCache(Config));
 		CHECKED_IF(Cache)
 		{
-			CacheStoreTest::TestCacheValues(*Cache);
-			CacheStoreTest::TestCacheRecords(*Cache);
-			CacheStoreTest::TestPartialCacheRecords(*Cache);
-			CacheStoreTest::TestCacheDeterminism(*Cache);
+			TestCacheValues(*Cache);
+			// TODO: Zen needs to be fixed to work with no flags passed to these functions.
+			TestCacheRecords(*Cache, ETestCacheRecordsFlags::SkipPutMissingData);
+			TestPartialCacheRecords(*Cache, ETestPartialCacheRecordsFlags::SkipMetaDiff);
+			TestCacheDeterminism(*Cache);
 		}
 	}
 
@@ -1144,8 +1407,9 @@ TEST_CASE("DerivedData::CacheStore", "[DerivedData]")
 		FStringView Config = TEXTVIEW("(Cloud)");
 		if (TUniquePtr<ICache> Cache{CreateCache(Config)})
 		{
-			CacheStoreTest::TestCacheValues(*Cache);
-			CacheStoreTest::TestCacheRecords(*Cache);
+			TestCacheValues(*Cache, ETestCacheValuesFlags::SkipPutMissingData);
+			TestCacheRecords(*Cache, ETestCacheRecordsFlags::SkipPutMissingData);
+			TestCacheDeterminism(*Cache, ETestCacheDeterminismFlags::SkipOverwrite);
 		}
 		else
 		{
