DDB: Fixed up required memory functions to be consistent

#rnx
#Synced-CL 50967057

diff --git a/Engine/Source/Developer/DerivedDataCache/Private/DerivedDataBuildJobContext.h b/Engine/Source/Developer/DerivedDataCache/Private/DerivedDataBuildJobContext.h
--- a/Engine/Source/Developer/DerivedDataCache/Private/DerivedDataBuildJobContext.h
+++ b/Engine/Source/Developer/DerivedDataCache/Private/DerivedDataBuildJobContext.h
@@ -43,7 +43,7 @@ public:
 
 	inline ECachePolicy GetCachePolicyMask() const final { return CachePolicyMask; }
 	inline EBuildPolicy GetBuildPolicyMask() const final { return BuildPolicyMask; }
-	uint64 GetRequiredMemory() const override { return RequiredMemory; }
+	inline uint64 GetRequiredMemory() const final { return RequiredMemory; }
 	inline bool ShouldCheckDeterministicOutput() const { return bDeterministicOutputCheck; }
 
 	void AddConstant(FUtf8StringView Key, FCbObject&& Value);
diff --git a/Engine/Source/Developer/DerivedDataCache/Public/DerivedDataBuildFunction.h b/Engine/Source/Developer/DerivedDataCache/Public/DerivedDataBuildFunction.h
--- a/Engine/Source/Developer/DerivedDataCache/Public/DerivedDataBuildFunction.h
+++ b/Engine/Source/Developer/DerivedDataCache/Public/DerivedDataBuildFunction.h
@@ -85,9 +85,6 @@ public:
 
 	/** Returns the name by which to identify this build for logging and profiling. */
 	virtual const FSharedString& GetName() const = 0;
-	
-	/** Returns estimate of memory required for this operation. */
-	virtual uint64 GetRequiredMemory() const = 0;
 
 	/** Returns the constant with the matching key, or an object with no fields if not found. */
 	virtual FCbObject FindConstant(FUtf8StringView Key) const = 0;
@@ -121,6 +118,9 @@ public:
 	/** Overrides the cache policy mask used when writing this build in the cache. */
 	virtual void SetCachePolicyMask(ECachePolicy Policy) = 0;
 
+	/** Returns the estimated peak memory required to execute the build, excluding input size. Defaults to 0. */
+	virtual uint64 GetRequiredMemory() const = 0;
+
 	/**
 	 * Make this an asynchronous build by making the caller responsible for completing the build.
 	 *
