DDC: Fixed reauthentication during shutdown by ensuring that DesktopPlatform unloads after DDC

#Synced-CL 50967057

diff --git a/Engine/Source/Developer/DerivedDataCache/Private/DerivedDataCacheModule.cpp b/Engine/Source/Developer/DerivedDataCache/Private/DerivedDataCacheModule.cpp
--- a/Engine/Source/Developer/DerivedDataCache/Private/DerivedDataCacheModule.cpp
+++ b/Engine/Source/Developer/DerivedDataCache/Private/DerivedDataCacheModule.cpp
@@ -87,9 +87,16 @@ public:
 
 	void StartupModule() final
 	{
+		FModuleManager& ModuleManager = FModuleManager::Get();
+
 		// Required to guarantee that SSL shuts down after DDC. Without this, waiting for active
 		// cache requests on shutdown can crash when accessing SSL.
-		FModuleManager::Get().LoadModuleChecked(TEXT("SSL"));
+		ModuleManager.LoadModuleChecked("SSL");
+
+		// Required to guarantee that DesktopPlatform shuts down after DDC. Without this,
+		// waiting for active cache requests on shutdown can fail to reauthenticate.
+		// DesktopPlatform is an *optional* dependency and this load may fail.
+		ModuleManager.LoadModule("DesktopPlatform");
 	}
 
 	void ShutdownModule() final
