Separating out request pools in the Jupiter client to allow for different transfer rate limits for small operations (eg: refs/exists) than for large operations (blobs).

#Synced-CL 51289657

diff --git a/Engine/Source/Developer/DerivedDataCache/Private/HttpCacheStore.cpp b/Engine/Source/Developer/DerivedDataCache/Private/HttpCacheStore.cpp
--- a/Engine/Source/Developer/DerivedDataCache/Private/HttpCacheStore.cpp
+++ b/Engine/Source/Developer/DerivedDataCache/Private/HttpCacheStore.cpp
@@ -71,7 +71,9 @@
 #include "Ssl.h"
 #endif
 
-#define UE_HTTPDDC_GET_REQUEST_POOL_SIZE 128
+#define UE_HTTPDDC_GETREF_REQUEST_POOL_SIZE 64
+#define UE_HTTPDDC_GETBLOBS_REQUEST_POOL_SIZE 64
+#define UE_HTTPDDC_EXISTS_REQUEST_POOL_SIZE 64
 #define UE_HTTPDDC_PUTREF_REQUEST_POOL_SIZE 64
 #define UE_HTTPDDC_PUTBLOBS_REQUEST_POOL_SIZE 64
 #define UE_HTTPDDC_PUTFINALIZE_REQUEST_POOL_SIZE 64
@@ -313,7 +315,9 @@ private:
 	FDerivedDataCacheUsageStats UsageStats;
 	FBackendDebugOptions DebugOptions;
 	THttpUniquePtr<IHttpConnectionPool> ConnectionPool;
-	FHttpRequestQueue GetRequestQueue;
+	FHttpRequestQueue GetRefRequestQueue;
+	FHttpRequestQueue GetBlobsRequestQueue;
+	FHttpRequestQueue ExistsRequestQueue;
 	FHttpRequestQueue PutRefRequestQueue;
 	FHttpRequestQueue PutBlobsRequestQueue;
 	FHttpRequestQueue PutFinalizeRequestQueue;
@@ -340,7 +344,9 @@ private:
 	
 	enum class EOperationCategory
 	{
-		Get,
+		GetRef,
+		GetBlobs,
+		Exists,
 		PutRef,
 		PutBlobs,
 		PutFinalize
@@ -1231,7 +1237,7 @@ void FHttpCacheStore::FGetRecordOp::GetRecordOnly(const FCacheKey& InKey, const
 	RequestStats.Bucket = Key.Bucket;
 
 	RequestTimer.Stop();
-	CacheStore.WaitForHttpOperationAsync(Owner, EOperationCategory::Get, [Self = TRefCountPtr(this)](TUniquePtr<FHttpOperation>&& Operation)
+	CacheStore.WaitForHttpOperationAsync(Owner, EOperationCategory::GetRef, [Self = TRefCountPtr(this)](TUniquePtr<FHttpOperation>&& Operation)
 	{
 		if (UNLIKELY(!Operation))
 		{
@@ -1488,7 +1494,7 @@ void FHttpCacheStore::FGetRecordOp::GetValues(TConstArrayView<FValueWithId> Valu
 			continue;
 		}
 
-		CacheStore.WaitForHttpOperationAsync(Owner, EOperationCategory::Get, [Self = TRefCountPtr(this), SharedOnComplete, Value](TUniquePtr<FHttpOperation>&& Operation)
+		CacheStore.WaitForHttpOperationAsync(Owner, EOperationCategory::GetBlobs, [Self = TRefCountPtr(this), SharedOnComplete, Value](TUniquePtr<FHttpOperation>&& Operation)
 		{
 			Self->BeginGetValue(MoveTemp(Operation), Value, SharedOnComplete);
 		});
@@ -1670,7 +1676,7 @@ void FHttpCacheStore::FGetValueOp::Get(const FCacheKey& InKey, ECachePolicy InPo
 	OnComplete = MoveTemp(InOnComplete);
 
 	RequestTimer.Stop();
-	CacheStore.WaitForHttpOperationAsync(Owner, EOperationCategory::Get, [Self = TRefCountPtr(this)](TUniquePtr<FHttpOperation>&& Operation)
+	CacheStore.WaitForHttpOperationAsync(Owner, EOperationCategory::GetRef, [Self = TRefCountPtr(this)](TUniquePtr<FHttpOperation>&& Operation)
 	{
 		Self->BeginGetRef(MoveTemp(Operation));
 	});
@@ -1941,7 +1947,7 @@ void FHttpCacheStore::TExistsBatchOp<RequestType>::Exists(TConstArrayView<Reques
 	FCbFieldIterator Body = BodyWriter.Save();
 
 	RequestTimer.Stop();
-	CacheStore.WaitForHttpOperationAsync(Owner, EOperationCategory::Get, [Self = TRefCountPtr(this), Body = MoveTemp(Body)](TUniquePtr<FHttpOperation>&& Operation) mutable
+	CacheStore.WaitForHttpOperationAsync(Owner, EOperationCategory::Exists, [Self = TRefCountPtr(this), Body = MoveTemp(Body)](TUniquePtr<FHttpOperation>&& Operation) mutable
 	{
 		Self->BeginExists(MoveTemp(Operation), MoveTemp(Body));
 	});
@@ -2316,13 +2322,23 @@ FHttpCacheStore::FHttpCacheStore(const FHttpCacheStoreParams& Params, ICacheStor
 	FHealthCheckOp HealthCheck(*this, *Client, StoreOwner);
 	if (AcquireAccessToken(Client.Get()) && HealthCheck.IsReady())
 	{
-		ClientParams.MaxRequests = UE_HTTPDDC_GET_REQUEST_POOL_SIZE;
-		ClientParams.MinRequests = UE_HTTPDDC_GET_REQUEST_POOL_SIZE;
-		GetRequestQueue.Initialize(*ConnectionPool, ClientParams);
+		// Default rate limits for getting blobs
+		ClientParams.MaxRequests = UE_HTTPDDC_GETBLOBS_REQUEST_POOL_SIZE;
+		ClientParams.MinRequests = UE_HTTPDDC_GETBLOBS_REQUEST_POOL_SIZE;
+		GetBlobsRequestQueue.Initialize(*ConnectionPool, ClientParams);
 
-		// Giving very generous rate limits during PUT operations as they cause too many spurious failures to put blobs or finalize refs
+		// More generous rate limits for operations that are expected to be low byte
+		// DOWNLOADED (not uploaded) count but can be delayed when the server stalls.
+		// Put operations are configured with these more generous limits even when
+		// putting high byte count items (eg: blobs).
 		ClientParams.LowSpeedLimit = 1;
 		ClientParams.LowSpeedTime = 60;
+		ClientParams.MaxRequests = UE_HTTPDDC_GETREF_REQUEST_POOL_SIZE;
+		ClientParams.MinRequests = UE_HTTPDDC_GETREF_REQUEST_POOL_SIZE;
+		GetRefRequestQueue.Initialize(*ConnectionPool, ClientParams);
+		ClientParams.MaxRequests = UE_HTTPDDC_EXISTS_REQUEST_POOL_SIZE;
+		ClientParams.MinRequests = UE_HTTPDDC_EXISTS_REQUEST_POOL_SIZE;
+		ExistsRequestQueue.Initialize(*ConnectionPool, ClientParams);
 		ClientParams.MaxRequests = UE_HTTPDDC_PUTREF_REQUEST_POOL_SIZE;
 		ClientParams.MinRequests = UE_HTTPDDC_PUTREF_REQUEST_POOL_SIZE;
 		PutRefRequestQueue.Initialize(*ConnectionPool, ClientParams);
@@ -2629,8 +2645,12 @@ FHttpRequestQueue& FHttpCacheStore::PickRequestQueue(EOperationCategory Category
 {
 	switch (Category)
 	{
-	case EOperationCategory::Get:
-		return GetRequestQueue;
+	case EOperationCategory::GetRef:
+		return GetRefRequestQueue;
+	case EOperationCategory::GetBlobs:
+		return GetBlobsRequestQueue;
+	case EOperationCategory::Exists:
+		return ExistsRequestQueue;
 	case EOperationCategory::PutRef:
 		return PutRefRequestQueue;
 	case EOperationCategory::PutBlobs:
@@ -2639,7 +2659,7 @@ FHttpRequestQueue& FHttpCacheStore::PickRequestQueue(EOperationCategory Category
 		return PutFinalizeRequestQueue;
 	default:
 		checkNoEntry();
-		return GetRequestQueue;
+		return GetBlobsRequestQueue;
 	}
 }
 
