DDC: Added filtering by key to the cache key filter

- Adds -DDC-<Name>-MissKeys=<Bucket>/<Hash>[+...]
- Adds -DDC-ReplayKeys=<Bucket>/<Hash>[+...]
- Adds -DDC-VerifyKeys=<Bucket>/<Hash>[+...]
- Adds -FilterKeys=<Bucket>/<Hash>[+...] to DerivedDataTool Copy

#Synced-CL 50967057

diff --git a/Engine/Source/Developer/DerivedDataCache/Internal/DerivedDataCacheKeyFilter.h b/Engine/Source/Developer/DerivedDataCache/Internal/DerivedDataCacheKeyFilter.h
--- a/Engine/Source/Developer/DerivedDataCache/Internal/DerivedDataCacheKeyFilter.h
+++ b/Engine/Source/Developer/DerivedDataCache/Internal/DerivedDataCacheKeyFilter.h
@@ -22,16 +22,24 @@ public:
 	 *
 	 * Example:
 	 * - Config = "[REDACTED_EMAIL]+StaticMesh -Unused -DDC-Verify=Audio@75"
-	 * - Prefix = "-DDC-Verify"
+	 * - BucketPrefix = "-DDC-Verify="
 	 * - DefaultRate = 50.0f
 	 * - Filter matches Texture at a rate of 2.5%, StaticMesh at a rate of 100%,
 	 *   Audio at a rate of 75%, and everything else at a rate of 50%.
 	 *
-	 * @param Config        The configuration in the form of command line parameters.
-	 * @param Prefix        The prefix of the parameter to parse the filter from.
-	 * @param DefaultRate   The default rate at which to apply the filter. [0, 100]
+	 * - Config = "-DDC-All-MissTypes=StaticMesh@50 -DDC-All-MissKeys=Texture/0123456789abcdef0123456789abcdef01234567"
+	 * - BucketPrefix = "-DDC-All-MissTypes="
+	 * - KeyPrefix = "-DDC-All-MissKeys="
+	 * - Filter matches StaticMesh at 50%, the single Texture key specified, and everything else at a rate of 0%.
+	 *
+	 * Multiple buckets or keys in one argument are delimited by '+' or ','.
+	 *
+	 * @param Config         The configuration in the form of command line parameters.
+	 * @param BucketPrefix   The prefix of the parameter to parse the bucket filter from. Disabled if null.
+	 * @param KeyPrefix      The prefix of the parameter to parse the key filter from. Disabled if null.
+	 * @param DefaultRate    The default rate at which to apply the filter. [0, 100]
 	 */
-	UE_API static FCacheKeyFilter Parse(const TCHAR* Config, const TCHAR* Prefix, float DefaultRate = 0.0f);
+	UE_API static FCacheKeyFilter Parse(const TCHAR* Config, const TCHAR* BucketPrefix, const TCHAR* KeyPrefix, float DefaultRate = 0.0f);
 
 	/**
 	 * Accesses the salt that is combined with the cache key for filtering.
diff --git a/Engine/Source/Developer/DerivedDataCache/Private/DerivedDataBackends.cpp b/Engine/Source/Developer/DerivedDataCache/Private/DerivedDataBackends.cpp
--- a/Engine/Source/Developer/DerivedDataCache/Private/DerivedDataBackends.cpp
+++ b/Engine/Source/Developer/DerivedDataCache/Private/DerivedDataBackends.cpp
@@ -146,10 +146,11 @@ protected:
 		// Async must exist in the graph.
 		RootCache = CreateCacheStoreAsync(Hierarchy, RootCache, MemoryCache, /*bDeleteInnerCache*/ true);
 
-		// Create a Verify node when using -DDC-Verify[=Type1[@Rate2][+Type2[@Rate2]...]].
+		// Create a Verify node when using -DDC-Verify[=Type1[@Rate2][+Type2[@Rate2]...]] or -DDC-VerifyKeys=Key1[+Key2...].
 		if (FString VerifyArg;
+			FParse::Param(FCommandLine::Get(), TEXT("DDC-Verify")) ||
 			FParse::Value(FCommandLine::Get(), TEXT("-DDC-Verify="), VerifyArg) ||
-			FParse::Param(FCommandLine::Get(), TEXT("DDC-Verify")))
+			FParse::Value(FCommandLine::Get(), TEXT("-DDC-VerifyKeys="), VerifyArg))
 		{
 			if (!UE::GetMultiprocessId())
 			{
@@ -998,28 +999,25 @@ private:
 		float DefaultRate = bDefaultMatch ? 100.0f : 0.0f;
 		FParse::Value(*JoinedArgs, TEXT("Rate="), DefaultRate);
 
-		FCacheKeyFilter KeyFilter = FCacheKeyFilter::Parse(*JoinedArgs, TEXT("Types="), DefaultRate);
+		FCacheKeyFilter KeyFilter = FCacheKeyFilter::Parse(*JoinedArgs, TEXT("Types="), nullptr, DefaultRate);
 
-		if (KeyFilter)
+		uint32 Salt;
+		if (FParse::Value(*JoinedArgs, TEXT("Salt="), Salt))
 		{
-			uint32 Salt;
-			if (FParse::Value(*JoinedArgs, TEXT("Salt="), Salt))
+			if (Salt == 0)
 			{
-				if (Salt == 0)
-				{
-					UE_LOG(LogDerivedDataCache, Warning,
-						TEXT("Replay: Ignoring salt of 0. The salt must be a positive integer."));
-				}
-				else
-				{
-					KeyFilter.SetSalt(Salt);
-				}
+				UE_LOGF(LogDerivedDataCache, Warning,
+					"Replay: Ignoring salt of 0. The salt must be a positive integer.");
+			}
+			else
+			{
+				KeyFilter.SetSalt(Salt);
 			}
-
-			UE_LOG(LogDerivedDataCache, Display,
-				TEXT("Replay: Using salt %u to filter cache keys to replay."), KeyFilter.GetSalt());
 		}
 
+		UE_CLOGF(KeyFilter.RequiresSalt(), LogDerivedDataCache, Display,
+			"Replay: Using salt %u to filter cache keys to replay.", KeyFilter.GetSalt());
+
 		Replay.SetKeyFilter(MoveTemp(KeyFilter));
 
 		// Parse Method Filter
@@ -1108,40 +1106,32 @@ FBackendDebugOptions::FBackendDebugOptions()
 bool FBackendDebugOptions::ParseFromTokens(FBackendDebugOptions& OutOptions, const TCHAR* InNodeName, const TCHAR* InInputTokens)
 {
 	// Check if the input stream has any DDC options for this node.
-	TStringBuilder<64> Prefix;
-	Prefix.Append(TEXTVIEW("-DDC-")).Append(InNodeName).Append(TEXTVIEW("-"));
+	TStringBuilder<64> Prefix(InPlace, TEXTVIEW("-DDC-"), InNodeName, '-');
 	if (UE::String::FindFirst(InInputTokens, Prefix, ESearchCase::IgnoreCase) == INDEX_NONE)
 	{
 		return false;
 	}
 
-	const int32 PrefixLen = Prefix.Len();
-
 	// Look for -DDC-Local-MissRate=, -DDC-Shared-MissRate=, etc.
 	float MissRate = 0.0f;
-	Prefix.Append(TEXTVIEW("MissRate="));
-	FParse::Value(InInputTokens, *Prefix, MissRate);
-	Prefix.RemoveAt(PrefixLen, Prefix.Len() - PrefixLen);
+	FParse::Value(InInputTokens, *WriteToString<64>(Prefix, TEXTVIEW("MissRate=")), MissRate);
 
-	// Look for -DDC-Local-MissTypes=AnimSeq+Audio, -DDC-Shared-MissType=AnimSeq+Audio, etc.
-	Prefix.Append(TEXTVIEW("MissTypes="));
-	OutOptions.SimulateMissFilter = FCacheKeyFilter::Parse(InInputTokens, *Prefix, MissRate);
-	Prefix.RemoveAt(PrefixLen, Prefix.Len() - PrefixLen);
+	// Look for -DDC-Local-MissTypes=AnimSeq+Audio, -DDC-Shared-MissTypes=AnimSeq+Audio, etc.
+	// Look for -DDC-Local-MissKeys=StaticMesh/<Hash>+Texture/<Hash>, -DDC-Shared-MissKeys=StaticMesh/<Hash>+Texture/<Hash>, etc.
+	OutOptions.SimulateMissFilter = FCacheKeyFilter::Parse(InInputTokens,
+		*WriteToString<64>(Prefix, TEXTVIEW("MissTypes=")),
+		*WriteToString<64>(Prefix, TEXTVIEW("MissKeys=")), MissRate);
 
 	// Look for -DDC-Local-MissSalt=, -DDC-Shared-MissSalt=, etc.
 	uint32 Salt = 0;
-	Prefix.Append(TEXTVIEW("MissSalt="));
-	if (FParse::Value(InInputTokens, *Prefix, Salt))
+	if (FParse::Value(InInputTokens, *WriteToString<64>(Prefix, TEXTVIEW("MissSalt=")), Salt))
 	{
 		OutOptions.SimulateMissFilter.SetSalt(Salt);
 	}
-	if (OutOptions.SimulateMissFilter)
-	{
-		UE_CLOGF(OutOptions.SimulateMissFilter.RequiresSalt(), LogDerivedDataCache, Display,
-			"%ls: Using salt %ls%u to filter cache keys to simulate misses on.",
-			InNodeName, *Prefix, OutOptions.SimulateMissFilter.GetSalt());
-	}
-	Prefix.RemoveAt(PrefixLen, Prefix.Len() - PrefixLen);
+
+	UE_CLOGF(OutOptions.SimulateMissFilter.RequiresSalt(), LogDerivedDataCache, Display,
+		"%ls: Using salt %lsMissSalt=%u to filter cache keys to simulate misses on.",
+		InNodeName, *Prefix, OutOptions.SimulateMissFilter.GetSalt());
 
 	return true;
 }
diff --git a/Engine/Source/Developer/DerivedDataCache/Private/DerivedDataCacheKeyFilter.cpp b/Engine/Source/Developer/DerivedDataCache/Private/DerivedDataCacheKeyFilter.cpp
--- a/Engine/Source/Developer/DerivedDataCache/Private/DerivedDataCacheKeyFilter.cpp
+++ b/Engine/Source/Developer/DerivedDataCache/Private/DerivedDataCacheKeyFilter.cpp
@@ -3,8 +3,8 @@
 #include "DerivedDataCacheKeyFilter.h"
 
 #include "Algo/AnyOf.h"
+#include "Algo/BinarySearch.h"
 #include "Algo/Find.h"
-#include "Containers/SharedString.h"
 #include "DerivedDataCacheKey.h"
 #include "Hash/xxhash.h"
 #include "Math/UnrealMathUtility.h"
@@ -22,11 +22,13 @@ struct Private::FCacheKeyFilterState
 {
 	struct FTypeRate
 	{
-		FUtf8SharedString Type;
+		FCacheBucket Type;
+		FCacheBucket LegacyType;
 		uint32 Rate = 0;
 	};
 
-	TArray<FTypeRate, TInlineAllocator<2>> Types;
+	TArray<FTypeRate> Types;
+	TArray<FCacheKey> Keys;
 	uint32 DefaultRate = 0;
 	uint32 Salt = 0;
 
@@ -55,38 +57,65 @@ bool Private::FCacheKeyFilterState::TryParseTypeRate(const FStringView TypeConfi
 
 	double Rate;
 	LexFromString(Rate, RateView);
-	if (TypeView.IsEmpty() || Rate < 0.0 || Rate > 100.0)
+	if (!FCacheBucket::IsValidName(TypeView) || Rate < 0.0 || Rate > 100.0)
 	{
 		return false;
 	}
 
-	OutTypeRate.Type = WriteToUtf8String<64>(TypeView);
+	OutTypeRate.Type = FCacheBucket(WriteToAnsiString<64>(TypeView));
+	OutTypeRate.LegacyType = FCacheBucket(WriteToAnsiString<64>(ANSITEXTVIEW("Legacy"), TypeView));
 	OutTypeRate.Rate = ConvertMatchRate(Rate);
 	return true;
 }
 
-FCacheKeyFilter FCacheKeyFilter::Parse(const TCHAR* const Config, const TCHAR* const Prefix, const float DefaultRate)
+FCacheKeyFilter FCacheKeyFilter::Parse(
+	const TCHAR* const Config,
+	const TCHAR* const BucketPrefix,
+	const TCHAR* const KeyPrefix,
+	const float DefaultRate)
 {
 	using namespace UE::DerivedData::Private;
 
-	TArray<FCacheKeyFilterState::FTypeRate, TInlineAllocator<2>> Types;
+	TArray<FCacheKeyFilterState::FTypeRate> Types;
+	if (BucketPrefix)
+	{
+		FString TypeConfigArray;
+		for (FStringView ConfigView(Config), PrefixView(BucketPrefix); FParse::Value(ConfigView.GetData(), BucketPrefix, TypeConfigArray);)
+		{
+			const int32 PrefixIndex = String::FindFirst(ConfigView, PrefixView, ESearchCase::IgnoreCase);
+			ConfigView.RightChopInline(PrefixIndex + PrefixView.Len() + TypeConfigArray.Len());
+			String::ParseTokensMultiple(TypeConfigArray, {TEXT('+'), TEXT(',')}, [&Types](const FStringView TypeConfig)
+			{
+				FCacheKeyFilterState::FTypeRate TypeRate;
+				if (FCacheKeyFilterState::TryParseTypeRate(TypeConfig, TypeRate))
+				{
+					Types.Emplace(TypeRate);
+				}
+			}, String::EParseTokensOptions::SkipEmpty | String::EParseTokensOptions::Trim);
+		}
+	}
 
-	FString TypeConfigArray;
-	for (FStringView ConfigView(Config), PrefixView(Prefix); FParse::Value(ConfigView.GetData(), Prefix, TypeConfigArray);)
+	TArray<FCacheKey> Keys;
+	if (KeyPrefix)
 	{
-		const int32 PrefixIndex = String::FindFirst(ConfigView, PrefixView, ESearchCase::IgnoreCase);
-		ConfigView.RightChopInline(PrefixIndex + PrefixView.Len() + TypeConfigArray.Len());
-		String::ParseTokensMultiple(TypeConfigArray, {TEXT('+'), TEXT(',')}, [&Types](const FStringView TypeConfig)
+		FString KeyArray;
+		for (FStringView ConfigView(Config), PrefixView(KeyPrefix); FParse::Value(ConfigView.GetData(), KeyPrefix, KeyArray);)
 		{
-			FCacheKeyFilterState::FTypeRate TypeRate;
-			if (FCacheKeyFilterState::TryParseTypeRate(TypeConfig, TypeRate))
+			const int32 PrefixIndex = String::FindFirst(ConfigView, PrefixView, ESearchCase::IgnoreCase);
+			ConfigView.RightChopInline(PrefixIndex + PrefixView.Len() + KeyArray.Len());
+			String::ParseTokensMultiple(KeyArray, {TEXT('+'), TEXT(',')}, [&Keys](const FStringView KeyConfig)
 			{
-				Types.Emplace(TypeRate);
-			}
-		}, String::EParseTokensOptions::SkipEmpty | String::EParseTokensOptions::Trim);
+				FCacheKey Key;
+				if (TryLexFromString(Key, KeyConfig))
+				{
+					Keys.Emplace(Key);
+				}
+			}, String::EParseTokensOptions::SkipEmpty | String::EParseTokensOptions::Trim);
+		}
+		Keys.Sort();
 	}
 
-	if (Types.IsEmpty() && DefaultRate == 0.0f)
+	if (Types.IsEmpty() && Keys.IsEmpty() && DefaultRate == 0.0f)
 	{
 		return {};
 	}
@@ -94,6 +123,7 @@ FCacheKeyFilter FCacheKeyFilter::Parse(const TCHAR* const Config, const TCHAR* c
 	FCacheKeyFilter Filter;
 	Filter.State = MakePimpl<FCacheKeyFilterState, EPimplPtrMode::DeepCopy>();
 	Filter.State->Types = MoveTemp(Types);
+	Filter.State->Keys = MoveTemp(Keys);
 	Filter.State->DefaultRate = FCacheKeyFilterState::ConvertMatchRate(FMath::Clamp(DefaultRate, 0.0f, 100.0f));
 	Filter.SetSalt(0);
 	return Filter;
@@ -140,19 +170,21 @@ bool FCacheKeyFilter::IsMatch(const FCacheKey& Key) const
 		return false;
 	}
 
+	if (!State->Keys.IsEmpty() && Algo::BinarySearch(State->Keys, Key) != INDEX_NONE)
+	{
+		return true;
+	}
+
 	uint32 TargetRate = State->DefaultRate;
 
 	if (!State->Types.IsEmpty())
 	{
-		TUtf8StringBuilder<256> Bucket;
-		Bucket << Key.Bucket;
-		if (Bucket.ToView().StartsWith(UTF8TEXTVIEW("Legacy")))
-		{
-			Bucket.RemoveAt(0, UTF8TEXTVIEW("Legacy").Len());
-		}
-
 		using FTypeRate = FCacheKeyFilterState::FTypeRate;
-		if (const FTypeRate* TypeRate = Algo::FindBy(State->Types, Bucket.ToView(), &FTypeRate::Type))
+		const auto MatchType = [Bucket = Key.Bucket](const FTypeRate& TypeRate)
+		{
+			return TypeRate.Type == Bucket || TypeRate.LegacyType == Bucket;
+		};
+		if (const FTypeRate* TypeRate = Algo::FindByPredicate(State->Types, MatchType))
 		{
 			TargetRate = TypeRate->Rate;
 		}
diff --git a/Engine/Source/Developer/DerivedDataCache/Private/DerivedDataCacheReplay.cpp b/Engine/Source/Developer/DerivedDataCache/Private/DerivedDataCacheReplay.cpp
--- a/Engine/Source/Developer/DerivedDataCache/Private/DerivedDataCacheReplay.cpp
+++ b/Engine/Source/Developer/DerivedDataCache/Private/DerivedDataCacheReplay.cpp
@@ -1099,32 +1099,31 @@ static FCacheMethodFilter ParseReplayMethodFilter(const TCHAR* const CommandLine
 
 static FCacheKeyFilter ParseReplayKeyFilter(const TCHAR* const CommandLine)
 {
-	const bool bDefaultMatch = String::FindFirst(CommandLine, TEXT("-DDC-ReplayTypes="), ESearchCase::IgnoreCase) == INDEX_NONE;
+	const bool bDefaultMatch =
+		String::FindFirst(CommandLine, TEXT("-DDC-ReplayTypes="), ESearchCase::IgnoreCase) == INDEX_NONE &&
+		String::FindFirst(CommandLine, TEXT("-DDC-ReplayKeys="), ESearchCase::IgnoreCase) == INDEX_NONE;
 	float DefaultRate = bDefaultMatch ? 100.0f : 0.0f;
 	FParse::Value(CommandLine, TEXT("-DDC-ReplayRate="), DefaultRate);
 
-	FCacheKeyFilter KeyFilter = FCacheKeyFilter::Parse(CommandLine, TEXT("-DDC-ReplayTypes="), DefaultRate);
+	FCacheKeyFilter KeyFilter = FCacheKeyFilter::Parse(CommandLine, TEXT("-DDC-ReplayTypes="), TEXT("-DDC-ReplayKeys="), DefaultRate);
 
-	if (KeyFilter)
+	uint32 Salt;
+	if (FParse::Value(CommandLine, TEXT("-DDC-ReplaySalt="), Salt))
 	{
-		uint32 Salt;
-		if (FParse::Value(CommandLine, TEXT("-DDC-ReplaySalt="), Salt))
+		if (Salt == 0)
 		{
-			if (Salt == 0)
-			{
-				UE_LOG(LogDerivedDataCache, Warning,
-					TEXT("Replay: Ignoring salt of 0. The salt must be a positive integer."));
-			}
-			else
-			{
-				KeyFilter.SetSalt(Salt);
-			}
+			UE_LOGF(LogDerivedDataCache, Warning,
+				"Replay: Ignoring salt of 0. The salt must be a positive integer.");
+		}
+		else
+		{
+			KeyFilter.SetSalt(Salt);
 		}
-
-		UE_CLOGF(KeyFilter.RequiresSalt(), LogDerivedDataCache, Display,
-			"Replay: Using salt -DDC-ReplaySalt=%u to filter cache keys to replay.", KeyFilter.GetSalt());
 	}
 
+	UE_CLOGF(KeyFilter.RequiresSalt(), LogDerivedDataCache, Display,
+		"Replay: Using salt -DDC-ReplaySalt=%u to filter cache keys to replay.", KeyFilter.GetSalt());
+
 	return KeyFilter;
 }
 
diff --git a/Engine/Source/Developer/DerivedDataCache/Private/DerivedDataCacheStoreVerify.cpp b/Engine/Source/Developer/DerivedDataCache/Private/DerivedDataCacheStoreVerify.cpp
--- a/Engine/Source/Developer/DerivedDataCache/Private/DerivedDataCacheStoreVerify.cpp
+++ b/Engine/Source/Developer/DerivedDataCache/Private/DerivedDataCacheStoreVerify.cpp
@@ -40,35 +40,33 @@ public:
 		const TCHAR* const CommandLine = FCommandLine::Get();
 
 		const bool bDefaultMatch = FParse::Param(CommandLine, TEXT("DDC-Verify")) ||
-			String::FindFirst(CommandLine, TEXT("-DDC-Verify="), ESearchCase::IgnoreCase) == INDEX_NONE;
+			(String::FindFirst(CommandLine, TEXT("-DDC-Verify="), ESearchCase::IgnoreCase) == INDEX_NONE &&
+			 String::FindFirst(CommandLine, TEXT("-DDC-VerifyKeys="), ESearchCase::IgnoreCase) == INDEX_NONE);
 		float DefaultRate = bDefaultMatch ? 100.0f : 0.0f;
 		FParse::Value(CommandLine, TEXT("-DDC-VerifyRate="), DefaultRate);
 
-		Filter = FCacheKeyFilter::Parse(CommandLine, TEXT("-DDC-Verify="), DefaultRate);
+		KeyFilter = FCacheKeyFilter::Parse(CommandLine, TEXT("-DDC-Verify="), TEXT("-DDC-VerifyKeys="), DefaultRate);
 
 		uint32 Salt;
 		if (FParse::Value(CommandLine, TEXT("-DDC-VerifySalt="), Salt))
 		{
 			if (Salt == 0)
 			{
-				UE_LOG(LogDerivedDataCache, Warning,
-					TEXT("Verify: Ignoring salt of 0. The salt must be a positive integer."));
+				UE_LOGF(LogDerivedDataCache, Warning,
+					"Verify: Ignoring salt of 0. The salt must be a positive integer.");
 			}
 			else
 			{
-				Filter.SetSalt(Salt);
+				KeyFilter.SetSalt(Salt);
 			}
 		}
 
-		if (Filter)
-		{
-			UE_CLOGF(Filter.RequiresSalt(), LogDerivedDataCache, Display,
-				"Verify: Using salt -DDC-VerifySalt=%u to filter cache keys to verify.", Filter.GetSalt());
-		}
+		UE_CLOGF(KeyFilter.RequiresSalt(), LogDerivedDataCache, Display,
+			"Verify: Using salt -DDC-VerifySalt=%u to filter cache keys to verify.", KeyFilter.GetSalt());
 
 		bPutOnError = bPutOnError || FParse::Param(CommandLine, TEXT("DDC-VerifyFix"));
-		UE_CLOG(bPutOnError, LogDerivedDataCache, Display,
-			TEXT("Verify: Any record or value that differs will be overwritten."));
+		UE_CLOGF(bPutOnError, LogDerivedDataCache, Display,
+			"Verify: Any record or value that differs will be overwritten.");
 	}
 
 	~FCacheStoreVerify()
@@ -180,7 +178,7 @@ private:
 	ILegacyCacheStore* InnerCache;
 	FCriticalSection AlreadyTestedLock;
 	TSet<FCacheKey> AlreadyTested;
-	FCacheKeyFilter Filter;
+	FCacheKeyFilter KeyFilter;
 	bool bPutOnError;
 };
 
@@ -207,7 +205,7 @@ void FCacheStoreVerify::Put(
 		for (const FCachePutRequest& Request : Requests)
 		{
 			const FCacheKey& Key = Request.Record.GetKey();
-			bool bForward = EnumHasAnyFlags(Request.Policy.GetRecordPolicy(), ECachePolicy::NonDeterministic) || !Filter.IsMatch(Key);
+			bool bForward = EnumHasAnyFlags(Request.Policy.GetRecordPolicy(), ECachePolicy::NonDeterministic) || !KeyFilter.IsMatch(Key);
 			if (!bForward)
 			{
 				AlreadyTested.Add(Key, &bForward);
@@ -375,7 +373,7 @@ void FCacheStoreVerify::Get(
 		for (const FCacheGetRequest& Request : Requests)
 		{
 			const bool bForward = EnumHasAnyFlags(Request.Policy.GetRecordPolicy(), ECachePolicy::NonDeterministic) ||
-				!Filter.IsMatch(Request.Key) || AlreadyTested.Contains(Request.Key);
+				!KeyFilter.IsMatch(Request.Key) || AlreadyTested.Contains(Request.Key);
 			(bForward ? ForwardRequests : VerifyRequests).Add(Request);
 		}
 	}
@@ -399,7 +397,7 @@ void FCacheStoreVerify::PutValue(
 		FScopeLock Lock(&AlreadyTestedLock);
 		for (const FCachePutValueRequest& Request : Requests)
 		{
-			bool bForward = EnumHasAnyFlags(Request.Policy, ECachePolicy::NonDeterministic) || !Filter.IsMatch(Request.Key);
+			bool bForward = EnumHasAnyFlags(Request.Policy, ECachePolicy::NonDeterministic) || !KeyFilter.IsMatch(Request.Key);
 			if (!bForward)
 			{
 				AlreadyTested.Add(Request.Key, &bForward);
@@ -572,7 +570,7 @@ void FCacheStoreVerify::GetValue(
 		for (const FCacheGetValueRequest& Request : Requests)
 		{
 			const bool bForward = EnumHasAnyFlags(Request.Policy, ECachePolicy::NonDeterministic) ||
-				!Filter.IsMatch(Request.Key) || AlreadyTested.Contains(Request.Key);
+				!KeyFilter.IsMatch(Request.Key) || AlreadyTested.Contains(Request.Key);
 			(bForward ? ForwardRequests : VerifyRequests).Add(Request);
 		}
 	}
@@ -599,7 +597,7 @@ void FCacheStoreVerify::GetChunks(
 		for (const FCacheGetChunkRequest& Request : Requests)
 		{
 			const bool bForward = EnumHasAnyFlags(Request.Policy, ECachePolicy::NonDeterministic) ||
-				!Filter.IsMatch(Request.Key) || AlreadyTested.Contains(Request.Key);
+				!KeyFilter.IsMatch(Request.Key) || AlreadyTested.Contains(Request.Key);
 			(bForward ? ForwardRequests : VerifyRequests).Add(Request);
 		}
 	}
diff --git a/Engine/Source/Programs/DerivedDataTool/Private/DerivedDataCopyCommand.cpp b/Engine/Source/Programs/DerivedDataTool/Private/DerivedDataCopyCommand.cpp
--- a/Engine/Source/Programs/DerivedDataTool/Private/DerivedDataCopyCommand.cpp
+++ b/Engine/Source/Programs/DerivedDataTool/Private/DerivedDataCopyCommand.cpp
@@ -76,8 +76,9 @@ FCommand MakeCopyCommand()
 		//.AddSwitch("Key=<Bucket>/<Hash>", "Key to copy from source to target(s). Supports multiple keys.")
 		.AddSwitch("Replay=<Path>", "Replay to copy from source to target(s). Supports multiple replays.")
 		.AddSwitch("Filter=<Bucket>[@<Rate>][+...]", "Filter cache keys to copy. Rate is 0-100 percent.")
-		.AddSwitch("FilterRate=<Rate>", "Filter rate (0-100) to apply to keys not matched by a bucket filter.")
+		.AddSwitch("FilterRate=<Rate>", "Filter rate (0-100) to apply to keys not matched by a bucket or key filter.")
 		.AddSwitch("FilterSalt=<Salt>", "Filter salt when filtering cache keys with fractional match rates.")
+		.AddSwitch("FilterKeys=<Key>[+...]", "Filter cache keys to copy. Always matches these keys.")
 		.AddSwitch("ValueMemoryMB=<Limit>", "Limit memory used to cache values to this amount. Default is 16,384 MiB.")
 		.AddSwitch("Preview", "Preview the copy without modifying the target(s).")
 		.AddSwitch("Overwrite", "Overwrite data in the target(s) that differs from the replays.")
@@ -534,16 +535,13 @@ struct FCopyParams
 		}
 	}
 
-	const bool bDefaultMatch = !FParse::Value(Options, TEXT("-Filter="), Value);
+	const bool bDefaultMatch = !FParse::Value(Options, TEXT("-Filter="), Value) && !FParse::Value(Options, TEXT("-FilterKeys="), Value);
 	float FilterRate = bDefaultMatch ? 100.0f : 0.0f;
 	uint32 FilterSalt = 0;
 	FParse::Value(Options, TEXT("-FilterRate="), FilterRate);
 	FParse::Value(Options, TEXT("-FilterSalt="), FilterSalt);
-	OutParams.Filter = FCacheKeyFilter::Parse(Options, TEXT("-Filter="), FilterRate);
-	if (FilterSalt)
-	{
-		OutParams.Filter.SetSalt(FilterSalt);
-	}
+	OutParams.Filter = FCacheKeyFilter::Parse(Options, TEXT("-Filter="), TEXT("-FilterKeys="), FilterRate);
+	OutParams.Filter.SetSalt(FilterSalt);
 	UE_CLOGFMT(OutParams.Filter.RequiresSalt(), LogDerivedDataTool, Display,
 		"Using salt -FilterSalt={FilterSalt} to filter cache keys.", OutParams.Filter.GetSalt());
 
