Disabling a zen cache store if we have multiple possible hosts, but none are resolvable in the time limits.  To allow it to continue with an arbitrary host means it will delay startup by timing out again while waiting for a connection to the arbitrary host that was selected.  It also creates a situation where the client will re-try connection to that arbirary host periodically in the background.

An alterative would be to allow it to proceed, but give it the full list of candidate hosts that would be re-evaluated in the background.  This isn't easy to accomplish because right now the code expects a single resolved host by the time we have finished creating the cache store.

#Synced-CL 51289657

diff --git a/Engine/Source/Developer/DerivedDataCache/Private/ZenCacheStore.cpp b/Engine/Source/Developer/DerivedDataCache/Private/ZenCacheStore.cpp
--- a/Engine/Source/Developer/DerivedDataCache/Private/ZenCacheStore.cpp
+++ b/Engine/Source/Developer/DerivedDataCache/Private/ZenCacheStore.cpp
@@ -2411,11 +2411,11 @@ ILegacyCacheStore* CreateZenCacheStore(const TCHAR* NodeName, const TCHAR* Confi
 		}
 		else
 		{
-			// even if we fail to resolve a host to use the returned host will at least contain the first of the possible hosts which we can attempt to use
-			Params.Host = FString(ResolvedHost);
-			FString HostCandidates = HostBuilder.GetHostCandidatesString();
-			UE_LOG(LogDerivedDataCache, Display, TEXT("%s: Unable to resolve best host candidate to use, most likely none of the suggested hosts was reachable. Attempted hosts were: '%s' ."), NodeName, *HostCandidates);
-		}		
+			// If we have multiple hosts and fail to resolve any, we don't want to supply an arbitrary one to use.  We should just disable this cache storage layer.
+			UE_LOG(LogDerivedDataCache, Display, TEXT("%s: Disabled due to inability to resolve best host candidate to use, most likely none of the suggested hosts was reachable.")
+				TEXT(" Attempted hosts were: '%s'."), NodeName, *HostBuilder.GetHostCandidatesString());
+			return nullptr;
+		}
 	}
 
 	bool bHasSandbox = !Params.Sandbox.IsEmpty();
