Fix a race condition involving the use of the Monitor member ZenCacheStore and HttpCacheStore async package retreivers.

#Synced-CL 51529686

diff --git a/Engine/Source/Developer/DerivedDataCache/Private/HttpCacheStore.cpp b/Engine/Source/Developer/DerivedDataCache/Private/HttpCacheStore.cpp
--- a/Engine/Source/Developer/DerivedDataCache/Private/HttpCacheStore.cpp
+++ b/Engine/Source/Developer/DerivedDataCache/Private/HttpCacheStore.cpp
@@ -600,14 +600,17 @@ private:
 	// IRequest Interface
 
 	void SetPriority(EPriority Priority) final {}
-	void Cancel() final { Monitor->Cancel(); }
-	void Wait() final { Monitor->Wait(); }
+	void Cancel() final { CopyMonitorRef()->Cancel(); }
+	void Wait() final { CopyMonitorRef()->Wait(); }
 
 	// IHttpReceiver Interface
 
 	IHttpReceiver* OnCreate(IHttpResponse& LocalResponse) final
 	{
-		Monitor = LocalResponse.GetMonitor();
+		{
+			TUniqueLock Lock(MonitorMutex);
+			Monitor = LocalResponse.GetMonitor();
+		}
 		Owner->Begin(this);
 		return &BaseReceiver;
 	}
@@ -636,10 +639,17 @@ private:
 	}
 
 private:
+	TRefCountPtr<IHttpResponseMonitor> CopyMonitorRef()
+	{
+		TUniqueLock Lock(MonitorMutex);
+		return Monitor;
+	}
+
 	IRequestOwner* Owner;
 	FHttpOperationReceiver BaseReceiver;
 	TUniqueFunction<void ()> OperationComplete;
 	TRefCountPtr<IHttpResponseMonitor> Monitor;
+	FMutex MonitorMutex;
 };
 
 void FHttpCacheStore::FHttpOperation::Send()
diff --git a/Engine/Source/Developer/DerivedDataCache/Private/ZenCacheStore.cpp b/Engine/Source/Developer/DerivedDataCache/Private/ZenCacheStore.cpp
--- a/Engine/Source/Developer/DerivedDataCache/Private/ZenCacheStore.cpp
+++ b/Engine/Source/Developer/DerivedDataCache/Private/ZenCacheStore.cpp
@@ -1809,14 +1809,17 @@ private:
 	// IRequest Interface
 
 	void SetPriority(EPriority Priority) final {}
-	void Cancel() final { Monitor->Cancel(); }
-	void Wait() final { Monitor->Wait(); }
+	void Cancel() final { CopyMonitorRef()->Cancel(); }
+	void Wait() final { CopyMonitorRef()->Wait(); }
 
 	// IHttpReceiver Interface
 
 	IHttpReceiver* OnCreate(IHttpResponse& LocalResponse) final
 	{
-		Monitor = LocalResponse.GetMonitor();
+		{
+			TUniqueLock Lock(MonitorMutex);
+			Monitor = LocalResponse.GetMonitor();
+		}
 		Owner->Begin(this);
 		return FAsyncCbPackageReceiver::OnCreate(LocalResponse);
 	}
@@ -1831,8 +1834,15 @@ private:
 	}
 
 private:
+	TRefCountPtr<IHttpResponseMonitor> CopyMonitorRef()
+	{
+		TUniqueLock Lock(MonitorMutex);
+		return Monitor;
+	}
+
 	TRefCountPtr<IHttpResponseMonitor> Monitor;
 	IRequestOwner* Owner;
+	FMutex MonitorMutex;
 };
 
 FZenCacheStore::FZenCacheStore(
