DDC: Hoisted the debug options from the cache stores into the hierarchy

#rnx
#Synced-CL 50967057
#Synced-CL 50967057

diff --git a/Engine/Source/Developer/DerivedDataCache/Private/DerivedDataBackends.cpp b/Engine/Source/Developer/DerivedDataCache/Private/DerivedDataBackends.cpp
--- a/Engine/Source/Developer/DerivedDataCache/Private/DerivedDataBackends.cpp
+++ b/Engine/Source/Developer/DerivedDataCache/Private/DerivedDataBackends.cpp
@@ -88,8 +88,6 @@ public:
 		// Needs to be initialized from the main thread because it uses config and console variables
 		check(IsInGameThread());
 		check(GConfig && GConfig->IsReadyForUse());
-
-		bHasDefaultDebugOptions = FBackendDebugOptions::ParseFromTokens(DefaultDebugOptions, TEXT("All"), FCommandLine::Get());
 	}
 
 	static FDerivedDataBackendGraph* TryCreate(FStringView GraphNameOrConfig)
@@ -592,27 +590,10 @@ private:
 		}
 	}
 
-	void Add(ILegacyCacheStore* CacheStore, ECacheStoreFlags Flags) final
+	void Add(const FSharedString& Name, ILegacyCacheStore* CacheStore, ECacheStoreFlags Flags) final
 	{
 		check(Hierarchy);
 
-		// Parse any debug options for this node. E.g. -DDC-<Name>-MissRate
-		FBackendDebugOptions DebugOptions;
-		if (FBackendDebugOptions::ParseFromTokens(DebugOptions, *ActiveNodeName, FCommandLine::Get()))
-		{
-			if (!CacheStore->LegacyDebugOptions(DebugOptions))
-			{
-				UE_LOG(LogDerivedDataCache, Warning, TEXT("%s: Node is ignoring one or more -DDC-%s-<Option> debug options."), *ActiveNodeName, *ActiveNodeName);
-			}
-		}
-		else if (bHasDefaultDebugOptions)
-		{
-			if (!CacheStore->LegacyDebugOptions(DefaultDebugOptions))
-			{
-				UE_LOG(LogDerivedDataCache, Warning, TEXT("%s: Node is ignoring one or more -DDC-All-<Option> debug options."), *ActiveNodeName);
-			}
-		}
-
 		// Add a throttling layer if parameters are found
 		{
 			uint32 LatencyMS = 0;
@@ -629,7 +610,7 @@ private:
 			}
 		}
 
-		Hierarchy->Add(CacheStore, Flags);
+		Hierarchy->Add(Name, CacheStore, Flags);
 	}
 
 	void SetFlags(ILegacyCacheStore* CacheStore, ECacheStoreFlags Flags) final
@@ -843,9 +824,6 @@ private:
 	/** List of directories used by the DDC */
 	TArray<FString> Directories;
 
-	FBackendDebugOptions DefaultDebugOptions;
-	bool bHasDefaultDebugOptions = false;
-
 	/** Whether a shared cache is in use */
 	bool bUsingSharedDDC = false;
 
diff --git a/Engine/Source/Developer/DerivedDataCache/Private/DerivedDataCacheReplay.cpp b/Engine/Source/Developer/DerivedDataCache/Private/DerivedDataCacheReplay.cpp
--- a/Engine/Source/Developer/DerivedDataCache/Private/DerivedDataCacheReplay.cpp
+++ b/Engine/Source/Developer/DerivedDataCache/Private/DerivedDataCacheReplay.cpp
@@ -105,11 +105,6 @@ public:
 		InnerCache->LegacyStats(OutNode);
 	}
 
-	bool LegacyDebugOptions(FBackendDebugOptions& Options) final
-	{
-		return InnerCache->LegacyDebugOptions(Options);
-	}
-
 	void SetReader(FCacheReplayReader&& Reader)
 	{
 		ReplayReader = MoveTemp(Reader);
diff --git a/Engine/Source/Developer/DerivedDataCache/Private/DerivedDataCacheStore.h b/Engine/Source/Developer/DerivedDataCache/Private/DerivedDataCacheStore.h
--- a/Engine/Source/Developer/DerivedDataCache/Private/DerivedDataCacheStore.h
+++ b/Engine/Source/Developer/DerivedDataCache/Private/DerivedDataCacheStore.h
@@ -198,7 +198,7 @@ ENUM_CLASS_FLAGS(ECacheStoreFlags);
 class ICacheStoreOwner
 {
 public:
-	virtual void Add(ILegacyCacheStore* CacheStore, ECacheStoreFlags Flags) = 0;
+	virtual void Add(const FSharedString& Name, ILegacyCacheStore* CacheStore, ECacheStoreFlags Flags) = 0;
 	virtual void SetFlags(ILegacyCacheStore* CacheStore, ECacheStoreFlags Flags) = 0;
 	virtual void RemoveNotSafe(ILegacyCacheStore* CacheStore) = 0;
 
diff --git a/Engine/Source/Developer/DerivedDataCache/Private/DerivedDataCacheStoreAsync.cpp b/Engine/Source/Developer/DerivedDataCache/Private/DerivedDataCacheStoreAsync.cpp
--- a/Engine/Source/Developer/DerivedDataCache/Private/DerivedDataCacheStoreAsync.cpp
+++ b/Engine/Source/Developer/DerivedDataCache/Private/DerivedDataCacheStoreAsync.cpp
@@ -96,11 +96,6 @@ public:
 
 	virtual void LegacyStats(FDerivedDataCacheStatsNode& OutNode) override;
 
-	virtual bool LegacyDebugOptions(FBackendDebugOptions& Options) override
-	{
-		return InnerCache->LegacyDebugOptions(Options);
-	}
-
 private:
 	COOK_STAT(using CookStatsFunction = FCookStats::FScopedStatsCounter (FDerivedDataCacheUsageStats::*)());
 
diff --git a/Engine/Source/Developer/DerivedDataCache/Private/DerivedDataCacheStoreHierarchy.cpp b/Engine/Source/Developer/DerivedDataCache/Private/DerivedDataCacheStoreHierarchy.cpp
--- a/Engine/Source/Developer/DerivedDataCache/Private/DerivedDataCacheStoreHierarchy.cpp
+++ b/Engine/Source/Developer/DerivedDataCache/Private/DerivedDataCacheStoreHierarchy.cpp
@@ -10,6 +10,7 @@
 #include "Async/UniqueLock.h"
 #include "Containers/Array.h"
 #include "Containers/StaticArray.h"
+#include "DerivedDataBackendInterface.h"
 #include "DerivedDataCache.h"
 #include "DerivedDataCacheStats.h"
 #include "DerivedDataCacheStore.h"
@@ -21,6 +22,7 @@
 #include "HAL/PlatformProcess.h"
 #include "Logging/StructuredLog.h"
 #include "MemoryCacheStore.h"
+#include "Misc/CommandLine.h"
 #include "Misc/EnumClassFlags.h"
 #include "ProfilingDebugging/CookStats.h"
 #include "Serialization/CompactBinary.h"
@@ -44,7 +46,7 @@ public:
 	explicit FCacheStoreHierarchy(ICacheStoreOwner*& OutOwner, TFunctionRef<void (IMemoryCacheStore*&)> MemoryCacheCreator);
 	~FCacheStoreHierarchy() final;
 
-	void Add(ILegacyCacheStore* CacheStore, ECacheStoreFlags Flags) final;
+	void Add(const FSharedString& Name, ILegacyCacheStore* CacheStore, ECacheStoreFlags Flags) final;
 	void SetFlags(ILegacyCacheStore* CacheStore, ECacheStoreFlags Flags) final;
 	void RemoveNotSafe(ILegacyCacheStore* CacheStore) final;
 
@@ -91,8 +93,6 @@ public:
 
 	void LegacyStats(FDerivedDataCacheStatsNode& OutNode) final;
 
-	bool LegacyDebugOptions(FBackendDebugOptions& Options) final;
-
 	template <typename PutRequestType, typename GetRequestType> struct TBatchParams;
 	using FCacheRecordBatchParams = TBatchParams<FCachePutRequest, FCacheGetRequest>;
 	using FCacheValueBatchParams = TBatchParams<FCachePutValueRequest, FCacheGetValueRequest>;
@@ -131,9 +131,6 @@ private:
 	static uint64 MeasureLogicalRecordSize(const FCacheRecord& Record);
 	static uint64 MeasureLogicalValueSize(const FValue& Value);
 
-	static const FCacheKey& GetKey(const FCachePutRequest& Request) { return Request.Record.GetKey(); }
-	static const FCacheKey& GetKey(const FCachePutValueRequest& Request) { return Request.Key; }
-
 	struct FCacheStoreNode
 	{
 		inline void SetCacheFlags(ECacheStoreFlags NewCacheFlags)
@@ -146,11 +143,13 @@ private:
 			return (ECacheStoreFlags)static_cast<std::underlying_type_t<ECacheStoreFlags>>(FPlatformAtomics::AtomicRead(reinterpret_cast<const std::make_signed_t<std::underlying_type_t<ECacheStoreFlags>>*>(&CacheFlags)));
 		}
 
+		FSharedString Name;
 		ILegacyCacheStore* Cache{};
 		ECacheStoreFlags CacheFlags{};
 		ECacheStoreNodeFlags NodeFlags{};
 		ILegacyCacheStore* AsyncCache{};
 		TArray<FCacheStoreStats*> CacheStats;
+		mutable FBackendDebugOptions DebugOptions;
 	};
 
 	mutable FSharedRecursiveMutex NodesMutex;
@@ -313,14 +312,19 @@ FCacheStoreHierarchy::~FCacheStoreHierarchy()
 	}
 }
 
-void FCacheStoreHierarchy::Add(ILegacyCacheStore* CacheStore, ECacheStoreFlags Flags)
+void FCacheStoreHierarchy::Add(const FSharedString& Name, ILegacyCacheStore* CacheStore, ECacheStoreFlags Flags)
 {
+	// Parse debug options for this node from, e.g., -DDC-<Name>-MissRate, with a fallback to, e.g., -DDC-All-MissRate.
+	FBackendDebugOptions DebugOptions;
+	FBackendDebugOptions::ParseFromTokens(DebugOptions, *Name, FCommandLine::Get()) ||
+	FBackendDebugOptions::ParseFromTokens(DebugOptions, TEXT("All"), FCommandLine::Get());
+
 	TUniqueLock Lock(NodesMutex);
 	checkf(!Algo::FindBy(Nodes, CacheStore, &FCacheStoreNode::Cache),
 		TEXT("Attempting to add a cache store that was previously registered to the hierarchy."));
 	ILegacyCacheStore* AsyncCacheStore = EnumHasAnyFlags(Flags, ECacheStoreFlags::NativeAsync) ? CacheStore
 		: CreateCacheStoreAsync(this, CacheStore, /*MemoryCache*/ nullptr, /*bDeleteInnerCache*/ false);
-	Nodes.Add({CacheStore, Flags, {}, AsyncCacheStore});
+	Nodes.Emplace(Name, CacheStore, Flags, ECacheStoreNodeFlags::None, AsyncCacheStore, TArray<FCacheStoreStats*>{}, MoveTemp(DebugOptions));
 	UpdateNodeFlags();
 }
 
@@ -770,7 +774,15 @@ bool FCacheStoreHierarchy::TPutBatch<Params>::DispatchGetRequests()
 	{
 		if (!States[RequestIndex].bStop && CanQuery(GetCombinedPolicy(Request.Policy), CacheFlags))
 		{
-			NodeRequests.Add(MakeGetRequest(Request, RequestIndex));
+			if (Node.DebugOptions.ShouldSimulateGetMiss(GetCacheKey(Request)))
+			{
+				UE_LOGF(LogDerivedDataCache, Verbose, "%ls: Simulated miss for get of %ls from '%ls'",
+					*Node.Name, *WriteToString<96>(GetCacheKey(Request)), *Request.Name);
+			}
+			else
+			{
+				NodeRequests.Add(MakeGetRequest(Request, RequestIndex));
+			}
 		}
 		++RequestIndex;
 	}
@@ -834,6 +846,8 @@ bool FCacheStoreHierarchy::TPutBatch<Params>::DispatchPutRequests()
 		const FRequestState& State = States[RequestIndex];
 		if (!State.bStop && CanStore(GetCombinedPolicy(Request.Policy), CacheFlags))
 		{
+			// We never simulate a put miss but we query it for the side effect of allowing subsequent gets.
+			Node.DebugOptions.ShouldSimulatePutMiss(GetCacheKey(Request));
 			(State.bFinished ? AsyncNodeRequests : NodeRequests).Add_GetRef(Request).UserData = uint64(RequestIndex);
 		}
 		++RequestIndex;
@@ -1019,12 +1033,23 @@ void FCacheStoreHierarchy::TGetBatch<Params>::DispatchRequests()
 			{
 				if (HasResponseData(Response) && CanStore(GetCombinedPolicy(Request.Policy), CacheFlags))
 				{
+					// We never simulate a put miss but we query it for the side effect of allowing subsequent gets.
+					Node.DebugOptions.ShouldSimulatePutMiss(GetCacheKey(Request));
 					AsyncNodeRequests.Add(MakePutRequest(Response, Request));
 					++State.NodeIndex;
 				}
 				else if (EnumHasAnyFlags(CacheFlags, ECacheStoreFlags::StopGetStore) && CanQuery(GetCombinedPolicy(Request.Policy), CacheFlags))
 				{
-					NodeRequests.Add({Request.Name, Request.Key, AddPolicy(Request.Policy, ECachePolicy::SkipData), StateIndex});
+					if (Node.DebugOptions.ShouldSimulateGetMiss(Request.Key))
+					{
+						UE_LOGF(LogDerivedDataCache, Verbose, "%ls: Simulated miss for get of %ls from '%ls'",
+							*Node.Name, *WriteToString<96>(Request.Key), *Request.Name);
+						++State.NodeIndex;
+					}
+					else
+					{
+						NodeRequests.Add({Request.Name, Request.Key, AddPolicy(Request.Policy, ECachePolicy::SkipData), StateIndex});
+					}
 				}
 				else
 				{
@@ -1035,17 +1060,26 @@ void FCacheStoreHierarchy::TGetBatch<Params>::DispatchRequests()
 			{
 				if (const ECachePolicy CombinedPolicy = GetCombinedPolicy(Request.Policy); CanQuery(CombinedPolicy, CacheFlags))
 				{
-					auto Policy = Request.Policy;
-					if (CanStoreIfOk(CombinedPolicy, Node.NodeFlags))
+					if (Node.DebugOptions.ShouldSimulateGetMiss(Request.Key))
 					{
-						Policy = RemovePolicy(Policy, ECachePolicy::SkipData | ECachePolicy::SkipMeta);
+						UE_LOGF(LogDerivedDataCache, Verbose, "%ls: Simulated miss for get of %ls from '%ls'",
+							*Node.Name, *WriteToString<96>(Request.Key), *Request.Name);
+						++State.NodeIndex;
 					}
-					if (CanQueryIfError(CombinedPolicy, Node.NodeFlags))
+					else
 					{
-						Policy = AddPolicy(Policy, ECachePolicy::PartialRecord);
+						auto Policy = Request.Policy;
+						if (CanStoreIfOk(CombinedPolicy, Node.NodeFlags))
+						{
+							Policy = RemovePolicy(Policy, ECachePolicy::SkipData | ECachePolicy::SkipMeta);
+						}
+						if (CanQueryIfError(CombinedPolicy, Node.NodeFlags))
+						{
+							Policy = AddPolicy(Policy, ECachePolicy::PartialRecord);
+						}
+						ModifyPolicyForResponse(Policy, Response);
+						NodeRequests.Add({Request.Name, Request.Key, Policy, StateIndex});
 					}
-					ModifyPolicyForResponse(Policy, Response);
-					NodeRequests.Add({Request.Name, Request.Key, Policy, StateIndex});
 				}
 				else
 				{
@@ -1126,6 +1160,8 @@ void FCacheStoreHierarchy::TGetBatch<Params>::CompleteRequest(FGetResponse&& Res
 			const FCacheStoreNode& PutNode = Hierarchy.Nodes[PutNodeIndex];
 			if (CanStore(GetCombinedPolicy(State.Request.Policy), PutNode.GetCacheFlags()))
 			{
+				// We never simulate a put miss but we query it for the side effect of allowing subsequent gets.
+				PutNode.DebugOptions.ShouldSimulatePutMiss(GetCacheKey(PutRequest));
 				FRequestBarrier AsyncBarrier(AsyncOwner);
 				Invoke(Put(), PutNode.AsyncCache, MakeArrayView(&PutRequest, 1), AsyncOwner, [](auto&&){});
 			}
@@ -1520,7 +1556,15 @@ void FCacheStoreHierarchy::FGetChunksBatch::DispatchRequests()
 			const FCacheGetChunkRequest& Request = State.Request;
 			if (State.Status == EStatus::Error && CanQuery(Request.Policy, CacheFlags))
 			{
-				NodeRequests.Add_GetRef(Request).UserData = StateIndex;
+				if (Node.DebugOptions.ShouldSimulateGetMiss(Request.Key))
+				{
+					UE_LOGF(LogDerivedDataCache, Verbose, "%ls: Simulated miss for get of %ls from '%ls'",
+						*Node.Name, *WriteToString<96>(Request.Key), *Request.Name);
+				}
+				else
+				{
+					NodeRequests.Add_GetRef(Request).UserData = StateIndex;
+				}
 			}
 			++StateIndex;
 		}
@@ -1641,11 +1685,6 @@ void FCacheStoreHierarchy::LegacyStats(FDerivedDataCacheStatsNode& OutNode)
 	}
 }
 
-bool FCacheStoreHierarchy::LegacyDebugOptions(FBackendDebugOptions& Options)
-{
-	return false;
-}
-
 void FCacheStoreHierarchy::RecordStats(FCacheBucket Bucket, ERequestType Type, ERequestOp Op, EStatus Status, uint64 LogicalReadSize, uint64 LogicalWriteSize)
 {
 	// Accumulate only request count and logical size from the hierarchy.
diff --git a/Engine/Source/Developer/DerivedDataCache/Private/DerivedDataCacheStoreThrottle.cpp b/Engine/Source/Developer/DerivedDataCache/Private/DerivedDataCacheStoreThrottle.cpp
--- a/Engine/Source/Developer/DerivedDataCache/Private/DerivedDataCacheStoreThrottle.cpp
+++ b/Engine/Source/Developer/DerivedDataCache/Private/DerivedDataCacheStoreThrottle.cpp
@@ -131,11 +131,6 @@ public:
 		InnerCache->LegacyStats(OutNode);
 	}
 
-	bool LegacyDebugOptions(FBackendDebugOptions& Options) final
-	{
-		return InnerCache->LegacyDebugOptions(Options);
-	}
-
 private:
 	struct FThrottlingState
 	{
diff --git a/Engine/Source/Developer/DerivedDataCache/Private/DerivedDataCacheStoreVerify.cpp b/Engine/Source/Developer/DerivedDataCache/Private/DerivedDataCacheStoreVerify.cpp
--- a/Engine/Source/Developer/DerivedDataCache/Private/DerivedDataCacheStoreVerify.cpp
+++ b/Engine/Source/Developer/DerivedDataCache/Private/DerivedDataCacheStoreVerify.cpp
@@ -104,11 +104,6 @@ public:
 		InnerCache->LegacyStats(OutNode);
 	}
 
-	bool LegacyDebugOptions(FBackendDebugOptions& Options) final
-	{
-		return InnerCache->LegacyDebugOptions(Options);
-	}
-
 	void GatherVerificationStats(FDerivedDataCacheVerificationStats& OutStats) const final;
 
 private:
diff --git a/Engine/Source/Developer/DerivedDataCache/Private/DerivedDataLegacyCacheStore.h b/Engine/Source/Developer/DerivedDataCache/Private/DerivedDataLegacyCacheStore.h
--- a/Engine/Source/Developer/DerivedDataCache/Private/DerivedDataLegacyCacheStore.h
+++ b/Engine/Source/Developer/DerivedDataCache/Private/DerivedDataLegacyCacheStore.h
@@ -48,8 +48,6 @@ public:
 
 	virtual void LegacyStats(FDerivedDataCacheStatsNode& OutNode) = 0;
 
-	virtual bool LegacyDebugOptions(FBackendDebugOptions& Options) = 0;
-
 	virtual void GatherVerificationStats(FDerivedDataCacheVerificationStats& OutStats) const { }
 };
 
diff --git a/Engine/Source/Developer/DerivedDataCache/Private/FileSystemCacheStore.cpp b/Engine/Source/Developer/DerivedDataCache/Private/FileSystemCacheStore.cpp
--- a/Engine/Source/Developer/DerivedDataCache/Private/FileSystemCacheStore.cpp
+++ b/Engine/Source/Developer/DerivedDataCache/Private/FileSystemCacheStore.cpp
@@ -914,7 +914,6 @@ private:
 	// ILegacyCacheStore
 
 	void LegacyStats(FDerivedDataCacheStatsNode& OutNode) final;
-	bool LegacyDebugOptions(FBackendDebugOptions& Options) final;
 
 private:
 	[[nodiscard]] FOptionalCacheRecord PutCacheRecord(FStringView Name, const FCacheRecord& Record, const FCacheRecordPolicy& Policy, FRequestStats& Stats);
@@ -1006,9 +1005,6 @@ private:
 	/** Access log to write to */
 	TUniquePtr<FAccessLogWriter> AccessLogWriter;
 
-	/** Debug Options */
-	FBackendDebugOptions DebugOptions;
-
 	/** Speed stats */
 	FDerivedDataCacheSpeedStats SpeedStats;
 
@@ -1290,7 +1286,7 @@ FFileSystemCacheStore::FFileSystemCacheStore(
 	Flags |= Params.bDeleteOnly || bReadOnly ? ECacheStoreFlags::None : ECacheStoreFlags::Store;
 	Flags |= bRemote ? ECacheStoreFlags::Remote : ECacheStoreFlags::Local;
 
-	StoreOwner.Add(this, Flags);
+	StoreOwner.Add({Params.CacheName}, this, Flags);
 	if (!Params.bDeleteOnly)
 	{
 		StoreStats = StoreOwner.CreateStats(this, Flags, TEXTVIEW("File System"), Params.CacheName, CachePath);
@@ -1610,12 +1606,6 @@ void FFileSystemCacheStore::LegacyStats(FDerivedDataCacheStatsNode& OutNode)
 	checkNoEntry();
 }
 
-bool FFileSystemCacheStore::LegacyDebugOptions(FBackendDebugOptions& InOptions)
-{
-	DebugOptions = InOptions;
-	return true;
-}
-
 void FFileSystemCacheStore::Put(
 	const TConstArrayView<FCachePutRequest> Requests,
 	IRequestOwner& Owner,
@@ -1899,13 +1889,6 @@ FOptionalCacheRecord FFileSystemCacheStore::PutCacheRecord(
 		return {};
 	}
 
-	if (DebugOptions.ShouldSimulatePutMiss(Key))
-	{
-		UE_LOG(LogDerivedDataCache, Verbose, TEXT("%s: Simulated miss for put of %s from '%.*s'"),
-			*CachePath, *WriteToString<96>(Key), Name.Len(), Name.GetData());
-		return {};
-	}
-
 	TStringBuilder<256> Path;
 	BuildCachePackagePath(Key, Path);
 
@@ -2091,13 +2074,6 @@ FOptionalCacheRecord FFileSystemCacheStore::GetCacheRecordOnly(
 		return FOptionalCacheRecord();
 	}
 
-	if (DebugOptions.ShouldSimulateGetMiss(Key))
-	{
-		UE_LOG(LogDerivedDataCache, Verbose, TEXT("%s: Simulated miss for get of %s from '%.*s'"),
-			*CachePath, *WriteToString<96>(Key), Name.Len(), Name.GetData());
-		return FOptionalCacheRecord();
-	}
-
 	TStringBuilder<256> Path;
 	BuildCachePackagePath(Key, Path);
 
@@ -2228,13 +2204,6 @@ TOptional<FValue> FFileSystemCacheStore::PutCacheValue(
 		return {};
 	}
 
-	if (DebugOptions.ShouldSimulatePutMiss(Key))
-	{
-		UE_LOG(LogDerivedDataCache, Verbose, TEXT("%s: Simulated miss for put of %s from '%.*s'"),
-			*CachePath, *WriteToString<96>(Key), Name.Len(), Name.GetData());
-		return {};
-	}
-
 	// Check if there is an existing value package.
 	FValue ExistingValue;
 	FCbPackage ExistingPackage;
@@ -2397,13 +2366,6 @@ bool FFileSystemCacheStore::GetCacheValueOnly(
 		return false;
 	}
 
-	if (DebugOptions.ShouldSimulateGetMiss(Key))
-	{
-		UE_LOG(LogDerivedDataCache, Verbose, TEXT("%s: Simulated miss for get of %s from '%.*s'"),
-			*CachePath, *WriteToString<96>(Key), Name.Len(), Name.GetData());
-		return false;
-	}
-
 	TStringBuilder<256> Path;
 	BuildCachePackagePath(Key, Path);
 
diff --git a/Engine/Source/Developer/DerivedDataCache/Private/HttpCacheStore.cpp b/Engine/Source/Developer/DerivedDataCache/Private/HttpCacheStore.cpp
--- a/Engine/Source/Developer/DerivedDataCache/Private/HttpCacheStore.cpp
+++ b/Engine/Source/Developer/DerivedDataCache/Private/HttpCacheStore.cpp
@@ -273,7 +273,6 @@ public:
 		FOnCacheGetChunkComplete&& OnComplete) final;
 
 	void LegacyStats(FDerivedDataCacheStatsNode& OutNode) final;
-	bool LegacyDebugOptions(FBackendDebugOptions& Options) final;
 
 	static FHttpCacheStore* GetAny()
 	{
@@ -294,7 +293,7 @@ public:
 	}
 
 private:
-	FString NodeName;
+	FSharedString NodeName;
 	FString Domain;
 	FString Namespace;
 	FString OAuthProvider;
@@ -313,7 +312,6 @@ private:
 	ICacheStoreStats* StoreStats = nullptr;
 
 	FDerivedDataCacheUsageStats UsageStats;
-	FBackendDebugOptions DebugOptions;
 	THttpUniquePtr<IHttpConnectionPool> ConnectionPool;
 	FHttpRequestQueue GetRefRequestQueue;
 	FHttpRequestQueue GetBlobsRequestQueue;
@@ -1236,13 +1234,6 @@ void FHttpCacheStore::FGetRecordOp::GetRecordOnly(const FCacheKey& InKey, const
 		return InOnComplete({FCacheRecord(Key), EStatus::Error});
 	}
 
-	if (CacheStore.DebugOptions.ShouldSimulateGetMiss(Key))
-	{
-		UE_LOG(LogDerivedDataCache, Verbose, TEXT("%s: Simulated miss for get of %s from '%s'"),
-			*CacheStore.NodeName, *WriteToString<96>(Key), *Name);
-		return InOnComplete({FCacheRecord(Key), EStatus::Error});
-	}
-
 	OnRecordComplete = MoveTemp(InOnComplete);
 	RequestStats.Bucket = Key.Bucket;
 
@@ -1918,14 +1909,6 @@ void FHttpCacheStore::TExistsBatchOp<RequestType>::Exists(TConstArrayView<Reques
 			continue;
 		}
 
-		if (CacheStore.DebugOptions.ShouldSimulateGetMiss(Request.Key))
-		{
-			UE_LOG(LogDerivedDataCache, Verbose, TEXT("%s: Simulated miss for get of %s from '%s'"),
-				*CacheStore.NodeName, *WriteToString<96>(Request.Key), *Request.Name);
-			OnComplete(Request.MakeResponse(EStatus::Error));
-			continue;
-		}
-
 		Requests.Emplace(Request);
 	}
 
@@ -2367,7 +2350,7 @@ FHttpCacheStore::FHttpCacheStore(const FHttpCacheStoreParams& Params, ICacheStor
 			Flags |= Params.bWriteOnly ? ECacheStoreFlags::None : ECacheStoreFlags::Query;
 			Flags |= Params.bReadOnly ? ECacheStoreFlags::None : ECacheStoreFlags::Store;
 			TStringBuilder<256> Path(InPlace, ResolvedHost, TEXTVIEW(" ("), Namespace, TEXTVIEW(")"));
-			StoreOwner->Add(this, Flags);
+			StoreOwner->Add(NodeName, this, Flags);
 			StoreStats = StoreOwner->CreateStats(this, Flags, TEXTVIEW("Unreal Cloud DDC"), Params.Name, Path);
 
 			StoreStats->SetAttribute(TEXTVIEW("Domain"), Domain);
@@ -2454,12 +2437,6 @@ FHttpClientParams FHttpCacheStore::GetDefaultClientParams() const
 	return ClientParams;
 }
 
-bool FHttpCacheStore::LegacyDebugOptions(FBackendDebugOptions& InOptions)
-{
-	DebugOptions = InOptions;
-	return true;
-}
-
 bool FHttpCacheStore::AcquireAccessToken(IHttpClient* Client)
 {
 	if (Domain.StartsWith(TEXT("http://localhost")))
@@ -2725,13 +2702,6 @@ void FHttpCacheStore::PutCacheRecordAsync(IRequestOwner& Owner, const FCachePutR
 		return OnComplete(Request.MakeResponse(EStatus::Error));
 	}
 
-	if (DebugOptions.ShouldSimulatePutMiss(Key))
-	{
-		UE_LOG(LogDerivedDataCache, Verbose, TEXT("%s: Simulated miss for put of %s from '%s'"),
-			*NodeName, *WriteToString<96>(Key), *Request.Name);
-		return OnComplete(Request.MakeResponse(EStatus::Error));
-	}
-
 	TRefCountPtr<FPutPackageOp> Op = FPutPackageOp::New(*this, Owner, Request.Name);
 
 	FCbPackage Package;
@@ -2829,13 +2799,6 @@ void FHttpCacheStore::PutCacheValueAsync(IRequestOwner& Owner, const FCachePutVa
 		return OnComplete(Request.MakeResponse(EStatus::Error));
 	}
 
-	if (DebugOptions.ShouldSimulatePutMiss(Request.Key))
-	{
-		UE_LOG(LogDerivedDataCache, Verbose, TEXT("%s: Simulated miss for put of %s from '%s'"),
-			*NodeName, *WriteToString<96>(Request.Key), *Request.Name);
-		return OnComplete(Request.MakeResponse(EStatus::Error));
-	}
-
 	TRefCountPtr<FPutPackageOp> Op = FPutPackageOp::New(*this, Owner, Request.Name);
 
 	FCbPackage Package;
@@ -2909,14 +2872,6 @@ void FHttpCacheStore::GetCacheValueAsync(
 		return;
 	}
 
-	if (DebugOptions.ShouldSimulateGetMiss(Key))
-	{
-		UE_LOG(LogDerivedDataCache, Verbose, TEXT("%s: Simulated miss for get of %s from '%s'"),
-			*NodeName, *WriteToString<96>(Key), *Name);
-		OnComplete({Name, Key, {}, UserData, EStatus::Error});
-		return;
-	}
-
 	TRefCountPtr<FGetValueOp> Op = FGetValueOp::New(*this, Owner, Name);
 
 	FRequestStats& RequestStats = Op->EditStats();
diff --git a/Engine/Source/Developer/DerivedDataCache/Private/MemoryCacheStore.cpp b/Engine/Source/Developer/DerivedDataCache/Private/MemoryCacheStore.cpp
--- a/Engine/Source/Developer/DerivedDataCache/Private/MemoryCacheStore.cpp
+++ b/Engine/Source/Developer/DerivedDataCache/Private/MemoryCacheStore.cpp
@@ -56,7 +56,6 @@ public:
 	// ILegacyCacheStore Interface
 
 	void LegacyStats(FDerivedDataCacheStatsNode& OutNode) final;
-	bool LegacyDebugOptions(FBackendDebugOptions& Options) final;
 
 	// IMemoryCacheStore Interface
 
@@ -73,7 +72,7 @@ private:
 	};
 
 	/** Name of this cache (used for debugging) */
-	FString Name;
+	FSharedString Name;
 
 	ICacheStoreOwner* StoreOwner;
 	ICacheStoreStats* StoreStats = nullptr;
@@ -114,8 +113,6 @@ private:
 	bool bCanBeDisabled = false;
 	bool bShuttingDown  = false;
 
-protected:
-	FBackendDebugOptions DebugOptions;
 };
 
 FMemoryCacheStore::FMemoryCacheStore(
@@ -152,7 +149,7 @@ FMemoryCacheStore::FMemoryCacheStore(
 		{
 			Flags |= ECacheStoreFlags::StopGetStore;
 		}
-		StoreOwner->Add(this, Flags);
+		StoreOwner->Add(Name, this, Flags);
 		if (!FParse::Param(InConfig, TEXT("NoStats")))
 		{
 			StoreStats = StoreOwner->CreateStats(this, Flags, TEXTVIEW("Memory"), Name);
@@ -185,12 +182,6 @@ void FMemoryCacheStore::LegacyStats(FDerivedDataCacheStatsNode& OutNode)
 	checkNoEntry();
 }
 
-bool FMemoryCacheStore::LegacyDebugOptions(FBackendDebugOptions& InOptions)
-{
-	DebugOptions = InOptions;
-	return true;
-}
-
 void FMemoryCacheStore::Put(
 	const TConstArrayView<FCachePutRequest> Requests,
 	IRequestOwner& Owner,
@@ -233,13 +224,6 @@ void FMemoryCacheStore::Put(
 
 		FRequestTimer RequestTimer(RequestStats);
 
-		if (bCanBeDisabled && DebugOptions.ShouldSimulatePutMiss(Key))
-		{
-			UE_LOG(LogDerivedDataCache, Verbose, TEXT("%s: Simulated miss for put of %s from '%s'"),
-				*Name, *WriteToString<96>(Key), *Request.Name);
-			continue;
-		}
-
 		const bool bReplaceExisting = !EnumHasAnyFlags(Request.Policy.GetRecordPolicy(), ECachePolicy::QueryLocal);
 
 		FWriteScopeLock ScopeLock(SynchronizationObject);
@@ -432,12 +416,7 @@ void FMemoryCacheStore::Get(
 
 		FRequestTimer RequestTimer(RequestStats);
 
-		if (bCanBeDisabled && DebugOptions.ShouldSimulateGetMiss(Key))
-		{
-			UE_LOG(LogDerivedDataCache, Verbose, TEXT("%s: Simulated miss for get of %s from '%s'"),
-				*Name, *WriteToString<96>(Key), *Request.Name);
-		}
-		else if (FReadScopeLock ScopeLock(SynchronizationObject); const FCacheRecordComponents* FoundComponents = CacheRecords.Find(Key))
+		if (FReadScopeLock ScopeLock(SynchronizationObject); const FCacheRecordComponents* FoundComponents = CacheRecords.Find(Key))
 		{
 			if (!FoundComponents->Values.IsEmpty())
 			{
@@ -544,13 +523,6 @@ void FMemoryCacheStore::PutValue(
 
 		FRequestTimer RequestTimer(RequestStats);
 
-		if (bCanBeDisabled && DebugOptions.ShouldSimulatePutMiss(Key))
-		{
-			UE_LOG(LogDerivedDataCache, Verbose, TEXT("%s: Simulated miss for put of %s from '%s'"),
-				*Name, *WriteToString<96>(Key), *Request.Name);
-			continue;
-		}
-
 		const bool bReplaceExisting = !EnumHasAnyFlags(Request.Policy, ECachePolicy::QueryLocal);
 
 		FWriteScopeLock ScopeLock(SynchronizationObject);
@@ -634,12 +606,7 @@ void FMemoryCacheStore::GetValue(
 
 		FRequestTimer RequestTimer(RequestStats);
 
-		if (bCanBeDisabled && DebugOptions.ShouldSimulateGetMiss(Key))
-		{
-			UE_LOG(LogDerivedDataCache, Verbose, TEXT("%s: Simulated miss for get of %s from '%s'"),
-				*Name, *WriteToString<96>(Key), *Request.Name);
-		}
-		else if (bExistsOnly && bCanBeDisabled)
+		if (bExistsOnly && bCanBeDisabled)
 		{
 		}
 		else if (FReadScopeLock ScopeLock(SynchronizationObject); const FValue* CacheValue = CacheValues.Find(Key))
@@ -701,12 +668,7 @@ void FMemoryCacheStore::GetChunks(
 
 		bool bProcessHit = false;
 		const bool bExistsOnly = EnumHasAnyFlags(Request.Policy, ECachePolicy::SkipData);
-		if (bCanBeDisabled && DebugOptions.ShouldSimulateGetMiss(Request.Key))
-		{
-			UE_LOG(LogDerivedDataCache, Verbose, TEXT("%s: Simulated miss for get of %s from '%s'"),
-				*Name, *WriteToString<96>(Request.Key, '/', Request.Id), *Request.Name);
-		}
-		else if (bHasValue && (ValueKey == Request.Key) && (ValueId == Request.Id) && (bExistsOnly || Reader.HasSource()))
+		if (bHasValue && (ValueKey == Request.Key) && (ValueId == Request.Id) && (bExistsOnly || Reader.HasSource()))
 		{
 			// Value matches the request.
 			bProcessHit = true;
diff --git a/Engine/Source/Developer/DerivedDataCache/Private/PakFileCacheStore.cpp b/Engine/Source/Developer/DerivedDataCache/Private/PakFileCacheStore.cpp
--- a/Engine/Source/Developer/DerivedDataCache/Private/PakFileCacheStore.cpp
+++ b/Engine/Source/Developer/DerivedDataCache/Private/PakFileCacheStore.cpp
@@ -109,7 +109,6 @@ public:
 	// ILegacyCacheStore
 
 	void LegacyStats(FDerivedDataCacheStatsNode& OutNode) final;
-	bool LegacyDebugOptions(FBackendDebugOptions& Options) final;
 
 private:
 	[[nodiscard]] bool PutCacheRecord(FStringView Name, const FCacheRecord& Record, const FCacheRecordPolicy& Policy, FRequestStats& Stats);
@@ -159,7 +158,6 @@ private:
 
 private:
 	FDerivedDataCacheUsageStats UsageStats;
-	FBackendDebugOptions DebugOptions;
 
 	struct FCacheValue
 	{
@@ -248,7 +246,7 @@ FPakFileCacheStore::FPakFileCacheStore(const TCHAR* InName, const TCHAR* const I
 	if (StoreOwner)
 	{
 		const ECacheStoreFlags Flags = ECacheStoreFlags::Local | (bWriting ? ECacheStoreFlags::Store : ECacheStoreFlags::Query | ECacheStoreFlags::StopStore);
-		StoreOwner->Add(this, Flags);
+		StoreOwner->Add({InName}, this, Flags);
 		StoreStats = StoreOwner->CreateStats(this, Flags, TEXT("Pak File"), InName, CachePath);
 	}
 }
@@ -638,12 +636,6 @@ void FPakFileCacheStore::LegacyStats(FDerivedDataCacheStatsNode& OutNode)
 	checkNoEntry();
 }
 
-bool FPakFileCacheStore::LegacyDebugOptions(FBackendDebugOptions& InOptions)
-{
-	DebugOptions = InOptions;
-	return true;
-}
-
 void FPakFileCacheStore::Put(
 	const TConstArrayView<FCachePutRequest> Requests,
 	IRequestOwner& Owner,
@@ -898,13 +890,6 @@ bool FPakFileCacheStore::PutCacheRecord(
 		return false;
 	}
 
-	if (DebugOptions.ShouldSimulatePutMiss(Key))
-	{
-		UE_LOG(LogDerivedDataCache, Verbose, TEXT("%s: Simulated miss for put of %s from '%.*s'"),
-			*CachePath, *WriteToString<96>(Key), Name.Len(), Name.GetData());
-		return false;
-	}
-
 	TStringBuilder<256> Path;
 	FPathViews::Append(Path, TEXT("Buckets"), Key);
 
@@ -972,13 +957,6 @@ FOptionalCacheRecord FPakFileCacheStore::GetCacheRecordOnly(
 		return FOptionalCacheRecord();
 	}
 
-	if (DebugOptions.ShouldSimulateGetMiss(Key))
-	{
-		UE_LOG(LogDerivedDataCache, Verbose, TEXT("%s: Simulated miss for get of %s from '%.*s'"),
-			*CachePath, *WriteToString<96>(Key), Name.Len(), Name.GetData());
-		return FOptionalCacheRecord();
-	}
-
 	TStringBuilder<256> Path;
 	FPathViews::Append(Path, TEXT("Buckets"), Key);
 
@@ -1089,13 +1067,6 @@ bool FPakFileCacheStore::PutCacheValue(
 		return false;
 	}
 
-	if (DebugOptions.ShouldSimulatePutMiss(Key))
-	{
-		UE_LOG(LogDerivedDataCache, Verbose, TEXT("%s: Simulated miss for put of %s from '%.*s'"),
-			*CachePath, *WriteToString<96>(Key), Name.Len(), Name.GetData());
-		return false;
-	}
-
 	// Check if there is an existing value package.
 	bool bValueExists = false;
 	TStringBuilder<256> Path;
@@ -1169,13 +1140,6 @@ bool FPakFileCacheStore::GetCacheValueOnly(
 		return false;
 	}
 
-	if (DebugOptions.ShouldSimulateGetMiss(Key))
-	{
-		UE_LOG(LogDerivedDataCache, Verbose, TEXT("%s: Simulated miss for get of %s from '%.*s'"),
-			*CachePath, *WriteToString<96>(Key), Name.Len(), Name.GetData());
-		return false;
-	}
-
 	TStringBuilder<256> Path;
 	FPathViews::Append(Path, TEXT("Buckets"), Key);
 
diff --git a/Engine/Source/Developer/DerivedDataCache/Private/S3CacheStore.cpp b/Engine/Source/Developer/DerivedDataCache/Private/S3CacheStore.cpp
--- a/Engine/Source/Developer/DerivedDataCache/Private/S3CacheStore.cpp
+++ b/Engine/Source/Developer/DerivedDataCache/Private/S3CacheStore.cpp
@@ -239,7 +239,6 @@ public:
 	// ILegacyCacheStore
 
 	void LegacyStats(FDerivedDataCacheStatsNode& OutNode) final;
-	bool LegacyDebugOptions(FBackendDebugOptions& Options) final;
 
 private:
 	struct FBundle;
@@ -311,7 +310,6 @@ private:
 	void RemoveUnusedBundles();
 	void ReadBundle(FBundle& Bundle);
 
-	FBackendDebugOptions DebugOptions;
 };
 
 /**
@@ -1025,7 +1023,7 @@ FS3CacheStore::FS3CacheStore(const TCHAR* InName, const TCHAR* InRootManifestPat
 				bEnabled = true;
 
 				constexpr ECacheStoreFlags Flags = ECacheStoreFlags::Local | ECacheStoreFlags::Query | ECacheStoreFlags::StopStore;
-				StoreOwner.Add(this, Flags);
+				StoreOwner.Add({InName}, this, Flags);
 				StoreStats = StoreOwner.CreateStats(this, Flags, TEXTVIEW("S3"), InName, BaseUrl);
 			}
 
@@ -1199,12 +1197,6 @@ void FS3CacheStore::ReadBundle(FBundle& Bundle)
 	}
 }
 
-bool FS3CacheStore::LegacyDebugOptions(FBackendDebugOptions& InOptions)
-{
-	DebugOptions = InOptions;
-	return true;
-}
-
 FOptionalCacheRecord FS3CacheStore::GetCacheRecordOnly(
 	const FStringView Name,
 	const FCacheKey& Key,
@@ -1227,13 +1219,6 @@ FOptionalCacheRecord FS3CacheStore::GetCacheRecordOnly(
 		return FOptionalCacheRecord();
 	}
 
-	if (DebugOptions.ShouldSimulateGetMiss(Key))
-	{
-		UE_LOG(LogDerivedDataCache, Verbose, TEXT("%s: Simulated miss for get of %s from '%.*s'"),
-			*GetName(), *WriteToString<96>(Key), Name.Len(), Name.GetData());
-		return FOptionalCacheRecord();
-	}
-
 	TStringBuilder<256> Path;
 	BuildCachePackagePath(Key, Path);
 
@@ -1337,13 +1322,6 @@ bool FS3CacheStore::GetCacheValueOnly(
 		return false;
 	}
 
-	if (DebugOptions.ShouldSimulateGetMiss(Key))
-	{
-		UE_LOG(LogDerivedDataCache, Verbose, TEXT("%s: Simulated miss for get of %s from '%.*s'"),
-			*GetName(), *WriteToString<96>(Key), Name.Len(), Name.GetData());
-		return false;
-	}
-
 	TStringBuilder<256> Path;
 	BuildCachePackagePath(Key, Path);
 
diff --git a/Engine/Source/Developer/DerivedDataCache/Private/ZenCacheStore.cpp b/Engine/Source/Developer/DerivedDataCache/Private/ZenCacheStore.cpp
--- a/Engine/Source/Developer/DerivedDataCache/Private/ZenCacheStore.cpp
+++ b/Engine/Source/Developer/DerivedDataCache/Private/ZenCacheStore.cpp
@@ -157,7 +157,7 @@ public:
 
 	~FZenCacheStore() final;
 
-	inline const FString& GetName() const { return NodeName; }
+	inline const FSharedString& GetName() const { return NodeName; }
 
 	/**
 	 * Checks if cache service is usable (reachable and accessible).
@@ -201,7 +201,6 @@ public:
 	// ILegacyCacheStore
 
 	void LegacyStats(FDerivedDataCacheStatsNode& OutNode) final;
-	bool LegacyDebugOptions(FBackendDebugOptions& Options) final;
 
 	const Zen::FZenServiceInstance& GetServiceInstance() const { return ZenService.GetInstance(); }
 
@@ -270,7 +269,7 @@ private:
 
 	class FRequestAsyncCbPackageReceiver;
 
-	FString NodeName;
+	FSharedString NodeName;
 	FString Namespace;
 	UE::Zen::FScopeZenService ZenService;
 	ICacheStoreOwner* StoreOwner = nullptr;
@@ -284,7 +283,6 @@ private:
 	int32 MaxBatchPutKB = 1024;
 	int32 CacheRecordBatchSize = 8;
 	int32 CacheChunksBatchSize = 8;
-	FBackendDebugOptions DebugOptions;
 	FAnsiStringView RpcUri = ANSITEXTVIEW("/z$/$rpc");
 	std::atomic<int64> LastPerformanceEvaluationTicks;
 	TOptional<FThread> PerformanceEvaluationThread;
@@ -513,14 +511,7 @@ public:
 						RequestWithStats.Stats.Latency = FMonotonicTimeSpan::Infinity();
 
 						const FCacheKey& Key = RequestWithStats.Request.Record.GetKey();
-						bool bPutSucceeded = ResponseField.AsBool();
-
-						if (CacheStore.DebugOptions.ShouldSimulatePutMiss(Key))
-						{
-							UE_LOG(LogDerivedDataCache, Verbose, TEXT("%s: Simulated miss for put of %s from '%s'"),
-								*CacheStore.GetName(), *WriteToString<96>(Key), *RequestWithStats.Request.Name);
-							bPutSucceeded = false;
-						}
+						const bool bPutSucceeded = ResponseField.AsBool();
 
 						bPutSucceeded ? OnHit(RequestWithStats, DetailsFieldIt.AsObject()) : OnMiss(RequestWithStats, DetailsFieldIt.AsObject());
 						++DetailsFieldIt;
@@ -786,12 +777,7 @@ public:
 						const FCacheKey& Key = RequestWithStats.Request.Key;
 						FOptionalCacheRecord Record;
 
-						if (CacheStore.DebugOptions.ShouldSimulateGetMiss(Key))
-						{
-							UE_LOG(LogDerivedDataCache, Verbose, TEXT("%s: Simulated miss for get of '%s' from '%s'"),
-								*CacheStore.GetName(), *WriteToString<96>(Key), *RequestWithStats.Request.Name);
-						}
-						else if (!RecordField.IsNull())
+						if (!RecordField.IsNull())
 						{
 							Record = FCacheRecord::Load(Response, RecordField.AsObject());
 						}
@@ -1018,13 +1004,7 @@ public:
 						// Latency can't be measured for Put operations because it is intertwined with upload time.
 						RequestWithStats.Stats.Latency = FMonotonicTimeSpan::Infinity();
 
-						bool bPutSucceeded = ResponseField.AsBool();
-						if (CacheStore.DebugOptions.ShouldSimulatePutMiss(RequestWithStats.Request.Key))
-						{
-							UE_LOG(LogDerivedDataCache, Verbose, TEXT("%s: Simulated miss for PutValue of %s from '%s'"),
-								*CacheStore.GetName(), *WriteToString<96>(RequestWithStats.Request.Key), *RequestWithStats.Request.Name);
-							bPutSucceeded = false;
-						}
+						const bool bPutSucceeded = ResponseField.AsBool();
 						bPutSucceeded ? OnHit(RequestWithStats, DetailsFieldIt.AsObjectView()) : OnMiss(RequestWithStats, DetailsFieldIt.AsObjectView());
 						++DetailsFieldIt;
 					}
@@ -1248,12 +1228,7 @@ public:
 
 						FCbObjectView ResultObj = ResultField.AsObjectView();
 						TOptional<FValue> Value;
-						if (CacheStore.DebugOptions.ShouldSimulateGetMiss(RequestWithStats.Request.Key))
-						{
-							UE_LOG(LogDerivedDataCache, Verbose, TEXT("%s: Simulated miss for GetValue of '%s' from '%s'"),
-								*CacheStore.GetName(), *WriteToString<96>(RequestWithStats.Request.Key), *RequestWithStats.Request.Name);
-						}
-						else if (FIoHash RawHash = ResultObj[ANSITEXTVIEW("RawHash")].AsHash(); !RawHash.IsZero())
+						if (FIoHash RawHash = ResultObj[ANSITEXTVIEW("RawHash")].AsHash(); !RawHash.IsZero())
 						{
 							if (EnumHasAnyFlags(RequestWithStats.Request.Policy, ECachePolicy::SkipData))
 							{
@@ -1482,12 +1457,7 @@ public:
 						uint64 RawSize = 0;
 						FSharedBuffer RequestedBytes;
 						const FCacheGetChunkRequest& Request = RequestWithStats.Request;
-						if (CacheStore.DebugOptions.ShouldSimulateGetMiss(Request.Key))
-						{
-							UE_LOG(LogDerivedDataCache, Verbose, TEXT("%s: Simulated miss for get of '%s' from '%s'"),
-								*CacheStore.GetName(), *WriteToString<96>(Request.Key, '/', Request.Id), *Request.Name);
-						}
-						else if (RawHash = ResultView[ANSITEXTVIEW("RawHash")].AsHash(); !RawHash.IsZero())
+						if (RawHash = ResultView[ANSITEXTVIEW("RawHash")].AsHash(); !RawHash.IsZero())
 						{
 							if (EnumHasAnyFlags(RequestWithStats.Request.Policy, ECachePolicy::SkipData))
 							{
@@ -2012,7 +1982,7 @@ void FZenCacheStore::Initialize(const FZenCacheStoreParams& Params)
 			bDeactivatedForPerformance.store(true, std::memory_order_relaxed);
 		}
 
-		StoreOwner->Add(this, Flags);
+		StoreOwner->Add(NodeName, this, Flags);
 		TStringBuilder<256> Path(InPlace, ZenService.GetInstance().GetPath(), TEXTVIEW(" ("), Namespace, TEXTVIEW(")"));
 		StoreStats = StoreOwner->CreateStats(this, Flags, TEXT("Zen"), *Params.Name, Path);
 		bTryEvaluatePerformance = !GIsBuildMachine && (StoreStats != nullptr) && (DeactivateAtMs > 0.0f);
@@ -2276,12 +2246,6 @@ void FZenCacheStore::LegacyStats(FDerivedDataCacheStatsNode& OutNode)
 	checkNoEntry();
 }
 
-bool FZenCacheStore::LegacyDebugOptions(FBackendDebugOptions& InOptions)
-{
-	DebugOptions = InOptions;
-	return true;
-}
-
 void FZenCacheStore::Put(
 	const TConstArrayView<FCachePutRequest> Requests,
 	IRequestOwner& Owner,
diff --git a/Engine/Source/Developer/DerivedDataCache/Tests/DerivedDataCacheStoreHierarchyTest.cpp b/Engine/Source/Developer/DerivedDataCache/Tests/DerivedDataCacheStoreHierarchyTest.cpp
--- a/Engine/Source/Developer/DerivedDataCache/Tests/DerivedDataCacheStoreHierarchyTest.cpp
+++ b/Engine/Source/Developer/DerivedDataCache/Tests/DerivedDataCacheStoreHierarchyTest.cpp
@@ -26,11 +26,11 @@ TEST_CASE("DerivedData::Cache::Hierarchy::PartialRecordPropagation", "[DerivedDa
 	TUniquePtr<ILegacyCacheStore> Hierarchy(CreateCacheStoreHierarchy(StoreOwner, [](IMemoryCacheStore*& OutCache) { OutCache = nullptr; }));
 
 	TUniquePtr<ITestCacheStore> Store0(CreateTestCacheStore(ECacheStoreFlags::Local | ECacheStoreFlags::Query | ECacheStoreFlags::Store, ETestCacheStoreFlags::Async));
-	StoreOwner->Add(Store0.Get(), Store0->GetFlags());
+	StoreOwner->Add({TEXT("Store0")}, Store0.Get(), Store0->GetFlags());
 	TUniquePtr<ITestCacheStore> Store1(CreateTestCacheStore(ECacheStoreFlags::Local | ECacheStoreFlags::Query | ECacheStoreFlags::Store, ETestCacheStoreFlags::Async));
-	StoreOwner->Add(Store1.Get(), Store1->GetFlags());
+	StoreOwner->Add({TEXT("Store1")}, Store1.Get(), Store1->GetFlags());
 	TUniquePtr<ITestCacheStore> Store2(CreateTestCacheStore(ECacheStoreFlags::Local | ECacheStoreFlags::Query | ECacheStoreFlags::Store, ETestCacheStoreFlags::Async));
-	StoreOwner->Add(Store2.Get(), Store2->GetFlags());
+	StoreOwner->Add({TEXT("Store2")}, Store2.Get(), Store2->GetFlags());
 
 	const FValueId Id0 = FValueId::FromName("0");
 	const FValueId Id1 = FValueId::FromName("1");
@@ -142,11 +142,11 @@ TEST_CASE("DerivedData::Cache::Hierarchy::PartialNonDeterministicRecordPropagati
 	TUniquePtr<ILegacyCacheStore> Hierarchy(CreateCacheStoreHierarchy(StoreOwner, [](IMemoryCacheStore*& OutCache) { OutCache = nullptr; }));
 
 	TUniquePtr<ITestCacheStore> Store0(CreateTestCacheStore(ECacheStoreFlags::Local | ECacheStoreFlags::Query | ECacheStoreFlags::Store, ETestCacheStoreFlags::Async));
-	StoreOwner->Add(Store0.Get(), Store0->GetFlags());
+	StoreOwner->Add({TEXT("Store0")}, Store0.Get(), Store0->GetFlags());
 	TUniquePtr<ITestCacheStore> Store1(CreateTestCacheStore(ECacheStoreFlags::Local | ECacheStoreFlags::Query | ECacheStoreFlags::Store, ETestCacheStoreFlags::Async));
-	StoreOwner->Add(Store1.Get(), Store1->GetFlags());
+	StoreOwner->Add({TEXT("Store1")}, Store1.Get(), Store1->GetFlags());
 	TUniquePtr<ITestCacheStore> Store2(CreateTestCacheStore(ECacheStoreFlags::Local | ECacheStoreFlags::Query | ECacheStoreFlags::Store, ETestCacheStoreFlags::Async));
-	StoreOwner->Add(Store2.Get(), Store2->GetFlags());
+	StoreOwner->Add({TEXT("Store2")}, Store2.Get(), Store2->GetFlags());
 
 	const FValueId Id0 = FValueId::FromName("0");
 	const FValueId Id1 = FValueId::FromName("1");
@@ -291,9 +291,9 @@ TEST_CASE("DerivedData::Cache::Hierarchy::BlockingRequestOwner", "[DerivedData]"
 	TUniquePtr<ILegacyCacheStore> Hierarchy(CreateCacheStoreHierarchy(StoreOwner, [](IMemoryCacheStore*& OutCache) { OutCache = nullptr; }));
 
 	TUniquePtr<ITestCacheStore> Store0(CreateTestCacheStore(ECacheStoreFlags::Local | ECacheStoreFlags::Query | ECacheStoreFlags::Store, ETestCacheStoreFlags::Async | ETestCacheStoreFlags::Wait));
-	StoreOwner->Add(Store0.Get(), Store0->GetFlags());
+	StoreOwner->Add({TEXT("Store0")}, Store0.Get(), Store0->GetFlags());
 	TUniquePtr<ITestCacheStore> Store1(CreateTestCacheStore(ECacheStoreFlags::Local | ECacheStoreFlags::Query | ECacheStoreFlags::Store, ETestCacheStoreFlags::Async));
-	StoreOwner->Add(Store1.Get(), Store1->GetFlags());
+	StoreOwner->Add({TEXT("Store1")}, Store1.Get(), Store1->GetFlags());
 
 	const FCacheKey Key{FCacheBucket("Test"), FIoHash::HashBuffer(MakeMemoryView<uint8>({0, 1, 2, 3}))};
 
diff --git a/Engine/Source/Developer/DerivedDataCache/Tests/TestCacheStore.cpp b/Engine/Source/Developer/DerivedDataCache/Tests/TestCacheStore.cpp
--- a/Engine/Source/Developer/DerivedDataCache/Tests/TestCacheStore.cpp
+++ b/Engine/Source/Developer/DerivedDataCache/Tests/TestCacheStore.cpp
@@ -68,7 +68,6 @@ private:
 		FOnCacheGetChunkComplete&& OnComplete) final;
 
 	void LegacyStats(FDerivedDataCacheStatsNode& OutNode) final {}
-	bool LegacyDebugOptions(FBackendDebugOptions& Options) final { return false; }
 
 	void ExecuteOrQueue(IRequestOwner& Owner, TUniqueFunction<void (bool bCancel)>&& Function);
 
