DDC: Fixed parsing of cache store debug options from -DDC-<Name>-Miss...

#rnx
#Synced-CL 50967057

diff --git a/Engine/Source/Developer/DerivedDataCache/Private/DerivedDataBackendInterface.h b/Engine/Source/Developer/DerivedDataCache/Private/DerivedDataBackendInterface.h
--- a/Engine/Source/Developer/DerivedDataCache/Private/DerivedDataBackendInterface.h
+++ b/Engine/Source/Developer/DerivedDataCache/Private/DerivedDataBackendInterface.h
@@ -11,7 +11,6 @@
 #include "DerivedDataLegacyCacheStore.h"
 #include "HAL/LowLevelMemTracker.h"
 #include "Stats/Stats.h"
-#include "Templates/DontCopy.h"
 #include "Templates/PimplPtr.h"
 
 class FDerivedDataCacheUsageStats;
@@ -58,7 +57,7 @@ struct FBackendDebugOptions
 	FCacheKeyFilter SimulateMissFilter;
 
 	/** State for simulated misses. */
-	TDontCopy<TPimplPtr<Private::FBackendDebugMissState>> SimulateMissState;
+	TPimplPtr<Private::FBackendDebugMissState> SimulateMissState;
 
 	FBackendDebugOptions();
 
diff --git a/Engine/Source/Developer/DerivedDataCache/Private/DerivedDataBackends.cpp b/Engine/Source/Developer/DerivedDataCache/Private/DerivedDataBackends.cpp
--- a/Engine/Source/Developer/DerivedDataCache/Private/DerivedDataBackends.cpp
+++ b/Engine/Source/Developer/DerivedDataCache/Private/DerivedDataBackends.cpp
@@ -1075,7 +1075,7 @@ struct Private::FBackendDebugMissState
 
 FBackendDebugOptions::FBackendDebugOptions()
 {
-	SimulateMissState.Get() = MakePimpl<Private::FBackendDebugMissState>();
+	SimulateMissState = MakePimpl<Private::FBackendDebugMissState>();
 }
 
 /**
