DDC: Fixed a crash when recovering from a failure to create a non-default graph containing a pak store

#Synced-CL 50967057

diff --git a/Engine/Source/Developer/DerivedDataCache/Private/DerivedDataCacheStoreHierarchy.cpp b/Engine/Source/Developer/DerivedDataCache/Private/DerivedDataCacheStoreHierarchy.cpp
--- a/Engine/Source/Developer/DerivedDataCache/Private/DerivedDataCacheStoreHierarchy.cpp
+++ b/Engine/Source/Developer/DerivedDataCache/Private/DerivedDataCacheStoreHierarchy.cpp
@@ -291,6 +291,9 @@ FCacheStoreHierarchy::~FCacheStoreHierarchy()
 		if (Node.AsyncCache != Node.Cache)
 		{
 			delete Node.AsyncCache;
+			// Deleting Node.Cache below can access Node.AsyncCache if it calls RemoveNotSafe.
+			// Assign it here as if this was a native-async node.
+			Node.AsyncCache = Node.Cache;
 		}
 		delete Node.Cache;
 
