Change jupiter DDC client to use the ref-specific blob endpoint instead of the older and global compressed-blobs endpoint.

#Synced-CL 51529686

diff --git a/Engine/Source/Developer/DerivedDataCache/Private/HttpCacheStore.cpp b/Engine/Source/Developer/DerivedDataCache/Private/HttpCacheStore.cpp
--- a/Engine/Source/Developer/DerivedDataCache/Private/HttpCacheStore.cpp
+++ b/Engine/Source/Developer/DerivedDataCache/Private/HttpCacheStore.cpp
@@ -90,6 +90,12 @@ TRACE_DECLARE_ATOMIC_INT_COUNTER(HttpDDC_PutHit, TEXT("HttpDDC Put Hit"));
 TRACE_DECLARE_ATOMIC_INT_COUNTER(HttpDDC_BytesReceived, TEXT("HttpDDC Bytes Received"));
 TRACE_DECLARE_ATOMIC_INT_COUNTER(HttpDDC_BytesSent, TEXT("HttpDDC Bytes Sent"));
 
+static TAutoConsoleVariable<bool> GDerivedDataCacheHttpUseOldBlobEndpoints(
+	TEXT("DDC.Http.UseOldBlobEndpoints"),
+	false,
+	TEXT("Set to true to use old endpoints in CloudDDC. This will be removed in a future engine version."),
+	ECVF_Default);
+
 static FAnsiStringView GetDomainFromUri(const FAnsiStringView Uri)
 {
 	FAnsiStringView Domain = Uri;
@@ -1065,7 +1071,14 @@ void FHttpCacheStore::FPutPackageOp::BeginPutBlobs(FCbPackage&& Package, FCacheP
 			}
 
 			FHttpOperation& LocalOperation = *Operation;
-			LocalOperation.SetUri(WriteToAnsiString<256>(Self->CacheStore.EffectiveDomain, ANSITEXTVIEW("/api/v1/compressed-blobs/"), Self->CacheStore.Namespace, '/', Blob.GetRawHash()));
+			if (GDerivedDataCacheHttpUseOldBlobEndpoints.GetValueOnAnyThread())
+			{
+				LocalOperation.SetUri(WriteToAnsiString<256>(Self->CacheStore.EffectiveDomain, ANSITEXTVIEW("/api/v1/compressed-blobs/"), Self->CacheStore.Namespace, '/', Blob.GetRawHash()));
+			}
+			else 
+			{
+				LocalOperation.SetUri(WriteToAnsiString<256>(Self->CacheStore.EffectiveDomain, ANSITEXTVIEW("/api/v1/refs/"), Self->CacheStore.Namespace, '/', Self->Key.Bucket, '/', Self->Key.Hash, ANSITEXTVIEW("/blobs/"), Blob.GetRawHash()));
+			}
 			LocalOperation.SetUnixSocketPath(StringCast<ANSICHAR>(*Self->CacheStore.UnixSocketPath));
 			LocalOperation.SetMethod(EHttpMethod::Put);
 			LocalOperation.SetContentType(EHttpMediaType::CompressedBinary);
@@ -1514,7 +1527,14 @@ void FHttpCacheStore::FGetRecordOp::BeginGetValue(
 	}
 
 	FHttpOperation& LocalOperation = *Operation;
-	LocalOperation.SetUri(WriteToAnsiString<256>(CacheStore.EffectiveDomain, ANSITEXTVIEW("/api/v1/compressed-blobs/"), CacheStore.Namespace, '/', Value.GetRawHash()));
+	if (GDerivedDataCacheHttpUseOldBlobEndpoints.GetValueOnAnyThread())
+	{
+		LocalOperation.SetUri(WriteToAnsiString<256>(CacheStore.EffectiveDomain, ANSITEXTVIEW("/api/v1/compressed-blobs/"), CacheStore.Namespace, '/', Value.GetRawHash()));
+	}
+	else 
+	{
+		LocalOperation.SetUri(WriteToAnsiString<256>(CacheStore.EffectiveDomain, ANSITEXTVIEW("/api/v1/refs/"), CacheStore.Namespace, '/', Key.Bucket, '/', Key.Hash, ANSITEXTVIEW("/blobs/"), Value.GetRawHash()));
+	}
 	LocalOperation.SetUnixSocketPath(StringCast<ANSICHAR>(*CacheStore.UnixSocketPath));
 	LocalOperation.SetMethod(EHttpMethod::Get);
 	LocalOperation.AddAcceptType(EHttpMediaType::Any);
