Use the namespace scoped version of the rpc endpoint when communicating with zenserver's cache API.

#Synced-CL 51529686

diff --git a/Engine/Source/Developer/DerivedDataCache/Private/ZenCacheStore.cpp b/Engine/Source/Developer/DerivedDataCache/Private/ZenCacheStore.cpp
--- a/Engine/Source/Developer/DerivedDataCache/Private/ZenCacheStore.cpp
+++ b/Engine/Source/Developer/DerivedDataCache/Private/ZenCacheStore.cpp
@@ -11,6 +11,7 @@
 #include "Async/ManualResetEvent.h"
 #include "Async/UniqueLock.h"
 #include "BatchView.h"
+#include "Containers/AnsiString.h"
 #include "DerivedDataBackendInterface.h"
 #include "DerivedDataCachePrivate.h"
 #include "DerivedDataCacheRecord.h"
@@ -283,7 +284,7 @@ private:
 	int32 MaxBatchPutKB = 1024;
 	int32 CacheRecordBatchSize = 8;
 	int32 CacheChunksBatchSize = 8;
-	FAnsiStringView RpcUri = ANSITEXTVIEW("/z$/$rpc");
+	FAnsiString RpcUri;
 	std::atomic<int64> LastPerformanceEvaluationTicks;
 	TOptional<FThread> PerformanceEvaluationThread;
 	FManualResetEvent PerformanceEvaluationThreadShutdownEvent;
@@ -1858,7 +1859,8 @@ void FZenCacheStore::Initialize(const FZenCacheStoreParams& Params)
 	NodeName = Params.Name;
 	LastPerformanceEvaluationTicks.store(FDateTime::UtcNow().GetTicks(), std::memory_order_relaxed);
 	LastStorageSizeUpdateTicks.store(FDateTime::UtcNow().GetTicks(), std::memory_order_relaxed);
-	Namespace = Params.Namespace;
+	Namespace = Params.Namespace.ToLower();
+	RpcUri = WriteToAnsiString<64>(ANSITEXTVIEW("/z$/"), Namespace, ANSITEXTVIEW("/$rpc"));
 
 	const uint32 MaxConnections = uint32(FMath::Clamp(FPlatformMisc::NumberOfCoresIncludingHyperthreads(), 8, 64));
 	constexpr uint32 RequestPoolSize = 128;
