Help users fix invalid DDC path instead of crashing

When opening the editor, if the local DDC cache path is invalid, we used to log a fatal error and crash. The user had to
A) manually edit the Windows registry to edit/remove the `HKEY_CURRENT_USER/Software/Epic Games/GlobalDataCache/UE-LocalDataCachePath` entry
or
B) launch the editor with `-DDC-ForceMemoryCache` which is not so trivial.

otherwise, the editor will always crash on launch.

When the DDC path is invalid, a message box will appear to automatically reset the local data cache path. If the user answers yes, the registry entry will be removed and the editor will restart.

#Synced-CL 51753244

diff --git a/Engine/Source/Developer/DerivedDataCache/Private/DerivedDataBackends.cpp b/Engine/Source/Developer/DerivedDataCache/Private/DerivedDataBackends.cpp
--- a/Engine/Source/Developer/DerivedDataCache/Private/DerivedDataBackends.cpp
+++ b/Engine/Source/Developer/DerivedDataCache/Private/DerivedDataBackends.cpp
@@ -904,14 +904,49 @@ public:
 			}
 			else
 			{
-				UE_CLOGFMT(!FParse::Param(FCommandLine::Get(), TEXT("DDC-AllowNoActiveStores")),
-					LogDerivedDataCache, Fatal,
+				const bool bAllowsNoActiveStore = FParse::Param(FCommandLine::Get(), TEXT("DDC-AllowNoActiveStores"));
+				const ANSICHAR* MissingTypes = 
+					Owner->HasAllFlags(ECacheStoreFlags::Query) ? "writable" :
+					Owner->HasAllFlags(ECacheStoreFlags::Store) ? "readable" : "readable or writable";
+#if WITH_EDITOR
+				if (!FApp::IsUnattended() && !IsRunningCommandlet() && !GIsRunningUnattendedScript)
+				{
+					const FString StoreId = TEXT("Epic Games");
+					const FString SectionName = TEXT("GlobalDataCachePath");
+					const FString KeyName = TEXT("UE-LocalDataCachePath");
+
+					//only offer to reset local data cache path if it is already set
+					if (FString LocalDataCachePath; FPlatformMisc::GetStoredValue(StoreId, SectionName, KeyName, LocalDataCachePath))
+					{
+						const FText ErrorTitle = LOCTEXT("DDC_NoActiveStore_ErrorMsgTitle", "Derived Data Cache Error");
+						const FText ErrorMessage = FText::Format(LOCTEXT("DDC_NoActiveStore_ErrorMsgText",
+							"Unable to use default cache graph '{0}'.\n\n"
+							"The local cache path is currently overridden by the setting\nUE-LocalDataCachePath: {1}\n\n"
+							"Do you want to reset this setting and restart the editor?"), 
+							FText::FromString(DefaultGraphName), FText::FromString(LocalDataCachePath));
+						const EAppReturnType::Type Result = FPlatformMisc::MessageBoxExt(EAppMsgType::YesNo, *ErrorMessage.ToString(), *ErrorTitle.ToString());
+						if (Result == EAppReturnType::Yes)
+						{
+							UE_CLOGFMT(!bAllowsNoActiveStore, LogDerivedDataCache, Error,
+								"Unable to use default cache graph '{Graph}' because there are no {MissingTypes} nodes available."
+								"'UE-LocalDataCachePath' setting will be reset.",
+								DefaultGraphName, MissingTypes);
+
+							FPlatformMisc::DeleteStoredValue(StoreId, SectionName, KeyName);
+							FPlatformProcess::CreateProc(FPlatformProcess::ExecutablePath(), FCommandLine::GetOriginal(),
+								/*bLaunchDetached*/true, /*bLaunchHidden*/false, /*bLaunchReallyHidden*/false,
+								/*OutProcessId*/nullptr, /*PriorityModifier*/0, /*OptionalWorkingDirectory*/nullptr, /*PipeWriteChild*/nullptr);
+							FPlatformMisc::RequestExit(/*Force*/false);
+							return;
+						}
+					}
+				}
+#endif
+				UE_CLOGFMT(!bAllowsNoActiveStore, LogDerivedDataCache, Fatal,
 					"Unable to use default cache graph '{Graph}' because there are no {MissingTypes} nodes available. "
 					"Add -DDC-ForceMemoryCache to the command line to bypass this if you need access "
 					"to the editor settings to fix the cache configuration.",
-					DefaultGraphName,
-					Owner->HasAllFlags(ECacheStoreFlags::Query) ? "writable" :
-					Owner->HasAllFlags(ECacheStoreFlags::Store) ? "readable" : "readable or writable");
+					DefaultGraphName, MissingTypes);
 			}
 		}
 
