DDC: Send the "X-Jupiter-Force-Consistent-Read: true" header to Jupiter for keys that have been overwritten in this session

#Synced-CL 50967057

diff --git a/Engine/Source/Developer/DerivedDataCache/Private/HttpCacheStore.cpp b/Engine/Source/Developer/DerivedDataCache/Private/HttpCacheStore.cpp
--- a/Engine/Source/Developer/DerivedDataCache/Private/HttpCacheStore.cpp
+++ b/Engine/Source/Developer/DerivedDataCache/Private/HttpCacheStore.cpp
@@ -6,6 +6,7 @@
 #if WITH_HTTP_DDC_BACKEND
 
 #include "Algo/AllOf.h"
+#include "Algo/AnyOf.h"
 #include "Algo/Compare.h"
 #include "Algo/NoneOf.h"
 #include "Algo/Transform.h"
@@ -339,6 +340,9 @@ private:
 	bool bWriteOnly = false;
 	bool bBypassProxy = false;
 
+	mutable FMutex OverwriteKeysMutex;
+	TSet<FCacheKey> OverwriteKeys;
+
 	static inline FHttpCacheStore* AnyInstance = nullptr;
 
 	FHttpClientParams GetDefaultClientParams() const;
@@ -1267,6 +1271,10 @@ void FHttpCacheStore::FGetRecordOp::GetRecordOnly(const FCacheKey& InKey, const
 		LocalOperation.SetUnixSocketPath(StringCast<ANSICHAR>(*Self->CacheStore.UnixSocketPath));
 		LocalOperation.SetMethod(EHttpMethod::Get);
 		LocalOperation.AddAcceptType(EHttpMediaType::CbObject);
+		if (TUniqueLock Lock(Self->CacheStore.OverwriteKeysMutex); Self->CacheStore.OverwriteKeys.Contains(Self->Key))
+		{
+			LocalOperation.AddHeader(ANSITEXTVIEW("X-Jupiter-Force-Consistent-Read"), ANSITEXTVIEW("true"));
+		}
 		LocalOperation.SetExpectedStatusCodes({404});
 
 		LocalOperation.SendAsync(Self->Owner, [Self, Operation = MoveTemp(Operation)]() mutable
@@ -1732,6 +1740,10 @@ void FHttpCacheStore::FGetValueOp::BeginGetRef(TUniquePtr<FHttpOperation>&& Oper
 	{
 		LocalOperation.AddHeader(ANSITEXTVIEW("Accept"), ANSITEXTVIEW("application/x-jupiter-inline"));
 	}
+	if (TUniqueLock Lock(CacheStore.OverwriteKeysMutex); CacheStore.OverwriteKeys.Contains(Key))
+	{
+		LocalOperation.AddHeader(ANSITEXTVIEW("X-Jupiter-Force-Consistent-Read"), ANSITEXTVIEW("true"));
+	}
 	LocalOperation.SetExpectedStatusCodes({404});
 
 	RequestTimer.Stop();
@@ -1990,6 +2002,15 @@ void FHttpCacheStore::TExistsBatchOp<RequestType>::BeginExists(TUniquePtr<FHttpO
 	LocalOperation.SetContentType(EHttpMediaType::CbObject);
 	LocalOperation.AddAcceptType(EHttpMediaType::CbObject);
 	LocalOperation.SetBody(FCompositeBuffer(Body.GetOuterBuffer()));
+	const auto HasOverwriteKey = [this](const RequestType& Request)
+	{
+		TUniqueLock Lock(CacheStore.OverwriteKeysMutex);
+		return CacheStore.OverwriteKeys.Contains(Request.Key);
+	};
+	if (Algo::AnyOf(Requests, HasOverwriteKey))
+	{
+		LocalOperation.AddHeader(ANSITEXTVIEW("X-Jupiter-Force-Consistent-Read"), ANSITEXTVIEW("true"));
+	}
 
 	RequestTimer.Stop();
 	LocalOperation.SendAsync(Owner, [Self = TRefCountPtr(this), Operation = MoveTemp(Operation)]() mutable
@@ -2735,6 +2756,13 @@ void FHttpCacheStore::PutCacheRecordAsync(IRequestOwner& Owner, const FCachePutR
 		Package = Request.Record.Save();
 	}
 
+	if (!EnumHasAnyFlags(RecordPolicy, ECachePolicy::QueryRemote))
+	{
+		// Track keys put with an overwrite policy to force consistent read on subsequent gets.
+		TUniqueLock Lock(OverwriteKeysMutex);
+		OverwriteKeys.Add(Key);
+	}
+
 	Op->Put(Key, Request.Policy, MoveTemp(Package), [this, Op, Request, OnComplete = MoveTemp(OnComplete)](FPutPackageOp::FResponse&& OpResponse) mutable
 	{
 		if (OpResponse.Status == EStatus::Ok)
@@ -2835,6 +2863,13 @@ void FHttpCacheStore::PutCacheValueAsync(IRequestOwner& Owner, const FCachePutVa
 		Package.SetObject(Writer.Save().AsObject());
 	}
 
+	if (!EnumHasAnyFlags(Request.Policy, ECachePolicy::QueryRemote))
+	{
+		// Track keys put with an overwrite policy to force consistent read on subsequent gets.
+		TUniqueLock Lock(OverwriteKeysMutex);
+		OverwriteKeys.Add(Request.Key);
+	}
+
 	Op->Put(Request.Key, Request.Policy, MoveTemp(Package), [this, Op, Request, OnComplete = MoveTemp(OnComplete)](FPutPackageOp::FResponse&& OpResponse) mutable
 	{
 		TRACE_COUNTER_ADD(HttpDDC_BytesReceived, Op->ReadStats().PhysicalReadSize);
diff --git a/Engine/Source/Developer/DerivedDataCache/Tests/DerivedDataCacheStoreTest.cpp b/Engine/Source/Developer/DerivedDataCache/Tests/DerivedDataCacheStoreTest.cpp
--- a/Engine/Source/Developer/DerivedDataCache/Tests/DerivedDataCacheStoreTest.cpp
+++ b/Engine/Source/Developer/DerivedDataCache/Tests/DerivedDataCacheStoreTest.cpp
@@ -6,6 +6,7 @@
 #include "DerivedDataRequestOwner.h"
 #include "DerivedDataPrivate.h"
 #include "HAL/FileManager.h"
+#include "HAL/PlatformProcess.h"
 #include "Logging/StructuredLog.h"
 #include "Misc/Guid.h"
 #include "Misc/Paths.h"
@@ -62,12 +63,16 @@ static FCacheBucket GetBucket()
 	return Bucket;
 }
 
-static FCacheKey CreateKeyFromGuid(const ANSICHAR* GuidStr, int32 Index = 0)
+static FCacheKey CreateKeyFromGuid(const ANSICHAR* GuidStr, int32 Index = 0, FAnsiStringView Extra = {})
 {
 	const FGuid Guid(GuidStr);
 	FIoHashBuilder HashBuilder;
 	HashBuilder.Update(&Guid, sizeof(FGuid));
 	HashBuilder.Update(&Index, sizeof(int32));
+	if (!Extra.IsEmpty())
+	{
+		HashBuilder.Update(Extra.GetData(), Extra.Len() * sizeof(ANSICHAR));
+	}
 	return {GetBucket(), HashBuilder.Finalize()};
 }
 
@@ -1114,6 +1119,7 @@ enum class ETestCacheDeterminismFlags : uint32
 {
 	None = 0,
 	SkipOverwrite = 1 << 0,
+	UniqueByComputerName = 1 << 1,
 };
 
 ENUM_CLASS_FLAGS(ETestCacheDeterminismFlags);
@@ -1121,6 +1127,9 @@ ENUM_CLASS_FLAGS(ETestCacheDeterminismFlags);
 static void TestCacheDeterminism(ICache& Cache, ETestCacheDeterminismFlags Flags = ETestCacheDeterminismFlags::None)
 {
 	const FValue Value(MakeSequentialIntegerBuffer(8 * 1024 * 1024));
+	// A new guid is the most reliable random value that the engine can generate at this time.
+	// Add [1, 1024] bytes to create an OtherValue that tends to differ between tests.
+	const int32 OtherValueExtraSize = 1 + (FGuid::NewGuid().D & 1023);
 	const FValue OtherValue(MakeSequentialIntegerBuffer(8 * 1024 * 1024 + 1024));
 
 	const FValueId ValueId = FValueId::FromName(ANSITEXTVIEW("Value"));
@@ -1129,11 +1138,12 @@ static void TestCacheDeterminism(ICache& Cache, ETestCacheDeterminismFlags Flags
 	const FCbObject Meta = CreateMeta("Value");
 	const FCbObject OtherMeta = CreateMeta("OtherValue");
 
-	const FCacheKey RecordKey = CreateKeyFromGuid("8660fd50-a31c-4ea1-ace8-eb25b21b6db6");
+	const FCacheKey RecordKey = CreateKeyFromGuid("8660fd50-a31c-4ea1-ace8-eb25b21b6db6", 0,
+		EnumHasAnyFlags(Flags, ETestCacheDeterminismFlags::UniqueByComputerName) ? WriteToAnsiString<64>(FPlatformProcess::ComputerName()) : FAnsiStringView{});
 
-	// Put Original Record w/ Overwrite
 	const auto PutOriginalRecordWithOverwrite = [&]
 	{
+		INFO("Put Original Record w/ Overwrite");
 		FCacheRecordBuilder RecordBuilder(RecordKey);
 		RecordBuilder.AddValue(ValueId, Value);
 		RecordBuilder.AddValue(AttachmentId, Value);
@@ -1152,8 +1162,8 @@ static void TestCacheDeterminism(ICache& Cache, ETestCacheDeterminismFlags Flags
 	};
 	PutOriginalRecordWithOverwrite();
 
-	// Put Original Record w/ New Meta
 	{
+		INFO("Put Original Record w/ New Meta");
 		FCacheRecordBuilder RecordBuilder(RecordKey);
 		RecordBuilder.AddValue(ValueId, Value);
 		RecordBuilder.AddValue(AttachmentId, Value);
@@ -1169,8 +1179,8 @@ static void TestCacheDeterminism(ICache& Cache, ETestCacheDeterminismFlags Flags
 		BlockingOwner.Wait();
 	}
 
-	// Put Different Record
 	{
+		INFO("Put Different Record");
 		FCacheRecordBuilder RecordBuilder(RecordKey);
 		RecordBuilder.AddValue(ValueId, OtherValue);
 		RecordBuilder.AddValue(AttachmentId, Value);
@@ -1185,8 +1195,8 @@ static void TestCacheDeterminism(ICache& Cache, ETestCacheDeterminismFlags Flags
 		BlockingOwner.Wait();
 	}
 
-	// Put Different Record w/ NonDeterministic Flag
 	{
+		INFO("Put Different Record w/ NonDeterministic Flag");
 		FCacheRecordBuilder RecordBuilder(RecordKey);
 		RecordBuilder.AddValue(ValueId, OtherValue);
 		RecordBuilder.AddValue(AttachmentId, Value);
@@ -1203,8 +1213,8 @@ static void TestCacheDeterminism(ICache& Cache, ETestCacheDeterminismFlags Flags
 		BlockingOwner.Wait();
 	}
 
-	// Get Original Record
 	{
+		INFO("Get Original Record");
 		FRequestOwner BlockingOwner(EPriority::Blocking);
 		Cache.Get({{{TEXTVIEW("OriginalRecord")}, RecordKey}}, BlockingOwner, [&](FCacheGetResponse&& Response)
 		{
@@ -1215,9 +1225,9 @@ static void TestCacheDeterminism(ICache& Cache, ETestCacheDeterminismFlags Flags
 		BlockingOwner.Wait();
 	}
 
-	// Put Different Record w/ Overwrite
 	if (!EnumHasAnyFlags(Flags, ETestCacheDeterminismFlags::SkipOverwrite))
 	{
+		INFO("Put Different Record w/ Overwrite");
 		FCacheRecordBuilder RecordBuilder(RecordKey);
 		RecordBuilder.AddValue(ValueId, OtherValue);
 		RecordBuilder.AddValue(AttachmentId, Value);
@@ -1233,27 +1243,77 @@ static void TestCacheDeterminism(ICache& Cache, ETestCacheDeterminismFlags Flags
 		BlockingOwner.Wait();
 	}
 
-	// Get Different Record
 	if (!EnumHasAnyFlags(Flags, ETestCacheDeterminismFlags::SkipOverwrite))
 	{
+		INFO("Get Different Record");
 		FRequestOwner BlockingOwner(EPriority::Blocking);
 		Cache.Get({{{TEXTVIEW("DifferentRecord")}, RecordKey}}, BlockingOwner, [&](FCacheGetResponse&& Response)
 		{
 			CHECK(Response.Status == EStatus::Ok);
+			CHECK(Response.Record.GetMeta().Equals(OtherMeta));
 			CHECK(HasValueWithData(Response.Record, FValueWithId(ValueId, OtherValue)));
+			CHECK(HasValueWithData(Response.Record, FValueWithId(AttachmentId, Value)));
+		});
+		BlockingOwner.Wait();
+	}
+
+	if (!EnumHasAnyFlags(Flags, ETestCacheDeterminismFlags::SkipOverwrite))
+	{
+		INFO("Get Different Record Hash + Size");
+		FRequestOwner BlockingOwner(EPriority::Blocking);
+		Cache.Get({{{TEXTVIEW("DifferentRecord")}, RecordKey, ECachePolicy::Default | ECachePolicy::SkipData}}, BlockingOwner, [&](FCacheGetResponse&& Response)
+		{
+			CHECK(Response.Status == EStatus::Ok);
 			CHECK(Response.Record.GetMeta().Equals(OtherMeta));
+			CHECK(HasValueWithNoData(Response.Record, FValueWithId(ValueId, OtherValue)));
+			CHECK(HasValueWithNoData(Response.Record, FValueWithId(AttachmentId, Value)));
+		});
+		BlockingOwner.Wait();
+	}
+
+	if (!EnumHasAnyFlags(Flags, ETestCacheDeterminismFlags::SkipOverwrite))
+	{
+		INFO("Get Different Record Chunk");
+		FRequestOwner BlockingOwner(EPriority::Blocking);
+		Cache.GetChunks({{{TEXTVIEW("DifferentRecordChunk")}, RecordKey, ValueId}}, BlockingOwner, [&](FCacheGetChunkResponse&& Response)
+		{
+			CHECK(Response.Status == EStatus::Ok);
+			CHECK(Response.Key == RecordKey);
+			CHECK(Response.Id == ValueId);
+			CHECK(Response.RawOffset == 0);
+			CHECK(Response.RawSize == OtherValue.GetRawSize());
+			CHECK(Response.RawHash == OtherValue.GetRawHash());
+			CHECK(Response.RawData.GetSize() == Response.RawSize);
+			CHECK(FIoHash::HashBuffer(Response.RawData) == Response.RawHash);
+		});
+		BlockingOwner.Wait();
+	}
+
+	if (!EnumHasAnyFlags(Flags, ETestCacheDeterminismFlags::SkipOverwrite))
+	{
+		INFO("Put Different Record w/o Overwrite");
+		FCacheRecordBuilder RecordBuilder(RecordKey);
+		RecordBuilder.AddValue(ValueId, OtherValue);
+		RecordBuilder.AddValue(AttachmentId, Value);
+		RecordBuilder.SetMeta(CopyTemp(OtherMeta));
+		FRequestOwner BlockingOwner(EPriority::Blocking);
+		Cache.Put({{{TEXTVIEW("DifferentRecord")}, RecordBuilder.Build(), ECachePolicy::Default}}, BlockingOwner, [&](FCachePutResponse&& Response)
+		{
+			CHECK(Response.Status == EStatus::Ok);
+			CHECK(Response.Record.GetMeta().Equals(OtherMeta));
+			CHECK(HasValueWithData(Response.Record, FValueWithId(ValueId, OtherValue)));
+			CHECK(HasValueWithData(Response.Record, FValueWithId(AttachmentId, Value)));
 		});
 		BlockingOwner.Wait();
 	}
 
-	// Put Original Record w/ Overwrite
 	PutOriginalRecordWithOverwrite();
 
 	const FCacheKey ValueKey = CreateKeyFromGuid("926ae1e6-8ef8-4efa-8f15-83d6c0468909");
 
-	// Put Original Value w/ Overwrite
 	const auto PutOriginalValueWithOverwrite = [&]
 	{
+		INFO("Put Original Value w/ Overwrite");
 		FRequestOwner BlockingOwner(EPriority::Blocking);
 		Cache.PutValue({{{TEXTVIEW("OriginalValue")}, ValueKey, Value, ECachePolicy::Store | ECachePolicy::SkipData | ECachePolicy::NonDeterministic}}, BlockingOwner, [&](FCachePutValueResponse&& Response)
 		{
@@ -1265,8 +1325,8 @@ static void TestCacheDeterminism(ICache& Cache, ETestCacheDeterminismFlags Flags
 	};
 	PutOriginalValueWithOverwrite();
 
-	// Put Original Value
 	{
+		INFO("Put Original Value");
 		FRequestOwner BlockingOwner(EPriority::Blocking);
 		Cache.PutValue({{{TEXTVIEW("OriginalValueCopy")}, ValueKey, Value}}, BlockingOwner, [&](FCachePutValueResponse&& Response)
 		{
@@ -1277,8 +1337,8 @@ static void TestCacheDeterminism(ICache& Cache, ETestCacheDeterminismFlags Flags
 		BlockingOwner.Wait();
 	}
 
-	// Put Different Value
 	{
+		INFO("Put Different Value");
 		FRequestOwner BlockingOwner(EPriority::Blocking);
 		Cache.PutValue({{{TEXTVIEW("DifferentValue")}, ValueKey, OtherValue}}, BlockingOwner, [&](FCachePutValueResponse&& Response)
 		{
@@ -1289,8 +1349,8 @@ static void TestCacheDeterminism(ICache& Cache, ETestCacheDeterminismFlags Flags
 		BlockingOwner.Wait();
 	}
 
-	// Put Different Value w/ NonDeterministic Flag
 	{
+		INFO("Put Different Value w/ NonDeterministic Flag");
 		FRequestOwner BlockingOwner(EPriority::Blocking);
 		Cache.PutValue({{{TEXTVIEW("DifferentValueIgnore")}, ValueKey, OtherValue, ECachePolicy::Default | ECachePolicy::NonDeterministic}}, BlockingOwner, [&](FCachePutValueResponse&& Response)
 		{
@@ -1301,8 +1361,8 @@ static void TestCacheDeterminism(ICache& Cache, ETestCacheDeterminismFlags Flags
 		BlockingOwner.Wait();
 	}
 
-	// Get Original Value
 	{
+		INFO("Get Original Value");
 		FRequestOwner BlockingOwner(EPriority::Blocking);
 		Cache.GetValue({{{TEXTVIEW("OriginalValue")}, ValueKey}}, BlockingOwner, [&](FCacheGetValueResponse&& Response)
 		{
@@ -1312,9 +1372,9 @@ static void TestCacheDeterminism(ICache& Cache, ETestCacheDeterminismFlags Flags
 		BlockingOwner.Wait();
 	}
 
-	// Put Different Value w/ Overwrite
 	if (!EnumHasAnyFlags(Flags, ETestCacheDeterminismFlags::SkipOverwrite))
 	{
+		INFO("Put Different Value w/ Overwrite");
 		FRequestOwner BlockingOwner(EPriority::Blocking);
 		Cache.PutValue({{{TEXTVIEW("DifferentValueOverwrite")}, ValueKey, OtherValue, ECachePolicy::Store}}, BlockingOwner, [&](FCachePutValueResponse&& Response)
 		{
@@ -1325,9 +1385,9 @@ static void TestCacheDeterminism(ICache& Cache, ETestCacheDeterminismFlags Flags
 		BlockingOwner.Wait();
 	}
 
-	// Get Different Value
 	if (!EnumHasAnyFlags(Flags, ETestCacheDeterminismFlags::SkipOverwrite))
 	{
+		INFO("Get Different Value");
 		FRequestOwner BlockingOwner(EPriority::Blocking);
 		Cache.GetValue({{{TEXTVIEW("DifferentValue")}, ValueKey}}, BlockingOwner, [&OtherValue](FCacheGetValueResponse&& Response)
 		{
@@ -1337,7 +1397,49 @@ static void TestCacheDeterminism(ICache& Cache, ETestCacheDeterminismFlags Flags
 		BlockingOwner.Wait();
 	}
 
-	// Put Original Value w/ Overwrite
+	if (!EnumHasAnyFlags(Flags, ETestCacheDeterminismFlags::SkipOverwrite))
+	{
+		INFO("Get Different Value Hash + Size");
+		FRequestOwner BlockingOwner(EPriority::Blocking);
+		Cache.GetValue({{{TEXTVIEW("DifferentValue")}, ValueKey, ECachePolicy::Default | ECachePolicy::SkipData}}, BlockingOwner, [&OtherValue](FCacheGetValueResponse&& Response)
+		{
+			CHECK(Response.Status == EStatus::Ok);
+			CHECK(Response.Value.GetRawHash() == OtherValue.GetRawHash());
+		});
+		BlockingOwner.Wait();
+	}
+
+	if (!EnumHasAnyFlags(Flags, ETestCacheDeterminismFlags::SkipOverwrite))
+	{
+		INFO("Get Different Value Chunk");
+		FRequestOwner BlockingOwner(EPriority::Blocking);
+		Cache.GetChunks({{{TEXTVIEW("DifferentValueChunk")}, ValueKey}}, BlockingOwner, [&](FCacheGetChunkResponse&& Response)
+		{
+			CHECK(Response.Status == EStatus::Ok);
+			CHECK(Response.Key == ValueKey);
+			CHECK(Response.Id.IsNull());
+			CHECK(Response.RawOffset == 0);
+			CHECK(Response.RawSize == OtherValue.GetRawSize());
+			CHECK(Response.RawHash == OtherValue.GetRawHash());
+			CHECK(Response.RawData.GetSize() == Response.RawSize);
+			CHECK(FIoHash::HashBuffer(Response.RawData) == Response.RawHash);
+		});
+		BlockingOwner.Wait();
+	}
+
+	if (!EnumHasAnyFlags(Flags, ETestCacheDeterminismFlags::SkipOverwrite))
+	{
+		INFO("Put Different Value w/o Overwrite");
+		FRequestOwner BlockingOwner(EPriority::Blocking);
+		Cache.PutValue({{{TEXTVIEW("DifferentValue")}, ValueKey, OtherValue, ECachePolicy::Default}}, BlockingOwner, [&](FCachePutValueResponse&& Response)
+		{
+			CHECK(Response.Status == EStatus::Ok);
+			CHECK(Response.Value == OtherValue);
+			CHECK(Response.Value.HasData());
+		});
+		BlockingOwner.Wait();
+	}
+
 	PutOriginalValueWithOverwrite();
 }
 
@@ -1409,6 +1511,7 @@ TEST_CASE("DerivedData::CacheStore", "[DerivedData]")
 		{
 			TestCacheValues(*Cache, ETestCacheValuesFlags::SkipPutMissingData);
 			TestCacheRecords(*Cache, ETestCacheRecordsFlags::SkipPutMissingData);
+			// Skip the overwrite tests because they are not reliable due to the replication delay for the refs.
 			TestCacheDeterminism(*Cache, ETestCacheDeterminismFlags::SkipOverwrite);
 		}
 		else
@@ -1418,6 +1521,22 @@ TEST_CASE("DerivedData::CacheStore", "[DerivedData]")
 	}
 }
 
+// Disabled by default because the determinism tests are not reliable due to the replication delay for the refs.
+// Everything here is tested by the preceding test case except for the overwrites in the determinism test.
+TEST_CASE("DerivedData::CacheStore::Cloud", "[.][DerivedData]")
+{
+	using namespace CacheStoreTest;
+
+	FStringView Config = TEXTVIEW("(Cloud)");
+	TUniquePtr<ICache> Cache(CreateCache(Config));
+	CHECKED_IF(Cache)
+	{
+		TestCacheValues(*Cache, ETestCacheValuesFlags::SkipPutMissingData);
+		TestCacheRecords(*Cache, ETestCacheRecordsFlags::SkipPutMissingData);
+		TestCacheDeterminism(*Cache, ETestCacheDeterminismFlags::UniqueByComputerName);
+	}
+}
+
 } // UE::DerivedData
 
 #endif // WITH_LOW_LEVEL_TESTS
