Use a config parameter to control use of the namespace scoped version of the rpc endpoint when communicating with zenserver's cache API.  Default to false for now.

#Synced-CL 51897070

diff --git a/Engine/Source/Developer/DerivedDataCache/Private/ZenCacheStore.cpp b/Engine/Source/Developer/DerivedDataCache/Private/ZenCacheStore.cpp
--- a/Engine/Source/Developer/DerivedDataCache/Private/ZenCacheStore.cpp
+++ b/Engine/Source/Developer/DerivedDataCache/Private/ZenCacheStore.cpp
@@ -102,6 +102,7 @@ struct FZenCacheStoreParams
 	bool bReadOnly = false;
 	bool bWriteOnly = false;
 	bool bBypassProxy = true;
+	bool bNamespacedRpcEndpoint = false;
 
 	void Parse(const TCHAR* NodeName, const TCHAR* Config);
 };
@@ -1860,7 +1861,14 @@ void FZenCacheStore::Initialize(const FZenCacheStoreParams& Params)
 	LastPerformanceEvaluationTicks.store(FDateTime::UtcNow().GetTicks(), std::memory_order_relaxed);
 	LastStorageSizeUpdateTicks.store(FDateTime::UtcNow().GetTicks(), std::memory_order_relaxed);
 	Namespace = Params.Namespace.ToLower();
-	RpcUri = WriteToAnsiString<64>(ANSITEXTVIEW("/z$/"), Namespace, ANSITEXTVIEW("/$rpc"));
+	if (Params.bNamespacedRpcEndpoint)
+	{
+		RpcUri = WriteToAnsiString<64>(ANSITEXTVIEW("/z$/"), Namespace, ANSITEXTVIEW("/$rpc"));
+	}
+	else
+	{
+		RpcUri = ANSITEXTVIEW("/z$/$rpc");
+	}
 
 	const uint32 MaxConnections = uint32(FMath::Clamp(FPlatformMisc::NumberOfCoresIncludingHyperthreads(), 8, 64));
 	constexpr uint32 RequestPoolSize = 128;
@@ -2333,6 +2341,7 @@ void FZenCacheStoreParams::Parse(const TCHAR* NodeName, const TCHAR* Config)
 	FParse::Value(Config, TEXT("Sandbox="), Sandbox);
 	FParse::Bool(Config, TEXT("Flush="), bFlush);
 	FParse::Bool(Config, TEXT("BypassProxy="), bBypassProxy);
+	FParse::Bool(Config, TEXT("NamespacedRpcEndpoint="), bNamespacedRpcEndpoint);
 
 	// Performance deactivation
 	FParse::Value(Config, TEXT("DeactivateAt="), DeactivateAtMs);
