DDC: Improvements to cache graph selection

- Consider every -DDC=... command line argument as a candidate for construction.
- Remove broken -DDC=None functionality because the build workers no longer require it.
- Fix DDC.Graph CVar overriding the installed default by making it empty by default.
- Improve reporting of why a cache graph was unable to be used before trying other candidates.
- Added -DDC-NoDefaultGraph command line argument to remove the default graph from consideration.

#Synced-CL 50967057

diff --git a/Engine/Source/Developer/DerivedDataCache/Private/DerivedDataBackends.cpp b/Engine/Source/Developer/DerivedDataCache/Private/DerivedDataBackends.cpp
--- a/Engine/Source/Developer/DerivedDataCache/Private/DerivedDataBackends.cpp
+++ b/Engine/Source/Developer/DerivedDataCache/Private/DerivedDataBackends.cpp
@@ -50,7 +50,7 @@ LLM_DEFINE_TAG(DDCBackend);
 
 static TAutoConsoleVariable<FString> GDerivedDataCacheGraphName(
 	TEXT("DDC.Graph"),
-	TEXT("Default"),
+	TEXT(""),
 	TEXT("Name or config of the graph to use for the Derived Data Cache."),
 	ECVF_ReadOnly);
 
@@ -106,22 +106,6 @@ protected:
 	bool CreateGraph(FStringView GraphNameOrConfig)
 	{
 		ActiveGraphName = GraphNameOrConfig;
-
-#if IS_PROGRAM
-		// A DDC graph of "None" is used by build worker programs that use the DDC build code paths but avoid use of the DDC cache
-		// code paths. Unfortunately the cache must currently be instantiated as part of initializing the build code, so we have
-		// to disable the cache portion by injecting "-DDC=None" to the commandline args during process startup. This mode is not
-		// compatible with use in the editor/commandlets which is not written to operate with a non-functional cache layer. This
-		// can lead to confusion when people attempt to use "-DDC=None" in the editor expecting it to behave like "-DDC=Cold".
-		// To avoid this confusion (and until we can use the build without the cache) it is restricted to use in programs only
-		// and not the editor.
-		if (GraphNameOrConfig == TEXTVIEW("None"))
-		{
-			UE_LOG(LogDerivedDataCache, Display, TEXT("Requested cache graph of 'None'. Every cache operation will fail."));
-			return true;
-		}
-#endif
-
 		FParsedNodeMap ParsedNodes;
 		TGuardValue ParsedNodeMapGuard(ActiveParsedNodeMap, &ParsedNodes);
 		return ParseGraph(ActiveGraphName, ParsedNodes) && Hierarchy->HasAllFlags(ECacheStoreFlags::Query | ECacheStoreFlags::Store);
@@ -304,20 +288,12 @@ private:
 		}
 	}
 
-	/**
-	 * Creates Hierarchical interface from ini settings.
-	 *
-	 * @param NodeName Node name.
-	 * @param Entry Node definition.
-	 * @param IniFilename ini filename.
-	 * @param IniSection ini section containing graph definition
-	 * @param InParsedNodes map of nodes and their names which have already been parsed
-	 * @return Hierarchical backend interface instance or nullptr if unsuccessful
-	 */
 	bool ParseGraph(FStringView GraphNameOrConfig, FParsedNodeMap& InParsedNodes)
 	{
 		TRACE_CPUPROFILER_EVENT_SCOPE(FDerivedDataBackendGraph::ParseGraph);
 
+		UE_LOGFMT(LogDerivedDataCache, Verbose, "Parsing cache graph '{Graph}'...", GraphNameOrConfig);
+
 		RootCache = CreateCacheStoreHierarchy(Hierarchy, [this](IMemoryCacheStore*& OutCache) { GetMemoryCache(OutCache); });
 
 		TStringBuilder<256> GraphConfig;
@@ -857,82 +833,115 @@ public:
 		check(!StaticGraph);
 		StaticGraph = this;
 
-		FString GraphName;
-		if (!FParse::Value(FCommandLine::Get(), TEXT("-DDC="), GraphName, /*bShouldStopOnSeparator*/ false))
+		TArray<FString> CandidateGraphs;
+
+		// Gather graphs from the command line as candidates.
 		{
-			GraphName = GDerivedDataCacheGraphName.GetValueOnGameThread();
+			FString GraphNameOrConfig;
+			for (const TCHAR* Remaining = FCommandLine::Get(); FParse::Value(Remaining, TEXT("-DDC="), GraphNameOrConfig, /*bShouldStopOnSeparator*/ false, &Remaining);)
+			{
+				CandidateGraphs.Add(MoveTemp(GraphNameOrConfig));
+			}
 		}
 
-		const TCHAR* const DefaultGraphName = FApp::IsEngineInstalled() ? TEXT("Installed") : TEXT("Default");
-		const TCHAR* const LegacyDefaultGraphName = FApp::IsEngineInstalled() ? TEXT("InstalledDerivedDataBackendGraph") : TEXT("DerivedDataBackendGraph");
+		// Add the graph from the CVar as a candidate.
+		if (FString GraphNameOrConfig = GDerivedDataCacheGraphName.GetValueOnGameThread(); !GraphNameOrConfig.IsEmpty())
+		{
+			CandidateGraphs.Add(MoveTemp(GraphNameOrConfig));
+		}
 
-		// Try to use the graph from DDC.Graph or -DDC=...
-		if (!GraphName.IsEmpty() && GraphName != DefaultGraphName && !CreateGraph(GraphName))
+		// Add the default graph as a candidate unless it has been suppressed.
+		if (!FParse::Param(FCommandLine::Get(), TEXT("DDC-NoDefaultGraph")))
 		{
-			ResetGraph();
-			UE_LOGFMT(LogDerivedDataCache, Warning,
-				"Failed to create graph '{Graph}'. Attempting to create the default graph...", GraphName);
+			// UE_DEPRECATED(5.8, "Remove the legacy default block after the standard deprecation period.")
+			const TCHAR* const LegacyDefaultGraphName = FApp::IsEngineInstalled() ? TEXT("InstalledDerivedDataBackendGraph") : TEXT("DerivedDataBackendGraph");
+			if (GConfig->DoesSectionExist(LegacyDefaultGraphName, GEngineIni))
+			{
+				CandidateGraphs.Add(LegacyDefaultGraphName);
+			}
+
+			const TCHAR* const DefaultGraphName = FApp::IsEngineInstalled() ? TEXT("Installed") : TEXT("Default");
+			CandidateGraphs.Add(DefaultGraphName);
 		}
 
-		// UE_DEPRECATED(5.8, "Remove the legacy default block after the standard deprecation period.")
-		// Try to use the legacy default graph.
-		if (!HasRoot() && GConfig->DoesSectionExist(LegacyDefaultGraphName, GEngineIni) && !CreateGraph(LegacyDefaultGraphName))
+		// Construct candidate graphs in the order they were added until one is valid.
+		bool bCreatedGraph = false;
+		FStringView GraphNameOrConfig;
+		for (const FString& CandidateGraph : CandidateGraphs)
 		{
-			ResetGraph();
+			const FStringView FailedGraphNameOrConfig = GraphNameOrConfig;
+			GraphNameOrConfig = CandidateGraph;
+
+			if (!FailedGraphNameOrConfig.IsEmpty())
+			{
+				ICacheStoreOwner* Owner = GetOwner();
+				const ANSICHAR* MissingTypes = 
+					Owner->HasAllFlags(ECacheStoreFlags::Query) ? "writable" :
+					Owner->HasAllFlags(ECacheStoreFlags::Store) ? "readable" : "readable or writable";
+				UE_LOGFMT(LogDerivedDataCache, Warning,
+					"Failed to create graph '{Graph}' because it has no {MissingTypes} nodes available. "
+					"Attempting to create graph '{NextGraph}'...",
+					FailedGraphNameOrConfig, MissingTypes, GraphNameOrConfig);
+				ResetGraph();
+			}
+
+			bCreatedGraph = CreateGraph(GraphNameOrConfig);
+			if (bCreatedGraph)
+			{
+				break;
+			}
 		}
 
-		// Try to use the default graph.
-		if (!HasRoot() && !CreateGraph(DefaultGraphName))
+		if (!bCreatedGraph)
 		{
 			ICacheStoreOwner* Owner = GetOwner();
 			TStringBuilder<256> GraphConfig;
-			if (!TryFindCacheGraphConfig(DefaultGraphName, GraphConfig) &&
-				!TryFindCacheGraphConfig(LegacyDefaultGraphName, GraphConfig))
+			if (!TryFindCacheGraphConfig(GraphNameOrConfig, GraphConfig))
 			{
 				UE_LOGFMT(LogDerivedDataCache, Fatal,
-					"Failed to create default cache graph '{Graph}' because its config is missing in the '{Config}' config.",
-					DefaultGraphName, GEngineIni);
+					"Failed to create cache graph '{Graph}' because its config is missing in the '{Config}' config.",
+					GraphNameOrConfig, GEngineIni);
 			}
-			else if (FParse::Param(FCommandLine::Get(), TEXT("DDC-ForceMemoryCache")))
+
+			const ANSICHAR* MissingTypes = 
+				Owner->HasAllFlags(ECacheStoreFlags::Query) ? "writable" :
+				Owner->HasAllFlags(ECacheStoreFlags::Store) ? "readable" : "readable or writable";
+
+			if (FParse::Param(FCommandLine::Get(), TEXT("DDC-ForceMemoryCache")))
 			{
 				UE_LOGFMT(LogDerivedDataCache, Error,
-					"Unable to use default cache graph '{Graph}' because there are no {MissingTypes} nodes available.",
-					DefaultGraphName,
-					Owner->HasAllFlags(ECacheStoreFlags::Query) ? "writable" :
-					Owner->HasAllFlags(ECacheStoreFlags::Store) ? "readable" : "readable or writable");
+					"Unable to use cache graph '{Graph}' because it has no {MissingTypes} nodes available. "
+					"Creating a memory cache as a fallback to avoid crashing immediately.",
+					GraphNameOrConfig, MissingTypes);
 				IMemoryCacheStore* Memory = nullptr;
 				CreateMemoryCacheStore(Memory, TEXT("ForceMemoryCache"), TEXT(""), Owner);
 				check(Owner->HasAllFlags(ECacheStoreFlags::Query | ECacheStoreFlags::Store));
 			}
-			else
+			else if (!FParse::Param(FCommandLine::Get(), TEXT("DDC-AllowNoActiveStores")))
 			{
-				const bool bAllowsNoActiveStore = FParse::Param(FCommandLine::Get(), TEXT("DDC-AllowNoActiveStores"));
-				const ANSICHAR* MissingTypes = 
-					Owner->HasAllFlags(ECacheStoreFlags::Query) ? "writable" :
-					Owner->HasAllFlags(ECacheStoreFlags::Store) ? "readable" : "readable or writable";
-#if WITH_EDITOR
+			#if WITH_EDITOR
 				if (!FApp::IsUnattended() && !IsRunningCommandlet() && !GIsRunningUnattendedScript)
 				{
 					const FString StoreId = TEXT("Epic Games");
 					const FString SectionName = TEXT("GlobalDataCachePath");
 					const FString KeyName = TEXT("UE-LocalDataCachePath");
 
-					//only offer to reset local data cache path if it is already set
+					// Offer to reset the local data cache path if it has been set.
 					if (FString LocalDataCachePath; FPlatformMisc::GetStoredValue(StoreId, SectionName, KeyName, LocalDataCachePath))
 					{
 						const FText ErrorTitle = LOCTEXT("DDC_NoActiveStore_ErrorMsgTitle", "Derived Data Cache Error");
 						const FText ErrorMessage = FText::Format(LOCTEXT("DDC_NoActiveStore_ErrorMsgText",
-							"Unable to use default cache graph '{0}'.\n\n"
-							"The local cache path is currently overridden by the setting\nUE-LocalDataCachePath: {1}\n\n"
-							"Do you want to reset this setting and restart the editor?"), 
-							FText::FromString(DefaultGraphName), FText::FromString(LocalDataCachePath));
+							"Unable to use cache graph '{Graph}'.\n\n"
+							"The local cache path is currently overridden by the setting\nUE-LocalDataCachePath: {LocalDataCachePath}\n\n"
+							"Do you want to reset this setting and restart the editor?"),
+							FText::FromStringView(GraphNameOrConfig), FText::FromString(LocalDataCachePath));
 						const EAppReturnType::Type Result = FPlatformMisc::MessageBoxExt(EAppMsgType::YesNo, *ErrorMessage.ToString(), *ErrorTitle.ToString());
 						if (Result == EAppReturnType::Yes)
 						{
-							UE_CLOGFMT(!bAllowsNoActiveStore, LogDerivedDataCache, Error,
-								"Unable to use default cache graph '{Graph}' because there are no {MissingTypes} nodes available."
+							UE_LOGFMT(LogDerivedDataCache, Error,
+								"Unable to use default cache graph '{Graph}' because it has no {MissingTypes} nodes available."
 								"'UE-LocalDataCachePath' setting will be reset.",
-								DefaultGraphName, MissingTypes);
+								GraphNameOrConfig, MissingTypes);
 
 							FPlatformMisc::DeleteStoredValue(StoreId, SectionName, KeyName);
 							FPlatformProcess::CreateProc(FPlatformProcess::ExecutablePath(), FCommandLine::GetOriginal(),
@@ -943,12 +952,13 @@ public:
 						}
 					}
 				}
-#endif
-				UE_CLOGFMT(!bAllowsNoActiveStore, LogDerivedDataCache, Fatal,
-					"Unable to use default cache graph '{Graph}' because there are no {MissingTypes} nodes available. "
+			#endif
+
+				UE_LOGFMT(LogDerivedDataCache, Fatal,
+					"Unable to use cache graph '{Graph}' because it has no {MissingTypes} nodes available. "
 					"Add -DDC-ForceMemoryCache to the command line to bypass this if you need access "
 					"to the editor settings to fix the cache configuration.",
-					DefaultGraphName, MissingTypes);
+					GraphNameOrConfig, MissingTypes);
 			}
 		}
 
