DDC: Make the API for cache store creation consistent

This is in preparation for a refactor of API used for the graph/hierarchy and the cache stores to communicate with each other.

#Synced-CL 52679933

diff --git a/Engine/Source/Developer/DerivedDataCache/Private/DerivedDataBackends.cpp b/Engine/Source/Developer/DerivedDataCache/Private/DerivedDataBackends.cpp
--- a/Engine/Source/Developer/DerivedDataCache/Private/DerivedDataBackends.cpp
+++ b/Engine/Source/Developer/DerivedDataCache/Private/DerivedDataBackends.cpp
@@ -61,13 +61,17 @@ ILegacyCacheStore* CreateCacheStoreAsync(ICacheStoreOwner* Owner, ILegacyCacheSt
 ILegacyCacheStore* CreateCacheStoreHierarchy(ICacheStoreOwner*& OutOwner, TFunctionRef<void (IMemoryCacheStore*&)> MemoryCacheCreator);
 ILegacyCacheStore* CreateCacheStoreThrottle(ILegacyCacheStore* InnerCache, uint32 LatencyMS, uint32 MaxBytesPerSecond);
 ILegacyCacheStore* CreateCacheStoreVerify(ILegacyCacheStore* InnerCache, bool bPutOnError);
-ILegacyCacheStore* CreateFileSystemCacheStore(const TCHAR* Name, const TCHAR* Config, ICacheStoreOwner& Owner, ICacheStoreGraph* Graph, FString& OutPath, bool& OutRedirected);
-ILegacyCacheStore* CreateHttpCacheStore(const TCHAR* Name, const TCHAR* Config, ICacheStoreOwner* Owner);
+ILegacyCacheStore* TryCreateCacheStoreReplay(ILegacyCacheStore* InnerCache, ICacheStoreOwner* Owner);
+
+ILegacyCacheStore* CreateFileSystemCacheStore(const TCHAR* Name, const TCHAR* Config, ICacheStoreOwner* Owner, ICacheStoreGraph* Graph);
+ILegacyCacheStore* CreateHttpCacheStore(const TCHAR* Name, const TCHAR* Config, ICacheStoreOwner* Owner, ICacheStoreGraph* Graph);
+ILegacyCacheStore* CreateMemoryCacheStore(const TCHAR* Name, const TCHAR* Config, ICacheStoreOwner* Owner, ICacheStoreGraph* Graph);
+ILegacyCacheStore* CreatePakFileCacheStore(const TCHAR* Name, const TCHAR* Config, ICacheStoreOwner* Owner, ICacheStoreGraph* Graph);
+ILegacyCacheStore* CreateS3CacheStore(const TCHAR* Name, const TCHAR* Config, ICacheStoreOwner* Owner, ICacheStoreGraph* Graph);
+ILegacyCacheStore* CreateZenCacheStore(const TCHAR* Name, const TCHAR* Config, ICacheStoreOwner* Owner, ICacheStoreGraph* Graph);
+
 void CreateMemoryCacheStore(IMemoryCacheStore*& OutCache, const TCHAR* Name, const TCHAR* Config, ICacheStoreOwner* Owner);
 IPakFileCacheStore* CreatePakFileCacheStore(const TCHAR* Name, const TCHAR* Filename, bool bWriting, bool bCompressed, ICacheStoreOwner* Owner);
-ILegacyCacheStore* CreateS3CacheStore(const TCHAR* Name, const TCHAR* Config, ICacheStoreOwner& Owner);
-ILegacyCacheStore* CreateZenCacheStore(const TCHAR* Name, const TCHAR* Config, ICacheStoreOwner* Owner);
-ILegacyCacheStore* TryCreateCacheStoreReplay(ILegacyCacheStore* InnerCache, ICacheStoreOwner* Owner);
 
 /**
  * This class is used to create a singleton that represents the derived data cache hierarchy and all of the wrappers necessary
@@ -188,23 +192,51 @@ private:
 			{
 				if (NodeType == TEXT("FileSystem"))
 				{
-					return ParseFileSystemCache(*StoreName, *Entry, InParsedNodes);
+					if (ILegacyCacheStore* CacheStore = CreateFileSystemCacheStore(*StoreName, *Entry, this, this))
+					{
+						InParsedNodes.Add(StoreName, CacheStore);
+						return true;
+					}
 				}
 				else if (NodeType == TEXT("Memory"))
 				{
-					return ParseMemoryCache(*StoreName, *Entry, InParsedNodes);
-				}
-				else if (NodeType == TEXT("ReadPak"))
-				{
-					return ParsePak(*StoreName, *Entry, /*bWriting*/ false, InParsedNodes);
+					if (ILegacyCacheStore* CacheStore = CreateMemoryCacheStore(*StoreName, *Entry, this, this))
+					{
+						InParsedNodes.Add(StoreName, CacheStore);
+						return true;
+					}
 				}
-				else if (NodeType == TEXT("WritePak"))
+				else if (NodeType == TEXT("ReadPak") || NodeType == TEXT("WritePak"))
 				{
-					return ParsePak(*StoreName, *Entry, /*bWriting*/ true, InParsedNodes);
+					if (IPakFileCacheStore* PakStore = static_cast<IPakFileCacheStore*>(CreatePakFileCacheStore(*StoreName, *Entry, this, this)))
+					{
+						FString PakFilename;
+						FParse::Value(*Entry, TEXT("Filename="), PakFilename);
+
+						if (PakStore->IsWritable())
+						{
+							if (WritePakCache)
+							{
+								UE_LOG(LogDerivedDataCache, Warning, TEXT("Unable to create %s pak cache because only one pak cache write node is supported."), *StoreName);
+								delete PakStore;
+								return false;
+							}
+							ReadPakFilename = PakFilename;
+							WritePakFilename = PakStore->GetFilename();
+							WritePakCache = PakStore;
+						}
+						else
+						{
+							ReadPakFilename = PakFilename;
+							ReadPakCache.Add(PakStore);
+						}
+						InParsedNodes.Add(StoreName, PakStore);
+						return true;
+					}
 				}
 				else if (NodeType == TEXT("S3"))
 				{
-					if (ILegacyCacheStore* CacheStore = CreateS3CacheStore(*StoreName, *Entry, *this))
+					if (ILegacyCacheStore* CacheStore = CreateS3CacheStore(*StoreName, *Entry, this, this))
 					{
 						InParsedNodes.Add(StoreName, CacheStore);
 						return true;
@@ -212,7 +244,7 @@ private:
 				}
 				else if (NodeType == TEXT("Cloud") || NodeType == TEXT("Http"))
 				{
-					if (ILegacyCacheStore* CacheStore = CreateHttpCacheStore(*StoreName, *Entry, this))
+					if (ILegacyCacheStore* CacheStore = CreateHttpCacheStore(*StoreName, *Entry, this, this))
 					{
 						InParsedNodes.Add(StoreName, CacheStore);
 						return true;
@@ -220,7 +252,7 @@ private:
 				}
 				else if (NodeType == TEXT("Zen"))
 				{
-					if (ILegacyCacheStore* CacheStore = CreateZenCacheStore(*StoreName, *Entry, this))
+					if (ILegacyCacheStore* CacheStore = CreateZenCacheStore(*StoreName, *Entry, this, this))
 					{
 						InParsedNodes.Add(StoreName, CacheStore);
 						return true;
@@ -232,62 +264,6 @@ private:
 		return false;
 	}
 
-	/**
-	 * Creates Read/write Pak file interface from ini settings
-	 *
-	 * @param NodeName Node name
-	 * @param Entry Node definition
-	 * @param bWriting true to create pak interface for writing
-	 * @return Pak file data backend interface instance or nullptr if unsuccessful
-	 */
-	bool ParsePak(const TCHAR* NodeName, const TCHAR* Entry, const bool bWriting, FParsedNodeMap& InParsedNodes)
-	{
-		ILegacyCacheStore* PakNode = nullptr;
-		FString PakFilename;
-		FParse::Value(Entry, TEXT("Filename="), PakFilename);
-		bool bCompressed = false;
-		FParse::Bool(Entry, TEXT("Compressed="), bCompressed);
-
-		if (PakFilename.IsEmpty())
-		{
-			UE_LOGF(LogDerivedDataCache, Log, "FDerivedDataBackendGraph: %ls pak cache Filename not found in *engine.ini, will not use a pak cache.", NodeName);
-			return false;
-		}
-
-		if (bWriting)
-		{
-			if (WritePakCache)
-			{
-				UE_LOGF(LogDerivedDataCache, Warning, "Unable to create %ls pak cache because only one pak cache write node is supported.", NodeName);
-				return false;
-			}
-
-			FGuid Temp = FGuid::NewGuid();
-			ReadPakFilename = PakFilename;
-			WritePakFilename = PakFilename + TEXT(".") + Temp.ToString();
-			WritePakCache = CreatePakFileCacheStore(NodeName, *WritePakFilename, /*bWriting*/ true, bCompressed, this);
-			PakNode = WritePakCache;
-			InParsedNodes.Add(NodeName, WritePakCache);
-			return true;
-		}
-		else
-		{
-			bool bReadPak = FPlatformFileManager::Get().GetPlatformFile().FileExists(*PakFilename);
-			if (!bReadPak)
-			{
-				UE_LOGF(LogDerivedDataCache, Log, "FDerivedDataBackendGraph: %ls pak cache file %ls not found, will not use a pak cache.", NodeName, *PakFilename);
-				return false;
-			}
-
-			IPakFileCacheStore* ReadPak = CreatePakFileCacheStore(NodeName, *PakFilename, /*bWriting*/ false, bCompressed, this);
-			ReadPakFilename = PakFilename;
-			PakNode = ReadPak;
-			ReadPakCache.Add(ReadPak);
-			InParsedNodes.Add(NodeName, ReadPak);
-			return true;
-		}
-	}
-
 	bool ParseGraph(FStringView GraphNameOrConfig, FParsedNodeMap& InParsedNodes)
 	{
 		TRACE_CPUPROFILER_EVENT_SCOPE(FDerivedDataBackendGraph::ParseGraph);
@@ -324,53 +300,6 @@ private:
 		return true;
 	}
 
-	/**
-	 * Creates Filesystem data cache interface from ini settings.
-	 *
-	 * @param NodeName Node name.
-	 * @param Entry Node definition.
-	 * @return Filesystem data cache backend interface instance or nullptr if unsuccessful
-	 */
-	bool ParseFileSystemCache(const TCHAR* NodeName, const TCHAR* Config, FParsedNodeMap& InParsedNodes)
-	{
-		bool bRedirected = false;
-		FString Path;
-		if (ILegacyCacheStore* Store = CreateFileSystemCacheStore(NodeName, Config, *this, this, Path, bRedirected))
-		{
-			if (!bRedirected)
-			{
-				bUsingSharedDDC |= NodeName == TEXTVIEW("Shared");
-				Directories.AddUnique(Path);
-			}
-			InParsedNodes.Add(NodeName, Store);
-			return true;
-		}
-
-		return false;
-	}
-
-	/**
-	 * Creates Memory data cache interface from ini settings.
-	 *
-	 * @param NodeName Node name.
-	 * @param Entry Node definition.
-	 * @return Memory data cache backend interface instance or nullptr if unsuccessful
-	 */
-	bool ParseMemoryCache(const TCHAR* NodeName, const TCHAR* Entry, FParsedNodeMap& InParsedNodes)
-	{
-		FString Filename;
-		FParse::Value(Entry, TEXT("Filename="), Filename);
-		IMemoryCacheStore* Cache;
-		CreateMemoryCacheStore(Cache, NodeName, TEXT(""), this);
-		check(Cache);
-		InParsedNodes.Add(NodeName, Cache);
-		if (Filename.Len())
-		{
-			UE_LOGF(LogDerivedDataCache, Display, "Memory nodes that load from a file are deprecated. Please remove the filename from the cache configuration.");
-		}
-		return true;
-	}
-
 	void GetMemoryCache(IMemoryCacheStore*& OutCache)
 	{
 		if (MemoryCache)
@@ -476,16 +405,6 @@ private:
 		return bUsingSharedDDC;
 	}
 
-	virtual const TCHAR* GetGraphName() const override
-	{
-		return *ActiveGraphName;
-	}
-
-	virtual const TCHAR* GetDefaultGraphName() const override
-	{
-		return FApp::IsEngineInstalled() ? TEXT("Installed") : TEXT("Default");
-	}
-
 	virtual void AddToAsyncCompletionCounter(int32 Addend) override
 	{
 		verify(AddToAsyncTaskCounter(Addend) + Addend >= 0);
@@ -502,6 +421,16 @@ private:
 	}
 
 protected:
+	virtual const TCHAR* GetGraphName() const override
+	{
+		return *ActiveGraphName;
+	}
+
+	virtual const TCHAR* GetDefaultGraphName() const final
+	{
+		return FApp::IsEngineInstalled() ? TEXT("Installed") : TEXT("Default");
+	}
+
 	virtual ILegacyCacheStore* MountPakFile(const TCHAR* PakFilename) override
 	{
 		// Assumptions: there's at least one read-only pak backend in the hierarchy
@@ -860,8 +789,7 @@ public:
 				CandidateGraphs.Add(LegacyDefaultGraphName);
 			}
 
-			const TCHAR* const DefaultGraphName = FApp::IsEngineInstalled() ? TEXT("Installed") : TEXT("Default");
-			CandidateGraphs.Add(DefaultGraphName);
+			CandidateGraphs.Add(GetDefaultGraphName());
 		}
 
 		// Construct candidate graphs in the order they were added until one is valid.
diff --git a/Engine/Source/Developer/DerivedDataCache/Private/FileSystemCacheStore.cpp b/Engine/Source/Developer/DerivedDataCache/Private/FileSystemCacheStore.cpp
--- a/Engine/Source/Developer/DerivedDataCache/Private/FileSystemCacheStore.cpp
+++ b/Engine/Source/Developer/DerivedDataCache/Private/FileSystemCacheStore.cpp
@@ -862,9 +862,7 @@ public:
 	static ILegacyCacheStore* TryCreate(
 		const FFileSystemCacheStoreParams& Params,
 		ICacheStoreOwner& Owner,
-		ICacheStoreGraph* Graph,
-		FString& OutPath,
-		bool& bOutRedirected);
+		ICacheStoreGraph* Graph);
 
 private:
 	FFileSystemCacheStore(
@@ -1032,9 +1030,7 @@ private:
 ILegacyCacheStore* FFileSystemCacheStore::TryCreate(
 	const FFileSystemCacheStoreParams& Params,
 	ICacheStoreOwner& Owner,
-	ICacheStoreGraph* Graph,
-	FString& OutPath,
-	bool& bOutRedirected)
+	ICacheStoreGraph* Graph)
 {
 	// If we find a platform that has more stringent limits, this needs to be rethought.
 	checkf(GMaxCacheRootLen + GMaxCacheKeyLen <= FPlatformMisc::GetMaxPathLength(),
@@ -1072,7 +1068,6 @@ ILegacyCacheStore* FFileSystemCacheStore::TryCreate(
 		{
 			if (ILegacyCacheStore* RedirectedStore = TryRedirection(Params, Owner, Graph))
 			{
-				bOutRedirected = true;
 				return RedirectedStore;
 			}
 
@@ -1089,7 +1084,6 @@ ILegacyCacheStore* FFileSystemCacheStore::TryCreate(
 					TUniqueLock Lock(ActiveStoresMutex);
 					ActiveStores.Add(CacheStore);
 				}
-				OutPath = CacheStore->CachePath;
 				return CacheStore;
 			}
 
@@ -2523,7 +2517,7 @@ bool FFileSystemCacheStore::GetCacheContent(
 			}
 			UE_LOGF(LogDerivedDataCache, Display,
 				"%ls: Cache miss with corrupted value %ls with hash %ls for %ls from '%.*ls'",
-				*CachePath, *WriteToString<16>(Id), *WriteToString<48>(RawHash),
+				*CachePath, *WriteToString<32>(Id), *WriteToString<48>(RawHash),
 				*WriteToString<96>(Key), Name.Len(), Name.GetData());
 			DeleteCacheContent(Name, Value);
 			return false;
@@ -3182,17 +3176,11 @@ void FFileSystemCacheStoreParams::Parse(const TCHAR* Name, const TCHAR* Config)
 	FParse::Value(Config, TEXT("RedirectionKeyName="), RedirectionKeyName);
 }
 
-ILegacyCacheStore* CreateFileSystemCacheStore(
-	const TCHAR* Name,
-	const TCHAR* Config,
-	ICacheStoreOwner& Owner,
-	ICacheStoreGraph* Graph,
-	FString& OutPath,
-	bool& bOutRedirected)
+ILegacyCacheStore* CreateFileSystemCacheStore(const TCHAR* Name, const TCHAR* Config, ICacheStoreOwner* Owner, ICacheStoreGraph* Graph)
 {
 	FFileSystemCacheStoreParams Params;
 	Params.Parse(Name, Config);
-	return FFileSystemCacheStore::TryCreate(Params, Owner, Graph, OutPath, bOutRedirected);
+	return FFileSystemCacheStore::TryCreate(Params, *Owner, Graph);
 }
 
 } // UE::DerivedData
diff --git a/Engine/Source/Developer/DerivedDataCache/Private/HttpCacheStore.cpp b/Engine/Source/Developer/DerivedDataCache/Private/HttpCacheStore.cpp
--- a/Engine/Source/Developer/DerivedDataCache/Private/HttpCacheStore.cpp
+++ b/Engine/Source/Developer/DerivedDataCache/Private/HttpCacheStore.cpp
@@ -3569,7 +3569,7 @@ void FHttpCacheStoreParams::Parse(const TCHAR* NodeName, const TCHAR* Config)
 namespace UE::DerivedData
 {
 
-ILegacyCacheStore* CreateHttpCacheStore(const TCHAR* NodeName, const TCHAR* Config, ICacheStoreOwner* Owner)
+ILegacyCacheStore* CreateHttpCacheStore(const TCHAR* NodeName, const TCHAR* Config, ICacheStoreOwner* Owner, ICacheStoreGraph* Graph)
 {
 #if !WITH_HTTP_DDC_BACKEND
 	UE_LOGF(LogDerivedDataCache, Warning, "%ls: HTTP cache is not yet supported in the current build configuration.", NodeName);
diff --git a/Engine/Source/Developer/DerivedDataCache/Private/MemoryCacheStore.cpp b/Engine/Source/Developer/DerivedDataCache/Private/MemoryCacheStore.cpp
--- a/Engine/Source/Developer/DerivedDataCache/Private/MemoryCacheStore.cpp
+++ b/Engine/Source/Developer/DerivedDataCache/Private/MemoryCacheStore.cpp
@@ -738,4 +738,11 @@ void CreateMemoryCacheStore(IMemoryCacheStore*& OutCache, const TCHAR* Name, con
 	new FMemoryCacheStore(OutCache, Name, Config, Owner);
 }
 
+ILegacyCacheStore* CreateMemoryCacheStore(const TCHAR* Name, const TCHAR* Config, ICacheStoreOwner* Owner, ICacheStoreGraph* Graph)
+{
+	IMemoryCacheStore* Cache = nullptr;
+	CreateMemoryCacheStore(Cache, Name, Config, Owner);
+	return Cache;
+}
+
 } // UE::DerivedData
diff --git a/Engine/Source/Developer/DerivedDataCache/Private/PakFileCacheStore.cpp b/Engine/Source/Developer/DerivedDataCache/Private/PakFileCacheStore.cpp
--- a/Engine/Source/Developer/DerivedDataCache/Private/PakFileCacheStore.cpp
+++ b/Engine/Source/Developer/DerivedDataCache/Private/PakFileCacheStore.cpp
@@ -1624,4 +1624,37 @@ IPakFileCacheStore* CreatePakFileCacheStore(const TCHAR* Name, const TCHAR* Cach
 	return bCompressed ? new FCompressedPakFileCacheStore(Name, CachePath, bWriting, Owner) : new FPakFileCacheStore(Name, CachePath, bWriting, Owner);
 }
 
+ILegacyCacheStore* CreatePakFileCacheStore(const TCHAR* Name, const TCHAR* Config, ICacheStoreOwner* Owner, ICacheStoreGraph* Graph)
+{
+	FString NodeType;
+	FParse::Value(Config, TEXT("Type="), NodeType);
+	const bool bWriting = (NodeType == TEXT("WritePak"));
+
+	FString Filename;
+	if (!FParse::Value(Config, TEXT("Filename="), Filename))
+	{
+		UE_LOG(LogDerivedDataCache, Log, TEXT("%s: Pak cache Filename not found, will not use a pak cache."), Name);
+		return nullptr;
+	}
+
+	bool bCompressed = false;
+	FParse::Bool(Config, TEXT("Compressed="), bCompressed);
+
+	if (bWriting)
+	{
+		FGuid TempGuid = FGuid::NewGuid();
+		FString WriteFilename = Filename + TEXT(".") + TempGuid.ToString();
+		return CreatePakFileCacheStore(Name, *WriteFilename, /*bWriting*/ true, bCompressed, Owner);
+	}
+	else
+	{
+		if (!FPlatformFileManager::Get().GetPlatformFile().FileExists(*Filename))
+		{
+			UE_LOG(LogDerivedDataCache, Log, TEXT("%s: Pak cache file %s not found, will not use a pak cache."), Name, *Filename);
+			return nullptr;
+		}
+		return CreatePakFileCacheStore(Name, *Filename, /*bWriting*/ false, bCompressed, Owner);
+	}
+}
+
 } // UE::DerivedData
diff --git a/Engine/Source/Developer/DerivedDataCache/Private/S3CacheStore.cpp b/Engine/Source/Developer/DerivedDataCache/Private/S3CacheStore.cpp
--- a/Engine/Source/Developer/DerivedDataCache/Private/S3CacheStore.cpp
+++ b/Engine/Source/Developer/DerivedDataCache/Private/S3CacheStore.cpp
@@ -1846,7 +1846,7 @@ void FS3CacheStore::GetChunks(
 namespace UE::DerivedData
 {
 
-ILegacyCacheStore* CreateS3CacheStore(const TCHAR* Name, const TCHAR* Config, ICacheStoreOwner& Owner)
+ILegacyCacheStore* CreateS3CacheStore(const TCHAR* Name, const TCHAR* Config, ICacheStoreOwner* Owner, ICacheStoreGraph* Graph)
 {
 #if WITH_S3_DDC_BACKEND
 	FString ManifestPath;
@@ -1896,7 +1896,7 @@ ILegacyCacheStore* CreateS3CacheStore(const TCHAR* Name, const TCHAR* Config, IC
 		return nullptr;
 	}
 
-	TUniquePtr<FS3CacheStore> Store(new FS3CacheStore(Name, *ManifestPath, *BaseUrl, *Region, *CanaryObjectKey, *CachePath, Owner));
+	TUniquePtr<FS3CacheStore> Store(new FS3CacheStore(Name, *ManifestPath, *BaseUrl, *Region, *CanaryObjectKey, *CachePath, *Owner));
 	if (!Store->IsUsable())
 	{
 		return nullptr;
diff --git a/Engine/Source/Developer/DerivedDataCache/Private/Tests/HttpCacheStoreTest.cpp b/Engine/Source/Developer/DerivedDataCache/Private/Tests/HttpCacheStoreTest.cpp
--- a/Engine/Source/Developer/DerivedDataCache/Private/Tests/HttpCacheStoreTest.cpp
+++ b/Engine/Source/Developer/DerivedDataCache/Private/Tests/HttpCacheStoreTest.cpp
@@ -45,7 +45,7 @@ ILegacyCacheStore* GetAnyHttpCacheStore(
 	FString& OutAccessToken,
 	FString& OutNamespace);
 
-ILegacyCacheStore* CreateZenCacheStore(const TCHAR* NodeName, const TCHAR* Config, ICacheStoreOwner* Owner);
+ILegacyCacheStore* CreateZenCacheStore(const TCHAR* NodeName, const TCHAR* Config, ICacheStoreOwner* Owner, ICacheStoreGraph* Graph);
 
 class FHttpCacheStoreTestBase : public FAutomationTestBase
 {
@@ -675,11 +675,11 @@ bool FHttpCacheStoreTest::RunTest(const FString& Parameters)
 
 	FScopeZenService ScopeZenSiblingService(MoveTemp(ZenTestServiceSiblingSettings));
 	TUniquePtr<ILegacyCacheStore> ZenIntermediarySiblingBackend(CreateZenCacheStore(TEXT("TestSibling"),
-		*FString::Printf(TEXT("Host=%s, StructuredNamespace=%s"), *ScopeZenSiblingService.GetInstance().GetEndpoint().GetURL(), *TestNamespace), nullptr));
+		*FString::Printf(TEXT("Host=%s, StructuredNamespace=%s"), *ScopeZenSiblingService.GetInstance().GetEndpoint().GetURL(), *TestNamespace), nullptr, nullptr));
 
 	FScopeZenService ScopeZenService(MoveTemp(ZenTestServiceSettings));
 	TUniquePtr<ILegacyCacheStore> ZenIntermediaryBackend(CreateZenCacheStore(TEXT("Test"),
-		*FString::Printf(TEXT("Host=%s, StructuredNamespace=%s"), *ScopeZenService.GetInstance().GetEndpoint().GetURL(), *TestNamespace), nullptr));
+		*FString::Printf(TEXT("Host=%s, StructuredNamespace=%s"), *ScopeZenService.GetInstance().GetEndpoint().GetURL(), *TestNamespace), nullptr, nullptr));
 	auto WaitForZenPushToUpstream = [](ILegacyCacheStore* ZenBackend, TConstArrayView<FCacheRecord> Records)
 	{
 		// TODO: Expecting a legitimate means to wait for zen to finish pushing records to its upstream in the future
diff --git a/Engine/Source/Developer/DerivedDataCache/Private/ZenCacheStore.cpp b/Engine/Source/Developer/DerivedDataCache/Private/ZenCacheStore.cpp
--- a/Engine/Source/Developer/DerivedDataCache/Private/ZenCacheStore.cpp
+++ b/Engine/Source/Developer/DerivedDataCache/Private/ZenCacheStore.cpp
@@ -2363,7 +2363,7 @@ void FZenCacheStoreParams::Parse(const TCHAR* NodeName, const TCHAR* Config)
 	FParse::Value(Config, TEXT("ChunksBatchSize="), ChunksBatchSize);
 }
 
-ILegacyCacheStore* CreateZenCacheStore(const TCHAR* NodeName, const TCHAR* Config, ICacheStoreOwner* Owner)
+ILegacyCacheStore* CreateZenCacheStore(const TCHAR* NodeName, const TCHAR* Config, ICacheStoreOwner* Owner, ICacheStoreGraph* Graph)
 {
 	FZenCacheStoreParams Params;
 	Params.Parse(NodeName, Config);
@@ -2497,7 +2497,7 @@ ILegacyCacheStore* CreateZenCacheStore(const TCHAR* NodeName, const TCHAR* Confi
 namespace UE::DerivedData
 {
 
-ILegacyCacheStore* CreateZenCacheStore(const TCHAR* NodeName, const TCHAR* Config, ICacheStoreOwner* Owner)
+ILegacyCacheStore* CreateZenCacheStore(const TCHAR* NodeName, const TCHAR* Config, ICacheStoreOwner* Owner, ICacheStoreGraph* Graph)
 {
 	UE_LOGF(LogDerivedDataCache, Warning, "%ls: Zen cache is not yet supported in the current build configuration.", NodeName);
 	return nullptr;
