DDC: Deprecate functionality in the UE4 API that is no longer supported

#Synced-CL 52727211

diff --git a/Engine/Source/Developer/DerivedDataCache/Private/DerivedDataBackendInterface.h b/Engine/Source/Developer/DerivedDataCache/Private/DerivedDataBackendInterface.h
--- a/Engine/Source/Developer/DerivedDataCache/Private/DerivedDataBackendInterface.h
+++ b/Engine/Source/Developer/DerivedDataCache/Private/DerivedDataBackendInterface.h
@@ -96,7 +96,6 @@ public:
 	virtual bool AnyAsyncRequestsRemaining() = 0;
 	virtual bool IsShuttingDown() = 0;
 	virtual void WaitForQuiescence(bool bShutdown = false) = 0;
-	virtual void GetDirectories(TArray<FString>& OutResults) = 0;
 	virtual bool GetUsingSharedDDC() const = 0;
 	virtual const TCHAR* GetGraphName() const = 0;
 	virtual const TCHAR* GetDefaultGraphName() const = 0;
diff --git a/Engine/Source/Developer/DerivedDataCache/Private/DerivedDataBackends.cpp b/Engine/Source/Developer/DerivedDataCache/Private/DerivedDataBackends.cpp
--- a/Engine/Source/Developer/DerivedDataCache/Private/DerivedDataBackends.cpp
+++ b/Engine/Source/Developer/DerivedDataCache/Private/DerivedDataBackends.cpp
@@ -123,7 +123,6 @@ protected:
 		WritePakCache = nullptr;
 		Hierarchy = nullptr;
 		ReadPakCache.Empty();
-		Directories.Empty();
 		ActiveGraphName.Empty();
 	}
 
@@ -415,11 +414,6 @@ private:
 		return AddToAsyncTaskCounter(0) > 0;
 	}
 
-	virtual void GetDirectories(TArray<FString>& OutResults) override
-	{
-		OutResults = Directories;
-	}
-
 protected:
 	virtual const TCHAR* GetGraphName() const override
 	{
@@ -728,9 +722,6 @@ private:
 	/** Support for multiple read only pak files. */
 	TArray<IPakFileCacheStore*> ReadPakCache;
 
-	/** List of directories used by the DDC */
-	TArray<FString> Directories;
-
 	/** Whether a shared cache is in use */
 	bool bUsingSharedDDC = false;
 
diff --git a/Engine/Source/Developer/DerivedDataCache/Private/DerivedDataCache.cpp b/Engine/Source/Developer/DerivedDataCache/Private/DerivedDataCache.cpp
--- a/Engine/Source/Developer/DerivedDataCache/Private/DerivedDataCache.cpp
+++ b/Engine/Source/Developer/DerivedDataCache/Private/DerivedDataCache.cpp
@@ -827,10 +827,6 @@ public:
 		return true;
 	}
 
-	void NotifyBootComplete() override
-	{
-	}
-
 	void AddToAsyncCompletionCounter(int32 Addend) override
 	{
 		Backend->AddToAsyncCompletionCounter(Addend);
@@ -863,11 +859,6 @@ public:
 		return Backend->GetDefaultGraphName();
 	}
 
-	void GetDirectories(TArray<FString>& OutResults) override
-	{
-		Backend->GetDirectories(OutResults);
-	}
-
 	virtual bool IsFinished() const override
 	{
 		return Backend->GetMaintainer().IsIdle();
@@ -909,10 +900,12 @@ public:
 	}
 
 	/** Get event delegate for data cache notifications */
+	PRAGMA_DISABLE_DEPRECATION_WARNINGS;
 	virtual FOnDDCNotification& GetDDCNotificationEvent()
 	{
 		return DDCNotificationEvent;
 	}
+	PRAGMA_ENABLE_DEPRECATION_WARNINGS;
 
 	ICache* GetCache() const
 	{
@@ -947,8 +940,10 @@ private:
 	/** Map of handle to pending task **/
 	TMap<uint32, TUniquePtr<FLegacyFetchOrBuildTask>> PendingTasks;
 
+	PRAGMA_DISABLE_DEPRECATION_WARNINGS;
 	/** Cache notification delegate */
 	FOnDDCNotification DDCNotificationEvent;
+	PRAGMA_ENABLE_DEPRECATION_WARNINGS;
 };
 
 void CreateCache(ICache*& OutCache, FDerivedDataCacheInterface*& OutLegacyCache)
diff --git a/Engine/Source/Developer/DerivedDataCache/Public/DerivedDataCacheInterface.h b/Engine/Source/Developer/DerivedDataCache/Public/DerivedDataCacheInterface.h
--- a/Engine/Source/Developer/DerivedDataCache/Public/DerivedDataCacheInterface.h
+++ b/Engine/Source/Developer/DerivedDataCache/Public/DerivedDataCacheInterface.h
@@ -215,14 +215,10 @@ public:
 	// System Interface
 	//--------------------
 
-	/**
-	 * Notify the system that the boot process is complete and so we can write the boot cache and get rid of it
-	 */
-	virtual void NotifyBootComplete() = 0;
+	UE_DEPRECATED(5.8, "The cache no longer uses this internally.")
+	void NotifyBootComplete() {}
 
-	/**
-	 * Adds or subtracts a number from the thread safe counter which tracks outstanding async requests. This is used to ensure everything is complete prior to shutdown.
-	 */
+	UE_DEPRECATED(5.8, "The cache tracks its async operations internally. External modifications are no longer supported.")
 	virtual void AddToAsyncCompletionCounter(int32 Addend) = 0;
 
 	/**
@@ -235,14 +231,10 @@ public:
 	 */
 	virtual void WaitForQuiescence(bool bShutdown = false) = 0;
 
-	/**
-	 * Retrieve the directories used by the DDC
-	 */
-	virtual void GetDirectories(TArray<FString>& OutResults) = 0;
+	UE_DEPRECATED(5.8, "Querying cache directories is no longer supported.")
+	void GetDirectories(TArray<FString>& OutResults) {}
 
-	/**
-	* Retrieve whether a Shared DDC is being used
-	*/
+	UE_DEPRECATED(5.8, "Querying whether a shared cache is being used is no longer supported.")
 	virtual bool GetUsingSharedDDC() const = 0;
 
 	/**
@@ -288,18 +280,18 @@ public:
 	// Notification Interface
 	//-----------------------
 
-	enum EDDCNotification
+	PRAGMA_DISABLE_DEPRECATION_WARNINGS;
+
+	enum UE_DEPRECATED(5.8, "Notifications from the cache are no longer supported.") EDDCNotification
 	{
 		// DDC performance notification generated when not using a shared cache
 		SharedDDCPerformanceNotification
 	};
 
+	using FOnDDCNotification UE_DEPRECATED(5.8, "Notifications from the cache are no longer supported.") = TMulticastDelegate<void (EDDCNotification)>;
 
-	DECLARE_MULTICAST_DELEGATE_OneParam(FOnDDCNotification, EDDCNotification);
-
-	/**
-	* Retrieve the event delegate for cache put notifications
-	*/
+	UE_DEPRECATED(5.8, "Notifications from the cache are no longer supported.")
 	virtual FOnDDCNotification& GetDDCNotificationEvent() = 0;
 
+	PRAGMA_ENABLE_DEPRECATION_WARNINGS;
 };
diff --git a/Engine/Source/Editor/DerivedDataEditor/Private/DerivedDataEditor.cpp b/Engine/Source/Editor/DerivedDataEditor/Private/DerivedDataEditor.cpp
--- a/Engine/Source/Editor/DerivedDataEditor/Private/DerivedDataEditor.cpp
+++ b/Engine/Source/Editor/DerivedDataEditor/Private/DerivedDataEditor.cpp
@@ -1,8 +1,6 @@
 // Copyright Epic Games, Inc. All Rights Reserved.
 
 #include "DerivedDataEditor.h"
-#include "DerivedDataCacheNotifications.h"
-#include "IDerivedDataCacheNotifications.h"
 #include "DerivedDataInformation.h"
 #include "Framework/Application/SlateApplication.h"
 #include "Modules/ModuleManager.h"
@@ -49,8 +47,6 @@ void FDerivedDataEditor::StartupModule()
 #endif // WITH_RELOAD
 
 	FDerivedDataEditorMenuCommands::Register();
-
-	DerivedDataCacheNotifications.Reset(new FDerivedDataCacheNotifications);
 }
 
 void FDerivedDataEditor::ShutdownModule()
diff --git a/Engine/Source/Editor/DerivedDataEditor/Private/SDerivedDataStatusBar.cpp b/Engine/Source/Editor/DerivedDataEditor/Private/SDerivedDataStatusBar.cpp
--- a/Engine/Source/Editor/DerivedDataEditor/Private/SDerivedDataStatusBar.cpp
+++ b/Engine/Source/Editor/DerivedDataEditor/Private/SDerivedDataStatusBar.cpp
@@ -6,7 +6,6 @@
 #include "Containers/Array.h"
 #include "Delegates/Delegate.h"
 #include "DerivedDataCacheInterface.h"
-#include "IDerivedDataCacheNotifications.h"
 #include "DerivedDataEditor.h"
 #include "DerivedDataInformation.h"
 #include "Framework/Commands/InputChord.h"
diff --git a/Engine/Source/Editor/DerivedDataEditor/Public/DerivedDataEditor.h b/Engine/Source/Editor/DerivedDataEditor/Public/DerivedDataEditor.h
--- a/Engine/Source/Editor/DerivedDataEditor/Public/DerivedDataEditor.h
+++ b/Engine/Source/Editor/DerivedDataEditor/Public/DerivedDataEditor.h
@@ -7,7 +7,6 @@
 #include "Modules/ModuleManager.h"
 #include "Templates/SharedPointer.h"
 #include "Templates/UniquePtr.h"
-#include "IDerivedDataCacheNotifications.h"
 
 #define UE_API DERIVEDDATAEDITOR_API
 
@@ -53,7 +52,6 @@ private:
 
 	TSharedPtr<SWindow>	SettingsWindow;
 	TSharedPtr<SDerivedDataCacheSettingsDialog> SettingsDialog;
-	TUniquePtr<IDerivedDataCacheNotifications>	DerivedDataCacheNotifications;
 };
 
 
diff --git a/Engine/Source/Editor/DerivedDataWidgets/Private/DerivedDataCacheNotifications.cpp b/Engine/Source/Editor/DerivedDataWidgets/Private/DerivedDataCacheNotifications.cpp
deleted file mode 100644
diff --git a/Engine/Source/Editor/DerivedDataWidgets/Public/DerivedDataCacheNotifications.h b/Engine/Source/Editor/DerivedDataWidgets/Public/DerivedDataCacheNotifications.h
deleted file mode 100644
diff --git a/Engine/Source/Editor/DerivedDataWidgets/Public/IDerivedDataCacheNotifications.h b/Engine/Source/Editor/DerivedDataWidgets/Public/IDerivedDataCacheNotifications.h
deleted file mode 100644
diff --git a/Engine/Source/Editor/ZenEditor/Private/SZenStatusBar.cpp b/Engine/Source/Editor/ZenEditor/Private/SZenStatusBar.cpp
--- a/Engine/Source/Editor/ZenEditor/Private/SZenStatusBar.cpp
+++ b/Engine/Source/Editor/ZenEditor/Private/SZenStatusBar.cpp
@@ -6,7 +6,6 @@
 #include "Containers/Array.h"
 #include "Delegates/Delegate.h"
 #include "DerivedDataCacheInterface.h"
-#include "IDerivedDataCacheNotifications.h"
 #include "ZenEditor.h"
 #include "DerivedDataInformation.h"
 #include "Framework/Commands/InputChord.h"
diff --git a/Engine/Source/Editor/ZenEditor/Private/ZenEditor.cpp b/Engine/Source/Editor/ZenEditor/Private/ZenEditor.cpp
--- a/Engine/Source/Editor/ZenEditor/Private/ZenEditor.cpp
+++ b/Engine/Source/Editor/ZenEditor/Private/ZenEditor.cpp
@@ -1,8 +1,6 @@
 // Copyright Epic Games, Inc. All Rights Reserved.
 
 #include "ZenEditor.h"
-#include "DerivedDataCacheNotifications.h"
-#include "IDerivedDataCacheNotifications.h"
 #include "DerivedDataInformation.h"
 #include "Experimental/ZenServerInterface.h"
 #include "Framework/Application/SlateApplication.h"
@@ -61,8 +59,6 @@ void FZenEditor::StartupModule()
 #endif // WITH_RELOAD
 
 	FZenStausBarCommands::Register();
-
-	DerivedDataCacheNotifications.Reset(new FDerivedDataCacheNotifications);
 }
 
 void FZenEditor::ShutdownModule()
diff --git a/Engine/Source/Editor/ZenEditor/Public/ZenEditor.h b/Engine/Source/Editor/ZenEditor/Public/ZenEditor.h
--- a/Engine/Source/Editor/ZenEditor/Public/ZenEditor.h
+++ b/Engine/Source/Editor/ZenEditor/Public/ZenEditor.h
@@ -7,7 +7,6 @@
 #include "Modules/ModuleManager.h"
 #include "Templates/SharedPointer.h"
 #include "Templates/UniquePtr.h"
-#include "IDerivedDataCacheNotifications.h"
 
 #define UE_API ZENEDITOR_API
 
@@ -61,7 +60,6 @@ private:
 
 	TSharedPtr<SWindow>	SettingsWindow;
 	TSharedPtr<SDerivedDataCacheSettingsDialog> SettingsDialog;
-	TUniquePtr<IDerivedDataCacheNotifications>	DerivedDataCacheNotifications;
 };
 
 
diff --git a/Engine/Source/Runtime/Engine/Private/UnrealEngine.cpp b/Engine/Source/Runtime/Engine/Private/UnrealEngine.cpp
--- a/Engine/Source/Runtime/Engine/Private/UnrealEngine.cpp
+++ b/Engine/Source/Runtime/Engine/Private/UnrealEngine.cpp
@@ -2562,11 +2562,6 @@ void UEngine::Init(IEngineLoop* InEngineLoop)
 	}
 #endif // !(UE_BUILD_SHIPPING) || ENABLE_PGO_PROFILE
 
-	if (GetDerivedDataCache())
-	{
-		GetDerivedDataCacheRef().NotifyBootComplete();
-	}
-	
 	// register the engine with the travel and network failure broadcasts
 	// games can override these to provide proper behavior in each error case
 	OnTravelFailure().AddUObject(this, &UEngine::HandleTravelFailure);
diff --git a/Engine/Source/Runtime/NetworkFile/Private/NetworkPlatformFile.cpp b/Engine/Source/Runtime/NetworkFile/Private/NetworkPlatformFile.cpp
--- a/Engine/Source/Runtime/NetworkFile/Private/NetworkPlatformFile.cpp
+++ b/Engine/Source/Runtime/NetworkFile/Private/NetworkPlatformFile.cpp
@@ -45,8 +45,7 @@ static FString UnsafeProjectPlatformExtensionDir()
 }
 
 FNetworkPlatformFile::FNetworkPlatformFile()
-	: bHasLoadedDDCDirectories(false)
-	, InnerPlatformFile(NULL)
+	: InnerPlatformFile(NULL)
 	, bIsUsable(false)
 	, FileServerPort(0)
 	, ConnectionFlags(EConnectionFlags::None)
@@ -1380,7 +1379,7 @@ void FNetworkPlatformFile::MakeStandardNetworkFilename(FString& Filename)
 	Filename.RemoveFromEnd(NetworkPlatformFileEndChop, ESearchCase::CaseSensitive);
 }
 
-bool FNetworkPlatformFile::IsInLocalDirectoryUnGuarded(const FString& Filename)
+bool FNetworkPlatformFile::IsInLocalDirectory(const FString& Filename)
 {
 	// cache the directory of the input file
 	FString Directory = FPaths::GetPath(Filename);
@@ -1398,38 +1397,6 @@ bool FNetworkPlatformFile::IsInLocalDirectoryUnGuarded(const FString& Filename)
 	return false;
 }
 
-bool FNetworkPlatformFile::IsInLocalDirectory(const FString& Filename)
-{
-	if (!bHasLoadedDDCDirectories)
-	{
-		// need to be careful here to avoid initializing the DDC from the wrong thread or using LocalDirectories while it is being initialized
-		FScopeLock ScopeLock(&LocalDirectoriesCriticalSection);
-
-		if (IsInGameThread() && GConfig && GConfig->IsReadyForUse())
-		{
-			// one time DDC directory initialization
-			// add any DDC directories to our list of local directories (local = inner platform file, it may
-			// actually live on a server, but it will use the platform's file system)
-			if (GetDerivedDataCache())
-			{
-				TArray<FString> DdcDirectories;
-
-				GetDerivedDataCacheRef().GetDirectories(DdcDirectories);
-
-				LocalDirectories.Append(DdcDirectories);
-			}
-
-			FPlatformMisc::MemoryBarrier();
-			bHasLoadedDDCDirectories = true;
-		}
-
-		return IsInLocalDirectoryUnGuarded(Filename);
-	}
-
-	// once the DDC is initialized, we don't need to lock a critical section anymore
-	return IsInLocalDirectoryUnGuarded(Filename);
-}
-
 void FNetworkPlatformFile::PerformHeartbeat()
 {
 	// send the filename over (cast away const here because we know this << will not modify the string)
diff --git a/Engine/Source/Runtime/NetworkFile/Public/NetworkPlatformFile.h b/Engine/Source/Runtime/NetworkFile/Public/NetworkPlatformFile.h
--- a/Engine/Source/Runtime/NetworkFile/Public/NetworkPlatformFile.h
+++ b/Engine/Source/Runtime/NetworkFile/Public/NetworkPlatformFile.h
@@ -223,12 +223,6 @@ private:
 	 */
 	static NETWORKFILE_API bool IsAdditionalCookedFileExtension(const TCHAR* Ext);
 
-	/**
-	 * @return true if the path exists in a directory that should always use the local filesystem
-	 * This version does not worry about initialization or thread safety, do not call directly
-	 */
-	NETWORKFILE_API bool IsInLocalDirectoryUnGuarded(const FString& Filename);
-
 	/**
 	 * @return true if the path exists in a directory that should always use the local filesystem
 	 */
@@ -250,9 +244,6 @@ protected:
 
 protected:
 
-	/** This is true after the DDC directories have been loaded from the DDC system */
-	bool				bHasLoadedDDCDirectories;
-
 	/** The file interface to read/write local files with */
 	IPlatformFile*		InnerPlatformFile;
 
@@ -278,7 +269,6 @@ protected:
 	TArray<FString> LocalDirectories;
 
 	FCriticalSection	SynchronizationObject;
-	FCriticalSection	LocalDirectoriesCriticalSection;
 	bool				bIsUsable;
 	int32				FileServerPort;
 
